* refactored Material HLSL compiling - less redundancy, less error prone

diff --git a/Engine/Shaders/ShaderVersion.usf b/Engine/Shaders/ShaderVersion.usf
--- a/Engine/Shaders/ShaderVersion.usf
+++ b/Engine/Shaders/ShaderVersion.usf
@@ -2,4 +2,4 @@
 // Each time the console command is executed it generates a new GUID. As this file is included
 // in common.usf (which should be included in any shader) it allows to invalidate the shader DDC.
 // 
-// GUID = {0C2B0080-9297-4D12-9954-EFA779F1122A}..
+// GUID = 72EE1CFA4957C3B1A3BD51B124CD3652
\ No newline at end of file
diff --git a/Engine/Source/Runtime/Engine/Classes/Materials/Material.h b/Engine/Source/Runtime/Engine/Classes/Materials/Material.h
--- a/Engine/Source/Runtime/Engine/Classes/Materials/Material.h
+++ b/Engine/Source/Runtime/Engine/Classes/Materials/Material.h
@@ -777,7 +777,7 @@ public:
 	/** Checks to see if an input property should be active, based on the state of the material */
 	ENGINE_API virtual bool IsPropertyActive(EMaterialProperty InProperty) const;
 	/** Allows material properties to be compiled with the option of being overridden by the material attributes input. */
-	ENGINE_API virtual int32 CompileProperty( class FMaterialCompiler* Compiler, EMaterialProperty Property, float DefaultFloat, FLinearColor DefaultColor, const FVector4& DefaultVector );
+	ENGINE_API virtual int32 CompilePropertyEx( class FMaterialCompiler* Compiler, EMaterialProperty Property );
 	ENGINE_API virtual void ForceRecompileForRendering() override;
 	// End UMaterialInterface interface.
 
diff --git a/Engine/Source/Runtime/Engine/Classes/Materials/MaterialInstance.h b/Engine/Source/Runtime/Engine/Classes/Materials/MaterialInstance.h
--- a/Engine/Source/Runtime/Engine/Classes/Materials/MaterialInstance.h
+++ b/Engine/Source/Runtime/Engine/Classes/Materials/MaterialInstance.h
@@ -260,7 +260,7 @@ public:
 	/** Checks to see if an input property should be active, based on the state of the material */
 	ENGINE_API virtual bool IsPropertyActive(EMaterialProperty InProperty) const;
 	/** Allows material properties to be compiled with the option of being overridden by the material attributes input. */
-	ENGINE_API virtual int32 CompileProperty(class FMaterialCompiler* Compiler, EMaterialProperty Property, float DefaultFloat, FLinearColor DefaultColor, const FVector4& DefaultVector);
+	ENGINE_API virtual int32 CompilePropertyEx(class FMaterialCompiler* Compiler, EMaterialProperty Property);
 	// End UMaterialInterface interface.
 
 	// Begin UObject interface.
diff --git a/Engine/Source/Runtime/Engine/Classes/Materials/MaterialInterface.h b/Engine/Source/Runtime/Engine/Classes/Materials/MaterialInterface.h
--- a/Engine/Source/Runtime/Engine/Classes/Materials/MaterialInterface.h
+++ b/Engine/Source/Runtime/Engine/Classes/Materials/MaterialInterface.h
@@ -563,7 +563,7 @@ public:
 	ENGINE_API int32 CompileProperty( FMaterialCompiler* Compiler, EMaterialProperty Property );
 
 	/** Allows material properties to be compiled with the option of being overridden by the material attributes input. */
-	ENGINE_API virtual int32 CompileProperty( class FMaterialCompiler* Compiler, EMaterialProperty Property, float DefaultFloat, FLinearColor DefaultColor, const FVector4& DefaultVector );
+	ENGINE_API virtual int32 CompilePropertyEx( class FMaterialCompiler* Compiler, EMaterialProperty Property );
 
 protected:
 	/** Returns a bitfield indicating which feature levels should be compiled for rendering. */
diff --git a/Engine/Source/Runtime/Engine/Private/Materials/HLSLMaterialTranslator.h b/Engine/Source/Runtime/Engine/Private/Materials/HLSLMaterialTranslator.h
--- a/Engine/Source/Runtime/Engine/Private/Materials/HLSLMaterialTranslator.h
+++ b/Engine/Source/Runtime/Engine/Private/Materials/HLSLMaterialTranslator.h
@@ -337,7 +337,13 @@ public:
 			Chunk[MP_ClearCoat]						= ForceCast(Material->CompileProperty(MP_ClearCoat		       ,GetMaterialPropertyShaderFrequency(MP_ClearCoat),this),MCT_Float1);
 			Chunk[MP_ClearCoatRoughness]			= ForceCast(Material->CompileProperty(MP_ClearCoatRoughness    ,GetMaterialPropertyShaderFrequency(MP_ClearCoatRoughness),this),MCT_Float1);
 			Chunk[MP_AmbientOcclusion]				= ForceCast(Material->CompileProperty(MP_AmbientOcclusion      ,GetMaterialPropertyShaderFrequency(MP_AmbientOcclusion),this),MCT_Float1);
-			Chunk[MP_Refraction]					= ForceCast(Material->CompileProperty(MP_Refraction            ,GetMaterialPropertyShaderFrequency(MP_Refraction),this),MCT_Float2);
+
+			{
+				int32 UserRefraction = ForceCast(Material->CompileProperty(MP_Refraction, GetMaterialPropertyShaderFrequency(MP_Refraction), this), MCT_Float2);
+				int32 RefractionDepthBias = ForceCast(ScalarParameter(FName(TEXT("RefractionDepthBias")), GetRefractionDepthBiasValue()), MCT_Float1);
+
+				Chunk[MP_Refraction] = AppendVector(ForceCast(UserRefraction, MCT_Float1), RefractionDepthBias);
+			}
 
 			if (bCompileForComputeShader)
 			{
@@ -376,14 +382,9 @@ public:
 				MaterialCompilationOutput.bModifiesMeshPosition = true;
 				if( WPOChunk.UniformExpression && WPOChunk.UniformExpression->IsConstant() )
 				{
-					float DummyFloat;
-					FColor DummyColor;
-					FVector DefaultVector;
-					GetDefaultForMaterialProperty(MP_WorldPositionOffset, DummyFloat,DummyColor,DefaultVector);
-
 					FLinearColor WPOValue;
 					WPOChunk.UniformExpression->GetNumberValue(DummyContext, WPOValue);
-					if( FVector(WPOValue) == DefaultVector)
+					if (FVector(WPOValue) == FVector::ZeroVector)
 					{
 						MaterialCompilationOutput.bModifiesMeshPosition = false;
 					}
diff --git a/Engine/Source/Runtime/Engine/Private/Materials/Material.cpp b/Engine/Source/Runtime/Engine/Private/Materials/Material.cpp
--- a/Engine/Source/Runtime/Engine/Private/Materials/Material.cpp
+++ b/Engine/Source/Runtime/Engine/Private/Materials/Material.cpp
@@ -3194,6 +3194,9 @@ void UMaterial::AppendReferencedTextures(TArray<UTexture*>& InOutTextures) const
 		}
 		else
 		{
+			// should never be 0 but there is a bug, see TTP 340934
+			check(Expression);
+
 			UTexture* ReferencedTexture = Expression->GetReferencedTexture();
 
 			if (ReferencedTexture)
@@ -3277,59 +3280,54 @@ void UMaterial::GetReferencedParameterCollectionIds(TArray<FGuid>& Ids) const
 	}
 }
 
-int32 UMaterial::CompileProperty( FMaterialCompiler* Compiler, EMaterialProperty Property, float DefaultFloat, FLinearColor DefaultColor, const FVector4& DefaultVector )
+int32 UMaterial::CompilePropertyEx( FMaterialCompiler* Compiler, EMaterialProperty Property )
 {
-	int32 Ret = INDEX_NONE;
-
 	if( bUseMaterialAttributes && MP_DiffuseColor != Property && MP_SpecularColor != Property )
 	{
-		Ret = MaterialAttributes.Compile(Compiler,Property, DefaultFloat, DefaultColor, DefaultVector);
-	}
-	else
-	{
-		switch(Property)
-		{
-		case MP_Opacity: Ret = Opacity.Compile(Compiler,DefaultFloat); break;
-		case MP_OpacityMask: Ret = OpacityMask.Compile(Compiler,DefaultFloat); break;
-		case MP_Metallic: Ret = Metallic.Compile(Compiler,DefaultFloat); break;
-		case MP_Specular: Ret = Specular.Compile(Compiler,DefaultFloat); break;
-		case MP_Roughness: Ret = Roughness.Compile(Compiler,DefaultFloat); break;
-		case MP_TessellationMultiplier: Ret = TessellationMultiplier.Compile(Compiler,DefaultFloat); break;
-		case MP_ClearCoat: Ret = ClearCoat.Compile(Compiler,DefaultFloat); break;
-		case MP_ClearCoatRoughness: Ret = ClearCoatRoughness.Compile(Compiler,DefaultFloat); break;
-		case MP_AmbientOcclusion: Ret = AmbientOcclusion.Compile(Compiler,DefaultFloat); break;
-		case MP_Refraction: 
-			Ret = Compiler->AppendVector( 
-					Compiler->ForceCast(Refraction.Compile(Compiler,DefaultFloat), MCT_Float1), 
-					Compiler->ForceCast(Compiler->ScalarParameter(FName(TEXT("RefractionDepthBias")), Compiler->GetRefractionDepthBiasValue()), MCT_Float1) 
-					);
-			break;
-		case MP_EmissiveColor: Ret = EmissiveColor.Compile(Compiler,DefaultColor); break;
-		case MP_DiffuseColor: Ret = DiffuseColor.Compile(Compiler,DefaultColor); break;
-		case MP_SpecularColor: Ret = SpecularColor.Compile(Compiler,DefaultColor); break;
-		case MP_BaseColor: Ret = BaseColor.Compile(Compiler,DefaultColor); break;
-		case MP_SubsurfaceColor: Ret = SubsurfaceColor.Compile(Compiler,DefaultColor); break;
-		case MP_Normal: Ret = Normal.Compile(Compiler,DefaultVector); break;
-		case MP_WorldPositionOffset: Ret = WorldPositionOffset.Compile(Compiler,DefaultVector); break;
-		case MP_WorldDisplacement: Ret = WorldDisplacement.Compile(Compiler,DefaultVector); break;
-		};
+		return MaterialAttributes.CompileWithDefault(Compiler, Property);
+	}
+
+	switch (Property)
+	{
+		case MP_Opacity:				return Opacity.CompileWithDefault(Compiler, Property);
+		case MP_OpacityMask:			return OpacityMask.CompileWithDefault(Compiler, Property);
+		case MP_Metallic:				return Metallic.CompileWithDefault(Compiler, Property);
+		case MP_Specular:				return Specular.CompileWithDefault(Compiler, Property);
+		case MP_Roughness:				return Roughness.CompileWithDefault(Compiler, Property);
+		case MP_TessellationMultiplier:	return TessellationMultiplier.CompileWithDefault(Compiler, Property);
+		case MP_ClearCoat:				return ClearCoat.CompileWithDefault(Compiler, Property);
+		case MP_ClearCoatRoughness:		return ClearCoatRoughness.CompileWithDefault(Compiler, Property);
+		case MP_AmbientOcclusion:		return AmbientOcclusion.CompileWithDefault(Compiler, Property);
+		case MP_Refraction:				return Refraction.CompileWithDefault(Compiler, Property);
+		case MP_EmissiveColor:			return EmissiveColor.CompileWithDefault(Compiler, Property);
+		case MP_DiffuseColor:			return DiffuseColor.CompileWithDefault(Compiler, Property);
+		case MP_SpecularColor:			return SpecularColor.CompileWithDefault(Compiler, Property);
+		case MP_BaseColor:				return BaseColor.CompileWithDefault(Compiler, Property);
+		case MP_SubsurfaceColor:		return SubsurfaceColor.CompileWithDefault(Compiler, Property);
+		case MP_Normal:					return Normal.CompileWithDefault(Compiler, Property);
+		case MP_WorldPositionOffset:	return WorldPositionOffset.CompileWithDefault(Compiler, Property);
+		case MP_WorldDisplacement:		return WorldDisplacement.CompileWithDefault(Compiler, Property);
 
+		default:
 			if (Property >= MP_CustomizedUVs0 && Property <= MP_CustomizedUVs7)
 			{
 				const int32 TextureCoordinateIndex = Property - MP_CustomizedUVs0;
 
 				if (CustomizedUVs[TextureCoordinateIndex].Expression && TextureCoordinateIndex < NumCustomizedUVs)
 				{
-					Ret = CustomizedUVs[TextureCoordinateIndex].Compile(Compiler,FVector2D(DefaultVector.X, DefaultVector.Y));
+					return CustomizedUVs[TextureCoordinateIndex].CompileWithDefault(Compiler, Property);
 				}
 				else
 				{
 					// The user did not customize this UV, pass through the vertex texture coordinates
-					Ret = Compiler->TextureCoordinate(TextureCoordinateIndex, false, false);
+					return Compiler->TextureCoordinate(TextureCoordinateIndex, false, false);
 				}
 			}
+		
 	}
-	return Ret;
+
+	check(0);
+	return INDEX_NONE;
 }
 
 void UMaterial::NotifyCompilationFinished(FMaterialResource* CompiledResource)
diff --git a/Engine/Source/Runtime/Engine/Private/Materials/MaterialExpressions.cpp b/Engine/Source/Runtime/Engine/Private/Materials/MaterialExpressions.cpp
--- a/Engine/Source/Runtime/Engine/Private/Materials/MaterialExpressions.cpp
+++ b/Engine/Source/Runtime/Engine/Private/Materials/MaterialExpressions.cpp
@@ -3317,11 +3317,8 @@ int32 UMaterialExpressionBreakMaterialAttributes::Compile(class FMaterialCompile
 	//Here we don't care about any multiplex index coming in.
 	//We pass through our output index as the multiplex index so the MakeMaterialAttriubtes node at the other end can send us the right data.
 	EMaterialProperty Property = GetMaterialPropertyFromInputOutputIndex(OutputIndex);
-	float DefaultFloat;
-	FColor DefaultColor;
-	FVector DefaultVector;
-	GetDefaultForMaterialProperty(Property, DefaultFloat, DefaultColor, DefaultVector);
-	return MaterialAttributes.Compile(Compiler, Property, DefaultFloat, DefaultColor, DefaultVector);
+
+	return MaterialAttributes.CompileWithDefault(Compiler, Property);
 }
 
 void UMaterialExpressionBreakMaterialAttributes::GetCaption(TArray<FString>& OutCaptions) const
diff --git a/Engine/Source/Runtime/Engine/Private/Materials/MaterialInstance.cpp b/Engine/Source/Runtime/Engine/Private/Materials/MaterialInstance.cpp
--- a/Engine/Source/Runtime/Engine/Private/Materials/MaterialInstance.cpp
+++ b/Engine/Source/Runtime/Engine/Private/Materials/MaterialInstance.cpp
@@ -2439,7 +2439,7 @@ bool UMaterialInstance::IsPropertyActive(EMaterialProperty InProperty) const
 	return true;
 }
 
-int32 UMaterialInstance::CompileProperty( class FMaterialCompiler* Compiler, EMaterialProperty Property, float DefaultFloat, FLinearColor DefaultColor, const FVector4& DefaultVector )
+int32 UMaterialInstance::CompilePropertyEx( class FMaterialCompiler* Compiler, EMaterialProperty Property )
 {
-	return Parent ? Parent->CompileProperty(Compiler,Property,DefaultFloat,DefaultColor,DefaultVector) : INDEX_NONE;
+	return Parent ? Parent->CompilePropertyEx(Compiler, Property) : INDEX_NONE;
 }
diff --git a/Engine/Source/Runtime/Engine/Private/Materials/MaterialShared.cpp b/Engine/Source/Runtime/Engine/Private/Materials/MaterialShared.cpp
--- a/Engine/Source/Runtime/Engine/Private/Materials/MaterialShared.cpp
+++ b/Engine/Source/Runtime/Engine/Private/Materials/MaterialShared.cpp
@@ -91,126 +91,79 @@ void FExpressionInput::Connect( int32 InOutputIndex, class UMaterialExpression*
 	MaskA = Output->MaskA;
 }
 
-//
-//	FColorMaterialInput::Compile
-//
-
-int32 FColorMaterialInput::Compile(class FMaterialCompiler* Compiler,const FColor& Default)
+int32 FColorMaterialInput::CompileWithDefault(class FMaterialCompiler* Compiler, EMaterialProperty Property)
 {
-	if(UseConstant)
+	if (UseConstant)
 	{
-		FLinearColor	LinearColor(Constant);
-		return Compiler->Constant3(LinearColor.R,LinearColor.G,LinearColor.B);
+		FLinearColor LinearColor(Constant);
+		return Compiler->Constant3(LinearColor.R, LinearColor.G, LinearColor.B);
 	}
-	else if(Expression)
+	else if (Expression)
 	{
 		int32 ResultIndex = FExpressionInput::Compile(Compiler);
-		if (ResultIndex == INDEX_NONE)
-		{
-			FLinearColor	LinearColor(Default);
-			return Compiler->Constant3(LinearColor.R,LinearColor.G,LinearColor.B);
-		}
-		else
+		if (ResultIndex != INDEX_NONE)
 		{
 			return ResultIndex;
 		}
 	}
-	else
-	{
-		FLinearColor	LinearColor(Default);
-		return Compiler->Constant3(LinearColor.R,LinearColor.G,LinearColor.B);
-	}
-}
 
-//
-//	FScalarMaterialInput::Compile
-//
+	return GetDefaultExpressionForMaterialProperty(Compiler, Property);
+}
 
-int32 FScalarMaterialInput::Compile(class FMaterialCompiler* Compiler,float Default)
+int32 FScalarMaterialInput::CompileWithDefault(class FMaterialCompiler* Compiler, EMaterialProperty Property)
 {
-	if(UseConstant)
+	if (UseConstant)
 	{
 		return Compiler->Constant(Constant);
 	}
-	else if(Expression)
+	else if (Expression)
 	{
 		int32 ResultIndex = FExpressionInput::Compile(Compiler);
-		if (ResultIndex == INDEX_NONE)
-		{
-			return Compiler->Constant(Default);
-		}
-		else
+		if (ResultIndex != INDEX_NONE)
 		{
 			return ResultIndex;
 		}
 	}
-	else
-	{
-		return Compiler->Constant(Default);
-	}
-}
 
-//
-//	FVectorMaterialInput::Compile
-//
+	return GetDefaultExpressionForMaterialProperty(Compiler, Property);
+}
 
-int32 FVectorMaterialInput::Compile(class FMaterialCompiler* Compiler,const FVector& Default)
+int32 FVectorMaterialInput::CompileWithDefault(class FMaterialCompiler* Compiler, EMaterialProperty Property)
 {
-	if(UseConstant)
+	if (UseConstant)
 	{
-		return Compiler->Constant3(Constant.X,Constant.Y,Constant.Z);
+		return Compiler->Constant3(Constant.X, Constant.Y, Constant.Z);
 	}
 	else if(Expression)
 	{
 		int32 ResultIndex = FExpressionInput::Compile(Compiler);
-		if (ResultIndex == INDEX_NONE)
-		{
-			return Compiler->Constant3(Default.X,Default.Y,Default.Z);
-		}
-		else
+		if (ResultIndex != INDEX_NONE)
 		{
 			return ResultIndex;
 		}
 	}
-	else
-	{
-		return Compiler->Constant3(Default.X,Default.Y,Default.Z);
-	}
+	return GetDefaultExpressionForMaterialProperty(Compiler, Property);
 }
 
-//
-//	FVector2MaterialInput::Compile
-//
-
-int32 FVector2MaterialInput::Compile(class FMaterialCompiler* Compiler,const FVector2D& Default)
+int32 FVector2MaterialInput::CompileWithDefault(class FMaterialCompiler* Compiler, EMaterialProperty Property)
 {
-	if(UseConstant)
+	if (UseConstant)
 	{
-		return Compiler->Constant2(Constant.X,Constant.Y);
+		return Compiler->Constant2(Constant.X, Constant.Y);
 	}
-	else if(Expression)
+	else if (Expression)
 	{
 		int32 ResultIndex = FExpressionInput::Compile(Compiler);
-		if (ResultIndex == INDEX_NONE)
-		{
-			return Compiler->Constant2(Default.X,Default.Y);
-		}
-		else
+		if (ResultIndex != INDEX_NONE)
 		{
 			return ResultIndex;
 		}
 	}
-	else
-	{
-		return Compiler->Constant2(Default.X,Default.Y);
-	}
-}
 
-//
-//	FMaterialAttributesInput::Compile
-//
+	return GetDefaultExpressionForMaterialProperty(Compiler, Property);
+}
 
-int32 FMaterialAttributesInput::Compile(class FMaterialCompiler* Compiler, EMaterialProperty Property, float DefaultFloat, const FColor& DefaultColor, const FVector& DefaultVector)
+int32 FMaterialAttributesInput::CompileWithDefault(class FMaterialCompiler* Compiler, EMaterialProperty Property)
 {
 	int32 Ret = INDEX_NONE;
 	if(Expression)
@@ -230,35 +183,7 @@ int32 FMaterialAttributesInput::Compile(class FMaterialCompiler* Compiler, EMate
 
 	if( Ret == INDEX_NONE )
 	{
-		//Fallback to defaults
-		switch(Property)
-		{
-		case MP_EmissiveColor: Ret = Compiler->Constant3(DefaultColor.R,DefaultColor.G,DefaultColor.B); break;
-		case MP_Opacity: Ret = Compiler->Constant(DefaultFloat); break;
-		case MP_OpacityMask: Ret = Compiler->Constant(DefaultFloat); break;
-		case MP_DiffuseColor: Ret = Compiler->Constant3(DefaultColor.R,DefaultColor.G,DefaultColor.B); break;
-		case MP_SpecularColor: Ret = Compiler->Constant3(DefaultColor.R,DefaultColor.G,DefaultColor.B); break;
-		case MP_BaseColor: Ret = Compiler->Constant3(DefaultColor.R,DefaultColor.G,DefaultColor.B); break;
-		case MP_Metallic: Ret = Compiler->Constant(DefaultFloat); break;
-		case MP_Specular: Ret = Compiler->Constant(DefaultFloat); break;
-		case MP_Roughness: Ret = Compiler->Constant(DefaultFloat); break;
-		case MP_Normal: Ret = Compiler->Constant3(DefaultVector.X,DefaultVector.Y,DefaultVector.Z); break;
-		case MP_WorldPositionOffset: Ret = Compiler->Constant3(DefaultVector.X,DefaultVector.Y,DefaultVector.Z); break;
-		case MP_WorldDisplacement : Ret = Compiler->Constant3(DefaultVector.X,DefaultVector.Y,DefaultVector.Z); break;
-		case MP_TessellationMultiplier: Ret = Compiler->Constant(DefaultFloat); break;
-		case MP_SubsurfaceColor: Ret = Compiler->Constant3(DefaultColor.R,DefaultColor.G,DefaultColor.B); break;
-		case MP_ClearCoat: Ret = Compiler->Constant(DefaultFloat); break;
-		case MP_ClearCoatRoughness: Ret = Compiler->Constant(DefaultFloat); break;
-		case MP_AmbientOcclusion: Ret = Compiler->Constant(DefaultFloat); break;
-		case MP_Refraction: Ret = Compiler->Constant2(DefaultVector.X,DefaultVector.Y); break;
-		};
-
-		if (Property >= MP_CustomizedUVs0 && Property <= MP_CustomizedUVs7)
-		{
-			const int32 TextureCoordinateIndex = Property - MP_CustomizedUVs0;
-			// The user did not customize this UV, pass through the vertex texture coordinates
-			Ret = Compiler->TextureCoordinate(TextureCoordinateIndex, false, false);
-		}
+		Ret = GetDefaultExpressionForMaterialProperty(Compiler, Property);
 	}
 
 	return Ret;
@@ -2201,33 +2126,48 @@ int32 GetInputOutputIndexFromMaterialProperty(EMaterialProperty Property)
 	return -1;
 }
 
-void GetDefaultForMaterialProperty(EMaterialProperty Property, float& OutDefaultFloat, FColor& OutDefaultColor, FVector& OutDefaultVector )
-{
-	OutDefaultFloat = 0;
-	OutDefaultColor = FColor::Black;
-	OutDefaultVector = FVector::ZeroVector;
+int32 GetDefaultExpressionForMaterialProperty(FMaterialCompiler* Compiler, EMaterialProperty Property)
+{
+	switch (Property)
+	{
+		case MP_Opacity:				return Compiler->Constant(1.0f);
+		case MP_OpacityMask:			return Compiler->Constant(1.0f);
+		case MP_Metallic:				return Compiler->Constant(0.0f);
+		case MP_Specular:				return Compiler->Constant(0.5f);
+		case MP_Roughness:				return Compiler->Constant(0.5f);
+		case MP_TessellationMultiplier:	return Compiler->Constant(1.0f);
+		case MP_ClearCoat:				return Compiler->Constant(1.0f);
+		case MP_ClearCoatRoughness:		return Compiler->Constant(0.1f);
+		case MP_AmbientOcclusion:		return Compiler->Constant(1.0f);
+
+		case MP_EmissiveColor:			return Compiler->Constant3(0, 0, 0);
+		case MP_DiffuseColor:			return Compiler->Constant3(0, 0, 0);
+		case MP_SpecularColor:			return Compiler->Constant3(0, 0, 0);
+		case MP_BaseColor:				return Compiler->Constant3(0, 0, 0);
+		case MP_SubsurfaceColor:		return Compiler->Constant3(1, 1, 1);
+		case MP_Normal:					return Compiler->Constant3(0, 0, 1);
+		case MP_WorldPositionOffset:	return Compiler->Constant3(0, 0, 0);
+		case MP_WorldDisplacement:		return Compiler->Constant3(0, 0, 0);
+		case MP_Refraction:				return Compiler->Constant3(1, 0, 0);
+
+		case MP_CustomizedUVs0:
+		case MP_CustomizedUVs1:
+		case MP_CustomizedUVs2:
+		case MP_CustomizedUVs3:
+		case MP_CustomizedUVs4:
+		case MP_CustomizedUVs5:
+		case MP_CustomizedUVs6:
+		case MP_CustomizedUVs7:
+			{
+//				return Compiler->Constant2(0.0f, 0.0f);				
+				const int32 TextureCoordinateIndex = Property - MP_CustomizedUVs0;
+				// The user did not customize this UV, pass through the vertex texture coordinates
+				return Compiler->TextureCoordinate(TextureCoordinateIndex, false, false);
+			}
+	}
 
-	switch(Property)
-	{
-	case MP_Opacity:				OutDefaultFloat = 1.0f; break;
-	case MP_OpacityMask:			OutDefaultFloat = 1.0f; break;
-	case MP_Metallic:				OutDefaultFloat = 0.0f; break;
-	case MP_Specular:				OutDefaultFloat = 0.5f; break;
-	case MP_Roughness:				OutDefaultFloat = 0.5f; break;
-	case MP_TessellationMultiplier: OutDefaultFloat = 1.0f; break;
-	case MP_ClearCoat:				OutDefaultFloat = 1.0f; break;
-	case MP_ClearCoatRoughness:		OutDefaultFloat = 0.1f; break;
-	case MP_AmbientOcclusion:		OutDefaultFloat = 1.0f; break;
-	case MP_EmissiveColor:			OutDefaultColor = FColor(0,0,0); break;
-	case MP_DiffuseColor:			OutDefaultColor = FColor(0,0,0); break;
-	case MP_SpecularColor:			OutDefaultColor = FColor(0,0,0); break;
-	case MP_BaseColor:				OutDefaultColor = FColor(0,0,0); break;
-	case MP_SubsurfaceColor:		OutDefaultColor = FColor(1,1,1); break;
-	case MP_Normal:					OutDefaultVector = FVector(0,0,1); break;
-	case MP_WorldPositionOffset:	OutDefaultVector = FVector::ZeroVector; break;
-	case MP_WorldDisplacement:		OutDefaultVector = FVector::ZeroVector; break;
-	case MP_Refraction:				OutDefaultVector = FVector(1,0,0); break;
-	};
+	check(0);
+	return INDEX_NONE;
 }
 
 FString GetNameOfMaterialProperty(EMaterialProperty Property)
@@ -2269,7 +2209,7 @@ bool UMaterialInterface::IsPropertyActive(EMaterialProperty InProperty)const
 	return false;
 }
 
-int32 UMaterialInterface::CompileProperty( class FMaterialCompiler* Compiler, EMaterialProperty Property, float DefaultFloat, FLinearColor DefaultColor, const FVector4& DefaultVector )
+int32 UMaterialInterface::CompilePropertyEx( class FMaterialCompiler* Compiler, EMaterialProperty Property )
 {
 	return INDEX_NONE;
 }
@@ -2278,63 +2218,13 @@ int32 UMaterialInterface::CompileProperty( FMaterialCompiler* Compiler, EMateria
 {
 	int32 Ret = INDEX_NONE;
 
-	float DefaultFloat;
-	FColor DefaultColor;
-	FVector DefaultVector;
-	GetDefaultForMaterialProperty(Property, DefaultFloat, DefaultColor, DefaultVector);
-
 	if (IsPropertyActive(Property))
 	{
-		return CompileProperty(Compiler, Property, DefaultFloat, DefaultColor, DefaultVector);
+		return CompilePropertyEx(Compiler, Property);
 	}
 	else
 	{
-		FLinearColor DefaultLinearColor(DefaultColor.ReinterpretAsLinear());
-
-		switch(Property)
-		{
-		case MP_Opacity:
-		case MP_OpacityMask:
-		case MP_Metallic:
-		case MP_Specular:
-		case MP_Roughness:
-		case MP_TessellationMultiplier:
-		case MP_ClearCoat:
-		case MP_ClearCoatRoughness:
-		case MP_AmbientOcclusion:
-			{
-				return Compiler->Constant(DefaultFloat);
-			}
-		case MP_EmissiveColor:
-		case MP_DiffuseColor:
-		case MP_SpecularColor:
-		case MP_BaseColor:
-		case MP_SubsurfaceColor:
-			{
-				return Compiler->Constant3(DefaultLinearColor.R, DefaultLinearColor.G, DefaultLinearColor.B);
-			}
-
-		case MP_Normal:
-		case MP_WorldPositionOffset:
-		case MP_WorldDisplacement:
-		case MP_Refraction:
-			{
-				return Compiler->Constant3(DefaultVector.X, DefaultVector.Y, DefaultVector.Z);
-			}
-
-		default:
-			{
-				if (Property >= MP_CustomizedUVs0 && Property <= MP_CustomizedUVs7)
-				{
-					return Compiler->Constant2(DefaultVector.X, DefaultVector.Y);
-				}
-				else
-				{
-					check(false);
-					break;
-				}
-			}
-		};
+		return GetDefaultExpressionForMaterialProperty(Compiler, Property);
 	}
 	return Ret;
 }
diff --git a/Engine/Source/Runtime/Engine/Public/MaterialExpressionIO.h b/Engine/Source/Runtime/Engine/Public/MaterialExpressionIO.h
--- a/Engine/Source/Runtime/Engine/Public/MaterialExpressionIO.h
+++ b/Engine/Source/Runtime/Engine/Public/MaterialExpressionIO.h
@@ -96,24 +96,21 @@ template<class InputType> struct FMaterialInput : FExpressionInput
 
 struct FColorMaterialInput : FMaterialInput<FColor>
 {
-	ENGINE_API int32 Compile(class FMaterialCompiler* Compiler, const FColor& Default);
+	ENGINE_API int32 CompileWithDefault(class FMaterialCompiler* Compiler, EMaterialProperty Property);
 };
 struct FScalarMaterialInput : FMaterialInput<float>
 {
-	ENGINE_API int32 Compile(class FMaterialCompiler* Compiler, float Default);
+	ENGINE_API int32 CompileWithDefault(class FMaterialCompiler* Compiler, EMaterialProperty Property);
 };
-
 struct FVectorMaterialInput : FMaterialInput<FVector>
 {
-	ENGINE_API int32 Compile(class FMaterialCompiler* Compiler, const FVector& Default);
+	ENGINE_API int32 CompileWithDefault(class FMaterialCompiler* Compiler, EMaterialProperty Property);
 };
-
 struct FVector2MaterialInput : FMaterialInput<FVector2D>
 {
-	int32 Compile(class FMaterialCompiler* Compiler, const FVector2D& Default);
+	ENGINE_API int32 CompileWithDefault(class FMaterialCompiler* Compiler, EMaterialProperty Property);
 };
-
 struct FMaterialAttributesInput : FMaterialInput<int32>
 {
-	ENGINE_API int32 Compile(class FMaterialCompiler* Compiler, EMaterialProperty Property, float DefaultFloat, const FColor& DefaultColor, const FVector& DefaultVector);
+	ENGINE_API int32 CompileWithDefault(class FMaterialCompiler* Compiler, EMaterialProperty Property);
 };
\ No newline at end of file
diff --git a/Engine/Source/Runtime/Engine/Public/MaterialShared.h b/Engine/Source/Runtime/Engine/Public/MaterialShared.h
--- a/Engine/Source/Runtime/Engine/Public/MaterialShared.h
+++ b/Engine/Source/Runtime/Engine/Public/MaterialShared.h
@@ -1583,9 +1583,9 @@ ENGINE_API EMaterialProperty GetMaterialPropertyFromInputOutputIndex(int32 Index
 ENGINE_API int32 GetInputOutputIndexFromMaterialProperty(EMaterialProperty Property);
 
 /**
- * @return Gets the default value for a material property
+ * @return Gets the default (usually constant) for a material property
  */
-ENGINE_API void GetDefaultForMaterialProperty(EMaterialProperty Property, float& OutDefaultFloat, FColor& OutDefaultColor, FVector& OutDefaultVector);
+ENGINE_API int32 GetDefaultExpressionForMaterialProperty(FMaterialCompiler* Compiler, EMaterialProperty Property);
 
 /**
  * @return Gets the name of a property.
