- Fixed a compile error

diff --git a/Engine/Source/Runtime/Engine/Private/ContentStreaming.cpp b/Engine/Source/Runtime/Engine/Private/ContentStreaming.cpp
--- a/Engine/Source/Runtime/Engine/Private/ContentStreaming.cpp
+++ b/Engine/Source/Runtime/Engine/Private/ContentStreaming.cpp
@@ -159,7 +159,7 @@ static TAutoConsoleVariable<int32> CVarStreamingPoolSize(
 /** For debugging purposes: Whether to consider the time factor when streaming. Turning it off is easier for debugging. */
 bool GStreamWithTimeFactor		= true;
 
-extern UPrimitiveComponent* GDebugSelectedComponent;
+extern ENGINE_API UPrimitiveComponent* GDebugSelectedComponent;
 
 #if STATS
 /** Ringbuffer of bandwidth samples for streaming in mip-levels (MB/sec). */
