Don't crash in GetMousePosition() if Player is NULL.

diff --git a/Engine/Source/Runtime/Engine/Classes/GameFramework/PlayerController.h b/Engine/Source/Runtime/Engine/Classes/GameFramework/PlayerController.h
--- a/Engine/Source/Runtime/Engine/Classes/GameFramework/PlayerController.h
+++ b/Engine/Source/Runtime/Engine/Classes/GameFramework/PlayerController.h
@@ -860,9 +860,9 @@ public:
 	UFUNCTION(BlueprintCallable, Category="Game|Player")
 	void GetInputMotionState(FVector& Tilt, FVector& RotationRate, FVector& Gravity, FVector& Acceleration) const;
 
-	/** Retrieves the X and Y screen coordinates of the mouse cursor. Returns false if the touch index is not down */
+	/** Retrieves the X and Y screen coordinates of the mouse cursor. Returns false if there is no associated mouse device */
 	UFUNCTION(BlueprintCallable, Category="Game|Player")
-	void GetMousePosition(float& LocationX, float& LocationY) const;
+	bool GetMousePosition(float& LocationX, float& LocationY) const;
 
 	/** Returns how long the given key/button has been down.  Returns 0 if it's up or it just went down this frame. */
 	UFUNCTION(BlueprintCallable, Category="Game|Player")
diff --git a/Engine/Source/Runtime/Engine/Private/PlayerController.cpp b/Engine/Source/Runtime/Engine/Private/PlayerController.cpp
--- a/Engine/Source/Runtime/Engine/Private/PlayerController.cpp
+++ b/Engine/Source/Runtime/Engine/Private/PlayerController.cpp
@@ -3860,11 +3860,19 @@ float APlayerController::GetInputKeyTimeDown(const FKey Key) const
 	return (PlayerInput ? PlayerInput->GetTimeDown(Key) : 0.f);
 }
 
-void APlayerController::GetMousePosition(float& LocationX, float& LocationY) const
+bool APlayerController::GetMousePosition(float& LocationX, float& LocationY) const
 {
+	// TODO: Return false if there is no mouse attached
+	if (Player == nullptr)
+	{
+		return false;
+	}
+
 	const FVector2D MousePosition = CastChecked<ULocalPlayer>(Player)->ViewportClient->GetMousePosition();
 	LocationX = MousePosition.X;
 	LocationY = MousePosition.Y;
+
+	return true;
 }
 
 void APlayerController::GetInputMouseDelta(float& DeltaX, float& DeltaY) const
