Fix for TTP 341159 - Use GMaxRHIFeature level for selecting a particle systems material when it isn't attached to a valid world

diff --git a/Engine/Source/Runtime/Engine/Private/Particles/ParticleModules_Collision.cpp b/Engine/Source/Runtime/Engine/Private/Particles/ParticleModules_Collision.cpp
--- a/Engine/Source/Runtime/Engine/Private/Particles/ParticleModules_Collision.cpp
+++ b/Engine/Source/Runtime/Engine/Private/Particles/ParticleModules_Collision.cpp
@@ -658,7 +658,7 @@ bool UParticleModuleCollisionGPU::IsValidForLODLevel(UParticleLODLevel* LODLevel
 	check(Material);
 
 	EBlendMode BlendMode = BLEND_Opaque;
-	const FMaterialResource* MaterialResource = Material->GetMaterialResource(GetWorld()->FeatureLevel);
+	const FMaterialResource* MaterialResource = Material->GetMaterialResource(GetWorld() ? GetWorld()->FeatureLevel : GMaxRHIFeatureLevel);
 	if(MaterialResource)
 	{
 		BlendMode = MaterialResource->GetBlendMode();
