Ensure that actors get attached to the correct blueprint components in cases where RerunConstructionScript has been run on its parent before RegisterComponent is called on the child

diff --git a/Engine/Source/Runtime/Engine/Private/Actor.cpp b/Engine/Source/Runtime/Engine/Private/Actor.cpp
--- a/Engine/Source/Runtime/Engine/Private/Actor.cpp
+++ b/Engine/Source/Runtime/Engine/Private/Actor.cpp
@@ -3088,6 +3088,18 @@ void AActor::RegisterAllComponents()
 				USceneComponent* PendingAttachParent = SceneComp->AttachParent;
 				SceneComp->AttachParent = NULL;
 
+				// If our attach parent is a blueprint component that is marked pending kill belonging to another actor,
+				// then rerun construction script has almost certainly been run on it and we are just going to pick its
+				// current root component instead
+				if (PendingAttachParent->bCreatedByConstructionScript && PendingAttachParent->IsPendingKill())
+				{
+					AActor* PendingAttachOwner = PendingAttachParent->GetOwner();
+					if (PendingAttachOwner != this)
+					{
+						PendingAttachParent = PendingAttachOwner->GetRootComponent();
+					}
+				}
+
 				FName PendingAttachSocketName = SceneComp->AttachSocketName;
 				SceneComp->AttachSocketName = NAME_None;
 
