* refactored material compiling, simpler, less code, less error prone, should not have any visual impact

diff --git a/Engine/Shaders/ShaderVersion.usf b/Engine/Shaders/ShaderVersion.usf
--- a/Engine/Shaders/ShaderVersion.usf
+++ b/Engine/Shaders/ShaderVersion.usf
@@ -2,4 +2,4 @@
 // Each time the console command is executed it generates a new GUID. As this file is included
 // in common.usf (which should be included in any shader) it allows to invalidate the shader DDC.
 // 
-// GUID = 72EE1CFA4957C3B1A3BD51B124CD3652
\ No newline at end of file
+// GUID = A7702884495089DC221506B8C4961E6C
\ No newline at end of file
diff --git a/Engine/Source/Editor/MaterialEditor/Private/MaterialEditor.cpp b/Engine/Source/Editor/MaterialEditor/Private/MaterialEditor.cpp
--- a/Engine/Source/Editor/MaterialEditor/Private/MaterialEditor.cpp
+++ b/Engine/Source/Editor/MaterialEditor/Private/MaterialEditor.cpp
@@ -111,32 +111,37 @@ bool FMatExpressionPreview::ShouldCache(EShaderPlatform Platform, const FShaderT
 	return false;
 }
 
-// Material properties.
-/** Entry point for compiling a specific material property.  This must call SetMaterialProperty. */
-int32 FMatExpressionPreview::CompileProperty(EMaterialProperty Property,EShaderFrequency InShaderFrequency,FMaterialCompiler* Compiler) const
+int32 FMatExpressionPreview::CompilePropertyAndSetMaterialProperty(EMaterialProperty Property, FMaterialCompiler* Compiler, EShaderFrequency OverrideShaderFrequency) const
 {
-	Compiler->SetMaterialProperty(Property, InShaderFrequency);
+	// needs to be called in this function!!
+	Compiler->SetMaterialProperty(Property, OverrideShaderFrequency);
+
+	int32 Ret = INDEX_NONE;
+
 	if( Property == MP_EmissiveColor && Expression.IsValid())
 	{
 		// Hardcoding output 0 as we don't have the UI to specify any other output
 		const int32 OutputIndex = 0;
 		// Get back into gamma corrected space, as DrawTile does not do this adjustment.
-		return Compiler->Power(Compiler->Max(Expression->CompilePreview(Compiler, OutputIndex, -1), Compiler->Constant(0)), Compiler->Constant(1.f/2.2f));
+		Ret = Compiler->Power(Compiler->Max(Expression->CompilePreview(Compiler, OutputIndex, -1), Compiler->Constant(0)), Compiler->Constant(1.f / 2.2f));
 	}
 	else if ( Property == MP_WorldPositionOffset)
 	{
 		//set to 0 to prevent off by 1 pixel errors
-		return Compiler->Constant(0.0f);
+		Ret = Compiler->Constant(0.0f);
 	}
 	else if (Property >= MP_CustomizedUVs0 && Property <= MP_CustomizedUVs7)
 	{
 		const int32 TextureCoordinateIndex = Property - MP_CustomizedUVs0;
-		return Compiler->TextureCoordinate(TextureCoordinateIndex, false, false);
+		Ret = Compiler->TextureCoordinate(TextureCoordinateIndex, false, false);
 	}
 	else
 	{
-		return Compiler->Constant(1.0f);
+		Ret = Compiler->Constant(1.0f);
 	}
+
+	// output should always be the right type for this property
+	return Compiler->ForceCast(Ret, GetMaterialPropertyType(Property));
 }
 
 void FMatExpressionPreview::NotifyCompilationFinished()
diff --git a/Engine/Source/Editor/MaterialEditor/Private/MaterialEditor.h b/Engine/Source/Editor/MaterialEditor/Private/MaterialEditor.h
--- a/Engine/Source/Editor/MaterialEditor/Private/MaterialEditor.h
+++ b/Engine/Source/Editor/MaterialEditor/Private/MaterialEditor.h
@@ -107,7 +107,7 @@ public:
 	
 	// Material properties.
 	/** Entry point for compiling a specific material property.  This must call SetMaterialProperty. */
-	virtual int32 CompileProperty(EMaterialProperty Property,EShaderFrequency InShaderFrequency,FMaterialCompiler* Compiler) const;
+	virtual int32 CompilePropertyAndSetMaterialProperty(EMaterialProperty Property, FMaterialCompiler* Compiler, EShaderFrequency OverrideShaderFrequency) const;
 
 	virtual int32 GetMaterialDomain() const override { return MD_Surface; }
 	virtual FString GetMaterialUsageDescription() const { return FString::Printf(TEXT("FMatExpressionPreview %s"), Expression.IsValid() ? *Expression->GetName() : TEXT("NULL")); }
diff --git a/Engine/Source/Editor/UnrealEd/Classes/MaterialGraph/MaterialGraph.h b/Engine/Source/Editor/UnrealEd/Classes/MaterialGraph/MaterialGraph.h
--- a/Engine/Source/Editor/UnrealEd/Classes/MaterialGraph/MaterialGraph.h
+++ b/Engine/Source/Editor/UnrealEd/Classes/MaterialGraph/MaterialGraph.h
@@ -15,8 +15,6 @@ struct FMaterialInputInfo
 {
 	/** Name of the input shown to user */
 	FText Name;
-	/** The actual input */
-	FExpressionInput* Input;
 	/** Type of the input */
 	EMaterialProperty Property;
 
@@ -26,12 +24,23 @@ struct FMaterialInputInfo
 	}
 
 	/** Constructor */
-	FMaterialInputInfo(const FText& InName, FExpressionInput* InInput, EMaterialProperty InProperty)
+	FMaterialInputInfo(const FText& InName, EMaterialProperty InProperty)
 		:	Name( InName )
-		,	Input( InInput )
 		,	Property( InProperty )
 	{
 	}
+
+	FExpressionInput* GetInput(UMaterial* Material) const
+	{
+		check(Material);
+
+		auto Ret = Material->GetExpressionInputForProperty(Property);
+
+		check(Ret);
+
+		return Ret;
+	}
+
 };
 
 UCLASS(MinimalAPI)
diff --git a/Engine/Source/Editor/UnrealEd/Private/EditorExporters.cpp b/Engine/Source/Editor/UnrealEd/Private/EditorExporters.cpp
--- a/Engine/Source/Editor/UnrealEd/Private/EditorExporters.cpp
+++ b/Engine/Source/Editor/UnrealEd/Private/EditorExporters.cpp
@@ -808,11 +808,21 @@ public:
 
 	// Material properties.
 	/** Entry point for compiling a specific material property.  This must call SetMaterialProperty. */
-	virtual int32 CompileProperty(EMaterialProperty Property,EShaderFrequency InShaderFrequency,FMaterialCompiler* Compiler) const
+	virtual int32 CompilePropertyAndSetMaterialProperty(EMaterialProperty Property, FMaterialCompiler* Compiler, EShaderFrequency OverrideShaderFrequency) const
+	{
+		// needs to be called in this function!!
+		Compiler->SetMaterialProperty(Property, OverrideShaderFrequency);
+
+		int32 Ret = CompilePropertyAndSetMaterialPropertyWithoutCast(Property, Compiler);
+
+		return Compiler->ForceCast(Ret, GetMaterialPropertyType(Property));
+	}
+
+	/** helper for CompilePropertyAndSetMaterialProperty() */
+	int32 CompilePropertyAndSetMaterialPropertyWithoutCast(EMaterialProperty Property, FMaterialCompiler* Compiler) const
 	{
 		static const auto UseDiffuseSpecularMaterialInputs = IConsoleManager::Get().FindTConsoleVariableDataInt(TEXT("r.UseDiffuseSpecularMaterialInputs"));
 				
-		Compiler->SetMaterialProperty(Property, InShaderFrequency);
 		// MAKE SURE THIS MATCHES THE CHART IN WillFillData
 		// 						  RETURNED VALUES (F16 'textures')
 		// 	BLEND MODE  | DIFFUSE     | SPECULAR     | EMISSIVE    | NORMAL    | TRANSMISSIVE              |
diff --git a/Engine/Source/Editor/UnrealEd/Private/Lightmass/LightmassRender.cpp b/Engine/Source/Editor/UnrealEd/Private/Lightmass/LightmassRender.cpp
--- a/Engine/Source/Editor/UnrealEd/Private/Lightmass/LightmassRender.cpp
+++ b/Engine/Source/Editor/UnrealEd/Private/Lightmass/LightmassRender.cpp
@@ -20,6 +20,13 @@ struct FLightmassMaterialCompiler : public FProxyMaterialCompiler
 		FProxyMaterialCompiler(InCompiler)
 	{}
 
+	// gets value stored by SetMaterialProperty()
+	virtual EShaderFrequency GetCurrentShaderFrequency() const
+	{
+		// not used by Lightmass
+		return SF_Pixel;
+	}
+
 	virtual int32 ParticleMacroUV()
 	{
 		return Compiler->ParticleMacroUV();
@@ -283,12 +290,23 @@ public:
 	}
 
 	// Material properties.
+
 	/** Entry point for compiling a specific material property.  This must call SetMaterialProperty. */
-	virtual int32 CompileProperty(EMaterialProperty Property,EShaderFrequency InShaderFrequency,FMaterialCompiler* Compiler) const
+	virtual int32 CompilePropertyAndSetMaterialProperty(EMaterialProperty Property, FMaterialCompiler* Compiler, EShaderFrequency OverrideShaderFrequency) const
+	{
+		// needs to be called in this function!!
+		Compiler->SetMaterialProperty(Property, OverrideShaderFrequency);
+
+		int32 Ret = CompilePropertyAndSetMaterialPropertyWithoutCast(Property, Compiler);
+
+		return Compiler->ForceCast(Ret, GetMaterialPropertyType(Property));
+	}
+
+	/** helper for CompilePropertyAndSetMaterialProperty() */
+	int32 CompilePropertyAndSetMaterialPropertyWithoutCast(EMaterialProperty Property, FMaterialCompiler* Compiler) const
 	{
 		static const auto UseDiffuseSpecularMaterialInputs = IConsoleManager::Get().FindTConsoleVariableDataInt(TEXT("r.UseDiffuseSpecularMaterialInputs"));
 
-		Compiler->SetMaterialProperty(Property, InShaderFrequency);
 		// MAKE SURE THIS MATCHES THE CHART IN WillFillData
 		// 						  RETURNED VALUES (F16 'textures')
 		// 	BLEND MODE  | DIFFUSE     | SPECULAR     | EMISSIVE    | NORMAL    | TRANSMISSIVE              |
diff --git a/Engine/Source/Editor/UnrealEd/Private/MaterialGraph.cpp b/Engine/Source/Editor/UnrealEd/Private/MaterialGraph.cpp
--- a/Engine/Source/Editor/UnrealEd/Private/MaterialGraph.cpp
+++ b/Engine/Source/Editor/UnrealEd/Private/MaterialGraph.cpp
@@ -30,32 +30,32 @@ void UMaterialGraph::RebuildGraph()
 	if (!MaterialFunction)
 	{
 		// Initialize the material input list.
-		MaterialInputs.Add( FMaterialInputInfo( LOCTEXT("DiffuseColor", "Diffuse Color"), &Material->DiffuseColor, MP_DiffuseColor ) );
-		MaterialInputs.Add( FMaterialInputInfo( LOCTEXT("SpecularColor", "Specular Color"), &Material->SpecularColor, MP_SpecularColor ) );
-		MaterialInputs.Add( FMaterialInputInfo( LOCTEXT("BaseColor", "Base Color"), &Material->BaseColor, MP_BaseColor ) );	
-		MaterialInputs.Add( FMaterialInputInfo( LOCTEXT("Metallic", "Metallic"), &Material->Metallic, MP_Metallic ) );
-		MaterialInputs.Add( FMaterialInputInfo( LOCTEXT("Specular", "Specular"), &Material->Specular, MP_Specular ) );
-		MaterialInputs.Add( FMaterialInputInfo( LOCTEXT("Roughness", "Roughness"), &Material->Roughness, MP_Roughness ) );
-		MaterialInputs.Add( FMaterialInputInfo( LOCTEXT("EmissiveColor", "Emissive Color"), &Material->EmissiveColor, MP_EmissiveColor ) );
-		MaterialInputs.Add( FMaterialInputInfo( LOCTEXT("Opacity", "Opacity"), &Material->Opacity, MP_Opacity ) );
-		MaterialInputs.Add( FMaterialInputInfo( LOCTEXT("OpacityMask", "Opacity Mask"), &Material->OpacityMask, MP_OpacityMask ) );
-		MaterialInputs.Add( FMaterialInputInfo( LOCTEXT("Normal", "Normal"), &Material->Normal, MP_Normal ) );
-		MaterialInputs.Add( FMaterialInputInfo( LOCTEXT("WorldPositionOffset", "World Position Offset"), &Material->WorldPositionOffset, MP_WorldPositionOffset ) );
-		MaterialInputs.Add( FMaterialInputInfo( LOCTEXT("WorldDisplacement", "World Displacement"), &Material->WorldDisplacement, MP_WorldDisplacement ) );
-		MaterialInputs.Add( FMaterialInputInfo( LOCTEXT("TessellationMultiplier", "Tessellation Multiplier"), &Material->TessellationMultiplier, MP_TessellationMultiplier ) );
-		MaterialInputs.Add( FMaterialInputInfo( LOCTEXT("SubsurfaceColor", "Subsurface Color"), &Material->SubsurfaceColor, MP_SubsurfaceColor ) );
-		MaterialInputs.Add( FMaterialInputInfo( LOCTEXT("ClearCoat", "Clear Coat"), &Material->ClearCoat, MP_ClearCoat ) );
-		MaterialInputs.Add( FMaterialInputInfo( LOCTEXT("ClearCoatRoughness", "Clear Coat Roughness"), &Material->ClearCoatRoughness, MP_ClearCoatRoughness ) );
-		MaterialInputs.Add( FMaterialInputInfo( LOCTEXT("AmbientOcclusion", "Ambient Occlusion"), &Material->AmbientOcclusion, MP_AmbientOcclusion ) );
-		MaterialInputs.Add( FMaterialInputInfo( LOCTEXT("Refraction", "Refraction"), &Material->Refraction, MP_Refraction) );
+		MaterialInputs.Add( FMaterialInputInfo( LOCTEXT("DiffuseColor", "Diffuse Color"), MP_DiffuseColor ) );
+		MaterialInputs.Add( FMaterialInputInfo( LOCTEXT("SpecularColor", "Specular Color"), MP_SpecularColor ) );
+		MaterialInputs.Add( FMaterialInputInfo( LOCTEXT("BaseColor", "Base Color"), MP_BaseColor ) );	
+		MaterialInputs.Add( FMaterialInputInfo( LOCTEXT("Metallic", "Metallic"), MP_Metallic ) );
+		MaterialInputs.Add( FMaterialInputInfo( LOCTEXT("Specular", "Specular"), MP_Specular ) );
+		MaterialInputs.Add( FMaterialInputInfo( LOCTEXT("Roughness", "Roughness"), MP_Roughness ) );
+		MaterialInputs.Add( FMaterialInputInfo( LOCTEXT("EmissiveColor", "Emissive Color"), MP_EmissiveColor ) );
+		MaterialInputs.Add( FMaterialInputInfo( LOCTEXT("Opacity", "Opacity"), MP_Opacity ) );
+		MaterialInputs.Add( FMaterialInputInfo( LOCTEXT("OpacityMask", "Opacity Mask"), MP_OpacityMask ) );
+		MaterialInputs.Add( FMaterialInputInfo( LOCTEXT("Normal", "Normal"), MP_Normal ) );
+		MaterialInputs.Add( FMaterialInputInfo( LOCTEXT("WorldPositionOffset", "World Position Offset"), MP_WorldPositionOffset ) );
+		MaterialInputs.Add( FMaterialInputInfo( LOCTEXT("WorldDisplacement", "World Displacement"), MP_WorldDisplacement ) );
+		MaterialInputs.Add( FMaterialInputInfo( LOCTEXT("TessellationMultiplier", "Tessellation Multiplier"), MP_TessellationMultiplier ) );
+		MaterialInputs.Add( FMaterialInputInfo( LOCTEXT("SubsurfaceColor", "Subsurface Color"), MP_SubsurfaceColor ) );
+		MaterialInputs.Add( FMaterialInputInfo( LOCTEXT("ClearCoat", "Clear Coat"), MP_ClearCoat ) );
+		MaterialInputs.Add( FMaterialInputInfo( LOCTEXT("ClearCoatRoughness", "Clear Coat Roughness"), MP_ClearCoatRoughness ) );
+		MaterialInputs.Add( FMaterialInputInfo( LOCTEXT("AmbientOcclusion", "Ambient Occlusion"), MP_AmbientOcclusion ) );
+		MaterialInputs.Add( FMaterialInputInfo( LOCTEXT("Refraction", "Refraction"), MP_Refraction) );
 
 		for (int32 UVIndex = 0; UVIndex < ARRAY_COUNT(Material->CustomizedUVs); UVIndex++)
 		{
 			//@todo - localize
-			MaterialInputs.Add( FMaterialInputInfo( FText::FromString(FString::Printf(TEXT("Customized UV%u"), UVIndex)), &Material->CustomizedUVs[UVIndex], (EMaterialProperty)(MP_CustomizedUVs0 + UVIndex)) );
+			MaterialInputs.Add( FMaterialInputInfo( FText::FromString(FString::Printf(TEXT("Customized UV%u"), UVIndex)), (EMaterialProperty)(MP_CustomizedUVs0 + UVIndex)) );
 		}
 
-		MaterialInputs.Add( FMaterialInputInfo( LOCTEXT("MaterialAttributes", "Material Attributes"), &Material->MaterialAttributes, MP_MaterialAttributes) );
+		MaterialInputs.Add( FMaterialInputInfo( LOCTEXT("MaterialAttributes", "Material Attributes"), MP_MaterialAttributes) );
 
 		// Add Root Node
 		FGraphNodeCreator<UMaterialGraphNode_Root> NodeCreator(*this);
@@ -125,11 +125,12 @@ void UMaterialGraph::LinkGraphNodesFromMaterial()
 		for (int32 Index = 0; Index < MaterialInputs.Num(); ++Index)
 		{
 			UEdGraphPin* InputPin = RootNode->GetInputPin(Index);
+			auto ExpressionInput = MaterialInputs[Index].GetInput(Material);
 
-			if (MaterialInputs[Index].Input->Expression)
+			if (ExpressionInput->Expression)
 			{
-				UMaterialGraphNode* GraphNode = CastChecked<UMaterialGraphNode>(MaterialInputs[Index].Input->Expression->GraphNode);
-				InputPin->MakeLinkTo(GraphNode->GetOutputPin(GetValidOutputIndex(MaterialInputs[Index].Input)));
+				UMaterialGraphNode* GraphNode = CastChecked<UMaterialGraphNode>(ExpressionInput->Expression->GraphNode);
+				InputPin->MakeLinkTo(GraphNode->GetOutputPin(GetValidOutputIndex(ExpressionInput)));
 			}
 		}
 	}
@@ -175,7 +176,7 @@ void UMaterialGraph::LinkMaterialExpressionsFromGraph() const
 			check(InputPins.Num() == MaterialInputs.Num());
 			for (int32 PinIndex = 0; PinIndex < InputPins.Num() && PinIndex < MaterialInputs.Num(); ++PinIndex)
 			{
-				FExpressionInput* MaterialInput = MaterialInputs[PinIndex].Input;
+				FExpressionInput* MaterialInput = MaterialInputs[PinIndex].GetInput(Material);
 				if (InputPins[PinIndex]->LinkedTo.Num() > 0)
 				{
 					UMaterialGraphNode* ConnectedNode = CastChecked<UMaterialGraphNode>(InputPins[PinIndex]->LinkedTo[0]->GetOwningNode());
diff --git a/Engine/Source/Runtime/Engine/Private/Materials/HLSLMaterialTranslator.h b/Engine/Source/Runtime/Engine/Private/Materials/HLSLMaterialTranslator.h
--- a/Engine/Source/Runtime/Engine/Private/Materials/HLSLMaterialTranslator.h
+++ b/Engine/Source/Runtime/Engine/Private/Materials/HLSLMaterialTranslator.h
@@ -312,34 +312,37 @@ public:
 			// Generate code
 			int32 Chunk[CompiledMP_MAX];
 
-			Chunk[MP_Normal]						= ForceCast(Material->CompileProperty(MP_Normal                ,GetMaterialPropertyShaderFrequency(MP_Normal),this),MCT_Float3);
-			Chunk[MP_EmissiveColor]					= ForceCast(Material->CompileProperty(MP_EmissiveColor         ,GetMaterialPropertyShaderFrequency(MP_EmissiveColor),this),MCT_Float3);
-			Chunk[MP_DiffuseColor]					= ForceCast(Material->CompileProperty(MP_DiffuseColor          ,GetMaterialPropertyShaderFrequency(MP_DiffuseColor),this),MCT_Float3);
-			Chunk[MP_SpecularColor]					= ForceCast(Material->CompileProperty(MP_SpecularColor         ,GetMaterialPropertyShaderFrequency(MP_SpecularColor),this),MCT_Float3);
-			Chunk[MP_BaseColor]						= ForceCast(Material->CompileProperty(MP_BaseColor             ,GetMaterialPropertyShaderFrequency(MP_BaseColor),this),MCT_Float3);
-			Chunk[MP_Metallic]						= ForceCast(Material->CompileProperty(MP_Metallic              ,GetMaterialPropertyShaderFrequency(MP_Metallic),this),MCT_Float1);
-			Chunk[MP_Specular]						= ForceCast(Material->CompileProperty(MP_Specular              ,GetMaterialPropertyShaderFrequency(MP_Specular),this),MCT_Float1);
-			Chunk[MP_Roughness]						= ForceCast(Material->CompileProperty(MP_Roughness             ,GetMaterialPropertyShaderFrequency(MP_Roughness),this),MCT_Float1);
-			Chunk[MP_Opacity]						= ForceCast(Material->CompileProperty(MP_Opacity               ,GetMaterialPropertyShaderFrequency(MP_Opacity),this),MCT_Float1);
-			Chunk[MP_OpacityMask]					= ForceCast(Material->CompileProperty(MP_OpacityMask           ,GetMaterialPropertyShaderFrequency(MP_OpacityMask),this),MCT_Float1);
-			Chunk[MP_WorldPositionOffset]			= ForceCast(Material->CompileProperty(MP_WorldPositionOffset   ,GetMaterialPropertyShaderFrequency(MP_WorldPositionOffset),this),MCT_Float3);
+			memset(Chunk, -1, sizeof(Chunk));
+
+			Chunk[MP_Normal]						= Material->CompilePropertyAndSetMaterialProperty(MP_Normal                ,this);
+			Chunk[MP_EmissiveColor]					= Material->CompilePropertyAndSetMaterialProperty(MP_EmissiveColor         ,this);
+			Chunk[MP_DiffuseColor]					= Material->CompilePropertyAndSetMaterialProperty(MP_DiffuseColor          ,this);
+			Chunk[MP_SpecularColor]					= Material->CompilePropertyAndSetMaterialProperty(MP_SpecularColor         ,this);
+			Chunk[MP_BaseColor]						= Material->CompilePropertyAndSetMaterialProperty(MP_BaseColor             ,this);
+			Chunk[MP_Metallic]						= Material->CompilePropertyAndSetMaterialProperty(MP_Metallic              ,this);
+			Chunk[MP_Specular]						= Material->CompilePropertyAndSetMaterialProperty(MP_Specular              ,this);
+			Chunk[MP_Roughness]						= Material->CompilePropertyAndSetMaterialProperty(MP_Roughness             ,this);
+			Chunk[MP_Opacity]						= Material->CompilePropertyAndSetMaterialProperty(MP_Opacity               ,this);
+			Chunk[MP_OpacityMask]					= Material->CompilePropertyAndSetMaterialProperty(MP_OpacityMask           ,this);
+			Chunk[MP_WorldPositionOffset]			= Material->CompilePropertyAndSetMaterialProperty(MP_WorldPositionOffset   ,this);
 			if (FeatureLevel >= ERHIFeatureLevel::SM5)
 			{
-				Chunk[MP_WorldDisplacement]			= ForceCast(Material->CompileProperty(MP_WorldDisplacement     ,GetMaterialPropertyShaderFrequency(MP_WorldDisplacement),this),MCT_Float3);
+				Chunk[MP_WorldDisplacement]			= Material->CompilePropertyAndSetMaterialProperty(MP_WorldDisplacement     ,this);
 			}
 			else
 			{
-				SetMaterialProperty(MP_WorldDisplacement, GetMaterialPropertyShaderFrequency(MP_WorldDisplacement));
-				Chunk[MP_WorldDisplacement]			= ForceCast(Constant3(0.0f,0.0f,0.0f),MCT_Float3);
+				// normally called in CompilePropertyAndSetMaterialProperty, needs to be called!!
+				SetMaterialProperty(MP_WorldDisplacement);
+				Chunk[MP_WorldDisplacement]			= Constant3(0.0f, 0.0f, 0.0f);
 			}
-			Chunk[MP_TessellationMultiplier]		= ForceCast(Material->CompileProperty(MP_TessellationMultiplier,GetMaterialPropertyShaderFrequency(MP_TessellationMultiplier),this),MCT_Float1);
-			Chunk[MP_SubsurfaceColor]				= ForceCast(Material->CompileProperty(MP_SubsurfaceColor       ,GetMaterialPropertyShaderFrequency(MP_SubsurfaceColor),this),MCT_Float3);
-			Chunk[MP_ClearCoat]						= ForceCast(Material->CompileProperty(MP_ClearCoat		       ,GetMaterialPropertyShaderFrequency(MP_ClearCoat),this),MCT_Float1);
-			Chunk[MP_ClearCoatRoughness]			= ForceCast(Material->CompileProperty(MP_ClearCoatRoughness    ,GetMaterialPropertyShaderFrequency(MP_ClearCoatRoughness),this),MCT_Float1);
-			Chunk[MP_AmbientOcclusion]				= ForceCast(Material->CompileProperty(MP_AmbientOcclusion      ,GetMaterialPropertyShaderFrequency(MP_AmbientOcclusion),this),MCT_Float1);
+			Chunk[MP_TessellationMultiplier]		= Material->CompilePropertyAndSetMaterialProperty(MP_TessellationMultiplier,this);
+			Chunk[MP_SubsurfaceColor]				= Material->CompilePropertyAndSetMaterialProperty(MP_SubsurfaceColor       ,this);
+			Chunk[MP_ClearCoat]						= Material->CompilePropertyAndSetMaterialProperty(MP_ClearCoat			   ,this);
+			Chunk[MP_ClearCoatRoughness]			= Material->CompilePropertyAndSetMaterialProperty(MP_ClearCoatRoughness    ,this);
+			Chunk[MP_AmbientOcclusion]				= Material->CompilePropertyAndSetMaterialProperty(MP_AmbientOcclusion      ,this);
 
 			{
-				int32 UserRefraction = ForceCast(Material->CompileProperty(MP_Refraction, GetMaterialPropertyShaderFrequency(MP_Refraction), this), MCT_Float2);
+				int32 UserRefraction = ForceCast(Material->CompilePropertyAndSetMaterialProperty(MP_Refraction, this), MCT_Float2);
 				int32 RefractionDepthBias = ForceCast(ScalarParameter(FName(TEXT("RefractionDepthBias")), GetRefractionDepthBiasValue()), MCT_Float1);
 
 				Chunk[MP_Refraction] = AppendVector(ForceCast(UserRefraction, MCT_Float1), RefractionDepthBias);
@@ -347,7 +350,7 @@ public:
 
 			if (bCompileForComputeShader)
 			{
-				Chunk[CompiledMP_EmissiveColorCS]		= ForceCast(Material->CompileProperty(MP_EmissiveColor     ,SF_Compute,this),MCT_Float3);
+				Chunk[CompiledMP_EmissiveColorCS]		= Material->CompilePropertyAndSetMaterialProperty(MP_EmissiveColor,this, SF_Compute);
 			}
 
 			for (uint32 CustomUVIndex = MP_CustomizedUVs0; CustomUVIndex <= MP_CustomizedUVs7; CustomUVIndex++)
@@ -357,7 +360,7 @@ public:
 				// Note: this is using NumUserTexCoords, which is set by translating all the pixel properties above
 				if (CustomUVIndex - MP_CustomizedUVs0 < NumUserTexCoords)
 				{
-					Chunk[CustomUVIndex] = ForceCast(Material->CompileProperty((EMaterialProperty)CustomUVIndex, GetMaterialPropertyShaderFrequency((EMaterialProperty)CustomUVIndex), this), MCT_Float2);
+					Chunk[CustomUVIndex] = Material->CompilePropertyAndSetMaterialProperty((EMaterialProperty)CustomUVIndex, this);
 				}
 				else
 				{
@@ -606,15 +609,26 @@ public:
 		LazyPrintf.PushParam(*FString::Printf(TEXT("%u"),NumUserVertexTexCoords));
 		LazyPrintf.PushParam(*FString::Printf(TEXT("%u"),NumUserTexCoords));
 		LazyPrintf.PushParam(*ResourcesString);
-		LazyPrintf.PushParam(*(TranslatedCodeChunkDefinitions[MP_Normal] + TEXT("	return ") + TranslatedCodeChunks[MP_Normal] + TEXT(";")));
-		LazyPrintf.PushParam(*(TranslatedCodeChunkDefinitions[MP_EmissiveColor] + TEXT("	return ") + TranslatedCodeChunks[MP_EmissiveColor] + TEXT(";")));
-		LazyPrintf.PushParam(bCompileForComputeShader ? *(TranslatedCodeChunkDefinitions[CompiledMP_EmissiveColorCS] + TEXT("	return ") + TranslatedCodeChunks[CompiledMP_EmissiveColorCS] + TEXT(";")) : TEXT("return 0"));
-		LazyPrintf.PushParam(*(TranslatedCodeChunkDefinitions[MP_DiffuseColor] + TEXT("	return ") + TranslatedCodeChunks[MP_DiffuseColor] + TEXT(";")));
-		LazyPrintf.PushParam(*(TranslatedCodeChunkDefinitions[MP_SpecularColor] + TEXT("	return ") + TranslatedCodeChunks[MP_SpecularColor] + TEXT(";")));
-		LazyPrintf.PushParam(*(TranslatedCodeChunkDefinitions[MP_BaseColor] + TEXT("	return ") + TranslatedCodeChunks[MP_BaseColor] + TEXT(";")));
-		LazyPrintf.PushParam(*(TranslatedCodeChunkDefinitions[MP_Metallic] + TEXT("	return ") + TranslatedCodeChunks[MP_Metallic] + TEXT(";")));
-		LazyPrintf.PushParam(*(TranslatedCodeChunkDefinitions[MP_Specular] + TEXT("	return ") + TranslatedCodeChunks[MP_Specular] + TEXT(";")));
-		LazyPrintf.PushParam(*(TranslatedCodeChunkDefinitions[MP_Roughness] + TEXT("	return ") + TranslatedCodeChunks[MP_Roughness] + TEXT(";")));
+
+		LazyPrintf.PushParam(*GenerateFunctionCode(MP_Normal));
+		LazyPrintf.PushParam(*GenerateFunctionCode(MP_EmissiveColor));
+
+		if(bCompileForComputeShader)
+		{
+			LazyPrintf.PushParam(*GenerateFunctionCode(CompiledMP_EmissiveColorCS));
+		}
+		else
+		{
+			LazyPrintf.PushParam(TEXT("return 0"));
+		}
+
+		LazyPrintf.PushParam(*GenerateFunctionCode(MP_DiffuseColor));
+		LazyPrintf.PushParam(*GenerateFunctionCode(MP_SpecularColor));
+		LazyPrintf.PushParam(*GenerateFunctionCode(MP_BaseColor));
+		LazyPrintf.PushParam(*GenerateFunctionCode(MP_Metallic));
+		LazyPrintf.PushParam(*GenerateFunctionCode(MP_Specular));
+		LazyPrintf.PushParam(*GenerateFunctionCode(MP_Roughness));
+
 		LazyPrintf.PushParam(*FString::Printf(TEXT("return %.5f"),Material->GetTranslucencyDirectionalLightingIntensity()));
 		
 		LazyPrintf.PushParam(*FString::Printf(TEXT("return %.5f"),Material->GetTranslucentShadowDensityScale()));
@@ -622,26 +636,25 @@ public:
 		LazyPrintf.PushParam(*FString::Printf(TEXT("return %.5f"),Material->GetTranslucentSelfShadowSecondDensityScale()));
 		LazyPrintf.PushParam(*FString::Printf(TEXT("return %.5f"),Material->GetTranslucentSelfShadowSecondOpacity()));
 		LazyPrintf.PushParam(*FString::Printf(TEXT("return %.5f"),Material->GetTranslucentBackscatteringExponent()));
-		LazyPrintf.PushParam(*FString::Printf(TEXT("return MaterialFloat3(%.5f, %.5f, %.5f)"),
-			Material->GetTranslucentMultipleScatteringExtinction().R,
-			Material->GetTranslucentMultipleScatteringExtinction().G,
-			Material->GetTranslucentMultipleScatteringExtinction().B));
+
+		{
+			FLinearColor Extinction = Material->GetTranslucentMultipleScatteringExtinction();
+
+			LazyPrintf.PushParam(*FString::Printf(TEXT("return MaterialFloat3(%.5f, %.5f, %.5f)"), Extinction.R, Extinction.G, Extinction.B));
+		}
+
 		LazyPrintf.PushParam(*FString::Printf(TEXT("return %.5f"),Material->GetOpacityMaskClipValue()));
-		LazyPrintf.PushParam(*(TranslatedCodeChunkDefinitions[MP_Opacity] + TEXT("	return ") + TranslatedCodeChunks[MP_Opacity] + TEXT(";")));
-		LazyPrintf.PushParam(*(TranslatedCodeChunkDefinitions[MP_OpacityMask] + TEXT("	return ") + TranslatedCodeChunks[MP_OpacityMask] + TEXT(";")));
-		/*
-		LazyPrintf.PushParam(*FString::Printf(TEXT("%s return %s"),
-			*GetDefinitions(MP_OpacityMask, GetMaterialPropertyShaderFrequency(MP_OpacityMask)),
-			*CodeChunk[MP_OpacityMask]));
-		*/
-		LazyPrintf.PushParam(*(TranslatedCodeChunkDefinitions[MP_WorldPositionOffset] + TEXT("	return ") + TranslatedCodeChunks[MP_WorldPositionOffset] + TEXT(";")));
-		LazyPrintf.PushParam(*(TranslatedCodeChunkDefinitions[MP_WorldDisplacement] + TEXT("	return ") + TranslatedCodeChunks[MP_WorldDisplacement] + TEXT(";")));
-		LazyPrintf.PushParam(*(TranslatedCodeChunkDefinitions[MP_TessellationMultiplier] + TEXT("	return ") + TranslatedCodeChunks[MP_TessellationMultiplier] + TEXT(";")));
-		LazyPrintf.PushParam(*(TranslatedCodeChunkDefinitions[MP_SubsurfaceColor] + TEXT("	return ") + TranslatedCodeChunks[MP_SubsurfaceColor] + TEXT(";")));
-		LazyPrintf.PushParam(*(TranslatedCodeChunkDefinitions[MP_ClearCoat] + TEXT("	return ") + TranslatedCodeChunks[MP_ClearCoat] + TEXT(";")));
-		LazyPrintf.PushParam(*(TranslatedCodeChunkDefinitions[MP_ClearCoatRoughness] + TEXT("	return ") + TranslatedCodeChunks[MP_ClearCoatRoughness] + TEXT(";")));
-		LazyPrintf.PushParam(*(TranslatedCodeChunkDefinitions[MP_AmbientOcclusion] + TEXT("	return ") + TranslatedCodeChunks[MP_AmbientOcclusion] + TEXT(";")));
-		LazyPrintf.PushParam(*(TranslatedCodeChunkDefinitions[MP_Refraction] + TEXT("	return ") + TranslatedCodeChunks[MP_Refraction] + TEXT(";")));
+
+		LazyPrintf.PushParam(*GenerateFunctionCode(MP_Opacity));
+		LazyPrintf.PushParam(*GenerateFunctionCode(MP_OpacityMask));
+		LazyPrintf.PushParam(*GenerateFunctionCode(MP_WorldPositionOffset));
+		LazyPrintf.PushParam(*GenerateFunctionCode(MP_WorldDisplacement));
+		LazyPrintf.PushParam(*GenerateFunctionCode(MP_TessellationMultiplier));
+		LazyPrintf.PushParam(*GenerateFunctionCode(MP_SubsurfaceColor));
+		LazyPrintf.PushParam(*GenerateFunctionCode(MP_ClearCoat));
+		LazyPrintf.PushParam(*GenerateFunctionCode(MP_ClearCoatRoughness));
+		LazyPrintf.PushParam(*GenerateFunctionCode(MP_AmbientOcclusion));
+		LazyPrintf.PushParam(*GenerateFunctionCode(MP_Refraction));
 
 		FString CustomUVAssignments;
 
@@ -659,6 +672,14 @@ public:
 
 protected:
 
+	// only used by GetMaterialShaderCode()
+	// @param Index ECompiledMaterialProperty or EMaterialProperty
+	FString GenerateFunctionCode(uint32 Index) const
+	{
+		check(Index < CompiledMP_MAX);
+		return TranslatedCodeChunkDefinitions[Index] + TEXT("	return ") + TranslatedCodeChunks[Index] + TEXT(";");
+	}
+
 	// GetParameterCode
 	virtual FString GetParameterCode(int32 Index, const TCHAR* Default = 0)
 	{
@@ -718,7 +739,7 @@ protected:
 	{
 		if(Index != INDEX_NONE)
 		{
-			checkf(Index >= 0 && Index < CodeChunks[InProperty][InFrequency].Num(), TEXT("Index out of range %d/%d [%s]"), Index, CodeChunks[InProperty][InFrequency].Num(), *this->Material->GetFriendlyName());
+			checkf(Index >= 0 && Index < CodeChunks[InProperty][InFrequency].Num(), TEXT("Index out of range %d/%d [%s]"), Index, CodeChunks[InProperty][InFrequency].Num(), *Material->GetFriendlyName());
 			check(!CodeChunks[InProperty][InFrequency][Index].UniformExpression || CodeChunks[InProperty][InFrequency][Index].UniformExpression->IsConstant());
 			if (CodeChunks[InProperty][InFrequency][Index].UniformExpression && CodeChunks[InProperty][InFrequency][Index].UniformExpression->IsConstant())
 			{
@@ -1149,11 +1170,24 @@ protected:
 	/** 
 	 * Sets the current material property being compiled.  
 	 * This affects the internal state of the compiler and the results of all functions except GetFixedParameterCode.
+	 * @param OverrideShaderFrequency SF_NumFrequencies to not override
 	 */
-	virtual void SetMaterialProperty(EMaterialProperty InProperty, EShaderFrequency InShaderFrequency)
+	virtual void SetMaterialProperty(EMaterialProperty InProperty, EShaderFrequency OverrideShaderFrequency = SF_NumFrequencies)
 	{
 		MaterialProperty = InProperty;
-		ShaderFrequency = InShaderFrequency;
+
+		if(OverrideShaderFrequency != SF_NumFrequencies)
+		{
+			ShaderFrequency = OverrideShaderFrequency;
+		}
+		else
+		{
+			ShaderFrequency = GetMaterialPropertyShaderFrequency(InProperty);
+		}
+	}
+	virtual EShaderFrequency GetCurrentShaderFrequency() const
+	{
+		return ShaderFrequency;
 	}
 
 	virtual int32 Error(const TCHAR* Text)
diff --git a/Engine/Source/Runtime/Engine/Private/Materials/Material.cpp b/Engine/Source/Runtime/Engine/Private/Materials/Material.cpp
--- a/Engine/Source/Runtime/Engine/Private/Materials/Material.cpp
+++ b/Engine/Source/Runtime/Engine/Private/Materials/Material.cpp
@@ -64,37 +64,54 @@ FMaterialResource::FMaterialResource()
 {
 }
 
-int32 FMaterialResource::CompileProperty(EMaterialProperty Property,EShaderFrequency InShaderFrequency,FMaterialCompiler* Compiler) const
+int32 FMaterialResource::CompilePropertyAndSetMaterialProperty(EMaterialProperty Property, FMaterialCompiler* Compiler, EShaderFrequency OverrideShaderFrequency) const
 {
-	Compiler->SetMaterialProperty(Property, InShaderFrequency);
+	// needs to be called in this function!!
+	// sets CurrentShaderFrequency
+	Compiler->SetMaterialProperty(Property, OverrideShaderFrequency);
+
+	EShaderFrequency ShaderFrequency = Compiler->GetCurrentShaderFrequency();
+
 	int32 SelectionColorIndex = INDEX_NONE;
-	if (InShaderFrequency == SF_Pixel)
+
+	if (ShaderFrequency == SF_Pixel)
 	{
 		SelectionColorIndex = Compiler->ComponentMask(Compiler->VectorParameter(NAME_SelectionColor,FLinearColor::Black),1,1,1,0);
 	}
 	
 	//Compile the material instance if we have one.
 	UMaterialInterface* MaterialInterface = MaterialInstance ? Cast<UMaterialInterface>(MaterialInstance) : Cast<UMaterialInterface>(Material);
+
+	int32 Ret = INDEX_NONE;
+
 	switch(Property)
 	{
-	case MP_EmissiveColor:
-		if (SelectionColorIndex != INDEX_NONE)
-		{
-			return Compiler->Add(Compiler->ForceCast(MaterialInterface->CompileProperty(Compiler, MP_EmissiveColor),MCT_Float3),SelectionColorIndex);
-		}
-		else
-		{
-			return Compiler->ForceCast(MaterialInterface->CompileProperty(Compiler, MP_EmissiveColor),MCT_Float3);
-		}
-	case MP_DiffuseColor: 
-		return Compiler->Mul(Compiler->ForceCast(MaterialInterface->CompileProperty(Compiler, MP_DiffuseColor),MCT_Float3),Compiler->Sub(Compiler->Constant(1.0f),SelectionColorIndex));
-	case MP_BaseColor: 
-		return Compiler->Mul(Compiler->ForceCast(MaterialInterface->CompileProperty(Compiler, MP_BaseColor),MCT_Float3),Compiler->Sub(Compiler->Constant(1.0f),SelectionColorIndex));
-	case MP_MaterialAttributes:
-		return INDEX_NONE;
-	default:
-		return MaterialInterface->CompileProperty(Compiler, Property);
+		case MP_EmissiveColor:
+			if (SelectionColorIndex != INDEX_NONE)
+			{
+				Ret = Compiler->Add(Compiler->ForceCast(MaterialInterface->CompileProperty(Compiler, MP_EmissiveColor), MCT_Float3), SelectionColorIndex);
+			}
+			else
+			{
+				Ret = Compiler->ForceCast(MaterialInterface->CompileProperty(Compiler, MP_EmissiveColor), MCT_Float3);
+			}
+			break;
+
+		case MP_DiffuseColor: 
+			Ret = Compiler->Mul(Compiler->ForceCast(MaterialInterface->CompileProperty(Compiler, MP_DiffuseColor), MCT_Float3), Compiler->Sub(Compiler->Constant(1.0f), SelectionColorIndex));
+			break;
+		case MP_BaseColor: 
+			Ret = Compiler->Mul(Compiler->ForceCast(MaterialInterface->CompileProperty(Compiler, MP_BaseColor),MCT_Float3),Compiler->Sub(Compiler->Constant(1.0f),SelectionColorIndex));
+			break;
+		case MP_MaterialAttributes:
+			Ret = INDEX_NONE;
+			break;
+		default:
+			Ret = MaterialInterface->CompileProperty(Compiler, Property);
 	};
+
+	// output should always be the right type for this property
+	return Compiler->ForceCast(Ret, GetMaterialPropertyType(Property));
 }
 
 void FMaterialResource::GetShaderMapId(EShaderPlatform Platform, FMaterialShaderMapId& OutId) const
@@ -2967,44 +2984,25 @@ FExpressionInput* UMaterial::GetExpressionInputForProperty(EMaterialProperty InP
 {
 	switch (InProperty)
 	{
-	case MP_EmissiveColor:
-		return &EmissiveColor;
-	case MP_Opacity:
-		return &Opacity;
-	case MP_OpacityMask:
-		return &OpacityMask;
-	case MP_DiffuseColor:
-		return &DiffuseColor;
-	case MP_SpecularColor:
-		return &SpecularColor;
-	case MP_BaseColor:
-		return &BaseColor;
-	case MP_Metallic:
-		return &Metallic;
-	case MP_Specular:
-		return &Specular;
-	case MP_Roughness:
-		return &Roughness;
-	case MP_Normal:
-		return &Normal;
-	case MP_WorldPositionOffset:
-		return &WorldPositionOffset;
-	case MP_WorldDisplacement:
-		return &WorldDisplacement;
-	case MP_TessellationMultiplier:
-		return &TessellationMultiplier;
-	case MP_SubsurfaceColor:
-		return &SubsurfaceColor;
-	case MP_ClearCoat:
-		return &ClearCoat;
-	case MP_ClearCoatRoughness:
-		return &ClearCoatRoughness;
-	case MP_AmbientOcclusion:
-		return &AmbientOcclusion;
-	case MP_Refraction:
-		return &Refraction;
-	case MP_MaterialAttributes: 
-		return &MaterialAttributes;
+		case MP_EmissiveColor:			return &EmissiveColor;
+		case MP_Opacity:				return &Opacity;
+		case MP_OpacityMask:			return &OpacityMask;
+		case MP_DiffuseColor:			return &DiffuseColor;
+		case MP_SpecularColor:			return &SpecularColor;
+		case MP_BaseColor:				return &BaseColor;
+		case MP_Metallic:				return &Metallic;
+		case MP_Specular:				return &Specular;
+		case MP_Roughness:				return &Roughness;
+		case MP_Normal:					return &Normal;
+		case MP_WorldPositionOffset:	return &WorldPositionOffset;
+		case MP_WorldDisplacement:		return &WorldDisplacement;
+		case MP_TessellationMultiplier:	return &TessellationMultiplier;
+		case MP_SubsurfaceColor:		return &SubsurfaceColor;
+		case MP_ClearCoat:				return &ClearCoat;
+		case MP_ClearCoatRoughness:		return &ClearCoatRoughness;
+		case MP_AmbientOcclusion:		return &AmbientOcclusion;
+		case MP_Refraction:				return &Refraction;
+		case MP_MaterialAttributes:		return &MaterialAttributes;
 	}
 
 	if (InProperty >= MP_CustomizedUVs0 && InProperty <= MP_CustomizedUVs7)
diff --git a/Engine/Source/Runtime/Engine/Private/Materials/MaterialShared.cpp b/Engine/Source/Runtime/Engine/Private/Materials/MaterialShared.cpp
--- a/Engine/Source/Runtime/Engine/Private/Materials/MaterialShared.cpp
+++ b/Engine/Source/Runtime/Engine/Private/Materials/MaterialShared.cpp
@@ -2057,8 +2057,14 @@ FMaterialUpdateContext::~FMaterialUpdateContext()
 		);
 }
 
+uint32 Global1 = 0;
+uint32 Global2 = 0;
+
 EMaterialProperty GetMaterialPropertyFromInputOutputIndex(int32 Index)
 {
+	++Global1;
+	FPlatformMisc::LowLevelOutputDebugStringf(TEXT(">>>> %d %d\n"), Global1, Global2);
+
 	switch(Index)
 	{
 	case 0: return MP_BaseColor;
@@ -2097,6 +2103,9 @@ EMaterialProperty GetMaterialPropertyFromInputOutputIndex(int32 Index)
 
 int32 GetInputOutputIndexFromMaterialProperty(EMaterialProperty Property)
 {
+	++Global2;
+	FPlatformMisc::LowLevelOutputDebugStringf(TEXT(">>>> %d %d\n"), Global1, Global2);
+
 	switch(Property)
 	{
 	case MP_BaseColor: return 0;
@@ -2214,10 +2223,8 @@ int32 UMaterialInterface::CompilePropertyEx( class FMaterialCompiler* Compiler,
 	return INDEX_NONE;
 }
 
-int32 UMaterialInterface::CompileProperty( FMaterialCompiler* Compiler, EMaterialProperty Property )
+int32 UMaterialInterface::CompileProperty(FMaterialCompiler* Compiler, EMaterialProperty Property)
 {
-	int32 Ret = INDEX_NONE;
-
 	if (IsPropertyActive(Property))
 	{
 		return CompilePropertyEx(Compiler, Property);
@@ -2226,7 +2233,6 @@ int32 UMaterialInterface::CompileProperty( FMaterialCompiler* Compiler, EMateria
 	{
 		return GetDefaultExpressionForMaterialProperty(Compiler, Property);
 	}
-	return Ret;
 }
 
 /** TODO - This can be removed whenever VER_UE4_MATERIAL_ATTRIBUTES_REORDERING is no longer relevant. */
diff --git a/Engine/Source/Runtime/Engine/Public/MaterialCompiler.h b/Engine/Source/Runtime/Engine/Public/MaterialCompiler.h
--- a/Engine/Source/Runtime/Engine/Public/MaterialCompiler.h
+++ b/Engine/Source/Runtime/Engine/Public/MaterialCompiler.h
@@ -18,8 +18,12 @@
 class FMaterialCompiler
 {
 public:
-
-	virtual void SetMaterialProperty(EMaterialProperty InProperty, EShaderFrequency InShaderFrequency) = 0;
+	// sets internal state CurrentShaderFrequency 
+	// @param OverrideShaderFrequency SF_NumFrequencies to not override
+	virtual void SetMaterialProperty(EMaterialProperty InProperty, EShaderFrequency OverrideShaderFrequency = SF_NumFrequencies) = 0;
+	// gets value stored by SetMaterialProperty()
+	virtual EShaderFrequency GetCurrentShaderFrequency() const = 0;
+	//
 	virtual int32 Error(const TCHAR* Text) = 0;
 	ENGINE_API int32 Errorf(const TCHAR* Format,...);
 
@@ -188,7 +192,8 @@ public:
 
 	// Simple pass through all other material operations unmodified.
 
-	virtual void SetMaterialProperty(EMaterialProperty InProperty, EShaderFrequency InShaderFrequency) override { Compiler->SetMaterialProperty(InProperty, InShaderFrequency); }
+	virtual void SetMaterialProperty(EMaterialProperty InProperty, EShaderFrequency OverrideShaderFrequency) override { Compiler->SetMaterialProperty(InProperty, OverrideShaderFrequency); }
+	virtual EShaderFrequency GetCurrentShaderFrequency() const	{ return Compiler->GetCurrentShaderFrequency(); }
 	virtual int32 Error(const TCHAR* Text) override { return Compiler->Error(Text); }
 
 	virtual int32 CallExpression(FMaterialExpressionKey ExpressionKey,FMaterialCompiler* InCompiler) override { return Compiler->CallExpression(ExpressionKey,InCompiler); }
diff --git a/Engine/Source/Runtime/Engine/Public/MaterialShared.h b/Engine/Source/Runtime/Engine/Public/MaterialShared.h
--- a/Engine/Source/Runtime/Engine/Public/MaterialShared.h
+++ b/Engine/Source/Runtime/Engine/Public/MaterialShared.h
@@ -1065,8 +1065,12 @@ protected:
 	*/
 	void GetShaderMapIDsWithUnfinishedCompilation(TArray<int32>& ShaderMapIds);
 
-	/** Entry point for compiling a specific material property.  This must call SetMaterialProperty. */
-	virtual int32 CompileProperty(EMaterialProperty Property,EShaderFrequency InShaderFrequency,class FMaterialCompiler* Compiler) const = 0;
+	/**
+	 * Entry point for compiling a specific material property.  This must call SetMaterialProperty. 
+	 * @param OverrideShaderFrequency SF_NumFrequencies to not override
+	 * @return cases to the proper type e.g. Compiler->ForceCast(Ret, GetMaterialPropertyType(Property));
+	 */
+	virtual int32 CompilePropertyAndSetMaterialProperty(EMaterialProperty Property, class FMaterialCompiler* Compiler, EShaderFrequency OverrideShaderFrequency = SF_NumFrequencies) const = 0;
 
 	/** Returns the index to the Expression in the Expressions array, or -1 if not found. */
 	int32 FindExpression(const TArray<TRefCountPtr<FMaterialUniformExpressionTexture> >&Expressions, const FMaterialUniformExpressionTexture &Expression);
@@ -1508,7 +1512,7 @@ protected:
 	UMaterialInstance* MaterialInstance;
 
 	/** Entry point for compiling a specific material property.  This must call SetMaterialProperty. */
-	ENGINE_API virtual int32 CompileProperty(EMaterialProperty Property,EShaderFrequency InShaderFrequency,class FMaterialCompiler* Compiler) const;
+	ENGINE_API virtual int32 CompilePropertyAndSetMaterialProperty(EMaterialProperty Property, class FMaterialCompiler* Compiler, EShaderFrequency OverrideShaderFrequency) const;
 	ENGINE_API virtual bool HasVertexPositionOffsetConnected() const;
 	ENGINE_API virtual bool HasMaterialAttributesConnected() const;
 	/** Useful for debugging. */
