Make sure the ChildActorName is cached before RerunConstructionScript occurs

diff --git a/Engine/Source/Runtime/Engine/Classes/Components/ChildActorComponent.h b/Engine/Source/Runtime/Engine/Classes/Components/ChildActorComponent.h
--- a/Engine/Source/Runtime/Engine/Classes/Components/ChildActorComponent.h
+++ b/Engine/Source/Runtime/Engine/Classes/Components/ChildActorComponent.h
@@ -29,6 +29,7 @@ class UChildActorComponent : public USceneComponent
 	// Begin ActorComponent interface.
 	virtual void OnComponentCreated() override;
 	virtual void OnComponentDestroyed() override;
+	virtual void OnRegister() override;
 	virtual TSharedPtr<FComponentInstanceDataBase> GetComponentInstanceData() const override;
 	virtual FName GetComponentInstanceDataType() const override;
 	virtual void ApplyComponentInstanceData(TSharedPtr<FComponentInstanceDataBase> ComponentInstanceData) override;
diff --git a/Engine/Source/Runtime/Engine/Private/ChildActorComponent.cpp b/Engine/Source/Runtime/Engine/Private/ChildActorComponent.cpp
--- a/Engine/Source/Runtime/Engine/Private/ChildActorComponent.cpp
+++ b/Engine/Source/Runtime/Engine/Private/ChildActorComponent.cpp
@@ -10,6 +10,16 @@ UChildActorComponent::UChildActorComponent(const class FPostConstructInitializeP
 {
 }
 
+void UChildActorComponent::OnRegister()
+{
+	Super::OnRegister();
+
+	if (ChildActor)
+	{
+		ChildActorName = ChildActor->GetFName();
+	}
+}
+
 void UChildActorComponent::OnComponentCreated()
 {
 	Super::OnComponentCreated();
