deleting non leaf bodies in phat will automatically fixup constraints from parent and children

diff --git a/Engine/Source/Editor/PhAT/Private/PhATSharedData.cpp b/Engine/Source/Editor/PhAT/Private/PhATSharedData.cpp
--- a/Engine/Source/Editor/PhAT/Private/PhATSharedData.cpp
+++ b/Engine/Source/Editor/PhAT/Private/PhATSharedData.cpp
@@ -1061,16 +1061,45 @@ void FPhATSharedData::DeleteBody(int32 DelBodyIndex, bool bRefreshComponent)
 	PhysicsAsset->Modify();
 
 	// .. the body..
-	PhysicsAsset->BodySetup[DelBodyIndex]->Modify();	
+	UBodySetup * BodySetup = PhysicsAsset->BodySetup[DelBodyIndex];
+	BodySetup->Modify();	
 
 	// .. and any constraints to the body.
 	TArray<int32>	Constraints;
 	PhysicsAsset->BodyFindConstraints(DelBodyIndex, Constraints);
 
+	//we want to fixup constraints so that nearest child bodies get constraint with parent body
+	TArray<int32> NearestBodiesBelow;
+	PhysicsAsset->GetNearestBodyIndicesBelow(NearestBodiesBelow, BodySetup->BoneName, EditorSkelMesh);
+	
+	int32 BoneIndex = EditorSkelMesh->RefSkeleton.FindBoneIndex(BodySetup->BoneName);
+	check(BoneIndex != INDEX_NONE);
+	int32 ParentBodyIndex = PhysicsAsset->FindParentBodyIndex(EditorSkelMesh, BoneIndex);
+
+	UBodySetup * ParentBody = ParentBodyIndex != INDEX_NONE ? PhysicsAsset->BodySetup[ParentBodyIndex] : NULL;
+
 	for (int32 i = 0; i <Constraints.Num(); ++i)
 	{
 		int32 ConstraintIndex = Constraints[i];
-		PhysicsAsset->ConstraintSetup[ConstraintIndex]->Modify();
+		UPhysicsConstraintTemplate * Constraint = PhysicsAsset->ConstraintSetup[ConstraintIndex];
+		Constraint->Modify();
+
+		if (ParentBody)
+		{
+			//for all constraints that contain a nearest child of this body, create a copy of the constraint between the child and parent
+			for (int32 i = 0; i < NearestBodiesBelow.Num(); ++i)
+			{
+				int32 BodyBelowIndex = NearestBodiesBelow[i];
+				UBodySetup * BodyBelow = PhysicsAsset->BodySetup[BodyBelowIndex];
+
+				if (Constraint->DefaultInstance.ConstraintBone1 == BodyBelow->BoneName)
+				{
+					int32 NewConstraintIndex = FPhysicsAssetUtils::CreateNewConstraint(PhysicsAsset, BodyBelow->BoneName, Constraint);
+					UPhysicsConstraintTemplate * NewConstraint = PhysicsAsset->ConstraintSetup[NewConstraintIndex];
+					InitConstraintSetup(NewConstraint, BodyBelowIndex, ParentBodyIndex);
+				}
+			}
+		}
 	}
 
 	// Now actually destroy body. This will destroy any constraints associated with the body as well.
diff --git a/Engine/Source/Runtime/Engine/Classes/PhysicsEngine/PhysicsAsset.h b/Engine/Source/Runtime/Engine/Classes/PhysicsEngine/PhysicsAsset.h
--- a/Engine/Source/Runtime/Engine/Classes/PhysicsEngine/PhysicsAsset.h
+++ b/Engine/Source/Runtime/Engine/Classes/PhysicsEngine/PhysicsAsset.h
@@ -127,11 +127,14 @@ public:
 
 	// Find the index of the physics bone that is controlling this graphics bone.
 	ENGINE_API int32		FindControllingBodyIndex(class USkeletalMesh* skelMesh, int32 BoneIndex);
+	ENGINE_API int32		FindParentBodyIndex(class USkeletalMesh * skelMesh, int32 StartBoneIndex) const;
 	ENGINE_API int32		FindConstraintIndex(FName ConstraintName);
 	FName					FindConstraintBoneName(int32 ConstraintIndex);
 
 	/** Utility for getting indices of all bodies below (and including) the one with the supplied name. */
-	ENGINE_API void			GetBodyIndicesBelow(TArray<int32>& OutBodyIndices, FName InBoneName, USkeletalMesh* InSkelMesh);
+	ENGINE_API void			GetBodyIndicesBelow(TArray<int32>& OutBodyIndices, FName InBoneName, USkeletalMesh* InSkelMesh, bool bIncludeParent = true);
+
+	ENGINE_API void			GetNearestBodyIndicesBelow(TArray<int32> & OutBodyIndices, FName InBoneName, USkeletalMesh * InSkelMesh);
 
 	ENGINE_API FBox			CalcAABB(const class USkinnedMeshComponent* MeshComponent) const;
 
@@ -164,7 +167,7 @@ public:
 
 
 	// @todo document
-	ENGINE_API int32 FindBodyIndex(FName BodyName);
+	ENGINE_API int32 FindBodyIndex(FName BodyName) const;
 
 	/** Find all the constraints that are connected to a particular body.
 	 * 
diff --git a/Engine/Source/Runtime/Engine/Private/PhysicsEngine/PhysicsAsset.cpp b/Engine/Source/Runtime/Engine/Private/PhysicsEngine/PhysicsAsset.cpp
--- a/Engine/Source/Runtime/Engine/Private/PhysicsEngine/PhysicsAsset.cpp
+++ b/Engine/Source/Runtime/Engine/Private/PhysicsEngine/PhysicsAsset.cpp
@@ -208,10 +208,25 @@ int32	UPhysicsAsset::FindControllingBodyIndex(class USkeletalMesh* skelMesh, int
 	return INDEX_NONE; // Shouldn't reach here.
 }
 
+int32 UPhysicsAsset::FindParentBodyIndex(class USkeletalMesh * skelMesh, int32 StartBoneIndex) const
+{
+	int32 BoneIndex = StartBoneIndex;
+	while ((BoneIndex = skelMesh->RefSkeleton.GetParentIndex(BoneIndex)) != 0)
+	{
+		FName BoneName = skelMesh->RefSkeleton.GetBoneName(BoneIndex);
+		int32 BodyIndex = FindBodyIndex(BoneName);
+
+		if (BodyIndex != INDEX_NONE)
+			return BodyIndex;
+	}
+
+	return INDEX_NONE;
+}
 
-int32 UPhysicsAsset::FindBodyIndex(FName bodyName)
+
+int32 UPhysicsAsset::FindBodyIndex(FName bodyName) const
 {
-	int32 * IdxData = BodySetupIndexMap.Find(bodyName);
+	const int32 * IdxData = BodySetupIndexMap.Find(bodyName);
 	if (IdxData)
 	{
 		return *IdxData;
@@ -243,7 +258,7 @@ FName UPhysicsAsset::FindConstraintBoneName(int32 ConstraintIndex)
 	return ConstraintSetup[ConstraintIndex]->DefaultInstance.JointName;
 }
 
-void UPhysicsAsset::GetBodyIndicesBelow(TArray<int32>& OutBodyIndices, FName InBoneName, USkeletalMesh* SkelMesh)
+void UPhysicsAsset::GetBodyIndicesBelow(TArray<int32>& OutBodyIndices, FName InBoneName, USkeletalMesh* SkelMesh, bool bIncludeParent /*= true*/)
 {
 	int32 BaseIndex = SkelMesh->RefSkeleton.FindBoneIndex(InBoneName);
 
@@ -254,14 +269,51 @@ void UPhysicsAsset::GetBodyIndicesBelow(TArray<int32>& OutBodyIndices, FName InB
 		FName TestName = BS->BoneName;
 		int32 TestIndex = SkelMesh->RefSkeleton.FindBoneIndex(TestName);
 
-		// We want to return this body as well.
-		if(TestIndex == BaseIndex || SkelMesh->RefSkeleton.BoneIsChildOf(TestIndex, BaseIndex))
+		if( (bIncludeParent && TestIndex == BaseIndex) || SkelMesh->RefSkeleton.BoneIsChildOf(TestIndex, BaseIndex))
 		{
 			OutBodyIndices.Add(i);
 		}
 	}
 }
 
+void UPhysicsAsset::GetNearestBodyIndicesBelow(TArray<int32> & OutBodyIndices, FName InBoneName, USkeletalMesh * InSkelMesh)
+{
+	TArray<int32> AllBodiesBelow;
+	GetBodyIndicesBelow(AllBodiesBelow, InBoneName, InSkelMesh, false);
+
+	//we need to filter all bodies below to first in the chain
+	TArray<bool> Nearest;
+	Nearest.AddUninitialized(BodySetup.Num());
+	for (int32 i = 0; i < Nearest.Num(); ++i)
+	{
+		Nearest[i] = true;
+	}
+
+	for (int32 i = 0; i < AllBodiesBelow.Num(); i++)
+	{
+		int32 BodyIndex = AllBodiesBelow[i];
+		if (Nearest[BodyIndex] == false) continue;
+
+		UBodySetup * Body = BodySetup[BodyIndex];
+		TArray<int32> BodiesBelowMe;
+		GetBodyIndicesBelow(BodiesBelowMe, Body->BoneName, InSkelMesh, false);
+		
+		for (int j = 0; j < BodiesBelowMe.Num(); ++j)
+		{
+			Nearest[BodiesBelowMe[j]] = false;	
+		}
+	}
+
+	for (int32 i = 0; i < AllBodiesBelow.Num(); i++)
+	{
+		int32 BodyIndex = AllBodiesBelow[i];
+		if (Nearest[BodyIndex])
+		{
+			OutBodyIndices.Add(BodyIndex);
+		}
+	}
+}
+
 void UPhysicsAsset::ClearAllPhysicsMeshes()
 {
 	for(int32 i=0; i<BodySetup.Num(); i++)
