Warn when we have queued bunches on the client for a very long time

diff --git a/Engine/Source/Runtime/Engine/Classes/Engine/ActorChannel.h b/Engine/Source/Runtime/Engine/Classes/Engine/ActorChannel.h
--- a/Engine/Source/Runtime/Engine/Classes/Engine/ActorChannel.h
+++ b/Engine/Source/Runtime/Engine/Classes/Engine/ActorChannel.h
@@ -68,9 +68,10 @@ class ENGINE_API UActorChannel : public UChannel
 
 	TMap< TWeakObjectPtr< UObject >, TSharedRef< FObjectReplicator > > ReplicationMap;
 
-	TArray< class FInBunch * > QueuedBunches;		// Queued bunches waiting on pending guids to resolve
-
-	TSet< FNetworkGUID > PendingGuidResolves;		// These guids are waiting for their resolves, we need to queue up bunches until these are resolved
+	// Async networking loading support state
+	TArray< class FInBunch * >	QueuedBunches;			// Queued bunches waiting on pending guids to resolve
+	double						QueuedBunchStartTime;	// Time when since queued bunches was last empty
+	TSet< FNetworkGUID >		PendingGuidResolves;	// These guids are waiting for their resolves, we need to queue up bunches until these are resolved
 
 	/**
 	 * Default constructor
diff --git a/Engine/Source/Runtime/Engine/Private/DataChannel.cpp b/Engine/Source/Runtime/Engine/Private/DataChannel.cpp
--- a/Engine/Source/Runtime/Engine/Private/DataChannel.cpp
+++ b/Engine/Source/Runtime/Engine/Private/DataChannel.cpp
@@ -1235,10 +1235,10 @@ void UActorChannel::Init( UNetConnection* InConnection, int32 InChannelIndex, bo
 {
 	Super::Init( InConnection, InChannelIndex, InOpenedLocally );
 
-	RelevantTime	= Connection->Driver->Time;
-	LastUpdateTime	= Connection->Driver->Time - Connection->Driver->SpawnPrioritySeconds;
-	bActorMustStayDirty = false;
-	bActorStillInitial = false;
+	RelevantTime			= Connection->Driver->Time;
+	LastUpdateTime			= Connection->Driver->Time - Connection->Driver->SpawnPrioritySeconds;
+	bActorMustStayDirty		= false;
+	bActorStillInitial		= false;
 }
 
 void UActorChannel::SetClosingFlag()
@@ -1496,6 +1496,18 @@ void UActorChannel::Tick()
 
 		QueuedBunches.Empty();
 	}
+
+	// Warn when we have queued bunches for a very long time
+	if ( QueuedBunches.Num() > 0 )
+	{
+		const double QUEUED_BUNCH_TIMEOUT_IN_SECONDS = 30;
+
+		if ( FPlatformTime::Seconds() - QueuedBunchStartTime > QUEUED_BUNCH_TIMEOUT_IN_SECONDS )
+		{
+			UE_LOG( LogNet, Warning, TEXT( "UActorChannel::Tick: Queued bunches for longer than normal. ChIndex: %i, Actor: %s, Queued: %i" ), ChIndex, Actor != NULL ? *Actor->GetPathName() : TEXT( "NULL" ), QueuedBunches.Num() );
+			QueuedBunchStartTime = FPlatformTime::Seconds();
+		}
+	}
 }
 
 void UActorChannel::ReceivedBunch( FInBunch & Bunch )
@@ -1536,6 +1548,12 @@ void UActorChannel::ReceivedBunch( FInBunch & Bunch )
 	// If we have guids that need to be resolved, or we have existing pending bunches, we must process this bunch later
 	if ( PendingGuidResolves.Num() > 0 || QueuedBunches.Num() > 0 )
 	{
+		if ( QueuedBunches.Num() == 0 )
+		{
+			// Remember when we first started queuing
+			QueuedBunchStartTime = FPlatformTime::Seconds();
+		}
+
 		QueuedBunches.Add( new FInBunch( Bunch ) );
 		return;
 	}
