Removed a benign warning case where a package wasn't fully loaded, but it would warn when it couldn't find assets in this package

diff --git a/Engine/Source/Runtime/Engine/Private/PackageMapClient.cpp b/Engine/Source/Runtime/Engine/Private/PackageMapClient.cpp
--- a/Engine/Source/Runtime/Engine/Private/PackageMapClient.cpp
+++ b/Engine/Source/Runtime/Engine/Private/PackageMapClient.cpp
@@ -1960,24 +1960,9 @@ UObject * FNetGUIDCache::GetObjectFromNetGUID( const FNetworkGUID & NetGUID, con
 		// Try to resolve the outer
 		ObjOuter = GetObjectFromNetGUID( CacheObjectPtr->OuterGUID, bIgnoreMustBeMapped );
 
-		// If we can't resolve the outer (or the outer package is not fully loaded, and we are async loading), we need to bail
-		if ( ObjOuter == NULL || ( CVarAllowAsyncLoading.GetValueOnGameThread() > 0 && !ObjOuter->GetOutermost()->IsFullyLoaded() ) )
+		// If we can't resolve the outer
+		if ( ObjOuter == NULL )
 		{
-			if ( !CacheObjectPtr->bNoLoad && CVarAllowAsyncLoading.GetValueOnGameThread() > 0 )
-			{
-				if ( !bIgnoreMustBeMapped )
-				{
-					// If we get here, we are trying to process a "must be mapped" guid when we shouldn't be (the asset hasn't loaded yet)
-					UE_LOG( LogNetPackageMap, Error, TEXT( "GetObjectFromNetGUID: Missing \"must be mapped\" asset. Path: %s, OuterPath: %s, NetGUID: %s" ), *CacheObjectPtr->PathName.ToString(), *OuterCacheObject->PathName.ToString(), *NetGUID.ToString() );
-				}
-
-				if ( ObjOuter != NULL && !ObjOuter->GetOutermost()->IsFullyLoaded() )
-				{
-					// Not done loading outer
-					return NULL;
-				}
-			}
-
 			// If the outer is missing, warn unless told to ignore
 			if ( !ShouldIgnoreWhenMissing( CacheObjectPtr->OuterGUID ) )
 			{
@@ -2007,6 +1992,7 @@ UObject * FNetGUIDCache::GetObjectFromNetGUID( const FNetworkGUID & NetGUID, con
 		}
 	}
 
+	// At this point, we either have an outer, or we are a package
 	check( !CacheObjectPtr->bIsPending );
 	check( ObjOuter == NULL || ObjOuter->GetOutermost()->IsFullyLoaded() );
 
@@ -2069,12 +2055,32 @@ UObject * FNetGUIDCache::GetObjectFromNetGUID( const FNetworkGUID & NetGUID, con
 
 		return NULL;
 	}
-	
-	if ( bIsPackage && CastChecked< UPackage >( Object )->GetGuid() != CacheObjectPtr->PackageGuid )
+
+	if ( bIsPackage )
 	{
-		CacheObjectPtr->bIsBroken = true;
-		UE_LOG( LogNetPackageMap, Error, TEXT( "GetObjectFromNetGUID: Package GUID mismatch! Path: %s, NetGUID: %s" ), *CacheObjectPtr->PathName.ToString(), *NetGUID.ToString() );
-		return NULL;
+		UPackage * Package = Cast< UPackage >( Object );
+
+		if ( Package == NULL )
+		{
+			// This isn't really a package but it should be
+			CacheObjectPtr->bIsBroken = true;
+			UE_LOG( LogNetPackageMap, Error, TEXT( "GetObjectFromNetGUID: Object is not a package but shoule be! Path: %s, NetGUID: %s" ), *CacheObjectPtr->PathName.ToString(), *NetGUID.ToString() );
+			return NULL;
+		}
+
+		if ( !Package->IsFullyLoaded() )
+		{
+			// If this package isn't fully loaded, don't complain, assume it will fully load at some point
+			return NULL;
+		}
+
+		if ( Package->GetGuid() != CacheObjectPtr->PackageGuid )
+		{
+			// If the package guid doesn't match, don't allow it to load
+			CacheObjectPtr->bIsBroken = true;
+			UE_LOG( LogNetPackageMap, Error, TEXT( "GetObjectFromNetGUID: Package GUID mismatch! Path: %s, NetGUID: %s" ), *CacheObjectPtr->PathName.ToString(), *NetGUID.ToString() );
+			return NULL;
+		}
 	}
 
 	// Assign the resolved object to this guid
@@ -2128,8 +2134,23 @@ bool FNetGUIDCache::ShouldIgnoreWhenMissing( const FNetworkGUID & NetGUID ) cons
 		OutermostCacheObject = ObjectLookup.Find( OutermostCacheObject->OuterGUID );
 	}
 
-	// Ignore warnings when we explicitly are told to, or if the outer most package is pending
-	return CacheObject->bIgnoreWhenMissing || OutermostCacheObject->bIsPending;
+	if ( OutermostCacheObject != NULL )
+	{
+		// If our outer package is not fully loaded, then don't warn, assume it will eventually come in
+		if ( OutermostCacheObject->bIsPending )
+		{
+			// Outer is pending, don't warn
+			return true;
+		}
+		// Sometimes, other systems async load packages, which we don't track, but still must be aware of
+		if ( OutermostCacheObject->Object != NULL && !OutermostCacheObject->Object->GetOutermost()->IsFullyLoaded() )
+		{
+			return true;
+		}
+	}
+
+	// Ignore warnings when we explicitly are told to
+	return CacheObject->bIgnoreWhenMissing;
 }
 
 bool FNetGUIDCache::IsGUIDRegistered( const FNetworkGUID & NetGUID ) const
