Pull Request #256 - Exposed PxJoint::setLocalPose and PxConstraint::getForce to Unreal API. Now FConstraintInstace::SetRefFrame will update an active constraint's frame at run time. Also provided FConstraintInstace::SetRefPosition and FConstraintInstace::SetRefOrientation to set reference position and orientation individually. Finally added FConstraintInstace::GetConstraintForce to call PxConstraint::getForce and added similar Blueprint methods to ConstraintComponent to make it easier to drive this constraint behavior.

diff --git a/Engine/Source/Runtime/Engine/Classes/PhysicsEngine/ConstraintInstance.h b/Engine/Source/Runtime/Engine/Classes/PhysicsEngine/ConstraintInstance.h
--- a/Engine/Source/Runtime/Engine/Classes/PhysicsEngine/ConstraintInstance.h
+++ b/Engine/Source/Runtime/Engine/Classes/PhysicsEngine/ConstraintInstance.h
@@ -25,6 +25,7 @@ enum ELinearConstraintMotion
 };
 
 /** Enum to indicate which frame we want */
+UENUM()
 namespace EConstraintFrame
 {
 	enum Type
@@ -409,9 +410,16 @@ public:
 	/** See if this constraint is valid. */
 	bool IsValidConstraintInstance() const;
 
-	// Pass in reference frame in
+	// Pass in reference frame in. If the constraint is currently active, this will set its active local pose. Otherwise the change will take affect in InitConstraint. 
 	void SetRefFrame(EConstraintFrame::Type Frame, const FTransform& RefFrame);
 
+	// Pass in reference position in (maintains reference orientation). If the constraint is currently active, this will set its active local pose. Otherwise the change will take affect in InitConstraint.
+	void SetRefPosition(EConstraintFrame::Type Frame, const FVector& RefPosition);
+	
+	// Pass in reference orientation in (maintains reference position). If the constraint is currently active, this will set its active local pose. Otherwise the change will take affect in InitConstraint.
+	void SetRefOrientation(EConstraintFrame::Type Frame, const FVector& PriAxis, const FVector& SecAxis);
+	
+
 	// Get component ref frame
 	FTransform GetRefFrame(EConstraintFrame::Type Frame) const;
 
@@ -424,6 +432,9 @@ public:
 	/** Get the position of this constraint in world space. */
 	FVector GetConstraintLocation();
 
+	// Retrieve the constraint force most recently applied to maintain this constraint. Returns 0 forces if the constraint is not initialized or broken.
+	void GetConstraintForce(FVector& OutLinearForce, FVector& OutAngularForce);
+
 	void SetLinearPositionDrive(bool bEnableXDrive, bool bEnableYDrive, bool bEnableZDrive);
 	void SetLinearVelocityDrive(bool bEnableXDrive, bool bEnableYDrive, bool bEnableZDrive);
 	void SetAngularPositionDrive(bool bEnableSwingDrive, bool bEnableTwistDrive);
diff --git a/Engine/Source/Runtime/Engine/Classes/PhysicsEngine/PhysicsConstraintComponent.h b/Engine/Source/Runtime/Engine/Classes/PhysicsEngine/PhysicsConstraintComponent.h
--- a/Engine/Source/Runtime/Engine/Classes/PhysicsEngine/PhysicsConstraintComponent.h
+++ b/Engine/Source/Runtime/Engine/Classes/PhysicsEngine/PhysicsConstraintComponent.h
@@ -225,6 +225,22 @@ public:
 	 */
 	ENGINE_API void UpdateConstraintFrames();
 
+	// Pass in reference frame in. If the constraint is currently active, this will set its active local pose. Otherwise the change will take affect in InitConstraint. 
+	UFUNCTION(BlueprintCallable, Category = "Physics|Components|PhysicsConstraint")
+	ENGINE_API void SetConstraintReferenceFrame(EConstraintFrame::Type Frame, const FTransform& RefFrame);
+	
+	// Pass in reference position in (maintains reference orientation). If the constraint is currently active, this will set its active local pose. Otherwise the change will take affect in InitConstraint.
+	UFUNCTION(BlueprintCallable, Category = "Physics|Components|PhysicsConstraint")
+	ENGINE_API void SetConstraintReferencePosition(EConstraintFrame::Type Frame, const FVector& RefPosition);
+	
+	// Pass in reference orientation in (maintains reference position). If the constraint is currently active, this will set its active local pose. Otherwise the change will take affect in InitConstraint.
+	UFUNCTION(BlueprintCallable, Category = "Physics|Components|PhysicsConstraint")
+	ENGINE_API void SetConstraintReferenceOrientation(EConstraintFrame::Type Frame, const FVector& PriAxis, const FVector& SecAxis);
+	
+	// Retrieve the constraint force most recently applied to maintain this constraint. Returns 0 forces if the constraint is not initialized or broken.
+	UFUNCTION(BlueprintCallable, Category = "Physics|Components|PhysicsConstraint")
+	ENGINE_API void GetConstraintForce(FVector& OutLinearForce, FVector& OutAngularForce);
+
 #if WITH_EDITORONLY_DATA
 	UPROPERTY(transient)
 	class UBillboardComponent* SpriteComponent;
diff --git a/Engine/Source/Runtime/Engine/Private/PhysicsEngine/ConstraintInstance.cpp b/Engine/Source/Runtime/Engine/Private/PhysicsEngine/ConstraintInstance.cpp
--- a/Engine/Source/Runtime/Engine/Private/PhysicsEngine/ConstraintInstance.cpp
+++ b/Engine/Source/Runtime/Engine/Private/PhysicsEngine/ConstraintInstance.cpp
@@ -520,8 +520,20 @@ FTransform FConstraintInstance::GetRefFrame(EConstraintFrame::Type Frame) const
 	return Result;
 }
 
+#if WITH_PHYSX
+FORCEINLINE physx::PxJointActorIndex::Enum U2PConstraintFrame(EConstraintFrame::Type Frame)
+{
+	// Swap frame order, since Unreal reverses physx order
+	return (Frame == EConstraintFrame::Frame1) ? physx::PxJointActorIndex::eACTOR1 : physx::PxJointActorIndex::eACTOR0;
+}
+#endif
+
+
 void FConstraintInstance::SetRefFrame(EConstraintFrame::Type Frame, const FTransform& RefFrame)
 {
+#if WITH_PHYSX
+	PxJointActorIndex::Enum PxFrame = U2PConstraintFrame(Frame);
+#endif
 	if(Frame == EConstraintFrame::Frame1)
 	{
 		Pos1 = RefFrame.GetTranslation();
@@ -534,6 +546,70 @@ void FConstraintInstance::SetRefFrame(EConstraintFrame::Type Frame, const FTrans
 		PriAxis2 = RefFrame.GetScaledAxis( EAxis::X );
 		SecAxis2 = RefFrame.GetScaledAxis( EAxis::Y );
 	}
+
+#if WITH_PHYSX
+	if (ConstraintData && !(ConstraintData->getConstraintFlags()&PxConstraintFlag::eBROKEN))
+	{
+	PxTransform PxRefFrame = U2PTransform(RefFrame);
+	ConstraintData->setLocalPose(PxFrame, PxRefFrame);
+	}
+#endif
+
+}
+
+void FConstraintInstance::SetRefPosition(EConstraintFrame::Type Frame, const FVector& RefPosition)
+{
+#if WITH_PHYSX
+	PxJointActorIndex::Enum PxFrame = U2PConstraintFrame(Frame);
+#endif
+
+	if (Frame == EConstraintFrame::Frame1)
+	{
+		Pos1 = RefPosition;
+	}
+	else
+	{
+		Pos2 = RefPosition;
+	}
+
+#if WITH_PHYSX
+	if (ConstraintData && !(ConstraintData->getConstraintFlags()&PxConstraintFlag::eBROKEN))
+	{
+		PxTransform PxRefFrame = ConstraintData->getLocalPose(PxFrame);
+		PxRefFrame.p = U2PVector(RefPosition);
+		ConstraintData->setLocalPose(PxFrame, PxRefFrame);
+	}
+#endif
+}
+
+void FConstraintInstance::SetRefOrientation(EConstraintFrame::Type Frame, const FVector& PriAxis, const FVector& SecAxis)
+{
+#if WITH_PHYSX
+	PxJointActorIndex::Enum PxFrame = U2PConstraintFrame(Frame);
+	FVector RefPos;
+#endif
+		
+	if (Frame == EConstraintFrame::Frame1)
+	{
+		RefPos = Pos1;
+		PriAxis1 = PriAxis;
+		SecAxis1 = SecAxis;
+	}
+	else
+	{
+		RefPos = Pos2;
+		PriAxis2 = PriAxis;
+		SecAxis2 = SecAxis;
+	}
+	
+#if WITH_PHYSX
+	if (ConstraintData && !(ConstraintData->getConstraintFlags()&PxConstraintFlag::eBROKEN))
+	{
+		FTransform URefTransform = FTransform(PriAxis1, SecAxis, PriAxis ^ SecAxis, RefPos);
+		PxTransform PxRefFrame = U2PTransform(URefTransform);
+		ConstraintData->setLocalPose(PxFrame, PxRefFrame);
+		}
+#endif
 }
 
 /** Get the position of this constraint in world space. */
@@ -572,6 +648,32 @@ FVector FConstraintInstance::GetConstraintLocation()
 #endif
 }
 
+void FConstraintInstance::GetConstraintForce(FVector& OutLinearForce, FVector& OutAngularForce)
+{
+#if WITH_PHYSX
+	if (ConstraintData && !(ConstraintData->getConstraintFlags()&PxConstraintFlag::eBROKEN))
+	{
+		PxVec3 PxOutLinearForce;
+		PxVec3 PxOutAngularForce;
+		ConstraintData->getConstraint()->getForce(PxOutLinearForce, PxOutAngularForce);
+
+		OutLinearForce = P2UVector(PxOutLinearForce);
+		OutAngularForce = P2UVector(PxOutAngularForce);
+	}
+	else
+	{
+		OutLinearForce = FVector::ZeroVector;
+		OutAngularForce = FVector::ZeroVector;
+	}
+
+#else
+	OutLinearForce = FVector::ZeroVector;
+	OutAngularForce = FVector::ZeroVector;
+#endif
+}
+
+
+
 /** Function for turning linear position drive on and off. */
 void FConstraintInstance::SetLinearPositionDrive(bool bEnableXDrive, bool bEnableYDrive, bool bEnableZDrive)
 {
diff --git a/Engine/Source/Runtime/Engine/Private/PhysicsEngine/PhysicsConstraintComponent.cpp b/Engine/Source/Runtime/Engine/Private/PhysicsEngine/PhysicsConstraintComponent.cpp
--- a/Engine/Source/Runtime/Engine/Private/PhysicsEngine/PhysicsConstraintComponent.cpp
+++ b/Engine/Source/Runtime/Engine/Private/PhysicsEngine/PhysicsConstraintComponent.cpp
@@ -405,6 +405,27 @@ void UPhysicsConstraintComponent::UpdateConstraintFrames()
 	ConstraintInstance.SecAxis2 = A2Transform.InverseTransformVectorNoScale(ConstraintInstance.AngularRotationOffset.RotateVector(WOrth));
 }
 
+void UPhysicsConstraintComponent::SetConstraintReferenceFrame(EConstraintFrame::Type Frame, const FTransform& RefFrame)
+{
+	ConstraintInstance.SetRefFrame(Frame, RefFrame);
+}
+
+void UPhysicsConstraintComponent::SetConstraintReferencePosition(EConstraintFrame::Type Frame, const FVector& RefPosition)
+{
+	ConstraintInstance.SetRefPosition(Frame, RefPosition);
+}
+
+void UPhysicsConstraintComponent::SetConstraintReferenceOrientation(EConstraintFrame::Type Frame, const FVector& PriAxis, const FVector& SecAxis)
+{
+	ConstraintInstance.SetRefOrientation(Frame, PriAxis, SecAxis);
+}
+
+void UPhysicsConstraintComponent::GetConstraintForce(FVector& OutLinearForce, FVector& OutAngularForce)
+{
+	ConstraintInstance.GetConstraintForce(OutLinearForce, OutAngularForce);
+}
+
+
 #if WITH_EDITOR
 void UPhysicsConstraintComponent::UpdateSpriteTexture()
 {
