#ue4 - Restoring Pawn.AIControllerClass default to AAIController. Fixes placed pawns in a level not performing CharacterMovement. Broken on 5/27 by AIModule refactoring.

TTP 341111

(integrate from 4.3)

diff --git a/Engine/Source/Runtime/Engine/Private/Pawn.cpp b/Engine/Source/Runtime/Engine/Private/Pawn.cpp
--- a/Engine/Source/Runtime/Engine/Private/Pawn.cpp
+++ b/Engine/Source/Runtime/Engine/Private/Pawn.cpp
@@ -14,6 +14,7 @@
 #include "DisplayDebugHelpers.h"
 #include "NetworkingDistanceConstants.h"
 #include "VisualLog.h"
+#include "AIController.h"
 
 DEFINE_LOG_CATEGORY(LogDamage);
 DEFINE_LOG_CATEGORY_STATIC(LogPawn, Warning, All);
@@ -28,7 +29,7 @@ APawn::APawn(const class FPostConstructInitializeProperties& PCIP)
 {
 	PrimaryActorTick.bCanEverTick = true;
 	PrimaryActorTick.TickGroup = TG_PrePhysics;
-	AIControllerClass = NULL;
+	AIControllerClass = AAIController::StaticClass();
 	bCanBeDamaged = true;
 	
 	SetRemoteRoleForBackwardsCompat(ROLE_SimulatedProxy);
