#ue4 - Always consume the input vector in CharacterMovement tick,  no matter the control path.

diff --git a/Engine/Source/Runtime/Engine/Private/CharacterMovementComponent.cpp b/Engine/Source/Runtime/Engine/Private/CharacterMovementComponent.cpp
--- a/Engine/Source/Runtime/Engine/Private/CharacterMovementComponent.cpp
+++ b/Engine/Source/Runtime/Engine/Private/CharacterMovementComponent.cpp
@@ -734,10 +734,9 @@ void UCharacterMovementComponent::TickComponent(float DeltaTime, enum ELevelTick
 {
 	Super::TickComponent(DeltaTime, TickType, ThisTickFunction);
 
+	const FVector InputVector = ConsumeInputVector();
 	if (!HasValidData() || UpdatedComponent->IsSimulatingPhysics())
 	{
-		// Movement is disabled, don't queue up actions that would apply when it is enabled again.
-		ConsumeInputVector();
 		return;
 	}
 
@@ -773,9 +772,8 @@ void UCharacterMovementComponent::TickComponent(float DeltaTime, enum ELevelTick
 			CharacterOwner->CheckJumpInput(DeltaTime);
 
 			// apply input to acceleration
-			Acceleration = ScaleInputAcceleration(ConstrainInputAcceleration(GetInputVector()));
+			Acceleration = ScaleInputAcceleration(ConstrainInputAcceleration(InputVector));
 			AnalogInputModifier = ComputeAnalogInputModifier();
-			ConsumeInputVector();
 
 			if (CharacterOwner->Role == ROLE_Authority)
 			{
