Expose disable constraint physics to blueprint

diff --git a/Engine/Source/Runtime/Engine/Classes/PhysicsEngine/ConstraintInstance.h b/Engine/Source/Runtime/Engine/Classes/PhysicsEngine/ConstraintInstance.h
--- a/Engine/Source/Runtime/Engine/Classes/PhysicsEngine/ConstraintInstance.h
+++ b/Engine/Source/Runtime/Engine/Classes/PhysicsEngine/ConstraintInstance.h
@@ -448,6 +448,8 @@ public:
 	void SetAngularVelocityTarget(const FVector& InVelTarget);
 	void SetAngularDriveParams(float InSpring, float InDamping, float InForceLimit);
 
+	void SetDisableCollision(bool InDisableCollision);
+
 	void UpdateLinearLimit();
 	void UpdateAngularLimit();
 
diff --git a/Engine/Source/Runtime/Engine/Classes/PhysicsEngine/PhysicsConstraintComponent.h b/Engine/Source/Runtime/Engine/Classes/PhysicsEngine/PhysicsConstraintComponent.h
--- a/Engine/Source/Runtime/Engine/Classes/PhysicsEngine/PhysicsConstraintComponent.h
+++ b/Engine/Source/Runtime/Engine/Classes/PhysicsEngine/PhysicsConstraintComponent.h
@@ -236,6 +236,10 @@ public:
 	// Pass in reference orientation in (maintains reference position). If the constraint is currently active, this will set its active local pose. Otherwise the change will take affect in InitConstraint.
 	UFUNCTION(BlueprintCallable, Category = "Physics|Components|PhysicsConstraint")
 	ENGINE_API void SetConstraintReferenceOrientation(EConstraintFrame::Type Frame, const FVector& PriAxis, const FVector& SecAxis);
+
+	// If true, the collision between the two rigid bodies of the constraint will be disabled.
+	UFUNCTION(BlueprintCallable, Category = "Physics|Components|PhysicsConstraint")
+	ENGINE_API void SetDisableCollision(bool bDisableCollision);
 	
 	// Retrieve the constraint force most recently applied to maintain this constraint. Returns 0 forces if the constraint is not initialized or broken.
 	UFUNCTION(BlueprintCallable, Category = "Physics|Components|PhysicsConstraint")
diff --git a/Engine/Source/Runtime/Engine/Private/PhysicsEngine/ConstraintInstance.cpp b/Engine/Source/Runtime/Engine/Private/PhysicsEngine/ConstraintInstance.cpp
--- a/Engine/Source/Runtime/Engine/Private/PhysicsEngine/ConstraintInstance.cpp
+++ b/Engine/Source/Runtime/Engine/Private/PhysicsEngine/ConstraintInstance.cpp
@@ -197,6 +197,28 @@ FConstraintInstance::FConstraintInstance()
 	ProjectionAngularTolerance = 10.f;// Angular projection when error > 10 degrees
 }
 
+void FConstraintInstance::SetDisableCollision(bool InDisableCollision)
+{
+	bDisableCollision = InDisableCollision;
+#if WITH_PHYSX
+	PxD6Joint* Joint = ConstraintData;
+	if (Joint && !(Joint->getConstraintFlags()&PxConstraintFlag::eBROKEN))
+	{
+		PxConstraintFlags Flags = Joint->getConstraintFlags();
+		if (bDisableCollision)
+		{
+			Flags &= ~PxConstraintFlag::eCOLLISION_ENABLED;
+		}
+		else
+		{
+			Flags |= PxConstraintFlag::eCOLLISION_ENABLED;
+		}
+
+		Joint->setConstraintFlags(Flags);
+	}
+#endif
+}
+
 /** 
  *	Create physics engine constraint.
  */
diff --git a/Engine/Source/Runtime/Engine/Private/PhysicsEngine/PhysicsConstraintComponent.cpp b/Engine/Source/Runtime/Engine/Private/PhysicsEngine/PhysicsConstraintComponent.cpp
--- a/Engine/Source/Runtime/Engine/Private/PhysicsEngine/PhysicsConstraintComponent.cpp
+++ b/Engine/Source/Runtime/Engine/Private/PhysicsEngine/PhysicsConstraintComponent.cpp
@@ -425,6 +425,11 @@ void UPhysicsConstraintComponent::GetConstraintForce(FVector& OutLinearForce, FV
 	ConstraintInstance.GetConstraintForce(OutLinearForce, OutAngularForce);
 }
 
+void UPhysicsConstraintComponent::SetDisableCollision(bool bDisableCollision)
+{
+	ConstraintInstance.SetDisableCollision(bDisableCollision);
+}
+
 
 #if WITH_EDITOR
 void UPhysicsConstraintComponent::UpdateSpriteTexture()
