#ue4 - Enabled movement on skeletal meshes, now that networking support is complete. Networking is still very dependent on bones being in sync between client and server.

- Added bone name to server and client movement RPCs.
- Moved Character movement RPCs from Character to CharacterMovementComponent
- Deprecated MaybeSaveBaseLocation() in favor of just calling SaveBaseLocation(). It is not harmful to resave the location and removes any confusion about this function.
- Added BoneName to MovementBaseUtility functions GetMovementBaseVelocity() and GetMovementBaseTangentialVelocity() and added functionality to get velocity from the bone.

diff --git a/Engine/Source/Runtime/Engine/Classes/GameFramework/Character.h b/Engine/Source/Runtime/Engine/Classes/GameFramework/Character.h
--- a/Engine/Source/Runtime/Engine/Classes/GameFramework/Character.h
+++ b/Engine/Source/Runtime/Engine/Classes/GameFramework/Character.h
@@ -38,6 +38,10 @@ struct FRepRootMotionMontage
 	UPROPERTY()
 	UPrimitiveComponent* MovementBase;
 
+	/** Bone on the MovementBase, if a skeletal mesh. */
+	UPROPERTY()
+	FName MovementBaseBoneName;
+
 	/** Additional replicated flag, if MovementBase can't be resolved on the client. So we don't use wrong data. */
 	UPROPERTY()
 	bool bRelativePosition;
@@ -78,28 +82,28 @@ struct FSimulatedRootMotionReplicatedMove
 namespace MovementBaseUtility
 {
 	/** Determine whether MovementBase can possibly move. */
-	FORCEINLINE bool IsDynamicBase(const class UPrimitiveComponent* MovementBase)
+	FORCEINLINE bool IsDynamicBase(const UPrimitiveComponent* MovementBase)
 	{
 		return (MovementBase && MovementBase->Mobility == EComponentMobility::Movable);
 	}
 
 	/** Determine if we should use relative positioning when based on a component (because it may move). */
-	FORCEINLINE bool UseRelativeLocation(const class UPrimitiveComponent* MovementBase)
+	FORCEINLINE bool UseRelativeLocation(const UPrimitiveComponent* MovementBase)
 	{
 		return IsDynamicBase(MovementBase);
 	}
 
 	/** Ensure that BasedObjectTick ticks after NewBase */
-	ENGINE_API void AddTickDependency(FTickFunction& BasedObjectTick, class UPrimitiveComponent* NewBase);
+	ENGINE_API void AddTickDependency(FTickFunction& BasedObjectTick, UPrimitiveComponent* NewBase);
 
 	/** Remove tick dependency of BasedObjectTick on OldBase */
-	ENGINE_API void RemoveTickDependency(FTickFunction& BasedObjectTick, class UPrimitiveComponent* OldBase);
+	ENGINE_API void RemoveTickDependency(FTickFunction& BasedObjectTick, UPrimitiveComponent* OldBase);
 
 	/** Get the velocity of the given component, first checking the ComponentVelocity and falling back to the physics velocity if necessary. */
-	ENGINE_API FVector GetMovementBaseVelocity(const class UPrimitiveComponent* MovementBase);
+	ENGINE_API FVector GetMovementBaseVelocity(const UPrimitiveComponent* MovementBase, const FName BoneName);
 
 	/** Get the tangential velocity at WorldLocation for the given component. */
-	ENGINE_API FVector GetMovementBaseTangentialVelocity(const class UPrimitiveComponent* MovementBase, const FVector& WorldLocation);
+	ENGINE_API FVector GetMovementBaseTangentialVelocity(const UPrimitiveComponent* MovementBase, const FName BoneName, const FVector& WorldLocation);
 
 	/** Get the transforms for the given MovementBase, optionally at the location of a bone. Returns false if MovementBase is NULL, or if BoneName is not a valid bone. */
 	ENGINE_API bool GetMovementBaseTransform(const UPrimitiveComponent* MovementBase, const FName BoneName, FVector& OutLocation, FQuat& OutQuat);
@@ -107,7 +111,7 @@ namespace MovementBaseUtility
 
 	// Deprecated function
 	DEPRECATED(4.4, "MovementBaseUtility::UseRelativePosition() is deprecated, use MovementBaseUtility::UseRelativeLocation() instead.")
-	FORCEINLINE bool UseRelativePosition(const class UPrimitiveComponent* MovementBase)
+	FORCEINLINE bool UseRelativePosition(const UPrimitiveComponent* MovementBase)
 	{
 		return UseRelativeLocation(MovementBase);
 	}
@@ -336,7 +340,7 @@ public:
 	// Begin APawn Interface.
 	virtual void PostInitializeComponents() override;
 	virtual class UPawnMovementComponent* GetMovementComponent() const override { return CharacterMovement; }
-	virtual class UPrimitiveComponent* GetMovementBase() const override final { return BasedMovement.MovementBase; }
+	virtual UPrimitiveComponent* GetMovementBase() const override final { return BasedMovement.MovementBase; }
 	virtual float GetDefaultHalfHeight() const override;
 	virtual void TurnOff() override;
 	virtual void Restart() override;
@@ -569,35 +573,6 @@ public:
 	virtual void UpdateFromCompressedFlags(uint8 Flags) { }
 
 public:
-	// Note: these network functions should be moved to the Character movement component. They are currently just wrappers that pass the calls to CharacterMovement.
-
-	/** Replicated function sent by client to server - contains client movement and view info. */
-	UFUNCTION(unreliable, server, WithValidation)
-	void ServerMove(float TimeStamp, FVector_NetQuantize100 InAccel, FVector_NetQuantize100 ClientLoc, uint8 MoveFlags, uint8 ClientRoll, uint32 View, class UPrimitiveComponent* ClientMovementBase, uint8 ClientMovementMode);
-
-	/** Replicated function sent by client to server - contains client movement and view info for two moves. */
-	UFUNCTION(unreliable, server, WithValidation)
-	void ServerMoveDual(float TimeStamp0, FVector_NetQuantize100 InAccel0, uint8 PendingFlags, uint32 View0, float TimeStamp, FVector_NetQuantize100 InAccel, FVector_NetQuantize100 ClientLoc, uint8 NewFlags, uint8 ClientRoll, uint32 View, class UPrimitiveComponent* ClientMovementBase, uint8 ClientMovementMode);
-
-	/* Resending an (important) old move.  Process it if not already processed. */
-	UFUNCTION(unreliable, server, WithValidation)
-	void ServerMoveOld(float OldTimeStamp, FVector_NetQuantize100 OldAccel, uint8 OldMoveFlags);
-
-	/** If no client adjustment is needed after processing received ServerMove(), ack the good move so client can remove it from SavedMoves */
-	UFUNCTION(unreliable, client)
-	void ClientAckGoodMove(float TimeStamp);
-
-	/** Replicate position correction to client, associated with a timestamped servermove.  Client will replay subsequent moves after applying adjustment.  */
-	UFUNCTION(unreliable, client)
-	void ClientAdjustPosition(float TimeStamp, FVector NewLoc, FVector NewVel, class UPrimitiveComponent* NewBase, bool bHasBase, bool bBaseRelativePosition, uint8 ServerMovementMode);
-
-	/* Bandwidth saving version, when velocity is zeroed */
-	UFUNCTION(unreliable, client)
-	void ClientVeryShortAdjustPosition(float TimeStamp, FVector NewLoc, class UPrimitiveComponent* NewBase, bool bHasBase, bool bBaseRelativePosition, uint8 ServerMovementMode);
-
-	/** Replicate position correction to client when using root motion for mevement. */
-	UFUNCTION(unreliable, client)
-	void ClientAdjustRootMotionPosition(float TimeStamp, float ServerMontageTrackPosition, FVector ServerLoc, FVector_NetQuantizeNormal ServerRotation, float ServerVelZ, class UPrimitiveComponent* ServerBase, bool bHasBase, bool bBaseRelativePosition, uint8 ServerMovementMode);
 
 	UFUNCTION(reliable, client)
 	void ClientCheatWalk();
diff --git a/Engine/Source/Runtime/Engine/Classes/GameFramework/CharacterMovementComponent.h b/Engine/Source/Runtime/Engine/Classes/GameFramework/CharacterMovementComponent.h
--- a/Engine/Source/Runtime/Engine/Classes/GameFramework/CharacterMovementComponent.h
+++ b/Engine/Source/Runtime/Engine/Classes/GameFramework/CharacterMovementComponent.h
@@ -327,7 +327,7 @@ public:
 
 	/** What to update CharacterOwner and UpdatedComponent after movement ends */
 	UPROPERTY()
-	class UPrimitiveComponent* DeferredUpdatedMoveComponent;
+	UPrimitiveComponent* DeferredUpdatedMoveComponent;
 
 	/** Maximum step height for getting out of water */
 	UPROPERTY(Category="Character Movement", EditAnywhere, BlueprintReadWrite, AdvancedDisplay, meta=(ClampMin="0", UIMin="0"))
@@ -743,9 +743,10 @@ public:
 	virtual void UpdateBasedRotation(FRotator &FinalRotation, const FRotator& ReducedRotation);
 
 	/** Update (or defer updating) OldBaseLocation and OldBaseQuat if there is a valid movement base. */
+	DEPRECATED(4.4, "CharacterMovementComponent::MaybeSaveBaseLocation() will be removed, call SaveBaseLocation().")
 	virtual void MaybeSaveBaseLocation();
 
-	/** Update OldBaseLocation and OldBaseQuat if there is a valid movement base. */
+	/** Update OldBaseLocation and OldBaseQuat if there is a valid movement base, and store the relative location/rotation if necessary. */
 	virtual void SaveBaseLocation();
 
 	/** changes physics based on MovementMode */
@@ -967,7 +968,7 @@ public:
 	virtual void MoveSmooth(const FVector& InVelocity, const float DeltaSeconds, FStepDownResult* OutStepDownResult = NULL );
 
 	
-	virtual void SetUpdatedComponent(class UPrimitiveComponent* NewUpdatedComponent) override;
+	virtual void SetUpdatedComponent(UPrimitiveComponent* NewUpdatedComponent) override;
 	
 	/** @Return MovementMode string */
 	virtual FString GetMovementName();
@@ -1317,7 +1318,7 @@ protected:
 	virtual void CallServerMove(const class FSavedMove_Character* NewMove, const class FSavedMove_Character* OldMove);
 	
 	/** Have the server check if the client is outside an error tolerance, and set a client adjustment if so. ClientLoc will be a relative location if MovementBaseUtility::UseRelativePosition(ClientMovementBase) is true. */
-	virtual void ServerMoveHandleClientError(float TimeStamp, float DeltaTime, const FVector& Accel, const FVector& ClientLoc, class UPrimitiveComponent* ClientMovementBase, uint8 ClientMovementMode);
+	virtual void ServerMoveHandleClientError(float TimeStamp, float DeltaTime, const FVector& Accel, const FVector& ClientLoc, UPrimitiveComponent* ClientMovementBase, FName ClientBaseBoneName, uint8 ClientMovementMode);
 
 	/* Process a move at the given time stamp, given the compressed flags representing various events that occurred (ie jump). */
 	virtual void MoveAutonomous( float ClientTimeStamp, float DeltaTime, uint8 CompressedFlags, const FVector& NewAccel);
@@ -1347,14 +1348,38 @@ public:
 		@returns true if TimeStamp is valid, or false if it has expired. */
 	bool VerifyClientTimeStamp(float TimeStamp, FNetworkPredictionData_Server_Character & ServerData);
 
-	// Callbacks for RPCs on Character
-	virtual void ServerMove_Implementation(float TimeStamp, FVector_NetQuantize100 InAccel, FVector_NetQuantize100 ClientLoc, uint8 CompressedMoveFlags, uint8 ClientRoll, uint32 View, class UPrimitiveComponent* ClientMovementBase, uint8 ClientMovementMode);
-	virtual void ServerMoveDual_Implementation(float TimeStamp0, FVector_NetQuantize100 InAccel0, uint8 PendingFlags, uint32 View0, float TimeStamp, FVector_NetQuantize100 InAccel, FVector_NetQuantize100 ClientLoc, uint8 NewFlags, uint8 ClientRoll, uint32 View, class UPrimitiveComponent* ClientMovementBase, uint8 ClientMovementMode);
-	virtual void ServerMoveOld_Implementation(float OldTimeStamp, FVector_NetQuantize100 OldAccel, uint8 OldMoveFlags);
-	virtual void ClientAckGoodMove_Implementation(float TimeStamp);
-	virtual void ClientAdjustPosition_Implementation(float TimeStamp, FVector NewLoc, FVector NewVel, class UPrimitiveComponent* NewBase, bool bHasBase, bool bBaseRelativePosition, uint8 ServerMovementMode);
-	virtual void ClientVeryShortAdjustPosition_Implementation(float TimeStamp, FVector NewLoc, class UPrimitiveComponent* NewBase, bool bHasBase, bool bBaseRelativePosition, uint8 ServerMovementMode);
-	void ClientAdjustRootMotionPosition_Implementation(float TimeStamp, float ServerMontageTrackPosition, FVector ServerLoc, FVector_NetQuantizeNormal ServerRotation, float ServerVelZ, class UPrimitiveComponent* ServerBase, bool bHasBase, bool bBaseRelativePosition, uint8 ServerMovementMode);
+	////////////////////////////////////
+	// Network RPCs for movement
+	////////////////////////////////////
+
+	/** Replicated function sent by client to server - contains client movement and view info. */
+	UFUNCTION(unreliable, server, WithValidation)
+	virtual void ServerMove(float TimeStamp, FVector_NetQuantize100 InAccel, FVector_NetQuantize100 ClientLoc, uint8 CompressedMoveFlags, uint8 ClientRoll, uint32 View, UPrimitiveComponent* ClientMovementBase, FName ClientBaseBoneName, uint8 ClientMovementMode);
+
+	/** Replicated function sent by client to server - contains client movement and view info for two moves. */
+	UFUNCTION(unreliable, server, WithValidation)
+	virtual void ServerMoveDual(float TimeStamp0, FVector_NetQuantize100 InAccel0, uint8 PendingFlags, uint32 View0, float TimeStamp, FVector_NetQuantize100 InAccel, FVector_NetQuantize100 ClientLoc, uint8 NewFlags, uint8 ClientRoll, uint32 View, UPrimitiveComponent* ClientMovementBase, FName ClientBaseBoneName, uint8 ClientMovementMode);
+	
+	/* Resending an (important) old move. Process it if not already processed. */
+	UFUNCTION(unreliable, server, WithValidation)
+	virtual void ServerMoveOld(float OldTimeStamp, FVector_NetQuantize100 OldAccel, uint8 OldMoveFlags);
+	
+	/** If no client adjustment is needed after processing received ServerMove(), ack the good move so client can remove it from SavedMoves */
+	UFUNCTION(unreliable, client)
+	virtual void ClientAckGoodMove(float TimeStamp);
+
+	/** Replicate position correction to client, associated with a timestamped servermove.  Client will replay subsequent moves after applying adjustment.  */
+	UFUNCTION(unreliable, client)
+	virtual void ClientAdjustPosition(float TimeStamp, FVector NewLoc, FVector NewVel, UPrimitiveComponent* NewBase, FName NewBaseBoneName, bool bHasBase, bool bBaseRelativePosition, uint8 ServerMovementMode);
+
+	/* Bandwidth saving version, when velocity is zeroed */
+	UFUNCTION(unreliable, client)
+	virtual void ClientVeryShortAdjustPosition(float TimeStamp, FVector NewLoc, UPrimitiveComponent* NewBase, FName NewBaseBoneName, bool bHasBase, bool bBaseRelativePosition, uint8 ServerMovementMode);
+	
+	/** Replicate position correction to client when using root motion for movement. */
+	UFUNCTION(unreliable, client)
+	void ClientAdjustRootMotionPosition(float TimeStamp, float ServerMontageTrackPosition, FVector ServerLoc, FVector_NetQuantizeNormal ServerRotation, float ServerVelZ, UPrimitiveComponent* ServerBase, FName ServerBoneName, bool bHasBase, bool bBaseRelativePosition, uint8 ServerMovementMode);
+
 
 	// Root Motion
 public:
@@ -1443,9 +1468,11 @@ public:
 	FFindFloorResult StartFloor;
 	FRotator StartRotation;
 	FRotator StartControlRotation;
-	FQuat StartBaseRotation; // rotation of the base component, only saved if it can move.
+	FQuat StartBaseRotation;	// rotation of the base component (or bone), only saved if it can move.
 	float StartCapsuleRadius;
 	float StartCapsuleHalfHeight;
+	TWeakObjectPtr<UPrimitiveComponent> StartBase;
+	FName StartBoneName;
 
 	// Information after the move has been performed
 	FVector SavedLocation;
@@ -1453,15 +1480,14 @@ public:
 	FVector SavedVelocity;
 	FVector SavedRelativeLocation;
 	FRotator SavedControlRotation;
+	TWeakObjectPtr<UPrimitiveComponent> EndBase;
+	FName EndBoneName;
 
 	FVector Acceleration;
 
 	// Cached to speed up iteration over IsImportantMove().
 	FVector AccelNormal;
 	float AccelMag;
-	
-	TWeakObjectPtr<UPrimitiveComponent> StartBase;
-	TWeakObjectPtr<UPrimitiveComponent> EndBase;
 
 	TWeakObjectPtr<class UAnimMontage> RootMotionMontage;
 	float RootMotionTrackPosition;
@@ -1484,8 +1510,8 @@ public:
 	/** @Return true if this move is an "important" move that should be sent again if not acked by the server */
 	virtual bool IsImportantMove(const FSavedMovePtr& LastAckedMove) const;
 	
-	/** Returns starting position of move, either absolute StartLocation, or StartRelativeLocation offset from MovementBase. */
-	virtual FVector GetStartLocation(const class UPrimitiveComponent* MovementBase) const;
+	/** Returns starting position if we were to revert the move, either absolute StartLocation, or StartRelativeLocation offset from MovementBase's current location (since we want to try to move forward at this time). */
+	virtual FVector GetRevertedLocation() const;
 
 	enum EPostUpdateMode
 	{
@@ -1532,6 +1558,7 @@ public:
 	, NewVel(ForceInitToZero)
 	, NewRot(ForceInitToZero)
 	, NewBase(NULL)
+	, NewBaseBoneName(NAME_None)
 	, bAckGoodMove(false)
 	, bBaseRelativePosition(false)
 	, MovementMode(0)
@@ -1543,7 +1570,8 @@ public:
 	FVector NewLoc;
 	FVector NewVel;
 	FRotator NewRot;
-	class UPrimitiveComponent* NewBase;
+	UPrimitiveComponent* NewBase;
+	FName NewBaseBoneName;
 	bool bAckGoodMove;
 	bool bBaseRelativePosition;
 	uint8 MovementMode;
diff --git a/Engine/Source/Runtime/Engine/Private/Character.cpp b/Engine/Source/Runtime/Engine/Private/Character.cpp
--- a/Engine/Source/Runtime/Engine/Private/Character.cpp
+++ b/Engine/Source/Runtime/Engine/Private/Character.cpp
@@ -369,7 +369,7 @@ void ACharacter::ClearCrossLevelReferences()
 
 namespace MovementBaseUtility
 {
-	void AddTickDependency(FTickFunction& BasedObjectTick, class UPrimitiveComponent* NewBase)
+	void AddTickDependency(FTickFunction& BasedObjectTick, UPrimitiveComponent* NewBase)
 	{
 		if (NewBase && MovementBaseUtility::UseRelativeLocation(NewBase))
 		{
@@ -389,7 +389,7 @@ namespace MovementBaseUtility
 		}
 	}
 
-	void RemoveTickDependency(FTickFunction& BasedObjectTick, class UPrimitiveComponent* OldBase)
+	void RemoveTickDependency(FTickFunction& BasedObjectTick, UPrimitiveComponent* OldBase)
 	{
 		if (OldBase && MovementBaseUtility::UseRelativeLocation(OldBase))
 		{
@@ -409,12 +409,21 @@ namespace MovementBaseUtility
 		}
 	}
 
-	// TODO: GetPhysicsLinearVelocity(BoneName)
-	FVector GetMovementBaseVelocity(const class UPrimitiveComponent* MovementBase)
+	FVector GetMovementBaseVelocity(const UPrimitiveComponent* MovementBase, const FName BoneName)
 	{
 		FVector BaseVelocity = FVector::ZeroVector;
 		if (MovementBaseUtility::IsDynamicBase(MovementBase))
 		{
+			if (BoneName != NAME_None)
+			{
+				const FBodyInstance* BodyInstance = MovementBase->GetBodyInstance(BoneName);
+				if (BodyInstance)
+				{
+					BaseVelocity = BodyInstance->GetUnrealWorldVelocity();
+					return BaseVelocity;
+				}
+			}
+
 			BaseVelocity = MovementBase->GetComponentVelocity();
 			if (BaseVelocity.IsZero())
 			{
@@ -437,18 +446,23 @@ namespace MovementBaseUtility
 		return BaseVelocity;
 	}
 
-	FVector GetMovementBaseTangentialVelocity(const class UPrimitiveComponent* MovementBase, const FVector& WorldLocation)
+	FVector GetMovementBaseTangentialVelocity(const UPrimitiveComponent* MovementBase, const FName BoneName, const FVector& WorldLocation)
 	{
 		if (MovementBaseUtility::IsDynamicBase(MovementBase))
 		{
-			if (MovementBase->GetBodyInstance())
+			if (const FBodyInstance* BodyInstance = MovementBase->GetBodyInstance(BoneName))
 			{
-				const FVector BaseAngVel = MovementBase->GetBodyInstance()->GetUnrealWorldAngularVelocity();
-				if (!BaseAngVel.IsZero())
+				const FVector BaseAngVel = BodyInstance->GetUnrealWorldAngularVelocity();
+				if (!BaseAngVel.IsNearlyZero())
 				{
-					const FVector RadialDistanceToBase = WorldLocation - MovementBase->GetComponentLocation();
-					const FVector TangentialVel = FMath::DegreesToRadians(BaseAngVel) ^ RadialDistanceToBase;
-					return TangentialVel;
+					FVector BaseLocation;
+					FQuat BaseRotation;
+					if (MovementBaseUtility::GetMovementBaseTransform(MovementBase, BoneName, BaseLocation, BaseRotation))
+					{
+						const FVector RadialDistanceToBase = WorldLocation - BaseLocation;
+						const FVector TangentialVel = FMath::DegreesToRadians(BaseAngVel) ^ RadialDistanceToBase;
+						return TangentialVel;
+					}
 				}
 			}			
 		}
@@ -780,77 +794,6 @@ float ACharacter::GetJumpMaxHoldTime() const
 	return JumpMaxHoldTime;
 }
 
-bool ACharacter::ServerMove_Validate(float TimeStamp, FVector_NetQuantize100 InAccel, FVector_NetQuantize100 ClientLoc, uint8 MoveFlags, uint8 ClientRoll, uint32 View, class UPrimitiveComponent* ClientMovementBase, uint8 ClientMovementMode)
-{
-	return true;
-}
-
-void ACharacter::ServerMove_Implementation(float TimeStamp, FVector_NetQuantize100 InAccel, FVector_NetQuantize100 ClientLoc, uint8 MoveFlags, uint8 ClientRoll, uint32 View, class UPrimitiveComponent* ClientMovementBase, uint8 ClientMovementMode)
-{
-	if (CharacterMovement)
-	{
-		CharacterMovement->ServerMove_Implementation(TimeStamp, InAccel, ClientLoc, MoveFlags, ClientRoll, View, ClientMovementBase, ClientMovementMode);
-	}
-}
-
-bool ACharacter::ServerMoveDual_Validate(float TimeStamp0, FVector_NetQuantize100 InAccel0, uint8 PendingFlags, uint32 View0, float TimeStamp, FVector_NetQuantize100 InAccel, FVector_NetQuantize100 ClientLoc, uint8 NewFlags, uint8 ClientRoll, uint32 View, class UPrimitiveComponent* ClientMovementBase, uint8 ClientMovementMode)
-{
-	return true;
-}
-
-void ACharacter::ServerMoveDual_Implementation(float TimeStamp0, FVector_NetQuantize100 InAccel0, uint8 PendingFlags, uint32 View0, float TimeStamp, FVector_NetQuantize100 InAccel, FVector_NetQuantize100 ClientLoc, uint8 NewFlags, uint8 ClientRoll, uint32 View, class UPrimitiveComponent* ClientMovementBase, uint8 ClientMovementMode)
-{
-	if (CharacterMovement)
-	{
-		CharacterMovement->ServerMoveDual_Implementation(TimeStamp0, InAccel0, PendingFlags, View0, TimeStamp, InAccel, ClientLoc, NewFlags, ClientRoll, View, ClientMovementBase, ClientMovementMode);
-	}
-}
-
-bool ACharacter::ServerMoveOld_Validate(float OldTimeStamp, FVector_NetQuantize100 OldAccel, uint8 OldMoveFlags)
-{
-	return true;
-}
-
-void ACharacter::ServerMoveOld_Implementation(float OldTimeStamp, FVector_NetQuantize100 OldAccel, uint8 OldMoveFlags)
-{
-	if (CharacterMovement)
-	{
-		CharacterMovement->ServerMoveOld_Implementation(OldTimeStamp, OldAccel, OldMoveFlags);
-	}
-}
-
-void ACharacter::ClientAckGoodMove_Implementation(float TimeStamp)
-{
-	if (CharacterMovement)
-	{
-		CharacterMovement->ClientAckGoodMove_Implementation(TimeStamp);
-	}
-}
-
-void ACharacter::ClientAdjustPosition_Implementation(float TimeStamp, FVector NewLoc, FVector NewVel, class UPrimitiveComponent* NewBase, bool bHasBase, bool bBaseRelativePosition, uint8 ServerMovementMode)
-{
-	if (CharacterMovement)
-	{
-		CharacterMovement->ClientAdjustPosition_Implementation(TimeStamp, NewLoc, NewVel, NewBase, bHasBase, bBaseRelativePosition, ServerMovementMode);
-	}
-}
-
-void ACharacter::ClientVeryShortAdjustPosition_Implementation(float TimeStamp, FVector NewLoc, class UPrimitiveComponent* NewBase, bool bHasBase, bool bBaseRelativePosition, uint8 ServerMovementMode)
-{
-	if (CharacterMovement)
-	{
-		CharacterMovement->ClientVeryShortAdjustPosition_Implementation(TimeStamp, NewLoc, NewBase, bHasBase, bBaseRelativePosition, ServerMovementMode);
-	}
-}
-
-void ACharacter::ClientAdjustRootMotionPosition_Implementation(float TimeStamp, float ServerMontageTrackPosition, FVector ServerLoc, FVector_NetQuantizeNormal ServerRotation, float ServerVelZ, class UPrimitiveComponent * ServerBase, bool bHasBase, bool bBaseRelativePosition, uint8 ServerMovementMode)
-{
-	if( CharacterMovement )
-	{
-		CharacterMovement->ClientAdjustRootMotionPosition_Implementation(TimeStamp, ServerMontageTrackPosition, ServerLoc, ServerRotation, ServerVelZ, ServerBase, bHasBase, bBaseRelativePosition, ServerMovementMode);
-	}
-}
-
 //
 // Static variables for networking.
 //
@@ -1073,7 +1016,8 @@ bool ACharacter::CanUseRootMotionRepMove(const FSimulatedRootMotionReplicatedMov
 
 bool ACharacter::RestoreReplicatedMove(const FSimulatedRootMotionReplicatedMove & RootMotionRepMove)
 {
-	UPrimitiveComponent * ServerBase = RootMotionRepMove.RootMotion.MovementBase;
+	UPrimitiveComponent* ServerBase = RootMotionRepMove.RootMotion.MovementBase;
+	const FName ServerBaseBoneName = RootMotionRepMove.RootMotion.MovementBaseBoneName;
 
 	// Relative Position
 	if( RootMotionRepMove.RootMotion.bRelativePosition )
@@ -1081,13 +1025,16 @@ bool ACharacter::RestoreReplicatedMove(const FSimulatedRootMotionReplicatedMove
 		bool bSuccess = false;
 		if( MovementBaseUtility::UseRelativeLocation(ServerBase) )
 		{
-			const FVector ServerLocation = ServerBase->GetComponentLocation() + RootMotionRepMove.RootMotion.Location;
+			FVector BaseLocation;
+			FQuat BaseRotation;
+			MovementBaseUtility::GetMovementBaseTransform(ServerBase, ServerBaseBoneName, BaseLocation, BaseRotation);
+
+			const FVector ServerLocation = BaseLocation + RootMotionRepMove.RootMotion.Location;
 			FRotator ServerRotation;
 			if (RootMotionRepMove.RootMotion.bRelativeRotation)
 			{
 				// Relative rotation
-				// TODO: Bone rotation
-				ServerRotation = (FRotationMatrix(RootMotionRepMove.RootMotion.Rotation) * FRotationMatrix(ServerBase->GetComponentRotation())).Rotator();
+				ServerRotation = (FRotationMatrix(RootMotionRepMove.RootMotion.Rotation) * FQuatRotationTranslationMatrix(BaseRotation, FVector::ZeroVector)).Rotator();
 			}
 			else
 			{
@@ -1110,11 +1057,7 @@ bool ACharacter::RestoreReplicatedMove(const FSimulatedRootMotionReplicatedMove
 		UpdateSimulatedPosition(RootMotionRepMove.RootMotion.Location, RootMotionRepMove.RootMotion.Rotation);
 	}
 
-	// Set base
-	if( BasedMovement.MovementBase != ServerBase /* TODO: or bone changes */)
-	{
-		SetBase( ServerBase /* TODO: bone */ );
-	}
+	SetBase( ServerBase, ServerBaseBoneName );
 
 	return true;
 }
@@ -1177,6 +1120,7 @@ void ACharacter::PreReplication( IRepChangedPropertyTracker & ChangedPropertyTra
 		RepRootMotion.Location			= RepRootMotion.bRelativePosition ? BasedMovement.Location : GetActorLocation();
 		RepRootMotion.Rotation			= RepRootMotion.bRelativeRotation ? BasedMovement.Rotation : GetActorRotation();
 		RepRootMotion.MovementBase		= BasedMovement.MovementBase;
+		RepRootMotion.MovementBaseBoneName = BasedMovement.BoneName;
 		RepRootMotion.AnimMontage		= RootMotionMontageInstance->Montage;
 		RepRootMotion.Position			= RootMotionMontageInstance->Position;
 
diff --git a/Engine/Source/Runtime/Engine/Private/CharacterMovementComponent.cpp b/Engine/Source/Runtime/Engine/Private/CharacterMovementComponent.cpp
--- a/Engine/Source/Runtime/Engine/Private/CharacterMovementComponent.cpp
+++ b/Engine/Source/Runtime/Engine/Private/CharacterMovementComponent.cpp
@@ -446,12 +446,12 @@ FVector UCharacterMovementComponent::GetImpartedMovementBaseVelocity() const
 		UPrimitiveComponent* MovementBase = CharacterOwner->GetMovementBase();
 		if (MovementBaseUtility::IsDynamicBase(MovementBase))
 		{
-			FVector BaseVelocity = MovementBaseUtility::GetMovementBaseVelocity(MovementBase);
+			FVector BaseVelocity = MovementBaseUtility::GetMovementBaseVelocity(MovementBase, CharacterOwner->GetBasedMovement().BoneName);
 			
 			if (bImpartBaseAngularVelocity)
 			{
 				const FVector CharacterBasePosition = (UpdatedComponent->GetComponentLocation() - FVector(0.f, 0.f, CharacterOwner->CapsuleComponent->GetScaledCapsuleHalfHeight()));
-				const FVector BaseTangentialVel = MovementBaseUtility::GetMovementBaseTangentialVelocity(MovementBase, CharacterBasePosition);
+				const FVector BaseTangentialVel = MovementBaseUtility::GetMovementBaseTangentialVelocity(MovementBase, CharacterOwner->GetBasedMovement().BoneName, CharacterBasePosition);
 				BaseVelocity += BaseTangentialVel;
 			}
 
@@ -792,7 +792,7 @@ void UCharacterMovementComponent::TickComponent(float DeltaTime, enum ELevelTick
 			// Between net updates from the client we need to update position if based on another object,
 			// otherwise the object will move on intermediate frames and we won't follow it.
 			MaybeUpdateBasedMovement(DeltaTime);
-			MaybeSaveBaseLocation();
+			SaveBaseLocation();
 		}
 	}
 	else if (CharacterOwner->Role == ROLE_SimulatedProxy)
@@ -1122,19 +1122,18 @@ void UCharacterMovementComponent::SimulateMovement(float DeltaSeconds)
 	// Call custom post-movement events. These happen after the scoped movement completes in case the events want to use the current state of overlaps etc.
 	CallMovementUpdateDelegate(DeltaSeconds, OldLocation, OldVelocity);
 
-	MaybeSaveBaseLocation();
+	SaveBaseLocation();
 	UpdateComponentVelocity();
 	bJustTeleported = false;
 
 	LastUpdateLocation = UpdatedComponent ? UpdatedComponent->GetComponentLocation() : FVector::ZeroVector;
 }
 
-void UCharacterMovementComponent::SetBase( UPrimitiveComponent* NewBase, const FName BoneName, bool bNotifyActor )
+void UCharacterMovementComponent::SetBase( UPrimitiveComponent* NewBase, FName BoneName, bool bNotifyActor )
 {
 	if (CharacterOwner)
 	{
-		// TODO: temporarily disabling movement on bone until networking is complete.
-		CharacterOwner->SetBase(NewBase, NAME_None /*BoneName*/, bNotifyActor);
+		CharacterOwner->SetBase(NewBase, BoneName, bNotifyActor);
 	}
 }
 
@@ -1157,12 +1156,10 @@ void UCharacterMovementComponent::MaybeUpdateBasedMovement(float DeltaSeconds)
 	}
 }
 
+// todo: deprecated, remove.
 void UCharacterMovementComponent::MaybeSaveBaseLocation()
 {
-	if(!bDeferUpdateBasedMovement)
-	{
-		SaveBaseLocation();
-	}
+	SaveBaseLocation();
 }
 
 // @todo UE4 - handle lift moving up and down through encroachment
@@ -1454,7 +1451,7 @@ void UCharacterMovementComponent::PerformMovement(float DeltaSeconds)
 	// Call external post-movement events. These happen after the scoped movement completes in case the events want to use the current state of overlaps etc.
 	CallMovementUpdateDelegate(DeltaSeconds, OldLocation, OldVelocity);
 
-	MaybeSaveBaseLocation();
+	SaveBaseLocation();
 	UpdateComponentVelocity();
 
 	LastUpdateLocation = UpdatedComponent ? UpdatedComponent->GetComponentLocation() : FVector::ZeroVector;
@@ -5023,7 +5020,7 @@ void UCharacterMovementComponent::ReplicateMoveToServer(float DeltaTime, const F
 	if( ClientData->PendingMove.IsValid() && !ClientData->PendingMove->bOldTimeStampBeforeReset && ClientData->PendingMove->CanCombineWith(NewMove, CharacterOwner, ClientData->MaxResponseTime * CharacterOwner->GetWorldSettings()->GetEffectiveTimeDilation()))
 	{
 		// Only combine and move back to the start location if we don't move back in to a spot that would make us collide with something new.
-		const FVector OldStartLocation = ClientData->PendingMove->GetStartLocation(ClientData->PendingMove->StartBase.Get());
+		const FVector OldStartLocation = ClientData->PendingMove->GetRevertedLocation();
 		if (!OverlapTest(OldStartLocation, ClientData->PendingMove->StartRotation.Quaternion(), UpdatedComponent->GetCollisionObjectType(), GetPawnCapsuleCollisionShape(SHRINK_None), CharacterOwner))
 		{
 			FScopedMovementUpdate ScopedMovementUpdate(UpdatedComponent, EScopedUpdate::DeferredUpdates);
@@ -5034,10 +5031,7 @@ void UCharacterMovementComponent::ReplicateMoveToServer(float DeltaTime, const F
 			UpdatedComponent->SetWorldLocationAndRotation(OldStartLocation, ClientData->PendingMove->StartRotation, false);
 			Velocity = ClientData->PendingMove->StartVelocity;
 
-			if (ClientData->PendingMove->StartBase != CharacterOwner->GetMovementBase())
-			{
-				SetBase(ClientData->PendingMove->StartBase.Get() /* TODO: BoneName on saved move */);
-			}			
+			SetBase(ClientData->PendingMove->StartBase.Get(), ClientData->PendingMove->StartBoneName);
 			CurrentFloor = ClientData->PendingMove->StartFloor;
 
 			// Now that we have reverted to the old position, prepare a new move from that position,
@@ -5142,12 +5136,13 @@ void UCharacterMovementComponent::CallServerMove
 
 	// Determine if we send absolute or relative location
 	UPrimitiveComponent* ClientMovementBase = NewMove->EndBase.Get();
+	const FName ClientBaseBone = NewMove->EndBoneName;
 	const FVector SendLocation = MovementBaseUtility::UseRelativeLocation(ClientMovementBase) ? NewMove->SavedRelativeLocation : NewMove->SavedLocation;
 
 	// send old move if it exists
 	if (OldMove)
 	{
-		CharacterOwner->ServerMoveOld(OldMove->TimeStamp, OldMove->Acceleration, OldMove->GetCompressedFlags());
+		ServerMoveOld(OldMove->TimeStamp, OldMove->Acceleration, OldMove->GetCompressedFlags());
 	}
 
 	FNetworkPredictionData_Client_Character* ClientData = GetPredictionData_Client_Character();
@@ -5156,7 +5151,7 @@ void UCharacterMovementComponent::CallServerMove
 		const uint32 OldClientYawPitchINT = PackYawAndPitchTo32(ClientData->PendingMove->SavedControlRotation.Yaw, ClientData->PendingMove->SavedControlRotation.Pitch);
 
 		// send two moves simultaneously
-		CharacterOwner->ServerMoveDual
+		ServerMoveDual
 			(
 			ClientData->PendingMove->TimeStamp,
 			ClientData->PendingMove->Acceleration,
@@ -5169,12 +5164,13 @@ void UCharacterMovementComponent::CallServerMove
 			ClientRollBYTE,
 			ClientYawPitchINT,
 			ClientMovementBase,
+			ClientBaseBone,
 			PackNetworkMovementMode()
 			);
 	}
 	else
 	{
-		CharacterOwner->ServerMove
+		ServerMove
 			(
 			NewMove->TimeStamp,
 			NewMove->Acceleration,
@@ -5183,6 +5179,7 @@ void UCharacterMovementComponent::CallServerMove
 			ClientRollBYTE,
 			ClientYawPitchINT,
 			ClientMovementBase,
+			ClientBaseBone,
 			PackNetworkMovementMode()
 			);
 	}
@@ -5227,10 +5224,23 @@ void UCharacterMovementComponent::ServerMoveOld_Implementation
 }
 
 
-void UCharacterMovementComponent::ServerMoveDual_Implementation(float TimeStamp0, FVector_NetQuantize100 InAccel0, uint8 PendingFlags, uint32 View0, float TimeStamp, FVector_NetQuantize100 InAccel, FVector_NetQuantize100 ClientLoc, uint8 NewFlags, uint8 ClientRoll, uint32 View, class UPrimitiveComponent* ClientMovementBase, uint8 ClientMovementMode)
+void UCharacterMovementComponent::ServerMoveDual_Implementation(
+	float TimeStamp0,
+	FVector_NetQuantize100 InAccel0,
+	uint8 PendingFlags,
+	uint32 View0,
+	float TimeStamp,
+	FVector_NetQuantize100 InAccel,
+	FVector_NetQuantize100 ClientLoc,
+	uint8 NewFlags,
+	uint8 ClientRoll,
+	uint32 View,
+	UPrimitiveComponent* ClientMovementBase,
+	FName ClientBaseBone,
+	uint8 ClientMovementMode)
 {
-	ServerMove_Implementation(TimeStamp0, InAccel0, FVector(1.f,2.f,3.f), PendingFlags, ClientRoll, View0, ClientMovementBase, ClientMovementMode);
-	ServerMove_Implementation(TimeStamp, InAccel, ClientLoc, NewFlags, ClientRoll, View, ClientMovementBase, ClientMovementMode);
+	ServerMove_Implementation(TimeStamp0, InAccel0, FVector(1.f,2.f,3.f), PendingFlags, ClientRoll, View0, ClientMovementBase, ClientBaseBone, ClientMovementMode);
+	ServerMove_Implementation(TimeStamp, InAccel, ClientLoc, NewFlags, ClientRoll, View, ClientMovementBase, ClientBaseBone, ClientMovementMode);
 }
 
 bool UCharacterMovementComponent::VerifyClientTimeStamp(float TimeStamp, FNetworkPredictionData_Server_Character & ServerData)
@@ -5266,7 +5276,16 @@ bool UCharacterMovementComponent::VerifyClientTimeStamp(float TimeStamp, FNetwor
 	return true;
 }
 
-void UCharacterMovementComponent::ServerMove_Implementation(float TimeStamp, FVector_NetQuantize100 InAccel, FVector_NetQuantize100 ClientLoc, uint8 MoveFlags, uint8 ClientRoll, uint32 View, class UPrimitiveComponent* ClientMovementBase, uint8 ClientMovementMode)
+void UCharacterMovementComponent::ServerMove_Implementation(
+	float TimeStamp,
+	FVector_NetQuantize100 InAccel,
+	FVector_NetQuantize100 ClientLoc,
+	uint8 MoveFlags,
+	uint8 ClientRoll,
+	uint32 View,
+	UPrimitiveComponent* ClientMovementBase,
+	FName ClientBaseBoneName,
+	uint8 ClientMovementMode)
 {
 	if (!HasValidData() || !IsComponentTickEnabled())
 	{
@@ -5331,11 +5350,11 @@ void UCharacterMovementComponent::ServerMove_Implementation(float TimeStamp, FVe
 	UE_LOG(LogNetPlayerMovement, Verbose, TEXT("ServerMove Time %f Acceleration %s Position %s DeltaTime %f"),
 			TimeStamp, *Accel.ToString(), *CharacterOwner->GetActorLocation().ToString(), DeltaTime);
 
-	ServerMoveHandleClientError(TimeStamp, DeltaTime, Accel, ClientLoc, ClientMovementBase, ClientMovementMode);
+	ServerMoveHandleClientError(TimeStamp, DeltaTime, Accel, ClientLoc, ClientMovementBase, ClientBaseBoneName, ClientMovementMode);
 }
 
 
-void UCharacterMovementComponent::ServerMoveHandleClientError(float TimeStamp, float DeltaTime, const FVector& Accel, const FVector& RelativeClientLoc, class UPrimitiveComponent* ClientMovementBase, uint8 ClientMovementMode)
+void UCharacterMovementComponent::ServerMoveHandleClientError(float TimeStamp, float DeltaTime, const FVector& Accel, const FVector& RelativeClientLoc, UPrimitiveComponent* ClientMovementBase, FName ClientBaseBoneName, uint8 ClientMovementMode)
 {
 	if (RelativeClientLoc == FVector(1.f,2.f,3.f)) // first part of double servermove
 	{
@@ -5357,8 +5376,10 @@ void UCharacterMovementComponent::ServerMoveHandleClientError(float TimeStamp, f
 	FVector ClientLoc = RelativeClientLoc;
 	if (MovementBaseUtility::UseRelativeLocation(ClientMovementBase))
 	{
-		// TODO: use MovementBaseUtility::GetMovementBaseTransform(ClientMovementBase) instead.
-		ClientLoc += ClientMovementBase->GetComponentLocation();
+		FVector BaseLocation;
+		FQuat BaseRotation;
+		MovementBaseUtility::GetMovementBaseTransform(ClientMovementBase, ClientBaseBoneName, BaseLocation, BaseRotation);
+		ClientLoc += BaseLocation;
 	}
 
 	// Compute the client error from the server's position
@@ -5371,24 +5392,26 @@ void UCharacterMovementComponent::ServerMoveHandleClientError(float TimeStamp, f
 		UPrimitiveComponent* MovementBase = CharacterOwner->GetMovementBase();
 		ServerData->PendingAdjustment.NewVel = Velocity;
 		ServerData->PendingAdjustment.NewBase = MovementBase;
+		ServerData->PendingAdjustment.NewBaseBoneName = CharacterOwner->GetBasedMovement().BoneName;
+		ServerData->PendingAdjustment.NewLoc = CharacterOwner->GetActorLocation();
+		ServerData->PendingAdjustment.NewRot = CharacterOwner->GetActorRotation();
+
 		ServerData->PendingAdjustment.bBaseRelativePosition = MovementBaseUtility::UseRelativeLocation(MovementBase);
 		if (ServerData->PendingAdjustment.bBaseRelativePosition)
 		{
-			// TODO: use MovementBaseUtility::GetMovementBaseTransform() instead.
-			ServerData->PendingAdjustment.NewLoc = CharacterOwner->GetActorLocation() - MovementBase->GetComponentLocation();
-		}
-		else
-		{
-			ServerData->PendingAdjustment.NewLoc = CharacterOwner->GetActorLocation();
+			// Relative location
+			ServerData->PendingAdjustment.NewLoc = CharacterOwner->GetBasedMovement().Location;
+			
+			// TODO: this could be a relative rotation, but all client corrections ignore rotation right now except the root motion one, which would need to be updated.
+			//ServerData->PendingAdjustment.NewRot = CharacterOwner->GetBasedMovement().Rotation;
 		}
 
-		// Character specific logic
-		ServerData->PendingAdjustment.NewRot = CharacterOwner->GetActorRotation();
 
 #if !(UE_BUILD_SHIPPING || UE_BUILD_TEST)
 		if (CVarNetShowCorrections.GetValueOnGameThread() != 0)
 		{
-			UE_LOG(LogNetPlayerMovement, Warning, TEXT("******** Client Error at %f is %f Accel %s LocDiff %s ClientLoc %s, ServerLoc: %s, Base: %s"), TimeStamp, LocDiff.Size(), *Accel.ToString(), *LocDiff.ToString(), *ClientLoc.ToString(), *CharacterOwner->GetActorLocation().ToString(), *GetNameSafe(MovementBase));
+			UE_LOG(LogNetPlayerMovement, Warning, TEXT("******** Client Error at %f is %f Accel %s LocDiff %s ClientLoc %s, ServerLoc: %s, Base: %s, Bone: %s"),
+				TimeStamp, LocDiff.Size(), *Accel.ToString(), *LocDiff.ToString(), *ClientLoc.ToString(), *CharacterOwner->GetActorLocation().ToString(), *GetNameSafe(MovementBase), *ServerData->PendingAdjustment.NewBaseBoneName.ToString());
 			const float DebugLifetime = CVarNetCorrectionLifetime.GetValueOnGameThread();
 			DrawDebugCapsule(GetWorld(), CharacterOwner->GetActorLocation(), CharacterOwner->GetSimpleCollisionHalfHeight(), CharacterOwner->GetSimpleCollisionRadius(), FQuat::Identity, FColor(100, 255, 100), true, DebugLifetime);
 			DrawDebugCapsule(GetWorld(), ClientLoc                    , CharacterOwner->GetSimpleCollisionHalfHeight(), CharacterOwner->GetSimpleCollisionRadius(), FQuat::Identity, FColor(255, 100, 100), true, DebugLifetime);
@@ -5414,7 +5437,7 @@ void UCharacterMovementComponent::ServerMoveHandleClientError(float TimeStamp, f
 				ApplyNetworkMovementMode(ClientMovementMode);
 
 				// Update base and floor at new location.
-				SetBase(ClientMovementBase /* TODO: client bone? */);
+				SetBase(ClientMovementBase, ClientBaseBoneName);
 				UpdateFloorFromAdjustment();
 
 				// Even if base has not changed, we need to recompute the relative offsets (since we've moved).
@@ -5429,6 +5452,23 @@ void UCharacterMovementComponent::ServerMoveHandleClientError(float TimeStamp, f
 }
 
 
+bool UCharacterMovementComponent::ServerMove_Validate(float TimeStamp, FVector_NetQuantize100 InAccel, FVector_NetQuantize100 ClientLoc, uint8 MoveFlags, uint8 ClientRoll, uint32 View, UPrimitiveComponent* ClientMovementBase, FName ClientBaseBoneName, uint8 ClientMovementMode)
+{
+	return true;
+}
+
+bool UCharacterMovementComponent::ServerMoveDual_Validate(float TimeStamp0, FVector_NetQuantize100 InAccel0, uint8 PendingFlags, uint32 View0, float TimeStamp, FVector_NetQuantize100 InAccel, FVector_NetQuantize100 ClientLoc, uint8 NewFlags, uint8 ClientRoll, uint32 View, UPrimitiveComponent* ClientMovementBase, FName ClientBaseBoneName, uint8 ClientMovementMode)
+{
+	return true;
+}
+
+bool UCharacterMovementComponent::ServerMoveOld_Validate(float OldTimeStamp, FVector_NetQuantize100 OldAccel, uint8 OldMoveFlags)
+{
+	return true;
+}
+
+
+
 void UCharacterMovementComponent::MoveAutonomous
 	(
 	float ClientTimeStamp,
@@ -5455,6 +5495,7 @@ void UCharacterMovementComponent::MoveAutonomous
 	if( !CharacterOwner->bClientUpdating && !CharacterOwner->IsPlayingRootMotion() && CharacterOwner->Mesh )
 	{
 		TickCharacterPose(DeltaTime);
+		// TODO: SaveBaseLocation() in case tick moves us?
 	}
 }
 
@@ -5499,7 +5540,7 @@ void UCharacterMovementComponent::SendClientAdjustment()
 	if (ServerData->PendingAdjustment.bAckGoodMove == true)
 	{
 		// just notify client this move was received
-		CharacterOwner->ClientAckGoodMove(ServerData->PendingAdjustment.TimeStamp);
+		ClientAckGoodMove(ServerData->PendingAdjustment.TimeStamp);
 	}
 	else
 	{
@@ -5507,7 +5548,7 @@ void UCharacterMovementComponent::SendClientAdjustment()
 		{
 			FRotator Rotation = ServerData->PendingAdjustment.NewRot.GetNormalized();
 			FVector_NetQuantizeNormal CompressedRotation(Rotation.Pitch / 180.f, Rotation.Yaw / 180.f, Rotation.Roll / 180.f);
-			CharacterOwner->ClientAdjustRootMotionPosition
+			ClientAdjustRootMotionPosition
 				(
 				ServerData->PendingAdjustment.TimeStamp,
 				CharacterOwner->GetRootMotionAnimMontageInstance()->Position,
@@ -5515,6 +5556,7 @@ void UCharacterMovementComponent::SendClientAdjustment()
 				CompressedRotation,
 				ServerData->PendingAdjustment.NewVel.Z,
 				ServerData->PendingAdjustment.NewBase,
+				ServerData->PendingAdjustment.NewBaseBoneName,
 				ServerData->PendingAdjustment.NewBase != NULL,
 				ServerData->PendingAdjustment.bBaseRelativePosition,
 				PackNetworkMovementMode()
@@ -5522,11 +5564,12 @@ void UCharacterMovementComponent::SendClientAdjustment()
 		}
 		else if (ServerData->PendingAdjustment.NewVel.IsZero())
 		{
-			CharacterOwner->ClientVeryShortAdjustPosition
+			ClientVeryShortAdjustPosition
 				(
 				ServerData->PendingAdjustment.TimeStamp,
 				ServerData->PendingAdjustment.NewLoc,
 				ServerData->PendingAdjustment.NewBase,
+				ServerData->PendingAdjustment.NewBaseBoneName,
 				ServerData->PendingAdjustment.NewBase != NULL,
 				ServerData->PendingAdjustment.bBaseRelativePosition,
 				PackNetworkMovementMode()
@@ -5534,12 +5577,13 @@ void UCharacterMovementComponent::SendClientAdjustment()
 		}
 		else
 		{
-			CharacterOwner->ClientAdjustPosition
+			ClientAdjustPosition
 				(
 				ServerData->PendingAdjustment.TimeStamp,
 				ServerData->PendingAdjustment.NewLoc,
 				ServerData->PendingAdjustment.NewVel,
 				ServerData->PendingAdjustment.NewBase,
+				ServerData->PendingAdjustment.NewBaseBoneName,
 				ServerData->PendingAdjustment.NewBase != NULL,
 				ServerData->PendingAdjustment.bBaseRelativePosition,
 				PackNetworkMovementMode()
@@ -5557,6 +5601,7 @@ void UCharacterMovementComponent::ClientVeryShortAdjustPosition_Implementation
 	float TimeStamp,
 	FVector NewLoc,
 	UPrimitiveComponent* NewBase,
+	FName NewBaseBoneName,
 	bool bHasBase,
 	bool bBaseRelativePosition,
 	uint8 ServerMovementMode
@@ -5564,7 +5609,7 @@ void UCharacterMovementComponent::ClientVeryShortAdjustPosition_Implementation
 {
 	if (HasValidData())
 	{
-		CharacterOwner->ClientAdjustPosition(TimeStamp, NewLoc, FVector::ZeroVector, NewBase, bHasBase, bBaseRelativePosition, ServerMovementMode);
+		ClientAdjustPosition(TimeStamp, NewLoc, FVector::ZeroVector, NewBase, NewBaseBoneName, bHasBase, bBaseRelativePosition, ServerMovementMode);
 	}
 }
 
@@ -5575,6 +5620,7 @@ void UCharacterMovementComponent::ClientAdjustPosition_Implementation
 	FVector NewLocation,
 	FVector NewVelocity,
 	UPrimitiveComponent* NewBase,
+	FName NewBaseBoneName,
 	bool bHasBase,
 	bool bBaseRelativePosition,
 	uint8 ServerMovementMode
@@ -5616,18 +5662,20 @@ void UCharacterMovementComponent::ClientAdjustPosition_Implementation
 	}
 	ClientData->AckMove(MoveIndex);
 		
-	//  Received Location is relative to non-world geometry (moving) bases
+	//  Received Location is relative to dynamic base
 	if (bBaseRelativePosition)
 	{
-		// TODO: fix if there is a bone
-		NewLocation += NewBase->GetComponentLocation();
+		FVector BaseLocation;
+		FQuat BaseRotation;
+		MovementBaseUtility::GetMovementBaseTransform(NewBase, NewBaseBoneName, BaseLocation, BaseRotation); // TODO: error handling if returns false
+		NewLocation += BaseLocation;
 	}
 
 #if !(UE_BUILD_SHIPPING || UE_BUILD_TEST)
 	static const auto CVarShowCorrections = IConsoleManager::Get().FindTConsoleVariableDataInt(TEXT("p.NetShowCorrections"));
 	if (CVarShowCorrections && CVarShowCorrections->GetValueOnGameThread() != 0)
 	{
-		UE_LOG(LogNetPlayerMovement, Warning, TEXT("******** ClientAdjustPosition Time %f velocity %s position %s NewBase: %s SavedMoves %d"), TimeStamp, *NewVelocity.ToString(), *NewLocation.ToString(), *GetNameSafe(NewBase), ClientData->SavedMoves.Num());
+		UE_LOG(LogNetPlayerMovement, Warning, TEXT("******** ClientAdjustPosition Time %f velocity %s position %s NewBase: %s NewBone: %s SavedMoves %d"), TimeStamp, *NewVelocity.ToString(), *NewLocation.ToString(), *GetNameSafe(NewBase), *NewBaseBoneName.ToString(), ClientData->SavedMoves.Num());
 		static const auto CVarCorrectionLifetime = IConsoleManager::Get().FindTConsoleVariableDataFloat(TEXT("p.NetCorrectionLifetime"));
 		const float DebugLifetime = CVarCorrectionLifetime ? CVarCorrectionLifetime->GetValueOnGameThread() : 1.f;
 		DrawDebugCapsule(GetWorld(), CharacterOwner->GetActorLocation()	, CharacterOwner->GetSimpleCollisionHalfHeight(), CharacterOwner->GetSimpleCollisionRadius(), FQuat::Identity, FColor(255, 100, 100), true, DebugLifetime);
@@ -5645,6 +5693,7 @@ void UCharacterMovementComponent::ClientAdjustPosition_Implementation
 
 	// Set base component
 	UPrimitiveComponent* FinalBase = NewBase;
+	FName FinalBaseBoneName = NewBaseBoneName;
 	if (bUnresolvedBase)
 	{
 		check(NewBase == NULL);
@@ -5656,9 +5705,10 @@ void UCharacterMovementComponent::ClientAdjustPosition_Implementation
 		{
 			FindFloor(UpdatedComponent->GetComponentLocation(), CurrentFloor, false);
 			FinalBase = CurrentFloor.HitResult.Component.Get();
+			FinalBaseBoneName = CurrentFloor.HitResult.BoneName;
 		}
 	}
-	SetBase(FinalBase /* TODO: bone? */);
+	SetBase(FinalBase, FinalBaseBoneName);
 
 	// Update floor at new location
 	UpdateFloorFromAdjustment();
@@ -5671,7 +5721,17 @@ void UCharacterMovementComponent::ClientAdjustPosition_Implementation
 	ClientData->bUpdatePosition = true;
 }
 
-void UCharacterMovementComponent::ClientAdjustRootMotionPosition_Implementation(float TimeStamp, float ServerMontageTrackPosition, FVector ServerLoc, FVector_NetQuantizeNormal ServerRotation, float ServerVelZ, class UPrimitiveComponent * ServerBase, bool bHasBase, bool bBaseRelativePosition, uint8 ServerMovementMode)
+void UCharacterMovementComponent::ClientAdjustRootMotionPosition_Implementation(
+	float TimeStamp,
+	float ServerMontageTrackPosition,
+	FVector ServerLoc,
+	FVector_NetQuantizeNormal ServerRotation,
+	float ServerVelZ,
+	UPrimitiveComponent * ServerBase,
+	FName ServerBaseBoneName,
+	bool bHasBase,
+	bool bBaseRelativePosition,
+	uint8 ServerMovementMode)
 {
 	if (!HasValidData() || !IsComponentTickEnabled())
 	{
@@ -5679,7 +5739,7 @@ void UCharacterMovementComponent::ClientAdjustRootMotionPosition_Implementation(
 	}
 
 	// Call ClientAdjustPosition first. This will Ack the move if it's not outdated.
-	CharacterOwner->ClientAdjustPosition(TimeStamp, ServerLoc, FVector(0.f, 0.f, ServerVelZ), ServerBase, bHasBase, bBaseRelativePosition, ServerMovementMode);
+	ClientAdjustPosition(TimeStamp, ServerLoc, FVector(0.f, 0.f, ServerVelZ), ServerBase, ServerBaseBoneName, bHasBase, bBaseRelativePosition, ServerMovementMode);
 	
 	FNetworkPredictionData_Client_Character* ClientData = GetPredictionData_Client_Character();
 	check(ClientData);
@@ -6203,15 +6263,16 @@ void FSavedMove_Character::Clear()
 	StartBaseRotation = FQuat::Identity;
 	StartCapsuleRadius = 0.f;
 	StartCapsuleHalfHeight = 0.f;
+	StartBase = NULL;
+	StartBoneName = NAME_None;
 
 	SavedLocation = FVector::ZeroVector;
 	SavedRotation = FRotator::ZeroRotator;
 	SavedRelativeLocation = FVector::ZeroVector;
 	Acceleration = FVector::ZeroVector;
 	SavedControlRotation = FRotator::ZeroRotator;
-
-	StartBase = NULL;
 	EndBase = NULL;
+	EndBoneName = NAME_None;
 
 	RootMotionMontage = NULL;
 	RootMotionTrackPosition = 0.f;
@@ -6248,17 +6309,16 @@ void FSavedMove_Character::SetInitialPosition(ACharacter* Character)
 	StartBaseRotation = FQuat::Identity;
 	StartFloor = Character->CharacterMovement->CurrentFloor;
 	CustomTimeDilation = Character->CustomTimeDilation;
-
-	StartControlRotation = Character->GetControlRotation().Clamp();
+	StartBoneName = Character->GetBasedMovement().BoneName;
 
 	if (MovementBaseUtility::UseRelativeLocation(MovementBase))
 	{
 		StartRelativeLocation = Character->GetBasedMovement().Location;
-		// TODO: bone
-		FVector StartBaseLocation;
-		MovementBaseUtility::GetMovementBaseTransform(MovementBase, NAME_None, StartBaseLocation, StartBaseRotation);
+		FVector StartBaseLocation_Unused;
+		MovementBaseUtility::GetMovementBaseTransform(MovementBase, StartBoneName, StartBaseLocation_Unused, StartBaseRotation);
 	}
 
+	StartControlRotation = Character->GetControlRotation().Clamp();
 	Character->CapsuleComponent->GetScaledCapsuleSize(StartCapsuleRadius, StartCapsuleHalfHeight);
 }
 
@@ -6270,6 +6330,7 @@ void FSavedMove_Character::PostUpdate(ACharacter* Character, FSavedMove_Characte
 		SavedRotation = Character->GetActorRotation();
 		UPrimitiveComponent* const MovementBase = Character->GetMovementBase();
 		EndBase = MovementBase;
+		EndBoneName = Character->GetBasedMovement().BoneName;
 		if (MovementBaseUtility::UseRelativeLocation(MovementBase))
 		{
 			SavedRelativeLocation = Character->GetBasedMovement().Location;
@@ -6316,12 +6377,14 @@ bool FSavedMove_Character::IsImportantMove(const FSavedMovePtr& LastAckedMove) c
 	return false;
 }
 
-FVector FSavedMove_Character::GetStartLocation(const class UPrimitiveComponent* MovementBase) const
+FVector FSavedMove_Character::GetRevertedLocation() const
 {
+	const UPrimitiveComponent* MovementBase = StartBase.Get();
 	if (MovementBaseUtility::UseRelativeLocation(MovementBase))
 	{
-		// TODO: use MovementBaseUtility::GetMovementBaseTransform() instead.
-		return MovementBase->GetComponentLocation() + StartRelativeLocation;
+		FVector BaseLocation; FQuat BaseRotation;
+		MovementBaseUtility::GetMovementBaseTransform(MovementBase, StartBoneName, BaseLocation, BaseRotation);
+		return BaseLocation + StartRelativeLocation;
 	}
 
 	return StartLocation;
@@ -6350,6 +6413,7 @@ bool FSavedMove_Character::CanCombineWith(const FSavedMovePtr& NewMove, ACharact
 			&& !bPressedJump && !NewMove->bPressedJump
 			&& bWantsToCrouch == NewMove->bWantsToCrouch
 			&& StartBase == NewMove->StartBase
+			&& StartBoneName == NewMove->StartBoneName
 			&& MovementMode == NewMove->MovementMode
 			&& StartCapsuleRadius == NewMove->StartCapsuleRadius
 			&& StartCapsuleHalfHeight == NewMove->StartCapsuleHalfHeight
@@ -6362,6 +6426,7 @@ bool FSavedMove_Character::CanCombineWith(const FSavedMovePtr& NewMove, ACharact
 			&& !bPressedJump && !NewMove->bPressedJump
 			&& bWantsToCrouch == NewMove->bWantsToCrouch
 			&& StartBase == NewMove->StartBase
+			&& StartBoneName == NewMove->StartBoneName
 			&& MovementMode == NewMove->MovementMode
 			&& StartCapsuleRadius == NewMove->StartCapsuleRadius
 			&& StartCapsuleHalfHeight == NewMove->StartCapsuleHalfHeight
