Remove visibility on unnecessary properties

diff --git a/Engine/Source/Runtime/Engine/Classes/Animation/BlendSpaceBase.h b/Engine/Source/Runtime/Engine/Classes/Animation/BlendSpaceBase.h
--- a/Engine/Source/Runtime/Engine/Classes/Animation/BlendSpaceBase.h
+++ b/Engine/Source/Runtime/Engine/Classes/Animation/BlendSpaceBase.h
@@ -202,7 +202,7 @@ class UBlendSpaceBase : public UAnimationAsset
 
 protected:
 	/** Blend Parameters for each axis. **/
-	UPROPERTY(Category=BlendSpace, VisibleAnywhere)
+	UPROPERTY()
 	struct FBlendParameter BlendParameters[3];
 
 	/** Input interpolation parameter for all 3 axis, for each axis input, decide how you'd like to interpolate input to*/
@@ -235,7 +235,7 @@ public:
 	bool bRotationBlendInMeshSpace;
 
 	/** Number of dimensions for this blend space (1 or 2) **/
-	UPROPERTY(Category=BlendSpace, VisibleAnywhere)
+	UPROPERTY()
 	int32 NumOfDimension;
 
 #if WITH_EDITORONLY_DATA
@@ -409,7 +409,7 @@ protected:
 public:
 	
 	/** Sample animation data **/
-	UPROPERTY(editfixedsize, Category=BlendSpace, VisibleAnywhere)
+	UPROPERTY()
 	TArray<struct FBlendSample> SampleData;
 
 	/** Grid samples, indexing scheme imposed by subclass **/
