Change landscape texture/mic outer back to package to fix slow PIE with LandscapeMountains

TTP#341135

diff --git a/Engine/Source/Editor/LandscapeEditor/Private/LandscapeEdMode.cpp b/Engine/Source/Editor/LandscapeEditor/Private/LandscapeEdMode.cpp
--- a/Engine/Source/Editor/LandscapeEditor/Private/LandscapeEdMode.cpp
+++ b/Engine/Source/Editor/LandscapeEditor/Private/LandscapeEdMode.cpp
@@ -95,7 +95,7 @@ void ALandscape::SplitHeightmap(ULandscapeComponent* Comp, bool bMoveToCurrentLe
 		LandscapeEdit.GetHeightDataFast(Comp->GetSectionBase().X, Comp->GetSectionBase().Y, Comp->GetSectionBase().X + Comp->ComponentSizeQuads, Comp->GetSectionBase().Y + Comp->ComponentSizeQuads, (uint16*)HeightData.GetTypedData(), 0, (uint16*)NormalData.GetTypedData());
 
 		// Construct the heightmap textures
-		UObject* TextureOuter = bMoveToCurrentLevel ? Comp->GetWorld()->GetCurrentLevel() : nullptr;
+		UObject* TextureOuter = bMoveToCurrentLevel ? Comp->GetWorld()->GetCurrentLevel()->GetOutermost() : nullptr;
 		HeightmapTexture = Comp->GetLandscapeProxy()->CreateLandscapeTexture(HeightmapSizeU, HeightmapSizeV, TEXTUREGROUP_Terrain_Heightmap, TSF_BGRA8, TextureOuter);
 
 		int32 MipSubsectionSizeQuads = Comp->SubsectionSizeQuads;
diff --git a/Engine/Source/Runtime/Engine/Private/Landscape/Landscape.cpp b/Engine/Source/Runtime/Engine/Private/Landscape/Landscape.cpp
--- a/Engine/Source/Runtime/Engine/Private/Landscape/Landscape.cpp
+++ b/Engine/Source/Runtime/Engine/Private/Landscape/Landscape.cpp
@@ -425,30 +425,30 @@ void ULandscapeComponent::PostLoad()
 #if WITH_EDITOR
 	if (GIsEditor && !HasAnyFlags(RF_ClassDefaultObject))
 	{
-		if (GetLinkerUE4Version() < VER_UE4_MOVE_LANDSCAPE_MICS_AND_TEXTURES_WITHIN_LEVEL)
+		// Move the MICs and Textures back to the Package if they're currently in the level
 		{
 			ULevel* Level = GetLevel();
 			if (ensure(Level))
 			{
-				TArray<UObject*> ObjectsToMoveFromPackgeToLevel;
+				TArray<UObject*> ObjectsToMoveFromLevelToPackage;
 				for (UMaterialInstance* CurrentMIC = MaterialInstance; CurrentMIC; CurrentMIC = Cast<UMaterialInstance>(CurrentMIC->Parent))
 				{
-					ObjectsToMoveFromPackgeToLevel.Add(CurrentMIC);
+					ObjectsToMoveFromLevelToPackage.Add(CurrentMIC);
 				}
-				ObjectsToMoveFromPackgeToLevel.Add(HeightmapTexture);
+				ObjectsToMoveFromLevelToPackage.Add(HeightmapTexture);
 				for (auto* Tex : WeightmapTextures)
 				{
-					ObjectsToMoveFromPackgeToLevel.Add(Tex);
+					ObjectsToMoveFromLevelToPackage.Add(Tex);
 				}
-				ObjectsToMoveFromPackgeToLevel.Add(XYOffsetmapTexture);
+				ObjectsToMoveFromLevelToPackage.Add(XYOffsetmapTexture);
 
 				UPackage* MyPackage = GetOutermost();
-				for (auto* Obj : ObjectsToMoveFromPackgeToLevel)
+				for (auto* Obj : ObjectsToMoveFromLevelToPackage)
 				{
-					if (Obj && Obj->GetOuter() == MyPackage)
+					if (Obj && Obj->GetOuter() == Level)
 					{
 						Obj->ClearFlags(RF_Public);
-						Obj->Rename(NULL, Level, REN_DoNotDirty | REN_DontCreateRedirectors | REN_ForceNoResetLoaders | REN_NonTransactional);
+						Obj->Rename(NULL, MyPackage, REN_DoNotDirty | REN_DontCreateRedirectors | REN_ForceNoResetLoaders | REN_NonTransactional);
 					}
 				}
 			}
@@ -1229,7 +1229,7 @@ void ALandscapeProxy::PostLoad()
 		if( Comp )
 		{
 			ComponentSizeQuads = Comp->ComponentSizeQuads;
-			SubsectionSizeQuads = Comp->SubsectionSizeQuads;	
+			SubsectionSizeQuads = Comp->SubsectionSizeQuads;
 			NumSubsections = Comp->NumSubsections;
 		}
 	}
diff --git a/Engine/Source/Runtime/Engine/Private/Landscape/LandscapeEdit.cpp b/Engine/Source/Runtime/Engine/Private/Landscape/LandscapeEdit.cpp
--- a/Engine/Source/Runtime/Engine/Private/Landscape/LandscapeEdit.cpp
+++ b/Engine/Source/Runtime/Engine/Private/Landscape/LandscapeEdit.cpp
@@ -94,9 +94,9 @@ ULandscapeMaterialInstanceConstant* ALandscapeProxy::GetLayerThumbnailMIC(UMater
 	{
 		FlushRenderingCommands();
 		UObject* MICOuter = GetTransientPackage();
-		if ( Proxy && Proxy->GetLevel() )
+		if (Proxy)
 		{
-			MICOuter = Proxy->GetLevel();
+			MICOuter = Proxy->GetOutermost();
 		}
 		CombinationMaterialInstance = ConstructObject<ULandscapeMaterialInstanceConstant>(ULandscapeMaterialInstanceConstant::StaticClass(), MICOuter);
 		if (Proxy)
@@ -170,7 +170,7 @@ UMaterialInstanceConstant* ULandscapeComponent::GetCombinationMaterial(bool bMob
 		{
 			FlushRenderingCommands();
 
-			CombinationMaterialInstance = ConstructObject<ULandscapeMaterialInstanceConstant>(ULandscapeMaterialInstanceConstant::StaticClass(), GetLevel());
+			CombinationMaterialInstance = ConstructObject<ULandscapeMaterialInstanceConstant>(ULandscapeMaterialInstanceConstant::StaticClass(), GetOutermost());
 			UE_LOG(LogLandscape, Log, TEXT("Looking for key %s, making new combination %s"), *LayerKey, *CombinationMaterialInstance->GetName());
 			Proxy->MaterialInstanceConstantMap.Add(*LayerKey,CombinationMaterialInstance);
 			CombinationMaterialInstance->SetParentEditorOnly(LandscapeMaterial);
@@ -225,7 +225,7 @@ void ULandscapeComponent::UpdateMaterialInstances()
 		// Create the instance for this component, that will use the layer combination instance.
 		if( MaterialInstance == NULL || GetOutermost() != MaterialInstance->GetOutermost() )
 		{
-			MaterialInstance = ConstructObject<ULandscapeMaterialInstanceConstant>(ULandscapeMaterialInstanceConstant::StaticClass(), GetLevel());
+			MaterialInstance = ConstructObject<ULandscapeMaterialInstanceConstant>(ULandscapeMaterialInstanceConstant::StaticClass(), GetOutermost());
 		}
 
 		// For undo
@@ -4784,7 +4784,7 @@ UMaterialInstance* ULandscapeComponent::GeneratePlatformPixelData(TArray<UTextur
 
 	if (!bIsCooking)
 	{
-		UMaterialInstanceDynamic* MobileMaterialInstance = UMaterialInstanceDynamic::Create(MaterialInstance, GetLevel()); 
+		UMaterialInstanceDynamic* MobileMaterialInstance = UMaterialInstanceDynamic::Create(MaterialInstance, GetOutermost()); 
 		// Set the layer mask
 		int32 CurrentIdx = 0;
 		for( int32 AllocIdx=0;AllocIdx<WeightmapLayerAllocations.Num();AllocIdx++ )
@@ -4805,7 +4805,7 @@ UMaterialInstance* ULandscapeComponent::GeneratePlatformPixelData(TArray<UTextur
 	else // for cooking
 	{
 		UMaterialInstanceConstant* CombinationMaterialInstance = GetCombinationMaterial(true); 
-		UMaterialInstanceConstant* MobileMaterialInstance = ConstructObject<ULandscapeMaterialInstanceConstant>(ULandscapeMaterialInstanceConstant::StaticClass(), GetLevel());
+		UMaterialInstanceConstant* MobileMaterialInstance = ConstructObject<ULandscapeMaterialInstanceConstant>(ULandscapeMaterialInstanceConstant::StaticClass(), GetOutermost());
 
 		MobileMaterialInstance->SetParentEditorOnly(CombinationMaterialInstance);
 
@@ -5002,7 +5002,7 @@ void ULandscapeComponent::GeneratePlatformVertexData()
 
 UTexture2D* ALandscapeProxy::CreateLandscapeTexture(int32 InSizeX, int32 InSizeY, TextureGroup InLODGroup, ETextureSourceFormat InFormat, UObject* OptionalOverrideOuter) const
 {
-	UObject* TexOuter = OptionalOverrideOuter ? OptionalOverrideOuter : GetLevel();
+	UObject* TexOuter = OptionalOverrideOuter ? OptionalOverrideOuter : GetOutermost();
 	UTexture2D* NewTexture = ConstructObject<UTexture2D>(UTexture2D::StaticClass(), TexOuter);
 	NewTexture->Source.Init2DWithMipChain(InSizeX, InSizeY, InFormat);
 	NewTexture->SRGB = false;
