Workaround for TTP #341211

diff --git a/Engine/Source/Runtime/Engine/Private/Materials/MaterialShared.cpp b/Engine/Source/Runtime/Engine/Private/Materials/MaterialShared.cpp
--- a/Engine/Source/Runtime/Engine/Private/Materials/MaterialShared.cpp
+++ b/Engine/Source/Runtime/Engine/Private/Materials/MaterialShared.cpp
@@ -393,12 +393,17 @@ bool FMaterial::NeedsGBuffer() const
 {
 	checkSlow(IsInRenderingThread());
 
+	if (IsOpenGLPlatform(GRHIShaderPlatform)) // @todo: TTP #341211
+	{
+		return true;
+	}
+
 	if (RenderingThreadShaderMap)
 	{
 		return RenderingThreadShaderMap->NeedsGBuffer();
 	}
 
-	return 0;
+	return false;
 }
 
 
