Fixed couple cases where code assumes that level always has at least 2 actors

diff --git a/Engine/Source/Runtime/Engine/Private/Level.cpp b/Engine/Source/Runtime/Engine/Private/Level.cpp
--- a/Engine/Source/Runtime/Engine/Private/Level.cpp
+++ b/Engine/Source/Runtime/Engine/Private/Level.cpp
@@ -480,8 +480,8 @@ void ULevel::SortActorList()
 	TArray<AActor*> NewActors;
 	NewActors.Reserve(Actors.Num());
 
-	// The world info and default brush have fixed actor indices.
-	NewActors.Add(Actors[StartIndex++]);
+	// The WorldSettings has fixed actor index.
+	check(Actors.Num() && Actors[StartIndex] == GetWorldSettings());
 	NewActors.Add(Actors[StartIndex++]);
 
 	// Static not net relevant actors.
@@ -1527,7 +1527,7 @@ ABrush* ULevel::GetDefaultBrush() const
 
 AWorldSettings* ULevel::GetWorldSettings() const
 {
-	checkf( Actors.Num() >= 2, *GetPathName() );
+	checkf( Actors.Num() >= 1, *GetPathName() );
 	AWorldSettings* WorldSettings = Cast<AWorldSettings>( Actors[0] );
 	checkf( WorldSettings != NULL, *GetPathName() );
 	return WorldSettings;
