Added code to reset OwnedComponents after Undo

#ttp 341131: CRASH: EDITOR: Editor crashes if the user attempts to Cut (Ctrl+X) then Undo (Ctrl+Z)

diff --git a/Engine/Source/Runtime/Engine/Private/ActorEditor.cpp b/Engine/Source/Runtime/Engine/Private/ActorEditor.cpp
--- a/Engine/Source/Runtime/Engine/Private/ActorEditor.cpp
+++ b/Engine/Source/Runtime/Engine/Private/ActorEditor.cpp
@@ -294,6 +294,12 @@ void AActor::PostEditUndo()
 		GetLevel()->LevelBoundsActor.Get()->OnLevelBoundsDirtied();
 	}
 
+	// Restore OwnedComponents array
+	if (!IsPendingKill())
+	{
+		ResetOwnedComponents();
+	}
+
 	Super::PostEditUndo();
 }
 
