moved navigation octree updates to scene component's transform update

diff --git a/Engine/Source/Runtime/Engine/Classes/AI/Navigation/NavigationSystem.h b/Engine/Source/Runtime/Engine/Classes/AI/Navigation/NavigationSystem.h
--- a/Engine/Source/Runtime/Engine/Classes/AI/Navigation/NavigationSystem.h
+++ b/Engine/Source/Runtime/Engine/Classes/AI/Navigation/NavigationSystem.h
@@ -505,9 +505,9 @@ public:
 	// Fetch the array of all nav-agent properties.
 	void GetNavAgentPropertiesArray(TArray<FNavAgentProperties>& OutNavAgentProperties) const;
 
-	static FORCEINLINE bool ShouldUpdateNavOctreeOnPrimitiveComponentChange()
+	static FORCEINLINE bool ShouldUpdateNavOctreeOnComponentChange()
 	{
-		return bUpdateNavOctreeOnPrimitiveComponentChange;
+		return bUpdateNavOctreeOnComponentChange;
 	}
 
 	/** 
@@ -577,7 +577,7 @@ protected:
 	/** whether seamless navigation building is enabled */
 	static bool bNavigationAutoUpdateEnabled;
 
-	static bool bUpdateNavOctreeOnPrimitiveComponentChange;
+	static bool bUpdateNavOctreeOnComponentChange;
 
 	static TArray<UClass*> PendingNavAreaRegistration;
 	static TArray<const UClass*> NavAreaClasses;
diff --git a/Engine/Source/Runtime/Engine/Classes/Components/PrimitiveComponent.h b/Engine/Source/Runtime/Engine/Classes/Components/PrimitiveComponent.h
--- a/Engine/Source/Runtime/Engine/Classes/Components/PrimitiveComponent.h
+++ b/Engine/Source/Runtime/Engine/Classes/Components/PrimitiveComponent.h
@@ -1115,7 +1115,6 @@ public:
 	// End UObject interface.
 
 	//Begin USceneComponent Interface
-	virtual void SetRelativeScale3D(FVector NewScale3D) override final;
 	virtual bool MoveComponent(const FVector& Delta, const FRotator& NewRotation, bool bSweep, FHitResult* OutHit = NULL, EMoveComponentFlags MoveFlags = MOVECOMP_NoFlags) override;
 	virtual bool IsWorldGeometry() const override;
 	virtual ECollisionEnabled::Type GetCollisionEnabled() const override;
@@ -1123,6 +1122,7 @@ public:
 	virtual ECollisionChannel GetCollisionObjectType() const override;
 	virtual const FCollisionResponseContainer & GetCollisionResponseToChannels() const override;
 	virtual FVector GetComponentVelocity() const override;
+	virtual bool IsNavigationRelevant(bool bSkipCollisionEnabledCheck = false) const override;
 	//End USceneComponent Interface
 
 	/**
@@ -1419,23 +1419,12 @@ public:
 	 */
 	virtual bool CanCharacterStepUp(class APawn* Pawn) const;
 
-	/** 
-	 *	Indicates whether this actor is to be considered by navigation system a valid actor
-	 *	@param bDoChannelCheckOnly allow caller to check whether this component affects/would affect
-	 *	navigation, regardless of its current BodyInstance collision value
-	 *	@note this function will return "true" also if HasCustomNavigableGeometry == EHasCustomNavigableGeometry::EvenIfNonCollidable
-	 */
-	bool IsNavigationRelevant(bool bSkipCollisionEnabledCheck = false) const;
-
 	/** Can this component potentially influence navigation */
 	bool CanEverAffectNavigation() const
 	{
 		return bCanEverAffectNavigation;
 	}
 
-	/** check if navigation system needs to be updated after changing transform of this component (include attached components as well) */
-	bool ShouldUpdateNavigationOnTransformChange() const;
-
 	/** turn off navigation relevance, must be called before component is registered! */
 	void DisableNavigationRelevance();
 
diff --git a/Engine/Source/Runtime/Engine/Classes/Components/SceneComponent.h b/Engine/Source/Runtime/Engine/Classes/Components/SceneComponent.h
--- a/Engine/Source/Runtime/Engine/Classes/Components/SceneComponent.h
+++ b/Engine/Source/Runtime/Engine/Classes/Components/SceneComponent.h
@@ -704,6 +704,9 @@ public:
 	/** Get the extent used when placing this component in the editor, used for 'pulling back' hit. */
 	virtual FBoxSphereBounds GetPlacementExtent() const;
 
+	/** Is component (not including attachments) relevant for navigation updates? */
+	virtual bool IsNavigationRelevant(bool bSkipCollisionEnabledCheck = false) const;
+
 protected:
 	/**
 	 * Called after a child scene component is attached to this component.
@@ -716,6 +719,12 @@ protected:
 	 * Note: Do not change the attachment state of the child during this call.
 	 */
 	virtual void OnChildDetached(USceneComponent* ChildComponent) {}
+
+	/** Called after changing transform, tries to update navigation octree */
+	void UpdateNavigationData();
+
+	/** Check if given component or any of its attachments are relevant for navigation updates */
+	bool CheckNavigationRelevancy(USceneComponent* TestComponent);
 };
 
 
diff --git a/Engine/Source/Runtime/Engine/Private/AI/Navigation/NavigationSystem.cpp b/Engine/Source/Runtime/Engine/Private/AI/Navigation/NavigationSystem.cpp
--- a/Engine/Source/Runtime/Engine/Private/AI/Navigation/NavigationSystem.cpp
+++ b/Engine/Source/Runtime/Engine/Private/AI/Navigation/NavigationSystem.cpp
@@ -171,7 +171,7 @@ FNavigationSystemExec UNavigationSystem::ExecHandler;
 /** called after navigation influencing event takes place*/
 UNavigationSystem::FOnNavigationDirty UNavigationSystem::NavigationDirtyEvent;
 
-bool UNavigationSystem::bUpdateNavOctreeOnPrimitiveComponentChange = true;
+bool UNavigationSystem::bUpdateNavOctreeOnComponentChange = true;
 //----------------------------------------------------------------------//
 // life cycle stuff                                                                
 //----------------------------------------------------------------------//
diff --git a/Engine/Source/Runtime/Engine/Private/Actor.cpp b/Engine/Source/Runtime/Engine/Private/Actor.cpp
--- a/Engine/Source/Runtime/Engine/Private/Actor.cpp
+++ b/Engine/Source/Runtime/Engine/Private/Actor.cpp
@@ -410,13 +410,6 @@ bool AActor::TeleportTo( const FVector& DestLocation, const FRotator& DestRotati
 		GetRootComponent()->SetWorldLocationAndRotation(NewLocation, DestRotation);
 	}
 
-	if ( !bIsATest && bTeleportSucceeded )
-	{
-		if (IsNavigationRelevant() && GetWorld() && GetWorld()->GetNavigationSystem())
-		{
-			GetWorld()->GetNavigationSystem()->UpdateNavOctree(this);
-		}
-	}
 	return bTeleportSucceeded; 
 }
 
@@ -2120,10 +2113,10 @@ bool AActor::UpdateNavigationRelevancy()
 
 	for (int32 CompIdx = 0; CompIdx < Components.Num() && bNewRelevancy == false; ++CompIdx)
 	{
-		UPrimitiveComponent* Comp = Cast<UPrimitiveComponent>(Components[CompIdx]);
+		USceneComponent* Comp = Cast<USceneComponent>(Components[CompIdx]);
 		if (Comp)
 		{
-			bNewRelevancy = Comp->CanEverAffectNavigation() && Comp->IsRegistered() && Comp->IsNavigationRelevant();
+			bNewRelevancy = Comp->IsRegistered() && Comp->IsNavigationRelevant();
 		}
 		else
 		{
diff --git a/Engine/Source/Runtime/Engine/Private/PrimitiveComponent.cpp b/Engine/Source/Runtime/Engine/Private/PrimitiveComponent.cpp
--- a/Engine/Source/Runtime/Engine/Private/PrimitiveComponent.cpp
+++ b/Engine/Source/Runtime/Engine/Private/PrimitiveComponent.cpp
@@ -1217,24 +1217,6 @@ FCollisionShape UPrimitiveComponent::GetCollisionShape(float Inflation) const
 	return FCollisionShape::MakeBox(Bounds.BoxExtent + Inflation);
 }
 
-void UPrimitiveComponent::SetRelativeScale3D(FVector NewScale3D)
-{
-	const FVector OldScale3D = RelativeScale3D;
-	Super::SetRelativeScale3D(NewScale3D);
-
-	AActor* Actor = GetOwner();
-	if (OldScale3D != RelativeScale3D && UNavigationSystem::ShouldUpdateNavOctreeOnPrimitiveComponentChange() &&
-		Actor && IsRegistered() && Actor->IsNavigationRelevant() && 
-		ShouldUpdateNavigationOnTransformChange() &&
-		World != NULL && World->IsGameWorld() && World->GetNetMode() < ENetMode::NM_Client)
-	{
-		if (UNavigationSystem* NavSys = World->GetNavigationSystem())
-		{
-			NavSys->UpdateNavOctree(Actor);
-		}
-	}
-}
-
 bool UPrimitiveComponent::MoveComponent( const FVector& Delta, const FRotator& NewRotation, bool bSweep, FHitResult* OutHit, EMoveComponentFlags MoveFlags)
 {
 	SCOPE_CYCLE_COUNTER(STAT_MoveComponentTime);
@@ -1534,18 +1516,6 @@ bool UPrimitiveComponent::MoveComponent( const FVector& Delta, const FRotator& N
 		}
 	}
 
-	// Update navigation system
-	if (bMoved && UNavigationSystem::ShouldUpdateNavOctreeOnPrimitiveComponentChange() &&
-		Actor && IsRegistered() && Actor->IsNavigationRelevant() &&
-		ShouldUpdateNavigationOnTransformChange() &&
-		World != NULL && World->IsGameWorld() && World->GetNetMode() < ENetMode::NM_Client)
-	{
-		if (UNavigationSystem* NavSys = World->GetNavigationSystem())
-		{
-			NavSys->UpdateNavOctree(Actor);
-		}
-	}
-
 #if defined(PERF_SHOW_MOVECOMPONENT_TAKING_LONG_TIME) || LOOKING_FOR_PERF_ISSUES
 	UNCLOCK_CYCLES(MoveCompTakingLongTime);
 	const float MSec = FPlatformTime::ToMilliseconds(MoveCompTakingLongTime);
@@ -1587,6 +1557,11 @@ void UPrimitiveComponent::DispatchBlockingHit(AActor& Owner, FHitResult const& B
 
 bool UPrimitiveComponent::IsNavigationRelevant(bool bSkipCollisionEnabledCheck) const 
 { 
+	if (!CanEverAffectNavigation())
+	{
+		return false;
+	}
+
 	if (HasCustomNavigableGeometry() >= EHasCustomNavigableGeometry::EvenIfNotCollidable)
 	{
 		return true;
@@ -1605,21 +1580,6 @@ void UPrimitiveComponent::DisableNavigationRelevance()
 	bCanEverAffectNavigation = false;
 }
 
-bool UPrimitiveComponent::ShouldUpdateNavigationOnTransformChange() const
-{
-	bool bResult = CanEverAffectNavigation() && IsNavigationRelevant();
-	for (int32 Idx = 0; Idx < AttachChildren.Num() && !bResult; Idx++)
-	{
-		const UPrimitiveComponent* PrimComp = Cast<const UPrimitiveComponent>(AttachChildren[Idx]);
-		if (PrimComp)
-		{
-			bResult = PrimComp->ShouldUpdateNavigationOnTransformChange();
-		}
-	}
-
-	return bResult;
-}
-
 
 //////////////////////////////////////////////////////////////////////////
 // COLLISION
diff --git a/Engine/Source/Runtime/Engine/Private/SceneComponent.cpp b/Engine/Source/Runtime/Engine/Private/SceneComponent.cpp
--- a/Engine/Source/Runtime/Engine/Private/SceneComponent.cpp
+++ b/Engine/Source/Runtime/Engine/Private/SceneComponent.cpp
@@ -132,6 +132,9 @@ void USceneComponent::PropagateTransformUpdate(bool bTransformChanged, bool bSki
 		
 		// Now go and update children
 		UpdateChildTransforms();
+
+		// Refresh navigation
+		UpdateNavigationData();
 	}
 	else
 	{
@@ -1697,4 +1700,41 @@ const int32 USceneComponent::GetNumUncachedStaticLightingInteractions() const
 }
 #endif
 
+void USceneComponent::UpdateNavigationData()
+{
+	AActor* MyActor = GetOwner();
+
+	if (UNavigationSystem::ShouldUpdateNavOctreeOnComponentChange() &&
+		IsRegistered() && MyActor && MyActor->IsNavigationRelevant() &&
+		World && World->IsGameWorld() && World->GetNetMode() < ENetMode::NM_Client)
+	{
+		const bool bShouldUpdate = CheckNavigationRelevancy(this);
+		UNavigationSystem* NavSys = UNavigationSystem::GetCurrent(World);
+		
+		if (bShouldUpdate && NavSys)
+		{
+			NavSys->UpdateNavOctree(MyActor);
+		}
+	}
+}
+
+bool USceneComponent::CheckNavigationRelevancy(USceneComponent* TestComponent)
+{
+	bool bResult = TestComponent->IsNavigationRelevant();
+	for (int32 Idx = 0; Idx < AttachChildren.Num() && !bResult; Idx++)
+	{
+		if (AttachChildren[Idx])
+		{
+			bResult = CheckNavigationRelevancy(AttachChildren[Idx]);
+		}
+	}
+
+	return bResult;
+}
+
+bool USceneComponent::IsNavigationRelevant(bool bSkipCollisionEnabledCheck) const
+{
+	return false;
+}
+
 #undef LOCTEXT_NAMESPACE
