Level Loading - Fix an ordering issue with rerunning blueprint construction scripts which was breaking attachments on load, also fix up some whitespace issues.

TTP# 339716

diff --git a/Engine/Source/Runtime/Engine/Private/Level.cpp b/Engine/Source/Runtime/Engine/Private/Level.cpp
--- a/Engine/Source/Runtime/Engine/Private/Level.cpp
+++ b/Engine/Source/Runtime/Engine/Private/Level.cpp
@@ -753,43 +753,52 @@ void ULevel::IncrementalUpdateComponents(int32 NumActorsToUpdate, bool bRerunCon
 		UpdateModelComponents();
 	}
 
+	TArray<AActor*> RerunConstructionScriptActors;
 	for( int32 i=0; i < NumActorsToUpdate && CurrentActorIndexForUpdateComponents < Actors.Num(); i++ )
 	{
 		AActor* Actor = Actors[CurrentActorIndexForUpdateComponents++];
 		if( Actor )
 		{
 #if PERF_TRACK_DETAILED_ASYNC_STATS
-				double Start = FPlatformTime::Seconds();
+			double Start = FPlatformTime::Seconds();
 #endif
 
-				Actor->ReregisterAllComponents();
+			Actor->ReregisterAllComponents();
 
-				// Rerun the construction script (if the actor is blueprint based and this is being run in the editor)
-				// @TODO: Only do this if blueprint has changed!
-				if (bRerunConstructionScripts && !IsTemplate() && !GIsUCCMakeStandaloneHeaderGenerator)
-				{
-					Actor->RerunConstructionScripts();
-				}
+			// Cache the actor to rerun the construction script (if the actor is blueprint based and this is being run in the editor)
+			// @TODO: Only do this if blueprint has changed!
+			if (bRerunConstructionScripts && !IsTemplate() && !GIsUCCMakeStandaloneHeaderGenerator)
+			{
+				RerunConstructionScriptActors.Add(Actor);
+			}
 
 #if PERF_TRACK_DETAILED_ASYNC_STATS
-				// Add how long this took to class->time map
-				double Time = FPlatformTime::Seconds() - Start;
-				UClass* ActorClass = Actor->GetClass();
-				FMapTimeEntry* CurrentEntry = UpdateComponentsTimePerActorClass.Find(ActorClass);
-				// Is an existing entry - add to it
-				if(CurrentEntry)
-				{
-					CurrentEntry->Time += Time;
-					CurrentEntry->ObjCount += 1;
-				}
-				// Make a new entry for this class
-				else
-				{
-					UpdateComponentsTimePerActorClass.Add(ActorClass, FMapTimeEntry(ActorClass, 1, Time));
-				}
-#endif
+			// Add how long this took to class->time map
+			double Time = FPlatformTime::Seconds() - Start;
+			UClass* ActorClass = Actor->GetClass();
+			FMapTimeEntry* CurrentEntry = UpdateComponentsTimePerActorClass.Find(ActorClass);
+			// Is an existing entry - add to it
+			if(CurrentEntry)
+			{
+				CurrentEntry->Time += Time;
+				CurrentEntry->ObjCount += 1;
 			}
+			// Make a new entry for this class
+			else
+			{
+				UpdateComponentsTimePerActorClass.Add(ActorClass, FMapTimeEntry(ActorClass, 1, Time));
+			}
+#endif
 		}
+	}
+
+	// Don't rerun construction scripts until after all actors' components have been registered.  This
+	// is necessary because child attachment lists are populated during registration, and running construction
+	// scripts requires that the attachments are correctly initialized.
+	for (auto RerunConstructionScriptActor : RerunConstructionScriptActors)
+	{
+		RerunConstructionScriptActor->RerunConstructionScripts();
+	}
 
 	// See whether we are done.
 	if( CurrentActorIndexForUpdateComponents == Actors.Num() )
