Fix sometimes stale GPU particle rendering resource.  Fixes PS4 artifacts.

diff --git a/Engine/Source/Runtime/Engine/Private/Particles/ParticleGpuSimulation.cpp b/Engine/Source/Runtime/Engine/Private/Particles/ParticleGpuSimulation.cpp
--- a/Engine/Source/Runtime/Engine/Private/Particles/ParticleGpuSimulation.cpp
+++ b/Engine/Source/Runtime/Engine/Private/Particles/ParticleGpuSimulation.cpp
@@ -2516,42 +2516,7 @@ struct FGPUSpriteDynamicEmitterData : FDynamicEmitterDataBase
 	 * Called to create render thread resources.
 	 */
 	virtual void CreateRenderThreadResources(const FParticleSystemSceneProxy* InOwnerProxy)
-	{
-		if (RHISupportsGPUParticles(InOwnerProxy->GetScene()->GetFeatureLevel()))
-		{
-			Simulation->PendingDeltaSeconds = PendingDeltaSeconds;
-
-			// Create the per-frame uniform buffer.
-			Simulation->PerFrameSimulationUniformBuffer =
-				FParticlePerFrameSimulationBufferRef::CreateUniformBufferImmediate(PerFrameSimulationParameters, UniformBuffer_SingleFrame);
-
-			// Create per-emitter uniform buffer for dynamic parameters
-			DynamicUniformBuffer = FGPUSpriteEmitterDynamicUniformBufferRef::CreateUniformBufferImmediate(EmitterDynamicParameters, UniformBuffer_SingleFrame);
-
-			// Local vector field parameters.
-			Simulation->LocalVectorField.Intensity = LocalVectorFieldIntensity;
-			Simulation->LocalVectorField.Tightness = LocalVectorFieldTightness;
-			Simulation->LocalVectorField.bTileX = bLocalVectorFieldTileX;
-			Simulation->LocalVectorField.bTileY = bLocalVectorFieldTileY;
-			Simulation->LocalVectorField.bTileZ = bLocalVectorFieldTileZ;
-			if (Simulation->LocalVectorField.Resource)
-			{
-				Simulation->LocalVectorField.UpdateTransforms(LocalVectorFieldToWorld);
-			}
-
-			// Update world bounds.
-			Simulation->Bounds = SimulationBounds;
-
-			// Transfer ownership of new data.
-			if (NewParticles.Num())
-			{
-				Exchange(Simulation->NewParticles, NewParticles);
-			}
-			if (TilesToClear.Num())
-			{
-				Exchange(Simulation->TilesToClear, TilesToClear);
-			}
-		}
+	{		
 	}
 
 	/**
@@ -2581,6 +2546,42 @@ struct FGPUSpriteDynamicEmitterData : FDynamicEmitterDataBase
 			if (Simulation->SimulationIndex != INDEX_NONE
 				&& Simulation->VertexBuffer.ParticleCount > 0)
 			{
+				// Create view agnostic render data.  Do here rather than in CreateRenderThreadResources because in some cases Render can be called before CreateRenderThreadResources
+				{
+					Simulation->PendingDeltaSeconds = PendingDeltaSeconds;
+
+					// Create the per-frame uniform buffer.
+					Simulation->PerFrameSimulationUniformBuffer =
+						FParticlePerFrameSimulationBufferRef::CreateUniformBufferImmediate(PerFrameSimulationParameters, UniformBuffer_SingleFrame);
+
+					// Create per-emitter uniform buffer for dynamic parameters
+					DynamicUniformBuffer = FGPUSpriteEmitterDynamicUniformBufferRef::CreateUniformBufferImmediate(EmitterDynamicParameters, UniformBuffer_SingleFrame);
+
+					// Local vector field parameters.
+					Simulation->LocalVectorField.Intensity = LocalVectorFieldIntensity;
+					Simulation->LocalVectorField.Tightness = LocalVectorFieldTightness;
+					Simulation->LocalVectorField.bTileX = bLocalVectorFieldTileX;
+					Simulation->LocalVectorField.bTileY = bLocalVectorFieldTileY;
+					Simulation->LocalVectorField.bTileZ = bLocalVectorFieldTileZ;
+					if (Simulation->LocalVectorField.Resource)
+					{
+						Simulation->LocalVectorField.UpdateTransforms(LocalVectorFieldToWorld);
+					}
+
+					// Update world bounds.
+					Simulation->Bounds = SimulationBounds;
+
+					// Transfer ownership of new data.
+					if (NewParticles.Num())
+					{
+						Exchange(Simulation->NewParticles, NewParticles);
+					}
+					if (TilesToClear.Num())
+					{
+						Exchange(Simulation->TilesToClear, TilesToClear);
+					}
+				}
+
 				EBlendMode BlendMode = Material->GetRenderProxy(false)->GetMaterial(FeatureLevel)->GetBlendMode();
 				const bool bOpaque = (BlendMode == BLEND_Opaque || BlendMode == BLEND_Masked);
 
