Ensure that rendering resources have been released when a level is deleted

diff --git a/Engine/Source/Runtime/Engine/Private/Level.cpp b/Engine/Source/Runtime/Engine/Private/Level.cpp
--- a/Engine/Source/Runtime/Engine/Private/Level.cpp
+++ b/Engine/Source/Runtime/Engine/Private/Level.cpp
@@ -696,6 +696,8 @@ bool ULevel::IsReadyForFinishDestroy()
 
 void ULevel::FinishDestroy()
 {
+	ReleaseRenderingResources();
+
 	delete PrecomputedLightVolume;
 	PrecomputedLightVolume = NULL;
 
