Making it more clear that for schema actions, the search title memeber does not have to be initialized (to save of perf with the new blueprint menu refactor).

diff --git a/Engine/Source/Runtime/Engine/Classes/EdGraph/EdGraphSchema.h b/Engine/Source/Runtime/Engine/Classes/EdGraph/EdGraphSchema.h
--- a/Engine/Source/Runtime/Engine/Classes/EdGraph/EdGraphSchema.h
+++ b/Engine/Source/Runtime/Engine/Classes/EdGraph/EdGraphSchema.h
@@ -80,9 +80,9 @@ struct ENGINE_API FEdGraphSchemaAction
 	UPROPERTY()
 	int32 SectionID;
 
-	/** Search title for the action */
+	/** Search title for the action (doesn't have to be set when instantiated, will be constructed by GetSearchTitle() if left empty)*/
 	UPROPERTY()
-	FText SearchTitle;
+	FText CachedSearchTitle;
 
 	FEdGraphSchemaAction() 
 		: Grouping(0)
@@ -125,17 +125,17 @@ struct ENGINE_API FEdGraphSchemaAction
 	/** Retrieves the full searchable title for this action */
 	FText GetSearchTitle()
 	{
-		if(SearchTitle.IsEmpty())
+		if(CachedSearchTitle.IsEmpty())
 		{
 			if(const FString* SourceString = FTextInspector::GetSourceString(MenuDescription))
 			{
 				FFormatNamedArguments Args;
 				Args.Add(TEXT("LocalizedTitle"), MenuDescription);
 				Args.Add(TEXT("SourceTitle"), FText::FromString(*SourceString));
-				SearchTitle = FText::Format(FText::FromString("{LocalizedTitle} {SourceTitle}"), Args);
+				CachedSearchTitle = FText::Format(FText::FromString("{LocalizedTitle} {SourceTitle}"), Args);
 			}
 		}
-		return SearchTitle;
+		return CachedSearchTitle;
 	}
 
 	// GC.
