Don't allow the client to set the bHasMustBeMappedGUIDs flag

diff --git a/Engine/Source/Runtime/Engine/Private/DataChannel.cpp b/Engine/Source/Runtime/Engine/Private/DataChannel.cpp
--- a/Engine/Source/Runtime/Engine/Private/DataChannel.cpp
+++ b/Engine/Source/Runtime/Engine/Private/DataChannel.cpp
@@ -1521,6 +1521,13 @@ void UActorChannel::ReceivedBunch( FInBunch & Bunch )
 
 	if ( Bunch.bHasMustBeMappedGUIDs )
 	{
+		if ( Connection->Driver->IsServer() )
+		{
+			UE_LOG( LogNetTraffic, Error, TEXT( "UActorChannel::ReceivedBunch: Client attempted to set bHasMustBeMappedGUIDs. Actor: %s" ), Actor != NULL ? *Actor->GetName() : TEXT( "NULL" ) );
+			Bunch.SetError();
+			return;
+		}
+
 		// If this bunch has any guids that must be mapped, we need to wait until they resolve before we can 
 		// process the rest of the stream on this channel
 		uint16 NumMustBeMappedGUIDs = 0;
