Fix selecting large amounts of actors freezing the edtitor

diff --git a/Engine/Source/Editor/LevelEditor/Private/LevelEditorActions.cpp b/Engine/Source/Editor/LevelEditor/Private/LevelEditorActions.cpp
--- a/Engine/Source/Editor/LevelEditor/Private/LevelEditorActions.cpp
+++ b/Engine/Source/Editor/LevelEditor/Private/LevelEditorActions.cpp
@@ -1420,15 +1420,19 @@ void FLevelEditorActionCallbacks::OnApplyMaterialToSurface()
 
 void FLevelEditorActionCallbacks::OnSelectAllLights()
 {
+	GEditor->GetSelectedActors()->BeginBatchSelectOperation();
 	// Select all light actors.
 	for( FActorIterator It(GetWorld()); It; ++It )
 	{
 		ALight* Light = Cast<ALight>(*It);
 		if( Light )
 		{
-			GUnrealEd->SelectActor( Light, true, true, false );
+			GUnrealEd->SelectActor( Light, true, false, false );
 		}
 	}
+
+	GEditor->GetSelectedActors()->EndBatchSelectOperation();
+
 }
 
 void FLevelEditorActionCallbacks::OnSelectStationaryLightsExceedingOverlap()
diff --git a/Engine/Source/Runtime/Engine/Private/Selection.cpp b/Engine/Source/Runtime/Engine/Private/Selection.cpp
--- a/Engine/Source/Runtime/Engine/Private/Selection.cpp
+++ b/Engine/Source/Runtime/Engine/Private/Selection.cpp
@@ -37,8 +37,11 @@ void USelection::Select(UObject* InObject)
 	SelectedObjects.AddUnique( InObject );
 	SelectedClasses.AddUnique( InObject->GetClass() );
 
-	// Call this after the item has been added from the selection set.
-	USelection::SelectObjectEvent.Broadcast( InObject );
+	if( !IsBatchSelecting() )
+	{
+		// Call this after the item has been added from the selection set.
+		USelection::SelectObjectEvent.Broadcast( InObject );
+	}
 
 	if ( bSelectionChanged )
 	{
@@ -57,8 +60,11 @@ void USelection::Deselect(UObject* InObject)
 	// Remove from selected list.
 	SelectedObjects.Remove( InObject );
 
-	// Call this after the item has been removed from the selection set.
-	USelection::SelectObjectEvent.Broadcast( InObject );
+	if( !IsBatchSelecting() )
+	{
+		// Call this after the item has been removed from the selection set.
+		USelection::SelectObjectEvent.Broadcast( InObject );
+	}
 
 	if ( bSelectionChanged )
 	{
