Fixed marquee select in orthographic view.
#ttp 333836 - Editor: Unable to marquee select certain actors using top orthographic view
#branch UE4
#proj Runtime.Engine, Editor.UnrealEd
#change Made some fairly hefty changes to FSeparatingAxisPointCheck (which was returning a false negative when testing degenerate axes, e.g. when a poly was axis-aligned): tidied the code up, made it a little clearer, allowed it to accept any convex poly as an input, along with the axis-aligned bounding box.  Improved documentation.
#change Corrected FLevelEditorViewportClient::ComponentIsTouchingSelectionBox so that BSP brushes are correctly selected (changed it from using FMath::PointBoxIntersection to FSeparatingAxisPointCheck, taking an arbitrary sided convex poly as its input).  This now correctly selects brushes when no points are contained within the marquee, but one of its polys is nonetheless within the box.

diff --git a/Engine/Source/Editor/UnrealEd/Private/LevelEditorViewport.cpp b/Engine/Source/Editor/UnrealEd/Private/LevelEditorViewport.cpp
--- a/Engine/Source/Editor/UnrealEd/Private/LevelEditorViewport.cpp
+++ b/Engine/Source/Editor/UnrealEd/Private/LevelEditorViewport.cpp
@@ -3037,31 +3037,43 @@ bool FLevelEditorViewportClient::ComponentIsTouchingSelectionBox(AActor* InActor
 		// Check the brush component
 		if( BrushComponent->Brush && BrushComponent->Brush->Polys && !bAlreadyProcessed )
 		{
-			for( int32 PolyIndex = 0 ; PolyIndex < BrushComponent->Brush->Polys->Element.Num() && !bAlreadyProcessed ; ++PolyIndex )
+			TArray<FVector> Vertices;
+
+			for (int32 PolyIndex = 0; PolyIndex < BrushComponent->Brush->Polys->Element.Num() && !bAlreadyProcessed; ++PolyIndex)
 			{
 				const FPoly& Poly = BrushComponent->Brush->Polys->Element[ PolyIndex ];
 
-				for( int32 VertexIndex = 0 ; VertexIndex < Poly.Vertices.Num() ; ++VertexIndex )
+				if (!bMustEncompassEntireComponent)
 				{
-					const FVector Location = InComponent->ComponentToWorld.TransformPosition( Poly.Vertices[VertexIndex] );
-					const bool bLocationIntersected = FMath::PointBoxIntersection( Location, InSelBBox );
+					Vertices.Empty(Poly.Vertices.Num());
+					for (const auto& Vertex : Poly.Vertices)
+					{
+						Vertices.Add(InComponent->ComponentToWorld.TransformPosition(Vertex));
+					}
 
-					// If the selection box doesn't have to encompass the entire component and a poly vertex has intersected with
-					// the selection box, this component is being touched by the selection box
-					if ( !bMustEncompassEntireComponent && bLocationIntersected )
+					FSeparatingAxisPointCheck PointCheck(Vertices, InSelBBox.GetCenter(), InSelBBox.GetExtent(), false);
+					if (PointCheck.bHit)
 					{
 						bResult = true;
 						bAlreadyProcessed = true;
 						break;
 					}
-
-					// If the selection box has to encompass the entire component and a poly vertex didn't intersect with the selection
-					// box, this component does not qualify
-					else if ( bMustEncompassEntireComponent && !bLocationIntersected )
+				}
+				else
+				{
+					for (int32 VertexIndex = 0; VertexIndex < Poly.Vertices.Num(); ++VertexIndex)
 					{
-						bResult = false;
-						bAlreadyProcessed = true;
-						break;
+						const FVector Location = InComponent->ComponentToWorld.TransformPosition(Poly.Vertices[VertexIndex]);
+						const bool bLocationIntersected = FMath::PointBoxIntersection(Location, InSelBBox);
+
+						// If the selection box has to encompass the entire component and a poly vertex didn't intersect with the selection
+						// box, this component does not qualify
+						if (!bLocationIntersected)
+						{
+							bResult = false;
+							bAlreadyProcessed = true;
+							break;
+						}
 					}
 				}
 			}
@@ -3085,6 +3097,8 @@ bool FLevelEditorViewportClient::ComponentIsTouchingSelectionBox(AActor* InActor
 			// Check if we are even inside it's bounding box, if we are not, there is no way we colliding via the more advanced checks we will do.
 			if(InSelBBox.Intersect(InComponent->Bounds.GetBox()))
 			{
+				TArray<FVector> Vertex;
+
 				FStaticMeshLODResources& LODModel = StaticMeshComponent->StaticMesh->RenderData->LODResources[0];
 				FIndexArrayView Indices = LODModel.IndexBuffer.GetArrayView();
 
@@ -3096,24 +3110,25 @@ bool FLevelEditorViewportClient::ComponentIsTouchingSelectionBox(AActor* InActor
 					// Iterate over each triangle.
 					for( int32 TriangleIndex = 0 ; TriangleIndex < (int32)Section.NumTriangles ; TriangleIndex++ )
 					{
-						FVector Vertex[3];
+						Vertex.Empty(3);
+
 						int32 FirstIndex = TriangleIndex * 3 + Section.FirstIndex;
 						for (int32 i = 0; i < 3; i++)
 						{
 							int32 VertexIndex = Indices[FirstIndex+i];
 							FVector LocalPosition = LODModel.PositionVertexBuffer.VertexPosition(VertexIndex);
-							Vertex[i] = StaticMeshComponent->ComponentToWorld.TransformPosition(LocalPosition);
+							Vertex.Emplace(StaticMeshComponent->ComponentToWorld.TransformPosition(LocalPosition));
 						}
 
 						// Check if the triangle is colliding with the bounding box.
-						FSeparatingAxisPointCheck ThePointCheck(Vertex[0], Vertex[1], Vertex[2], InSelBBox.GetCenter(), InSelBBox.GetExtent(), INT_MAX);
-						if( !bMustEncompassEntireComponent && ThePointCheck.Hit )
+						FSeparatingAxisPointCheck ThePointCheck(Vertex, InSelBBox.GetCenter(), InSelBBox.GetExtent(), false);
+						if( !bMustEncompassEntireComponent && ThePointCheck.bHit )
 						{
 							bResult = true;
 							bAlreadyProcessed = true;
 							break;
 						}
-						else if( bMustEncompassEntireComponent && !ThePointCheck.Hit )
+						else if( bMustEncompassEntireComponent && !ThePointCheck.bHit )
 						{
 							bResult = false;
 							bAlreadyProcessed = true;
diff --git a/Engine/Source/Runtime/Engine/Private/Collision/Collision.cpp b/Engine/Source/Runtime/Engine/Private/Collision/Collision.cpp
--- a/Engine/Source/Runtime/Engine/Private/Collision/Collision.cpp
+++ b/Engine/Source/Runtime/Engine/Private/Collision/Collision.cpp
@@ -5,6 +5,7 @@
 =============================================================================*/
 
 #include "EnginePrivate.h"
+#include "Collision.h"
 
 //////////////////////////////////////////////////////////////////////////
 // FHitResult
@@ -93,4 +94,187 @@ void FCollisionQueryParams::AddIgnoredActors(const TArray<TWeakObjectPtr<AActor>
 			IgnoreActors.Add(A->GetUniqueID());
 		}
 	}
-}
\ No newline at end of file
+}
+
+//////////////////////////////////////////////////////////////////////////
+// FSeparatingAxisPointCheck
+
+TArray<FVector> FSeparatingAxisPointCheck::TriangleVertices;
+
+bool FSeparatingAxisPointCheck::TestSeparatingAxisCommon(const FVector& Axis, float ProjectedPolyMin, float ProjectedPolyMax)
+{
+	const float ProjectedCenter = FVector::DotProduct(Axis, BoxCenter);
+	const float ProjectedExtent = FVector::DotProduct(Axis.GetAbs(), BoxExtent);
+	const float ProjectedBoxMin = ProjectedCenter - ProjectedExtent;
+	const float ProjectedBoxMax = ProjectedCenter + ProjectedExtent;
+
+	if (ProjectedPolyMin > ProjectedBoxMax || ProjectedPolyMax < ProjectedBoxMin)
+	{
+		return false;
+	}
+
+	if (bCalcLeastPenetration)
+	{
+		const float AxisMagnitudeSqr = Axis.SizeSquared();
+		if (AxisMagnitudeSqr > (SMALL_NUMBER * SMALL_NUMBER))
+		{
+			const float InvAxisMagnitude = FMath::InvSqrt(AxisMagnitudeSqr);
+			const float MinPenetrationDist = (ProjectedBoxMax - ProjectedPolyMin) * InvAxisMagnitude;
+			const float	MaxPenetrationDist = (ProjectedPolyMax - ProjectedBoxMin) * InvAxisMagnitude;
+
+			if (MinPenetrationDist < BestDist)
+			{
+				BestDist = MinPenetrationDist;
+				HitNormal = -Axis * InvAxisMagnitude;
+			}
+
+			if (MaxPenetrationDist < BestDist)
+			{
+				BestDist = MaxPenetrationDist;
+				HitNormal = Axis * InvAxisMagnitude;
+			}
+		}
+	}
+
+	return true;
+}
+
+bool FSeparatingAxisPointCheck::TestSeparatingAxisTriangle(const FVector& Axis)
+{
+	const float ProjectedV0 = FVector::DotProduct(Axis, PolyVertices[0]);
+	const float ProjectedV1 = FVector::DotProduct(Axis, PolyVertices[1]);
+	const float ProjectedV2 = FVector::DotProduct(Axis, PolyVertices[2]);
+	const float ProjectedTriMin = FMath::Min3(ProjectedV0, ProjectedV1, ProjectedV2);
+	const float ProjectedTriMax = FMath::Max3(ProjectedV0, ProjectedV1, ProjectedV2);
+
+	return TestSeparatingAxisCommon(Axis, ProjectedTriMin, ProjectedTriMax);
+}
+
+bool FSeparatingAxisPointCheck::TestSeparatingAxisGeneric(const FVector& Axis)
+{
+	float ProjectedPolyMin = TNumericLimits<float>::Max();
+	float ProjectedPolyMax = TNumericLimits<float>::Lowest();
+	for (const auto& Vertex : PolyVertices)
+	{
+		const float ProjectedVertex = FVector::DotProduct(Axis, Vertex);
+		ProjectedPolyMin = FMath::Min(ProjectedPolyMin, ProjectedVertex);
+		ProjectedPolyMax = FMath::Max(ProjectedPolyMax, ProjectedVertex);
+	}
+
+	return TestSeparatingAxisCommon(Axis, ProjectedPolyMin, ProjectedPolyMax);
+}
+
+bool FSeparatingAxisPointCheck::FindSeparatingAxisTriangle()
+{
+	check(PolyVertices.Num() == 3);
+	const FVector EdgeDir0 = PolyVertices[1] - PolyVertices[0];
+	const FVector EdgeDir1 = PolyVertices[2] - PolyVertices[1];
+	const FVector EdgeDir2 = PolyVertices[0] - PolyVertices[2];
+
+	// Box Z edge x triangle edges.
+
+	if (!TestSeparatingAxisTriangle(FVector(EdgeDir0.Y, -EdgeDir0.X, 0.0f)) ||
+		!TestSeparatingAxisTriangle(FVector(EdgeDir1.Y, -EdgeDir1.X, 0.0f)) ||
+		!TestSeparatingAxisTriangle(FVector(EdgeDir2.Y, -EdgeDir2.X, 0.0f)))
+	{
+		return false;
+	}
+
+	// Box Y edge x triangle edges.
+
+	if (!TestSeparatingAxisTriangle(FVector(-EdgeDir0.Z, 0.0f, EdgeDir0.X)) ||
+		!TestSeparatingAxisTriangle(FVector(-EdgeDir1.Z, 0.0f, EdgeDir1.X)) ||
+		!TestSeparatingAxisTriangle(FVector(-EdgeDir2.Z, 0.0f, EdgeDir2.X)))
+	{
+		return false;
+	}
+
+	// Box X edge x triangle edges.
+
+	if (!TestSeparatingAxisTriangle(FVector(0.0f, EdgeDir0.Z, -EdgeDir0.Y)) ||
+		!TestSeparatingAxisTriangle(FVector(0.0f, EdgeDir1.Z, -EdgeDir1.Y)) ||
+		!TestSeparatingAxisTriangle(FVector(0.0f, EdgeDir2.Z, -EdgeDir2.Y)))
+	{
+		return false;
+	}
+
+	// Box faces.
+
+	if (!TestSeparatingAxisTriangle(FVector(0.0f, 0.0f, 1.0f)) ||
+		!TestSeparatingAxisTriangle(FVector(1.0f, 0.0f, 0.0f)) ||
+		!TestSeparatingAxisTriangle(FVector(0.0f, 1.0f, 0.0f)))
+	{
+		return false;
+	}
+
+	// Triangle normal.
+
+	if (!TestSeparatingAxisTriangle(FVector::CrossProduct(EdgeDir1, EdgeDir0)))
+	{
+		return false;
+	}
+
+	return true;
+}
+
+bool FSeparatingAxisPointCheck::FindSeparatingAxisGeneric()
+{
+	check(PolyVertices.Num() > 3);
+	int32 LastIndex = PolyVertices.Num() - 1;
+	for (int32 Index = 0; Index < PolyVertices.Num(); Index++)
+	{
+		const FVector& V0 = PolyVertices[LastIndex];
+		const FVector& V1 = PolyVertices[Index];
+		const FVector EdgeDir = V1 - V0;
+
+		// Box edges x polygon edge
+
+		if (!TestSeparatingAxisGeneric(FVector(EdgeDir.Y, -EdgeDir.X, 0.0f)) ||
+			!TestSeparatingAxisGeneric(FVector(-EdgeDir.Z, 0.0f, EdgeDir.X)) ||
+			!TestSeparatingAxisGeneric(FVector(0.0f, EdgeDir.Z, -EdgeDir.Y)))
+		{
+			return false;
+		}
+
+		LastIndex = Index;
+	}
+
+	// Box faces.
+
+	if (!TestSeparatingAxisGeneric(FVector(0.0f, 0.0f, 1.0f)) ||
+		!TestSeparatingAxisGeneric(FVector(1.0f, 0.0f, 0.0f)) ||
+		!TestSeparatingAxisGeneric(FVector(0.0f, 1.0f, 0.0f)))
+	{
+		return false;
+	}
+
+	// Polygon normal.
+
+	int32 Index0 = PolyVertices.Num() - 2;
+	int32 Index1 = Index0 + 1;
+	for (int32 Index2 = 0; Index2 < PolyVertices.Num(); Index2++)
+	{
+		const FVector& V0 = PolyVertices[Index0];
+		const FVector& V1 = PolyVertices[Index1];
+		const FVector& V2 = PolyVertices[Index2];
+
+		const FVector EdgeDir0 = V1 - V0;
+		const FVector EdgeDir1 = V2 - V1;
+
+		FVector Normal = FVector::CrossProduct(EdgeDir1, EdgeDir0);
+		if (Normal.SizeSquared() > SMALL_NUMBER)
+		{
+			if (!TestSeparatingAxisGeneric(Normal))
+			{
+				return false;
+			}
+			break;
+		}
+
+		Index0 = Index1;
+		Index1 = Index2;
+	}
+
+	return true;
+}
+
diff --git a/Engine/Source/Runtime/Engine/Public/Collision.h b/Engine/Source/Runtime/Engine/Public/Collision.h
--- a/Engine/Source/Runtime/Engine/Public/Collision.h
+++ b/Engine/Source/Runtime/Engine/Public/Collision.h
@@ -28,147 +28,130 @@ DECLARE_CYCLE_STAT_EXTERN(TEXT("GeomComputePenetration"), STAT_Collision_GeomCom
 #endif
 
 //
-//	FSeparatingAxisPointCheck - Checks for intersection between an oriented bounding box and a triangle.
-//	HitNormal: The normal of the separating axis the bounding box is penetrating the least.
-//	BestDist: The amount the bounding box is penetrating the axis defined by HitNormal.
+//	FSeparatingAxisPointCheck - Checks for intersection between an AABB and a convex polygon.
 //
 
-struct FSeparatingAxisPointCheck
+class ENGINE_API FSeparatingAxisPointCheck
 {
-	FVector	HitNormal;
-	float	BestDist;
-	bool	Hit;
-
-	const FVector&	V0,
-					V1,
-					V2;
-
-	bool TestSeparatingAxis(
-		const FVector& Axis,
-		float ProjectedPoint,
-		float ProjectedExtent
-		)
-	{
-		float	ProjectedV0 = Axis | V0,
-				ProjectedV1 = Axis | V1,
-				ProjectedV2 = Axis | V2,
-				TriangleMin = FMath::Min3(ProjectedV0,ProjectedV1,ProjectedV2) - ProjectedExtent,
-				TriangleMax = FMath::Max3(ProjectedV0,ProjectedV1,ProjectedV2) + ProjectedExtent,
-				AxisMagnitude = Axis.Size();
-
-		if(ProjectedPoint >= TriangleMin && ProjectedPoint <= TriangleMax && AxisMagnitude > SMALL_NUMBER )
-		{
-			// Use inverse sqrt because that is fast and we do more math with the inverse value anyway
-			float	InvAxisMagnitude = FMath::InvSqrt(Axis.X * Axis.X +	Axis.Y * Axis.Y + Axis.Z * Axis.Z),
-					ScaledBestDist = BestDist / InvAxisMagnitude,
-					MinPenetrationDist = ProjectedPoint - TriangleMin,
-					MaxPenetrationDist = TriangleMax - ProjectedPoint;
-			if(MinPenetrationDist < ScaledBestDist)
-			{
-				BestDist = MinPenetrationDist * InvAxisMagnitude;
-				HitNormal = -Axis * InvAxisMagnitude;
-			}
-			if(MaxPenetrationDist < ScaledBestDist)
-			{
-				BestDist = MaxPenetrationDist * InvAxisMagnitude;
-				HitNormal = Axis * InvAxisMagnitude;
-			}
-			return 1;
-		}
-		else
-			return 0;
-	}
-
-		
-	bool TestSeparatingAxis(
-		const FVector& Axis,
-		const FVector& Point,
-		const FVector& BoxExtent
-		)
-	{
-		float	ProjectedExtent = BoxExtent.X * FMath::Abs(Axis.X) + BoxExtent.Y * FMath::Abs(Axis.Y) + BoxExtent.Z * FMath::Abs(Axis.Z);
-		return TestSeparatingAxis(Axis,Axis | Point,ProjectedExtent);
-	}
-
-	bool FindSeparatingAxis(
-		const FVector& Point,
-		const FVector& BoxExtent
+public:
+	/** The normal of the separating axis that the bounding box is penetrating the least */
+	FVector HitNormal;
+
+	/** The amount that the bounding box is penetrating the axis defined by HitNormal */
+	float BestDist;
+
+	/** Whether the bounding box intersects with the polygon */
+	bool bHit;
+
+	/**
+	 *	Creates an object representing the intersection between an axis-aligned bounding box and a convex polygon.
+	 *
+	 *	@param	InPolyVertices			Array of points forming a convex polygon.
+	 *	@param	InBoxCenter				The center of the axis-aligned bounding box being checked.
+	 *	@param	InBoxExtent				The extents of the axis-aligned bounding box being checked.
+	 *	@param	bInCalcLeastPenetration	Whether the axis and amount of penetration of the bounding box into the polygon should be calculated.
+	 */
+	FSeparatingAxisPointCheck(
+		const TArray<FVector>& InPolyVertices,
+		const FVector& InBoxCenter,
+		const FVector& InBoxExtent,
+		bool bInCalcLeastPenetration = true
 		)
+		: PolyVertices(InPolyVertices),
+		  BoxCenter(InBoxCenter),
+		  BoxExtent(InBoxExtent),
+		  HitNormal(FVector::ZeroVector),
+		  BestDist(TNumericLimits<float>::Max()),
+		  bCalcLeastPenetration(bInCalcLeastPenetration)
 	{
-		// Triangle normal.
-
-		if(!TestSeparatingAxis((V2 - V1) ^ (V1 - V0),Point,BoxExtent))
-			return 0;
-
-		const FVector& EdgeDir0 = V1 - V0;
-		const FVector& EdgeDir1 = V2 - V1;
-		const FVector& EdgeDir2 = V0 - V2;
-
-		// Box Z edge x triangle edges.
-
-		if(!TestSeparatingAxis(FVector(EdgeDir0.Y,-EdgeDir0.X,0.0f),Point,BoxExtent))
-			return 0;
-
-		if(!TestSeparatingAxis(FVector(EdgeDir1.Y,-EdgeDir1.X,0.0f),Point,BoxExtent))
-			return 0;
-
-		if(!TestSeparatingAxis(FVector(EdgeDir2.Y,-EdgeDir2.X,0.0f),Point,BoxExtent))
-			return 0;
-
-		// Box Y edge x triangle edges.
-
-		if(!TestSeparatingAxis(FVector(-EdgeDir0.Z,0.0f,EdgeDir0.X),Point,BoxExtent))
-			return 0;
-
-		if(!TestSeparatingAxis(FVector(-EdgeDir1.Z,0.0f,EdgeDir1.X),Point,BoxExtent))
-			return 0;
-
-		if(!TestSeparatingAxis(FVector(-EdgeDir2.Z,0.0f,EdgeDir2.X),Point,BoxExtent))
-			return 0;
-
-		// Box X edge x triangle edges.
-
-		if(!TestSeparatingAxis(FVector(0.0f,EdgeDir0.Z,-EdgeDir0.Y),Point,BoxExtent))
-			return 0;
-
-		if(!TestSeparatingAxis(FVector(0.0f,EdgeDir1.Z,-EdgeDir1.Y),Point,BoxExtent))
-			return 0;
-
-		if(!TestSeparatingAxis(FVector(0.0f,EdgeDir2.Z,-EdgeDir2.Y),Point,BoxExtent))
-			return 0;
-
-		// Box faces. We need to calculate this by crossing edges because BoxX etc are the _edge_ directions - not the faces.
-		// The box may be sheared due to non-unfiform scaling and rotation so FaceX normal != BoxX edge direction
-
-		if(!TestSeparatingAxis(FVector(0.0f,0.0f,1.0f),Point,BoxExtent))
-			return 0;
-
-		if(!TestSeparatingAxis(FVector(1.0f,0.0f,0.0f),Point,BoxExtent))
-			return 0;
-
-		if(!TestSeparatingAxis(FVector(0.0f,1.0f,0.0f),Point,BoxExtent))
-			return 0;
-
-		return 1;
+		// Optimization: if the poly is a triangle, use a more optimized code path
+		bHit = (PolyVertices.Num() == 3) ? FindSeparatingAxisTriangle() : FindSeparatingAxisGeneric();
 	}
 
+	/**
+	 *	Legacy constructor for the class (deprecated)
+	 */
 	FSeparatingAxisPointCheck(
 		const FVector& InV0,
 		const FVector& InV1,
 		const FVector& InV2,
-		const FVector& Point,
-		const FVector& BoxExtent,
+		const FVector& InBoxCenter,
+		const FVector& InBoxExtent,
 		float InBestDist
-		):
-	HitNormal(0,0,0),
-		BestDist(InBestDist),
-		Hit(0),
-		V0(InV0),
-		V1(InV1),
-		V2(InV2)
+		)
+		: PolyVertices(TriangleVertices),
+		  BoxCenter(InBoxCenter),
+		  BoxExtent(InBoxExtent),
+		  HitNormal(FVector::ZeroVector),
+		  BestDist(InBestDist),
+		  bCalcLeastPenetration(true)
 	{
-		Hit = FindSeparatingAxis(Point,BoxExtent);
+		TriangleVertices.Empty(3);
+		TriangleVertices[0] = InV0;
+		TriangleVertices[1] = InV1;
+		TriangleVertices[2] = InV2;
+		bHit = FindSeparatingAxisTriangle();
 	}
+
+private:
+
+	/**
+	 *	Given a separating axis, and a minimum and maximum point projected onto it, determines whether the bounding box
+	 *	encroaches on the polygon along that axis, according to the separating axis theorem.
+	 *	Optionally sets the HitNormal and BestDist members according to the least penetration of the bounding box into the polygon.
+	 *
+	 *	@param	Axis				The separating axis.
+	 *	@param	ProjectedPolyMin	The minimum polygon point, projected onto the separating axis.
+	 *	@param	ProjectedPolyMax	The maximum polygon point, projected onto the separating axis.
+	 */
+	bool TestSeparatingAxisCommon(const FVector& Axis, float ProjectedPolyMin, float ProjectedPolyMax);
+
+	/**
+	 *	Determines whether the bounding box encroaches on the triangle along the given axis.
+	 *	
+	 *	@param	Axis				The separating axis.
+	 *
+	 *	@return	True if the bounding box encroaches on the triangle along the given axis.
+	 */
+	bool TestSeparatingAxisTriangle(const FVector& Axis);
+
+	/**
+	 *	Determines whether the bounding box encroaches on the convex polygon along the given axis.
+	 *
+	 *	@param	Axis				The separating axis.
+	 *
+	 *	@return	True if the bounding box encroaches on the triangle along the given axis.
+	 */
+	bool TestSeparatingAxisGeneric(const FVector& Axis);
+
+	/**
+	 *	Determines whether the bounding box encroaches on the triangle, checking all relevant axes.
+	 *
+	 *	@return	True if the bounding box encroaches on the triangle.
+	 */
+	bool FindSeparatingAxisTriangle();
+
+	/**
+	 *	Determines whether the bounding box encroaches on the convex polygon, checking all relevant axes.
+	 *
+	 *	@return	True if the bounding box encroaches on the convex polygon.
+	 */
+	bool FindSeparatingAxisGeneric();
+
+	/** Array of vertices defining the convex polygon being checked. */
+	const TArray<FVector>& PolyVertices;
+
+	/** Center of the axis-aligned bounding box being checked. */
+	FVector BoxCenter;
+
+	/** Extents of the axis-aligned bounding box being checked. */
+	FVector BoxExtent;
+
+	/** Flag specifying whether the least penetration should be calculated. */
+	bool bCalcLeastPenetration;
+
+	/** Array into which triangle vertices are placed (legacy use only) */
+	static TArray<FVector> TriangleVertices;
 };
 
 /**
