Resolve TTP #323573 by disabling screen precentage on mobile (not implemented).

diff --git a/Engine/Source/Runtime/Engine/Private/SceneView.cpp b/Engine/Source/Runtime/Engine/Private/SceneView.cpp
--- a/Engine/Source/Runtime/Engine/Private/SceneView.cpp
+++ b/Engine/Source/Runtime/Engine/Private/SceneView.cpp
@@ -912,6 +912,12 @@ void FSceneView::EndFinalPostprocessSettings()
 				FinalPostProcessSettings.ScreenPercentage = Value;
 			}
 		}
+
+		// Not supported in ES2.
+		if(Family->Scene->GetFeatureLevel() == ERHIFeatureLevel::ES2)
+		{
+			FinalPostProcessSettings.ScreenPercentage = 100.0f;
+		}
 	}
 
 #if !(UE_BUILD_SHIPPING || UE_BUILD_TEST)
