#ue4 default values read from CDO of spawned class, displayed to user, assignment nodes generated if the pin values differ from the values in the CDO...

diff --git a/Engine/Source/Editor/BlueprintGraph/Classes/EdGraphSchema_K2.h b/Engine/Source/Editor/BlueprintGraph/Classes/EdGraphSchema_K2.h
--- a/Engine/Source/Editor/BlueprintGraph/Classes/EdGraphSchema_K2.h
+++ b/Engine/Source/Editor/BlueprintGraph/Classes/EdGraphSchema_K2.h
@@ -536,6 +536,9 @@ public:
 	// Returns whether the supplied Pin is a splittable struct
 	bool PinHasSplittableStructType(const UEdGraphPin* InGraphPin) const;
 
+	/** Returns true if the pin has a value field that can be edited inline */
+	bool PinDefaultValueIsEditable(const UEdGraphPin& InGraphPin) const;
+
 	/** Helper function to create the expansion node.  
 		If the CompilerContext is specified this will be created as an intermediate node */
 	class UK2Node* CreateSplitPinNode(UEdGraphPin* Pin, class FKismetCompilerContext* CompilerContext = NULL, UEdGraph* SourceGraph = NULL) const;
@@ -1019,5 +1022,11 @@ private:
 	 * @return true if there are out parameters, else false
 	 */
 	bool DoesFunctionHaveOutParameters( const UFunction* Function ) const;
+
+	const UScriptStruct* VectorStruct;
+	const UScriptStruct* RotatorStruct;
+	const UScriptStruct* TransformStruct;
+	const UScriptStruct* LinearColorStruct;
+	const UScriptStruct* ColorStruct;
 };
 
diff --git a/Engine/Source/Editor/BlueprintGraph/Classes/K2Node_AddComponent.h b/Engine/Source/Editor/BlueprintGraph/Classes/K2Node_AddComponent.h
--- a/Engine/Source/Editor/BlueprintGraph/Classes/K2Node_AddComponent.h
+++ b/Engine/Source/Editor/BlueprintGraph/Classes/K2Node_AddComponent.h
@@ -67,6 +67,8 @@ private:
 	{
 		return FindPin(TEXT("TemplateName"));
 	}
+
+	const UClass* GetSpawnedType() const;
 };
 
 
diff --git a/Engine/Source/Editor/BlueprintGraph/Private/EdGraphSchema_K2.cpp b/Engine/Source/Editor/BlueprintGraph/Private/EdGraphSchema_K2.cpp
--- a/Engine/Source/Editor/BlueprintGraph/Private/EdGraphSchema_K2.cpp
+++ b/Engine/Source/Editor/BlueprintGraph/Private/EdGraphSchema_K2.cpp
@@ -262,6 +262,12 @@ UEdGraphSchema_K2::UEdGraphSchema_K2(const class FPostConstructInitializePropert
 	GN_AnimGraph = TEXT("AnimGraph");
 	VR_DefaultCategory = TEXT("Default");
 	AG_LevelReference = 100;
+
+	VectorStruct = FindObjectChecked<UScriptStruct>(UObject::StaticClass(), TEXT("Vector"));
+	RotatorStruct = FindObjectChecked<UScriptStruct>(UObject::StaticClass(), TEXT("Rotator"));
+	TransformStruct = FindObjectChecked<UScriptStruct>(UObject::StaticClass(), TEXT("Transform"));
+	LinearColorStruct = FindObjectChecked<UScriptStruct>(UObject::StaticClass(), TEXT("LinearColor"));
+	ColorStruct = FindObjectChecked<UScriptStruct>(UObject::StaticClass(), TEXT("Color"));
 }
 
 
@@ -755,6 +761,30 @@ bool UEdGraphSchema_K2::PinHasSplittableStructType(const UEdGraphPin* InGraphPin
 	return bCanSplit;
 }
 
+bool UEdGraphSchema_K2::PinDefaultValueIsEditable(const UEdGraphPin& InGraphPin) const
+{
+	// Array types are not currently assignable without a 'make array' node:
+	if( InGraphPin.PinType.bIsArray )
+	{
+		return false;
+	}
+
+	// User defined structures (from code or from data) cannot accept default values:
+	if( InGraphPin.PinType.PinCategory == PC_Struct )
+	{
+		// Only the built in struct types are editable as 'default' values on a pin.
+		// See FNodeFactory::CreatePinWidget for justification of the above statement!
+		UObject const& SubCategoryObject = *InGraphPin.PinType.PinSubCategoryObject;
+		return &SubCategoryObject == VectorStruct 
+			|| &SubCategoryObject == RotatorStruct
+			|| &SubCategoryObject == TransformStruct
+			|| &SubCategoryObject == LinearColorStruct
+			|| &SubCategoryObject == ColorStruct;
+	}
+
+	return true;
+}
+
 void UEdGraphSchema_K2::GetContextMenuActions(const UEdGraph* CurrentGraph, const UEdGraphNode* InGraphNode, const UEdGraphPin* InGraphPin, FMenuBuilder* MenuBuilder, bool bIsDebugging) const
 {
 	check(CurrentGraph);
diff --git a/Engine/Source/Editor/BlueprintGraph/Private/K2Node_AddComponent.cpp b/Engine/Source/Editor/BlueprintGraph/Private/K2Node_AddComponent.cpp
--- a/Engine/Source/Editor/BlueprintGraph/Private/K2Node_AddComponent.cpp
+++ b/Engine/Source/Editor/BlueprintGraph/Private/K2Node_AddComponent.cpp
@@ -103,12 +103,26 @@ void UK2Node_AddComponent::AllocatePinsForExposedVariables()
 					UEdGraphPin* Pin = CreatePin(EGPD_Input, TEXT(""), TEXT(""), NULL, false, false, Property->GetName());
 					Pin->PinType = PinType;
 					bHasExposedVariable = true;
+
+					if (K2Schema->PinDefaultValueIsEditable(*Pin))
+					{
+						FString DefaultValueAsString;
+						const bool bDefaultValueSet = FBlueprintEditorUtils::PropertyValueToString(Property, (uint8*)ComponentClass->ClassDefaultObject, DefaultValueAsString);
+						check(bDefaultValueSet);
+						K2Schema->TrySetDefaultValue(*Pin, DefaultValueAsString);
+					}
 				}
 			}
 		}
 	}
 }
 
+const UClass* UK2Node_AddComponent::GetSpawnedType() const
+{
+	const UActorComponent* TemplateComponent = GetTemplateFromNode();
+	return TemplateComponent ? TemplateComponent->GetClass() : NULL;
+}
+
 void UK2Node_AddComponent::AllocateDefaultPinsWithoutExposedVariables()
 {
 	Super::AllocateDefaultPins();
@@ -396,7 +410,7 @@ void UK2Node_AddComponent::ExpandNode(class FKismetCompilerContext& CompilerCont
 		// exec in
 		CompilerContext.MovePinLinksToIntermediate(*GetExecPin(), *NewNode->GetExecPin());
 
-		UEdGraphPin* LastThen = FKismetCompilerUtilities::GenerateAssignmentNodes( CompilerContext, SourceGraph, NewNode, this, ReturnPin );
+		UEdGraphPin* LastThen = FKismetCompilerUtilities::GenerateAssignmentNodes( CompilerContext, SourceGraph, NewNode, this, ReturnPin, GetSpawnedType() );
 
 		CompilerContext.MovePinLinksToIntermediate(*GetThenPin(), *LastThen);
 		BreakAllNodeLinks();
diff --git a/Engine/Source/Editor/BlueprintGraph/Private/K2Node_SpawnActorFromClass.cpp b/Engine/Source/Editor/BlueprintGraph/Private/K2Node_SpawnActorFromClass.cpp
--- a/Engine/Source/Editor/BlueprintGraph/Private/K2Node_SpawnActorFromClass.cpp
+++ b/Engine/Source/Editor/BlueprintGraph/Private/K2Node_SpawnActorFromClass.cpp
@@ -75,7 +75,6 @@ void UK2Node_SpawnActorFromClass::SetPinToolTip(UEdGraphPin& MutatablePin, const
 	MutatablePin.PinToolTip += FString(TEXT("\n")) + PinDescription.ToString();
 }
 
-
 void UK2Node_SpawnActorFromClass::CreatePinsForClass(UClass* InClass)
 {
 	check(InClass != NULL);
@@ -94,10 +93,19 @@ void UK2Node_SpawnActorFromClass::CreatePinsForClass(UClass* InClass)
 			!Property->HasAnyPropertyFlags(CPF_Parm) && 
 			bIsSettableExternally &&
 			Property->HasAllPropertyFlags(CPF_BlueprintVisible) &&
-			!bIsDelegate )
+			!bIsDelegate &&
+			(NULL == FindPin(Property->GetName()) ) )
 		{
 			UEdGraphPin* Pin = CreatePin(EGPD_Input, TEXT(""), TEXT(""), NULL, false, false, Property->GetName());
-			const bool bPinGood = (Pin != NULL) && K2Schema->ConvertPropertyToPinType(Property, /*out*/ Pin->PinType);	
+			const bool bPinGood = (Pin != NULL) && K2Schema->ConvertPropertyToPinType(Property, /*out*/ Pin->PinType);
+			
+			if( K2Schema->PinDefaultValueIsEditable(*Pin) )
+			{
+				FString DefaultValueAsString;
+				const bool bDefaultValueSet = FBlueprintEditorUtils::PropertyValueToString(Property, (uint8*)InClass->ClassDefaultObject, DefaultValueAsString);
+				check( bDefaultValueSet );
+				K2Schema->TrySetDefaultValue(*Pin, DefaultValueAsString);
+			}
 		}
 	}
 
@@ -422,7 +430,7 @@ void UK2Node_SpawnActorFromClass::ExpandNode(class FKismetCompilerContext& Compi
 		// create 'set var' nodes
 
 		// Get 'result' pin from 'begin spawn', this is the actual actor we want to set properties on
-		UEdGraphPin* LastThen = FKismetCompilerUtilities::GenerateAssignmentNodes(CompilerContext, SourceGraph, CallBeginSpawnNode, SpawnNode, CallBeginResult );
+		UEdGraphPin* LastThen = FKismetCompilerUtilities::GenerateAssignmentNodes(CompilerContext, SourceGraph, CallBeginSpawnNode, SpawnNode, CallBeginResult, GetClassToSpawn() );
 
 		// Make exec connection between 'then' on last node and 'finish'
 		LastThen->MakeLinkTo(CallFinishExec);
diff --git a/Engine/Source/Editor/GraphEditor/Private/NodeFactory.cpp b/Engine/Source/Editor/GraphEditor/Private/NodeFactory.cpp
--- a/Engine/Source/Editor/GraphEditor/Private/NodeFactory.cpp
+++ b/Engine/Source/Editor/GraphEditor/Private/NodeFactory.cpp
@@ -290,6 +290,7 @@ TSharedPtr<SGraphPin> FNodeFactory::CreatePinWidget(UEdGraphPin* InPin)
 		}
 		else if (InPin->PinType.PinCategory == K2Schema->PC_Struct)
 		{
+			// If you update this logic you'll probably need to update UEdGraphSchema_K2::ShouldHidePinDefaultValue!
 			UScriptStruct* ColorStruct = FindObjectChecked<UScriptStruct>(UObject::StaticClass(), TEXT("LinearColor"));
 			UScriptStruct* VectorStruct = FindObjectChecked<UScriptStruct>(UObject::StaticClass(), TEXT("Vector"));
 			UScriptStruct* Vector2DStruct = FindObjectChecked<UScriptStruct>(UObject::StaticClass(), TEXT("Vector2D"));
diff --git a/Engine/Source/Editor/KismetCompiler/Private/KismetCompilerMisc.cpp b/Engine/Source/Editor/KismetCompiler/Private/KismetCompilerMisc.cpp
--- a/Engine/Source/Editor/KismetCompiler/Private/KismetCompilerMisc.cpp
+++ b/Engine/Source/Editor/KismetCompiler/Private/KismetCompilerMisc.cpp
@@ -503,7 +503,7 @@ void FKismetCompilerUtilities::ValidateEnumProperties(UObject* DefaultObject, FC
 	}
 }
 
-UEdGraphPin* FKismetCompilerUtilities::GenerateAssignmentNodes(class FKismetCompilerContext& CompilerContext, UEdGraph* SourceGraph, UK2Node_CallFunction* CallBeginSpawnNode, UEdGraphNode* SpawnNode, UEdGraphPin* CallBeginResult)
+UEdGraphPin* FKismetCompilerUtilities::GenerateAssignmentNodes(class FKismetCompilerContext& CompilerContext, UEdGraph* SourceGraph, UK2Node_CallFunction* CallBeginSpawnNode, UEdGraphNode* SpawnNode, UEdGraphPin* CallBeginResult, const UClass* ForClass )
 {
 	static FString ObjectParamName = FString(TEXT("Object"));
 	static FString ValueParamName = FString(TEXT("Value"));
@@ -520,6 +520,26 @@ UEdGraphPin* FKismetCompilerUtilities::GenerateAssignmentNodes(class FKismetComp
 		if (NULL == CallBeginSpawnNode->FindPin(OrgPin->PinName) &&
 			(OrgPin->LinkedTo.Num() > 0 || OrgPin->DefaultValue != FString()))
 		{
+			if( OrgPin->LinkedTo.Num() == 0 )
+			{
+				UProperty* Property = FindField<UProperty>(ForClass, *OrgPin->PinName);
+				// NULL property indicates that this pin was part of the original node, not the 
+				// class we're assigning to:
+				if( !Property )
+				{
+					continue;
+				}
+
+				// We don't want to generate an assignment node unless the default value 
+				// differs from the value in the CDO:
+				FString DefaultValueAsString;
+				FBlueprintEditorUtils::PropertyValueToString(Property, (uint8*)ForClass->ClassDefaultObject, DefaultValueAsString);
+				if (DefaultValueAsString == OrgPin->DefaultValue)
+				{
+					continue;
+				}
+			}
+
 			UFunction* SetByNameFunction = Schema->FindSetVariableByNameFunction(OrgPin->PinType);
 			if (SetByNameFunction)
 			{
diff --git a/Engine/Source/Editor/KismetCompiler/Public/KismetCompilerMisc.h b/Engine/Source/Editor/KismetCompiler/Public/KismetCompilerMisc.h
--- a/Engine/Source/Editor/KismetCompiler/Public/KismetCompilerMisc.h
+++ b/Engine/Source/Editor/KismetCompiler/Public/KismetCompilerMisc.h
@@ -77,7 +77,7 @@ public:
 	static void ValidateEnumProperties(UObject* DefaultObject, FCompilerResultsLog& MessageLog);
 
 	/** Create 'set var by name' nodes and hook them up - used to set values when components are added or actor are created at run time. Returns the 'last then' pin of the assignment nodes */
-	static UEdGraphPin* GenerateAssignmentNodes( class FKismetCompilerContext& CompilerContext, UEdGraph* SourceGraph, UK2Node_CallFunction* CallBeginSpawnNode, UEdGraphNode* SpawnNode, UEdGraphPin* CallBeginResult );
+	static UEdGraphPin* GenerateAssignmentNodes( class FKismetCompilerContext& CompilerContext, UEdGraph* SourceGraph, UK2Node_CallFunction* CallBeginSpawnNode, UEdGraphNode* SpawnNode, UEdGraphPin* CallBeginResult, const UClass* ForClass );
 };
 
 //////////////////////////////////////////////////////////////////////////
diff --git a/Engine/Source/Runtime/Engine/Classes/Kismet/KismetSystemLibrary.h b/Engine/Source/Runtime/Engine/Classes/Kismet/KismetSystemLibrary.h
--- a/Engine/Source/Runtime/Engine/Classes/Kismet/KismetSystemLibrary.h
+++ b/Engine/Source/Runtime/Engine/Classes/Kismet/KismetSystemLibrary.h
@@ -323,31 +323,31 @@ class UKismetSystemLibrary : public UBlueprintFunctionLibrary
 	static void SetObjectPropertyByName(UObject* Object, FName PropertyName, UObject* Value);
 
 	/** Set a NAME property by name */
-	UFUNCTION(BlueprintCallable, meta=(BlueprintInternalUseOnly = "true"))
+	UFUNCTION(BlueprintCallable, meta=(BlueprintInternalUseOnly = "true", AutoCreateRefTerm = "Value" ))
 	static void SetNamePropertyByName(UObject* Object, FName PropertyName, const FName& Value);
 
 	/** Set a STRING property by name */
-	UFUNCTION(BlueprintCallable, meta=(BlueprintInternalUseOnly = "true"))
+	UFUNCTION(BlueprintCallable, meta=(BlueprintInternalUseOnly = "true", AutoCreateRefTerm = "Value" ))
 	static void SetStringPropertyByName(UObject* Object, FName PropertyName, const FString& Value);
 
 	/** Set a TEXT property by name */
-	UFUNCTION(BlueprintCallable, meta=(BlueprintInternalUseOnly = "true"))
+	UFUNCTION(BlueprintCallable, meta=(BlueprintInternalUseOnly = "true", AutoCreateRefTerm = "Value" ))
 	static void SetTextPropertyByName(UObject* Object, FName PropertyName, const FText& Value);
 
 	/** Set a VECTOR property by name */
-	UFUNCTION(BlueprintCallable, meta=(BlueprintInternalUseOnly = "true"))
+	UFUNCTION(BlueprintCallable, meta=(BlueprintInternalUseOnly = "true", AutoCreateRefTerm = "Value" ))
 	static void SetVectorPropertyByName(UObject* Object, FName PropertyName, const FVector& Value);
 
 	/** Set a ROTATOR property by name */
-	UFUNCTION(BlueprintCallable, meta=(BlueprintInternalUseOnly = "true"))
+	UFUNCTION(BlueprintCallable, meta=(BlueprintInternalUseOnly = "true", AutoCreateRefTerm = "Value" ))
 	static void SetRotatorPropertyByName(UObject* Object, FName PropertyName, const FRotator& Value);
 
 	/** Set a LINEAR COLOR property by name */
-	UFUNCTION(BlueprintCallable, meta=(BlueprintInternalUseOnly = "true"))
+	UFUNCTION(BlueprintCallable, meta=(BlueprintInternalUseOnly = "true", AutoCreateRefTerm = "Value" ))
 	static void SetLinearColorPropertyByName(UObject* Object, FName PropertyName, const FLinearColor& Value);
 
 	/** Set a TRANSFORM property by name */
-	UFUNCTION(BlueprintCallable, meta=(BlueprintInternalUseOnly = "true"))
+	UFUNCTION(BlueprintCallable, meta=(BlueprintInternalUseOnly = "true", AutoCreateRefTerm = "Value" ))
 	static void SetTransformPropertyByName(UObject* Object, FName PropertyName, const FTransform& Value);
 
 	/** Set a custom structure property by name */
