Fix clang errors (initialization order issues).

diff --git a/Engine/Source/Runtime/Engine/Public/Collision.h b/Engine/Source/Runtime/Engine/Public/Collision.h
--- a/Engine/Source/Runtime/Engine/Public/Collision.h
+++ b/Engine/Source/Runtime/Engine/Public/Collision.h
@@ -57,11 +57,11 @@ public:
 		const FVector& InBoxExtent,
 		bool bInCalcLeastPenetration = true
 		)
-		: PolyVertices(InPolyVertices),
+		: HitNormal(FVector::ZeroVector),
+		  BestDist(TNumericLimits<float>::Max()),
+		  PolyVertices(InPolyVertices),
 		  BoxCenter(InBoxCenter),
 		  BoxExtent(InBoxExtent),
-		  HitNormal(FVector::ZeroVector),
-		  BestDist(TNumericLimits<float>::Max()),
 		  bCalcLeastPenetration(bInCalcLeastPenetration)
 	{
 		// Optimization: if the poly is a triangle, use a more optimized code path
@@ -79,11 +79,11 @@ public:
 		const FVector& InBoxExtent,
 		float InBestDist
 		)
-		: PolyVertices(TriangleVertices),
+		: HitNormal(FVector::ZeroVector),
+		  BestDist(InBestDist),
+		  PolyVertices(TriangleVertices),
 		  BoxCenter(InBoxCenter),
 		  BoxExtent(InBoxExtent),
-		  HitNormal(FVector::ZeroVector),
-		  BestDist(InBestDist),
 		  bCalcLeastPenetration(true)
 	{
 		TriangleVertices.Empty(3);
