- Changed AnimNotify to support any other AnimSequenceBase
- Fixed K2Node_FunctionEntry to support redirect

diff --git a/Engine/Config/BaseEngine.ini b/Engine/Config/BaseEngine.ini
--- a/Engine/Config/BaseEngine.ini
+++ b/Engine/Config/BaseEngine.ini
@@ -324,7 +324,7 @@ bUseStreamingPause=false
 +ActiveClassRedirects=(OldClassName="K2Node_MathExpression",NewClassName="/Script/BlueprintGraph.K2Node_MathExpression")
 +ActiveClassRedirects=(OldClassName="/Script/CoreUObject.K2Node_MathExpression",NewClassName="/Script/BlueprintGraph.K2Node_MathExpression")
 
-; Move anim nodes into their own module
+; Move ani mnodes into their own module
 +ActiveClassRedirects=(OldClassName="AnimationConduitGraphSchema",NewClassName="/Script/AnimGraph.AnimationConduitGraphSchema")
 +ActiveClassRedirects=(OldClassName="AnimationCustomTransitionGraph",NewClassName="/Script/AnimGraph.AnimationCustomTransitionGraph")
 +ActiveClassRedirects=(OldClassName="AnimationCustomTransitionSchema",NewClassName="/Script/AnimGraph.AnimationCustomTransitionSchema")
@@ -1233,6 +1233,13 @@ bUseStreamingPause=false
 
 +EnumRedirects=(EnumName="EMaxConcurrentResolutionRule",OldEnumEntry="EMaxConcurrentResolutionRule::StopFarthest",NewEnumEntry="EMaxConcurrentResolutionRule::StopFarthestThenPreventNew")
 
+; AnimNotify parameter change
+
++K2ParamRedirects=(NodeName="/Script/BlueprintGraph.K2Node_FunctionEntry", OldParamName="AnimNotify.Received_Notify.AnimSeq", NewParamName="Animation")
++K2ParamRedirects=(NodeName="/Script/BlueprintGraph.K2Node_FunctionEntry", OldParamName="AnimNotifyState.Received_NotifyBegin.AnimSeq", NewParamName="Animation")
++K2ParamRedirects=(NodeName="/Script/BlueprintGraph.K2Node_FunctionEntry", OldParamName="AnimNotifyState.Received_NotifyTick.AnimSeq", NewParamName="Animation")
++K2ParamRedirects=(NodeName="/Script/BlueprintGraph.K2Node_FunctionEntry", OldParamName="AnimNotifyState.Received_NotifyEnd.AnimSeq", NewParamName="Animation")
+
 ; VER_UE4_INVERSE_SQUARED_LIGHTS_DEFAULT
 +K2FieldRedirects=(OldFieldName="LightComponentBase.Brightness",NewFieldName="LightComponentBase.Intensity")
 +K2FieldRedirects=(OldFieldName="PointLightComponent.Radius",NewFieldName="PointLightComponent.AttenuationRadius")
diff --git a/Engine/Source/Editor/BlueprintGraph/Classes/K2Node_FunctionEntry.h b/Engine/Source/Editor/BlueprintGraph/Classes/K2Node_FunctionEntry.h
--- a/Engine/Source/Editor/BlueprintGraph/Classes/K2Node_FunctionEntry.h
+++ b/Engine/Source/Editor/BlueprintGraph/Classes/K2Node_FunctionEntry.h
@@ -36,6 +36,7 @@ class UK2Node_FunctionEntry : public UK2Node_FunctionTerminator
 	// Begin UK2Node interface
 	virtual bool DrawNodeAsEntry() const override { return true; }
 	virtual class FNodeHandlingFunctor* CreateNodeHandler(class FKismetCompilerContext& CompilerContext) const override;
+	virtual void GetRedirectPinNames(const UEdGraphPin& Pin, TArray<FString>& RedirectPinNames) const override;
 	// End UK2Node interface
 
 	// Begin UK2Node_EditablePinBase interface
diff --git a/Engine/Source/Editor/BlueprintGraph/Private/K2Node_FunctionEntry.cpp b/Engine/Source/Editor/BlueprintGraph/Private/K2Node_FunctionEntry.cpp
--- a/Engine/Source/Editor/BlueprintGraph/Private/K2Node_FunctionEntry.cpp
+++ b/Engine/Source/Editor/BlueprintGraph/Private/K2Node_FunctionEntry.cpp
@@ -196,7 +196,24 @@ FNodeHandlingFunctor* UK2Node_FunctionEntry::CreateNodeHandler(FKismetCompilerCo
 	return new FKCHandler_FunctionEntry(CompilerContext);
 }
 
+void UK2Node_FunctionEntry::GetRedirectPinNames(const UEdGraphPin& Pin, TArray<FString>& RedirectPinNames) const
+{
+	Super::GetRedirectPinNames(Pin, RedirectPinNames);
+
+	if(RedirectPinNames.Num() > 0)
+	{
+		const FString & OldPinName = RedirectPinNames[0];
 
+		
+		// first add functionname.param
+		RedirectPinNames.Add(FString::Printf(TEXT("%s.%s"), *SignatureName.ToString(), *OldPinName));
+		// if there is class, also add an option for class.functionname.param
+		if(SignatureClass!=NULL)
+		{
+			RedirectPinNames.Add(FString::Printf(TEXT("%s.%s.%s"), *SignatureClass->GetName(), *SignatureName.ToString(), *OldPinName));
+		}
+	}
+}
 
 bool UK2Node_FunctionEntry::IsDeprecated() const
 {
diff --git a/Engine/Source/Runtime/Engine/Classes/Animation/AnimNotifies/AnimNotify.h b/Engine/Source/Runtime/Engine/Classes/Animation/AnimNotifies/AnimNotify.h
--- a/Engine/Source/Runtime/Engine/Classes/Animation/AnimNotifies/AnimNotify.h
+++ b/Engine/Source/Runtime/Engine/Classes/Animation/AnimNotifies/AnimNotify.h
@@ -12,7 +12,7 @@ class ENGINE_API UAnimNotify : public UObject
 	GENERATED_UCLASS_BODY()
 
 	UFUNCTION(BlueprintImplementableEvent)
-	virtual bool Received_Notify(class USkeletalMeshComponent* MeshComp, class UAnimSequence* AnimSeq) const;
+	virtual bool Received_Notify(class USkeletalMeshComponent* MeshComp, class UAnimSequenceBase* Animation) const;
 
 #if WITH_EDITORONLY_DATA
 	/** Color of Notify in editor */
@@ -21,7 +21,7 @@ class ENGINE_API UAnimNotify : public UObject
 
 #endif // WITH_EDITORONLY_DATA
 
-	virtual void Notify(class USkeletalMeshComponent* MeshComp, class UAnimSequence* AnimSeq);
+	virtual void Notify(class USkeletalMeshComponent* MeshComp, class UAnimSequenceBase* Animation);
 
 	// @todo document 
 	virtual FString GetEditorComment() 
@@ -45,7 +45,7 @@ class ENGINE_API UAnimNotify : public UObject
 	 *	@param	AnimSeq			The animation sequence this notify is associated with.
 	 *	@param	OwnerEvent		The event that 'owns' this AnimNotify.
 	 */
-	virtual void AnimNotifyEventChanged(class USkeletalMeshComponent* MeshComp, class UAnimSequence* AnimSeq, FAnimNotifyEvent* OwnerEvent) {}
+	virtual void AnimNotifyEventChanged(class USkeletalMeshComponent* MeshComp, class UAnimSequenceBase* Animation, FAnimNotifyEvent* OwnerEvent) {}
 
 	/**
 	 * We don't instance UAnimNotify objects along with the animations they belong to, but
diff --git a/Engine/Source/Runtime/Engine/Classes/Animation/AnimNotifies/AnimNotifyState.h b/Engine/Source/Runtime/Engine/Classes/Animation/AnimNotifies/AnimNotifyState.h
--- a/Engine/Source/Runtime/Engine/Classes/Animation/AnimNotifies/AnimNotifyState.h
+++ b/Engine/Source/Runtime/Engine/Classes/Animation/AnimNotifies/AnimNotifyState.h
@@ -12,13 +12,13 @@ class ENGINE_API UAnimNotifyState : public UObject
 	GENERATED_UCLASS_BODY()
 
 	UFUNCTION(BlueprintImplementableEvent)
-	virtual bool Received_NotifyBegin(class USkeletalMeshComponent * MeshComp, class UAnimSequence * AnimSeq) const;
+	virtual bool Received_NotifyBegin(class USkeletalMeshComponent * MeshComp, class UAnimSequenceBase * Animation) const;
 	
 	UFUNCTION(BlueprintImplementableEvent)
-	virtual bool Received_NotifyTick(class USkeletalMeshComponent * MeshComp, class UAnimSequence * AnimSeq, float FrameDeltaTime) const;
+	virtual bool Received_NotifyTick(class USkeletalMeshComponent * MeshComp, class UAnimSequenceBase * Animation, float FrameDeltaTime) const;
 
 	UFUNCTION(BlueprintImplementableEvent)
-	virtual bool Received_NotifyEnd(class USkeletalMeshComponent * MeshComp, class UAnimSequence * AnimSeq) const;
+	virtual bool Received_NotifyEnd(class USkeletalMeshComponent * MeshComp, class UAnimSequenceBase * Animation) const;
 
 #if WITH_EDITORONLY_DATA
 	/** Color of Notify in editor */
@@ -27,9 +27,9 @@ class ENGINE_API UAnimNotifyState : public UObject
 
 #endif // WITH_EDITORONLY_DATA
 
-	virtual void NotifyBegin(class USkeletalMeshComponent * MeshComp, class UAnimSequence * AnimSeq);
-	virtual void NotifyTick(class USkeletalMeshComponent * MeshComp, class UAnimSequence * AnimSeq, float FrameDeltaTime);
-	virtual void NotifyEnd(class USkeletalMeshComponent * MeshComp, class UAnimSequence * AnimSeq);
+	virtual void NotifyBegin(class USkeletalMeshComponent * MeshComp, class UAnimSequenceBase * Animation);
+	virtual void NotifyTick(class USkeletalMeshComponent * MeshComp, class UAnimSequenceBase * Animation, float FrameDeltaTime);
+	virtual void NotifyEnd(class USkeletalMeshComponent * MeshComp, class UAnimSequenceBase * Animation);
 
 	// @todo document 
 	virtual FString GetEditorComment() 
@@ -53,7 +53,7 @@ class ENGINE_API UAnimNotifyState : public UObject
 	 *	@param	AnimSeq			The animation sequence this notify is associated with.
 	 *	@param	OwnerEvent		The event that 'owns' this AnimNotify.
 	 */
-	virtual void AnimNotifyEventChanged(class USkeletalMeshComponent* MeshComp, class UAnimSequence* AnimSeq, FAnimNotifyEvent * OwnerEvent) {}
+	virtual void AnimNotifyEventChanged(class USkeletalMeshComponent* MeshComp, class UAnimSequenceBase* Animation, FAnimNotifyEvent * OwnerEvent) {}
 };
 
 
diff --git a/Engine/Source/Runtime/Engine/Classes/Animation/AnimNotifies/AnimNotifyState_TimedParticleEffect.h b/Engine/Source/Runtime/Engine/Classes/Animation/AnimNotifies/AnimNotifyState_TimedParticleEffect.h
--- a/Engine/Source/Runtime/Engine/Classes/Animation/AnimNotifies/AnimNotifyState_TimedParticleEffect.h
+++ b/Engine/Source/Runtime/Engine/Classes/Animation/AnimNotifies/AnimNotifyState_TimedParticleEffect.h
@@ -49,9 +49,9 @@ class UAnimNotifyState_TimedParticleEffect : public UAnimNotifyState
 	virtual void PreEditChange(UProperty* PropertyAboutToChange) override;
 #endif
 
-	virtual void NotifyBegin(class USkeletalMeshComponent * MeshComp, class UAnimSequence * AnimSeq) override;
-	virtual void NotifyTick(class USkeletalMeshComponent * MeshComp, class UAnimSequence * AnimSeq, float FrameDeltaTime) override;
-	virtual void NotifyEnd(class USkeletalMeshComponent * MeshComp, class UAnimSequence * AnimSeq) override;
+	virtual void NotifyBegin(class USkeletalMeshComponent * MeshComp, class UAnimSequenceBase * Animation) override;
+	virtual void NotifyTick(class USkeletalMeshComponent * MeshComp, class UAnimSequenceBase * Animation, float FrameDeltaTime) override;
+	virtual void NotifyEnd(class USkeletalMeshComponent * MeshComp, class UAnimSequenceBase * Animation) override;
 
 private:
 	bool ValidateParameters(USkeletalMeshComponent* MeshComp);
diff --git a/Engine/Source/Runtime/Engine/Classes/Animation/AnimNotifies/AnimNotifyState_Trail.h b/Engine/Source/Runtime/Engine/Classes/Animation/AnimNotifies/AnimNotifyState_Trail.h
--- a/Engine/Source/Runtime/Engine/Classes/Animation/AnimNotifies/AnimNotifyState_Trail.h
+++ b/Engine/Source/Runtime/Engine/Classes/Animation/AnimNotifies/AnimNotifyState_Trail.h
@@ -55,9 +55,9 @@ class UAnimNotifyState_Trail : public UAnimNotifyState
 	uint32 bRenderTessellation : 1;
 #endif // WITH_EDITORONLY_DATA
 
-	virtual void NotifyBegin(class USkeletalMeshComponent * MeshComp, class UAnimSequence * AnimSeq);
-	virtual void NotifyTick(class USkeletalMeshComponent * MeshComp, class UAnimSequence * AnimSeq, float FrameDeltaTime);
-	virtual void NotifyEnd(class USkeletalMeshComponent * MeshComp, class UAnimSequence * AnimSeq);
+	virtual void NotifyBegin(class USkeletalMeshComponent * MeshComp, class UAnimSequenceBase * Animation) override;
+	virtual void NotifyTick(class USkeletalMeshComponent * MeshComp, class UAnimSequenceBase * Animation, float FrameDeltaTime) override;
+	virtual void NotifyEnd(class USkeletalMeshComponent * MeshComp, class UAnimSequenceBase * Animation) override;
 
 	bool ValidateInput(class USkeletalMeshComponent * MeshComp, bool bReportErrors = false);
 
@@ -67,7 +67,7 @@ class UAnimNotifyState_Trail : public UAnimNotifyState
 	 *	@param	AnimSeq			The animation sequence this notify is associated with.
 	 *	@param	OwnerEvent		The event that 'owns' this AnimNotify.
 	 */
-	virtual void AnimNotifyEventChanged(class USkeletalMeshComponent* MeshComp, class UAnimSequence* AnimSeq, FAnimNotifyEvent * OwnerEvent);
+	virtual void AnimNotifyEventChanged(class USkeletalMeshComponent* MeshComp, class UAnimSequenceBase* Animation, FAnimNotifyEvent * OwnerEvent) override;
 };
 
 
diff --git a/Engine/Source/Runtime/Engine/Private/Animation/AnimInstance.cpp b/Engine/Source/Runtime/Engine/Private/Animation/AnimInstance.cpp
--- a/Engine/Source/Runtime/Engine/Private/Animation/AnimInstance.cpp
+++ b/Engine/Source/Runtime/Engine/Private/Animation/AnimInstance.cpp
@@ -1087,7 +1087,7 @@ void UAnimInstance::TriggerAnimNotifies(float DeltaSeconds)
 		if(AnimNotifyEvent->Notify != NULL)
 		{
 			// Implemented notify: just call Notify. UAnimNotify will forward this to the event which will do the work.
-			AnimNotifyEvent->Notify->Notify(SkelMeshComp, Cast<UAnimSequence>(AnimNotifyEvent->Notify->GetOuter()));
+			AnimNotifyEvent->Notify->Notify(SkelMeshComp, Cast<UAnimSequenceBase>(AnimNotifyEvent->Notify->GetOuter()));
 		}
 		else if( AnimNotifyEvent->NotifyName != NAME_None )
 		{
@@ -1128,14 +1128,14 @@ void UAnimInstance::TriggerAnimNotifies(float DeltaSeconds)
 	for(int32 Index=0; Index<ActiveAnimNotifyState.Num(); Index++)
 	{
 		const FAnimNotifyEvent& AnimNotifyEvent = ActiveAnimNotifyState[Index];
-		AnimNotifyEvent.NotifyStateClass->NotifyEnd(SkelMeshComp, Cast<UAnimSequence>(AnimNotifyEvent.NotifyStateClass->GetOuter()));
+		AnimNotifyEvent.NotifyStateClass->NotifyEnd(SkelMeshComp, Cast<UAnimSequenceBase>(AnimNotifyEvent.NotifyStateClass->GetOuter()));
 	}
 
 	// Call 'NotifyBegin' event on freshly added AnimNotifyState.
 	for (int32 Index = 0; Index < NotifyStateBeginEvent.Num(); Index++)
 	{
 		const FAnimNotifyEvent * AnimNotifyEvent = NotifyStateBeginEvent[Index];
-		AnimNotifyEvent->NotifyStateClass->NotifyBegin(SkelMeshComp, Cast<UAnimSequence>(AnimNotifyEvent->NotifyStateClass->GetOuter()));
+		AnimNotifyEvent->NotifyStateClass->NotifyBegin(SkelMeshComp, Cast<UAnimSequenceBase>(AnimNotifyEvent->NotifyStateClass->GetOuter()));
 	}
 
 	// Switch our arrays.
@@ -1145,7 +1145,7 @@ void UAnimInstance::TriggerAnimNotifies(float DeltaSeconds)
 	for(int32 Index=0; Index<ActiveAnimNotifyState.Num(); Index++)
 	{
 		const FAnimNotifyEvent& AnimNotifyEvent = ActiveAnimNotifyState[Index];
-		AnimNotifyEvent.NotifyStateClass->NotifyTick(SkelMeshComp, Cast<UAnimSequence>(AnimNotifyEvent.NotifyStateClass->GetOuter()), DeltaSeconds);
+		AnimNotifyEvent.NotifyStateClass->NotifyTick(SkelMeshComp, Cast<UAnimSequenceBase>(AnimNotifyEvent.NotifyStateClass->GetOuter()), DeltaSeconds);
 	}
 }
 
diff --git a/Engine/Source/Runtime/Engine/Private/Animation/AnimNotify.cpp b/Engine/Source/Runtime/Engine/Private/Animation/AnimNotify.cpp
--- a/Engine/Source/Runtime/Engine/Private/Animation/AnimNotify.cpp
+++ b/Engine/Source/Runtime/Engine/Private/Animation/AnimNotify.cpp
@@ -17,11 +17,11 @@ UAnimNotify::UAnimNotify(const class FPostConstructInitializeProperties& PCIP)
 }
 
 
-void UAnimNotify::Notify(class USkeletalMeshComponent* MeshComp, class UAnimSequence* AnimSeq)
+void UAnimNotify::Notify(class USkeletalMeshComponent* MeshComp, class UAnimSequenceBase* Animation)
 {
 	USkeletalMeshComponent* PrevContext = MeshContext;
 	MeshContext = MeshComp;
-	Received_Notify(MeshComp, AnimSeq);
+	Received_Notify(MeshComp, Animation);
 	MeshContext = PrevContext;
 }
 
diff --git a/Engine/Source/Runtime/Engine/Private/Animation/AnimNotifyState.cpp b/Engine/Source/Runtime/Engine/Private/Animation/AnimNotifyState.cpp
--- a/Engine/Source/Runtime/Engine/Private/Animation/AnimNotifyState.cpp
+++ b/Engine/Source/Runtime/Engine/Private/Animation/AnimNotifyState.cpp
@@ -16,18 +16,18 @@ UAnimNotifyState::UAnimNotifyState(const class FPostConstructInitializePropertie
 }
 
 
-void UAnimNotifyState::NotifyBegin(class USkeletalMeshComponent * MeshComp, class UAnimSequence * AnimSeq)
+void UAnimNotifyState::NotifyBegin(class USkeletalMeshComponent * MeshComp, class UAnimSequenceBase * Animation)
 {
-	Received_NotifyBegin(MeshComp, AnimSeq);
+	Received_NotifyBegin(MeshComp, Animation);
 }
 
-void UAnimNotifyState::NotifyTick(class USkeletalMeshComponent * MeshComp, class UAnimSequence * AnimSeq, float FrameDeltaTime)
+void UAnimNotifyState::NotifyTick(class USkeletalMeshComponent * MeshComp, class UAnimSequenceBase * Animation, float FrameDeltaTime)
 {
-	Received_NotifyTick(MeshComp, AnimSeq, FrameDeltaTime);
+	Received_NotifyTick(MeshComp, Animation, FrameDeltaTime);
 }
 
-void UAnimNotifyState::NotifyEnd(class USkeletalMeshComponent * MeshComp, class UAnimSequence * AnimSeq)
+void UAnimNotifyState::NotifyEnd(class USkeletalMeshComponent * MeshComp, class UAnimSequenceBase * Animation)
 {
-	Received_NotifyEnd(MeshComp, AnimSeq);
+	Received_NotifyEnd(MeshComp, Animation);
 }
 
diff --git a/Engine/Source/Runtime/Engine/Private/Animation/AnimNotifyState_TimedParticleEffect.cpp b/Engine/Source/Runtime/Engine/Private/Animation/AnimNotifyState_TimedParticleEffect.cpp
--- a/Engine/Source/Runtime/Engine/Private/Animation/AnimNotifyState_TimedParticleEffect.cpp
+++ b/Engine/Source/Runtime/Engine/Private/Animation/AnimNotifyState_TimedParticleEffect.cpp
@@ -12,7 +12,7 @@ UAnimNotifyState_TimedParticleEffect::UAnimNotifyState_TimedParticleEffect(const
 	RotationOffset = FRotator(0.0f, 0.0f, 0.0f);
 }
 
-void UAnimNotifyState_TimedParticleEffect::NotifyBegin(USkeletalMeshComponent * MeshComp, class UAnimSequence * AnimSeq)
+void UAnimNotifyState_TimedParticleEffect::NotifyBegin(USkeletalMeshComponent * MeshComp, class UAnimSequenceBase * Animation)
 {
 	// Only spawn if we've got valid params
 	if(ValidateParameters(MeshComp))
@@ -20,15 +20,15 @@ void UAnimNotifyState_TimedParticleEffect::NotifyBegin(USkeletalMeshComponent *
 		UParticleSystemComponent* NewComponent = UGameplayStatics::SpawnEmitterAttached(PSTemplate, MeshComp, SocketName, LocationOffset, RotationOffset);
 	}
 
-	Received_NotifyBegin(MeshComp, AnimSeq);
+	Received_NotifyBegin(MeshComp, Animation);
 }
 
-void UAnimNotifyState_TimedParticleEffect::NotifyTick(USkeletalMeshComponent * MeshComp, class UAnimSequence * AnimSeq, float FrameDeltaTime)
+void UAnimNotifyState_TimedParticleEffect::NotifyTick(USkeletalMeshComponent * MeshComp, class UAnimSequenceBase * Animation, float FrameDeltaTime)
 {
-	Received_NotifyTick(MeshComp, AnimSeq, FrameDeltaTime);
+	Received_NotifyTick(MeshComp, Animation, FrameDeltaTime);
 }
 
-void UAnimNotifyState_TimedParticleEffect::NotifyEnd(USkeletalMeshComponent * MeshComp, class UAnimSequence * AnimSeq)
+void UAnimNotifyState_TimedParticleEffect::NotifyEnd(USkeletalMeshComponent * MeshComp, class UAnimSequenceBase * Animation)
 {
 	TArray<USceneComponent*> Children;
 	MeshComp->GetChildrenComponents(false, Children);
@@ -73,7 +73,7 @@ void UAnimNotifyState_TimedParticleEffect::NotifyEnd(USkeletalMeshComponent * Me
 		}
 	}
 
-	Received_NotifyEnd(MeshComp, AnimSeq);
+	Received_NotifyEnd(MeshComp, Animation);
 }
 
 bool UAnimNotifyState_TimedParticleEffect::ValidateParameters(USkeletalMeshComponent* MeshComp)
diff --git a/Engine/Source/Runtime/Engine/Private/Animation/AnimNotifyState_Trail.cpp b/Engine/Source/Runtime/Engine/Private/Animation/AnimNotifyState_Trail.cpp
--- a/Engine/Source/Runtime/Engine/Private/Animation/AnimNotifyState_Trail.cpp
+++ b/Engine/Source/Runtime/Engine/Private/Animation/AnimNotifyState_Trail.cpp
@@ -34,11 +34,11 @@ UAnimNotifyState_Trail::UAnimNotifyState_Trail(const class FPostConstructInitial
 #endif // WITH_EDITORONLY_DATA
 }
 
-void UAnimNotifyState_Trail::AnimNotifyEventChanged(class USkeletalMeshComponent* MeshComp, class UAnimSequence* AnimSeq, FAnimNotifyEvent * OwnerEvent)
+void UAnimNotifyState_Trail::AnimNotifyEventChanged(class USkeletalMeshComponent* MeshComp, class UAnimSequenceBase* Animation, FAnimNotifyEvent * OwnerEvent)
 {
 }
 
-void UAnimNotifyState_Trail::NotifyBegin(class USkeletalMeshComponent * MeshComp, class UAnimSequence * AnimSeq)
+void UAnimNotifyState_Trail::NotifyBegin(class USkeletalMeshComponent * MeshComp, class UAnimSequenceBase * Animation)
 {
 	bool bError = ValidateInput(MeshComp);
 
@@ -120,10 +120,10 @@ void UAnimNotifyState_Trail::NotifyBegin(class USkeletalMeshComponent * MeshComp
 		}
 	}
 
-	Received_NotifyBegin(MeshComp, AnimSeq);
+	Received_NotifyBegin(MeshComp, Animation);
 }
 
-void UAnimNotifyState_Trail::NotifyTick(class USkeletalMeshComponent * MeshComp, class UAnimSequence * AnimSeq, float FrameDeltaTime)
+void UAnimNotifyState_Trail::NotifyTick(class USkeletalMeshComponent * MeshComp, class UAnimSequenceBase * Animation, float FrameDeltaTime)
 {
 	bool bError = ValidateInput(MeshComp, true);
 
@@ -162,10 +162,10 @@ void UAnimNotifyState_Trail::NotifyTick(class USkeletalMeshComponent * MeshComp,
 		}
 	}
 
-	Received_NotifyTick(MeshComp, AnimSeq, FrameDeltaTime);
+	Received_NotifyTick(MeshComp, Animation, FrameDeltaTime);
 }
 
-void UAnimNotifyState_Trail::NotifyEnd(class USkeletalMeshComponent * MeshComp, class UAnimSequence * AnimSeq)
+void UAnimNotifyState_Trail::NotifyEnd(class USkeletalMeshComponent * MeshComp, class UAnimSequenceBase * Animation)
 {
 	TArray<USceneComponent*> Children;
 	MeshComp->GetChildrenComponents(false, Children);
@@ -184,7 +184,7 @@ void UAnimNotifyState_Trail::NotifyEnd(class USkeletalMeshComponent * MeshComp,
 		}
 	}
 
-	Received_NotifyEnd(MeshComp, AnimSeq);
+	Received_NotifyEnd(MeshComp, Animation);
 }
 
 bool UAnimNotifyState_Trail::ValidateInput(class USkeletalMeshComponent * MeshComp, bool bReportErrors/* =false */)
