UE4: Fixed a stack overflow crash that happened when entering PIE.

diff --git a/Engine/Source/Runtime/Engine/Classes/Components/SceneComponent.h b/Engine/Source/Runtime/Engine/Classes/Components/SceneComponent.h
--- a/Engine/Source/Runtime/Engine/Classes/Components/SceneComponent.h
+++ b/Engine/Source/Runtime/Engine/Classes/Components/SceneComponent.h
@@ -724,7 +724,7 @@ protected:
 	void UpdateNavigationData();
 
 	/** Check if given component or any of its attachments are relevant for navigation updates */
-	bool CheckNavigationRelevancy(USceneComponent* TestComponent);
+	static bool CheckNavigationRelevancy(USceneComponent* TestComponent);
 };
 
 
diff --git a/Engine/Source/Runtime/Engine/Private/SceneComponent.cpp b/Engine/Source/Runtime/Engine/Private/SceneComponent.cpp
--- a/Engine/Source/Runtime/Engine/Private/SceneComponent.cpp
+++ b/Engine/Source/Runtime/Engine/Private/SceneComponent.cpp
@@ -1720,12 +1720,15 @@ void USceneComponent::UpdateNavigationData()
 
 bool USceneComponent::CheckNavigationRelevancy(USceneComponent* TestComponent)
 {
+	check(TestComponent);
+
 	bool bResult = TestComponent->IsNavigationRelevant();
-	for (int32 Idx = 0; Idx < AttachChildren.Num() && !bResult; Idx++)
+	for (int32 Idx = 0; Idx < TestComponent->AttachChildren.Num() && !bResult; Idx++)
 	{
-		if (AttachChildren[Idx])
+		USceneComponent* Child = TestComponent->AttachChildren[Idx];
+		if (Child)
 		{
-			bResult = CheckNavigationRelevancy(AttachChildren[Idx]);
+			bResult = CheckNavigationRelevancy(Child);
 		}
 	}
 
