UE4: Shared landscape textures and materials whose outers are the world's package are now renamed and duplicated to the new package. These assets were recently moved back to being inners of the package since there were performance implications in duplicating worlds with many landscape components for PIE.

diff --git a/Engine/Source/Runtime/Engine/Classes/Engine/World.h b/Engine/Source/Runtime/Engine/Classes/Engine/World.h
--- a/Engine/Source/Runtime/Engine/Classes/Engine/World.h
+++ b/Engine/Source/Runtime/Engine/Classes/Engine/World.h
@@ -2389,6 +2389,9 @@ public:
 	/** Gets all LightMaps and ShadowMaps associated with this world. Specify the level or leave null for persistent */
 	void GetLightMapsAndShadowMaps(ULevel* Level, TArray<UTexture2D*>& OutLightMapsAndShadowMaps);
 
+	/** Gets all textures and materials used by all landscape components in the specified level */
+	void GetLandscapeTexturesAndMaterials(ULevel* Level, TArray<UObject*>& OutTexturesAndMaterials);
+
 public:
 	static FString ConvertToPIEPackageName(const FString& PackageName, int32 PIEInstanceID);
 	static FString BuildPIEPackagePrefix(int32 PIEInstanceID);
diff --git a/Engine/Source/Runtime/Engine/Classes/Landscape/LandscapeComponent.h b/Engine/Source/Runtime/Engine/Classes/Landscape/LandscapeComponent.h
--- a/Engine/Source/Runtime/Engine/Classes/Landscape/LandscapeComponent.h
+++ b/Engine/Source/Runtime/Engine/Classes/Landscape/LandscapeComponent.h
@@ -306,6 +306,9 @@ public:
 	/** Get the level in which the owning actor resides */
 	class ULevel* GetLevel() const;
 
+	/** Returns all textures and materials used by this component. */
+	ENGINE_API void GetAllReferencedTexturesAndMaterials(TArray<UObject*>& OutTexturesAndMaterials) const;
+
 	/** @todo document */
 	ENGINE_API class ALandscapeProxy* GetLandscapeProxy() const;
 
diff --git a/Engine/Source/Runtime/Engine/Private/Landscape/Landscape.cpp b/Engine/Source/Runtime/Engine/Private/Landscape/Landscape.cpp
--- a/Engine/Source/Runtime/Engine/Private/Landscape/Landscape.cpp
+++ b/Engine/Source/Runtime/Engine/Private/Landscape/Landscape.cpp
@@ -431,16 +431,7 @@ void ULandscapeComponent::PostLoad()
 			if (ensure(Level))
 			{
 				TArray<UObject*> ObjectsToMoveFromLevelToPackage;
-				for (UMaterialInstance* CurrentMIC = MaterialInstance; CurrentMIC; CurrentMIC = Cast<UMaterialInstance>(CurrentMIC->Parent))
-				{
-					ObjectsToMoveFromLevelToPackage.Add(CurrentMIC);
-				}
-				ObjectsToMoveFromLevelToPackage.Add(HeightmapTexture);
-				for (auto* Tex : WeightmapTextures)
-				{
-					ObjectsToMoveFromLevelToPackage.Add(Tex);
-				}
-				ObjectsToMoveFromLevelToPackage.Add(XYOffsetmapTexture);
+				GetAllReferencedTexturesAndMaterials(ObjectsToMoveFromLevelToPackage);
 
 				UPackage* MyPackage = GetOutermost();
 				for (auto* Obj : ObjectsToMoveFromLevelToPackage)
@@ -578,6 +569,35 @@ ULevel* ULandscapeComponent::GetLevel() const
 	return MyOwner ? MyOwner->GetLevel() : NULL;
 }
 
+void ULandscapeComponent::GetAllReferencedTexturesAndMaterials(TArray<UObject*>& OutTexturesAndMaterials) const
+{
+	for (UMaterialInstance* CurrentMIC = MaterialInstance; CurrentMIC; CurrentMIC = Cast<UMaterialInstance>(CurrentMIC->Parent))
+	{
+		if (CurrentMIC)
+		{
+			OutTexturesAndMaterials.Add(CurrentMIC);
+		}
+	}
+
+	if (HeightmapTexture)
+	{
+		OutTexturesAndMaterials.Add(HeightmapTexture);
+	}
+
+	for (auto* Tex : WeightmapTextures)
+	{
+		if (Tex)
+		{
+			OutTexturesAndMaterials.Add(Tex);
+		}
+	}
+
+	if (XYOffsetmapTexture)
+	{
+		OutTexturesAndMaterials.Add(XYOffsetmapTexture);
+	}
+}
+
 ALandscapeProxy* ULandscapeComponent::GetLandscapeProxy() const
 {
 	return CastChecked<ALandscapeProxy>(GetOuter());
diff --git a/Engine/Source/Runtime/Engine/Private/World.cpp b/Engine/Source/Runtime/Engine/Private/World.cpp
--- a/Engine/Source/Runtime/Engine/Private/World.cpp
+++ b/Engine/Source/Runtime/Engine/Private/World.cpp
@@ -45,6 +45,7 @@
 
 #include "EngineModule.h"
 #include "ContentStreaming.h"
+#include "Runtime/Engine/Classes/Landscape/LandscapeComponent.h"
 
 DEFINE_LOG_CATEGORY_STATIC(LogWorld, Log, All);
 DEFINE_LOG_CATEGORY(LogSpawn);
@@ -216,6 +217,8 @@ void UWorld::AddReferencedObjects(UObject* InThis, FReferenceCollector& Collecto
 #if WITH_EDITOR
 bool UWorld::Rename(const TCHAR* InName, UObject* NewOuter, ERenameFlags Flags)
 {
+	check(PersistentLevel);
+
 	// Rename LightMaps and ShadowMaps to the new location. Keep the old name, since they are not named after the world.
 	TArray<UTexture2D*> LightMapsAndShadowMaps;
 	GetLightMapsAndShadowMaps(PersistentLevel, LightMapsAndShadowMaps);
@@ -227,6 +230,22 @@ bool UWorld::Rename(const TCHAR* InName, UObject* NewOuter, ERenameFlags Flags)
 		}
 	}
 
+	// Also rename all textures and materials used by landscape components
+	TArray<UObject*> LandscapeTexturesAndMaterials;
+	GetLandscapeTexturesAndMaterials(PersistentLevel, LandscapeTexturesAndMaterials);
+	UPackage* PersistentLevelPackage = PersistentLevel->GetOutermost();
+	for (auto* OldTexOrMat : LandscapeTexturesAndMaterials)
+	{
+		if (OldTexOrMat && OldTexOrMat->GetOuter() == PersistentLevelPackage)
+		{
+			// The names for these objects are not important, just generate a new name to avoid collisions
+			if ( !OldTexOrMat->Rename(nullptr, NewOuter, Flags) )
+			{
+				return false;
+			}
+		}
+	}
+
 	// Rename the world itself
 	UPackage* OldPackage = GetOutermost();
 	if ( !Super::Rename(InName, NewOuter, Flags) )
@@ -240,7 +259,7 @@ bool UWorld::Rename(const TCHAR* InName, UObject* NewOuter, ERenameFlags Flags)
 	if ( !GIsRoutingPostLoad )
 	{
 		const bool bDontCreate = true;
-		UBlueprint* LevelScriptBlueprint = PersistentLevel ? PersistentLevel->GetLevelScriptBlueprint(bDontCreate) : NULL;
+		UBlueprint* LevelScriptBlueprint = PersistentLevel->GetLevelScriptBlueprint(bDontCreate);
 		if ( LevelScriptBlueprint )
 		{
 			// See if we are just testing for a rename. When testing, the world hasn't actually changed outers, so we need to test for name collisions at the target outer.
@@ -289,6 +308,8 @@ void UWorld::PostDuplicate(bool bDuplicateForPIE)
 	// Note that these objects should really be inners of the world, so if they become inners later, most of this code should not be necessary
 	if ( !bDuplicateForPIE )
 	{
+		check(PersistentLevel);
+
 		// Update the persistent level's owning world. This is needed for some initialization
 		if ( !PersistentLevel->OwningWorld )
 		{
@@ -301,9 +322,18 @@ void UWorld::PostDuplicate(bool bDuplicateForPIE)
 			CurrentLevel = PersistentLevel;
 		}
 
-		TMap<UObject*, UObject*> ReplacementMap;
 		UPackage* MyPackage = GetOutermost();
 
+		// Make sure PKG_ContainsMap is set
+		MyPackage->ThisContainsMap();
+
+#if WITH_EDITOR
+		TArray<UObject*> ObjectsToFixReferences;
+		TMap<UObject*, UObject*> ReplacementMap;
+
+		// Add the world to the list of objects in which to fix up references.
+		ObjectsToFixReferences.Add(this);
+
 		// Gather the textures
 		TArray<UTexture2D*> LightMapsAndShadowMaps;
 		GetLightMapsAndShadowMaps(PersistentLevel, LightMapsAndShadowMaps);
@@ -318,27 +348,50 @@ void UWorld::PostDuplicate(bool bDuplicateForPIE)
 			}
 		}
 
-#if WITH_EDITOR
+		// Also duplicate all textures and materials used by landscape components
+		TArray<UObject*> LandscapeTexturesAndMaterials;
+		GetLandscapeTexturesAndMaterials(PersistentLevel, LandscapeTexturesAndMaterials);
+		for (auto* OldTexOrMat : LandscapeTexturesAndMaterials)
+		{
+			if (OldTexOrMat && OldTexOrMat->GetOuter() != MyPackage)
+			{
+				// The names for these objects are not important, just generate a new name to avoid collisions
+				UObject* NewTextureOrMaterial = StaticDuplicateObject(OldTexOrMat, MyPackage, nullptr);
+				ReplacementMap.Add(OldTexOrMat, NewTextureOrMaterial);
+
+				// Materials hold references to the textures being moved, so they will need references fixed up as well
+				if ( OldTexOrMat->IsA(UMaterialInterface::StaticClass()) )
+				{
+					ObjectsToFixReferences.Add(NewTextureOrMaterial);
+				}
+			}
+		}
+
 		// Duplicate the level script blueprint generated classes as well
 		const bool bDontCreate = true;
-		UBlueprint* LevelScriptBlueprint = PersistentLevel ? PersistentLevel->GetLevelScriptBlueprint(bDontCreate) : NULL;
-		if ( LevelScriptBlueprint )
+		UBlueprint* LevelScriptBlueprint = PersistentLevel->GetLevelScriptBlueprint(bDontCreate);
+		if (LevelScriptBlueprint)
 		{
 			UObject* OldGeneratedClass = LevelScriptBlueprint->GeneratedClass;
-			if ( OldGeneratedClass )
+			if (OldGeneratedClass)
 			{
 				UObject* NewGeneratedClass = StaticDuplicateObject(OldGeneratedClass, MyPackage, *OldGeneratedClass->GetName());
 				ReplacementMap.Add(OldGeneratedClass, NewGeneratedClass);
+
+				// The class may have referenced a lightmap or landscape resource that is also being duplicated. Add it to the list of objects that need references fixed up.
+				ObjectsToFixReferences.Add(NewGeneratedClass);
 			}
 
 			UObject* OldSkeletonClass = LevelScriptBlueprint->SkeletonGeneratedClass;
-			if ( OldSkeletonClass )
+			if (OldSkeletonClass)
 			{
 				UObject* NewSkeletonClass = StaticDuplicateObject(OldSkeletonClass, MyPackage, *OldSkeletonClass->GetName());
 				ReplacementMap.Add(OldSkeletonClass, NewSkeletonClass);
+
+				// The class may have referenced a lightmap or landscape resource that is also being duplicated. Add it to the list of objects that need references fixed up.
+				ObjectsToFixReferences.Add(NewSkeletonClass);
 			}
 		}
-#endif // WITH_EDITOR
 
 		// Now replace references from the old textures/classes to the new ones, if any were duplicated
 		if (ReplacementMap.Num() > 0)
@@ -346,11 +399,12 @@ void UWorld::PostDuplicate(bool bDuplicateForPIE)
 			const bool bNullPrivateRefs = false;
 			const bool bIgnoreOuterRef = true;
 			const bool bIgnoreArchetypeRef = false;
-			FArchiveReplaceObjectRef<UObject> ReplaceAr(this, ReplacementMap, bNullPrivateRefs, bIgnoreOuterRef, bIgnoreArchetypeRef);
+			for ( auto* Obj : ObjectsToFixReferences )
+			{
+				FArchiveReplaceObjectRef<UObject> ReplaceAr(Obj, ReplacementMap, bNullPrivateRefs, bIgnoreOuterRef, bIgnoreArchetypeRef);
+			}
 		}
-
-		// Make sure PKG_ContainsMap is set
-		MyPackage->ThisContainsMap();
+#endif // WITH_EDITOR
 	}
 }
 
@@ -5038,6 +5092,27 @@ void UWorld::GetLightMapsAndShadowMaps(ULevel* Level, TArray<UTexture2D*>& OutLi
 	FFindLightmapsArchive FindArchive(SearchObject, OutLightMapsAndShadowMaps);
 }
 
+void UWorld::GetLandscapeTexturesAndMaterials(ULevel* Level, TArray<UObject*>& OutTexturesAndMaterials)
+{
+	UObject* SearchLevel = Level;
+	if (!SearchLevel)
+	{
+		SearchLevel = PersistentLevel;
+	}
+
+	TArray<UObject*> ObjectsInLevel;
+	const bool bIncludeNestedObjects = true;
+	GetObjectsWithOuter(SearchLevel, ObjectsInLevel, bIncludeNestedObjects);
+	for (auto* ObjInLevel : ObjectsInLevel)
+	{
+		ULandscapeComponent* LandscapeComponent = Cast<ULandscapeComponent>(ObjInLevel);
+		if (LandscapeComponent)
+		{
+			LandscapeComponent->GetAllReferencedTexturesAndMaterials(OutTexturesAndMaterials);
+		}
+	}
+}
+
 void UWorld::ChangeFeatureLevel(ERHIFeatureLevel::Type InFeatureLevel)
 {
 	if (InFeatureLevel != FeatureLevel)
