Incremental components registration during level streaming. Granularity of components registration can be adjusted using LevelStreamingComponentsRegistrationGranularity parameter in a BaseEngine.ini (currently 10)

diff --git a/Engine/Config/BaseEngine.ini b/Engine/Config/BaseEngine.ini
--- a/Engine/Config/BaseEngine.ini
+++ b/Engine/Config/BaseEngine.ini
@@ -119,7 +119,7 @@ MaxES2PixelShaderAdditiveComplexityCount=600
 TimeBetweenPurgingPendingKillObjects=60
 bUseBackgroundLevelStreaming=True
 LevelStreamingActorsUpdateTimeLimit = 5.0
-LevelStreamingNumActorsToUpdate = 50
+LevelStreamingComponentsRegistrationGranularity = 10
 bSubtitlesEnabled=True
 bSubtitlesForcedOff=false
 DefaultSoundName=/Engine/EngineSounds/WhiteNoise.WhiteNoise
diff --git a/Engine/Source/Runtime/Engine/Classes/Engine/Engine.h b/Engine/Source/Runtime/Engine/Classes/Engine/Engine.h
--- a/Engine/Source/Runtime/Engine/Classes/Engine/Engine.h
+++ b/Engine/Source/Runtime/Engine/Classes/Engine/Engine.h
@@ -1105,9 +1105,9 @@ public:
 	UPROPERTY(EditAnywhere, config, Category=LevelStreaming, AdvancedDisplay)
 	float LevelStreamingActorsUpdateTimeLimit;
 	
-	/** Batching granularity used to register actors during level streaming */
+	/** Batching granularity used to register actor components during level streaming */
 	UPROPERTY(EditAnywhere, config, Category=LevelStreaming, AdvancedDisplay)
-	int32 LevelStreamingNumActorsToUpdate;
+	int32 LevelStreamingComponentsRegistrationGranularity;
 	
 	/** @todo document */
 	UPROPERTY(config)
diff --git a/Engine/Source/Runtime/Engine/Classes/Engine/Level.h b/Engine/Source/Runtime/Engine/Classes/Engine/Level.h
--- a/Engine/Source/Runtime/Engine/Classes/Engine/Level.h
+++ b/Engine/Source/Runtime/Engine/Classes/Engine/Level.h
@@ -505,10 +505,10 @@ public:
 	/**
 	 * Incrementally updates all components of actors associated with this level.
 	 *
-	 * @param NumActorsToUpdate		Number of actors to update in this run, 0 for all
+	 * @param NumComponentsToUpdate		Number of components to update in this run, 0 for all
 	 * @param bRerunConstructionScripts	If we want to rerun construction scripts on actors in level
 	 */
-	void IncrementalUpdateComponents( int32 NumActorsToUpdate, bool bRerunConstructionScripts );
+	void IncrementalUpdateComponents( int32 NumComponentsToUpdate, bool bRerunConstructionScripts );
 
 	/**
 	 * Invalidates the cached data used to render the level's UModel.
diff --git a/Engine/Source/Runtime/Engine/Classes/GameFramework/Actor.h b/Engine/Source/Runtime/Engine/Classes/GameFramework/Actor.h
--- a/Engine/Source/Runtime/Engine/Classes/GameFramework/Actor.h
+++ b/Engine/Source/Runtime/Engine/Classes/GameFramework/Actor.h
@@ -1635,6 +1635,14 @@ public:
 	/** Will reregister all components on this actor. Does a lot of work - should only really be used in editor, generally use UpdateComponentTransforms or MarkComponentsRenderStateDirty. */
 	virtual void ReregisterAllComponents();
 
+	/**
+	 * Incrementally registers components associated with this actor
+	 *
+	 * @param NumComponentsToRegister  Number of components to register in this run, 0 for all
+	 * @return true when all components were registered for this actor
+	 */
+	bool IncrementalRegisterComponents(int32 NumComponentsToRegister);
+
 	/** Flags all component's render state as dirty	 */
 	void MarkComponentsRenderStateDirty();
 
diff --git a/Engine/Source/Runtime/Engine/Private/Actor.cpp b/Engine/Source/Runtime/Engine/Private/Actor.cpp
--- a/Engine/Source/Runtime/Engine/Private/Actor.cpp
+++ b/Engine/Source/Runtime/Engine/Private/Actor.cpp
@@ -2093,6 +2093,53 @@ void AActor::DisableComponentsSimulatePhysics()
 
 void AActor::PostRegisterAllComponents() 
 {
+	// Try and get transforms correct.
+	// Obviously a component could be attached to something in another actor that has not been processed yet, that will just move the registered components to the correct location
+
+	// We refetch the components because the registration process could have modified the list
+	TArray<USceneComponent*> SceneComponents;
+	GetComponents(SceneComponents);
+
+	for(int32 CompIdx=0; CompIdx < SceneComponents.Num(); CompIdx++)
+	{
+		USceneComponent* SceneComp = SceneComponents[CompIdx];
+#if !(UE_BUILD_SHIPPING || UE_BUILD_TEST)
+		AActor* CompOwner = SceneComp->GetOwner();
+		FString CompOwnerPathName = (CompOwner != NULL) ? CompOwner->GetFullName() : TEXT("NONE");
+		checkf(SceneComp->GetOwner() == this, TEXT("SceneComp '%s' has Owner of %x'%s', expected %x'%s' (TTP 241544)"), *SceneComp->GetFullName(), CompOwner, *CompOwnerPathName, this, *this->GetFullName());
+#endif
+		//UE_LOG(LogActorComponent, Log, TEXT("%d: %s (AP: %s)"), CompIdx, *SceneComp->GetPathName(), SceneComp->AttachParent ? *SceneComp->AttachParent->GetPathName() : TEXT("NONE"));
+		// If we need to perform a call to AttachTo to add it to AttachChildren array, do that now
+		if(SceneComp->AttachParent != NULL && !SceneComp->AttachParent->AttachChildren.Contains(SceneComp))
+		{
+			USceneComponent* PendingAttachParent = SceneComp->AttachParent;
+			SceneComp->AttachParent = NULL;
+
+			// If our attach parent is a blueprint component that is marked pending kill belonging to another actor,
+			// then rerun construction script has almost certainly been run on it and we are just going to pick its
+			// current root component instead
+			if (PendingAttachParent->bCreatedByConstructionScript && PendingAttachParent->IsPendingKill())
+			{
+				AActor* PendingAttachOwner = PendingAttachParent->GetOwner();
+				if (PendingAttachOwner != this)
+				{
+					PendingAttachParent = PendingAttachOwner->GetRootComponent();
+				}
+			}
+
+			FName PendingAttachSocketName = SceneComp->AttachSocketName;
+			SceneComp->AttachSocketName = NAME_None;
+
+			//UE_LOG(LogActorComponent, Log, TEXT("  Attached!"));
+			SceneComp->AttachTo(PendingAttachParent, PendingAttachSocketName);
+		}
+		// Otherwise just update its ComponentToWorld
+		else
+		{
+			SceneComp->UpdateComponentToWorld();
+		}
+	}
+		
 	if (bNavigationRelevant == true && GetWorld() != NULL 
 		&& GetWorld()->GetNavigationSystem() != NULL)
 	{
@@ -3026,8 +3073,18 @@ void AActor::UnregisterAllComponents()
 
 void AActor::RegisterAllComponents()
 {
-	// Second phase, register the components
+	// 0 - means register all components
+	verify(IncrementalRegisterComponents(0));
+}
 
+bool AActor::IncrementalRegisterComponents(int32 NumComponentsToRegister)
+{
+	if (NumComponentsToRegister == 0)
+	{
+		// 0 - means register all components
+		NumComponentsToRegister = MAX_int32;
+	}
+	
 	// Register RootComponent first so all other components can reliable use it (ie call GetLocation) when they register
 	if( RootComponent != NULL && !RootComponent->IsRegistered() )
 	{
@@ -3042,10 +3099,12 @@ void AActor::RegisterAllComponents()
 		RootComponent->RegisterComponentWithWorld(GetWorld());
 	}
 
+	int32 NumTotalRegisteredComponents = 0;
+	int32 NumRegisteredComponentsThisRun = 0;
 	TArray<UActorComponent*> Components;
 	GetComponents(Components);
-
-	for(int32 CompIdx=0; CompIdx < Components.Num(); CompIdx++)
+	
+	for (int32 CompIdx = 0; CompIdx < Components.Num() && NumRegisteredComponentsThisRun < NumComponentsToRegister; CompIdx++)
 	{
 		UActorComponent* Component = Components[CompIdx];
 		if(!Component->IsRegistered() && Component->bAutoRegister && !Component->IsPendingKill())
@@ -3056,58 +3115,22 @@ void AActor::RegisterAllComponents()
 
 			check(GetWorld());
 			Component->RegisterComponentWithWorld(GetWorld());
+			NumRegisteredComponentsThisRun++;
 		}
-	}
-
-	// First phase, try and get transforms correct.
-	// Obviously a component could be attached to something in another actor that has not been processed yet, that will just move the registered components to the correct location
 
-	// We refetch the components because the registration process could have modified the list
-	TArray<USceneComponent*> SceneComponents;
-	GetComponents(SceneComponents);
+		NumTotalRegisteredComponents++;
+	}
 
-	for(int32 CompIdx=0; CompIdx < SceneComponents.Num(); CompIdx++)
+	// See whether we are done
+	if (Components.Num() == NumTotalRegisteredComponents)
 	{
-		USceneComponent* SceneComp = SceneComponents[CompIdx];
-#if !(UE_BUILD_SHIPPING || UE_BUILD_TEST)
-			AActor* CompOwner = SceneComp->GetOwner();
-			FString CompOwnerPathName = (CompOwner != NULL) ? CompOwner->GetFullName() : TEXT("NONE");
-			checkf(SceneComp->GetOwner() == this, TEXT("SceneComp '%s' has Owner of %x'%s', expected %x'%s' (TTP 241544)"), *SceneComp->GetFullName(), CompOwner, *CompOwnerPathName, this, *this->GetFullName());
-#endif
-			//UE_LOG(LogActorComponent, Log, TEXT("%d: %s (AP: %s)"), CompIdx, *SceneComp->GetPathName(), SceneComp->AttachParent ? *SceneComp->AttachParent->GetPathName() : TEXT("NONE"));
-			// If we need to perform a call to AttachTo to add it to AttachChildren array, do that now
-			if(SceneComp->AttachParent != NULL && !SceneComp->AttachParent->AttachChildren.Contains(SceneComp))
-			{
-				USceneComponent* PendingAttachParent = SceneComp->AttachParent;
-				SceneComp->AttachParent = NULL;
-
-				// If our attach parent is a blueprint component that is marked pending kill belonging to another actor,
-				// then rerun construction script has almost certainly been run on it and we are just going to pick its
-				// current root component instead
-				if (PendingAttachParent->bCreatedByConstructionScript && PendingAttachParent->IsPendingKill())
-				{
-					AActor* PendingAttachOwner = PendingAttachParent->GetOwner();
-					if (PendingAttachOwner != this)
-					{
-						PendingAttachParent = PendingAttachOwner->GetRootComponent();
-					}
-				}
-
-				FName PendingAttachSocketName = SceneComp->AttachSocketName;
-				SceneComp->AttachSocketName = NAME_None;
-
-				//UE_LOG(LogActorComponent, Log, TEXT("  Attached!"));
-				SceneComp->AttachTo(PendingAttachParent, PendingAttachSocketName);
-			}
-			// Otherwise just update its ComponentToWorld
-			else
-			{
-				SceneComp->UpdateComponentToWorld();
-			}
-		}
-
-	// Finally, call PostRegisterAllComponents
-	PostRegisterAllComponents();
+		// Finally, call PostRegisterAllComponents
+		PostRegisterAllComponents();
+		return true;
+	}
+	
+	// Still have components to register
+	return false;
 }
 
 bool AActor::HasValidRootComponent()
diff --git a/Engine/Source/Runtime/Engine/Private/Level.cpp b/Engine/Source/Runtime/Engine/Private/Level.cpp
--- a/Engine/Source/Runtime/Engine/Private/Level.cpp
+++ b/Engine/Source/Runtime/Engine/Private/Level.cpp
@@ -728,91 +728,76 @@ struct FModelComponentKey
 	}
 };
 
-
 void ULevel::UpdateLevelComponents(bool bRerunConstructionScripts)
 {
 	// Update all components in one swoop.
 	IncrementalUpdateComponents( 0, bRerunConstructionScripts );
 }
 
-
-void ULevel::IncrementalUpdateComponents(int32 NumActorsToUpdate, bool bRerunConstructionScripts)
+void ULevel::IncrementalUpdateComponents(int32 NumComponentsToUpdate, bool bRerunConstructionScripts)
 {
 	// A value of 0 means that we want to update all components.
-	if( NumActorsToUpdate == 0 )
-	{
-		NumActorsToUpdate = MAX_int32;
-	}
-	// Only the game can use incremental update functionality.
-	else
+	if (NumComponentsToUpdate != 0)
 	{
+		// Only the game can use incremental update functionality.
 		checkf(!GIsEditor && OwningWorld->IsGameWorld(),TEXT("Cannot call IncrementalUpdateComponents with non 0 argument in the Editor/ commandlets."));
 	}
 
 	// Do BSP on the first pass.
-	if( CurrentActorIndexForUpdateComponents == 0 )
+	if (CurrentActorIndexForUpdateComponents == 0)
 	{
 		UpdateModelComponents();
 	}
 
-	TArray<AActor*> RerunConstructionScriptActors;
-	for( int32 i=0; i < NumActorsToUpdate && CurrentActorIndexForUpdateComponents < Actors.Num(); i++ )
+	// Find next valid actor to process components registration
+	while (CurrentActorIndexForUpdateComponents < Actors.Num())
 	{
-		AActor* Actor = Actors[CurrentActorIndexForUpdateComponents++];
-		if( Actor )
+		AActor* Actor = Actors[CurrentActorIndexForUpdateComponents];
+		bool bAllComponentsRegistered = true;
+		if (Actor)
 		{
-#if PERF_TRACK_DETAILED_ASYNC_STATS
-			double Start = FPlatformTime::Seconds();
-#endif
-
-			Actor->ReregisterAllComponents();
-
-			// Cache the actor to rerun the construction script (if the actor is blueprint based and this is being run in the editor)
-			// @TODO: Only do this if blueprint has changed!
-			if (bRerunConstructionScripts && !IsTemplate() && !GIsUCCMakeStandaloneHeaderGenerator)
-			{
-				RerunConstructionScriptActors.Add(Actor);
-			}
+			bAllComponentsRegistered = Actor->IncrementalRegisterComponents(NumComponentsToUpdate);
+		}
 
-#if PERF_TRACK_DETAILED_ASYNC_STATS
-			// Add how long this took to class->time map
-			double Time = FPlatformTime::Seconds() - Start;
-			UClass* ActorClass = Actor->GetClass();
-			FMapTimeEntry* CurrentEntry = UpdateComponentsTimePerActorClass.Find(ActorClass);
-			// Is an existing entry - add to it
-			if(CurrentEntry)
-			{
-				CurrentEntry->Time += Time;
-				CurrentEntry->ObjCount += 1;
-			}
-			// Make a new entry for this class
-			else
-			{
-				UpdateComponentsTimePerActorClass.Add(ActorClass, FMapTimeEntry(ActorClass, 1, Time));
-			}
-#endif
+		if (bAllComponentsRegistered)
+		{	
+			// All components have been registered fro this actor, move to a next one
+			CurrentActorIndexForUpdateComponents++;
 		}
-	}
 
-	// Don't rerun construction scripts until after all actors' components have been registered.  This
-	// is necessary because child attachment lists are populated during registration, and running construction
-	// scripts requires that the attachments are correctly initialized.
-	for (auto RerunConstructionScriptActor : RerunConstructionScriptActors)
-	{
-		RerunConstructionScriptActor->RerunConstructionScripts();
+		// If we do an incremental registration return to outer loop after each processed actor 
+		// so outer loop can decide whether we want to continue processing this frame
+		if (NumComponentsToUpdate != 0)
+		{
+			break;
+		}
 	}
 
 	// See whether we are done.
-	if( CurrentActorIndexForUpdateComponents == Actors.Num() )
+	if (CurrentActorIndexForUpdateComponents == Actors.Num())
 	{
 		CurrentActorIndexForUpdateComponents	= 0;
 		bAreComponentsCurrentlyRegistered		= true;
+		
+		if (bRerunConstructionScripts && !IsTemplate() && !GIsUCCMakeStandaloneHeaderGenerator)
+		{
+			// Don't rerun construction scripts until after all actors' components have been registered.  This
+			// is necessary because child attachment lists are populated during registration, and running construction
+			// scripts requires that the attachments are correctly initialized.
+			for (AActor* Actor : Actors)
+			{
+				if (Actor)
+				{
+					Actor->RerunConstructionScripts();
+				}
+			}
+		}
 	}
 	// Only the game can use incremental update functionality.
 	else
 	{
 		// The editor is never allowed to incrementally updated components.  Make sure to pass in a value of zero for NumActorsToUpdate.
-		check( OwningWorld->IsGameWorld() );
+		check(OwningWorld->IsGameWorld());
 	}
 }
 
diff --git a/Engine/Source/Runtime/Engine/Private/World.cpp b/Engine/Source/Runtime/Engine/Private/World.cpp
--- a/Engine/Source/Runtime/Engine/Private/World.cpp
+++ b/Engine/Source/Runtime/Engine/Private/World.cpp
@@ -1587,6 +1587,7 @@ double InitActorTime = 0.0;
 double RouteActorInitializeTime = 0.0;
 double CrossLevelRefsTime = 0.0;
 double SortActorListTime = 0.0;
+double PerformLastStepTime = 0.0;
 
 /** Helper class, to add the time between this objects creating and destruction to passed in variable. */
 class FAddWorldScopeTimeVar
@@ -1660,6 +1661,7 @@ void UWorld::AddToWorld( ULevel* Level, const FTransform& LevelTransform )
 		RouteActorInitializeTime = 0.0;
 		CrossLevelRefsTime = 0.0;
 		SortActorListTime = 0.0;
+		PerformLastStepTime = 0.0;
 #endif
 	}
 
@@ -1696,7 +1698,7 @@ void UWorld::AddToWorld( ULevel* Level, const FTransform& LevelTransform )
 		Level->bAreComponentsCurrentlyRegistered = false;
 
 		// Incrementally update components.
-		int32 NumActorsToUpdate = GEngine->LevelStreamingNumActorsToUpdate;
+		int32 NumComponentsToUpdate = GEngine->LevelStreamingComponentsRegistrationGranularity;
 		do
 		{
 #if WITH_EDITOR
@@ -1706,7 +1708,7 @@ void UWorld::AddToWorld( ULevel* Level, const FTransform& LevelTransform )
 			// We don't need to rerun construction scripts if we have cooked data or we are playing in editor unless the PIE world was loaded
 			// from disk rather than duplicated
 			const bool bRerunConstructionScript = !FPlatformProperties::RequiresCookedData() && (!IsPlayInEditor() || HasAnyFlags(RF_WasLoaded));
-			Level->IncrementalUpdateComponents( (GIsEditor || IsRunningCommandlet()) ? 0 : NumActorsToUpdate, bRerunConstructionScript );
+			Level->IncrementalUpdateComponents( (GIsEditor || IsRunningCommandlet()) ? 0 : NumComponentsToUpdate, bRerunConstructionScript );
 		}
 		while( (!bConsiderTimeLimit || !IsTimeLimitExceeded( TEXT("updating components"), StartTime, Level )) && !Level->bAreComponentsCurrentlyRegistered );
 
@@ -1760,25 +1762,8 @@ void UWorld::AddToWorld( ULevel* Level, const FTransform& LevelTransform )
 	// We're done.
 	if( bPerformedLastStep )
 	{
-#if PERF_TRACK_DETAILED_ASYNC_STATS
-		// Log out all of the timing information
-		double TotalTime = MoveActorTime + ShiftActorsTime + UpdateComponentsTime + InitBSPPhysTime + InitActorPhysTime + InitActorTime + RouteActorInitializeTime + CrossLevelRefsTime + SortActorListTime;
-		UE_LOG(LogStreaming, Log, TEXT("Detailed AddToWorld stats for '%s' - Total %6.2fms"), *Level->GetOutermost()->GetName(), TotalTime * 1000 );
-		UE_LOG(LogStreaming, Log, TEXT("Move Actors             : %6.2f ms"), MoveActorTime * 1000 );
-		UE_LOG(LogStreaming, Log, TEXT("Shift Actors            : %6.2f ms"), ShiftActorsTime * 1000 );
-		UE_LOG(LogStreaming, Log, TEXT("Update Components       : %6.2f ms"), UpdateComponentsTime * 1000 );
-
-		PrintSortedListFromMap(Level->UpdateComponentsTimePerActorClass);
-		Level->UpdateComponentsTimePerActorClass.Empty();
-
-		UE_LOG(LogStreaming, Log, TEXT("Init BSP Phys           : %6.2f ms"), InitBSPPhysTime * 1000 );
-		UE_LOG(LogStreaming, Log, TEXT("Init Actor Phys         : %6.2f ms"), InitActorPhysTime * 1000 );
-		UE_LOG(LogStreaming, Log, TEXT("Init Actors             : %6.2f ms"), InitActorTime * 1000 );
-		UE_LOG(LogStreaming, Log, TEXT("Initialize               : %6.2f ms"), RouteActorInitializeTime * 1000 );
-		UE_LOG(LogStreaming, Log, TEXT("Cross Level Refs        : %6.2f ms"), CrossLevelRefsTime * 1000 );
-		UE_LOG(LogStreaming, Log, TEXT("Sort Actor List         : %6.2f ms"), SortActorListTime * 1000 );	
-#endif // PERF_TRACK_DETAILED_ASYNC_STATS
-
+		SCOPE_TIME_TO_VAR(&PerformLastStepTime);
+		
 		Level->bAlreadyShiftedActors					= false;
 		Level->bAlreadyUpdatedComponents				= false;
 		Level->bAlreadyInitializedNetworkActors			= false;
@@ -1822,6 +1807,36 @@ void UWorld::AddToWorld( ULevel* Level, const FTransform& LevelTransform )
 
 		ULevelStreaming::BroadcastLevelVisibleStatus(this, Level->GetOutermost()->GetFName(), true);
 	}
+
+#if PERF_TRACK_DETAILED_ASYNC_STATS
+	if (bPerformedLastStep)
+	{
+		// Log out all of the timing information
+		double TotalTime = 
+			MoveActorTime + 
+			ShiftActorsTime + 
+			UpdateComponentsTime + 
+			InitBSPPhysTime + 
+			InitActorPhysTime + 
+			InitActorTime + 
+			RouteActorInitializeTime + 
+			CrossLevelRefsTime + 
+			SortActorListTime + 
+			PerformLastStepTime;
+
+		UE_LOG(LogStreaming, Log, TEXT("Detailed AddToWorld stats for '%s' - Total %6.2fms"), *Level->GetOutermost()->GetName(), TotalTime * 1000 );
+		UE_LOG(LogStreaming, Log, TEXT("Move Actors             : %6.2f ms"), MoveActorTime * 1000 );
+		UE_LOG(LogStreaming, Log, TEXT("Shift Actors            : %6.2f ms"), ShiftActorsTime * 1000 );
+		UE_LOG(LogStreaming, Log, TEXT("Update Components       : %6.2f ms"), UpdateComponentsTime * 1000 );
+		UE_LOG(LogStreaming, Log, TEXT("Init BSP Phys           : %6.2f ms"), InitBSPPhysTime * 1000 );
+		UE_LOG(LogStreaming, Log, TEXT("Init Actor Phys         : %6.2f ms"), InitActorPhysTime * 1000 );
+		UE_LOG(LogStreaming, Log, TEXT("Init Actors             : %6.2f ms"), InitActorTime * 1000 );
+		UE_LOG(LogStreaming, Log, TEXT("Initialize              : %6.2f ms"), RouteActorInitializeTime * 1000 );
+		UE_LOG(LogStreaming, Log, TEXT("Cross Level Refs        : %6.2f ms"), CrossLevelRefsTime * 1000 );
+		UE_LOG(LogStreaming, Log, TEXT("Sort Actor List         : %6.2f ms"), SortActorListTime * 1000 );
+		UE_LOG(LogStreaming, Log, TEXT("Perform Last Step       : %6.2f ms"), SortActorListTime * 1000 );
+	}
+#endif // PERF_TRACK_DETAILED_ASYNC_STATS
 }
 
 
