Flush level streaming right after map is loaded, so all relevant streaming levels will be also loaded before match start.

diff --git a/Engine/Source/Editor/UnrealEd/Private/PlayLevel.cpp b/Engine/Source/Editor/UnrealEd/Private/PlayLevel.cpp
--- a/Engine/Source/Editor/UnrealEd/Private/PlayLevel.cpp
+++ b/Engine/Source/Editor/UnrealEd/Private/PlayLevel.cpp
@@ -2645,6 +2645,10 @@ UWorld* UEditorEngine::CreatePlayInEditorWorld(FWorldContext &PieWorldContext, b
 	{
 		PlayWorld->NavigateTo(PlayWorld->GlobalOriginOffset);
 	}
+	else
+	{
+		PlayWorld->FlushLevelStreaming();
+	}
 	
 	UNavigationSystem::InitializeForWorld(PlayWorld, PieWorldContext.GamePlayers.Num() > 0 ? FNavigationSystem::PIEMode : FNavigationSystem::SimulationMode);
 	PlayWorld->CreateAISystem();
diff --git a/Engine/Source/Runtime/Engine/Private/UnrealEngine.cpp b/Engine/Source/Runtime/Engine/Private/UnrealEngine.cpp
--- a/Engine/Source/Runtime/Engine/Private/UnrealEngine.cpp
+++ b/Engine/Source/Runtime/Engine/Private/UnrealEngine.cpp
@@ -8528,11 +8528,16 @@ bool UEngine::LoadMap( FWorldContext& WorldContext, FURL URL, class UPendingNetG
 	check(WorldContext.World()->PersistentLevel);
 	LoadPackagesFully(WorldContext.World(), FULLYLOAD_Map, WorldContext.World()->PersistentLevel->GetOutermost()->GetName());
 
-	// Set initial world origin and stream in levels
+	// Make sure all relevant streaming levels are fully loaded
 	if (WorldContext.World()->WorldComposition)
 	{
+		// Set initial world origin and stream in levels
 		WorldContext.World()->NavigateTo(FIntPoint::ZeroValue);
 	}
+	else
+	{
+		WorldContext.World()->FlushLevelStreaming();
+	}
 	
 	UNavigationSystem::InitializeForWorld(WorldContext.World(), FNavigationSystem::GameMode);
 	
