Changed default rotation for SpawnDecalAtLocation to point downwards.
#ttp 340918 - REDNERING: DECAL: Spawned Decals render incorrectly through Blueprints
#branch UE4
#proj Runtime.Engine

diff --git a/Engine/Source/Runtime/Engine/Classes/Kismet/GameplayStatics.h b/Engine/Source/Runtime/Engine/Classes/Kismet/GameplayStatics.h
--- a/Engine/Source/Runtime/Engine/Classes/Kismet/GameplayStatics.h
+++ b/Engine/Source/Runtime/Engine/Classes/Kismet/GameplayStatics.h
@@ -352,7 +352,7 @@ class ENGINE_API UGameplayStatics : public UBlueprintFunctionLibrary
 	 * @param LifeSpan - destroy decal component after time runs out (0 = infinite)
 	 */
 	UFUNCTION(BlueprintCallable, Category="Rendering|Components|Decal", meta=(HidePin="WorldContextObject", DefaultToSelf="WorldContextObject", UnsafeDuringActorConstruction = "true"))
-	static UDecalComponent* SpawnDecalAtLocation(UObject* WorldContextObject, class UMaterialInterface* DecalMaterial, FVector DecalSize, FVector Location, FRotator Rotation = FRotator::ZeroRotator, float LifeSpan = 0);
+	static UDecalComponent* SpawnDecalAtLocation(UObject* WorldContextObject, class UMaterialInterface* DecalMaterial, FVector DecalSize, FVector Location, FRotator Rotation = FRotator(-90, 0, 0), float LifeSpan = 0);
 
 	/** Spawns a decal attached to and following the specified component. Does not replicate.
 	 * @param DecalMaterial - decal's material
