Fixed tiling UVs on beams.
Can now control tiling via the Texture Tiles property as well as the Tiling Distance.
Passing in second UV set to cover the whole beam 0-1 regardless of tiling so that beam ends can be faded etc.
This allows the same texturing functionality as ribbons and anim trails although the UV sets are the other way around :(

diff --git a/Engine/Source/Runtime/Engine/Private/Particles/ParticleSystemRender.cpp b/Engine/Source/Runtime/Engine/Private/Particles/ParticleSystemRender.cpp
--- a/Engine/Source/Runtime/Engine/Private/Particles/ParticleSystemRender.cpp
+++ b/Engine/Source/Runtime/Engine/Private/Particles/ParticleSystemRender.cpp
@@ -3487,6 +3487,11 @@ int32 FDynamicBeam2EmitterData::FillVertexData_NoNoise(FAsyncBufferFillData& Me)
 				float	Distance	= Direction.Size();
 				Tiles				= Distance / Source.TextureTileDistance;
 			}
+			else
+			{
+				Tiles = FMath::Max((float)Source.TextureTile, 1.0f);
+			}
+
 			fUEnd		= Tiles;
 
 			if (BeamPayloadData->TravelRatio > KINDA_SMALL_NUMBER)
@@ -3527,6 +3532,8 @@ int32 FDynamicBeam2EmitterData::FillVertexData_NoNoise(FAsyncBufferFillData& Me)
 				Vertex->Size		= Size;
 				Vertex->Tex_U		= 0.0f;
 				Vertex->Tex_V		= 0.0f;
+				Vertex->Tex_U2		= 0.0f;
+				Vertex->Tex_V2		= 0.0f;
 				Vertex->Rotation	= Particle->Rotation;
 				Vertex->Color		= Particle->Color;
 				Vertex++;
@@ -3538,6 +3545,8 @@ int32 FDynamicBeam2EmitterData::FillVertexData_NoNoise(FAsyncBufferFillData& Me)
 				Vertex->Size		= Size;
 				Vertex->Tex_U		= 0.0f;
 				Vertex->Tex_V		= 1.0f;
+				Vertex->Tex_U2		= 0.0f;
+				Vertex->Tex_V2		= 1.0f;
 				Vertex->Rotation	= Particle->Rotation;
 				Vertex->Color		= Particle->Color;
 				Vertex++;
@@ -3561,6 +3570,8 @@ int32 FDynamicBeam2EmitterData::FillVertexData_NoNoise(FAsyncBufferFillData& Me)
 				Vertex->Size		= Size;
 				Vertex->Tex_U		= fUEnd;
 				Vertex->Tex_V		= 0.0f;
+				Vertex->Tex_U2		= 1.0f;
+				Vertex->Tex_V2		= 0.0f;
 				Vertex->Rotation	= Particle->Rotation;
 				Vertex->Color		= Particle->Color;
 				Vertex++;
@@ -3572,6 +3583,8 @@ int32 FDynamicBeam2EmitterData::FillVertexData_NoNoise(FAsyncBufferFillData& Me)
 				Vertex->Size		= Size;
 				Vertex->Tex_U		= fUEnd;
 				Vertex->Tex_V		= 1.0f;
+				Vertex->Tex_U2		= 1.0f;
+				Vertex->Tex_V2		= 1.0f;
 				Vertex->Rotation	= Particle->Rotation;
 				Vertex->Color		= Particle->Color;
 				Vertex++;
@@ -3646,6 +3659,9 @@ int32 FDynamicBeam2EmitterData::FillVertexData_NoNoise(FAsyncBufferFillData& Me)
 			FVector WorkingUp;
 			float	fU;
 
+			float	Tex_U2 = 0.0f;
+			float	Tex_U2_Increment = 1.0f / BeamPayloadData->Steps;
+
 			check(InterpolatedPoints);	// TTP #33139
 			// For the direct case, this isn't a big deal, as it will not require much work per sheet.
 			for (int32 SheetIndex = 0; SheetIndex < Source.Sheets; SheetIndex++)
@@ -3696,6 +3712,8 @@ int32 FDynamicBeam2EmitterData::FillVertexData_NoNoise(FAsyncBufferFillData& Me)
 				Vertex->Size		= Size;
 				Vertex->Tex_U		= fU;
 				Vertex->Tex_V		= 0.0f;
+				Vertex->Tex_U2		= 0.0f;
+				Vertex->Tex_V2		= 0.0f;
 				Vertex->Rotation	= Particle->Rotation;
 				Vertex->Color		= Particle->Color;
 				Vertex++;
@@ -3707,6 +3725,8 @@ int32 FDynamicBeam2EmitterData::FillVertexData_NoNoise(FAsyncBufferFillData& Me)
 				Vertex->Size		= Size;
 				Vertex->Tex_U		= fU;
 				Vertex->Tex_V		= 1.0f;
+				Vertex->Tex_U2		= 0.0f;
+				Vertex->Tex_V2		= 1.0f;
 				Vertex->Rotation	= Particle->Rotation;
 				Vertex->Color		= Particle->Color;
 				Vertex++;
@@ -3714,6 +3734,8 @@ int32 FDynamicBeam2EmitterData::FillVertexData_NoNoise(FAsyncBufferFillData& Me)
 
 				for (int32 StepIndex = 0; StepIndex < BeamPayloadData->Steps; StepIndex++)
 				{
+					Tex_U2 += Tex_U2_Increment;
+
 					EndPoint	= InterpolatedPoints[StepIndex];
 					if (Source.UpVectorStepSize == 0)
 					{
@@ -3752,6 +3774,8 @@ int32 FDynamicBeam2EmitterData::FillVertexData_NoNoise(FAsyncBufferFillData& Me)
 					Vertex->Size		= Size;
 					Vertex->Tex_U		= fU + fTextureIncrement;
 					Vertex->Tex_V		= 0.0f;
+					Vertex->Tex_U2		= Tex_U2;
+					Vertex->Tex_V2		= 0.0f;
 					Vertex->Rotation	= Particle->Rotation;
 					Vertex->Color		= Particle->Color;
 					Vertex++;
@@ -3763,6 +3787,8 @@ int32 FDynamicBeam2EmitterData::FillVertexData_NoNoise(FAsyncBufferFillData& Me)
 					Vertex->Size		= Size;
 					Vertex->Tex_U		= fU + fTextureIncrement;
 					Vertex->Tex_V		= 1.0f;
+					Vertex->Tex_U2		= Tex_U2;
+					Vertex->Tex_V2		= 1.0f;
 					Vertex->Rotation	= Particle->Rotation;
 					Vertex->Color		= Particle->Color;
 					Vertex++;
