changed the tooltips for UV tiling properties in beam emitters to mention which UV set they affect.

diff --git a/Engine/Source/Runtime/Engine/Classes/Particles/TypeData/ParticleModuleTypeDataBeam2.h b/Engine/Source/Runtime/Engine/Classes/Particles/TypeData/ParticleModuleTypeDataBeam2.h
--- a/Engine/Source/Runtime/Engine/Classes/Particles/TypeData/ParticleModuleTypeDataBeam2.h
+++ b/Engine/Source/Runtime/Engine/Classes/Particles/TypeData/ParticleModuleTypeDataBeam2.h
@@ -68,11 +68,16 @@ class UParticleModuleTypeDataBeam2 : public UParticleModuleTypeDataBase
 	UPROPERTY(EditAnywhere, Category=Beam)
 	TEnumAsByte<enum EBeam2Method> BeamMethod;
 
-	/** The number of times to tile the texture along each beam									*/
+	/** The number of times to tile the texture along each beam. 
+	 *  Overridden by TextureTilingDistance if it is > 0.0.
+	 *	1st UV set only. 2nd UV set does not Tile.		
+	 */
 	UPROPERTY(EditAnywhere, Category=Beam)
 	int32 TextureTile;
 
-	/** The distance per texture tile															*/
+	/** The distance per texture tile. 
+	 *	1st UV set only. 2nd UV set does not Tile.				
+	 */
 	UPROPERTY(EditAnywhere, Category=Beam)
 	float TextureTileDistance;
 
