Remove superfluous headers part 2 of n - Engine module cleanup, Animation headers.

diff --git a/Engine/Source/Developer/AutomationController/AutomationController.Build.cs b/Engine/Source/Developer/AutomationController/AutomationController.Build.cs
--- a/Engine/Source/Developer/AutomationController/AutomationController.Build.cs
+++ b/Engine/Source/Developer/AutomationController/AutomationController.Build.cs
@@ -18,6 +18,7 @@ namespace UnrealBuildTool.Rules
 				{
 					"AutomationMessages",
 					"CoreUObject",
+					"Messaging",
 					"UnrealEdMessages",
 				}
 			);
diff --git a/Engine/Source/Developer/LauncherServices/LauncherServices.Build.cs b/Engine/Source/Developer/LauncherServices/LauncherServices.Build.cs
--- a/Engine/Source/Developer/LauncherServices/LauncherServices.Build.cs
+++ b/Engine/Source/Developer/LauncherServices/LauncherServices.Build.cs
@@ -21,6 +21,7 @@ namespace UnrealBuildTool.Rules
 			PrivateDependencyModuleNames.AddRange(
 				new string[] {
 					"CoreUObject",
+					"Messaging",
 					"SessionMessages",
 					"TargetPlatform",
 					"UnrealEdMessages",
diff --git a/Engine/Source/Developer/ProfilerClient/ProfilerClient.Build.cs b/Engine/Source/Developer/ProfilerClient/ProfilerClient.Build.cs
--- a/Engine/Source/Developer/ProfilerClient/ProfilerClient.Build.cs
+++ b/Engine/Source/Developer/ProfilerClient/ProfilerClient.Build.cs
@@ -19,6 +19,7 @@ namespace UnrealBuildTool.Rules
 				new string[]
 				{
 					"CoreUObject",
+					"Messaging",
 					"ProfilerMessages",
 				}
 			);
diff --git a/Engine/Source/Developer/ProfilerService/ProfilerService.Build.cs b/Engine/Source/Developer/ProfilerService/ProfilerService.Build.cs
--- a/Engine/Source/Developer/ProfilerService/ProfilerService.Build.cs
+++ b/Engine/Source/Developer/ProfilerService/ProfilerService.Build.cs
@@ -18,6 +18,7 @@ namespace UnrealBuildTool.Rules
 				{
 					"CoreUObject",
 					"Engine",
+					"Messaging",
 					"ProfilerMessages",
 				}
 			);
diff --git a/Engine/Source/Developer/ScreenShotComparisonTools/ScreenShotComparisonTools.Build.cs b/Engine/Source/Developer/ScreenShotComparisonTools/ScreenShotComparisonTools.Build.cs
--- a/Engine/Source/Developer/ScreenShotComparisonTools/ScreenShotComparisonTools.Build.cs
+++ b/Engine/Source/Developer/ScreenShotComparisonTools/ScreenShotComparisonTools.Build.cs
@@ -16,6 +16,7 @@ public class ScreenShotComparisonTools : ModuleRules
 			new string[] {
 				"AutomationMessages",
 				"CoreUObject",
+				"Messaging",
 				"UnrealEdMessages",
 				"Slate",
                 "EditorStyle",
diff --git a/Engine/Source/Developer/SessionServices/SessionServices.Build.cs b/Engine/Source/Developer/SessionServices/SessionServices.Build.cs
--- a/Engine/Source/Developer/SessionServices/SessionServices.Build.cs
+++ b/Engine/Source/Developer/SessionServices/SessionServices.Build.cs
@@ -18,6 +18,7 @@ public class SessionServices : ModuleRules
 			{
 				"CoreUObject",
 				"EngineMessages",
+				"Messaging",
 				"SessionMessages",
 			}
 		);
diff --git a/Engine/Source/Developer/TargetDeviceServices/TargetDeviceServices.Build.cs b/Engine/Source/Developer/TargetDeviceServices/TargetDeviceServices.Build.cs
--- a/Engine/Source/Developer/TargetDeviceServices/TargetDeviceServices.Build.cs
+++ b/Engine/Source/Developer/TargetDeviceServices/TargetDeviceServices.Build.cs
@@ -18,6 +18,7 @@ namespace UnrealBuildTool.Rules
 				new string[]
 				{
 					"CoreUObject",
+					"Messaging",
 				}
 			);
 
diff --git a/Engine/Source/Developer/iOS/IOSTargetPlatform/IOSTargetPlatform.Build.cs b/Engine/Source/Developer/iOS/IOSTargetPlatform/IOSTargetPlatform.Build.cs
--- a/Engine/Source/Developer/iOS/IOSTargetPlatform/IOSTargetPlatform.Build.cs
+++ b/Engine/Source/Developer/iOS/IOSTargetPlatform/IOSTargetPlatform.Build.cs
@@ -13,6 +13,7 @@ public class IOSTargetPlatform : ModuleRules
 			{
 				"Core",
 				"CoreUObject",
+				"Messaging",
 				"TargetPlatform",
 				"DesktopPlatform",
 				"LaunchDaemonMessages",
diff --git a/Engine/Source/Editor/UnrealEd/Private/Animation/DebugSkelMeshComponent.cpp b/Engine/Source/Editor/UnrealEd/Private/Animation/DebugSkelMeshComponent.cpp
--- a/Engine/Source/Editor/UnrealEd/Private/Animation/DebugSkelMeshComponent.cpp
+++ b/Engine/Source/Editor/UnrealEd/Private/Animation/DebugSkelMeshComponent.cpp
@@ -6,6 +6,7 @@
 #include "AnimationRuntime.h"
 #include "AnimPreviewInstance.h"
 #include "Animation/VertexAnim/VertexAnimation.h"
+#include "Animation/AnimInstance.h"
 
 //////////////////////////////////////////////////////////////////////////
 // FDebugSkelMeshSceneProxy
diff --git a/Engine/Source/Editor/UnrealEd/UnrealEd.Build.cs b/Engine/Source/Editor/UnrealEd/UnrealEd.Build.cs
--- a/Engine/Source/Editor/UnrealEd/UnrealEd.Build.cs
+++ b/Engine/Source/Editor/UnrealEd/UnrealEd.Build.cs
@@ -87,6 +87,7 @@ public class UnrealEd : ModuleRules
 				"LauncherServices",
 				"MaterialEditor",
 				"MessageLog",
+				"Messaging",
 				"NetworkFileSystem",
 				"ProjectSettingsViewer",
 				"PropertyEditor",
diff --git a/Engine/Source/Runtime/AutomationWorker/AutomationWorker.Build.cs b/Engine/Source/Runtime/AutomationWorker/AutomationWorker.Build.cs
--- a/Engine/Source/Runtime/AutomationWorker/AutomationWorker.Build.cs
+++ b/Engine/Source/Runtime/AutomationWorker/AutomationWorker.Build.cs
@@ -18,6 +18,7 @@ namespace UnrealBuildTool.Rules
 				{
 					"AutomationMessages",
 					"CoreUObject",
+					"Messaging",
 				}
 			);
 
diff --git a/Engine/Source/Runtime/Engine/Classes/Animation/AimOffsetBlendSpace.h b/Engine/Source/Runtime/Engine/Classes/Animation/AimOffsetBlendSpace.h
--- a/Engine/Source/Runtime/Engine/Classes/Animation/AimOffsetBlendSpace.h
+++ b/Engine/Source/Runtime/Engine/Classes/Animation/AimOffsetBlendSpace.h
@@ -8,7 +8,6 @@
 #pragma once 
 
 #include "BlendSpace.h"
-#include "AnimInstance.h"
 
 #include "AimOffsetBlendSpace.generated.h"
 
diff --git a/Engine/Source/Runtime/Engine/Classes/Animation/AimOffsetBlendSpace1D.h b/Engine/Source/Runtime/Engine/Classes/Animation/AimOffsetBlendSpace1D.h
--- a/Engine/Source/Runtime/Engine/Classes/Animation/AimOffsetBlendSpace1D.h
+++ b/Engine/Source/Runtime/Engine/Classes/Animation/AimOffsetBlendSpace1D.h
@@ -8,7 +8,6 @@
 #pragma once 
 
 #include "BlendSpace1D.h"
-#include "AnimInstance.h"
 
 #include "AimOffsetBlendSpace1D.generated.h"
 
diff --git a/Engine/Source/Runtime/Engine/Classes/Animation/AnimBlueprint.h b/Engine/Source/Runtime/Engine/Classes/Animation/AnimBlueprint.h
--- a/Engine/Source/Runtime/Engine/Classes/Animation/AnimBlueprint.h
+++ b/Engine/Source/Runtime/Engine/Classes/Animation/AnimBlueprint.h
@@ -2,9 +2,6 @@
 
 #pragma once
 
-#include "Engine/Blueprint.h"
-#include "AnimInstance.h"
-
 #include "AnimBlueprint.generated.h"
 
 USTRUCT()
diff --git a/Engine/Source/Runtime/Engine/Classes/Animation/AnimInstance.h b/Engine/Source/Runtime/Engine/Classes/Animation/AnimInstance.h
--- a/Engine/Source/Runtime/Engine/Classes/Animation/AnimInstance.h
+++ b/Engine/Source/Runtime/Engine/Classes/Animation/AnimInstance.h
@@ -3,19 +3,24 @@
 #pragma once
 
 #include "AnimationAsset.h"
-#include "Components/SkeletalMeshComponent.h"
 #include "AnimSequence.h"
-#include "AnimStateMachineTypes.h"
-
+#include "Components/SkeletalMeshComponent.h"
 #include "AnimInstance.generated.h"
 
-struct FBoneContainer;
 struct FAnimMontageInstance;
-class  USkeletalMesh;
 class UAnimMontage;
-
-DECLARE_DELEGATE_TwoParams(FOnMontageEnded, class UAnimMontage*, bool /*bInterrupted*/)
-DECLARE_DELEGATE_TwoParams(FOnMontageBlendingOutStarted, class UAnimMontage*, bool /*bInterrupted*/)
+class USkeleton;
+class AActor;
+class UAnimSequenceBase;
+class UBlendSpaceBase;
+class APawn;
+class UAnimationAsset;
+class UCanvas;
+class UWorld;
+class FTransform;
+
+DECLARE_DELEGATE_TwoParams(FOnMontageEnded, UAnimMontage*, bool /*bInterrupted*/)
+DECLARE_DELEGATE_TwoParams(FOnMontageBlendingOutStarted, UAnimMontage*, bool /*bInterrupted*/)
 
 /**
 * Delegate for when Montage is completed, whether interrupted or finished
@@ -23,7 +28,7 @@ DECLARE_DELEGATE_TwoParams(FOnMontageBlendingOutStarted, class UAnimMontage*, bo
 *
 * bInterrupted = true if it was not property finished
 */
-DECLARE_DYNAMIC_MULTICAST_DELEGATE_TwoParams(FOnMontageEndedMCDelegate, class UAnimMontage*, Montage, bool, bInterrupted);
+DECLARE_DYNAMIC_MULTICAST_DELEGATE_TwoParams(FOnMontageEndedMCDelegate, UAnimMontage*, Montage, bool, bInterrupted);
 
 /**
 * Delegate for when Montage started to blend out, whether interrupted or finished
@@ -31,7 +36,7 @@ DECLARE_DYNAMIC_MULTICAST_DELEGATE_TwoParams(FOnMontageEndedMCDelegate, class UA
 *
 * bInterrupted = true if it was not property finished
 */
-DECLARE_DYNAMIC_MULTICAST_DELEGATE_TwoParams(FOnMontageBlendingOutStartedMCDelegate, class UAnimMontage*, Montage, bool, bInterrupted);
+DECLARE_DYNAMIC_MULTICAST_DELEGATE_TwoParams(FOnMontageBlendingOutStartedMCDelegate, UAnimMontage*, Montage, bool, bInterrupted);
 
 /** Enum for controlling which reference frame a controller is applied in. */
 UENUM()
@@ -246,7 +251,7 @@ class ENGINE_API UAnimInstance : public UObject
 
 	/** This is used to extract animation. If Mesh exists, this will be overwritten by Mesh->Skeleton */
 	UPROPERTY(transient)
-	class USkeleton* CurrentSkeleton;
+	USkeleton* CurrentSkeleton;
 
 	// The list of animation assets which are going to be evaluated this frame and need to be ticked (ungrouped)
 	UPROPERTY(transient)
@@ -263,18 +268,18 @@ class ENGINE_API UAnimInstance : public UObject
 public:
 
 	// @todo document
-	void MakeSequenceTickRecord(FAnimTickRecord& TickRecord, class UAnimSequenceBase* Sequence, bool bLooping, float PlayRate, float FinalBlendWeight, float& CurrentTime) const;
-	void MakeBlendSpaceTickRecord(FAnimTickRecord& TickRecord, class UBlendSpaceBase* BlendSpace, const FVector& BlendInput, TArray<FBlendSampleData>& BlendSampleDataCache, FBlendFilter& BlendFilter, bool bLooping, float PlayRate, float FinalBlendWeight, float& CurrentTime) const;
+	void MakeSequenceTickRecord(FAnimTickRecord& TickRecord, UAnimSequenceBase* Sequence, bool bLooping, float PlayRate, float FinalBlendWeight, float& CurrentTime) const;
+	void MakeBlendSpaceTickRecord(FAnimTickRecord& TickRecord, UBlendSpaceBase* BlendSpace, const FVector& BlendInput, TArray<FBlendSampleData>& BlendSampleDataCache, FBlendFilter& BlendFilter, bool bLooping, float PlayRate, float FinalBlendWeight, float& CurrentTime) const;
 
-	void SequenceAdvanceImmediate(class UAnimSequenceBase* Sequence, bool bLooping, float PlayRate, float DeltaSeconds, /*inout*/ float& CurrentTime);
+	void SequenceAdvanceImmediate(UAnimSequenceBase* Sequence, bool bLooping, float PlayRate, float DeltaSeconds, /*inout*/ float& CurrentTime);
 
 	// @todo document
-	void BlendSpaceAdvanceImmediate(class UBlendSpaceBase* BlendSpace, const FVector& BlendInput, TArray<FBlendSampleData> & BlendSampleDataCache, FBlendFilter & BlendFilter, bool bLooping, float PlayRate, float DeltaSeconds, /*inout*/ float& CurrentTime);
+	void BlendSpaceAdvanceImmediate(UBlendSpaceBase* BlendSpace, const FVector& BlendInput, TArray<FBlendSampleData> & BlendSampleDataCache, FBlendFilter & BlendFilter, bool bLooping, float PlayRate, float DeltaSeconds, /*inout*/ float& CurrentTime);
 
 	// Creates an uninitialized tick record in the list for the correct group or the ungrouped array.  If the group is valid, OutSyncGroupPtr will point to the group.
 	FAnimTickRecord& CreateUninitializedTickRecord(int32 GroupIndex, FAnimGroupInstance*& OutSyncGroupPtr);
 
-	void SequenceEvaluatePose(class UAnimSequenceBase* Sequence, struct FA2Pose& Pose, const FAnimExtractContext & ExtractionContext);
+	void SequenceEvaluatePose(UAnimSequenceBase* Sequence, struct FA2Pose& Pose, const FAnimExtractContext & ExtractionContext);
 
 	void BlendSequences(const struct FA2Pose& Pose1, const struct FA2Pose& Pose2, float Alpha, struct FA2Pose& Blended);
 
@@ -282,7 +287,7 @@ public:
 
 	void ApplyAdditiveSequence(const struct FA2Pose& BasePose, const struct FA2Pose& AdditivePose, float Alpha, struct FA2Pose& Blended);
 
-	void BlendSpaceEvaluatePose(class UBlendSpaceBase* BlendSpace, TArray<FBlendSampleData>& BlendSampleDataCache, struct FA2Pose& Pose);
+	void BlendSpaceEvaluatePose(UBlendSpaceBase* BlendSpace, TArray<FBlendSampleData>& BlendSampleDataCache, struct FA2Pose& Pose);
 
 	// skeletal control related functions
 	void BlendRotationOffset(const struct FA2Pose& BasePose/* local space base pose */, struct FA2Pose const & RotationOffsetPose/* mesh space rotation only additive **/, float Alpha/*0 means no additive, 1 means whole additive */, struct FA2Pose& Pose /** local space blended pose **/);
@@ -301,16 +306,16 @@ public:
 	// kismet event functions
 
 	UFUNCTION(BlueprintPure, Category = "Animation")
-	virtual class APawn* TryGetPawnOwner();
+	virtual APawn* TryGetPawnOwner();
 
 protected:
 	/** Returns the owning actor of this AnimInstance */
 	UFUNCTION(BlueprintPure, Category="Animation")
-	class AActor* GetOwningActor();
+	AActor* GetOwningActor();
 	
 	// Returns the skeletal mesh component that has created this AnimInstance
 	UFUNCTION(BlueprintPure, Category="Animation")
-	class USkeletalMeshComponent* GetOwningComponent();
+	USkeletalMeshComponent* GetOwningComponent();
 
 public:
 
@@ -391,19 +396,19 @@ public:
 
 	/** Returns the length (in seconds) of an animation AnimAsset. */
 	UFUNCTION(BlueprintPure, Category="Animation", meta=(BlueprintInternalUseOnly = "true"))
-	static float GetAnimAssetPlayerLength(class UAnimationAsset* AnimAsset);
+	static float GetAnimAssetPlayerLength(UAnimationAsset* AnimAsset);
 
 	//** Returns how far through the animation AnimAsset we are (as a proportion between 0.0 and 1.0). */
 	UFUNCTION(BlueprintPure, Category="Animation", meta=(BlueprintInternalUseOnly = "true"))
-	static float GetAnimAssetPlayerTimeFraction(class UAnimationAsset* AnimAsset, float CurrentTime);
+	static float GetAnimAssetPlayerTimeFraction(UAnimationAsset* AnimAsset, float CurrentTime);
 
 	/** Returns how long until the end of the animation AnimAsset (in seconds). */
 	UFUNCTION(BlueprintPure, Category="Animation", meta=(BlueprintInternalUseOnly = "true"))
-	static float GetAnimAssetPlayerTimeFromEnd(class UAnimationAsset* AnimAsset, float CurrentTime);
+	static float GetAnimAssetPlayerTimeFromEnd(UAnimationAsset* AnimAsset, float CurrentTime);
 
 	/** Returns how long until the end of the animation AnimAsset we are (as a proportion between 0.0 and 1.0). */
 	UFUNCTION(BlueprintPure, Category="Animation", meta=(BlueprintInternalUseOnly = "true"))
-	static float GetAnimAssetPlayerTimeFromEndFraction(class UAnimationAsset* AnimAsset, float CurrentTime);
+	static float GetAnimAssetPlayerTimeFromEndFraction(UAnimationAsset* AnimAsset, float CurrentTime);
 
 	/** Returns the weight of a state in a state machine. */
 	UFUNCTION(BlueprintPure, Category="Animation", meta=(BlueprintInternalUseOnly = "true"))
@@ -469,7 +474,7 @@ public:
 	virtual bool NativeEvaluateAnimation(FPoseContext& Output);
 
 	// Debug output for this anim instance 
-	void DisplayDebug(class UCanvas* Canvas, const FDebugDisplayInfo& DebugDisplay, float& YL, float& YPos);
+	void DisplayDebug(UCanvas* Canvas, const FDebugDisplayInfo& DebugDisplay, float& YL, float& YPos);
 public:
 
 	/** Temporary array of bone indices required this frame. Should be subset of Skeleton and Mesh's RequiredBones */
@@ -527,7 +532,7 @@ public:
 	// @todo document
 	inline USkeletalMeshComponent* GetSkelMeshComponent() const { return CastChecked<USkeletalMeshComponent>(GetOuter()); }
 
-	virtual class UWorld* GetWorld() const override;
+	virtual UWorld* GetWorld() const override;
 
 	/** Add anim notifier **/
 	void AddAnimNotifies(const TArray<const FAnimNotifyEvent*>& NewNotifies, const float InstanceWeight);
diff --git a/Engine/Source/Runtime/Engine/Classes/Animation/AnimMontage.h b/Engine/Source/Runtime/Engine/Classes/Animation/AnimMontage.h
--- a/Engine/Source/Runtime/Engine/Classes/Animation/AnimMontage.h
+++ b/Engine/Source/Runtime/Engine/Classes/Animation/AnimMontage.h
@@ -7,9 +7,8 @@
  *
  */
 
-#include "Delegate.h"
 #include "AnimCompositeBase.h"
-#include "AnimInstance.h"
+#include "Animation/AnimInstance.h"
 #include "AnimMontage.generated.h"
 
 /**
diff --git a/Engine/Source/Runtime/Engine/Classes/Animation/AnimNode_LayeredBoneBlend.h b/Engine/Source/Runtime/Engine/Classes/Animation/AnimNode_LayeredBoneBlend.h
--- a/Engine/Source/Runtime/Engine/Classes/Animation/AnimNode_LayeredBoneBlend.h
+++ b/Engine/Source/Runtime/Engine/Classes/Animation/AnimNode_LayeredBoneBlend.h
@@ -3,8 +3,7 @@
 #pragma once
 
 #include "AnimNodeBase.h"
-#include "AnimData/BoneMaskFilter.h"
-
+#include "Animation/AnimData/BoneMaskFilter.h"
 #include "AnimNode_LayeredBoneBlend.generated.h"
 
 
diff --git a/Engine/Source/Runtime/Engine/Classes/Animation/AnimNode_TwoWayBlend.h b/Engine/Source/Runtime/Engine/Classes/Animation/AnimNode_TwoWayBlend.h
--- a/Engine/Source/Runtime/Engine/Classes/Animation/AnimNode_TwoWayBlend.h
+++ b/Engine/Source/Runtime/Engine/Classes/Animation/AnimNode_TwoWayBlend.h
@@ -2,8 +2,8 @@
 
 #pragma once
 
-#include "InputScaleBias.h"
 #include "Animation/AnimNodeBase.h"
+#include "Animation/InputScaleBias.h"
 #include "AnimNode_TwoWayBlend.generated.h"
 
 // This represents a baked transition
diff --git a/Engine/Source/Runtime/Engine/Classes/Animation/AnimNotifies/AnimNotify.h b/Engine/Source/Runtime/Engine/Classes/Animation/AnimNotifies/AnimNotify.h
--- a/Engine/Source/Runtime/Engine/Classes/Animation/AnimNotifies/AnimNotify.h
+++ b/Engine/Source/Runtime/Engine/Classes/Animation/AnimNotifies/AnimNotify.h
@@ -4,6 +4,9 @@
 #pragma once
 #include "AnimNotify.generated.h"
 
+class UAnimSequence;
+class UAnimSequenceBase;
+class USkeletalMeshComponent;
 struct FAnimNotifyEvent;
 
 UCLASS(abstract, Blueprintable, const, hidecategories=Object, collapsecategories)
@@ -12,7 +15,7 @@ class ENGINE_API UAnimNotify : public UObject
 	GENERATED_UCLASS_BODY()
 
 	UFUNCTION(BlueprintImplementableEvent)
-	virtual bool Received_Notify(class USkeletalMeshComponent* MeshComp, class UAnimSequenceBase* Animation) const;
+	virtual bool Received_Notify(USkeletalMeshComponent* MeshComp, UAnimSequenceBase* Animation) const;
 
 #if WITH_EDITORONLY_DATA
 	/** Color of Notify in editor */
@@ -21,7 +24,7 @@ class ENGINE_API UAnimNotify : public UObject
 
 #endif // WITH_EDITORONLY_DATA
 
-	virtual void Notify(class USkeletalMeshComponent* MeshComp, class UAnimSequenceBase* Animation);
+	virtual void Notify(USkeletalMeshComponent* MeshComp, UAnimSequenceBase* Animation);
 
 	// @todo document 
 	virtual FString GetEditorComment() 
@@ -45,7 +48,7 @@ class ENGINE_API UAnimNotify : public UObject
 	 *	@param	AnimSeq			The animation sequence this notify is associated with.
 	 *	@param	OwnerEvent		The event that 'owns' this AnimNotify.
 	 */
-	virtual void AnimNotifyEventChanged(class USkeletalMeshComponent* MeshComp, class UAnimSequenceBase* Animation, FAnimNotifyEvent* OwnerEvent) {}
+	virtual void AnimNotifyEventChanged(USkeletalMeshComponent* MeshComp, UAnimSequenceBase* Animation, FAnimNotifyEvent* OwnerEvent) {}
 
 	/**
 	 * We don't instance UAnimNotify objects along with the animations they belong to, but
diff --git a/Engine/Source/Runtime/Engine/Classes/Animation/AnimNotifies/AnimNotifyState.h b/Engine/Source/Runtime/Engine/Classes/Animation/AnimNotifies/AnimNotifyState.h
--- a/Engine/Source/Runtime/Engine/Classes/Animation/AnimNotifies/AnimNotifyState.h
+++ b/Engine/Source/Runtime/Engine/Classes/Animation/AnimNotifies/AnimNotifyState.h
@@ -4,6 +4,9 @@
 #pragma once
 #include "AnimNotifyState.generated.h"
 
+class USkeletalMeshComponent;
+class UAnimSequence;
+class UAnimSequenceBase;
 struct FAnimNotifyEvent;
 
 UCLASS(abstract, editinlinenew, Blueprintable, const, hidecategories=Object, collapsecategories, meta=(ShowHiddenSelfPins))
@@ -12,13 +15,13 @@ class ENGINE_API UAnimNotifyState : public UObject
 	GENERATED_UCLASS_BODY()
 
 	UFUNCTION(BlueprintImplementableEvent)
-	virtual bool Received_NotifyBegin(class USkeletalMeshComponent * MeshComp, class UAnimSequenceBase * Animation) const;
+	virtual bool Received_NotifyBegin(USkeletalMeshComponent * MeshComp, UAnimSequenceBase * Animation) const;
 	
 	UFUNCTION(BlueprintImplementableEvent)
-	virtual bool Received_NotifyTick(class USkeletalMeshComponent * MeshComp, class UAnimSequenceBase * Animation, float FrameDeltaTime) const;
+	virtual bool Received_NotifyTick(USkeletalMeshComponent * MeshComp, UAnimSequenceBase * Animation, float FrameDeltaTime) const;
 
 	UFUNCTION(BlueprintImplementableEvent)
-	virtual bool Received_NotifyEnd(class USkeletalMeshComponent * MeshComp, class UAnimSequenceBase * Animation) const;
+	virtual bool Received_NotifyEnd(USkeletalMeshComponent * MeshComp, UAnimSequenceBase * Animation) const;
 
 #if WITH_EDITORONLY_DATA
 	/** Color of Notify in editor */
@@ -27,9 +30,9 @@ class ENGINE_API UAnimNotifyState : public UObject
 
 #endif // WITH_EDITORONLY_DATA
 
-	virtual void NotifyBegin(class USkeletalMeshComponent * MeshComp, class UAnimSequenceBase * Animation);
-	virtual void NotifyTick(class USkeletalMeshComponent * MeshComp, class UAnimSequenceBase * Animation, float FrameDeltaTime);
-	virtual void NotifyEnd(class USkeletalMeshComponent * MeshComp, class UAnimSequenceBase * Animation);
+	virtual void NotifyBegin(USkeletalMeshComponent * MeshComp, UAnimSequenceBase * Animation);
+	virtual void NotifyTick(USkeletalMeshComponent * MeshComp, UAnimSequenceBase * Animation, float FrameDeltaTime);
+	virtual void NotifyEnd(USkeletalMeshComponent * MeshComp, UAnimSequenceBase * Animation);
 
 	// @todo document 
 	virtual FString GetEditorComment() 
@@ -53,7 +56,7 @@ class ENGINE_API UAnimNotifyState : public UObject
 	 *	@param	AnimSeq			The animation sequence this notify is associated with.
 	 *	@param	OwnerEvent		The event that 'owns' this AnimNotify.
 	 */
-	virtual void AnimNotifyEventChanged(class USkeletalMeshComponent* MeshComp, class UAnimSequenceBase* Animation, FAnimNotifyEvent * OwnerEvent) {}
+	virtual void AnimNotifyEventChanged(USkeletalMeshComponent* MeshComp, UAnimSequenceBase* Animation, FAnimNotifyEvent * OwnerEvent) {}
 };
 
 
diff --git a/Engine/Source/Runtime/Engine/Classes/Animation/AnimNotifies/AnimNotifyState_TimedParticleEffect.h b/Engine/Source/Runtime/Engine/Classes/Animation/AnimNotifies/AnimNotifyState_TimedParticleEffect.h
--- a/Engine/Source/Runtime/Engine/Classes/Animation/AnimNotifies/AnimNotifyState_TimedParticleEffect.h
+++ b/Engine/Source/Runtime/Engine/Classes/Animation/AnimNotifies/AnimNotifyState_TimedParticleEffect.h
@@ -1,8 +1,7 @@
 // Copyright 1998-2014 Epic Games, Inc. All Rights Reserved.
 
 #pragma once
-
-#include "AnimNotifyState.h"
+#include "Animation/AnimNotifies/AnimNotifyState.h"
 #include "AnimNotifyState_TimedParticleEffect.generated.h"
 
 // Timed Particle Effect Notify
@@ -17,7 +16,7 @@ class UAnimNotifyState_TimedParticleEffect : public UAnimNotifyState
 	UPROPERTY(EditAnywhere, Category = ParticleSystem, meta = (ToolTip = "The particle system to spawn for the notify state"))
 	UParticleSystem* PSTemplate;
 
-	// The socket within our mesh component to attach to when we spawn the particle componnet
+	// The socket within our mesh component to attach to when we spawn the particle component
 	UPROPERTY(EditAnywhere, Category = ParticleSystem, meta = (ToolTip = "The socket or bone to attach the system to"))
 	FName SocketName;
 
diff --git a/Engine/Source/Runtime/Engine/Classes/Animation/AnimNotifies/AnimNotifyState_Trail.h b/Engine/Source/Runtime/Engine/Classes/Animation/AnimNotifies/AnimNotifyState_Trail.h
--- a/Engine/Source/Runtime/Engine/Classes/Animation/AnimNotifies/AnimNotifyState_Trail.h
+++ b/Engine/Source/Runtime/Engine/Classes/Animation/AnimNotifies/AnimNotifyState_Trail.h
@@ -1,7 +1,6 @@
 // Copyright 1998-2014 Epic Games, Inc. All Rights Reserved.
 
 #pragma once
-#include "ParticleHelper.h"
 #include "AnimNotifyState.h"
 #include "AnimNotifyState_Trail.generated.h"
 
diff --git a/Engine/Source/Runtime/Engine/Classes/Animation/BlendSpace.h b/Engine/Source/Runtime/Engine/Classes/Animation/BlendSpace.h
--- a/Engine/Source/Runtime/Engine/Classes/Animation/BlendSpace.h
+++ b/Engine/Source/Runtime/Engine/Classes/Animation/BlendSpace.h
@@ -3,8 +3,6 @@
 #pragma once
 
 #include "BlendSpaceBase.h"
-#include "AnimInterpFilter.h"
-#include "AnimInstance.h"
 
 #include "BlendSpace.generated.h"
 
diff --git a/Engine/Source/Runtime/Engine/Classes/Animation/BlendSpace1D.h b/Engine/Source/Runtime/Engine/Classes/Animation/BlendSpace1D.h
--- a/Engine/Source/Runtime/Engine/Classes/Animation/BlendSpace1D.h
+++ b/Engine/Source/Runtime/Engine/Classes/Animation/BlendSpace1D.h
@@ -7,7 +7,6 @@
 
 #pragma once
 
-#include "AnimInterpFilter.h"
 #include "BlendSpaceBase.h"
 #include "BlendSpace1D.generated.h"
 
diff --git a/Engine/Source/Runtime/Engine/Classes/Animation/Skeleton.h b/Engine/Source/Runtime/Engine/Classes/Animation/Skeleton.h
--- a/Engine/Source/Runtime/Engine/Classes/Animation/Skeleton.h
+++ b/Engine/Source/Runtime/Engine/Classes/Animation/Skeleton.h
@@ -6,8 +6,6 @@
  */
 
 #pragma once
-#include "Core.h"
-#include "Delegate.h"
 #include "PreviewAssetAttachComponent.h"
 #include "Skeleton.generated.h"
 
diff --git a/Engine/Source/Runtime/Engine/Classes/Animation/VertexAnim/MorphTarget.h b/Engine/Source/Runtime/Engine/Classes/Animation/VertexAnim/MorphTarget.h
--- a/Engine/Source/Runtime/Engine/Classes/Animation/VertexAnim/MorphTarget.h
+++ b/Engine/Source/Runtime/Engine/Classes/Animation/VertexAnim/MorphTarget.h
@@ -5,7 +5,6 @@
 #include "MorphTarget.generated.h"
 
 class USkeletalMesh;
-class UStaticMesh;
 
 /**
 * Morph mesh vertex data used for comparisons and importing
diff --git a/Engine/Source/Runtime/Engine/Private/Animation/AnimInstance.cpp b/Engine/Source/Runtime/Engine/Private/Animation/AnimInstance.cpp
--- a/Engine/Source/Runtime/Engine/Private/Animation/AnimInstance.cpp
+++ b/Engine/Source/Runtime/Engine/Private/Animation/AnimInstance.cpp
@@ -18,6 +18,7 @@
 #include "Animation/AnimNotifies/AnimNotifyState.h"
 #include "Animation/AnimNotifies/AnimNotify.h"
 #include "Animation/AnimMontage.h"
+#include "Animation/AnimInstance.h"
 
 /** Anim stats */
 
diff --git a/Engine/Source/Runtime/Engine/Private/Slate/SlateSoundDevice.cpp b/Engine/Source/Runtime/Engine/Private/Slate/SlateSoundDevice.cpp
--- a/Engine/Source/Runtime/Engine/Private/Slate/SlateSoundDevice.cpp
+++ b/Engine/Source/Runtime/Engine/Private/Slate/SlateSoundDevice.cpp
@@ -5,6 +5,7 @@
 #include "AudioDevice.h"
 #include "Slate/SlateSoundDevice.h"
 #include "Sound/SoundBase.h"
+#include "Sound/SlateSound.h"
 
 DEFINE_LOG_CATEGORY_STATIC(LogSlateSoundDevice, Log, All);
 
diff --git a/Engine/Source/Runtime/Engine/Private/Slate/SlateSoundDevice.h b/Engine/Source/Runtime/Engine/Private/Slate/SlateSoundDevice.h
--- a/Engine/Source/Runtime/Engine/Private/Slate/SlateSoundDevice.h
+++ b/Engine/Source/Runtime/Engine/Private/Slate/SlateSoundDevice.h
@@ -2,7 +2,7 @@
 #pragma once
 
 #include "EnginePrivate.h"
-#include "Slate.h"
+#include "Sound/ISlateSoundDevice.h"
 
 class FSlateSoundDevice : public ISlateSoundDevice
 {
diff --git a/Engine/Source/Runtime/Engine/Private/SplineMeshSceneProxy.h b/Engine/Source/Runtime/Engine/Private/SplineMeshSceneProxy.h
--- a/Engine/Source/Runtime/Engine/Private/SplineMeshSceneProxy.h
+++ b/Engine/Source/Runtime/Engine/Private/SplineMeshSceneProxy.h
@@ -3,8 +3,10 @@
 #pragma once
 
 #include "StaticMeshResources.h"
+#include "Components/SplineMeshComponent.h"
 
 class FVertexFactoryShaderParameters;
+class USplineMeshComponent;
 
 //////////////////////////////////////////////////////////////////////////
 // SplineMeshVertexFactory
@@ -144,100 +146,22 @@ protected:
 	};
 public:
 
-	FSplineMeshSceneProxy(USplineMeshComponent* InComponent) :
-		FStaticMeshSceneProxy(InComponent)
-	{
-		// make sure all the materials are okay to be rendered as a spline mesh
-		for (FStaticMeshSceneProxy::FLODInfo& LODInfo : LODs)
-		{
-			for (FStaticMeshSceneProxy::FLODInfo::FSectionInfo& Section : LODInfo.Sections)
-			{
-				if (!Section.Material->CheckMaterialUsage(MATUSAGE_SplineMesh))
-				{
-					Section.Material = UMaterial::GetDefaultMaterial(MD_Surface);
-				}
-			}
-		}
-
-		// Copy spline params from component
-		SplineParams = InComponent->SplineParams;
-		SplineUpDir = InComponent->SplineUpDir;
-		bSmoothInterpRollScale = InComponent->bSmoothInterpRollScale;
-		ForwardAxis = InComponent->ForwardAxis;
-
-		// Fill in info about the mesh
-		FBoxSphereBounds StaticMeshBounds = StaticMesh->GetBounds();
-		SplineMeshScaleZ = 0.5f / USplineMeshComponent::GetAxisValue(StaticMeshBounds.BoxExtent, ForwardAxis); // 1/(2 * Extent)
-		SplineMeshMinZ = USplineMeshComponent::GetAxisValue(StaticMeshBounds.Origin, ForwardAxis) * SplineMeshScaleZ - 0.5f;
-
-		LODResources.Reset(InComponent->StaticMesh->RenderData->LODResources.Num());
-
-		for (int32 LODIndex = 0; LODIndex < LODs.Num(); LODIndex++)
-		{
-			FSplineMeshVertexFactory* VertexFactory = new FSplineMeshVertexFactory(this);
+	FSplineMeshSceneProxy(USplineMeshComponent* InComponent);
 
-			LODResources.Add(VertexFactory);
-
-			InitResources(InComponent, LODIndex);
-		}
-	}
-
-	virtual ~FSplineMeshSceneProxy() override
-	{
-		ReleaseResources();
-
-		for (FSplineMeshSceneProxy::FLODResources& LODResource : LODResources)
-		{
-			delete LODResource.VertexFactory;
-		}
-		LODResources.Empty();
-	}
+	virtual ~FSplineMeshSceneProxy() override;
 
 	void InitResources(USplineMeshComponent* InComponent, int32 InLODIndex);
 
 	void ReleaseResources();
 
 	/** Sets up a shadow FMeshBatch for a specific LOD. */
-	virtual bool GetShadowMeshElement(int32 LODIndex, uint8 InDepthPriorityGroup, FMeshBatch& OutMeshElement) const override
-	{
-		//checkf(LODIndex == 0, TEXT("Getting spline static mesh element with invalid LOD [%d]"), LODIndex);
-
-		if (FStaticMeshSceneProxy::GetShadowMeshElement(LODIndex, InDepthPriorityGroup, OutMeshElement))
-		{
-			OutMeshElement.VertexFactory = LODResources[LODIndex].VertexFactory;
-			OutMeshElement.ReverseCulling ^= (SplineParams.StartScale.X < 0) ^ (SplineParams.StartScale.Y < 0);
-			return true;
-		}
-		return false;
-	}
+	virtual bool GetShadowMeshElement(int32 LODIndex, uint8 InDepthPriorityGroup, FMeshBatch& OutMeshElement) const override;
 
 	/** Sets up a FMeshBatch for a specific LOD and element. */
-	virtual bool GetMeshElement(int32 LODIndex, int32 SectionIndex, uint8 InDepthPriorityGroup, FMeshBatch& OutMeshElement, const bool bUseSelectedMaterial, const bool bUseHoveredMaterial) const override
-	{
-		//checkf(LODIndex == 0 /*&& SectionIndex == 0*/, TEXT("Getting spline static mesh element with invalid params [%d, %d]"), LODIndex, SectionIndex);
-
-		if (FStaticMeshSceneProxy::GetMeshElement(LODIndex, SectionIndex, InDepthPriorityGroup, OutMeshElement, bUseSelectedMaterial, bUseHoveredMaterial))
-		{
-			OutMeshElement.VertexFactory = LODResources[LODIndex].VertexFactory;
-			OutMeshElement.ReverseCulling ^= (SplineParams.StartScale.X < 0) ^ (SplineParams.StartScale.Y < 0);
-			return true;
-		}
-		return false;
-	}
+	virtual bool GetMeshElement(int32 LODIndex, int32 SectionIndex, uint8 InDepthPriorityGroup, FMeshBatch& OutMeshElement, const bool bUseSelectedMaterial, const bool bUseHoveredMaterial) const override;
 
 	/** Sets up a wireframe FMeshBatch for a specific LOD. */
-	virtual bool GetWireframeMeshElement(int32 LODIndex, const FMaterialRenderProxy* WireframeRenderProxy, uint8 InDepthPriorityGroup, FMeshBatch& OutMeshElement) const override
-	{
-		//checkf(LODIndex == 0, TEXT("Getting spline static mesh element with invalid LOD [%d]"), LODIndex);
-
-		if (FStaticMeshSceneProxy::GetWireframeMeshElement(LODIndex, WireframeRenderProxy, InDepthPriorityGroup, OutMeshElement))
-		{
-			OutMeshElement.VertexFactory = LODResources[LODIndex].VertexFactory;
-			OutMeshElement.ReverseCulling ^= (SplineParams.StartScale.X < 0) ^ (SplineParams.StartScale.Y < 0);
-			return true;
-		}
-		return false;
-	}
+	virtual bool GetWireframeMeshElement(int32 LODIndex, const FMaterialRenderProxy* WireframeRenderProxy, uint8 InDepthPriorityGroup, FMeshBatch& OutMeshElement) const override;
 
 	virtual FPrimitiveViewRelevance GetViewRelevance(const FSceneView* View) override
 	{
diff --git a/Engine/Source/Runtime/Messaging/Private/MessageEndpointBuilder.cpp b/Engine/Source/Runtime/Messaging/Private/MessageEndpointBuilder.cpp
new file mode 100644
--- /dev/null
+++ b/Engine/Source/Runtime/Messaging/Private/MessageEndpointBuilder.cpp
@@ -0,0 +1,32 @@
+#include "MessagingPrivatePCH.h"
+#include "MessageEndpointBuilder.h"
+
+FMessageEndpointPtr FMessageEndpointBuilder::Build()
+{
+	FMessageEndpointPtr Endpoint;
+
+	IMessageBusPtr Bus = BusPtr.Pin();
+
+	if (Bus.IsValid())
+	{
+		Endpoint = MakeShareable(new FMessageEndpoint(Name, Bus.ToSharedRef(), Handlers));
+		Bus->Register(Endpoint->GetAddress(), Endpoint.ToSharedRef());
+
+		if (Disabled)
+		{
+			Endpoint->Disable();
+		}
+
+		if (InboxEnabled)
+		{
+			Endpoint->EnableInbox();
+			Endpoint->SetRecipientThread(ENamedThreads::AnyThread);
+		}
+		else
+		{
+			Endpoint->SetRecipientThread(RecipientThread);
+		}
+	}
+
+	return Endpoint;
+}
\ No newline at end of file
diff --git a/Engine/Source/Runtime/Messaging/Public/Common/MessageEndpoint.h b/Engine/Source/Runtime/Messaging/Public/Common/MessageEndpoint.h
--- a/Engine/Source/Runtime/Messaging/Public/Common/MessageEndpoint.h
+++ b/Engine/Source/Runtime/Messaging/Public/Common/MessageEndpoint.h
@@ -68,28 +68,12 @@ public:
 	 * @param InBus The message bus to attach this endpoint to.
 	 * @param InHandlers The collection of message handlers to register.
 	 */
-	FMessageEndpoint( const FName& InName, const IMessageBusRef& InBus, const TArray<IMessageHandlerPtr>& InHandlers )
-		: Address(FGuid::NewGuid())
-		, BusPtr(InBus)
-		, Enabled(true)
-		, Handlers(InHandlers)
-		, Id(FGuid::NewGuid())
-		, InboxEnabled(false)
-		, Name(InName)
-	{ }
+	MESSAGING_API FMessageEndpoint( const FName& InName, const IMessageBusRef& InBus, const TArray<IMessageHandlerPtr>& InHandlers );
 
 	/**
 	 * Destructor.
 	 */
-	~FMessageEndpoint( )
-	{
-		IMessageBusPtr Bus = BusPtr.Pin();
-
-		if (Bus.IsValid())
-		{
-			Bus->Unregister(Address);
-		}
-	}
+	MESSAGING_API ~FMessageEndpoint();
 
 public:
 
@@ -177,15 +161,7 @@ public:
 	 * @param Message The context of the message to defer.
 	 * @param Delay The time delay.
 	 */
-	void Defer( const IMessageContextRef& Context, const FTimespan& Delay )
-	{
-		IMessageBusPtr Bus = GetBusIfEnabled();
-
-		if (Bus.IsValid())
-		{
-			Bus->Forward(Context, TArrayBuilder<FMessageAddress>().Add(Address), Context->GetScope(), Delay, AsShared());
-		}
-	}
+	MESSAGING_API void Defer(const IMessageContextRef& Context, const FTimespan& Delay);
 
 	/**
 	 * Forwards a previously received message.
@@ -197,15 +173,7 @@ public:
 	 * @param ForwardingScope The scope of the forwarded message.
 	 * @param Delay The time delay.
 	 */
-	void Forward( const IMessageContextRef& Context, const TArray<FMessageAddress>& Recipients, EMessageScope::Type ForwardingScope, const FTimespan& Delay )
-	{
-		IMessageBusPtr Bus = GetBusIfEnabled();
-
-		if (Bus.IsValid())
-		{
-			Bus->Forward(Context, Recipients, ForwardingScope, Delay, AsShared());
-		}
-	}
+	MESSAGING_API void Forward(const IMessageContextRef& Context, const TArray<FMessageAddress>& Recipients, EMessageScope::Type ForwardingScope, const FTimespan& Delay);
 
 	/**
 	 * Publishes a message to all subscribed recipients within the specified scope.
@@ -217,15 +185,7 @@ public:
 	 * @param Delay The delay after which to publish the message.
 	 * @param Expiration The time at which the message expires.
 	 */
-	void Publish( void* Message, UScriptStruct* TypeInfo, EMessageScope::Type Scope, const FTimespan& Delay, const FDateTime& Expiration )
-	{
-		IMessageBusPtr Bus = GetBusIfEnabled();
-
-		if (Bus.IsValid())
-		{
-			Bus->Publish(Message, TypeInfo, Scope, Delay, Expiration, AsShared());
-		}
-	}
+	MESSAGING_API void Publish(void* Message, UScriptStruct* TypeInfo, EMessageScope::Type Scope, const FTimespan& Delay, const FDateTime& Expiration);
 
 	/**
 	 * Sends a message to the specified list of recipients.
@@ -237,15 +197,7 @@ public:
 	 * @param Delay The delay after which to send the message.
 	 * @param Expiration The time at which the message expires.
 	 */
-	void Send( void* Message, UScriptStruct* TypeInfo, const IMessageAttachmentPtr& Attachment, const TArray<FMessageAddress>& Recipients, const FTimespan& Delay, const FDateTime& Expiration )
-	{
-		IMessageBusPtr Bus = GetBusIfEnabled();
-
-		if (Bus.IsValid())
-		{
-			Bus->Send(Message, TypeInfo, Attachment, Recipients, Delay, Expiration, AsShared());
-		}
-	}
+	MESSAGING_API void Send(void* Message, UScriptStruct* TypeInfo, const IMessageAttachmentPtr& Attachment, const TArray<FMessageAddress>& Recipients, const FTimespan& Delay, const FDateTime& Expiration);
 
 	/**
 	 * Subscribes a message handler.
@@ -253,15 +205,7 @@ public:
 	 * @param MessageType The type name of the messages to subscribe to.
 	 * @param ScopeRange The range of message scopes to include in the subscription.
 	 */
-	void Subscribe( const FName& MessageType, const FMessageScopeRange& ScopeRange )
-	{
-		IMessageBusPtr Bus = GetBusIfEnabled();
-
-		if (Bus.IsValid())
-		{
-			Bus->Subscribe(AsShared(), MessageType, ScopeRange);
-		}
-	}
+	MESSAGING_API void Subscribe(const FName& MessageType, const FMessageScopeRange& ScopeRange);
 
 	/**
 	 * Unsubscribes this endpoint from the specified message type.
@@ -269,15 +213,7 @@ public:
 	 * @param MessageType The type of message to unsubscribe (NAME_All = all types).
 	 * @see Subscribe
 	 */
-	void Unsubscribe( const FName& TopicPattern )
-	{
-		IMessageBusPtr Bus = GetBusIfEnabled();
-
-		if (Bus.IsValid())
-		{
-			Bus->Unsubscribe(AsShared(), TopicPattern);
-		}
-	}
+	MESSAGING_API void Unsubscribe(const FName& TopicPattern);
 
 public:
 
@@ -775,23 +711,7 @@ protected:
 	 *
 	 * @param Context The context of the message to handle.
 	 */
-	void ProcessMessage( const IMessageContextRef& Context )
-	{
-		if (!Context->IsValid())
-		{
-			return;
-		}
-
-		for (int32 HandlerIndex = 0; HandlerIndex < Handlers.Num(); ++HandlerIndex)
-		{
-			IMessageHandlerPtr& Handler = Handlers[HandlerIndex];
-
-			if (Handler->GetHandledMessageType() == Context->GetMessageType())
-			{
-				Handler->HandleMessage(Context);
-			}
-		}
-	}
+	MESSAGING_API void ProcessMessage(const IMessageContextRef& Context);
 
 private:
 
diff --git a/Engine/Source/Runtime/Messaging/Public/Common/MessageEndpointBuilder.h b/Engine/Source/Runtime/Messaging/Public/Common/MessageEndpointBuilder.h
--- a/Engine/Source/Runtime/Messaging/Public/Common/MessageEndpointBuilder.h
+++ b/Engine/Source/Runtime/Messaging/Public/Common/MessageEndpointBuilder.h
@@ -150,35 +150,7 @@ public:
 	 *
 	 * @return A new message endpoint, or nullptr if it couldn't be built.
 	 */
-	FMessageEndpointPtr Build( )
-	{
-		FMessageEndpointPtr Endpoint;
-
-		IMessageBusPtr Bus = BusPtr.Pin();
-		
-		if (Bus.IsValid())
-		{
-			Endpoint = MakeShareable(new FMessageEndpoint(Name, Bus.ToSharedRef(), Handlers));
-			Bus->Register(Endpoint->GetAddress(), Endpoint.ToSharedRef());
-
-			if (Disabled)
-			{
-				Endpoint->Disable();
-			}
-
-			if (InboxEnabled)
-			{
-				Endpoint->EnableInbox();
-				Endpoint->SetRecipientThread(ENamedThreads::AnyThread);
-			}
-			else
-			{
-				Endpoint->SetRecipientThread(RecipientThread);
-			}
-		}
-
-		return Endpoint;
-	}
+	MESSAGING_API FMessageEndpointPtr Build();
 
 	/**
 	 * Implicit conversion operator to build the message endpoint as configured.
diff --git a/Engine/Source/Runtime/SlateCore/Public/Sound/ISlateSoundDevice.h b/Engine/Source/Runtime/SlateCore/Public/Sound/ISlateSoundDevice.h
--- a/Engine/Source/Runtime/SlateCore/Public/Sound/ISlateSoundDevice.h
+++ b/Engine/Source/Runtime/SlateCore/Public/Sound/ISlateSoundDevice.h
@@ -1,6 +1,7 @@
 // Copyright 1998-2014 Epic Games, Inc. All Rights Reserved.
 #pragma once
 
+struct FSlateSound;
 
 /** Interface that Slate uses to play sounds. */
 class SLATECORE_API ISlateSoundDevice
