Added missing SplineMeshSceneProxy.cpp

diff --git a/Engine/Source/Runtime/Engine/Private/SplineMeshSceneProxy.cpp b/Engine/Source/Runtime/Engine/Private/SplineMeshSceneProxy.cpp
new file mode 100644
--- /dev/null
+++ b/Engine/Source/Runtime/Engine/Private/SplineMeshSceneProxy.cpp
@@ -0,0 +1,90 @@
+#include "EnginePrivate.h"
+#include "SplineMeshSceneProxy.h"
+
+FSplineMeshSceneProxy::FSplineMeshSceneProxy(USplineMeshComponent* InComponent) :
+	FStaticMeshSceneProxy(InComponent)
+{
+	// make sure all the materials are okay to be rendered as a spline mesh
+	for (FStaticMeshSceneProxy::FLODInfo& LODInfo : LODs)
+	{
+		for (FStaticMeshSceneProxy::FLODInfo::FSectionInfo& Section : LODInfo.Sections)
+		{
+			if (!Section.Material->CheckMaterialUsage(MATUSAGE_SplineMesh))
+			{
+				Section.Material = UMaterial::GetDefaultMaterial(MD_Surface);
+			}
+		}
+	}
+
+	// Copy spline params from component
+	SplineParams = InComponent->SplineParams;
+	SplineUpDir = InComponent->SplineUpDir;
+	bSmoothInterpRollScale = InComponent->bSmoothInterpRollScale;
+	ForwardAxis = InComponent->ForwardAxis;
+
+	// Fill in info about the mesh
+	FBoxSphereBounds StaticMeshBounds = StaticMesh->GetBounds();
+	SplineMeshScaleZ = 0.5f / USplineMeshComponent::GetAxisValue(StaticMeshBounds.BoxExtent, ForwardAxis); // 1/(2 * Extent)
+	SplineMeshMinZ = USplineMeshComponent::GetAxisValue(StaticMeshBounds.Origin, ForwardAxis) * SplineMeshScaleZ - 0.5f;
+
+	LODResources.Reset(InComponent->StaticMesh->RenderData->LODResources.Num());
+
+	for (int32 LODIndex = 0; LODIndex < LODs.Num(); LODIndex++)
+	{
+		FSplineMeshVertexFactory* VertexFactory = new FSplineMeshVertexFactory(this);
+
+		LODResources.Add(VertexFactory);
+
+		InitResources(InComponent, LODIndex);
+	}
+}
+
+FSplineMeshSceneProxy::~FSplineMeshSceneProxy()
+{
+	ReleaseResources();
+
+	for (FSplineMeshSceneProxy::FLODResources& LODResource : LODResources)
+	{
+		delete LODResource.VertexFactory;
+	}
+	LODResources.Empty();
+}
+
+bool FSplineMeshSceneProxy::GetShadowMeshElement(int32 LODIndex, uint8 InDepthPriorityGroup, FMeshBatch& OutMeshElement) const
+{
+	//checkf(LODIndex == 0, TEXT("Getting spline static mesh element with invalid LOD [%d]"), LODIndex);
+
+	if (FStaticMeshSceneProxy::GetShadowMeshElement(LODIndex, InDepthPriorityGroup, OutMeshElement))
+	{
+		OutMeshElement.VertexFactory = LODResources[LODIndex].VertexFactory;
+		OutMeshElement.ReverseCulling ^= (SplineParams.StartScale.X < 0) ^ (SplineParams.StartScale.Y < 0);
+		return true;
+	}
+	return false;
+}
+
+bool FSplineMeshSceneProxy::GetMeshElement(int32 LODIndex, int32 SectionIndex, uint8 InDepthPriorityGroup, FMeshBatch& OutMeshElement, const bool bUseSelectedMaterial, const bool bUseHoveredMaterial) const
+{
+	//checkf(LODIndex == 0 /*&& SectionIndex == 0*/, TEXT("Getting spline static mesh element with invalid params [%d, %d]"), LODIndex, SectionIndex);
+
+	if (FStaticMeshSceneProxy::GetMeshElement(LODIndex, SectionIndex, InDepthPriorityGroup, OutMeshElement, bUseSelectedMaterial, bUseHoveredMaterial))
+	{
+		OutMeshElement.VertexFactory = LODResources[LODIndex].VertexFactory;
+		OutMeshElement.ReverseCulling ^= (SplineParams.StartScale.X < 0) ^ (SplineParams.StartScale.Y < 0);
+		return true;
+	}
+	return false;
+}
+
+bool FSplineMeshSceneProxy::GetWireframeMeshElement(int32 LODIndex, const FMaterialRenderProxy* WireframeRenderProxy, uint8 InDepthPriorityGroup, FMeshBatch& OutMeshElement) const
+{
+	//checkf(LODIndex == 0, TEXT("Getting spline static mesh element with invalid LOD [%d]"), LODIndex);
+
+	if (FStaticMeshSceneProxy::GetWireframeMeshElement(LODIndex, WireframeRenderProxy, InDepthPriorityGroup, OutMeshElement))
+	{
+		OutMeshElement.VertexFactory = LODResources[LODIndex].VertexFactory;
+		OutMeshElement.ReverseCulling ^= (SplineParams.StartScale.X < 0) ^ (SplineParams.StartScale.Y < 0);
+		return true;
+	}
+	return false;
+}
\ No newline at end of file
