Handle 'Switch User' feature on PS4 properly.

diff --git a/Engine/Source/Runtime/Core/Private/Misc/CallbackDevice.cpp b/Engine/Source/Runtime/Core/Private/Misc/CallbackDevice.cpp
--- a/Engine/Source/Runtime/Core/Private/Misc/CallbackDevice.cpp
+++ b/Engine/Source/Runtime/Core/Private/Misc/CallbackDevice.cpp
@@ -22,6 +22,7 @@ FCoreDelegates::FHotFixDelegate& FCoreDelegates::GetHotfixDelegate(EHotfixDelega
 
 FCoreDelegates::FOnMountPak FCoreDelegates::OnMountPak;
 FCoreDelegates::FOnUserLoginChangedEvent FCoreDelegates::OnUserLoginChangedEvent; 
+FCoreDelegates::FOnUserControllerConnectionChange FCoreDelegates::OnControllerConnectionChange;
 FCoreDelegates::FOnSafeFrameChangedEvent FCoreDelegates::OnSafeFrameChangedEvent;
 FCoreDelegates::FOnHandleSystemEnsure FCoreDelegates::OnHandleSystemEnsure;
 FCoreDelegates::FOnHandleSystemError FCoreDelegates::OnHandleSystemError;
diff --git a/Engine/Source/Runtime/Core/Public/Misc/CallbackDevice.h b/Engine/Source/Runtime/Core/Public/Misc/CallbackDevice.h
--- a/Engine/Source/Runtime/Core/Public/Misc/CallbackDevice.h
+++ b/Engine/Source/Runtime/Core/Public/Misc/CallbackDevice.h
@@ -91,11 +91,19 @@ public:
 	// Callback for handling safe frame area size changes
 	DECLARE_MULTICAST_DELEGATE(FOnSafeFrameChangedEvent);
 
+	// Callback for handling the Controller connection / disconnection
+	// first param is true for a connection, false for a disconnection.
+	// second param is UserID, third is UserIndex / ControllerId.
+	DECLARE_MULTICAST_DELEGATE_ThreeParams(FOnUserControllerConnectionChange, bool, int32, int32);
+
 	// get a hotfix delegate
 	static FHotFixDelegate& GetHotfixDelegate(EHotfixDelegates::Type HotFix);
 
 	// Callback when a user logs in/out of the platform.
 	static FOnUserLoginChangedEvent OnUserLoginChangedEvent;
+	
+	// Callback when controllers disconnected / reconnected
+	static FOnUserControllerConnectionChange OnControllerConnectionChange;
 
 	// Callback when a user changes the safe frame size
 	static FOnSafeFrameChangedEvent OnSafeFrameChangedEvent;
diff --git a/Engine/Source/Runtime/Engine/Classes/Engine/LocalPlayer.h b/Engine/Source/Runtime/Engine/Classes/Engine/LocalPlayer.h
--- a/Engine/Source/Runtime/Engine/Classes/Engine/LocalPlayer.h
+++ b/Engine/Source/Runtime/Engine/Classes/Engine/LocalPlayer.h
@@ -97,6 +97,7 @@ private:
 public:
 	// UObject interface
 	virtual void PostInitProperties() override;
+	virtual void BeginDestroy() override;
 	virtual void FinishDestroy() override;
 	static void AddReferencedObjects(UObject* InThis, FReferenceCollector& Collector);
 	// End of UObject interface
@@ -137,6 +138,14 @@ protected:
 	/** @todo document */
 	void ExecMacro( const TCHAR* Filename, FOutputDevice& Ar );
 
+	/**
+	 * Handle controller connection events from the platform
+	 * @param bConnected - determines connection or disconnection event
+	 * @param UserId - platform specfic UserId
+	 * @param UserIndex - Engine level UserIndex.
+	 */
+	void HandleControllerConnectionChange(bool bConnected, int32 InUserId, int32 InControllerId);
+
 public:
 
 	/**
diff --git a/Engine/Source/Runtime/Engine/Private/LocalPlayer.cpp b/Engine/Source/Runtime/Engine/Private/LocalPlayer.cpp
--- a/Engine/Source/Runtime/Engine/Private/LocalPlayer.cpp
+++ b/Engine/Source/Runtime/Engine/Private/LocalPlayer.cpp
@@ -128,6 +128,11 @@ void ULocalPlayer::PostInitProperties()
 			StereoViewState.Allocate();
 		}
 	}
+
+	if (!HasAnyFlags(RF_ClassDefaultObject))
+	{	
+		FCoreDelegates::OnControllerConnectionChange.AddUObject(this, &ULocalPlayer::HandleControllerConnectionChange);
+	}
 }
 
 void ULocalPlayer::PlayerAdded(UGameViewportClient* InViewportClient, int32 InControllerID)
@@ -271,6 +276,12 @@ void ULocalPlayer::SendSplitJoin()
 	}
 }
 
+void ULocalPlayer::BeginDestroy()
+{
+	FCoreDelegates::OnControllerConnectionChange.RemoveAll(this);
+	Super::BeginDestroy();
+}
+
 void ULocalPlayer::FinishDestroy()
 {
 	if ( !IsTemplate() )
@@ -1494,6 +1505,22 @@ void ULocalPlayer::ExecMacro( const TCHAR* Filename, FOutputDevice& Ar )
 	}
 }
 
+void ULocalPlayer::HandleControllerConnectionChange(bool bConnected, int32 InUserId, int32 InControllerId)
+{
+	// if this is an event for this LocalPlayer
+	if (InControllerId == ControllerId)
+	{
+		// if we lost the connection we need to flush all the keys on the PlayerInput to avoid the PC spinning in place, or firing forever, etc.
+		if (!bConnected)
+		{		
+			if (PlayerController && PlayerController->PlayerInput)
+			{
+				PlayerController->PlayerInput->FlushPressedKeys();	
+			}
+		}
+	}
+}
+
 void ULocalPlayer::SetControllerId( int32 NewControllerId )
 {
 	if ( ControllerId != NewControllerId )
