Add damage prevention channel for radial damage

diff --git a/Engine/Source/Runtime/Engine/Classes/Kismet/GameplayStatics.h b/Engine/Source/Runtime/Engine/Classes/Kismet/GameplayStatics.h
--- a/Engine/Source/Runtime/Engine/Classes/Kismet/GameplayStatics.h
+++ b/Engine/Source/Runtime/Engine/Classes/Kismet/GameplayStatics.h
@@ -157,10 +157,11 @@ class ENGINE_API UGameplayStatics : public UBlueprintFunctionLibrary
 	 * @param DamageCauser - Actor that actually caused the damage (e.g. the grenade that exploded).  This actor will not be damaged and it will not block damage.
 	 * @param InstigatedByController - Controller that was responsible for causing this damage (e.g. player who threw the grenade)
 	 * @param bFullDamage - if true, damage not scaled based on distance from Origin
+	 * @param DamagePreventionChannel - Damage will not be applied to victim if there is something between the origin and the victim which blocks traces on this channel
 	 * @return true if damage was applied to at least one actor.
 	 */
 	UFUNCTION(BlueprintCallable, BlueprintAuthorityOnly, Category="Game|Damage", meta=(HidePin="WorldContextObject", DefaultToSelf="WorldContextObject", AutoCreateRefTerm="IgnoreActors"))
-	static bool ApplyRadialDamage(UObject* WorldContextObject, float BaseDamage, const FVector& Origin, float DamageRadius, TSubclassOf<class UDamageType> DamageTypeClass, const TArray<AActor*>& IgnoreActors, AActor* DamageCauser=NULL, AController* InstigatedByController=NULL, bool bDoFullDamage=false );
+		static bool ApplyRadialDamage(UObject* WorldContextObject, float BaseDamage, const FVector& Origin, float DamageRadius, TSubclassOf<class UDamageType> DamageTypeClass, const TArray<AActor*>& IgnoreActors, AActor* DamageCauser = NULL, AController* InstigatedByController = NULL, bool bDoFullDamage = false, ECollisionChannel DamagePreventionChannel = ECC_Visibility);
 	
 	/** Hurt locally authoritative actors within the radius. Will only hit components that block the Visibility channel.
 	 * @param BaseDamage - The base damage to apply, i.e. the damage at the origin.
@@ -172,10 +173,11 @@ class ENGINE_API UGameplayStatics : public UBlueprintFunctionLibrary
 	 * @param DamageCauser - Actor that actually caused the damage (e.g. the grenade that exploded)
 	 * @param InstigatedByController - Controller that was responsible for causing this damage (e.g. player who threw the grenade)
 	 * @param bFullDamage - if true, damage not scaled based on distance from Origin
+	 * @param DamagePreventionChannel - Damage will not be applied to victim if there is something between the origin and the victim which blocks traces on this channel
 	 * @return true if damage was applied to at least one actor.
 	 */
 	UFUNCTION(BlueprintCallable, BlueprintAuthorityOnly, Category="Game|Damage", meta=(HidePin="WorldContextObject", DefaultToSelf="WorldContextObject", AutoCreateRefTerm="IgnoreActors"))
-	static bool ApplyRadialDamageWithFalloff(UObject* WorldContextObject, float BaseDamage, float MinimumDamage, const FVector& Origin, float DamageInnerRadius, float DamageOuterRadius, float DamageFalloff, TSubclassOf<class UDamageType> DamageTypeClass, const TArray<AActor*>& IgnoreActors, AActor* DamageCauser=NULL, AController* InstigatedByController=NULL );
+		static bool ApplyRadialDamageWithFalloff(UObject* WorldContextObject, float BaseDamage, float MinimumDamage, const FVector& Origin, float DamageInnerRadius, float DamageOuterRadius, float DamageFalloff, TSubclassOf<class UDamageType> DamageTypeClass, const TArray<AActor*>& IgnoreActors, AActor* DamageCauser = NULL, AController* InstigatedByController = NULL, ECollisionChannel DamagePreventionChannel = ECC_Visibility);
 	
 
 	/** Hurts the specified actor with the specified impact.
diff --git a/Engine/Source/Runtime/Engine/Private/GameplayStatics.cpp b/Engine/Source/Runtime/Engine/Private/GameplayStatics.cpp
--- a/Engine/Source/Runtime/Engine/Private/GameplayStatics.cpp
+++ b/Engine/Source/Runtime/Engine/Private/GameplayStatics.cpp
@@ -122,7 +122,7 @@ bool UGameplayStatics::SetGamePaused(UObject* WorldContextObject, bool bPaused)
 }
 
 /** @RETURN True if weapon trace from Origin hits component VictimComp.  OutHitResult will contain properties of the hit. */
-static bool ComponentIsVisibleFrom(UPrimitiveComponent* VictimComp, FVector const& Origin, AActor const* IgnoredActor, const TArray<AActor*>& IgnoreActors, FHitResult& OutHitResult)
+static bool ComponentIsDamageableFrom(UPrimitiveComponent* VictimComp, FVector const& Origin, AActor const* IgnoredActor, const TArray<AActor*>& IgnoreActors, ECollisionChannel TraceChannel, FHitResult& OutHitResult)
 {
 	static FName NAME_ComponentIsVisibleFrom = FName(TEXT("ComponentIsVisibleFrom"));
 	FCollisionQueryParams LineParams(NAME_ComponentIsVisibleFrom, true, IgnoredActor);
@@ -139,7 +139,7 @@ static bool ComponentIsVisibleFrom(UPrimitiveComponent* VictimComp, FVector cons
 		// tiny nudge so LineTraceSingle doesn't early out with no hits
 		TraceStart.Z += 0.01f;
 	}
-	bool const bHadBlockingHit = World->LineTraceSingle(OutHitResult, TraceStart, TraceEnd, ECC_Visibility, LineParams);
+	bool const bHadBlockingHit = World->LineTraceSingle(OutHitResult, TraceStart, TraceEnd, TraceChannel, LineParams);
 	//::DrawDebugLine(World, TraceStart, TraceEnd, FLinearColor::Red, true);
 
 	// If there was a blocking hit, it will be the last one
@@ -166,13 +166,13 @@ static bool ComponentIsVisibleFrom(UPrimitiveComponent* VictimComp, FVector cons
 	return true;
 }
 
-bool UGameplayStatics::ApplyRadialDamage(UObject* WorldContextObject, float BaseDamage, const FVector& Origin, float DamageRadius, TSubclassOf<UDamageType> DamageTypeClass, const TArray<AActor*>& IgnoreActors, AActor* DamageCauser, AController* InstigatedByController, bool bDoFullDamage )
+bool UGameplayStatics::ApplyRadialDamage(UObject* WorldContextObject, float BaseDamage, const FVector& Origin, float DamageRadius, TSubclassOf<UDamageType> DamageTypeClass, const TArray<AActor*>& IgnoreActors, AActor* DamageCauser, AController* InstigatedByController, bool bDoFullDamage, ECollisionChannel DamagePreventionChannel )
 {
 	float DamageFalloff = bDoFullDamage ? 0.f : 1.f;
-	return ApplyRadialDamageWithFalloff(WorldContextObject, BaseDamage, 0.f, Origin, 0.f, DamageRadius, DamageFalloff, DamageTypeClass, IgnoreActors, DamageCauser, InstigatedByController);
+	return ApplyRadialDamageWithFalloff(WorldContextObject, BaseDamage, 0.f, Origin, 0.f, DamageRadius, DamageFalloff, DamageTypeClass, IgnoreActors, DamageCauser, InstigatedByController, DamagePreventionChannel);
 }
 
-bool UGameplayStatics::ApplyRadialDamageWithFalloff(UObject* WorldContextObject, float BaseDamage, float MinimumDamage, const FVector& Origin, float DamageInnerRadius, float DamageOuterRadius, float DamageFalloff, TSubclassOf<class UDamageType> DamageTypeClass, const TArray<AActor*>& IgnoreActors, AActor* DamageCauser, AController* InstigatedByController )
+bool UGameplayStatics::ApplyRadialDamageWithFalloff(UObject* WorldContextObject, float BaseDamage, float MinimumDamage, const FVector& Origin, float DamageInnerRadius, float DamageOuterRadius, float DamageFalloff, TSubclassOf<class UDamageType> DamageTypeClass, const TArray<AActor*>& IgnoreActors, AActor* DamageCauser, AController* InstigatedByController, ECollisionChannel DamagePreventionChannel)
 {
 	static FName NAME_ApplyRadialDamage = FName(TEXT("ApplyRadialDamage"));
 	FCollisionQueryParams SphereParams(NAME_ApplyRadialDamage, false, DamageCauser);
@@ -197,7 +197,7 @@ bool UGameplayStatics::ApplyRadialDamageWithFalloff(UObject* WorldContextObject,
 			Overlap.Component.IsValid() )
 		{
 			FHitResult Hit;
-			if (ComponentIsVisibleFrom(Overlap.Component.Get(), Origin, DamageCauser, IgnoreActors, Hit))
+			if (DamagePreventionChannel == ECC_MAX || ComponentIsDamageableFrom(Overlap.Component.Get(), Origin, DamageCauser, IgnoreActors, DamagePreventionChannel, Hit))
 			{
 				TArray<FHitResult>& HitList = OverlapComponentMap.FindOrAdd(OverlapActor);
 				HitList.Add(Hit);
