Use false instead of 0

diff --git a/Engine/Source/Runtime/Engine/Private/PlayerController.cpp b/Engine/Source/Runtime/Engine/Private/PlayerController.cpp
--- a/Engine/Source/Runtime/Engine/Private/PlayerController.cpp
+++ b/Engine/Source/Runtime/Engine/Private/PlayerController.cpp
@@ -2967,7 +2967,7 @@ bool APlayerController::IsPrimaryPlayer() const
 
 bool APlayerController::IsSplitscreenPlayer(int32* OutSplitscreenPlayerIndex) const
 {
-	bool bResult = 0;
+	bool bResult = false;
 
 	if (OutSplitscreenPlayerIndex)
 	{
