#ue4 - Change some log messages on player startup to reduce log spam.

diff --git a/Engine/Source/Runtime/Engine/Private/GameMode.cpp b/Engine/Source/Runtime/Engine/Private/GameMode.cpp
--- a/Engine/Source/Runtime/Engine/Private/GameMode.cpp
+++ b/Engine/Source/Runtime/Engine/Private/GameMode.cpp
@@ -415,7 +415,7 @@ void AGameMode::RestartPlayer(AController* NewPlayer)
 		return;
 	}
 
-	UE_LOG(LogGameMode, Log, TEXT("RestartPlayer %s"), (NewPlayer && NewPlayer->PlayerState) ? *NewPlayer->PlayerState->PlayerName : TEXT("Unknown"));
+	UE_LOG(LogGameMode, Verbose, TEXT("RestartPlayer %s"), (NewPlayer && NewPlayer->PlayerState) ? *NewPlayer->PlayerState->PlayerName : TEXT("Unknown"));
 	AActor* StartSpot = FindPlayerStart(NewPlayer);
 
 	// if a start spot wasn't found,
diff --git a/Engine/Source/Runtime/Engine/Private/NetConnection.cpp b/Engine/Source/Runtime/Engine/Private/NetConnection.cpp
--- a/Engine/Source/Runtime/Engine/Private/NetConnection.cpp
+++ b/Engine/Source/Runtime/Engine/Private/NetConnection.cpp
@@ -1472,7 +1472,7 @@ void UNetConnection::HandleClientPlayer( APlayerController *PC, UNetConnection*
 	PC->Role = ROLE_AutonomousProxy;
 	PC->SetPlayer(LocalPlayer);
 	PC->NetConnection = NetConnection;
-	UE_LOG(LogNet, Log, TEXT("%s setplayer %s"),*PC->GetName(),*LocalPlayer->GetName());
+	UE_LOG(LogNet, Verbose, TEXT("%s setplayer %s"),*PC->GetName(),*LocalPlayer->GetName());
 	LastReceiveTime = Driver->Time;
 	State = USOCK_Open;
 	PlayerController = PC;
@@ -1573,7 +1573,7 @@ void UChildConnection::HandleClientPlayer(APlayerController* PC, UNetConnection*
 	PC->Role = ROLE_AutonomousProxy;
 	PC->SetPlayer(NewPlayer);
 	PC->NetConnection = NetConnection;
-	UE_LOG(LogNet, Log, TEXT("%s setplayer %s"), *PC->GetName(), *NewPlayer->GetName());
+	UE_LOG(LogNet, Verbose, TEXT("%s setplayer %s"), *PC->GetName(), *NewPlayer->GetName());
 	PlayerController = PC;
 	OwningActor = PC;
 }
diff --git a/Engine/Source/Runtime/Engine/Private/PlayerController.cpp b/Engine/Source/Runtime/Engine/Private/PlayerController.cpp
--- a/Engine/Source/Runtime/Engine/Private/PlayerController.cpp
+++ b/Engine/Source/Runtime/Engine/Private/PlayerController.cpp
@@ -243,7 +243,7 @@ void APlayerController::ServerUpdateLevelVisibility_Implementation(FName Package
 			if ( Linker || FPackageName::DoesPackageExist(PackageName.ToString(), NULL, &Filename ) || Local::IsInLevelList(GetWorld(), PackageName ) )
 			{
 				Connection->ClientVisibleLevelNames.AddUnique(PackageName);
-				UE_LOG( LogPlayerController, Log, TEXT("ServerUpdateLevelVisibility() Added '%s'"), *PackageName.ToString() );
+				UE_LOG( LogPlayerController, Verbose, TEXT("ServerUpdateLevelVisibility() Added '%s'"), *PackageName.ToString() );
 			}
 			else
 			{
@@ -254,7 +254,7 @@ void APlayerController::ServerUpdateLevelVisibility_Implementation(FName Package
 		else
 		{
 			Connection->ClientVisibleLevelNames.Remove(PackageName);
-			UE_LOG( LogPlayerController, Log, TEXT("ServerUpdateLevelVisibility() Removed '%s'"), *PackageName.ToString() );
+			UE_LOG( LogPlayerController, Verbose, TEXT("ServerUpdateLevelVisibility() Removed '%s'"), *PackageName.ToString() );
 			
 			// Close any channels now that have actors that were apart of the level the client just unloaded
 			for ( auto It = Connection->ActorChannels.CreateIterator(); It; ++It )
@@ -412,7 +412,7 @@ bool APlayerController::ServerNotifyLoadedWorld_Validate(FName WorldPackageName)
 
 void APlayerController::ServerNotifyLoadedWorld_Implementation(FName WorldPackageName)
 {
-	UE_LOG(LogPlayerController, Log, TEXT("APlayerController::ServerNotifyLoadedWorld_Implementation: Client loaded %s"), *WorldPackageName.ToString());
+	UE_LOG(LogPlayerController, Verbose, TEXT("APlayerController::ServerNotifyLoadedWorld_Implementation: Client loaded %s"), *WorldPackageName.ToString());
 
 	UWorld* CurWorld = GetWorld();
 
@@ -588,7 +588,7 @@ void APlayerController::ClientRetryClientRestart_Implementation(APawn* NewPawn)
 		return;
 	}
 
-	UE_LOG(LogPlayerController, Log, TEXT("ClientRetryClientRestart_Implementation %s, AcknowledgedPawn: %s"), *GetNameSafe(NewPawn), *GetNameSafe(AcknowledgedPawn));
+	UE_LOG(LogPlayerController, Verbose, TEXT("ClientRetryClientRestart_Implementation %s, AcknowledgedPawn: %s"), *GetNameSafe(NewPawn), *GetNameSafe(AcknowledgedPawn));
 
 	// Avoid calling ClientRestart if we have already accepted this pawn
 	if( (GetPawn() != NewPawn) || (NewPawn->Controller != this) || (NewPawn != AcknowledgedPawn) )
@@ -601,7 +601,7 @@ void APlayerController::ClientRetryClientRestart_Implementation(APawn* NewPawn)
 
 void APlayerController::ClientRestart_Implementation(APawn* NewPawn)
 {
-	UE_LOG(LogPlayerController, Log, TEXT("ClientRestart_Implementation %s"), *GetNameSafe(NewPawn));
+	UE_LOG(LogPlayerController, Verbose, TEXT("ClientRestart_Implementation %s"), *GetNameSafe(NewPawn));
 
 	ResetIgnoreInputFlags();
 	AcknowledgedPawn = NULL;
@@ -907,7 +907,7 @@ void APlayerController::SpawnDefaultHUD()
 		return;
 	}
 
-	UE_LOG(LogPlayerController, Log, TEXT("SpawnDefaultHUD"));
+	UE_LOG(LogPlayerController, Verbose, TEXT("SpawnDefaultHUD"));
 	FActorSpawnParameters SpawnInfo;
 	SpawnInfo.Owner = this;
 	SpawnInfo.Instigator = Instigator;
@@ -1014,7 +1014,7 @@ bool APlayerController::IsFrozen()
 
 void APlayerController::ServerAcknowledgePossession_Implementation(APawn* P)
 {
-	UE_LOG(LogPlayerController, Log, TEXT("ServerAcknowledgePossession_Implementation %s"), *GetNameSafe(P));
+	UE_LOG(LogPlayerController, Verbose, TEXT("ServerAcknowledgePossession_Implementation %s"), *GetNameSafe(P));
 	AcknowledgedPawn = P;
 }
 
@@ -1094,7 +1094,7 @@ void APlayerController::OnActorChannelOpen(FInBunch& InBunch, UNetConnection* Co
 				// create a split connection on the client
 				UChildConnection* Child = Connection->Driver->CreateChild(Connection); 
 				Child->HandleClientPlayer(this, Connection);
-				UE_LOG(LogNet, Log, TEXT("Client received PlayerController=%s. Num child connections=%i."), *GetName(), Connection->Children.Num());
+				UE_LOG(LogNet, Verbose, TEXT("Client received PlayerController=%s. Num child connections=%i."), *GetName(), Connection->Children.Num());
 			}
 		}
 	}
@@ -2551,7 +2551,7 @@ bool APlayerController::ServerRestartPlayer_Validate()
 
 void APlayerController::ServerRestartPlayer_Implementation()
 {
-	UE_LOG(LogPlayerController, Log, TEXT("SERVER RESTART PLAYER"));
+	UE_LOG(LogPlayerController, Verbose, TEXT("SERVER RESTART PLAYER"));
 	if ( GetNetMode() == NM_Client )
 	{
 		return;
@@ -3593,7 +3593,7 @@ ASpectatorPawn* APlayerController::SpawnSpectatorPawn()
 				SpawnedSpectator->PawnClientRestart();
 				SpawnedSpectator->SetActorTickEnabled(true);
 
-				UE_LOG(LogPlayerController, Log, TEXT("Spawned spectator %s [server:%d]"), *GetNameSafe(SpawnedSpectator), GetNetMode() < NM_Client);
+				UE_LOG(LogPlayerController, Verbose, TEXT("Spawned spectator %s [server:%d]"), *GetNameSafe(SpawnedSpectator), GetNetMode() < NM_Client);
 			}
 			else
 			{
@@ -3603,7 +3603,7 @@ ASpectatorPawn* APlayerController::SpawnSpectatorPawn()
 		else
 		{
 			// This normally happens on clients if the Player is replicated but the GameState has not yet.
-			UE_LOG(LogPlayerController, Log, TEXT("NULL GameState when trying to spawn spectator!"));
+			UE_LOG(LogPlayerController, Verbose, TEXT("NULL GameState when trying to spawn spectator!"));
 		}
 	}
 
@@ -3707,7 +3707,7 @@ void APlayerController::EndSpectatingState()
 	{
 		if ( PlayerState->bOnlySpectator )
 		{
-			UE_LOG(LogPlayerController, Log, TEXT("WARNING - Spectator only UPlayer* leaving spectating state"));
+			UE_LOG(LogPlayerController, Warning, TEXT("Spectator only UPlayer* leaving spectating state"));
 		}
 		PlayerState->bIsSpectator = false;
 	}
