Fix non-unity compile

diff --git a/Engine/Source/Runtime/Engine/Classes/AI/Navigation/NavigationTypes.h b/Engine/Source/Runtime/Engine/Classes/AI/Navigation/NavigationTypes.h
--- a/Engine/Source/Runtime/Engine/Classes/AI/Navigation/NavigationTypes.h
+++ b/Engine/Source/Runtime/Engine/Classes/AI/Navigation/NavigationTypes.h
@@ -24,6 +24,7 @@ typedef uint64 NavNodeRef;
 struct FNavigationQueryFilter;
 class AActor;
 class ANavigationData;
+class INavAgentInterface;
 
 namespace FNavigationSystem
 {
