TTP 341356 : Check Owner before accessing it in URadialForceComponent::TickComponent (is NULL in BSP case)

diff --git a/Engine/Source/Runtime/Engine/Private/PhysicsEngine/RadialForceComponent.cpp b/Engine/Source/Runtime/Engine/Private/PhysicsEngine/RadialForceComponent.cpp
--- a/Engine/Source/Runtime/Engine/Private/PhysicsEngine/RadialForceComponent.cpp
+++ b/Engine/Source/Runtime/Engine/Private/PhysicsEngine/RadialForceComponent.cpp
@@ -50,14 +50,18 @@ void URadialForceComponent::TickComponent(float DeltaTime, enum ELevelTick TickT
 				PokeComp->AddRadialForce( Origin, Radius, ForceStrength, Falloff );
 
 				// see if this is a target for a movement component
-				TArray<UMovementComponent*> MovementComponents;
-				PokeComp->GetOwner()->GetComponents<UMovementComponent>(MovementComponents);
-				for(const auto& MovementComponent : MovementComponents)
+				AActor* PokeOwner = PokeComp->GetOwner();
+				if(PokeOwner)
 				{
-					if(MovementComponent->UpdatedComponent == PokeComp)
+					TArray<UMovementComponent*> MovementComponents;
+					PokeOwner->GetComponents<UMovementComponent>(MovementComponents);
+					for(const auto& MovementComponent : MovementComponents)
 					{
-						MovementComponent->AddRadialForce( Origin, Radius, ForceStrength, Falloff );
-						break;
+						if(MovementComponent->UpdatedComponent == PokeComp)
+						{
+							MovementComponent->AddRadialForce( Origin, Radius, ForceStrength, Falloff );
+							break;
+						}
 					}
 				}
 			}
