Current blueprint menu system can now take advantage of the EditorCategoryUtils category key system.

diff --git a/Engine/Source/Runtime/Engine/Private/EdGraph/EdGraphSchema.cpp b/Engine/Source/Runtime/Engine/Private/EdGraph/EdGraphSchema.cpp
--- a/Engine/Source/Runtime/Engine/Private/EdGraph/EdGraphSchema.cpp
+++ b/Engine/Source/Runtime/Engine/Private/EdGraph/EdGraphSchema.cpp
@@ -4,6 +4,7 @@
 #include "EdGraph/EdGraph.h"
 #include "EdGraph/EdGraphSchema.h"
 #include "BlueprintUtilities.h"
+#include "EditorCategoryUtils.h"
 #if WITH_EDITOR
 #include "Editor/UnrealEd/Public/Kismet2/BlueprintEditorUtils.h"
 #include "Slate.h"
@@ -64,12 +65,14 @@ FGraphActionListBuilderBase::ActionGroup::ActionGroup( const TArray< TSharedPtr<
 void FGraphActionListBuilderBase::ActionGroup::GetCategoryChain(TArray<FString>& HierarchyOut) const
 {
 	static FString const CategoryDelim("|");
-	RootCategory.ParseIntoArray(&HierarchyOut, *CategoryDelim, true);
+	FEditorCategoryUtils::GetCategoryDisplayString(RootCategory).ParseIntoArray(&HierarchyOut, *CategoryDelim, true);
 
 	if (Actions.Num() > 0)
 	{
 		TArray<FString> SubCategoryChain;
-		Actions[0]->Category.ParseIntoArray(&SubCategoryChain, *CategoryDelim, true);
+
+		FString SubCategory = FEditorCategoryUtils::GetCategoryDisplayString(Actions[0]->Category);
+		SubCategory.ParseIntoArray(&SubCategoryChain, *CategoryDelim, true);
 
 		HierarchyOut.Append(SubCategoryChain);
 	}
