Fixing broken builder, bad include path.

diff --git a/Engine/Source/Runtime/Engine/Private/EdGraph/EdGraphSchema.cpp b/Engine/Source/Runtime/Engine/Private/EdGraph/EdGraphSchema.cpp
--- a/Engine/Source/Runtime/Engine/Private/EdGraph/EdGraphSchema.cpp
+++ b/Engine/Source/Runtime/Engine/Private/EdGraph/EdGraphSchema.cpp
@@ -4,12 +4,12 @@
 #include "EdGraph/EdGraph.h"
 #include "EdGraph/EdGraphSchema.h"
 #include "BlueprintUtilities.h"
-#include "EditorCategoryUtils.h"
 #if WITH_EDITOR
 #include "Editor/UnrealEd/Public/Kismet2/BlueprintEditorUtils.h"
 #include "Slate.h"
 #include "ScopedTransaction.h"
 #include "Editor/UnrealEd/Public/Kismet2/Kismet2NameValidators.h"
+#include "Editor/UnrealEd/Public/EditorCategoryUtils.h"
 #endif
 
 #define LOCTEXT_NAMESPACE "EdGraph"
@@ -64,6 +64,7 @@ FGraphActionListBuilderBase::ActionGroup::ActionGroup( const TArray< TSharedPtr<
 
 void FGraphActionListBuilderBase::ActionGroup::GetCategoryChain(TArray<FString>& HierarchyOut) const
 {
+#if WITH_EDITOR
 	static FString const CategoryDelim("|");
 	FEditorCategoryUtils::GetCategoryDisplayString(RootCategory).ParseIntoArray(&HierarchyOut, *CategoryDelim, true);
 
@@ -81,6 +82,7 @@ void FGraphActionListBuilderBase::ActionGroup::GetCategoryChain(TArray<FString>&
 	{
 		Category.Trim();
 	}
+#endif
 }
 
 void FGraphActionListBuilderBase::ActionGroup::PerformAction( class UEdGraph* ParentGraph, TArray<UEdGraphPin*>& FromPins, const FVector2D Location )
