Some improvements to "stat levels" command when there are many streaming levels

diff --git a/Engine/Source/Runtime/Engine/Private/EngineUtils.cpp b/Engine/Source/Runtime/Engine/Private/EngineUtils.cpp
--- a/Engine/Source/Runtime/Engine/Private/EngineUtils.cpp
+++ b/Engine/Source/Runtime/Engine/Private/EngineUtils.cpp
@@ -369,7 +369,7 @@ TArray<FSubLevelStatus> GetSubLevelsStatus( UWorld* World )
 				{
 					LevelStatus.StreamingStatus = LEVEL_UnloadedButStillAround;
 				}
-				else if( GetAsyncLoadPercentage( LevelStatus.PackageName.ToString() ) >= 0 )
+				else if( LevelStreaming->bHasLoadRequestPending )
 				{
 					LevelStatus.StreamingStatus = LEVEL_Loading;
 				}
diff --git a/Engine/Source/Runtime/Engine/Private/UnrealEngine.cpp b/Engine/Source/Runtime/Engine/Private/UnrealEngine.cpp
--- a/Engine/Source/Runtime/Engine/Private/UnrealEngine.cpp
+++ b/Engine/Source/Runtime/Engine/Private/UnrealEngine.cpp
@@ -10206,7 +10206,7 @@ int32 UEngine::RenderStatLevels(UWorld* World, FViewport* Viewport, FCanvas* Can
 		{
 			MaxY = FMath::Max(MaxY, Y);
 			Y = BaseY;
-			X += 250;
+			X += 350;
 		}
 
 		FColor	Color = GetColorForLevelStatus(LevelStatus.StreamingStatus);
@@ -10225,10 +10225,14 @@ int32 UEngine::RenderStatLevels(UWorld* World, FViewport* Viewport, FCanvas* Can
 		{
 			DisplayName += FString::Printf(TEXT(" - %4.1f sec"), LevelPackage->GetLoadTime());
 		}
-		else if (GetAsyncLoadPercentage(*LevelStatus.PackageName.ToString()) >= 0)
+		else
 		{
-			const int32 Percentage = FMath::TruncToInt(GetAsyncLoadPercentage(*LevelStatus.PackageName.ToString()));
-			DisplayName += FString::Printf(TEXT(" - %3i %%"), Percentage);
+			float AsyncLoadPercentage = GetAsyncLoadPercentage(*LevelStatus.PackageName.ToString());
+			if (AsyncLoadPercentage >= 0)
+			{
+				const int32 Percentage = FMath::TruncToInt(AsyncLoadPercentage);
+				DisplayName += FString::Printf(TEXT(" - %3i %%"), Percentage);
+			}
 		}
 
 		if (LevelStatus.bPlayerInside)
