fixed custom navlink ID for static navmesh

diff --git a/Engine/Source/Runtime/Engine/Classes/AI/Navigation/NavLinkCustomComponent.h b/Engine/Source/Runtime/Engine/Classes/AI/Navigation/NavLinkCustomComponent.h
--- a/Engine/Source/Runtime/Engine/Classes/AI/Navigation/NavLinkCustomComponent.h
+++ b/Engine/Source/Runtime/Engine/Classes/AI/Navigation/NavLinkCustomComponent.h
@@ -39,6 +39,13 @@ class ENGINE_API UNavLinkCustomComponent : public UNavRelevantComponent, public
 	virtual void OnApplyModifiers(struct FCompositeNavModifier& Modifiers) override;
 	// END UNavRelevantComponent
 
+	// BEGIN UObject Interface
+	virtual void PostLoad() override;
+#if WITH_EDITOR
+	virtual void PostEditImport() override;
+#endif
+	// END UObject Interface
+
 	/** set basic link data: end points and direction */
 	void SetLinkData(const FVector& RelativeStart, const FVector& RelativeEnd, ENavLinkDirection::Type Direction);
 	virtual FNavigationLink GetLinkModifier() const;
@@ -111,7 +118,7 @@ class ENGINE_API UNavLinkCustomComponent : public UNavRelevantComponent, public
 protected:
 
 	/** link Id assigned by navigation system */
-	UPROPERTY(DuplicateTransient)
+	UPROPERTY()
 	uint32 NavLinkUserId;
 
 	/** area class to use when link is enabled */
diff --git a/Engine/Source/Runtime/Engine/Classes/AI/Navigation/NavLinkCustomInterface.h b/Engine/Source/Runtime/Engine/Classes/AI/Navigation/NavLinkCustomInterface.h
--- a/Engine/Source/Runtime/Engine/Classes/AI/Navigation/NavLinkCustomInterface.h
+++ b/Engine/Source/Runtime/Engine/Classes/AI/Navigation/NavLinkCustomInterface.h
@@ -54,6 +54,11 @@ class ENGINE_API INavLinkCustomInterface
 	/** Helper function: returns unique ID number for custom links */
 	static uint32 GetUniqueId();
 
+	/** Helper function: bump unique ID numbers above given one */
+	static void UpdateUniqueId(uint32 AlreadyUsedId);
+
 	/** Helper function: create modifier for navigation data export */
 	static FNavigationLink GetModifier(const INavLinkCustomInterface* CustomNavLink);
+
+	static uint32 NextUniqueId;
 };
diff --git a/Engine/Source/Runtime/Engine/Classes/AI/Navigation/NavigationSystem.h b/Engine/Source/Runtime/Engine/Classes/AI/Navigation/NavigationSystem.h
--- a/Engine/Source/Runtime/Engine/Classes/AI/Navigation/NavigationSystem.h
+++ b/Engine/Source/Runtime/Engine/Classes/AI/Navigation/NavigationSystem.h
@@ -225,6 +225,12 @@ public:
 		OctreeUpdate_Refresh = 4,						// update is used for refresh, don't invalidate pending queue
 	};
 
+	enum ECleanupMode
+	{
+		CleanupWithWorld,
+		CleanupUnsafe,
+	};
+
 	// Begin UObject Interface
 	virtual void PostInitProperties() override;
 	static void AddReferencedObjects(UObject* InThis, FReferenceCollector& Collector);
@@ -477,7 +483,7 @@ public:
 	virtual void DoInitialSetup();
 
 	/** Called upon UWorld destruction to release what needs to be released */
-	void CleanUp();
+	void CleanUp(ECleanupMode Mode = ECleanupMode::CleanupUnsafe);
 
 	/** 
 	 *	Called when owner-UWorld initializes actors
diff --git a/Engine/Source/Runtime/Engine/Private/AI/Navigation/NavLinkCustomComponent.cpp b/Engine/Source/Runtime/Engine/Private/AI/Navigation/NavLinkCustomComponent.cpp
--- a/Engine/Source/Runtime/Engine/Private/AI/Navigation/NavLinkCustomComponent.cpp
+++ b/Engine/Source/Runtime/Engine/Private/AI/Navigation/NavLinkCustomComponent.cpp
@@ -30,6 +30,22 @@ UNavLinkCustomComponent::UNavLinkCustomComponent(const class FPostConstructIniti
 	BroadcastInterval = 0.0f;
 }
 
+void UNavLinkCustomComponent::PostLoad()
+{
+	Super::PostLoad();
+
+	INavLinkCustomInterface::UpdateUniqueId(NavLinkUserId);
+}
+
+#if WITH_EDITOR
+void UNavLinkCustomComponent::PostEditImport()
+{
+	Super::PostEditImport();
+
+	NavLinkUserId = INavLinkCustomInterface::GetUniqueId();
+}
+#endif
+
 void UNavLinkCustomComponent::GetLinkData(FVector& LeftPt, FVector& RightPt, ENavLinkDirection::Type& Direction) const
 {
 	LeftPt = LinkRelativeStart;
diff --git a/Engine/Source/Runtime/Engine/Private/AI/Navigation/NavigationInterfaces.cpp b/Engine/Source/Runtime/Engine/Private/AI/Navigation/NavigationInterfaces.cpp
--- a/Engine/Source/Runtime/Engine/Private/AI/Navigation/NavigationInterfaces.cpp
+++ b/Engine/Source/Runtime/Engine/Private/AI/Navigation/NavigationInterfaces.cpp
@@ -9,6 +9,8 @@
 #include "AI/Navigation/NavRelevantActorInterface.h"
 #include "AI/Navigation/NavLinkCustomInterface.h"
 
+uint32 INavLinkCustomInterface::NextUniqueId = 1;
+
 UNavAgentInterface::UNavAgentInterface(const class FPostConstructInitializeProperties& PCIP)
 	: Super(PCIP)
 {
@@ -36,10 +38,14 @@ UNavLinkCustomInterface::UNavLinkCustomInterface(const class FPostConstructIniti
 
 uint32 INavLinkCustomInterface::GetUniqueId()
 {
-	static uint32 NextUniqueId = 1;
 	return NextUniqueId++;
 }
 
+void INavLinkCustomInterface::UpdateUniqueId(uint32 AlreadyUsedId)
+{
+	NextUniqueId = FMath::Max(NextUniqueId, AlreadyUsedId + 1);
+}
+
 FNavigationLink INavLinkCustomInterface::GetModifier(const INavLinkCustomInterface* CustomNavLink)
 {
 	FNavigationLink LinkMod;
diff --git a/Engine/Source/Runtime/Engine/Private/AI/Navigation/NavigationSystem.cpp b/Engine/Source/Runtime/Engine/Private/AI/Navigation/NavigationSystem.cpp
--- a/Engine/Source/Runtime/Engine/Private/AI/Navigation/NavigationSystem.cpp
+++ b/Engine/Source/Runtime/Engine/Private/AI/Navigation/NavigationSystem.cpp
@@ -2723,7 +2723,7 @@ void UNavigationSystem::OnNavigationDirtied(const FBox& Bounds)
 #endif //WITH_NAVIGATION_GENERATOR
 }
 
-void UNavigationSystem::CleanUp()
+void UNavigationSystem::CleanUp(ECleanupMode Mode)
 {
 	UE_LOG(LogNavigation, Log, TEXT("UNavigationSystem::CleanUp"));
 
@@ -2761,6 +2761,13 @@ void UNavigationSystem::CleanUp()
 	SetCrowdManager(NULL);
 
 	NavDataSet.Reset();
+
+	// reset unique link Id for new map
+	const UWorld* MyWorld = (Mode == ECleanupMode::CleanupWithWorld) ? GetWorld() : NULL;
+	if (MyWorld && (MyWorld->WorldType == EWorldType::Game || MyWorld->WorldType == EWorldType::Editor))
+	{
+		INavLinkCustomInterface::NextUniqueId = 1;
+	}
 }
 
 //----------------------------------------------------------------------//
diff --git a/Engine/Source/Runtime/Engine/Private/World.cpp b/Engine/Source/Runtime/Engine/Private/World.cpp
--- a/Engine/Source/Runtime/Engine/Private/World.cpp
+++ b/Engine/Source/Runtime/Engine/Private/World.cpp
@@ -2917,7 +2917,7 @@ void UWorld::CleanupWorld(bool bSessionEnded, bool bCleanupResources, UWorld* Ne
 	{
 		if (NavigationSystem != NULL)
 		{
-			NavigationSystem->CleanUp();
+			NavigationSystem->CleanUp(UNavigationSystem::CleanupWithWorld);
 		}
 	}
 
@@ -5019,7 +5019,7 @@ void UWorld::SetNavigationSystem( UNavigationSystem* InNavigationSystem)
 {
 	if (NavigationSystem != NULL && NavigationSystem != InNavigationSystem)
 	{
-		NavigationSystem->CleanUp();
+		NavigationSystem->CleanUp(UNavigationSystem::CleanupWithWorld);
 	}
 
 	NavigationSystem = InNavigationSystem;
