The bAllowLooping parameter of UAnimSequenceBase::GetAnimNotifies was a float instead of bool, fixing.

diff --git a/Engine/Source/Runtime/Engine/Classes/Animation/AnimSequenceBase.h b/Engine/Source/Runtime/Engine/Classes/Animation/AnimSequenceBase.h
--- a/Engine/Source/Runtime/Engine/Classes/Animation/AnimSequenceBase.h
+++ b/Engine/Source/Runtime/Engine/Classes/Animation/AnimSequenceBase.h
@@ -308,7 +308,7 @@ class UAnimSequenceBase : public UAnimationAsset
 	 * Supports playing backwards (DeltaTime<0).
 	 * Returns notifies between StartTime (exclusive) and StartTime+DeltaTime (inclusive)
 	 */
-	void GetAnimNotifies(const float & StartTime, const float & DeltaTime, const float bAllowLooping, TArray<const FAnimNotifyEvent *> & OutActiveNotifies) const;
+	void GetAnimNotifies(const float & StartTime, const float & DeltaTime, const bool bAllowLooping, TArray<const FAnimNotifyEvent *> & OutActiveNotifies) const;
 
 	/** 
 	 * Retrieves AnimNotifies between two time positions. ]PreviousPosition, CurrentPosition]
diff --git a/Engine/Source/Runtime/Engine/Private/Animation/AnimSequenceBase.cpp b/Engine/Source/Runtime/Engine/Private/Animation/AnimSequenceBase.cpp
--- a/Engine/Source/Runtime/Engine/Private/Animation/AnimSequenceBase.cpp
+++ b/Engine/Source/Runtime/Engine/Private/Animation/AnimSequenceBase.cpp
@@ -221,7 +221,7 @@ void UAnimSequenceBase::SortNotifies()
  * Supports playing backwards (DeltaTime<0).
  * Returns notifies between StartTime (exclusive) and StartTime+DeltaTime (inclusive)
  */
-void UAnimSequenceBase::GetAnimNotifies(const float & StartTime, const float & DeltaTime, const float bAllowLooping, TArray<const FAnimNotifyEvent *> & OutActiveNotifies) const
+void UAnimSequenceBase::GetAnimNotifies(const float & StartTime, const float & DeltaTime, const bool bAllowLooping, TArray<const FAnimNotifyEvent *> & OutActiveNotifies) const
 {
 	// Early out if we have no notifies
 	if( (Notifies.Num() == 0) || (DeltaTime == 0.f) )
