Renamed some references to ' package' which isn't intended to be a user facing term, and reorganized some properties to be advanced.
Polished the custom Yaw UI to look more consistent.

diff --git a/Engine/Source/Editor/WorldBrowser/Private/StreamingLevels/StreamingLevelCustomization.cpp b/Engine/Source/Editor/WorldBrowser/Private/StreamingLevels/StreamingLevelCustomization.cpp
--- a/Engine/Source/Editor/WorldBrowser/Private/StreamingLevels/StreamingLevelCustomization.cpp
+++ b/Engine/Source/Editor/WorldBrowser/Private/StreamingLevels/StreamingLevelCustomization.cpp
@@ -77,6 +77,7 @@ void FStreamingLevelCustomization::CustomizeDetails(IDetailLayoutBuilder& Detail
 
 			+SHorizontalBox::Slot()
 			.FillWidth(1.f)
+			.Padding(0,2,0,2)
 			[
 				SNew(SNumericEntryBox<int32>)
 
@@ -108,6 +109,7 @@ void FStreamingLevelCustomization::CustomizeDetails(IDetailLayoutBuilder& Detail
 				.ToolTipText(LOCTEXT("EditLevelToolTip", "Edit level transform in viewport."))
 				.OnClicked(this, &FStreamingLevelCustomization::OnEditLevelClicked)
 				.IsEnabled(this, &FStreamingLevelCustomization::LevelViewportTransformAllowed)
+				.ContentPadding(1)
 			]
 		];
 }
diff --git a/Engine/Source/Runtime/Engine/Classes/Engine/LevelStreaming.h b/Engine/Source/Runtime/Engine/Classes/Engine/LevelStreaming.h
--- a/Engine/Source/Runtime/Engine/Classes/Engine/LevelStreaming.h
+++ b/Engine/Source/Runtime/Engine/Classes/Engine/LevelStreaming.h
@@ -73,8 +73,8 @@ class ULevelStreaming : public UObject
 {
 	GENERATED_UCLASS_BODY()
 
-	/** Name of the level package name used for loading.																		*/
-	UPROPERTY(Category=LevelStreaming, VisibleAnywhere, BlueprintReadOnly)
+	/** Path to the asset used for loading.																						*/
+	UPROPERTY(Category=LevelStreaming, VisibleAnywhere, BlueprintReadOnly, AdvancedDisplay, meta=(DisplayName = "Asset Path"))
 	FName PackageName;
 
 	/** If this isn't Name_None, then we load from this package on disk to the new package named PackageName					*/
@@ -122,7 +122,7 @@ class ULevelStreaming : public UObject
 	uint32 bIsRequestingUnloadAndRemoval:1;
 
 	/** If true, will be drawn on the 'level streaming status' map (STAT LEVELMAP console command) */
-	UPROPERTY(EditAnywhere, Category=LevelStreaming)
+	UPROPERTY(EditAnywhere, AdvancedDisplay, Category=LevelStreaming)
 	uint32 bDrawOnLevelStatusMap:1;
 
 	/** The level color used for visualization. (Show -> Advanced -> Level Coloration) */
@@ -134,7 +134,7 @@ class ULevelStreaming : public UObject
 	TArray<class ALevelStreamingVolume*> EditorStreamingVolumes;
 
 	/** Cooldown time in seconds between volume-based unload requests.  Used in preventing spurious unload requests.			*/
-	UPROPERTY(EditAnywhere, Category=LevelStreaming, meta=(ClampMin = "0", UIMin = "0", UIMax = "10"))
+	UPROPERTY(EditAnywhere, AdvancedDisplay, Category=LevelStreaming, meta=(ClampMin = "0", UIMin = "0", UIMax = "10"))
 	float MinTimeBetweenVolumeUnloadRequests;
 
 	/** Time of last volume unload request.  Used in preventing spurious unload requests.										*/
