Revised IsBuilderBrush to prevent check from triggering in GetDefaultBrush when there is no builder brush.

diff --git a/Engine/Source/Runtime/Engine/Private/ActorEditorUtils.cpp b/Engine/Source/Runtime/Engine/Private/ActorEditorUtils.cpp
--- a/Engine/Source/Runtime/Engine/Private/ActorEditorUtils.cpp
+++ b/Engine/Source/Runtime/Engine/Private/ActorEditorUtils.cpp
@@ -7,13 +7,24 @@ namespace FActorEditorUtils
 {
 	bool IsABuilderBrush( const AActor* InActor )
 	{
-#if WITH_EDITOR
+		bool bIsBuilder = false;
+#if WITH_EDITOR		
 		if ( InActor && InActor->GetWorld() && !InActor->HasAnyFlags(RF_ClassDefaultObject) )
 		{
-			return InActor->GetLevel()->GetDefaultBrush() == InActor;
+			ULevel* ActorLevel = InActor->GetLevel();
+			if ((ActorLevel != nullptr) && (ActorLevel->Actors.Num() >= 2))
+			{
+				// If the builder brush exists then it will be the 2nd actor in the actors array.
+				ABrush* BuilderBrush = Cast<ABrush>(ActorLevel->Actors[1]);
+				// If the second actor is not a brush then it certainly cannot be the builder brush.
+				if ((BuilderBrush != nullptr) && (BuilderBrush->BrushComponent != nullptr) && (BuilderBrush->Brush != nullptr))
+				{
+					bIsBuilder = (BuilderBrush == InActor);
+				}
+			}			
 		}
 #endif
-		return false;
+		return bIsBuilder;
 	}
 
 	bool IsAPreviewOrInactiveActor( const AActor* InActor )
