Hidden NavLinkProxy.SegmentLinks from edition since it's not working yet (an experimental feature) #UE4

diff --git a/Engine/Source/Runtime/Engine/Classes/AI/Navigation/NavLinkProxy.h b/Engine/Source/Runtime/Engine/Classes/AI/Navigation/NavLinkProxy.h
--- a/Engine/Source/Runtime/Engine/Classes/AI/Navigation/NavLinkProxy.h
+++ b/Engine/Source/Runtime/Engine/Classes/AI/Navigation/NavLinkProxy.h
@@ -16,8 +16,9 @@ class ENGINE_API ANavLinkProxy : public AActor, public INavLinkHostInterface, pu
 	UPROPERTY(EditAnywhere, Category=SimpleLink)
 	TArray<FNavigationLink> PointLinks;
 	
-	/** Navigation links (segment to segment) added to navigation data */
-	UPROPERTY(EditAnywhere, Category=SimpleLink)
+	/** Navigation links (segment to segment) added to navigation data
+	*	@todo hidden from use until we fix segment links. Not really working now*/
+	UPROPERTY()
 	TArray<FNavigationSegmentLink> SegmentLinks;
 
 	/** Smart link: can affect path following */
