#ue4 - Prevent boosting up slopes when lower portion of capsule impacts something while falling.

- Add HandleSlopeBoosting() to enable easier customization of this behavior.
- Check for slope boosting also after TwoWallAdjust() within PhysFalling().

diff --git a/Engine/Source/Runtime/Engine/Classes/GameFramework/CharacterMovementComponent.h b/Engine/Source/Runtime/Engine/Classes/GameFramework/CharacterMovementComponent.h
--- a/Engine/Source/Runtime/Engine/Classes/GameFramework/CharacterMovementComponent.h
+++ b/Engine/Source/Runtime/Engine/Classes/GameFramework/CharacterMovementComponent.h
@@ -1110,8 +1110,14 @@ protected:
 	virtual void TwoWallAdjust(FVector &Delta, const FHitResult& Hit, const FVector &OldHitNormal) const override;
 
 	/**
-	 * Calculate slide on a slope. Has special treatment when falling, to avoid boosting up slopes.
+	 * Calculate slide vector along a surface.
+	 * Has special treatment when falling, to avoid boosting up slopes (calling HandleSlopeBoosting in this case).
 	 * Calls AdjustUpperHemisphereImpact() for upward falling movement that hits the top of the capsule, because commonly we don't want this to behave like a smooth capsule.
+	 * @param Delta:	Attempted move.
+	 * @param Time:		Amount of move to apply (between 0 and 1).
+	 * @param Normal:	Normal opposed to movement. Not necessarily equal to Hit.Normal (but usually is).
+	 * @param Hit:		HitResult of the move that resulted in the slide.
+	 * @return			New deflected vector of movement.
 	 */
 	virtual FVector ComputeSlideVector(const FVector& Delta, const float Time, const FVector& Normal, const FHitResult& Hit) const override;
 
@@ -1121,6 +1127,17 @@ protected:
 	 */
 	virtual FVector AdjustUpperHemisphereImpact(const FVector& Delta, const FHitResult& Hit) const;
 
+	/** 
+	 * Limit the slide vector when falling if the resulting slide might boost the character faster upwards.
+	 * @param SlideResult:	Vector of movement for the slide (usually the result of ComputeSlideVector)
+	 * @param Delta:		Original attempted move
+	 * @param Time:			Amount of move to apply (between 0 and 1).
+	 * @param Normal:		Normal opposed to movement. Not necessarily equal to Hit.Normal (but usually is).
+	 * @param Hit:			HitResult of the move that resulted in the slide.
+	 * @return:				New slide result.
+	 */
+	virtual FVector HandleSlopeBoosting(const FVector& SlideResult, const FVector& Delta, const float Time, const FVector& Normal, const FHitResult& Hit) const;
+
 	/** Slows towards stop. */
 	virtual void ApplyVelocityBraking(float DeltaTime, float Friction, float BrakingDeceleration);
 
diff --git a/Engine/Source/Runtime/Engine/Private/CharacterMovementComponent.cpp b/Engine/Source/Runtime/Engine/Private/CharacterMovementComponent.cpp
--- a/Engine/Source/Runtime/Engine/Private/CharacterMovementComponent.cpp
+++ b/Engine/Source/Runtime/Engine/Private/CharacterMovementComponent.cpp
@@ -1899,6 +1899,19 @@ FVector UCharacterMovementComponent::ComputeSlideVector(const FVector& InDelta,
 
 	// prevent boosting up slopes
 	if ( bFalling && Result.Z > 0.f )
+	{
+		Result = HandleSlopeBoosting(Result, Delta, Time, Normal, Hit);
+	}
+
+	return Result;
+}
+
+
+FVector UCharacterMovementComponent::HandleSlopeBoosting(const FVector& SlideResult, const FVector& Delta, const float Time, const FVector& Normal, const FHitResult& Hit) const
+{
+	FVector Result = SlideResult;
+
+	if (Result.Z > 0.f)
 	{
 		if (Delta.Z < 0.f && (Hit.ImpactNormal.Z < MAX_STEP_SIDE_Z))
 		{
@@ -1908,13 +1921,27 @@ FVector UCharacterMovementComponent::ComputeSlideVector(const FVector& InDelta,
 		}
 		else
 		{
-			Result.Z = FMath::Min(Result.Z, Delta.Z * Time);
+			// Don't move any higher than we originally intended.
+			const float ZLimit = Delta.Z * Time;
+			if (Result.Z - ZLimit > KINDA_SMALL_NUMBER && ZLimit > KINDA_SMALL_NUMBER)
+			{
+				// Rescale the entire vector (not just the Z component) otherwise we change the direction and likely head right back into the impact.
+				const float UpPercent = ZLimit / Result.Z;
+				Result *= UpPercent;
+
+				// Make remaining portion of original result horizontal and parallel to impact normal.
+				const FVector RemainderXY = (SlideResult - Result) * FVector(1.f, 1.f, 0.f);
+				const FVector NormalXY = Normal.SafeNormal2D();
+				const FVector Adjust = Super::ComputeSlideVector(RemainderXY, 1.f, NormalXY, Hit);
+				Result += Adjust;
+			}
 		}
 	}
 
 	return Result;
 }
 
+
 FVector UCharacterMovementComponent::AdjustUpperHemisphereImpact(const FVector& Delta, const FHitResult& Hit) const
 {
 	const float ZScale = FMath::Clamp(1.f - (FMath::Abs(Hit.Normal.Z) * UpperImpactNormalScale), 0.f, 1.f);
@@ -2805,7 +2832,7 @@ void UCharacterMovementComponent::PhysFalling(float deltaTime, int32 Iterations)
 			}
 			else
 			{
-				HandleImpact(Hit, deltaTime, Adjusted);
+				HandleImpact(Hit, timeTick, Adjusted);
 				
 				if (Acceleration.SizeSquared2D() > 0.f)
 				{
@@ -2816,9 +2843,10 @@ void UCharacterMovementComponent::PhysFalling(float deltaTime, int32 Iterations)
 					}
 				}
 
+				const float FirstHitPercent = Hit.Time;
 				const FVector OldHitNormal = Hit.Normal;
 				const FVector OldHitImpactNormal = Hit.ImpactNormal;
-				FVector Delta = ComputeSlideVector(Adjusted, 1.f - Hit.Time, OldHitNormal, Hit);
+				FVector Delta = ComputeSlideVector(Adjusted, 1.f - FirstHitPercent, OldHitNormal, Hit);
 
 				if ((Delta | Adjusted) > 0.f)
 				{
@@ -2832,7 +2860,8 @@ void UCharacterMovementComponent::PhysFalling(float deltaTime, int32 Iterations)
 							return;
 						}
 
-						HandleImpact(Hit, timeTick, Delta);
+						const float SecondHitPercentOfTotal = (Hit.Time * (1.f - FirstHitPercent));
+						HandleImpact(Hit, timeTick * SecondHitPercentOfTotal, Delta);
 
 						// If we've changed physics mode, abort.
 						if (!HasValidData() || !IsFalling())
@@ -2840,7 +2869,12 @@ void UCharacterMovementComponent::PhysFalling(float deltaTime, int32 Iterations)
 							return;
 						}
 
+						const FVector PreTwoWallDelta = Delta;
 						TwoWallAdjust(Delta, Hit, OldHitNormal);
+						if (Delta.Z > 0.f)
+						{
+							HandleSlopeBoosting(Delta, PreTwoWallDelta, 1.f - Hit.Time, Hit.Normal, Hit);
+						}
 
 						// bDitch=true means that pawn is straddling two slopes, neither of which he can stand on
 						bool bDitch = ( (OldHitImpactNormal.Z > 0.f) && (Hit.ImpactNormal.Z > 0.f) && (FMath::Abs(Delta.Z) <= KINDA_SMALL_NUMBER) && ((Hit.ImpactNormal | OldHitImpactNormal) < 0.f) );
