Fixing compiler error, missed file in a recent checkin.

diff --git a/Engine/Source/Runtime/Engine/Classes/EdGraph/EdGraphSchema.h b/Engine/Source/Runtime/Engine/Classes/EdGraph/EdGraphSchema.h
--- a/Engine/Source/Runtime/Engine/Classes/EdGraph/EdGraphSchema.h
+++ b/Engine/Source/Runtime/Engine/Classes/EdGraph/EdGraphSchema.h
@@ -312,7 +312,7 @@ public:
 	ENGINE_API ActionGroup& GetAction( int32 Index );
 
 	/** Clears the action entries */
-	ENGINE_API void Empty();
+	ENGINE_API virtual void Empty();
 
 	// The temporary graph outer to store any template nodes created
 	UEdGraph* OwnerOfTemporaries;
