#ue4 - Characters that are falling and impact geometry with an unwalkable normal can check for a valid landing spot and land if possible.

- Previously the character would "float" for a bit before landing, which would lose velocity.
- Added virtual ShouldCheckForValidLandingSpot() for custom behavior. Base version returns true when hitting the edge of geometry with the lower portion of the capsule.
- Cleaned up a few other minor blocks of code in PhysFalling().

TTP #321357

diff --git a/Engine/Source/Runtime/Engine/Classes/GameFramework/CharacterMovementComponent.h b/Engine/Source/Runtime/Engine/Classes/GameFramework/CharacterMovementComponent.h
--- a/Engine/Source/Runtime/Engine/Classes/GameFramework/CharacterMovementComponent.h
+++ b/Engine/Source/Runtime/Engine/Classes/GameFramework/CharacterMovementComponent.h
@@ -1175,6 +1175,12 @@ protected:
 	/** Verify that the supplied hit result is a valid landing spot when falling. */
 	virtual bool IsValidLandingSpot(const FVector& CapsuleLocation, const FHitResult& Hit) const;
 
+	/**
+	 * Determine whether we should try to find a valid landing spot after an impact with an invalid one (based on the Hit result).
+	 * For example, landing on the lower portion of the capsule on the edge of geometry may be a walkable surface, but could have reported an unwalkable impact normal.
+	 */
+	virtual bool ShouldCheckForValidLandingSpot(const float DeltaTime, const FVector& Delta, const FHitResult& Hit) const;
+
 	/**
 	 * Check if the result of a sweep test (passed in InHit) might be a valid location to perch, in which case we should use ComputePerchResult to validate the location.
 	 * @see ComputePerchResult
diff --git a/Engine/Source/Runtime/Engine/Private/CharacterMovementComponent.cpp b/Engine/Source/Runtime/Engine/Private/CharacterMovementComponent.cpp
--- a/Engine/Source/Runtime/Engine/Private/CharacterMovementComponent.cpp
+++ b/Engine/Source/Runtime/Engine/Private/CharacterMovementComponent.cpp
@@ -2823,26 +2823,33 @@ void UCharacterMovementComponent::PhysFalling(float deltaTime, int32 Iterations)
 			if (IsValidLandingSpot(UpdatedComponent->GetComponentLocation(), Hit))
 			{
 				remainingTime += timeTick * (1.f - Hit.Time);
-				if (!bJustTeleported && (Hit.Time > 0.1f) && (Hit.Time * timeTick > 0.003f) )
-				{
-					Velocity = (CharacterOwner->GetActorLocation() - OldLocation)/(timeTick * Hit.Time);
-				}
 				ProcessLanded(Hit, remainingTime, Iterations);
 				return;
 			}
 			else
 			{
-				HandleImpact(Hit, timeTick, Adjusted);
-				
-				if (Acceleration.SizeSquared2D() > 0.f)
+				// See if we can convert a normally invalid landing spot (based on the hit result) to a usable one.
+				if (!Hit.bStartPenetrating && ShouldCheckForValidLandingSpot(timeTick, Adjusted, Hit))
 				{
-					// If we've changed physics mode, abort.
-					if (!HasValidData() || !IsFalling())
+					const FVector PawnLocation = UpdatedComponent->GetComponentLocation();
+					FFindFloorResult FloorResult;
+					FindFloor(PawnLocation, FloorResult, false);
+					if (FloorResult.IsWalkableFloor() && IsValidLandingSpot(PawnLocation, FloorResult.HitResult))
 					{
+						remainingTime += timeTick * (1.f - Hit.Time);
+						ProcessLanded(FloorResult.HitResult, remainingTime, Iterations);
 						return;
 					}
 				}
 
+				HandleImpact(Hit, timeTick, Adjusted);
+				
+				// If we've changed physics mode, abort.
+				if (!HasValidData() || !IsFalling())
+				{
+					return;
+				}
+
 				const float FirstHitPercent = Hit.Time;
 				const FVector OldHitNormal = Hit.Normal;
 				const FVector OldHitImpactNormal = Hit.ImpactNormal;
@@ -2860,8 +2867,7 @@ void UCharacterMovementComponent::PhysFalling(float deltaTime, int32 Iterations)
 							return;
 						}
 
-						const float SecondHitPercentOfTotal = (Hit.Time * (1.f - FirstHitPercent));
-						HandleImpact(Hit, timeTick * SecondHitPercentOfTotal, Delta);
+						HandleImpact(Hit, timeTick * (1.f - FirstHitPercent), Delta);
 
 						// If we've changed physics mode, abort.
 						if (!HasValidData() || !IsFalling())
@@ -4229,6 +4235,23 @@ bool UCharacterMovementComponent::IsValidLandingSpot(const FVector& CapsuleLocat
 }
 
 
+bool UCharacterMovementComponent::ShouldCheckForValidLandingSpot(const float DeltaTime, const FVector& Delta, const FHitResult& Hit) const
+{
+	// See if we hit an edge of a surface on the lower portion of the capsule.
+	// In this case the normal will not equal the impact normal, and a downward sweep may find a walkable surface on top of the edge.
+	if (Hit.Normal.Z > KINDA_SMALL_NUMBER && !Hit.Normal.Equals(Hit.ImpactNormal))
+	{
+		const FVector PawnLocation = UpdatedComponent->GetComponentLocation();
+		if (IsWithinEdgeTolerance(PawnLocation, Hit.ImpactPoint, CharacterOwner->CapsuleComponent->GetScaledCapsuleRadius()))
+		{						
+			return true;
+		}
+	}
+
+	return false;
+}
+
+
 float UCharacterMovementComponent::GetPerchRadiusThreshold() const
 {
 	// Don't allow negative values.
