#ue4 - Replaced Character::NotifyLanded() with ShouldNotifyLanded() and deprecated old function.

diff --git a/Engine/Source/Runtime/Engine/Classes/GameFramework/Character.h b/Engine/Source/Runtime/Engine/Classes/GameFramework/Character.h
--- a/Engine/Source/Runtime/Engine/Classes/GameFramework/Character.h
+++ b/Engine/Source/Runtime/Engine/Classes/GameFramework/Character.h
@@ -547,11 +547,13 @@ public:
 	UPROPERTY(BlueprintAssignable, Category="Pawn|Character")
 	FCharacterMovementUpdatedSignature OnCharacterMovementUpdated;
 
-	/**
-	 * Called when our pawn has landed from a fall.
-	 * @param	PrevMovementMode	Movement mode before the change
+	DEPRECATED(4.4, "NotifyLanded() has been renamed ShouldNotifyLanded() instead.")
+	virtual bool NotifyLanded(const struct FHitResult& Hit) { return ShouldNotifyLanded(Hit); }
+
+	/** 
+	 * Returns true if the Landed() event should be called. Used by CharacterMovement to prevent notifications while playing back network moves.
 	 */
-	virtual bool NotifyLanded(const struct FHitResult& Hit);
+	virtual bool ShouldNotifyLanded(const struct FHitResult& Hit);
 
 	/** Trigger jump if jump button has been pressed. */
 	virtual void CheckJumpInput(float DeltaTime);
diff --git a/Engine/Source/Runtime/Engine/Private/Character.cpp b/Engine/Source/Runtime/Engine/Private/Character.cpp
--- a/Engine/Source/Runtime/Engine/Private/Character.cpp
+++ b/Engine/Source/Runtime/Engine/Private/Character.cpp
@@ -740,7 +740,7 @@ void ACharacter::OnMovementModeChanged(EMovementMode PrevMovementMode, uint8 Pre
 /** Don't process landed notification if updating client position by replaying moves. 
  * Allow event to be called if Pawn was initially falling (before starting to replay moves), 
  * and this is going to cause him to land. . */
-bool ACharacter::NotifyLanded(const FHitResult& Hit)
+bool ACharacter::ShouldNotifyLanded(const FHitResult& Hit)
 {
 	if (bClientUpdating && !bClientWasFalling)
 	{
diff --git a/Engine/Source/Runtime/Engine/Private/CharacterMovementComponent.cpp b/Engine/Source/Runtime/Engine/Private/CharacterMovementComponent.cpp
--- a/Engine/Source/Runtime/Engine/Private/CharacterMovementComponent.cpp
+++ b/Engine/Source/Runtime/Engine/Private/CharacterMovementComponent.cpp
@@ -3499,7 +3499,7 @@ void UCharacterMovementComponent::StopActiveMovement()
 
 void UCharacterMovementComponent::ProcessLanded(const FHitResult& Hit, float remainingTime, int32 Iterations)
 {
-	if( CharacterOwner && CharacterOwner->NotifyLanded(Hit) )
+	if( CharacterOwner && CharacterOwner->ShouldNotifyLanded(Hit) )
 	{
 		CharacterOwner->Landed(Hit);
 	}
