#ue4 - Renamed Character::IsJumping() to IsJumpProvidingForce()

- Old function has been marked deprecated.

diff --git a/Engine/Config/BaseEngine.ini b/Engine/Config/BaseEngine.ini
--- a/Engine/Config/BaseEngine.ini
+++ b/Engine/Config/BaseEngine.ini
@@ -792,6 +792,7 @@ bUseStreamingPause=false
 +K2FieldRedirects=(OldFieldName="Pawn.OnLanded",NewFieldName="Character.OnLanded")
 +K2FieldRedirects=(OldFieldName="Pawn.LaunchPawn",NewFieldName="Character.LaunchCharacter")
 +K2FieldRedirects=(OldFieldName="Character.Launch",NewFieldName="Character.LaunchCharacter")
++K2FieldRedirects=(OldFieldName="Character.IsJumping",NewFieldName="Character.IsJumpProvidingForce")
 +K2FieldRedirects=(OldFieldName="Actor.SetActorLocation",NewFieldName="Actor.K2_SetActorLocation")
 +K2FieldRedirects=(OldFieldName="ProjectileMovementComponent.StopMovement",NewFieldName="ProjectileMovementComponent.StopSimulating")
 +K2FieldRedirects=(OldFieldName="PawnMovementComponent.StopActiveMovement",NewFieldName="NavMovementComponent.StopActiveMovement")
diff --git a/Engine/Source/Runtime/Engine/Classes/GameFramework/Character.h b/Engine/Source/Runtime/Engine/Classes/GameFramework/Character.h
--- a/Engine/Source/Runtime/Engine/Classes/GameFramework/Character.h
+++ b/Engine/Source/Runtime/Engine/Classes/GameFramework/Character.h
@@ -402,9 +402,12 @@ protected:
 
 public:
 
-	/** True if currently jumping; i.e. jump key is held and the time it has been held is less than JumpMaxHoldTime */
+	/** True if jump is actively providing a force; i.e. jump key is held and the time it has been held is less than JumpMaxHoldTime. */
 	UFUNCTION(BlueprintCallable, Category="Pawn|Character")
-	virtual bool IsJumping() const;
+	virtual bool IsJumpProvidingForce() const;
+
+	DEPRECATED(4.4, "IsJumping() has been renamed IsJumpProvidingForce().")
+	virtual bool IsJumping() const { return IsJumpProvidingForce(); }
 
 	/** Play Animation Montage on the character mesh **/
 	UFUNCTION(BlueprintCallable, Category=Animation)
diff --git a/Engine/Source/Runtime/Engine/Private/Character.cpp b/Engine/Source/Runtime/Engine/Private/Character.cpp
--- a/Engine/Source/Runtime/Engine/Private/Character.cpp
+++ b/Engine/Source/Runtime/Engine/Private/Character.cpp
@@ -199,7 +199,7 @@ bool ACharacter::CanJump() const
 
 bool ACharacter::CanJumpInternal_Implementation() const
 {
-	const bool bCanHoldToJumpHigher = (GetJumpMaxHoldTime() > 0.0f) && IsJumping();
+	const bool bCanHoldToJumpHigher = (GetJumpMaxHoldTime() > 0.0f) && IsJumpProvidingForce();
 
 	return !bIsCrouched && CharacterMovement && (CharacterMovement->IsMovingOnGround() || bCanHoldToJumpHigher) && CharacterMovement->CanEverJump() && !CharacterMovement->bWantsToCrouch;
 }
@@ -208,7 +208,7 @@ void ACharacter::OnJumped_Implementation()
 {
 }
 
-bool ACharacter::IsJumping() const
+bool ACharacter::IsJumpProvidingForce() const
 {
 	return (bPressedJump && JumpKeyHoldTime > 0.0f && JumpKeyHoldTime < GetJumpMaxHoldTime());
 }
@@ -769,7 +769,7 @@ void ACharacter::CheckJumpInput(float DeltaTime)
 	const bool bWasJumping = bPressedJump && JumpKeyHoldTime > 0.0f;
 	if (bPressedJump)
 	{
-		// Incrememnt our timer first so calls to IsJumping() will return true
+		// Increment our timer first so calls to IsJumpProvidingForce() will return true
 		JumpKeyHoldTime += DeltaTime;
 		const bool bDidJump = DoJump(bClientUpdating);
 		if(!bWasJumping && bDidJump)
