#ue4 Fix for all functions in BPs marked as Exec.  Added overrides for normal exec gating for KE / CE events, and removed the flagging

diff --git a/Engine/Source/Editor/KismetCompiler/Private/KismetCompiler.cpp b/Engine/Source/Editor/KismetCompiler/Private/KismetCompiler.cpp
--- a/Engine/Source/Editor/KismetCompiler/Private/KismetCompiler.cpp
+++ b/Engine/Source/Editor/KismetCompiler/Private/KismetCompiler.cpp
@@ -1484,8 +1484,6 @@ void FKismetCompilerContext::FinishCompilingFunction(FKismetFunctionContext& Con
 	}
 	else
 	{
-		Function->FunctionFlags |= FUNC_Exec;
-
 		// If this is the root of a blueprint-defined function or event, and if it's public, make it overrideable
 		if( !Context.IsEventGraph() && !Function->HasAnyFunctionFlags(FUNC_Private) )
 		{
diff --git a/Engine/Source/Runtime/CoreUObject/Private/UObject/ScriptCore.cpp b/Engine/Source/Runtime/CoreUObject/Private/UObject/ScriptCore.cpp
--- a/Engine/Source/Runtime/CoreUObject/Private/UObject/ScriptCore.cpp
+++ b/Engine/Source/Runtime/CoreUObject/Private/UObject/ScriptCore.cpp
@@ -611,7 +611,7 @@ void UObject::ProcessInternal( FFrame& Stack, RESULT_DECL )
 	}
 }
 
-bool UObject::CallFunctionByNameWithArguments( const TCHAR* Str, FOutputDevice& Ar, UObject* Executor )
+bool UObject::CallFunctionByNameWithArguments(const TCHAR* Str, FOutputDevice& Ar, UObject* Executor, bool bForceCallWithNonExec/*=false*/)
 {
 	// Find an exec function.
 	FString MsgStr;
@@ -632,7 +632,7 @@ bool UObject::CallFunctionByNameWithArguments( const TCHAR* Str, FOutputDevice&
 		UE_LOG(LogScriptCore, Verbose, TEXT("CallFunctionByNameWithArguments: Function not found '%s'"), Str);
 		return false;
 	}
-	if(0 == (Function->FunctionFlags & FUNC_Exec))
+	if(0 == (Function->FunctionFlags & FUNC_Exec) && !bForceCallWithNonExec)
 	{
 		UE_LOG(LogScriptCore, Verbose, TEXT("CallFunctionByNameWithArguments: Function not executable '%s'"), Str);
 		return false;
diff --git a/Engine/Source/Runtime/CoreUObject/Public/UObject/UObject.h b/Engine/Source/Runtime/CoreUObject/Public/UObject/UObject.h
--- a/Engine/Source/Runtime/CoreUObject/Public/UObject/UObject.h
+++ b/Engine/Source/Runtime/CoreUObject/Public/UObject/UObject.h
@@ -763,7 +763,7 @@ public:
 	}
 
 	/** Command line. */
-	bool CallFunctionByNameWithArguments( const TCHAR* Cmd, FOutputDevice& Ar, UObject* Executor );
+	bool CallFunctionByNameWithArguments( const TCHAR* Cmd, FOutputDevice& Ar, UObject* Executor, bool bForceCallWithNonExec = false );
 
 	// Call a function
 	void CallFunction( FFrame& Stack, RESULT_DECL, UFunction* Function );
diff --git a/Engine/Source/Runtime/Engine/Private/UnrealEngine.cpp b/Engine/Source/Runtime/Engine/Private/UnrealEngine.cpp
--- a/Engine/Source/Runtime/Engine/Private/UnrealEngine.cpp
+++ b/Engine/Source/Runtime/Engine/Private/UnrealEngine.cpp
@@ -3064,7 +3064,7 @@ bool UEngine::HandleKismetEventCommand( const TCHAR* Cmd, FOutputDevice& Ar )
 		for (TObjectIterator<UObject> It; It; ++It)
 		{
 			UObject* Obj = *It;
-			Obj->CallFunctionByNameWithArguments(Cmd, Ar, NULL);
+			Obj->CallFunctionByNameWithArguments(Cmd, Ar, NULL, true);
 		}
 	}
 	else
@@ -3077,7 +3077,7 @@ bool UEngine::HandleKismetEventCommand( const TCHAR* Cmd, FOutputDevice& Ar )
 		}
 		else
 		{
-			ObjectToMatch->CallFunctionByNameWithArguments(Cmd, Ar, NULL);
+			ObjectToMatch->CallFunctionByNameWithArguments(Cmd, Ar, NULL, true);
 		}
 	}
 
@@ -5023,7 +5023,7 @@ bool UEngine::HandleCeCommand( UWorld* InWorld, const TCHAR* Cmd, FOutputDevice&
 				ErrorMessage = 0;
 
 				// return true if at least one level handles the command
-				bResult |= CurrentLevel->GetLevelScriptActor()->CallFunctionByNameWithArguments(Cmd, Ar, NULL);
+				bResult |= CurrentLevel->GetLevelScriptActor()->CallFunctionByNameWithArguments(Cmd, Ar, NULL, true);
 			}
 		}
 	}
