#ue4 - Fix up comments referencing old AddMomentum() function.

diff --git a/Engine/Source/Runtime/Engine/Classes/GameFramework/CharacterMovementComponent.h b/Engine/Source/Runtime/Engine/Classes/GameFramework/CharacterMovementComponent.h
--- a/Engine/Source/Runtime/Engine/Classes/GameFramework/CharacterMovementComponent.h
+++ b/Engine/Source/Runtime/Engine/Classes/GameFramework/CharacterMovementComponent.h
@@ -884,7 +884,7 @@ public:
 	/** Applies repulsion force to all touched components */
 	void ApplyRepulsionForce(float DeltaTime);
 	
-	/** Applies momentum accumulated through Addmomentum() */
+	/** Applies momentum accumulated through AddImpulse()/AddForce() */
 	void ApplyAccumulatedForces(float DeltaSeconds);	
 
 	/** 
