Reduce dependencies for server builds.

diff --git a/Engine/Source/Runtime/Engine/Engine.Build.cs b/Engine/Source/Runtime/Engine/Engine.Build.cs
--- a/Engine/Source/Runtime/Engine/Engine.Build.cs
+++ b/Engine/Source/Runtime/Engine/Engine.Build.cs
@@ -91,10 +91,13 @@ public class Engine : ModuleRules
 			// at any time if needed.  The module isn't actually loaded by the engine so there is no runtime cost.
 			DynamicallyLoadedModuleNames.Add("BlankModule");
 
-			PrivateIncludePathModuleNames.Add("MeshUtilities");
-			DynamicallyLoadedModuleNames.Add("MeshUtilities");
+			if (Target.Type != TargetRules.TargetType.Server)
+			{
+				PrivateIncludePathModuleNames.Add("MeshUtilities");
+				DynamicallyLoadedModuleNames.Add("MeshUtilities");
+			}
 
-			if (Target.Configuration != UnrealTargetConfiguration.Shipping && Target.Configuration != UnrealTargetConfiguration.Test)
+			if (Target.Configuration != UnrealTargetConfiguration.Shipping && Target.Configuration != UnrealTargetConfiguration.Test && Target.Type != TargetRules.TargetType.Server)
 			{
 				PrivateDependencyModuleNames.Add("CollisionAnalyzer");
 				CircularlyReferencedDependentModules.Add("CollisionAnalyzer");
