Fix crash in PhAT deleting root body
#TTP-341380

diff --git a/Engine/Source/Runtime/Engine/Private/PhysicsEngine/PhysicsAsset.cpp b/Engine/Source/Runtime/Engine/Private/PhysicsEngine/PhysicsAsset.cpp
--- a/Engine/Source/Runtime/Engine/Private/PhysicsEngine/PhysicsAsset.cpp
+++ b/Engine/Source/Runtime/Engine/Private/PhysicsEngine/PhysicsAsset.cpp
@@ -211,11 +211,14 @@ int32	UPhysicsAsset::FindControllingBodyIndex(class USkeletalMesh* skelMesh, int
 int32 UPhysicsAsset::FindParentBodyIndex(class USkeletalMesh * skelMesh, int32 StartBoneIndex) const
 {
 	int32 BoneIndex = StartBoneIndex;
-	while ((BoneIndex = skelMesh->RefSkeleton.GetParentIndex(BoneIndex)) != 0)
+	while ((BoneIndex = skelMesh->RefSkeleton.GetParentIndex(BoneIndex)) != INDEX_NONE)
 	{
 		FName BoneName = skelMesh->RefSkeleton.GetBoneName(BoneIndex);
 		int32 BodyIndex = FindBodyIndex(BoneName);
 
+		if (StartBoneIndex == BoneIndex)
+			return INDEX_NONE;
+
 		if (BodyIndex != INDEX_NONE)
 			return BodyIndex;
 	}
