[GitHub] - Fixed sign errors in SuggestProjectileVelocity().

Pull request: https://github.com/EpicGames/UnrealEngine/pull/274

diff --git a/Engine/Source/Runtime/Engine/Private/GameplayStatics.cpp b/Engine/Source/Runtime/Engine/Private/GameplayStatics.cpp
--- a/Engine/Source/Runtime/Engine/Private/GameplayStatics.cpp
+++ b/Engine/Source/Runtime/Engine/Private/GameplayStatics.cpp
@@ -1121,7 +1121,7 @@ bool UGameplayStatics::SuggestProjectileVelocity(UObject* WorldContextObject, FV
 
 	UWorld* const World = GEngine->GetWorldFromContextObject(WorldContextObject);
 
-	const float GravityZ = (OverrideGravityZ != 0.f) ? OverrideGravityZ : World->GetGravityZ();
+	const float GravityZ = (OverrideGravityZ != 0.f) ? -OverrideGravityZ : -World->GetGravityZ();
 
 
 	// v^4 - g*(g*x^2 + 2*y*v^2)
@@ -1152,13 +1152,16 @@ bool UGameplayStatics::SuggestProjectileVelocity(UObject* WorldContextObject, FV
 	{
 		// choose which arc
 		const float FavoredMagXYSq = bFavorHighArc ? FMath::Min(MagXYSq_A, MagXYSq_B) : FMath::Max(MagXYSq_A, MagXYSq_B);
+		const float ZSign = bFavorHighArc ?
+							(MagXYSq_A < MagXYSq_B) ? FMath::Sign(TanSolutionAngleA) : FMath::Sign(TanSolutionAngleB) :
+							(MagXYSq_A > MagXYSq_B) ? FMath::Sign(TanSolutionAngleA) : FMath::Sign(TanSolutionAngleB);
 
 		// finish calculations
 		const float MagXY = FMath::Sqrt(FavoredMagXYSq);
 		const float MagZ = FMath::Sqrt(TossSpeedSq - FavoredMagXYSq);		// pythagorean
 
 		// final answer!
-		OutTossVelocity = (DirXY * MagXY) + (FVector::UpVector * MagZ);
+		OutTossVelocity = (DirXY * MagXY) + (FVector::UpVector * MagZ * ZSign);
 		bFoundAValidSolution = true;
 
 	 	if (bDrawDebug)
@@ -1170,7 +1173,7 @@ bool UGameplayStatics::SuggestProjectileVelocity(UObject* WorldContextObject, FV
 	 			const float TimeInFlight = (Step+StepSize) * DeltaXY/MagXY;
 	 
 	 			// d = vt + .5 a t^2
-	 			const FVector TraceEnd = Start + OutTossVelocity*TimeInFlight + FVector(0.f, 0.f, 0.5f * GravityZ * FMath::Square(TimeInFlight) - CollisionRadius);
+				const FVector TraceEnd = Start + OutTossVelocity*TimeInFlight + FVector(0.f, 0.f, 0.5f * -GravityZ * FMath::Square(TimeInFlight) - CollisionRadius);
 	 
 	 			DrawDebugLine( World, TraceStart, TraceEnd, (bFoundAValidSolution ? FColor::Yellow : FColor::Red), true );
 	 			TraceStart = TraceEnd;
@@ -1186,6 +1189,14 @@ bool UGameplayStatics::SuggestProjectileVelocity(UObject* WorldContextObject, FV
 		PrioritizedSolutionsMagXYSq[0] = bFavorHighArc ? FMath::Min(MagXYSq_A, MagXYSq_B) : FMath::Max(MagXYSq_A, MagXYSq_B);
 		PrioritizedSolutionsMagXYSq[1] = bFavorHighArc ? FMath::Max(MagXYSq_A, MagXYSq_B) : FMath::Min(MagXYSq_A, MagXYSq_B);
 
+		float PrioritizedSolutionZSign[2];
+		PrioritizedSolutionZSign[0] = bFavorHighArc ?
+										(MagXYSq_A < MagXYSq_B) ? FMath::Sign(TanSolutionAngleA) : FMath::Sign(TanSolutionAngleB) :
+										(MagXYSq_A > MagXYSq_B) ? FMath::Sign(TanSolutionAngleA) : FMath::Sign(TanSolutionAngleB);
+		PrioritizedSolutionZSign[1] = bFavorHighArc ?
+										(MagXYSq_A > MagXYSq_B) ? FMath::Sign(TanSolutionAngleA) : FMath::Sign(TanSolutionAngleB) :
+										(MagXYSq_A < MagXYSq_B) ? FMath::Sign(TanSolutionAngleA) : FMath::Sign(TanSolutionAngleB);
+
 		FVector PrioritizedProjVelocities[2];
 
 		// try solutions in priority order
@@ -1194,8 +1205,9 @@ bool UGameplayStatics::SuggestProjectileVelocity(UObject* WorldContextObject, FV
 		{
 			const float MagXY = FMath::Sqrt( PrioritizedSolutionsMagXYSq[CurrentSolutionIdx] );
 			const float MagZ = FMath::Sqrt( TossSpeedSq - PrioritizedSolutionsMagXYSq[CurrentSolutionIdx] );		// pythagorean
+			const float ZSign = PrioritizedSolutionZSign[CurrentSolutionIdx];
 
-			PrioritizedProjVelocities[CurrentSolutionIdx] = (DirXY * MagXY) + (FVector::UpVector * MagZ);
+			PrioritizedProjVelocities[CurrentSolutionIdx] = (DirXY * MagXY) + (FVector::UpVector * MagZ * ZSign);
 
 			// iterate along the arc, doing stepwise traces
 			bool bFailedTrace = false;
@@ -1206,7 +1218,7 @@ bool UGameplayStatics::SuggestProjectileVelocity(UObject* WorldContextObject, FV
 				const float TimeInFlight = (Step+StepSize) * DeltaXY/MagXY;
 
 				// d = vt + .5 a t^2
-				const FVector TraceEnd = Start + PrioritizedProjVelocities[CurrentSolutionIdx]*TimeInFlight + FVector(0.f, 0.f, 0.5f * GravityZ * FMath::Square(TimeInFlight) - CollisionRadius);
+				const FVector TraceEnd = Start + PrioritizedProjVelocities[CurrentSolutionIdx]*TimeInFlight + FVector(0.f, 0.f, 0.5f * -GravityZ * FMath::Square(TimeInFlight) - CollisionRadius);
 
 				if ( (TraceOption == ESuggestProjVelocityTraceOption::OnlyTraceWhileAsceding) && (TraceEnd.Z < TraceStart.Z) )
 				{
