Helpful comments while debugging connection flow

diff --git a/Engine/Source/Runtime/Engine/Private/LocalPlayer.cpp b/Engine/Source/Runtime/Engine/Private/LocalPlayer.cpp
--- a/Engine/Source/Runtime/Engine/Private/LocalPlayer.cpp
+++ b/Engine/Source/Runtime/Engine/Private/LocalPlayer.cpp
@@ -211,6 +211,12 @@ bool ULocalPlayer::SpawnPlayActor(const FString& URL,FString& OutError, UWorld*
 		UClass* PCClass = PendingLevelPlayerControllerClass;
 		// The PlayerController gets replicated from the client though the engine assumes that every Player always has
 		// a valid PlayerController so we spawn a dummy one that is going to be replaced later.
+
+		// 
+		// Look at APlayerController::OnActorChannelOpen + UNetConnection::HandleClientPlayer for the code the
+		// replaces this fake player controller with the real replicated one from the server
+		//
+
 		PlayerController = CastChecked<APlayerController>(InWorld->SpawnActor(PCClass));
 		const int32 PlayerIndex = GEngine->GetGamePlayers(InWorld).Find(this);
 		PlayerController->NetPlayerIndex = PlayerIndex;
