Removed streaming levels sorting on a map load, not really required

diff --git a/Engine/Source/Runtime/Engine/Private/World.cpp b/Engine/Source/Runtime/Engine/Private/World.cpp
--- a/Engine/Source/Runtime/Engine/Private/World.cpp
+++ b/Engine/Source/Runtime/Engine/Private/World.cpp
@@ -507,36 +507,10 @@ void UWorld::PostLoad()
 		StreamingLevels = WorldSettings->StreamingLevels_DEPRECATED;
 		WorldSettings->StreamingLevels_DEPRECATED.Empty();
 	}
-
-	// this is from AWorldSettings::PostLoad, where it rearrange order
-	// Make sure that 'always loaded' maps are first in the array
-	TArray<ULevelStreaming*> AlwaysLoadedLevels;
-	// Iterate over each LevelStreaming object
-	for( int32 LevelIndex=StreamingLevels.Num()-1; LevelIndex>=0; LevelIndex-- )
-	{
-		//Remove null entries
-		if (StreamingLevels[LevelIndex] == nullptr)
-		{
-			StreamingLevels.RemoveAt(LevelIndex);
-			continue;
-		}
-		
-		// See if its an 'always loaded' one
-		ULevelStreamingAlwaysLoaded* AlwaysLoadedLevel = Cast<ULevelStreamingAlwaysLoaded>( StreamingLevels[LevelIndex] );
-		if(AlwaysLoadedLevel)
-		{
-			// If it is, add to out list (preserving order), and remove from main list
-			AlwaysLoadedLevels.Insert(AlwaysLoadedLevel, 0);
-			StreamingLevels.RemoveAt(LevelIndex);
-		}
-	}
-
-	// Now make new array that starts with 'always loaded' levels, followed by the rest
-	TArray<ULevelStreaming*> NewStreamingLevels = AlwaysLoadedLevels;
-	NewStreamingLevels.Append( StreamingLevels );
-	// And use that for the new StreamingLevels array
-	StreamingLevels = NewStreamingLevels;
-
+	
+	// Remove null streaming level entries (could be if level was saved with transient level streaming objects)
+	StreamingLevels.Remove(nullptr);
+	
 	// Make sure that the persistent level isn't in this world's list of streaming levels.  This should
 	// never really happen, but was needed in at least one observed case of corrupt map data.
 	if( PersistentLevel != NULL )
