Move FAnimUpdateRateParameters into EngineTypes.h to keep Actor.h smaller
Add EngineDefines.h to Actor.h so that ENABLE_VISUAL_LOG is correctly defined (changes class layout)
Forward declare UGameInstance in GameViewportClient.h

diff --git a/Engine/Source/Runtime/Engine/Classes/Engine/EngineTypes.h b/Engine/Source/Runtime/Engine/Classes/Engine/EngineTypes.h
--- a/Engine/Source/Runtime/Engine/Classes/Engine/EngineTypes.h
+++ b/Engine/Source/Runtime/Engine/Classes/Engine/EngineTypes.h
@@ -1525,6 +1525,85 @@ struct FAnimSlotDesc
 
 };
 
+/** Container for Animation Update Rate parameters.
+ * They are shared for all components of an Actor, so they can be updated in sync. */
+USTRUCT()
+struct FAnimUpdateRateParameters
+{
+	GENERATED_USTRUCT_BODY()
+
+private:
+	/** How often animation will be updated/ticked. 1 = every frame, 2 = every 2 frames, etc. */
+	UPROPERTY()
+	int32 UpdateRate;
+
+	/** How often animation will be evaluated. 1 = every frame, 2 = every 2 frames, etc.
+	 *  has to be a multiple of UpdateRate. */
+	UPROPERTY()
+	int32 EvaluationRate;
+
+	/** When skipping a frame, should it be interpolated or frozen? */
+	UPROPERTY()
+	bool bInterpolateSkippedFrames;
+
+	/** (This frame) animation update should be skipped. */
+	UPROPERTY()
+	bool bSkipUpdate;
+
+	/** (This frame) animation evaluation should be skipped. */
+	UPROPERTY()
+	bool bSkipEvaluation;
+
+public:
+	/** Default constructor */
+	FAnimUpdateRateParameters()
+		: UpdateRate(1)
+		, EvaluationRate(1)
+		, bInterpolateSkippedFrames(false)
+		, bSkipUpdate(false)
+		, bSkipEvaluation(false)
+	{
+	}
+
+	/** Set parameters and verify inputs.
+	 * @param : Owner Actor owner calling this.
+	 * @param : NewUpdateRate. How often animation will be updated/ticked. 1 = every frame, 2 = every 2 frames, etc.
+	 * @param : NewEvaluationRate. How often animation will be evaluated. 1 = every frame, 2 = every 2 frames, etc.
+	 * @param : bNewInterpSkippedFrames. When skipping a frame, should it be interpolated or frozen?
+	 */
+	void Set(const class AActor & Owner, const int32 & NewUpdateRate, const int32 & NewEvaluationRate, const bool & bNewInterpSkippedFrames);
+
+	/* Getter for UpdateRate */
+	int32 GetUpdateRate() const
+	{
+		return UpdateRate;
+	}
+
+	/* Getter for EvaluationRate */
+	int32 GetEvaluationRate() const
+	{
+		return EvaluationRate;
+	}
+
+	/* Getter for bSkipUpdate */
+	bool ShouldSkipUpdate() const
+	{
+		return bSkipUpdate;
+	}
+
+	/* Getter for bSkipEvaluation */
+	bool ShouldSkipEvaluation() const
+	{
+		return bSkipEvaluation;
+	}
+
+	/* Getter for bInterpolateSkippedFrames */
+	bool ShouldInterpolateSkippedFrames() const
+	{
+		return bInterpolateSkippedFrames;
+	}
+};
+
 /**
  * Point Of View type.
  */
diff --git a/Engine/Source/Runtime/Engine/Classes/Engine/GameViewportClient.h b/Engine/Source/Runtime/Engine/Classes/Engine/GameViewportClient.h
--- a/Engine/Source/Runtime/Engine/Classes/Engine/GameViewportClient.h
+++ b/Engine/Source/Runtime/Engine/Classes/Engine/GameViewportClient.h
@@ -19,6 +19,8 @@
 #include "Engine/ScriptViewportClient.h"
 #include "GameViewportClient.generated.h"
 
+class UGameInstance;
+
 /**
  * Enum of the different splitscreen types
  */
@@ -254,7 +256,7 @@ public:
 
 	class FSceneViewport* GetGameViewport();
 
-	class UGameInstance* GetGameInstance() const;
+	UGameInstance* GetGameInstance() const;
 
 	virtual void Init(struct FWorldContext& WorldContext, UGameInstance* OwningGameInstance);
 
diff --git a/Engine/Source/Runtime/Engine/Classes/GameFramework/Actor.h b/Engine/Source/Runtime/Engine/Classes/GameFramework/Actor.h
--- a/Engine/Source/Runtime/Engine/Classes/GameFramework/Actor.h
+++ b/Engine/Source/Runtime/Engine/Classes/GameFramework/Actor.h
@@ -2,6 +2,8 @@
 
 #pragma once
 
+#include "EngineDefines.h"
+
 #include "Engine/EngineBaseTypes.h"
 #include "Components/ActorComponent.h"
 #include "Engine/EngineTypes.h"
@@ -56,86 +58,6 @@ DECLARE_DELEGATE_FourParams(FMakeNoiseDelegate, class AActor*, float, class APaw
 
 DECLARE_CYCLE_STAT_EXTERN(TEXT("GetComponentsTime"),STAT_GetComponentsTime,STATGROUP_Engine,ENGINE_API);
 
-/** Container for Animation Update Rate parameters.
- * They are shared for all components of an Actor, so they can be updated in sync. */
-USTRUCT()
-struct FAnimUpdateRateParameters
-{
-	GENERATED_USTRUCT_BODY()
-
-private:
-	/** How often animation will be updated/ticked. 1 = every frame, 2 = every 2 frames, etc. */
-	UPROPERTY()
-	int32 UpdateRate;
-
-	/** How often animation will be evaluated. 1 = every frame, 2 = every 2 frames, etc.
-	 *  has to be a multiple of UpdateRate. */
-	UPROPERTY()
-	int32 EvaluationRate;
-
-	/** When skipping a frame, should it be interpolated or frozen? */
-	UPROPERTY()
-	bool bInterpolateSkippedFrames;
-
-	/** (This frame) animation update should be skipped. */
-	UPROPERTY()
-	bool bSkipUpdate;
-
-	/** (This frame) animation evaluation should be skipped. */
-	UPROPERTY()
-	bool bSkipEvaluation;
-
-public:
-	/** Default constructor */
-	FAnimUpdateRateParameters()
-		: UpdateRate(1)
-		, EvaluationRate(1)
-		, bInterpolateSkippedFrames(false)
-		, bSkipUpdate(false)
-		, bSkipEvaluation(false)
-	{
-	}
-
-	/** Set parameters and verify inputs.
-	 * @param : Owner Actor owner calling this.
-	 * @param : NewUpdateRate. How often animation will be updated/ticked. 1 = every frame, 2 = every 2 frames, etc.
-	 * @param : NewEvaluationRate. How often animation will be evaluated. 1 = every frame, 2 = every 2 frames, etc.
-	 * @param : bNewInterpSkippedFrames. When skipping a frame, should it be interpolated or frozen?
-	 */
-	void Set(const class AActor & Owner, const int32 & NewUpdateRate, const int32 & NewEvaluationRate, const bool & bNewInterpSkippedFrames);
-
-	/* Getter for UpdateRate */
-	int32 GetUpdateRate() const
-	{
-		return UpdateRate;
-	}
-
-	/* Getter for EvaluationRate */
-	int32 GetEvaluationRate() const
-	{
-		return EvaluationRate;
-	}
-
-	/* Getter for bSkipUpdate */
-	bool ShouldSkipUpdate() const
-	{
-		return bSkipUpdate;
-	}
-
-	/* Getter for bSkipEvaluation */
-	bool ShouldSkipEvaluation() const
-	{
-		return bSkipEvaluation;
-	}
-
-	/* Getter for bInterpolateSkippedFrames */
-	bool ShouldInterpolateSkippedFrames() const
-	{
-		return bInterpolateSkippedFrames;
-	}
-};
-
-
 /**
  * Base class for an object that can be placed or spawned in a level. Actors may contain a collection of Components, and support network replication.
 *
