fixed updating navmesh version on rebuilld, restored navmesh error messages on hud

diff --git a/Engine/Source/Runtime/Engine/Classes/AI/Navigation/NavigationData.h b/Engine/Source/Runtime/Engine/Classes/AI/Navigation/NavigationData.h
--- a/Engine/Source/Runtime/Engine/Classes/AI/Navigation/NavigationData.h
+++ b/Engine/Source/Runtime/Engine/Classes/AI/Navigation/NavigationData.h
@@ -97,7 +97,8 @@ class ENGINE_API ANavigationData : public AActor
 
 	virtual void RerunConstructionScripts() override;
 
-	virtual bool NeedsRebuild() { return false; }
+	virtual bool NeedsRebuild() const { return false; }
+	virtual bool CanRebuild() const;
 
 	//----------------------------------------------------------------------//
 	// Generation & data access                                                      
diff --git a/Engine/Source/Runtime/Engine/Classes/AI/Navigation/NavigationSystem.h b/Engine/Source/Runtime/Engine/Classes/AI/Navigation/NavigationSystem.h
--- a/Engine/Source/Runtime/Engine/Classes/AI/Navigation/NavigationSystem.h
+++ b/Engine/Source/Runtime/Engine/Classes/AI/Navigation/NavigationSystem.h
@@ -363,7 +363,10 @@ public:
 	void ApplyWorldOffset(const FVector& InOffset, bool bWorldShift);
 
 	/** checks if navigation/navmesh is dirty and needs to be rebuilt */
-	bool IsNavigationDirty();
+	bool IsNavigationDirty() const;
+
+	/** checks if dirty navigation data can rebuild itself */
+	bool CanRebuildDirtyNavigation() const;
 
 	static bool DoesPathIntersectBox(const FNavigationPath* Path, const FBox& Box);
 
diff --git a/Engine/Source/Runtime/Engine/Classes/AI/Navigation/RecastNavMesh.h b/Engine/Source/Runtime/Engine/Classes/AI/Navigation/RecastNavMesh.h
--- a/Engine/Source/Runtime/Engine/Classes/AI/Navigation/RecastNavMesh.h
+++ b/Engine/Source/Runtime/Engine/Classes/AI/Navigation/RecastNavMesh.h
@@ -862,7 +862,7 @@ public:
 	const FTileSetItem* GetTileSet() const { return TileSet.GetTypedData(); }
 	FTileSetItem* GetTileSet() { return TileSet.GetTypedData(); }
 	
-	virtual bool NeedsRebuild() override;
+	virtual bool NeedsRebuild() const override;
 
 	/** update offset for navmesh in recast library, to current one */
 	/** @param Offset - current offset */
diff --git a/Engine/Source/Runtime/Engine/Private/AI/Navigation/NavigationData.cpp b/Engine/Source/Runtime/Engine/Private/AI/Navigation/NavigationData.cpp
--- a/Engine/Source/Runtime/Engine/Private/AI/Navigation/NavigationData.cpp
+++ b/Engine/Source/Runtime/Engine/Private/AI/Navigation/NavigationData.cpp
@@ -364,6 +364,15 @@ void ANavigationData::CleanUpAndMarkPendingKill()
 	MarkComponentsAsPendingKill();
 }
 
+bool ANavigationData::CanRebuild() const
+{
+#if WITH_NAVIGATION_GENERATOR
+	return NavDataGenerator.IsValid();
+#endif
+
+	return false;
+}
+
 #if WITH_NAVIGATION_GENERATOR
 
 FNavDataGenerator* ANavigationData::GetGenerator(FNavigationSystem::ECreateIfEmpty CreateIfNone)
diff --git a/Engine/Source/Runtime/Engine/Private/AI/Navigation/NavigationSystem.cpp b/Engine/Source/Runtime/Engine/Private/AI/Navigation/NavigationSystem.cpp
--- a/Engine/Source/Runtime/Engine/Private/AI/Navigation/NavigationSystem.cpp
+++ b/Engine/Source/Runtime/Engine/Private/AI/Navigation/NavigationSystem.cpp
@@ -2956,7 +2956,7 @@ void UNavigationSystem::CycleNavigationDataDrawn()
 	}
 }
 
-bool UNavigationSystem::IsNavigationDirty()
+bool UNavigationSystem::IsNavigationDirty() const
 {
 	for (int32 NavDataIndex=0; NavDataIndex < NavDataSet.Num(); ++NavDataIndex)
 	{
@@ -2969,6 +2969,22 @@ bool UNavigationSystem::IsNavigationDirty()
 	return false;
 }
 
+bool UNavigationSystem::CanRebuildDirtyNavigation() const
+{
+	for (int32 NavDataIndex = 0; NavDataIndex < NavDataSet.Num(); ++NavDataIndex)
+	{
+		const bool bIsDirty = NavDataSet[NavDataIndex]->NeedsRebuild();
+		const bool bCanRebuild = NavDataSet[NavDataIndex]->CanRebuild();
+
+		if (bIsDirty && !bCanRebuild)
+		{
+			return false;
+		}
+	}
+
+	return true;
+}
+
 bool UNavigationSystem::DoesPathIntersectBox(const FNavigationPath* Path, const FBox& Box)
 {
 	return Path != NULL && Path->DoesIntersectBox(Box);
diff --git a/Engine/Source/Runtime/Engine/Private/AI/Navigation/RecastNavMesh.cpp b/Engine/Source/Runtime/Engine/Private/AI/Navigation/RecastNavMesh.cpp
--- a/Engine/Source/Runtime/Engine/Private/AI/Navigation/RecastNavMesh.cpp
+++ b/Engine/Source/Runtime/Engine/Private/AI/Navigation/RecastNavMesh.cpp
@@ -1803,7 +1803,7 @@ void ARecastNavMesh::PostEditChangeProperty( struct FPropertyChangedEvent& Prope
 
 #endif // WITH_EDITOR
 
-bool ARecastNavMesh::NeedsRebuild()
+bool ARecastNavMesh::NeedsRebuild() const
 {
 	bool bLooksLikeNeeded = !RecastNavMeshImpl || RecastNavMeshImpl->GetRecastMesh() == 0;
 #if WITH_NAVIGATION_GENERATOR
diff --git a/Engine/Source/Runtime/Engine/Private/AI/Navigation/RecastNavMeshGenerator.cpp b/Engine/Source/Runtime/Engine/Private/AI/Navigation/RecastNavMeshGenerator.cpp
--- a/Engine/Source/Runtime/Engine/Private/AI/Navigation/RecastNavMeshGenerator.cpp
+++ b/Engine/Source/Runtime/Engine/Private/AI/Navigation/RecastNavMeshGenerator.cpp
@@ -3529,10 +3529,6 @@ void FRecastNavMeshGenerator::RequestDirtyTilesRebuild()
 bool FRecastNavMeshGenerator::Generate()
 {
 	bool const bGenerated = GenerateTiledNavMesh();
-
-	// we rebuilt, update to latest
-	DestNavMesh->UpdateNavVersion();
-
 	return bGenerated;
 }
 
@@ -4019,6 +4015,11 @@ void FRecastNavMeshGenerator::RebuildAll()
 		DirtyAreas.Add(BigArea);
 		RequestDirtyTilesRebuild();
 	}
+
+	if (DestNavMesh.IsValid())
+	{
+		DestNavMesh->UpdateNavVersion();
+	}
 }
 
 void FRecastNavMeshGenerator::MarkAbortedGenerator(const FRecastTileGenerator* TileGenerator)
diff --git a/Engine/Source/Runtime/Engine/Private/UnrealEngine.cpp b/Engine/Source/Runtime/Engine/Private/UnrealEngine.cpp
--- a/Engine/Source/Runtime/Engine/Private/UnrealEngine.cpp
+++ b/Engine/Source/Runtime/Engine/Private/UnrealEngine.cpp
@@ -6395,8 +6395,9 @@ void DrawStatsHUD( UWorld* World, FViewport* Viewport, FCanvas* Canvas, UCanvas*
 #else
 			const bool bIsNavigationAutoUpdateEnabled = true;
 #endif
-			if (World->GetNavigationSystem() != NULL && World->GetNavigationSystem()->IsNavigationDirty() && 
-				(!World->GetNavigationSystem()->bBuildNavigationAtRuntime || !bIsNavigationAutoUpdateEnabled))
+			UNavigationSystem* NavSys = UNavigationSystem::GetCurrent(World);
+			if (NavSys && NavSys->IsNavigationDirty() &&
+				(!bIsNavigationAutoUpdateEnabled || !NavSys->bBuildNavigationAtRuntime || !NavSys->CanRebuildDirtyNavigation()))
 			{
 				SmallTextItem.SetColor( FLinearColor::White );
 				SmallTextItem.Text =  LOCTEXT("NAVMESHERROR", "NAVMESH NEEDS TO BE REBUILT");				
