UE4 - Instanced static meshes for ES2

diff --git a/Engine/Shaders/LocalVertexFactory.usf b/Engine/Shaders/LocalVertexFactory.usf
--- a/Engine/Shaders/LocalVertexFactory.usf
+++ b/Engine/Shaders/LocalVertexFactory.usf
@@ -168,8 +168,8 @@ half3x3 CalcTangentToWorldNoScale(in half3x3 TangentToLocal)
 	half3x3 LocalToWorld = GetLocalToWorld3x3();
 	half3 InvScale = Primitive.InvNonUniformScale.xyz;
 	LocalToWorld[0] *= InvScale.x;
-    LocalToWorld[1] *= InvScale.y;
-    LocalToWorld[2] *= InvScale.z;
+	LocalToWorld[1] *= InvScale.y;
+	LocalToWorld[2] *= InvScale.z;
 	return mul(TangentToLocal, LocalToWorld); 
 }
 
@@ -247,7 +247,7 @@ FMaterialVertexParameters GetMaterialVertexParameters( FVertexFactoryInput Input
 		float3 BaseYVec = normalize( cross(SplineDir, BaseXVec) );
 	
 		// Offset from the spline, using the frenet frame
- 		float2 SliceOffset = lerp(SplineStartOffset, SplineEndOffset, HermiteAlpha);
+		float2 SliceOffset = lerp(SplineStartOffset, SplineEndOffset, HermiteAlpha);
 		SplinePos += SliceOffset.x * BaseXVec;
 		SplinePos += SliceOffset.y * BaseYVec;
 	
@@ -447,7 +447,11 @@ FVertexFactoryInterpolantsVSToPS VertexFactoryGetInterpolantsVSToPS(FVertexFacto
 #if USE_INSTANCING
 	//float2 InstanceLightMapBias = float2(Input.InstanceInverseTransform1.w, Input.InstanceInverseTransform2.w);
 	//Interpolants.LightMapCoordinate.xy = Input.LightMapCoordinate * LightMapCoordinateScaleBias.xy + InstanceLightMapBias;
+	#if ES2_PROFILE
+	SetLightMapCoordinate(Interpolants, (float4)0, (float2)0);
+	#else
 	Interpolants.LightMapCoordinate = 0;
+	#endif
 #else
 	
 	float4 LightMapCoordinate = 0;
diff --git a/Engine/Shaders/LocalVertexFactoryCommon.usf b/Engine/Shaders/LocalVertexFactoryCommon.usf
--- a/Engine/Shaders/LocalVertexFactoryCommon.usf
+++ b/Engine/Shaders/LocalVertexFactoryCommon.usf
@@ -140,6 +140,11 @@ struct FVertexFactoryInterpolantsVSToPS
 	// ShadowmapUV, tangent zw
 	float4 PackedInterpolator5 : TEXCOORD5;
 #endif
+
+#if USE_INSTANCING
+	// x = per-instance random, y = per-instance fade out amount
+	float2  PerInstanceParams : COLOR1;
+#endif
 };
 
 float2 GetUV(FVertexFactoryInterpolantsVSToPS Interpolants, int UVIndex)
diff --git a/Engine/Source/Runtime/Engine/Private/InstancedStaticMesh.cpp b/Engine/Source/Runtime/Engine/Private/InstancedStaticMesh.cpp
--- a/Engine/Source/Runtime/Engine/Private/InstancedStaticMesh.cpp
+++ b/Engine/Source/Runtime/Engine/Private/InstancedStaticMesh.cpp
@@ -354,7 +354,11 @@ public:
 	 */
 	static void ModifyCompilationEnvironment(EShaderPlatform Platform, const FMaterial* Material, FShaderCompilerEnvironment& OutEnvironment)
 	{
-		OutEnvironment.SetDefine(TEXT("USE_INSTANCING"),TEXT("1"));
+		const bool bInstanced = RHISupportsInstancing(Platform);
+		if (bInstanced)
+		{
+			OutEnvironment.SetDefine(TEXT("USE_INSTANCING"),TEXT("1"));
+		}
 	}
 
 	/**
@@ -389,8 +393,7 @@ private:
 bool FInstancedStaticMeshVertexFactory::ShouldCache(EShaderPlatform Platform, const class FMaterial* Material, const class FShaderType* ShaderType)
 {
 	return (Material->IsUsedWithInstancedStaticMeshes() || Material->IsSpecialEngineMaterial())
-			&& FLocalVertexFactory::ShouldCache(Platform, Material, ShaderType) 
-			&& IsFeatureLevelSupported(Platform, ERHIFeatureLevel::SM3);
+			&& FLocalVertexFactory::ShouldCache(Platform, Material, ShaderType);
 }
 
 
@@ -412,6 +415,8 @@ void FInstancedStaticMeshVertexFactory::Copy(const FInstancedStaticMeshVertexFac
 
 void FInstancedStaticMeshVertexFactory::InitRHI()
 {
+	const bool bInstanced = RHISupportsInstancing(GRHIShaderPlatform);
+
 	// If the vertex buffer containing position is not the same vertex buffer containing the rest of the data,
 	// then initialize PositionStream and PositionDeclaration.
 	if(Data.PositionComponent.VertexBuffer != Data.TangentBasisComponents[0].VertexBuffer)
@@ -419,10 +424,13 @@ void FInstancedStaticMeshVertexFactory::InitRHI()
 		FVertexDeclarationElementList PositionOnlyStreamElements;
 		PositionOnlyStreamElements.Add(AccessPositionStreamComponent(Data.PositionComponent,0));
 
-		// toss in the instanced location stream
-		PositionOnlyStreamElements.Add(AccessPositionStreamComponent(Data.InstancedTransformComponent[0],9));
-		PositionOnlyStreamElements.Add(AccessPositionStreamComponent(Data.InstancedTransformComponent[1],10));
-		PositionOnlyStreamElements.Add(AccessPositionStreamComponent(Data.InstancedTransformComponent[2],11));
+		if (bInstanced)
+		{
+			// toss in the instanced location stream
+			PositionOnlyStreamElements.Add(AccessPositionStreamComponent(Data.InstancedTransformComponent[0],9));
+			PositionOnlyStreamElements.Add(AccessPositionStreamComponent(Data.InstancedTransformComponent[1],10));
+			PositionOnlyStreamElements.Add(AccessPositionStreamComponent(Data.InstancedTransformComponent[2],11));
+		}
 		InitPositionDeclaration(PositionOnlyStreamElements);
 	}
 
@@ -484,13 +492,16 @@ void FInstancedStaticMeshVertexFactory::InitRHI()
 	}
 
 	// toss in the instanced location stream
-	Elements.Add(AccessStreamComponent(Data.InstancedShadowMapBiasComponent,8));
-	Elements.Add(AccessStreamComponent(Data.InstancedTransformComponent[0],9));
-	Elements.Add(AccessStreamComponent(Data.InstancedTransformComponent[1],10));
-	Elements.Add(AccessStreamComponent(Data.InstancedTransformComponent[2],11));
- 	Elements.Add(AccessStreamComponent(Data.InstancedInverseTransformComponent[0],12));
- 	Elements.Add(AccessStreamComponent(Data.InstancedInverseTransformComponent[1],13));
- 	Elements.Add(AccessStreamComponent(Data.InstancedInverseTransformComponent[2],14));
+	if (bInstanced)
+	{
+		Elements.Add(AccessStreamComponent(Data.InstancedShadowMapBiasComponent,8));
+		Elements.Add(AccessStreamComponent(Data.InstancedTransformComponent[0],9));
+		Elements.Add(AccessStreamComponent(Data.InstancedTransformComponent[1],10));
+		Elements.Add(AccessStreamComponent(Data.InstancedTransformComponent[2],11));
+		Elements.Add(AccessStreamComponent(Data.InstancedInverseTransformComponent[0],12));
+		Elements.Add(AccessStreamComponent(Data.InstancedInverseTransformComponent[1],13));
+		Elements.Add(AccessStreamComponent(Data.InstancedInverseTransformComponent[2],14));
+	}
 
 	// we don't need per-vertex shadow or lightmap rendering
 	InitDeclaration(Elements,Data);
@@ -677,6 +688,8 @@ void FInstancedStaticMeshRenderData::InitStaticMeshVertexFactories(
 		FInstancedStaticMeshRenderData* InstancedRenderData,
 		UStaticMesh* Parent)
 {
+	const bool bInstanced = RHISupportsInstancing(GRHIShaderPlatform);
+
 	for( int32 LODIndex=0;LODIndex<VertexFactories->Num(); LODIndex++ )
 	{
 		const FStaticMeshLODResources* RenderData = &InstancedRenderData->LODModels[LODIndex];
@@ -786,35 +799,13 @@ void FInstancedStaticMeshRenderData::InitStaticMeshVertexFactories(
 					);
 #endif
 			}
-		}	
-
-	
-		// Shadow map bias (and random instance ID)
-		int32 CurInstanceBufferOffset = 0;
-		Data.InstancedShadowMapBiasComponent = FVertexStreamComponent(
-			&InstancedRenderData->InstanceBuffer,
-			CurInstanceBufferOffset, 
-			InstancedRenderData->InstanceBuffer.GetStride(),
-			VET_Float4,
-			true
-			);
-		CurInstanceBufferOffset += sizeof(float) * 4;
-
-		for (int32 MatrixRow = 0; MatrixRow < 3; MatrixRow++)
-		{
-			Data.InstancedTransformComponent[MatrixRow] = FVertexStreamComponent(
-				&InstancedRenderData->InstanceBuffer,
-				CurInstanceBufferOffset, 
-				InstancedRenderData->InstanceBuffer.GetStride(),
-				VET_Float4,
-				true
-				);
-			CurInstanceBufferOffset += sizeof(float) * 4;
 		}
 
-		for (int32 MatrixRow = 0; MatrixRow < 3; MatrixRow++)
+		if (bInstanced)
 		{
-			Data.InstancedInverseTransformComponent[MatrixRow] = FVertexStreamComponent(
+			// Shadow map bias (and random instance ID)
+			int32 CurInstanceBufferOffset = 0;
+			Data.InstancedShadowMapBiasComponent = FVertexStreamComponent(
 				&InstancedRenderData->InstanceBuffer,
 				CurInstanceBufferOffset, 
 				InstancedRenderData->InstanceBuffer.GetStride(),
@@ -822,6 +813,30 @@ void FInstancedStaticMeshRenderData::InitStaticMeshVertexFactories(
 				true
 				);
 			CurInstanceBufferOffset += sizeof(float) * 4;
+
+			for (int32 MatrixRow = 0; MatrixRow < 3; MatrixRow++)
+			{
+				Data.InstancedTransformComponent[MatrixRow] = FVertexStreamComponent(
+					&InstancedRenderData->InstanceBuffer,
+					CurInstanceBufferOffset, 
+					InstancedRenderData->InstanceBuffer.GetStride(),
+					VET_Float4,
+					true
+					);
+				CurInstanceBufferOffset += sizeof(float) * 4;
+			}
+
+			for (int32 MatrixRow = 0; MatrixRow < 3; MatrixRow++)
+			{
+				Data.InstancedInverseTransformComponent[MatrixRow] = FVertexStreamComponent(
+					&InstancedRenderData->InstanceBuffer,
+					CurInstanceBufferOffset, 
+					InstancedRenderData->InstanceBuffer.GetStride(),
+					VET_Float4,
+					true
+					);
+				CurInstanceBufferOffset += sizeof(float) * 4;
+			}
 		}
 
 		// Assign to the vertex factory for this LOD.
@@ -1032,7 +1047,8 @@ bool FInstancedStaticMeshSceneProxy::GetShadowMeshElement(int32 LODIndex, uint8
 {
 	if (LODIndex < InstancedRenderData.VertexFactories.Num() && FStaticMeshSceneProxy::GetShadowMeshElement(LODIndex, InDepthPriorityGroup, OutMeshElement))
 	{
-		OutMeshElement.Elements[0].NumInstances = InstancedRenderData.InstanceBuffer.GetNumInstances();
+		const bool bInstanced = RHISupportsInstancing(GRHIShaderPlatform);
+		OutMeshElement.Elements[0].NumInstances = bInstanced ? InstancedRenderData.InstanceBuffer.GetNumInstances() : 1;
 		OutMeshElement.Elements[0].UserData = (void*)&UserData_AllInstances;
 		OutMeshElement.VertexFactory = &InstancedRenderData.VertexFactories[LODIndex];
 		return true;
@@ -1046,7 +1062,8 @@ bool FInstancedStaticMeshSceneProxy::GetMeshElement(int32 LODIndex,int32 Element
 {
 	if (LODIndex < InstancedRenderData.VertexFactories.Num() && FStaticMeshSceneProxy::GetMeshElement(LODIndex, ElementIndex, InDepthPriorityGroup, OutMeshElement, bUseSelectedMaterial, bUseHoveredMaterial))
 	{
-		OutMeshElement.Elements[0].NumInstances = InstancedRenderData.InstanceBuffer.GetNumInstances();
+		const bool bInstanced = RHISupportsInstancing(GRHIShaderPlatform);
+		OutMeshElement.Elements[0].NumInstances = bInstanced ? InstancedRenderData.InstanceBuffer.GetNumInstances() : 1;
 		OutMeshElement.Elements[0].UserData = (void*)&UserData_AllInstances;
 		OutMeshElement.VertexFactory = &InstancedRenderData.VertexFactories[LODIndex];
 		return true;
@@ -1059,7 +1076,8 @@ bool FInstancedStaticMeshSceneProxy::GetWireframeMeshElement(int32 LODIndex, con
 {
 	if (LODIndex < InstancedRenderData.VertexFactories.Num() && FStaticMeshSceneProxy::GetWireframeMeshElement(LODIndex, WireframeRenderProxy, InDepthPriorityGroup, OutMeshElement))
 	{
-		OutMeshElement.Elements[0].NumInstances = InstancedRenderData.InstanceBuffer.GetNumInstances();
+		const bool bInstanced = RHISupportsInstancing(GRHIShaderPlatform);
+		OutMeshElement.Elements[0].NumInstances = bInstanced ? InstancedRenderData.InstanceBuffer.GetNumInstances() : 1;
 		OutMeshElement.Elements[0].UserData = (void*)&UserData_AllInstances;
 		OutMeshElement.VertexFactory = &InstancedRenderData.VertexFactories[LODIndex];
 		return true;
@@ -1121,12 +1139,6 @@ void UInstancedStaticMeshComponent::ApplyComponentInstanceData(TSharedPtr<FCompo
 
 FPrimitiveSceneProxy* UInstancedStaticMeshComponent::CreateSceneProxy()
 {
-	// We don't support instancing on ES2
-	if (GetWorld()->FeatureLevel == ERHIFeatureLevel::ES2)
-	{
-		return NULL;
-	}
-
 	// Verify that the mesh is valid before using it.
 	const bool bMeshIsValid = 
 		// make sure we have instances
@@ -1256,10 +1268,10 @@ FBoxSphereBounds UInstancedStaticMeshComponent::CalcBounds(const FTransform& Bou
 		FBoxSphereBounds RenderBounds = StaticMesh->GetBounds();
 		FBoxSphereBounds NewBounds = RenderBounds.TransformBy(PerInstanceSMData[0].Transform * BoundTransformMatrix);
 
- 		for (int32 InstanceIndex = 1; InstanceIndex < PerInstanceSMData.Num(); InstanceIndex++)
- 		{
- 			NewBounds = NewBounds + RenderBounds.TransformBy(PerInstanceSMData[InstanceIndex].Transform * BoundTransformMatrix);
- 		}
+		for (int32 InstanceIndex = 1; InstanceIndex < PerInstanceSMData.Num(); InstanceIndex++)
+		{
+			NewBounds = NewBounds + RenderBounds.TransformBy(PerInstanceSMData[InstanceIndex].Transform * BoundTransformMatrix);
+		}
 
 		return NewBounds;
 	}
diff --git a/Engine/Source/Runtime/RHI/Public/RHIDefinitions.h b/Engine/Source/Runtime/RHI/Public/RHIDefinitions.h
--- a/Engine/Source/Runtime/RHI/Public/RHIDefinitions.h
+++ b/Engine/Source/Runtime/RHI/Public/RHIDefinitions.h
@@ -607,7 +607,8 @@ inline bool RHISupportsTessellation(const EShaderPlatform Platform)
 
 inline bool RHISupportsInstancing(const EShaderPlatform Platform)
 {
-	return IsFeatureLevelSupported(Platform, ERHIFeatureLevel::SM3) || (Platform == SP_METAL);
+	//@todo-rco: Add Metal support
+	return IsFeatureLevelSupported(Platform, ERHIFeatureLevel::SM3);// || (Platform == SP_METAL);
 }
 
 inline uint32 GetFeatureLevelMaxTextureSamplers(ERHIFeatureLevel::Type FeatureLevel)
