Remove FLandscapeToolSet

diff --git a/Engine/Source/Editor/LandscapeEditor/Private/LandscapeEditorDetailCustomization_ResizeLandscape.cpp b/Engine/Source/Editor/LandscapeEditor/Private/LandscapeEditorDetailCustomization_ResizeLandscape.cpp
--- a/Engine/Source/Editor/LandscapeEditor/Private/LandscapeEditorDetailCustomization_ResizeLandscape.cpp
+++ b/Engine/Source/Editor/LandscapeEditor/Private/LandscapeEditorDetailCustomization_ResizeLandscape.cpp
@@ -39,7 +39,7 @@ TSharedRef<IDetailCustomization> FLandscapeEditorDetailCustomization_ResizeLands
 BEGIN_SLATE_FUNCTION_BUILD_OPTIMIZATION
 void FLandscapeEditorDetailCustomization_ResizeLandscape::CustomizeDetails(IDetailLayoutBuilder& DetailBuilder)
 {
-	if (!IsToolActive("ToolSet_ResizeLandscape"))
+	if (!IsToolActive("ResizeLandscape"))
 	{
 		return;
 	}
@@ -311,7 +311,7 @@ TSharedRef<SWidget> FLandscapeEditorDetailCustomization_ResizeLandscape::GetSect
 
 	for (int32 i = 0; i < ARRAY_COUNT(SectionSizes); i++)
 	{
-		MenuBuilder.AddMenuEntry( FText::Format( LOCTEXT("NxNQuads", "{0}x{0} Quads"), FText::AsNumber( SectionSizes[i]) ), FText::GetEmpty(), FSlateIcon(), FExecuteAction::CreateStatic(&OnChangeSectionSize, PropertyHandle, SectionSizes[i]));
+		MenuBuilder.AddMenuEntry(FText::Format(LOCTEXT("NxNQuads", "{0}x{0} Quads"), FText::AsNumber(SectionSizes[i])), FText::GetEmpty(), FSlateIcon(), FExecuteAction::CreateStatic(&OnChangeSectionSize, PropertyHandle, SectionSizes[i]));
 	}
 
 	return MenuBuilder.MakeWidget();
@@ -381,7 +381,7 @@ TSharedRef<SWidget> FLandscapeEditorDetailCustomization_ResizeLandscape::GetSect
 		FFormatNamedArguments Args;
 		Args.Add(TEXT("Width"), NumSections[i]);
 		Args.Add(TEXT("Height"), NumSections[i]);
-		MenuBuilder.AddMenuEntry(FText::Format(NumSections[i] == 1 ? LOCTEXT("1x1Section", "{Width}\u00D7{Height} Section") : LOCTEXT("NxNSections", "{Width}\u00D7{Height} Sections"), Args ), FText::GetEmpty(), FSlateIcon(), FExecuteAction::CreateStatic(&OnChangeSectionsPerComponent, PropertyHandle, NumSections[i]));
+		MenuBuilder.AddMenuEntry(FText::Format(NumSections[i] == 1 ? LOCTEXT("1x1Section", "{Width}\u00D7{Height} Section") : LOCTEXT("NxNSections", "{Width}\u00D7{Height} Sections"), Args), FText::GetEmpty(), FSlateIcon(), FExecuteAction::CreateStatic(&OnChangeSectionsPerComponent, PropertyHandle, NumSections[i]));
 	}
 
 	return MenuBuilder.MakeWidget();
@@ -404,8 +404,8 @@ FText FLandscapeEditorDetailCustomization_ResizeLandscape::GetSectionsPerCompone
 	}
 
 	FFormatNamedArguments Args;
-	Args.Add( TEXT("Width"), SectionsPerComponent );
-	Args.Add( TEXT("Height"), SectionsPerComponent );
+	Args.Add(TEXT("Width"), SectionsPerComponent);
+	Args.Add(TEXT("Height"), SectionsPerComponent);
 	return FText::Format(SectionsPerComponent == 1 ? LOCTEXT("1x1Section", "{Width}\u00D7{Height} Section") : LOCTEXT("NxNSections", "{Width}\u00D7{Height} Sections"), Args);
 }
 
@@ -434,7 +434,7 @@ FText FLandscapeEditorDetailCustomization_ResizeLandscape::GetOriginalComponentC
 	FEdModeLandscape* LandscapeEdMode = GetEditorMode();
 	if (LandscapeEdMode != NULL)
 	{
-		return FText::Format(LOCTEXT("NxN", "{0}\u00D7{1}"), 
+		return FText::Format(LOCTEXT("NxN", "{0}\u00D7{1}"),
 			FText::AsNumber(LandscapeEdMode->UISettings->ResizeLandscape_Original_ComponentCount.X),
 			FText::AsNumber(LandscapeEdMode->UISettings->ResizeLandscape_Original_ComponentCount.Y));
 	}
@@ -443,7 +443,7 @@ FText FLandscapeEditorDetailCustomization_ResizeLandscape::GetOriginalComponentC
 
 FText FLandscapeEditorDetailCustomization_ResizeLandscape::GetComponentCount(TSharedRef<IPropertyHandle> PropertyHandle_ComponentCount_X, TSharedRef<IPropertyHandle> PropertyHandle_ComponentCount_Y)
 {
-	return FText::Format(LOCTEXT("NxN", "{0}\u00D7{1}"), 
+	return FText::Format(LOCTEXT("NxN", "{0}\u00D7{1}"),
 		FLandscapeEditorDetailCustomization_Base::GetPropertyValueText(PropertyHandle_ComponentCount_X),
 		FLandscapeEditorDetailCustomization_Base::GetPropertyValueText(PropertyHandle_ComponentCount_Y));
 }
@@ -513,8 +513,8 @@ FReply FLandscapeEditorDetailCustomization_ResizeLandscape::OnApplyButtonClicked
 
 		LandscapeEdMode->UpdateLandscapeList();
 
-		LandscapeEdMode->SetCurrentTool("ToolSet_Select"); // change tool so switching back to the manage mode doesn't give "Resize Landscape" again
-		LandscapeEdMode->SetCurrentTool("ToolSet_Sculpt"); // change to sculpting mode and tool
+		LandscapeEdMode->SetCurrentTool("Select"); // change tool so switching back to the manage mode doesn't give "Resize Landscape" again
+		LandscapeEdMode->SetCurrentTool("Sculpt"); // change to sculpting mode and tool
 	}
 
 	return FReply::Handled();
diff --git a/Engine/Source/Editor/LandscapeEditor/Private/LandscapeEditorDetailCustomization_TargetLayers.cpp b/Engine/Source/Editor/LandscapeEditor/Private/LandscapeEditorDetailCustomization_TargetLayers.cpp
--- a/Engine/Source/Editor/LandscapeEditor/Private/LandscapeEditorDetailCustomization_TargetLayers.cpp
+++ b/Engine/Source/Editor/LandscapeEditor/Private/LandscapeEditorDetailCustomization_TargetLayers.cpp
@@ -56,12 +56,13 @@ bool FLandscapeEditorDetailCustomization_TargetLayers::ShouldShowTargetLayers()
 	FEdModeLandscape* LandscapeEdMode = GetEditorMode();
 	if (LandscapeEdMode && LandscapeEdMode->CurrentToolMode)
 	{
-		bool bSupportsHeightmap = (LandscapeEdMode->CurrentToolMode->SupportedTargetTypes & ELandscapeToolTargetTypeMask::Heightmap) != 0;
-		bool bSupportsWeightmap = (LandscapeEdMode->CurrentToolMode->SupportedTargetTypes & ELandscapeToolTargetTypeMask::Weightmap) != 0;
-		bool bSupportsVisibility = (LandscapeEdMode->CurrentToolMode->SupportedTargetTypes & ELandscapeToolTargetTypeMask::Visibility) != 0;
+		//bool bSupportsHeightmap = (LandscapeEdMode->CurrentToolMode->SupportedTargetTypes & ELandscapeToolTargetTypeMask::Heightmap) != 0;
+		//bool bSupportsWeightmap = (LandscapeEdMode->CurrentToolMode->SupportedTargetTypes & ELandscapeToolTargetTypeMask::Weightmap) != 0;
+		//bool bSupportsVisibility = (LandscapeEdMode->CurrentToolMode->SupportedTargetTypes & ELandscapeToolTargetTypeMask::Visibility) != 0;
 
-		// Visible if there are possible choices
-		if (bSupportsWeightmap || bSupportsHeightmap || bSupportsVisibility)
+		//// Visible if there are possible choices
+		//if (bSupportsWeightmap || bSupportsHeightmap || bSupportsVisibility)
+		if (LandscapeEdMode->CurrentToolMode->SupportedTargetTypes != 0)
 		{
 			return true;
 		}
@@ -124,13 +125,13 @@ void FLandscapeEditorCustomNodeBuilder_TargetLayers::GenerateRow(IDetailChildren
 	FEdModeLandscape* LandscapeEdMode = GetEditorMode();
 	if (LandscapeEdMode)
 	{
-		if ( !(LandscapeEdMode->CurrentToolMode->SupportedTargetTypes & (1 << Target->TargetType)) )
+		if ((LandscapeEdMode->CurrentTool->GetSupportedTargetTypes() & LandscapeEdMode->CurrentToolMode->SupportedTargetTypes & ELandscapeToolTargetTypeMask::FromType(Target->TargetType)) == 0)
 		{
 			return;
 		}
 	}
 	
-	if (Target->TargetType == ELandscapeToolTargetType::Heightmap)
+	if (Target->TargetType != ELandscapeToolTargetType::Weightmap)
 	{
 		ChildrenBuilder.AddChildContent(Target->TargetName.ToString())
 		[
@@ -149,7 +150,7 @@ void FLandscapeEditorCustomNodeBuilder_TargetLayers::GenerateRow(IDetailChildren
 				//.AspectRatio()
 				[
 					SNew(SImage)
-					.Image(FEditorStyle::GetBrush(TEXT("LandscapeEditor.Target_Heightmap")))
+					.Image(FEditorStyle::GetBrush(Target->TargetType == ELandscapeToolTargetType::Heightmap ? TEXT("LandscapeEditor.Target_Heightmap") : TEXT("LandscapeEditor.Target_Visibility")))
 				]
 				+ SHorizontalBox::Slot()
 				.VAlign(VAlign_Center)
@@ -359,7 +360,7 @@ void FLandscapeEditorCustomNodeBuilder_TargetLayers::GenerateRow(IDetailChildren
 	}
 }
 END_SLATE_FUNCTION_BUILD_OPTIMIZATION
-	
+
 bool FLandscapeEditorCustomNodeBuilder_TargetLayers::GetTargetLayerIsSelected(const TSharedRef<FLandscapeTargetListInfo> Target)
 {
 	FEdModeLandscape* LandscapeEdMode = GetEditorMode();
@@ -557,7 +558,7 @@ void FLandscapeEditorCustomNodeBuilder_TargetLayers::OnReimportLayer(const TShar
 	}
 }
 
-bool FLandscapeEditorCustomNodeBuilder_TargetLayers::ShouldFilterLayerInfo(const class FAssetData& AssetData, FName LayerName)
+bool FLandscapeEditorCustomNodeBuilder_TargetLayers::ShouldFilterLayerInfo(const FAssetData& AssetData, FName LayerName)
 {
 	const FString* const LayerNameMetaData = AssetData.TagsAndValues.Find("LayerName");
 	if (LayerNameMetaData && !LayerNameMetaData->IsEmpty())
@@ -578,7 +579,7 @@ void FLandscapeEditorCustomNodeBuilder_TargetLayers::OnTargetLayerSetObject(cons
 		return;
 	}
 
-	FScopedTransaction Transaction( LOCTEXT("Undo_UseExisting", "Assigning Layer to Landscape") );
+	FScopedTransaction Transaction(LOCTEXT("Undo_UseExisting", "Assigning Layer to Landscape"));
 
 	ULandscapeLayerInfoObject* SelectedLayerInfo = const_cast<ULandscapeLayerInfoObject*>(CastChecked<ULandscapeLayerInfoObject>(Object));
 
@@ -687,7 +688,7 @@ void FLandscapeEditorCustomNodeBuilder_TargetLayers::OnTargetLayerCreateClicked(
 {
 	check(!Target->LayerInfoObj.IsValid());
 
-	FScopedTransaction Transaction( LOCTEXT("Undo_Create", "Creating New Landscape Layer") );
+	FScopedTransaction Transaction(LOCTEXT("Undo_Create", "Creating New Landscape Layer"));
 
 	FEdModeLandscape* LandscapeEdMode = GetEditorMode();
 	if (LandscapeEdMode)
@@ -704,7 +705,7 @@ void FLandscapeEditorCustomNodeBuilder_TargetLayers::OnTargetLayerCreateClicked(
 		}
 		FString PackageName = Path + LayerObjectName.ToString();
 
-		TSharedRef<SDlgPickAssetPath> NewLayerDlg = 
+		TSharedRef<SDlgPickAssetPath> NewLayerDlg =
 			SNew(SDlgPickAssetPath)
 			.Title(LOCTEXT("CreateNewLayerInfo", "Create New Landscape Layer Info Object"))
 			.DefaultAssetPath(FText::FromString(PackageName));
@@ -757,7 +758,7 @@ void FLandscapeEditorCustomNodeBuilder_TargetLayers::OnTargetLayerCreateClicked(
 
 FReply FLandscapeEditorCustomNodeBuilder_TargetLayers::OnTargetLayerMakePublicClicked(const TSharedRef<FLandscapeTargetListInfo> Target)
 {
-	FScopedTransaction Transaction( LOCTEXT("Undo_MakePublic", "Make Layer Public") );
+	FScopedTransaction Transaction(LOCTEXT("Undo_MakePublic", "Make Layer Public"));
 	TArray<UObject*> Objects;
 	Objects.Add(Target->LayerInfoObj.Get());
 
@@ -786,7 +787,7 @@ FReply FLandscapeEditorCustomNodeBuilder_TargetLayers::OnTargetLayerDeleteClicke
 
 	if (FMessageDialog::Open(EAppMsgType::YesNo, LOCTEXT("Prompt_DeleteLayer", "Are you sure you want to delete this layer?")) == EAppReturnType::Yes)
 	{
-		FScopedTransaction Transaction( LOCTEXT("Undo_Delete", "Delete Layer") );
+		FScopedTransaction Transaction(LOCTEXT("Undo_Delete", "Delete Layer"));
 
 		Target->LandscapeInfo->DeleteLayer(Target->LayerInfoObj.Get());
 
diff --git a/Engine/Source/Editor/LandscapeEditor/Private/LandscapeEditorDetails.cpp b/Engine/Source/Editor/LandscapeEditor/Private/LandscapeEditorDetails.cpp
--- a/Engine/Source/Editor/LandscapeEditor/Private/LandscapeEditorDetails.cpp
+++ b/Engine/Source/Editor/LandscapeEditor/Private/LandscapeEditorDetails.cpp
@@ -31,30 +31,30 @@ void FLandscapeEditorDetails::CustomizeDetails(IDetailLayoutBuilder& DetailBuild
 
 	auto NameToCommandMap = FLandscapeEditorCommands::Get().NameToCommandMap;
 
-#define MAP_TOOL(ToolSetName) CommandList->MapAction(NameToCommandMap.FindChecked(ToolSetName), FUIAction(FExecuteAction::CreateSP(this, &FLandscapeEditorDetails::OnChangeTool, FName(ToolSetName)), FCanExecuteAction::CreateSP(this, &FLandscapeEditorDetails::IsToolEnabled, FName(ToolSetName)), FIsActionChecked::CreateStatic(&FLandscapeEditorDetails::IsToolActive, FName(ToolSetName))));
-	MAP_TOOL("ToolSet_NewLandscape");
-	MAP_TOOL("ToolSet_ResizeLandscape");
-
-	MAP_TOOL("ToolSet_Sculpt");
-	MAP_TOOL("ToolSet_Paint");
-	MAP_TOOL("ToolSet_Smooth");
-	MAP_TOOL("ToolSet_Flatten");
-	MAP_TOOL("ToolSet_Ramp");
-	MAP_TOOL("ToolSet_Erosion");
-	MAP_TOOL("ToolSet_HydraErosion");
-	MAP_TOOL("ToolSet_Noise");
-	MAP_TOOL("ToolSet_Retopologize");
-	MAP_TOOL("ToolSet_Visibility");
-
-	MAP_TOOL("ToolSet_Select");
-	MAP_TOOL("ToolSet_AddComponent");
-	MAP_TOOL("ToolSet_DeleteComponent");
-	MAP_TOOL("ToolSet_MoveToLevel");
-
-	MAP_TOOL("ToolSet_Mask");
-	MAP_TOOL("ToolSet_CopyPaste");
-
-	MAP_TOOL("ToolSet_Splines");
+#define MAP_TOOL(ToolName) CommandList->MapAction(NameToCommandMap.FindChecked("Tool_" ToolName), FUIAction(FExecuteAction::CreateSP(this, &FLandscapeEditorDetails::OnChangeTool, FName(ToolName)), FCanExecuteAction::CreateSP(this, &FLandscapeEditorDetails::IsToolEnabled, FName(ToolName)), FIsActionChecked::CreateStatic(&FLandscapeEditorDetails::IsToolActive, FName(ToolName))));
+	MAP_TOOL("NewLandscape");
+	MAP_TOOL("ResizeLandscape");
+
+	MAP_TOOL("Sculpt");
+	MAP_TOOL("Paint");
+	MAP_TOOL("Smooth");
+	MAP_TOOL("Flatten");
+	MAP_TOOL("Ramp");
+	MAP_TOOL("Erosion");
+	MAP_TOOL("HydraErosion");
+	MAP_TOOL("Noise");
+	MAP_TOOL("Retopologize");
+	MAP_TOOL("Visibility");
+
+	MAP_TOOL("Select");
+	MAP_TOOL("AddComponent");
+	MAP_TOOL("DeleteComponent");
+	MAP_TOOL("MoveToLevel");
+
+	MAP_TOOL("Mask");
+	MAP_TOOL("CopyPaste");
+
+	MAP_TOOL("Splines");
 #undef MAP_TOOL
 #define MAP_BRUSH_SET(BrushSetName) CommandList->MapAction(NameToCommandMap.FindChecked(BrushSetName), FUIAction(FExecuteAction::CreateSP(this, &FLandscapeEditorDetails::OnChangeBrushSet, FName(BrushSetName)), FCanExecuteAction::CreateSP(this, &FLandscapeEditorDetails::IsBrushSetEnabled, FName(BrushSetName)), FIsActionChecked::CreateStatic(&FLandscapeEditorDetails::IsBrushSetActive, FName(BrushSetName))));
 	MAP_BRUSH_SET("BrushSet_Circle");
@@ -71,7 +71,7 @@ void FLandscapeEditorDetails::CustomizeDetails(IDetailLayoutBuilder& DetailBuild
 #undef MAP_BRUSH
 
 	static const FLinearColor BorderColor = FLinearColor(0.2f, 0.2f, 0.2f, 0.2f);
-	static const FSlateBrush* BorderStyle =  FEditorStyle::GetBrush("DetailsView.GroupSection");
+	static const FSlateBrush* BorderStyle = FEditorStyle::GetBrush("DetailsView.GroupSection");
 
 	IDetailCategoryBuilder& LandscapeEditorCategory = DetailBuilder.EditCategory("LandscapeEditor", TEXT(""), ECategoryPriority::TypeSpecific);
 
@@ -148,7 +148,7 @@ void FLandscapeEditorDetails::CustomizeDetails(IDetailLayoutBuilder& DetailBuild
 	Customization_TargetLayers->CustomizeDetails(DetailBuilder);
 }
 
-FText FLandscapeEditorDetails::GetLocalizedName(FName Name)
+FText FLandscapeEditorDetails::GetLocalizedName(FString Name)
 {
 	static bool bInitialized = false;
 	if (!bInitialized)
@@ -190,7 +190,7 @@ FText FLandscapeEditorDetails::GetLocalizedName(FName Name)
 	}
 
 	FText Result;
-	ensure(FText::FindText(LOCTEXT_NAMESPACE, Name.ToString(), Result));
+	ensure(FText::FindText(LOCTEXT_NAMESPACE, Name, Result));
 	return Result;
 }
 
@@ -227,7 +227,7 @@ TSharedRef<SWidget> FLandscapeEditorDetails::GetTargetLandscapeMenu()
 		for (auto It = LandscapeList.CreateConstIterator(); It; It++)
 		{
 			FUIAction Action = FUIAction(FExecuteAction::CreateStatic(&FLandscapeEditorDetails::OnChangeTargetLandscape, TWeakObjectPtr<ULandscapeInfo>(It->Info)));
-			MenuBuilder.AddMenuEntry( FText::FromString( It->Info->GetLandscapeProxy()->GetActorLabel() ), FText(), FSlateIcon(), Action);
+			MenuBuilder.AddMenuEntry(FText::FromString(It->Info->GetLandscapeProxy()->GetActorLabel()), FText(), FSlateIcon(), Action);
 		}
 
 		return MenuBuilder.MakeWidget();
@@ -253,10 +253,10 @@ void FLandscapeEditorDetails::OnChangeTargetLandscape(TWeakObjectPtr<ULandscapeI
 FText FLandscapeEditorDetails::GetCurrentToolName() const
 {
 	FEdModeLandscape* LandscapeEdMode = GetEditorMode();
-	if (LandscapeEdMode != NULL && LandscapeEdMode->CurrentToolSet != NULL)
+	if (LandscapeEdMode != NULL && LandscapeEdMode->CurrentTool != NULL)
 	{
-		const FName CurrentToolSetName = LandscapeEdMode->CurrentToolSet->GetToolSetName();
-		return GetLocalizedName(CurrentToolSetName);
+		const TCHAR* CurrentToolName = LandscapeEdMode->CurrentTool->GetToolName();
+		return GetLocalizedName(FString("ToolSet_") + CurrentToolName);
 	}
 
 	return LOCTEXT("Unknown", "Unknown");
@@ -265,10 +265,10 @@ FText FLandscapeEditorDetails::GetCurrentToolName() const
 FSlateIcon FLandscapeEditorDetails::GetCurrentToolIcon() const
 {
 	FEdModeLandscape* LandscapeEdMode = GetEditorMode();
-	if (LandscapeEdMode != NULL && LandscapeEdMode->CurrentToolSet != NULL)
+	if (LandscapeEdMode != NULL && LandscapeEdMode->CurrentTool != NULL)
 	{
-		const FName CurrentToolSetName = LandscapeEdMode->CurrentToolSet->GetToolSetName();
-		return FLandscapeEditorCommands::Get().NameToCommandMap.FindChecked(CurrentToolSetName)->GetIcon();
+		const TCHAR* CurrentToolName = LandscapeEdMode->CurrentTool->GetToolName();
+		return FLandscapeEditorCommands::Get().NameToCommandMap.FindChecked(*(FString("Tool_") + CurrentToolName))->GetIcon();
 	}
 
 	return FSlateIcon(FEditorStyle::GetStyleSetName(), "Default");
@@ -286,49 +286,49 @@ TSharedRef<SWidget> FLandscapeEditorDetails::GetToolSelector()
 		if (LandscapeEdMode->CurrentToolMode->ToolModeName == "ToolMode_Manage")
 		{
 			MenuBuilder.BeginSection(NAME_None, LOCTEXT("NewLandscapeToolsTitle", "New Landscape"));
-			MenuBuilder.AddToolButton(NameToCommandMap.FindChecked("ToolSet_NewLandscape"), NAME_None, LOCTEXT("Tool.NewLandscape", "New Landscape"), LOCTEXT("Tool.NewLandscape.Tooltip", "Create or import a new landscape"));
+			MenuBuilder.AddToolButton(NameToCommandMap.FindChecked("Tool_NewLandscape"), NAME_None, LOCTEXT("Tool.NewLandscape", "New Landscape"), LOCTEXT("Tool.NewLandscape.Tooltip", "Create or import a new landscape"));
 			MenuBuilder.EndSection();
 
 			MenuBuilder.BeginSection(NAME_None, LOCTEXT("ComponentToolsTitle", "Component Tools"));
-			MenuBuilder.AddToolButton(NameToCommandMap.FindChecked("ToolSet_Select"), NAME_None, LOCTEXT("Tool.SelectComponent", "Selection"), LOCTEXT("Tool.SelectComponent.Tooltip", "Select components to use with other tools"));
-			MenuBuilder.AddToolButton(NameToCommandMap.FindChecked("ToolSet_AddComponent"), NAME_None, LOCTEXT("Tool.AddComponent", "Add"), LOCTEXT("Tool.AddComponent.Tooltip", "Add components to the landscape"));
-			MenuBuilder.AddToolButton(NameToCommandMap.FindChecked("ToolSet_DeleteComponent"), NAME_None, LOCTEXT("Tool.DeleteComponent", "Delete"), LOCTEXT("Tool.DeleteComponent.Tooltip", "Delete components from the landscape, leaving a hole"));
-			MenuBuilder.AddToolButton(NameToCommandMap.FindChecked("ToolSet_MoveToLevel"), NAME_None, LOCTEXT("Tool.MoveToLevel", "Move Level"), LOCTEXT("Tool.MoveToLevel.Tooltip", "Move landscape components to a landscape proxy in the currently active streaming level, so that they can be streamed in/out independently of the rest of the landscape"));
-			MenuBuilder.AddToolButton(NameToCommandMap.FindChecked("ToolSet_ResizeLandscape"), NAME_None, LOCTEXT("Tool.ResizeLandscape", "Change Component Size"), LOCTEXT("Tool.ResizeLandscape.Tooltip", "Change the size of the landscape components"));
+			MenuBuilder.AddToolButton(NameToCommandMap.FindChecked("Tool_Select"), NAME_None, LOCTEXT("Tool.SelectComponent", "Selection"), LOCTEXT("Tool.SelectComponent.Tooltip", "Select components to use with other tools"));
+			MenuBuilder.AddToolButton(NameToCommandMap.FindChecked("Tool_AddComponent"), NAME_None, LOCTEXT("Tool.AddComponent", "Add"), LOCTEXT("Tool.AddComponent.Tooltip", "Add components to the landscape"));
+			MenuBuilder.AddToolButton(NameToCommandMap.FindChecked("Tool_DeleteComponent"), NAME_None, LOCTEXT("Tool.DeleteComponent", "Delete"), LOCTEXT("Tool.DeleteComponent.Tooltip", "Delete components from the landscape, leaving a hole"));
+			MenuBuilder.AddToolButton(NameToCommandMap.FindChecked("Tool_MoveToLevel"), NAME_None, LOCTEXT("Tool.MoveToLevel", "Move Level"), LOCTEXT("Tool.MoveToLevel.Tooltip", "Move landscape components to a landscape proxy in the currently active streaming level, so that they can be streamed in/out independently of the rest of the landscape"));
+			MenuBuilder.AddToolButton(NameToCommandMap.FindChecked("Tool_ResizeLandscape"), NAME_None, LOCTEXT("Tool.ResizeLandscape", "Change Component Size"), LOCTEXT("Tool.ResizeLandscape.Tooltip", "Change the size of the landscape components"));
 			MenuBuilder.EndSection();
 
 			MenuBuilder.BeginSection(NAME_None, LOCTEXT("SplineToolsTitle", "Spline Tools"));
-			MenuBuilder.AddToolButton(NameToCommandMap.FindChecked("ToolSet_Splines"), NAME_None, LOCTEXT("Tool.Spline", "Edit Splines"), LOCTEXT("Tool.Spline.Tooltip", "Ctrl+click to add control points\nHaving a control point selected when you ctrl+click will connect to the new control point with a segment\nSpline mesh settings can be found on the details panel when you have segments selected"));
+			MenuBuilder.AddToolButton(NameToCommandMap.FindChecked("Tool_Splines"), NAME_None, LOCTEXT("Tool.Spline", "Edit Splines"), LOCTEXT("Tool.Spline.Tooltip", "Ctrl+click to add control points\nHaving a control point selected when you ctrl+click will connect to the new control point with a segment\nSpline mesh settings can be found on the details panel when you have segments selected"));
 			MenuBuilder.EndSection();
 		}
 
 		if (LandscapeEdMode->CurrentToolMode->ToolModeName == "ToolMode_Sculpt")
 		{
 			MenuBuilder.BeginSection(NAME_None, LOCTEXT("SculptToolsTitle", "Sculpting Tools"));
-			MenuBuilder.AddToolButton(NameToCommandMap.FindChecked("ToolSet_Sculpt"), NAME_None, LOCTEXT("Tool.Sculpt", "Sculpt"), LOCTEXT("Tool.Sculpt.Tooltip", "Sculpt height data.\nCtrl+Click to Raise, Ctrl+Shift+Click to lower"));
-			MenuBuilder.AddToolButton(NameToCommandMap.FindChecked("ToolSet_Smooth"), NAME_None, LOCTEXT("Tool.Smooth", "Smooth"), LOCTEXT("Tool.Smooth.Tooltip", "Smooths heightmaps or blend layers"));
-			MenuBuilder.AddToolButton(NameToCommandMap.FindChecked("ToolSet_Flatten"), NAME_None, LOCTEXT("Tool.Flatten", "Flatten"), LOCTEXT("Tool.Flatten.Tooltip", "Flattens an area of heightmap or blend layer"));
-			MenuBuilder.AddToolButton(NameToCommandMap.FindChecked("ToolSet_Ramp"), NAME_None, LOCTEXT("Tool.Ramp", "Ramp"), LOCTEXT("Tool.Ramp.Tooltip", "Creates a ramp between two points"));
-			MenuBuilder.AddToolButton(NameToCommandMap.FindChecked("ToolSet_Erosion"), NAME_None, LOCTEXT("Tool.Erosion", "Erosion"), LOCTEXT("Tool.Erosion.Tooltip", "Thermal Erosion - Simulates erosion caused by the movement of soil from higher areas to lower areas"));
-			MenuBuilder.AddToolButton(NameToCommandMap.FindChecked("ToolSet_HydraErosion"), NAME_None, LOCTEXT("Tool.HydroErosion", "Hydro Erosion"), LOCTEXT("Tool.HydroErosion.Tooltip", "Hydro Erosion - Simulates erosion caused by rainfall"));
-			MenuBuilder.AddToolButton(NameToCommandMap.FindChecked("ToolSet_Noise"), NAME_None, LOCTEXT("Tool.Noise", "Noise"), LOCTEXT("Tool.Noise.Tooltip", "Adds noise to the heightmap or blend layer"));
-			MenuBuilder.AddToolButton(NameToCommandMap.FindChecked("ToolSet_Retopologize"), NAME_None, LOCTEXT("Tool.Retopologize", "Retopologize"), LOCTEXT("Tool.Retopologize.Tooltip", "Automatically adjusts landscape vertices with an X/Y offset map to improve vertex density on cliffs, reducing texture stretching.\nNote: An X/Y offset map makes the landscape slower to render and paint on with other tools, so only use if needed"));
-			MenuBuilder.AddToolButton(NameToCommandMap.FindChecked("ToolSet_Visibility"), NAME_None, LOCTEXT("Tool.Visibility", "Visibility"), LOCTEXT("Tool.Visibility.Tooltip", "Mask out individual quads in the landscape, leaving a hole."));
+			MenuBuilder.AddToolButton(NameToCommandMap.FindChecked("Tool_Sculpt"), NAME_None, LOCTEXT("Tool.Sculpt", "Sculpt"), LOCTEXT("Tool.Sculpt.Tooltip", "Sculpt height data.\nCtrl+Click to Raise, Ctrl+Shift+Click to lower"));
+			MenuBuilder.AddToolButton(NameToCommandMap.FindChecked("Tool_Smooth"), NAME_None, LOCTEXT("Tool.Smooth", "Smooth"), LOCTEXT("Tool.Smooth.Tooltip", "Smooths heightmaps or blend layers"));
+			MenuBuilder.AddToolButton(NameToCommandMap.FindChecked("Tool_Flatten"), NAME_None, LOCTEXT("Tool.Flatten", "Flatten"), LOCTEXT("Tool.Flatten.Tooltip", "Flattens an area of heightmap or blend layer"));
+			MenuBuilder.AddToolButton(NameToCommandMap.FindChecked("Tool_Ramp"), NAME_None, LOCTEXT("Tool.Ramp", "Ramp"), LOCTEXT("Tool.Ramp.Tooltip", "Creates a ramp between two points"));
+			MenuBuilder.AddToolButton(NameToCommandMap.FindChecked("Tool_Erosion"), NAME_None, LOCTEXT("Tool.Erosion", "Erosion"), LOCTEXT("Tool.Erosion.Tooltip", "Thermal Erosion - Simulates erosion caused by the movement of soil from higher areas to lower areas"));
+			MenuBuilder.AddToolButton(NameToCommandMap.FindChecked("Tool_HydraErosion"), NAME_None, LOCTEXT("Tool.HydroErosion", "Hydro Erosion"), LOCTEXT("Tool.HydroErosion.Tooltip", "Hydro Erosion - Simulates erosion caused by rainfall"));
+			MenuBuilder.AddToolButton(NameToCommandMap.FindChecked("Tool_Noise"), NAME_None, LOCTEXT("Tool.Noise", "Noise"), LOCTEXT("Tool.Noise.Tooltip", "Adds noise to the heightmap or blend layer"));
+			MenuBuilder.AddToolButton(NameToCommandMap.FindChecked("Tool_Retopologize"), NAME_None, LOCTEXT("Tool.Retopologize", "Retopologize"), LOCTEXT("Tool.Retopologize.Tooltip", "Automatically adjusts landscape vertices with an X/Y offset map to improve vertex density on cliffs, reducing texture stretching.\nNote: An X/Y offset map makes the landscape slower to render and paint on with other tools, so only use if needed"));
+			MenuBuilder.AddToolButton(NameToCommandMap.FindChecked("Tool_Visibility"), NAME_None, LOCTEXT("Tool.Visibility", "Visibility"), LOCTEXT("Tool.Visibility.Tooltip", "Mask out individual quads in the landscape, leaving a hole."));
 			MenuBuilder.EndSection();
 
 			MenuBuilder.BeginSection(NAME_None, LOCTEXT("RegionToolsTitle", "Region Tools"));
-			MenuBuilder.AddToolButton(NameToCommandMap.FindChecked("ToolSet_Mask"), NAME_None, LOCTEXT("Tool.RegionSelect", "Selection"), LOCTEXT("Tool.RegionSelect.Tooltip", "Select a region of landscape to use as a mask for other tools"));
-			MenuBuilder.AddToolButton(NameToCommandMap.FindChecked("ToolSet_CopyPaste"), NAME_None, LOCTEXT("Tool.RegionCopyPaste", "Copy/Paste"), LOCTEXT("Tool.RegionCopyPaste.Tooltip", "Copy/Paste areas of the landscape, or import/export a copied area of landscape from disk"));
+			MenuBuilder.AddToolButton(NameToCommandMap.FindChecked("Tool_Mask"), NAME_None, LOCTEXT("Tool.RegionSelect", "Selection"), LOCTEXT("Tool.RegionSelect.Tooltip", "Select a region of landscape to use as a mask for other tools"));
+			MenuBuilder.AddToolButton(NameToCommandMap.FindChecked("Tool_CopyPaste"), NAME_None, LOCTEXT("Tool.RegionCopyPaste", "Copy/Paste"), LOCTEXT("Tool.RegionCopyPaste.Tooltip", "Copy/Paste areas of the landscape, or import/export a copied area of landscape from disk"));
 			MenuBuilder.EndSection();
 		}
 
 		if (LandscapeEdMode->CurrentToolMode->ToolModeName == "ToolMode_Paint")
 		{
 			MenuBuilder.BeginSection(NAME_None, LOCTEXT("PaintToolsTitle", "Paint Tools"));
-			MenuBuilder.AddToolButton(NameToCommandMap.FindChecked("ToolSet_Paint"), NAME_None, LOCTEXT("Tool.Paint", "Paint"), LOCTEXT("Tool.Paint.Tooltip", "Paints weight data.\nCtrl+Click to paint, Ctrl+Shift+Click to erase"));
-			MenuBuilder.AddToolButton(NameToCommandMap.FindChecked("ToolSet_Smooth"), NAME_None, LOCTEXT("Tool.Smooth", "Smooth"), LOCTEXT("Tool.Smooth.Tooltip", "Smooths heightmaps or blend layers"));
-			MenuBuilder.AddToolButton(NameToCommandMap.FindChecked("ToolSet_Flatten"), NAME_None, LOCTEXT("Tool.Flatten", "Flatten"), LOCTEXT("Tool.Flatten.Tooltip", "Flattens an area of heightmap or blend layer"));
-			MenuBuilder.AddToolButton(NameToCommandMap.FindChecked("ToolSet_Noise"), NAME_None, LOCTEXT("Tool.Noise", "Noise"), LOCTEXT("Tool.Noise.Tooltip", "Adds noise to the heightmap or blend layer"));
+			MenuBuilder.AddToolButton(NameToCommandMap.FindChecked("Tool_Paint"), NAME_None, LOCTEXT("Tool.Paint", "Paint"), LOCTEXT("Tool.Paint.Tooltip", "Paints weight data.\nCtrl+Click to paint, Ctrl+Shift+Click to erase"));
+			MenuBuilder.AddToolButton(NameToCommandMap.FindChecked("Tool_Smooth"), NAME_None, LOCTEXT("Tool.Smooth", "Smooth"), LOCTEXT("Tool.Smooth.Tooltip", "Smooths heightmaps or blend layers"));
+			MenuBuilder.AddToolButton(NameToCommandMap.FindChecked("Tool_Flatten"), NAME_None, LOCTEXT("Tool.Flatten", "Flatten"), LOCTEXT("Tool.Flatten.Tooltip", "Flattens an area of heightmap or blend layer"));
+			MenuBuilder.AddToolButton(NameToCommandMap.FindChecked("Tool_Noise"), NAME_None, LOCTEXT("Tool.Noise", "Noise"), LOCTEXT("Tool.Noise.Tooltip", "Adds noise to the heightmap or blend layer"));
 			MenuBuilder.EndSection();
 		}
 
@@ -341,9 +341,9 @@ TSharedRef<SWidget> FLandscapeEditorDetails::GetToolSelector()
 bool FLandscapeEditorDetails::GetToolSelectorIsVisible() const
 {
 	FEdModeLandscape* LandscapeEdMode = GetEditorMode();
-	if (LandscapeEdMode && LandscapeEdMode->CurrentToolSet)
+	if (LandscapeEdMode && LandscapeEdMode->CurrentTool)
 	{
-		if (!IsToolActive("ToolSet_NewLandscape") || LandscapeEdMode->GetLandscapeList().Num() > 0)
+		if (!IsToolActive("NewLandscape") || LandscapeEdMode->GetLandscapeList().Num() > 0)
 		{
 			return true;
 		}
@@ -361,21 +361,21 @@ EVisibility FLandscapeEditorDetails::GetToolSelectorVisibility() const
 	return EVisibility::Collapsed;
 }
 
-void FLandscapeEditorDetails::OnChangeTool(FName ToolSetName)
+void FLandscapeEditorDetails::OnChangeTool(FName ToolName)
 {
 	FEdModeLandscape* LandscapeEdMode = GetEditorMode();
 	if (LandscapeEdMode != NULL)
 	{
-		LandscapeEdMode->SetCurrentTool(ToolSetName);
+		LandscapeEdMode->SetCurrentTool(ToolName);
 	}
 }
 
-bool FLandscapeEditorDetails::IsToolEnabled(FName ToolSetName) const
+bool FLandscapeEditorDetails::IsToolEnabled(FName ToolName) const
 {
 	FEdModeLandscape* LandscapeEdMode = GetEditorMode();
 	if (LandscapeEdMode != NULL)
 	{
-		if (ToolSetName == "ToolSet_NewLandscape" || LandscapeEdMode->GetLandscapeList().Num() > 0)
+		if (ToolName == "NewLandscape" || LandscapeEdMode->GetLandscapeList().Num() > 0)
 		{
 			return true;
 		}
@@ -390,7 +390,7 @@ FText FLandscapeEditorDetails::GetCurrentBrushName() const
 	if (LandscapeEdMode != NULL && LandscapeEdMode->CurrentBrush != NULL)
 	{
 		const FName CurrentBrushSetName = LandscapeEdMode->LandscapeBrushSets[LandscapeEdMode->CurrentBrushSetIndex].BrushSetName;
-		return GetLocalizedName(CurrentBrushSetName);
+		return GetLocalizedName(CurrentBrushSetName.ToString());
 	}
 
 	return LOCTEXT("Unknown", "Unknown");
@@ -415,34 +415,34 @@ FSlateIcon FLandscapeEditorDetails::GetCurrentBrushIcon() const
 TSharedRef<SWidget> FLandscapeEditorDetails::GetBrushSelector()
 {
 	FEdModeLandscape* LandscapeEdMode = GetEditorMode();
-	if (LandscapeEdMode && LandscapeEdMode->CurrentToolSet)
+	if (LandscapeEdMode && LandscapeEdMode->CurrentTool)
 	{
 		auto NameToCommandMap = FLandscapeEditorCommands::Get().NameToCommandMap;
 
 		FToolMenuBuilder MenuBuilder(true, CommandList);
 		MenuBuilder.BeginSection(NAME_None, LOCTEXT("BrushesTitle", "Brushes"));
 
-		if (LandscapeEdMode->CurrentToolSet->ValidBrushes.Contains("BrushSet_Circle"))
+		if (LandscapeEdMode->CurrentTool->ValidBrushes.Contains("BrushSet_Circle"))
 		{
 			MenuBuilder.AddToolButton(NameToCommandMap.FindChecked("BrushSet_Circle"), NAME_None, LOCTEXT("Brush.Circle", "Circle"), LOCTEXT("Brush.Circle.Brushtip", "Simple circular brush"));
 		}
-		if (LandscapeEdMode->CurrentToolSet->ValidBrushes.Contains("BrushSet_Alpha"))
+		if (LandscapeEdMode->CurrentTool->ValidBrushes.Contains("BrushSet_Alpha"))
 		{
 			MenuBuilder.AddToolButton(NameToCommandMap.FindChecked("BrushSet_Alpha"), NAME_None, LOCTEXT("Brush.Alpha.Alpha", "Alpha"), LOCTEXT("Brush.Alpha.Alpha.Tooltip", "Alpha brush, orients a mask image with the brush stroke"));
 		}
-		if (LandscapeEdMode->CurrentToolSet->ValidBrushes.Contains("BrushSet_Pattern"))
+		if (LandscapeEdMode->CurrentTool->ValidBrushes.Contains("BrushSet_Pattern"))
 		{
 			MenuBuilder.AddToolButton(NameToCommandMap.FindChecked("BrushSet_Pattern"), NAME_None, LOCTEXT("Brush.Alpha.Pattern", "Pattern"), LOCTEXT("Brush.Alpha.Pattern.Tooltip", "Pattern brush, tiles a mask image across the landscape"));
 		}
-		if (LandscapeEdMode->CurrentToolSet->ValidBrushes.Contains("BrushSet_Component"))
+		if (LandscapeEdMode->CurrentTool->ValidBrushes.Contains("BrushSet_Component"))
 		{
 			MenuBuilder.AddToolButton(NameToCommandMap.FindChecked("BrushSet_Component"), NAME_None, LOCTEXT("Brush.Component", "Component"), LOCTEXT("Brush.Component.Brushtip", "Work with entire landscape components"));
 		}
-		if (LandscapeEdMode->CurrentToolSet->ValidBrushes.Contains("BrushSet_Gizmo"))
+		if (LandscapeEdMode->CurrentTool->ValidBrushes.Contains("BrushSet_Gizmo"))
 		{
 			MenuBuilder.AddToolButton(NameToCommandMap.FindChecked("BrushSet_Gizmo"), NAME_None, LOCTEXT("Brush.Gizmo", "Gizmo"), LOCTEXT("Brush.Gizmo.Brushtip", "Work with the landscape gizmo, used for copy/pasting landscape"));
 		}
-		//if (LandscapeEdMode->CurrentToolSet->ValidBrushes.Contains("BrushSet_Splines"))
+		//if (LandscapeEdMode->CurrentTool->ValidBrushes.Contains("BrushSet_Splines"))
 		//{
 		//	MenuBuilder.AddToolButton(NameToCommandMap.FindChecked("BrushSet_Splines"), NAME_None, LOCTEXT("Brush.Splines", "Splines"), LOCTEXT("Brush.Splines.Brushtip", "Edit Splines"));
 		//}
@@ -457,9 +457,9 @@ TSharedRef<SWidget> FLandscapeEditorDetails::GetBrushSelector()
 bool FLandscapeEditorDetails::GetBrushSelectorIsVisible() const
 {
 	FEdModeLandscape* LandscapeEdMode = GetEditorMode();
-	if (LandscapeEdMode && LandscapeEdMode->CurrentToolSet)
+	if (LandscapeEdMode && LandscapeEdMode->CurrentTool)
 	{
-		if (LandscapeEdMode->CurrentToolSet->ValidBrushes.Num() >= 2)
+		if (LandscapeEdMode->CurrentTool->ValidBrushes.Num() >= 2)
 		{
 			return true;
 		}
@@ -480,9 +480,9 @@ void FLandscapeEditorDetails::OnChangeBrushSet(FName BrushSetName)
 bool FLandscapeEditorDetails::IsBrushSetEnabled(FName BrushSetName) const
 {
 	FEdModeLandscape* LandscapeEdMode = GetEditorMode();
-	if (LandscapeEdMode && LandscapeEdMode->CurrentToolSet)
+	if (LandscapeEdMode && LandscapeEdMode->CurrentTool)
 	{
-		return LandscapeEdMode->CurrentToolSet->ValidBrushes.Contains(BrushSetName);
+		return LandscapeEdMode->CurrentTool->ValidBrushes.Contains(BrushSetName);
 	}
 
 	return false;
@@ -493,7 +493,7 @@ FText FLandscapeEditorDetails::GetCurrentBrushFalloffName() const
 	FEdModeLandscape* LandscapeEdMode = GetEditorMode();
 	if (LandscapeEdMode != NULL && LandscapeEdMode->CurrentBrush != NULL)
 	{
-		const FName CurrentBrushName = LandscapeEdMode->CurrentBrush->GetBrushName();
+		const TCHAR* CurrentBrushName = LandscapeEdMode->CurrentBrush->GetBrushName();
 		return GetLocalizedName(CurrentBrushName);
 	}
 
@@ -519,7 +519,7 @@ FSlateIcon FLandscapeEditorDetails::GetCurrentBrushFalloffIcon() const
 TSharedRef<SWidget> FLandscapeEditorDetails::GetBrushFalloffSelector()
 {
 	FEdModeLandscape* LandscapeEdMode = GetEditorMode();
-	if (LandscapeEdMode && LandscapeEdMode->CurrentToolSet)
+	if (LandscapeEdMode && LandscapeEdMode->CurrentTool)
 	{
 		auto NameToCommandMap = FLandscapeEditorCommands::Get().NameToCommandMap;
 
diff --git a/Engine/Source/Editor/LandscapeEditor/Private/LandscapeEditorDetails.h b/Engine/Source/Editor/LandscapeEditor/Private/LandscapeEditorDetails.h
--- a/Engine/Source/Editor/LandscapeEditor/Private/LandscapeEditorDetails.h
+++ b/Engine/Source/Editor/LandscapeEditor/Private/LandscapeEditorDetails.h
@@ -5,6 +5,14 @@
 #include "PropertyEditorModule.h"
 #include "LandscapeEditorDetailCustomization_Base.h"
 
+// Forward declarations
+class FLandscapeEditorDetailCustomization_NewLandscape;
+class FLandscapeEditorDetailCustomization_ResizeLandscape;
+class FLandscapeEditorDetailCustomization_CopyPaste;
+class FLandscapeEditorDetailCustomization_MiscTools;
+class FLandscapeEditorDetailCustomization_AlphaBrush;
+class FLandscapeEditorDetailCustomization_TargetLayers;
+
 class FLandscapeEditorDetails : public FLandscapeEditorDetailCustomization_Base
 {
 public:
@@ -12,10 +20,10 @@ public:
 	static TSharedRef<IDetailCustomization> MakeInstance();
 
 	/** IDetailCustomization interface */
-	virtual void CustomizeDetails( IDetailLayoutBuilder& DetailBuilder ) override;
+	virtual void CustomizeDetails(IDetailLayoutBuilder& DetailBuilder) override;
 
 protected:
-	static FText GetLocalizedName(FName Name);
+	static FText GetLocalizedName(FString Name);
 
 	static EVisibility GetTargetLandscapeSelectorVisibility();
 	static FText GetTargetLandscapeName();
@@ -27,8 +35,8 @@ protected:
 	TSharedRef<SWidget> GetToolSelector();
 	bool GetToolSelectorIsVisible() const;
 	EVisibility GetToolSelectorVisibility() const;
-	void OnChangeTool(FName ToolSetName);
-	bool IsToolEnabled(FName ToolSetName) const;
+	void OnChangeTool(FName ToolName);
+	bool IsToolEnabled(FName ToolName) const;
 
 	FText GetCurrentBrushName() const;
 	FSlateIcon GetCurrentBrushIcon() const;
@@ -46,10 +54,10 @@ protected:
 
 	TSharedPtr<FUICommandList> CommandList;
 
-	TSharedPtr<class FLandscapeEditorDetailCustomization_NewLandscape> Customization_NewLandscape;
-	TSharedPtr<class FLandscapeEditorDetailCustomization_ResizeLandscape> Customization_ResizeLandscape;
-	TSharedPtr<class FLandscapeEditorDetailCustomization_CopyPaste> Customization_CopyPaste;
-	TSharedPtr<class FLandscapeEditorDetailCustomization_MiscTools> Customization_MiscTools;
-	TSharedPtr<class FLandscapeEditorDetailCustomization_AlphaBrush> Customization_AlphaBrush;
-	TSharedPtr<class FLandscapeEditorDetailCustomization_TargetLayers> Customization_TargetLayers;
+	TSharedPtr<FLandscapeEditorDetailCustomization_NewLandscape> Customization_NewLandscape;
+	TSharedPtr<FLandscapeEditorDetailCustomization_ResizeLandscape> Customization_ResizeLandscape;
+	TSharedPtr<FLandscapeEditorDetailCustomization_CopyPaste> Customization_CopyPaste;
+	TSharedPtr<FLandscapeEditorDetailCustomization_MiscTools> Customization_MiscTools;
+	TSharedPtr<FLandscapeEditorDetailCustomization_AlphaBrush> Customization_AlphaBrush;
+	TSharedPtr<FLandscapeEditorDetailCustomization_TargetLayers> Customization_TargetLayers;
 };
diff --git a/Engine/Source/Editor/LandscapeEditor/Private/SLandscapeEditor.cpp b/Engine/Source/Editor/LandscapeEditor/Private/SLandscapeEditor.cpp
--- a/Engine/Source/Editor/LandscapeEditor/Private/SLandscapeEditor.cpp
+++ b/Engine/Source/Editor/LandscapeEditor/Private/SLandscapeEditor.cpp
@@ -15,7 +15,7 @@ void SLandscapeAssetThumbnail::Construct(const FArguments& InArgs, UObject* Asse
 {
 	FIntPoint ThumbnailSize = InArgs._ThumbnailSize;
 
-	AssetThumbnail = MakeShareable( new FAssetThumbnail( Asset, ThumbnailSize.X, ThumbnailSize.Y, ThumbnailPool ) );
+	AssetThumbnail = MakeShareable(new FAssetThumbnail(Asset, ThumbnailSize.X, ThumbnailSize.Y, ThumbnailPool));
 
 	ChildSlot
 	[
@@ -52,9 +52,9 @@ void SLandscapeAssetThumbnail::OnMaterialCompilationFinished(UMaterialInterface*
 	}
 }
 
-void SLandscapeAssetThumbnail::SetAsset( UObject* Asset )
+void SLandscapeAssetThumbnail::SetAsset(UObject* Asset)
 {
-	AssetThumbnail->SetAsset( Asset );
+	AssetThumbnail->SetAsset(Asset);
 }
 
 //////////////////////////////////////////////////////////////////////////
@@ -116,9 +116,9 @@ void SLandscapeEditor::Construct(const FArguments& InArgs, TSharedRef<FLandscape
 	//CommandList = LevelEditor->GetLevelEditorActions();
 
 	// Modes:
-	CommandList->MapAction(FLandscapeEditorCommands::Get().ManageMode, FUIAction(FExecuteAction::CreateSP(this, &SLandscapeEditor::OnChangeMode, FName(TEXT("ToolMode_Manage"))), FCanExecuteAction::CreateSP(this, &SLandscapeEditor::IsModeEnabled, FName(TEXT("ToolMode_Manage"))), FIsActionChecked::CreateSP(this, &SLandscapeEditor::IsModeActive, FName(TEXT("ToolMode_Manage")))));
-	CommandList->MapAction(FLandscapeEditorCommands::Get().SculptMode, FUIAction(FExecuteAction::CreateSP(this, &SLandscapeEditor::OnChangeMode, FName(TEXT("ToolMode_Sculpt"))), FCanExecuteAction::CreateSP(this, &SLandscapeEditor::IsModeEnabled, FName(TEXT("ToolMode_Sculpt"))), FIsActionChecked::CreateSP(this, &SLandscapeEditor::IsModeActive, FName(TEXT("ToolMode_Sculpt")))));
-	CommandList->MapAction(FLandscapeEditorCommands::Get().PaintMode,  FUIAction(FExecuteAction::CreateSP(this, &SLandscapeEditor::OnChangeMode, FName(TEXT("ToolMode_Paint"))),  FCanExecuteAction::CreateSP(this, &SLandscapeEditor::IsModeEnabled, FName(TEXT("ToolMode_Paint"))),  FIsActionChecked::CreateSP(this, &SLandscapeEditor::IsModeActive, FName(TEXT("ToolMode_Paint")))));
+	CommandList->MapAction(FLandscapeEditorCommands::Get().ManageMode, FUIAction(FExecuteAction::CreateSP(this, &SLandscapeEditor::OnChangeMode, FName("ToolMode_Manage")), FCanExecuteAction::CreateSP(this, &SLandscapeEditor::IsModeEnabled, FName(TEXT("ToolMode_Manage"))), FIsActionChecked::CreateSP(this, &SLandscapeEditor::IsModeActive, FName(TEXT("ToolMode_Manage")))));
+	CommandList->MapAction(FLandscapeEditorCommands::Get().SculptMode, FUIAction(FExecuteAction::CreateSP(this, &SLandscapeEditor::OnChangeMode, FName("ToolMode_Sculpt")), FCanExecuteAction::CreateSP(this, &SLandscapeEditor::IsModeEnabled, FName(TEXT("ToolMode_Sculpt"))), FIsActionChecked::CreateSP(this, &SLandscapeEditor::IsModeActive, FName(TEXT("ToolMode_Sculpt")))));
+	CommandList->MapAction(FLandscapeEditorCommands::Get().PaintMode,  FUIAction(FExecuteAction::CreateSP(this, &SLandscapeEditor::OnChangeMode, FName("ToolMode_Paint" )), FCanExecuteAction::CreateSP(this, &SLandscapeEditor::IsModeEnabled, FName(TEXT("ToolMode_Paint" ))), FIsActionChecked::CreateSP(this, &SLandscapeEditor::IsModeActive, FName(TEXT("ToolMode_Paint" )))));
 
 	FToolBarBuilder ModeSwitchButtons(CommandList, FMultiBoxCustomization::None);
 	{
@@ -244,12 +244,12 @@ bool SLandscapeEditor::GetIsPropertyVisible(const FPropertyAndParent& PropertyAn
 	const UProperty& Property = PropertyAndParent.Property;
 
 	FEdModeLandscape* LandscapeEdMode = GetEditorMode();
-	if (LandscapeEdMode != NULL && LandscapeEdMode->CurrentToolSet != NULL)
+	if (LandscapeEdMode != NULL && LandscapeEdMode->CurrentTool != NULL)
 	{
 		if (Property.HasMetaData("ShowForMask"))
 		{
-			const bool bMaskEnabled = LandscapeEdMode->CurrentToolSet->GetTool() &&
-				LandscapeEdMode->CurrentToolSet->GetTool()->SupportsMask() &&
+			const bool bMaskEnabled = LandscapeEdMode->CurrentTool &&
+				LandscapeEdMode->CurrentTool->SupportsMask() &&
 				LandscapeEdMode->CurrentToolTarget.LandscapeInfo.IsValid() &&
 				LandscapeEdMode->CurrentToolTarget.LandscapeInfo->SelectedRegion.Num() > 0;
 
@@ -260,7 +260,7 @@ bool SLandscapeEditor::GetIsPropertyVisible(const FPropertyAndParent& PropertyAn
 		}
 		if (Property.HasMetaData("ShowForTools"))
 		{
-			const FName CurrentToolName = LandscapeEdMode->CurrentToolSet->GetToolSetName();
+			const FName CurrentToolName = LandscapeEdMode->CurrentTool->GetToolName();
 
 			TArray<FString> ShowForTools;
 			Property.GetMetaData("ShowForTools").ParseIntoArray(&ShowForTools, TEXT(","), true);
@@ -272,7 +272,7 @@ bool SLandscapeEditor::GetIsPropertyVisible(const FPropertyAndParent& PropertyAn
 		if (Property.HasMetaData("ShowForBrushes"))
 		{
 			const FName CurrentBrushSetName = LandscapeEdMode->LandscapeBrushSets[LandscapeEdMode->CurrentBrushSetIndex].BrushSetName;
-	//		const FName CurrentBrushName = LandscapeEdMode->CurrentBrush->GetBrushName();
+			// const FName CurrentBrushName = LandscapeEdMode->CurrentBrush->GetBrushName();
 
 			TArray<FString> ShowForBrushes;
 			Property.GetMetaData("ShowForBrushes").ParseIntoArray(&ShowForBrushes, TEXT(","), true);
@@ -330,7 +330,7 @@ bool SLandscapeEditor::IsModeEnabled(FName ModeName) const
 bool SLandscapeEditor::IsModeActive(FName ModeName) const
 {
 	FEdModeLandscape* LandscapeEdMode = GetEditorMode();
-	if (LandscapeEdMode && LandscapeEdMode->CurrentToolSet)
+	if (LandscapeEdMode && LandscapeEdMode->CurrentTool)
 	{
 		return LandscapeEdMode->CurrentToolMode->ToolModeName == ModeName;
 	}
diff --git a/Engine/Source/Editor/LandscapeEditor/Public/LandscapeEdMode.h b/Engine/Source/Editor/LandscapeEditor/Public/LandscapeEdMode.h
--- a/Engine/Source/Editor/LandscapeEditor/Public/LandscapeEdMode.h
+++ b/Engine/Source/Editor/LandscapeEditor/Public/LandscapeEdMode.h
@@ -8,6 +8,11 @@ DECLARE_LOG_CATEGORY_EXTERN(LogLandscapeEdMode, Log, All);
 
 // Forward declarations
 class ULandscapeLayerInfoObject;
+class FLandscapeToolSplines;
+
+struct FHeightmapToolTarget;
+template<typename TargetType> class FLandscapeToolCopyPaste;
+
 #include "Landscape/LandscapeInfo.h"
 #include "Landscape/LandscapeProxy.h"
 #include "Landscape/LandscapeGizmoActiveActor.h"
@@ -19,11 +24,12 @@ struct FLandscapeToolMousePosition
 	float PositionX, PositionY;
 	bool bShiftDown;
 
-	FLandscapeToolMousePosition( float InPositionX, float InPositionY, bool InbShiftDown )
-		:	PositionX(InPositionX)
-		,	PositionY(InPositionY)
-		,	bShiftDown(InbShiftDown)
-	{}
+	FLandscapeToolMousePosition(float InPositionX, float InPositionY, bool InbShiftDown)
+		: PositionX(InPositionX)
+		, PositionY(InPositionY)
+		, bShiftDown(InbShiftDown)
+	{
+	}
 };
 
 class FLandscapeBrush : public FGCObject
@@ -37,11 +43,11 @@ public:
 		BT_Gizmo,
 		BT_Splines
 	};
-	
-	virtual void MouseMove( float LandscapeX, float LandscapeY ) = 0;
-	virtual bool ApplyBrush( const TArray<FLandscapeToolMousePosition>& MousePositions, TMap<FIntPoint, float>& OutBrush, int32& OutX1, int32& OutY1, int32& OutX2, int32& OutY2 ) =0;
-	virtual bool InputKey( FEditorViewportClient* InViewportClient, FViewport* InViewport, FKey InKey, EInputEvent InEvent ) { return false; }
-	virtual void Tick(FEditorViewportClient* ViewportClient,float DeltaTime) {};
+
+	virtual void MouseMove(float LandscapeX, float LandscapeY) = 0;
+	virtual bool ApplyBrush(const TArray<FLandscapeToolMousePosition>& MousePositions, TMap<FIntPoint, float>& OutBrush, int32& OutX1, int32& OutY1, int32& OutX2, int32& OutY2) = 0;
+	virtual bool InputKey(FEditorViewportClient* InViewportClient, FViewport* InViewport, FKey InKey, EInputEvent InEvent) { return false; }
+	virtual void Tick(FEditorViewportClient* ViewportClient, float DeltaTime) {};
 	virtual void BeginStroke(float LandscapeX, float LandscapeY, class FLandscapeTool* CurrentTool);
 	virtual void EndStroke();
 	virtual void EnterBrush() {}
@@ -51,7 +57,7 @@ public:
 	virtual const TCHAR* GetBrushName() = 0;
 	virtual FText GetDisplayName() = 0;
 	virtual EBrushType GetBrushType() { return BT_Normal; }
-	
+
 	// FGCObject interface
 	virtual void AddReferencedObjects(FReferenceCollector& Collector) override {}
 };
@@ -70,7 +76,7 @@ struct FLandscapeBrushSet
 
 	virtual ~FLandscapeBrushSet()
 	{
-		for( int32 BrushIdx=0;BrushIdx<Brushes.Num();BrushIdx++ )
+		for (int32 BrushIdx = 0; BrushIdx < Brushes.Num(); BrushIdx++)
 		{
 			delete Brushes[BrushIdx];
 		}
@@ -79,27 +85,36 @@ struct FLandscapeBrushSet
 
 namespace ELandscapeToolTargetType
 {
-	enum Type
+	enum Type : int8
 	{
-		Heightmap	= 0,
-		Weightmap	= 1,
-		Visibility	= 2,
+		Heightmap  = 0,
+		Weightmap  = 1,
+		Visibility = 2,
 
-		Invalid		= -1, // only valid for LandscapeEdMode->CurrentToolTarget.TargetType
+		Invalid    = -1, // only valid for LandscapeEdMode->CurrentToolTarget.TargetType
 	};
 }
 
 namespace ELandscapeToolTargetTypeMask
 {
-	enum Type
+	enum Type : uint8
 	{
-		Heightmap	= 1<<ELandscapeToolTargetType::Heightmap,
-		Weightmap	= 1<<ELandscapeToolTargetType::Weightmap,
-		Visibility	= 1<<ELandscapeToolTargetType::Visibility,
+		Heightmap = 1 << ELandscapeToolTargetType::Heightmap,
+		Weightmap = 1 << ELandscapeToolTargetType::Weightmap,
+		Visibility = 1 << ELandscapeToolTargetType::Visibility,
 
-		NA			= 0,
-		All			= 0xFF,
+		NA = 0,
+		All = 0xFF,
 	};
+
+	inline ELandscapeToolTargetTypeMask::Type FromType(ELandscapeToolTargetType::Type TargetType)
+	{
+		if (TargetType == ELandscapeToolTargetType::Invalid)
+		{
+			return ELandscapeToolTargetTypeMask::NA;
+		}
+		return (ELandscapeToolTargetTypeMask::Type)(1 << TargetType);
+	}
 }
 
 namespace ELandscapeToolNoiseMode
@@ -114,7 +129,7 @@ namespace ELandscapeToolNoiseMode
 
 	inline float Conversion(Type Mode, float NoiseAmount, float OriginalValue)
 	{
-		switch( Mode )
+		switch (Mode)
 		{
 		case Add: // always +
 			OriginalValue += NoiseAmount;
@@ -135,7 +150,7 @@ struct FLandscapeToolTarget
 	ELandscapeToolTargetType::Type TargetType;
 	TWeakObjectPtr<ULandscapeLayerInfoObject> LayerInfo;
 	FName LayerName;
-	
+
 	FLandscapeToolTarget()
 		: LandscapeInfo(NULL)
 		, TargetType(ELandscapeToolTargetType::Heightmap)
@@ -158,14 +173,13 @@ public:
 	};
 	virtual void EnterTool() {}
 	virtual void ExitTool() {}
-	virtual bool IsValidForTarget(const FLandscapeToolTarget& Target) = 0;
-	virtual bool BeginTool( FEditorViewportClient* ViewportClient, const FLandscapeToolTarget& Target, const FVector& InHitLocation ) = 0;
+	virtual bool BeginTool(FEditorViewportClient* ViewportClient, const FLandscapeToolTarget& Target, const FVector& InHitLocation) = 0;
 	virtual void EndTool(FEditorViewportClient* ViewportClient) = 0;
-	virtual void Tick(FEditorViewportClient* ViewportClient,float DeltaTime) {};
-	virtual bool MouseMove( FEditorViewportClient* ViewportClient, FViewport* Viewport, int32 x, int32 y ) = 0;
-	virtual bool HandleClick( HHitProxy* HitProxy, const FViewportClick& Click ) { return false; }
-	virtual bool InputKey( FEditorViewportClient* InViewportClient, FViewport* InViewport, FKey InKey, EInputEvent InEvent ) { return false; }
-	virtual bool InputDelta( FEditorViewportClient* InViewportClient, FViewport* InViewport, FVector& InDrag, FRotator& InRot, FVector& InScale ) { return false; }
+	virtual void Tick(FEditorViewportClient* ViewportClient, float DeltaTime) {};
+	virtual bool MouseMove(FEditorViewportClient* ViewportClient, FViewport* Viewport, int32 x, int32 y) = 0;
+	virtual bool HandleClick(HHitProxy* HitProxy, const FViewportClick& Click) { return false; }
+	virtual bool InputKey(FEditorViewportClient* InViewportClient, FViewport* InViewport, FKey InKey, EInputEvent InEvent) { return false; }
+	virtual bool InputDelta(FEditorViewportClient* InViewportClient, FViewport* InViewport, FVector& InDrag, FRotator& InRot, FVector& InScale) { return false; }
 	virtual ~FLandscapeTool() {}
 	virtual const TCHAR* GetToolName() = 0;
 	virtual FText GetDisplayName() = 0;
@@ -174,7 +188,7 @@ public:
 	virtual bool SupportsMask() { return true; }
 	virtual bool SupportsComponentSelection() { return false; }
 	virtual bool OverrideSelection() const { return false; }
-	virtual bool IsSelectionAllowed( AActor* InActor, bool bInSelection ) const { return false; }
+	virtual bool IsSelectionAllowed(AActor* InActor, bool bInSelection) const { return false; }
 	virtual bool UsesTransformWidget() const { return false; }
 	virtual EAxisList::Type GetWidgetAxisToDraw(FWidget::EWidgetMode InWidgetMode) const { return EAxisList::All; }
 	virtual FVector GetWidgetLocation() const { return FVector::ZeroVector; }
@@ -195,72 +209,10 @@ public:
 	// Functions which doesn't need Viewport data...
 	virtual void Process(int32 Index, int32 Arg) {}
 	virtual EToolType GetToolType() { return TT_Normal; }
-};
+	virtual ELandscapeToolTargetTypeMask::Type GetSupportedTargetTypes() { return ELandscapeToolTargetTypeMask::NA; };
 
-class FLandscapeToolSet
-{
-	TArray<FLandscapeTool*> Tools;
-	FLandscapeTool*			CurrentTool;
-	FName					ToolSetName;
-public:
 	int32					PreviousBrushIndex;
 	TArray<FName>			ValidBrushes;
-
-	FLandscapeToolSet(FName InToolSetName)
-	:	ToolSetName(InToolSetName)
-	,	CurrentTool(NULL)
-	,	PreviousBrushIndex(INDEX_NONE)
-	{		
-	}
-
-	virtual ~FLandscapeToolSet()
-	{
-		for( int32 ToolIdx=0;ToolIdx<Tools.Num();ToolIdx++ )
-		{
-			delete Tools[ToolIdx];
-		}
-	}
-
-	virtual FText GetDisplayName() { return Tools[0]->GetDisplayName(); }
-
-	void AddTool(FLandscapeTool* InTool)
-	{
-		Tools.Add(InTool);
-	}
-
-	bool SetToolForTarget( const FLandscapeToolTarget& Target )
-	{
-		for( int32 ToolIdx=0;ToolIdx<Tools.Num();ToolIdx++ )
-		{
-			if( Tools[ToolIdx]->IsValidForTarget(Target) )
-			{
-				CurrentTool = Tools[ToolIdx];
-				return true;
-			}
-		}
-
-		return false;
-	}
-
-	FLandscapeTool* GetTool()
-	{
-		return CurrentTool;
-	}
-
-	FName GetToolSetName()
-	{
-		return ToolSetName;
-	}
-
-	bool operator==(const FLandscapeToolSet& OtherToolSet) const
-	{
-		return OtherToolSet.ToolSetName==ToolSetName;
-	}
-
-	bool operator==(const FName& OtherName) const
-	{
-		return OtherName==ToolSetName;
-	}
 };
 
 struct FLandscapeToolMode
@@ -268,8 +220,8 @@ struct FLandscapeToolMode
 	const FName				ToolModeName;
 	int32					SupportedTargetTypes; // ELandscapeToolTargetTypeMask::Type
 
-	TArray<FName>			ValidToolSets;
-	FName					CurrentToolSetName;
+	TArray<FName>			ValidTools;
+	FName					CurrentToolName;
 
 	FLandscapeToolMode(FName InToolModeName, int32 InSupportedTargetTypes)
 		: ToolModeName(InToolModeName)
@@ -290,7 +242,7 @@ struct FLandscapeTargetListInfo
 	TWeakObjectPtr<class ALandscapeProxy> Owner;					// ignored for heightmap
 	TWeakObjectPtr<class UMaterialInstanceConstant> ThumbnailMIC;	// ignored for heightmap
 	int32 DebugColorChannel;										// ignored for heightmap
-	uint32 bValid:1;												// ignored for heightmap
+	uint32 bValid : 1;												// ignored for heightmap
 
 	FLandscapeTargetListInfo(FText InTargetName, ELandscapeToolTargetType::Type InTargetType, const FLandscapeInfoLayerSettings& InLayerSettings)
 		: TargetName(InTargetName)
@@ -460,7 +412,7 @@ public:
 	class ULandscapeEditorObject* UISettings;
 
 	FLandscapeToolMode* CurrentToolMode;
-	FLandscapeToolSet* CurrentToolSet;
+	FLandscapeTool* CurrentTool;
 	FLandscapeBrush* CurrentBrush;
 	FLandscapeToolTarget CurrentToolTarget;
 
@@ -477,11 +429,11 @@ public:
 
 	TWeakObjectPtr<ALandscapeGizmoActiveActor> CurrentGizmoActor;
 	//
-	FLandscapeToolSet* CopyPasteToolSet;
+	FLandscapeToolCopyPaste<FHeightmapToolTarget>* CopyPasteTool;
 	void CopyDataToGizmo();
 	void PasteDataFromGizmo();
 
-	FLandscapeToolSet* SplinesToolSet;
+	FLandscapeToolSplines* SplinesTool;
 	void ShowSplineProperties();
 	virtual void SelectAllConnectedSplineControlPoints();
 	virtual void SelectAllConnectedSplineSegments();
@@ -498,31 +450,31 @@ public:
 
 	/** Initialization */
 	void InitializeBrushes();
-	void IntializeToolSet_Paint();
-	void IntializeToolSet_Smooth();
-	void IntializeToolSet_Flatten();
-	void IntializeToolSet_Erosion();
-	void IntializeToolSet_HydraErosion();
-	void IntializeToolSet_Noise();
-	void IntializeToolSet_Retopologize();
-	void InitializeToolSet_NewLandscape();
-	void InitializeToolSet_ResizeLandscape();
-	void IntializeToolSet_Select();
-	void IntializeToolSet_AddComponent();
-	void IntializeToolSet_DeleteComponent();
-	void IntializeToolSet_MoveToLevel();
-	void IntializeToolSet_Mask();
-	void IntializeToolSet_CopyPaste();
-	void IntializeToolSet_Visibility();
-	void IntializeToolSet_Splines();
-	void InitializeToolSet_Ramp();
+	void InitializeTool_Paint();
+	void InitializeTool_Smooth();
+	void InitializeTool_Flatten();
+	void InitializeTool_Erosion();
+	void InitializeTool_HydraErosion();
+	void InitializeTool_Noise();
+	void InitializeTool_Retopologize();
+	void InitializeTool_NewLandscape();
+	void InitializeTool_ResizeLandscape();
+	void InitializeTool_Select();
+	void InitializeTool_AddComponent();
+	void InitializeTool_DeleteComponent();
+	void InitializeTool_MoveToLevel();
+	void InitializeTool_Mask();
+	void InitializeTool_CopyPaste();
+	void InitializeTool_Visibility();
+	void InitializeTool_Splines();
+	void InitializeTool_Ramp();
 	void InitializeToolModes();
 
 	/** Destructor */
 	virtual ~FEdModeLandscape();
 
 	/** FGCObject interface */
-	virtual void AddReferencedObjects( FReferenceCollector& Collector ) override;
+	virtual void AddReferencedObjects(FReferenceCollector& Collector) override;
 
 	virtual bool UsesToolkits() const override;
 
@@ -533,7 +485,7 @@ public:
 	virtual void Exit() override;
 
 	/** FEdMode: Called when the mouse is moved over the viewport */
-	virtual bool MouseMove( FEditorViewportClient* ViewportClient, FViewport* Viewport, int32 x, int32 y ) override;
+	virtual bool MouseMove(FEditorViewportClient* ViewportClient, FViewport* Viewport, int32 x, int32 y) override;
 
 	/**
 	 * FEdMode: Called when the mouse is moved while a window input capture is in effect
@@ -545,7 +497,7 @@ public:
 	 *
 	 * @return	true if input was handled
 	 */
-	virtual bool CapturedMouseMove( FEditorViewportClient* InViewportClient, FViewport* InViewport, int32 InMouseX, int32 InMouseY ) override;
+	virtual bool CapturedMouseMove(FEditorViewportClient* InViewportClient, FViewport* InViewport, int32 InMouseX, int32 InMouseY) override;
 
 	/** FEdMode: Called when a mouse button is pressed */
 	virtual bool StartTracking(FEditorViewportClient* InViewportClient, FViewport* InViewport) override;
@@ -557,28 +509,28 @@ public:
 	virtual bool DisallowMouseDeltaTracking() const override;
 
 	/** FEdMode: Called once per frame */
-	virtual void Tick(FEditorViewportClient* ViewportClient,float DeltaTime) override;
+	virtual void Tick(FEditorViewportClient* ViewportClient, float DeltaTime) override;
 
 	/** FEdMode: Called when clicking on a hit proxy */
 	virtual bool HandleClick(FEditorViewportClient* InViewportClient, HHitProxy* HitProxy, const FViewportClick& Click) override;
 
 	/** FEdMode: Called when a key is pressed */
-	virtual bool InputKey( FEditorViewportClient* InViewportClient, FViewport* InViewport, FKey InKey, EInputEvent InEvent ) override;
+	virtual bool InputKey(FEditorViewportClient* InViewportClient, FViewport* InViewport, FKey InKey, EInputEvent InEvent) override;
 
 	/** FEdMode: Called when mouse drag input is applied */
-	virtual bool InputDelta( FEditorViewportClient* InViewportClient, FViewport* InViewport, FVector& InDrag, FRotator& InRot, FVector& InScale ) override;
+	virtual bool InputDelta(FEditorViewportClient* InViewportClient, FViewport* InViewport, FVector& InDrag, FRotator& InRot, FVector& InScale) override;
 
 	/** FEdMode: Render elements for the landscape tool */
-	virtual void Render( const FSceneView* View, FViewport* Viewport, FPrimitiveDrawInterface* PDI ) override;
+	virtual void Render(const FSceneView* View, FViewport* Viewport, FPrimitiveDrawInterface* PDI) override;
 
 	/** FEdMode: Render HUD elements for this tool */
-	virtual void DrawHUD( FEditorViewportClient* ViewportClient, FViewport* Viewport, const FSceneView* View, FCanvas* Canvas ) override;
+	virtual void DrawHUD(FEditorViewportClient* ViewportClient, FViewport* Viewport, const FSceneView* View, FCanvas* Canvas) override;
 
 	/** FEdMode: Handling SelectActor */
-	virtual bool Select( AActor* InActor, bool bInSelected ) override;
+	virtual bool Select(AActor* InActor, bool bInSelected) override;
 
 	/** FEdMode: Check to see if an actor can be selected in this mode - no side effects */
-	virtual bool IsSelectionAllowed( AActor* InActor, bool bInSelection ) const override;
+	virtual bool IsSelectionAllowed(AActor* InActor, bool bInSelection) const override;
 
 	/** FEdMode: Called when the currently selected actor has changed */
 	virtual void ActorSelectionChangeNotify() override;
@@ -605,32 +557,32 @@ public:
 	/** FEdMode: Returns true if this mode uses the transform widget */
 	virtual bool UsesTransformWidget() const override;
 
-	virtual EAxisList::Type GetWidgetAxisToDraw( FWidget::EWidgetMode InWidgetMode ) const override;
+	virtual EAxisList::Type GetWidgetAxisToDraw(FWidget::EWidgetMode InWidgetMode) const override;
 
 	virtual FVector GetWidgetLocation() const override;
-	virtual bool GetCustomDrawingCoordinateSystem( FMatrix& InMatrix, void* InData ) override;
-	virtual bool GetCustomInputCoordinateSystem( FMatrix& InMatrix, void* InData ) override;
+	virtual bool GetCustomDrawingCoordinateSystem(FMatrix& InMatrix, void* InData) override;
+	virtual bool GetCustomInputCoordinateSystem(FMatrix& InMatrix, void* InData) override;
 
 	/** Forces real-time perspective viewports */
-	void ForceRealTimeViewports( const bool bEnable, const bool bStoreCurrentState );
+	void ForceRealTimeViewports(const bool bEnable, const bool bStoreCurrentState);
 
 	/** Trace under the mouse cursor and return the landscape hit and the hit location (in landscape quad space) */
-	bool LandscapeMouseTrace( FEditorViewportClient* ViewportClient, float& OutHitX, float& OutHitY );
-	bool LandscapeMouseTrace( FEditorViewportClient* ViewportClient, FVector& OutHitLocation );
+	bool LandscapeMouseTrace(FEditorViewportClient* ViewportClient, float& OutHitX, float& OutHitY);
+	bool LandscapeMouseTrace(FEditorViewportClient* ViewportClient, FVector& OutHitLocation);
 
 	/** Trace under the specified coordinates and return the landscape hit and the hit location (in landscape quad space) */
-	bool LandscapeMouseTrace( FEditorViewportClient* ViewportClient, int32 MouseX, int32 MouseY, float& OutHitX, float& OutHitY );
-	bool LandscapeMouseTrace( FEditorViewportClient* ViewportClient, int32 MouseX, int32 MouseY, FVector& OutHitLocation );
+	bool LandscapeMouseTrace(FEditorViewportClient* ViewportClient, int32 MouseX, int32 MouseY, float& OutHitX, float& OutHitY);
+	bool LandscapeMouseTrace(FEditorViewportClient* ViewportClient, int32 MouseX, int32 MouseY, FVector& OutHitLocation);
 
 	/** Trace under the mouse cursor / specified screen coordinates against a world-space plane and return the hit location (in world space) */
 	bool LandscapePlaneTrace(FEditorViewportClient* ViewportClient, const FPlane& Plane, FVector& OutHitLocation);
 	bool LandscapePlaneTrace(FEditorViewportClient* ViewportClient, int32 MouseX, int32 MouseY, const FPlane& Plane, FVector& OutHitLocation);
 
-	void SetCurrentToolMode( FName ToolModeName, bool bRestoreCurrentTool = true);
+	void SetCurrentToolMode(FName ToolModeName, bool bRestoreCurrentTool = true);
 
 	/** Change current tool */
-	void SetCurrentTool( FName ToolSetName );
-	void SetCurrentTool( int32 ToolIdx );
+	void SetCurrentTool(FName ToolName);
+	void SetCurrentTool(int32 ToolIdx);
 
 	void SetCurrentBrushSet(FName BrushSetName);
 	void SetCurrentBrushSet(int32 BrushSetIndex);
@@ -659,7 +611,7 @@ public:
 	ALandscape* ChangeComponentSetting(int32 NumComponentsX, int32 NumComponentsY, int32 InNumSubsections, int32 InSubsectionSizeQuads, bool bResample);
 
 	TArray<FLandscapeToolMode> LandscapeToolModes;
-	TArray<FLandscapeToolSet> LandscapeToolSets;
+	TArray<TUniquePtr<FLandscapeTool>> LandscapeTools;
 	TArray<FLandscapeBrushSet> LandscapeBrushSets;
 
 	// For collision add visualization
@@ -679,8 +631,8 @@ namespace LandscapeEditorUtils
 {
 	template<typename T>
 	TArray<T> ExpandData(const TArray<T>& Data,
-	                     int32 OldMinX, int32 OldMinY, int32 OldMaxX, int32 OldMaxY,
-	                     int32 NewMinX, int32 NewMinY, int32 NewMaxX, int32 NewMaxY)
+		int32 OldMinX, int32 OldMinY, int32 OldMaxX, int32 OldMaxY,
+		int32 NewMinX, int32 NewMinY, int32 NewMaxX, int32 NewMaxY)
 	{
 		const int32 OldWidth = OldMaxX - OldMinX + 1;
 		const int32 OldHeight = OldMaxY - OldMinY + 1;
diff --git a/Engine/Source/Editor/UnrealEd/Private/EditorExporters.cpp b/Engine/Source/Editor/UnrealEd/Private/EditorExporters.cpp
--- a/Engine/Source/Editor/UnrealEd/Private/EditorExporters.cpp
+++ b/Engine/Source/Editor/UnrealEd/Private/EditorExporters.cpp
@@ -1455,7 +1455,7 @@ static void AddActorToOBJs(AActor* Actor, TArray<FOBJGeom*>& Objects, TSet<UMate
 			for( int32 AllocIdx=0;AllocIdx < Component->WeightmapLayerAllocations.Num(); AllocIdx++ )
 			{
 				FWeightmapLayerAllocationInfo& AllocInfo = Component->WeightmapLayerAllocations[AllocIdx];
-				if( AllocInfo.LayerInfo == ALandscapeProxy::DataLayer )
+				if( AllocInfo.LayerInfo == ALandscapeProxy::VisibilityLayer )
 				{
 					TexIndex = AllocInfo.WeightmapTextureIndex;
 					Component->WeightmapTextures[TexIndex]->Source.GetMipData(RawVisData, 0);
diff --git a/Engine/Source/Editor/UnrealEd/Private/Fbx/FbxMainExport.cpp b/Engine/Source/Editor/UnrealEd/Private/Fbx/FbxMainExport.cpp
--- a/Engine/Source/Editor/UnrealEd/Private/Fbx/FbxMainExport.cpp
+++ b/Engine/Source/Editor/UnrealEd/Private/Fbx/FbxMainExport.cpp
@@ -2393,7 +2393,7 @@ FbxNode* FFbxExporter::ExportLandscapeToFbx(ALandscapeProxy* Landscape, const TC
 		for (int32 AllocIdx = 0; AllocIdx < Component->WeightmapLayerAllocations.Num(); AllocIdx++)
 		{
 			FWeightmapLayerAllocationInfo& AllocInfo = Component->WeightmapLayerAllocations[AllocIdx];
-			if (AllocInfo.LayerInfo == ALandscapeProxy::DataLayer)
+			if (AllocInfo.LayerInfo == ALandscapeProxy::VisibilityLayer)
 			{
 				CDI.GetWeightmapTextureData(AllocInfo.LayerInfo, CompVisData);
 			}
diff --git a/Engine/Source/Editor/WorldBrowser/Private/Tiles/WorldTileCollectionModel.cpp b/Engine/Source/Editor/WorldBrowser/Private/Tiles/WorldTileCollectionModel.cpp
--- a/Engine/Source/Editor/WorldBrowser/Private/Tiles/WorldTileCollectionModel.cpp
+++ b/Engine/Source/Editor/WorldBrowser/Private/Tiles/WorldTileCollectionModel.cpp
@@ -1624,7 +1624,7 @@ void FWorldTileCollectionModel::ReimportTiledLandscape_Executed(FName TargetLaye
 			
 		ALandscapeProxy* Landscape = TileModel->GetLandscape();
 
-		ULandscapeLayerInfoObject* DataLayer = ALandscapeProxy::DataLayer;
+		ULandscapeLayerInfoObject* DataLayer = ALandscapeProxy::VisibilityLayer;
 	
 		if (TargetLayer == HeightmapLayerName) // Heightmap
 		{
diff --git a/Engine/Source/Runtime/Engine/Classes/Landscape/LandscapeProxy.h b/Engine/Source/Runtime/Engine/Classes/Landscape/LandscapeProxy.h
--- a/Engine/Source/Runtime/Engine/Classes/Landscape/LandscapeProxy.h
+++ b/Engine/Source/Runtime/Engine/Classes/Landscape/LandscapeProxy.h
@@ -12,6 +12,8 @@
 // Forward declarations.
 //
 class ULandscapeMaterialInstanceConstant;
+class ULandscapeLayerInfoObject;
+class ULandscapeSplinesComponent;
 
 /** Structure storing channel usage for weightmap textures */
 USTRUCT()
@@ -20,7 +22,7 @@ struct FLandscapeWeightmapUsage
 	GENERATED_USTRUCT_BODY()
 
 	UPROPERTY()
-	class ULandscapeComponent* ChannelUsage[4];
+	ULandscapeComponent* ChannelUsage[4];
 
 	FLandscapeWeightmapUsage()
 	{
@@ -46,7 +48,7 @@ struct FLandscapeEditorLayerSettings
 
 #if WITH_EDITORONLY_DATA
 	UPROPERTY()
-	class ULandscapeLayerInfoObject* LayerInfoObj;
+	ULandscapeLayerInfoObject* LayerInfoObj;
 
 	UPROPERTY()
 	FString ReimportLayerFilePath;
@@ -68,7 +70,7 @@ struct FLandscapeEditorLayerSettings
 		return LayerInfoObj == rhs.LayerInfoObj;
 	}
 
-	bool operator==(const class ULandscapeLayerInfoObject*& rhs)
+	bool operator==(const ULandscapeLayerInfoObject*& rhs)
 	{
 		return LayerInfoObj == rhs;
 	}
@@ -81,14 +83,14 @@ struct FLandscapeLayerStruct
 	GENERATED_USTRUCT_BODY()
 
 	UPROPERTY()
-	class ULandscapeLayerInfoObject* LayerInfoObj;
+	ULandscapeLayerInfoObject* LayerInfoObj;
 
 #if WITH_EDITORONLY_DATA
 	UPROPERTY(transient)
-	class UMaterialInstanceConstant* ThumbnailMIC;
+	UMaterialInstanceConstant* ThumbnailMIC;
 
 	UPROPERTY()
-	class ALandscapeProxy* Owner;
+	ALandscapeProxy* Owner;
 
 	UPROPERTY(transient)
 	int32 DebugColorChannel;
@@ -127,7 +129,7 @@ struct FLandscapeImportLayerInfo
 	ULandscapeLayerInfoObject* LayerInfo;
 
 	UPROPERTY(Category="Import", VisibleAnywhere)
-	class UMaterialInstanceConstant* ThumbnailMIC; // Optional
+	UMaterialInstanceConstant* ThumbnailMIC; // Optional
 
 	UPROPERTY(Category="Import", EditAnywhere)
 	FString SourceFilePath; // Optional
@@ -192,7 +194,7 @@ class ALandscapeProxy : public AActor, public INavRelevantActorInterface
 	GENERATED_UCLASS_BODY()
 
 	UPROPERTY()
-	class ULandscapeSplinesComponent* SplineComponent;
+	ULandscapeSplinesComponent* SplineComponent;
 
 protected:
 	/** Guid for LandscapeEditorInfo **/
@@ -220,7 +222,7 @@ public:
 
 	/** Default physical material, used when no per-layer values physical materials */
 	UPROPERTY(EditAnywhere, Category=Landscape)
-	class UPhysicalMaterial* DefaultPhysMaterial;
+	UPhysicalMaterial* DefaultPhysMaterial;
 
 	/**
 	 * Allows artists to adjust the distance where textures using UV 0 are streamed in/out.
@@ -231,22 +233,22 @@ public:
 
 	/** Combined material used to render the landscape */
 	UPROPERTY(EditAnywhere, Category=Landscape)
-	class UMaterialInterface* LandscapeMaterial;
+	UMaterialInterface* LandscapeMaterial;
 
 	/** Material used to render landscape components with holes. Should be a BLEND_Masked version of LandscapeMaterial. */
 	UPROPERTY(EditAnywhere, Category=Landscape)
-	class UMaterialInterface* LandscapeHoleMaterial;
+	UMaterialInterface* LandscapeHoleMaterial;
 
 	UPROPERTY(EditAnywhere, Category=LOD)
 	float LODDistanceFactor;
 
 	/** The array of LandscapeComponent that are used by the landscape */
 	UPROPERTY()
-	TArray<class ULandscapeComponent*> LandscapeComponents;
+	TArray<ULandscapeComponent*> LandscapeComponents;
 
 	/** Array of LandscapeHeightfieldCollisionComponent */
 	UPROPERTY()
-	TArray<class ULandscapeHeightfieldCollisionComponent*> CollisionComponents;
+	TArray<ULandscapeHeightfieldCollisionComponent*> CollisionComponents;
 
 	/**
 	 *	The resolution to cache lighting at, in texels/quad in one axis
@@ -257,7 +259,7 @@ public:
 	float StaticLightingResolution;
 
 	UPROPERTY(EditAnywhere, Category=LandscapeProxy)
-	TLazyObjectPtr<class ALandscape> LandscapeActor;
+	TLazyObjectPtr<ALandscape> LandscapeActor;
 
 	UPROPERTY(EditAnywhere, Category=Lighting)
 	uint32 bCastStaticShadow:1;
@@ -331,10 +333,10 @@ public:
 #endif
 
 public:
-	ENGINE_API static class ULandscapeLayerInfoObject* DataLayer;
+	ENGINE_API static ULandscapeLayerInfoObject* VisibilityLayer;
 
 	/** Map of material instance constants used to for the components. Key is generated with ULandscapeComponent::GetLayerAllocationKey() */
-	TMap<FString, class UMaterialInstanceConstant*> MaterialInstanceConstantMap;
+	TMap<FString, UMaterialInstanceConstant*> MaterialInstanceConstantMap;
 
 	/** Map of weightmap usage */
 	TMap<UTexture2D*, struct FLandscapeWeightmapUsage> WeightmapUsageMap;
@@ -357,7 +359,7 @@ public:
 #endif	//WITH_EDITOR
 
 	FGuid GetLandscapeGuid() const { return LandscapeGuid; }
-	virtual class ALandscape* GetLandscapeActor();
+	virtual ALandscape* GetLandscapeActor();
 
 	// Begin INavRelevantActorInterface Interface
 	virtual bool DoesSupplyPerComponentNavigationCollision() const override{ return true; }
@@ -382,7 +384,7 @@ public:
 	ENGINE_API static ULandscapeLayerInfoObject* CreateLayerInfo(const TCHAR* LayerName, ULevel* Level);
 	ENGINE_API ULandscapeLayerInfoObject* CreateLayerInfo(const TCHAR* LayerName);
 
-	ENGINE_API class ULandscapeInfo* GetLandscapeInfo(bool bSpawnNewActor = true);
+	ENGINE_API ULandscapeInfo* GetLandscapeInfo(bool bSpawnNewActor = true);
 
 	// Get Landscape Material assigned to this Landscape
 	virtual UMaterialInterface* GetLandscapeMaterial() const;
@@ -397,10 +399,10 @@ public:
 	ENGINE_API void ChangedPhysMaterial();
 
 	// Check input Landscape actor is match for this LandscapeProxy (by GUID)
-	ENGINE_API bool IsValidLandscapeActor(class ALandscape* Landscape);
+	ENGINE_API bool IsValidLandscapeActor(ALandscape* Landscape);
 
 	// Copy properties from parent Landscape actor
-	ENGINE_API void GetSharedProperties(class ALandscapeProxy* Landscape);
+	ENGINE_API void GetSharedProperties(ALandscapeProxy* Landscape);
 
 	/** Get the LandcapeActor-to-world transform with respect to landscape section offset*/
 	ENGINE_API FTransform LandscapeActorToWorld() const;
@@ -438,7 +440,7 @@ public:
 	ENGINE_API FIntRect GetBoundingRect() const;
 
 	/** Creates a Texture2D for use by this landscape proxy or one of it's components. If OptionalOverrideOuter is not specified, the level is used. */
-	ENGINE_API class UTexture2D* CreateLandscapeTexture(int32 InSizeX, int32 InSizeY, TextureGroup InLODGroup, ETextureSourceFormat InFormat, UObject* OptionalOverrideOuter = nullptr) const;
+	ENGINE_API UTexture2D* CreateLandscapeTexture(int32 InSizeX, int32 InSizeY, TextureGroup InLODGroup, ETextureSourceFormat InFormat, UObject* OptionalOverrideOuter = nullptr) const;
 #endif
 };
 
diff --git a/Engine/Source/Runtime/Engine/Private/Landscape/Landscape.cpp b/Engine/Source/Runtime/Engine/Private/Landscape/Landscape.cpp
--- a/Engine/Source/Runtime/Engine/Private/Landscape/Landscape.cpp
+++ b/Engine/Source/Runtime/Engine/Private/Landscape/Landscape.cpp
@@ -38,8 +38,7 @@ static void PrintNumLandscapeShadows()
 			NumShadowCasters++;
 		}
 	}
-	UE_LOG(LogConsoleResponse,Display,TEXT("%d/%d landscape components cast shadows"),
-		NumShadowCasters, NumComponents);
+	UE_LOG(LogConsoleResponse, Display, TEXT("%d/%d landscape components cast shadows"), NumShadowCasters, NumComponents);
 }
 
 FAutoConsoleCommand CmdPrintNumLandscapeShadows(
@@ -83,7 +82,7 @@ ULandscapeComponent::ULandscapeComponent(const class FPostConstructInitializePro
 	NeighborLODBias[5] = 128;
 	NeighborLODBias[6] = 128;
 	NeighborLODBias[7] = 128;
-	
+
 	Mobility = EComponentMobility::Static;
 
 	EditToolRenderData = NULL;
@@ -100,7 +99,7 @@ ULandscapeComponent::~ULandscapeComponent()
 void ULandscapeComponent::AddReferencedObjects(UObject* InThis, FReferenceCollector& Collector)
 {
 	ULandscapeComponent* This = CastChecked<ULandscapeComponent>(InThis);
-	if(This->LightMap != NULL)
+	if (This->LightMap != NULL)
 	{
 		This->LightMap->AddReferencedObjects(Collector);
 	}
@@ -117,7 +116,7 @@ void ULandscapeComponent::Serialize(FArchive& Ar)
 	// Saving for cooking path
 	if (Ar.IsCooking() && Ar.IsSaving() && !HasAnyFlags(RF_ClassDefaultObject))
 	{
-		if (!Ar.CookingTarget()->SupportsFeature(ETargetPlatformFeatures::VertexShaderTextureSampling) )
+		if (!Ar.CookingTarget()->SupportsFeature(ETargetPlatformFeatures::VertexShaderTextureSampling))
 		{
 			if (!PlatformData.HasValidPlatformData())
 			{
@@ -130,7 +129,7 @@ void ULandscapeComponent::Serialize(FArchive& Ar)
 #endif
 
 	Super::Serialize(Ar);
-	
+
 	Ar << LightMap;
 	if (Ar.UE4Ver() >= VER_UE4_PRECOMPUTED_SHADOW_MAPS_BSP)
 	{
@@ -138,7 +137,7 @@ void ULandscapeComponent::Serialize(FArchive& Ar)
 	}
 
 #if WITH_EDITOR
-	if ( Ar.IsTransacting() )
+	if (Ar.IsTransacting())
 	{
 		if (EditToolRenderData)
 		{
@@ -232,9 +231,9 @@ UMaterialInterface* ULandscapeComponent::GetLandscapeHoleMaterial() const
 
 bool ULandscapeComponent::ComponentHasVisibilityPainted() const
 {
-	for( int32 LayerIdx=0;LayerIdx < WeightmapLayerAllocations.Num();LayerIdx++ )
+	for (const FWeightmapLayerAllocationInfo& Allocation : WeightmapLayerAllocations)
 	{
-		if( WeightmapLayerAllocations[LayerIdx].LayerInfo == ALandscapeProxy::DataLayer )
+		if (Allocation.LayerInfo == ALandscapeProxy::VisibilityLayer)
 		{
 			return true;
 		}
@@ -255,16 +254,16 @@ FString ULandscapeComponent::GetLayerAllocationKey(bool bMobile /* = false */) c
 
 	// Sort the allocations
 	TArray<FString> LayerStrings;
-	for( int32 LayerIdx=0;LayerIdx < WeightmapLayerAllocations.Num();LayerIdx++ )
+	for (int32 LayerIdx = 0; LayerIdx < WeightmapLayerAllocations.Num(); LayerIdx++)
 	{
-		new(LayerStrings) FString( *FString::Printf(TEXT("_%s_%d"), *WeightmapLayerAllocations[LayerIdx].GetLayerName().ToString(), bMobile ? 0 : WeightmapLayerAllocations[LayerIdx].WeightmapTextureIndex) );
+		new(LayerStrings)FString(*FString::Printf(TEXT("_%s_%d"), *WeightmapLayerAllocations[LayerIdx].GetLayerName().ToString(), bMobile ? 0 : WeightmapLayerAllocations[LayerIdx].WeightmapTextureIndex));
 	}
-/**
- * Generate a key for this component's layer allocations to use with MaterialInstanceConstantMap.
- */
-	LayerStrings.Sort( TGreater<FString>() );
+	/**
+	 * Generate a key for this component's layer allocations to use with MaterialInstanceConstantMap.
+	 */
+	LayerStrings.Sort(TGreater<FString>());
 
-	for( int32 LayerIdx=0;LayerIdx < LayerStrings.Num();LayerIdx++ )
+	for (int32 LayerIdx = 0; LayerIdx < LayerStrings.Num(); LayerIdx++)
 	{
 		Result += LayerStrings[LayerIdx];
 	}
@@ -290,15 +289,15 @@ void ULandscapeComponent::GetLayerDebugColorKey(int32& R, int32& G, int32& B) co
 					{
 						if (LayerStruct.DebugColorChannel & 1) // R
 						{
-							R = (WeightmapLayerAllocations[LayerIdx].WeightmapTextureIndex*4+WeightmapLayerAllocations[LayerIdx].WeightmapTextureChannel);
+							R = (WeightmapLayerAllocations[LayerIdx].WeightmapTextureIndex * 4 + WeightmapLayerAllocations[LayerIdx].WeightmapTextureChannel);
 						}
 						if (LayerStruct.DebugColorChannel & 2) // G
 						{
-							G = (WeightmapLayerAllocations[LayerIdx].WeightmapTextureIndex*4+WeightmapLayerAllocations[LayerIdx].WeightmapTextureChannel);
+							G = (WeightmapLayerAllocations[LayerIdx].WeightmapTextureIndex * 4 + WeightmapLayerAllocations[LayerIdx].WeightmapTextureChannel);
 						}
 						if (LayerStruct.DebugColorChannel & 4) // B
 						{
-							B = (WeightmapLayerAllocations[LayerIdx].WeightmapTextureIndex*4+WeightmapLayerAllocations[LayerIdx].WeightmapTextureChannel);
+							B = (WeightmapLayerAllocations[LayerIdx].WeightmapTextureIndex * 4 + WeightmapLayerAllocations[LayerIdx].WeightmapTextureChannel);
 						}
 						break;
 					}
@@ -329,8 +328,8 @@ void ULandscapeInfo::UpdateDebugColorMaterial()
 {
 	FlushRenderingCommands();
 	//GWarn->BeginSlowTask( *FString::Printf(TEXT("Compiling layer color combinations for %s"), *GetName()), true);
-	
-	for (auto It = XYtoComponentMap.CreateIterator(); It; ++It )
+
+	for (auto It = XYtoComponentMap.CreateIterator(); It; ++It)
 	{
 		ULandscapeComponent* Comp = It.Value();
 		if (Comp && Comp->EditToolRenderData)
@@ -361,7 +360,7 @@ void ULandscapeComponent::PostLoad()
 			{
 				if (WeightmapLayerAllocations[i].LayerName_DEPRECATED == DataWeightmapName)
 				{
-					WeightmapLayerAllocations[i].LayerInfo = ALandscapeProxy::DataLayer;
+					WeightmapLayerAllocations[i].LayerInfo = ALandscapeProxy::VisibilityLayer;
 				}
 				else
 				{
@@ -374,17 +373,17 @@ void ULandscapeComponent::PostLoad()
 	}
 
 #if WITH_EDITOR
-	if( GIsEditor && !HasAnyFlags(RF_ClassDefaultObject) )
+	if (GIsEditor && !HasAnyFlags(RF_ClassDefaultObject))
 	{
 		// Remove standalone flags from data textures to ensure data is unloaded in the editor when reverting an unsaved level.
 		// Previous version of landscape set these flags on creation.
-		if( HeightmapTexture && HeightmapTexture->HasAnyFlags(RF_Standalone) )
+		if (HeightmapTexture && HeightmapTexture->HasAnyFlags(RF_Standalone))
 		{
 			HeightmapTexture->ClearFlags(RF_Standalone);
 		}
-		for( int32 Idx=0;Idx<WeightmapTextures.Num();Idx++ )
+		for (int32 Idx = 0; Idx < WeightmapTextures.Num(); Idx++)
 		{
-			if( WeightmapTextures[Idx] && WeightmapTextures[Idx]->HasAnyFlags(RF_Standalone) )
+			if (WeightmapTextures[Idx] && WeightmapTextures[Idx]->HasAnyFlags(RF_Standalone))
 			{
 				WeightmapTextures[Idx]->ClearFlags(RF_Standalone);
 			}
@@ -474,13 +473,13 @@ ALandscapeProxy::ALandscapeProxy(const class FPostConstructInitializeProperties&
 	bHidden = false;
 	bReplicateMovement = false;
 	bCanBeDamaged = false;
-		
+
 	TSubobjectPtr<USceneComponent> SceneComponent = PCIP.CreateDefaultSubobject<USceneComponent>(this, TEXT("RootComponent0"));
 	RootComponent = SceneComponent;
 	RootComponent->RelativeScale3D = FVector(128.0f, 128.0f, 256.0f); // Old default scale, preserved for compatibility. See ULandscapeEditorObject::NewLandscape_Scale
 	RootComponent->Mobility = EComponentMobility::Static;
 	LandscapeSectionOffset = FIntPoint::ZeroValue;
-	
+
 	StaticLightingResolution = 1.0f;
 	StreamingDistanceMultiplier = 1.0f;
 	bIsProxy = true;
@@ -500,7 +499,7 @@ ALandscapeProxy::ALandscapeProxy(const class FPostConstructInitializeProperties&
 	MaxPaintedLayersPerComponent = 0;
 #endif
 
-	if (DataLayer == NULL)
+	if (VisibilityLayer == NULL)
 	{
 		// Structure to hold one-time initialization
 		struct FConstructorStatics
@@ -513,9 +512,9 @@ ALandscapeProxy::ALandscapeProxy(const class FPostConstructInitializeProperties&
 		};
 		static FConstructorStatics ConstructorStatics;
 
-		DataLayer = ConstructorStatics.DataLayer.Get();
-		check(DataLayer);
-		DataLayer->AddToRoot();
+		VisibilityLayer = ConstructorStatics.DataLayer.Get();
+		check(VisibilityLayer);
+		VisibilityLayer->AddToRoot();
 	}
 }
 
@@ -637,23 +636,23 @@ void ULandscapeComponent::BeginDestroy()
 		EditToolRenderData = NULL;
 	}
 
-	if( GIsEditor && !HasAnyFlags(RF_ClassDefaultObject) )
+	if (GIsEditor && !HasAnyFlags(RF_ClassDefaultObject))
 	{
 		ALandscapeProxy* Proxy = GetLandscapeProxy();
 
 		// Remove any weightmap allocations from the Landscape Actor's map
-		for( int32 LayerIdx=0;LayerIdx < WeightmapLayerAllocations.Num();LayerIdx++ )
+		for (int32 LayerIdx = 0; LayerIdx < WeightmapLayerAllocations.Num(); LayerIdx++)
 		{
 			int32 WeightmapIndex = WeightmapLayerAllocations[LayerIdx].WeightmapTextureIndex;
-			if( WeightmapTextures.IsValidIndex(WeightmapIndex) )
+			if (WeightmapTextures.IsValidIndex(WeightmapIndex))
 			{
 				UTexture2D* WeightmapTexture = WeightmapTextures[WeightmapIndex];
 				FLandscapeWeightmapUsage* Usage = Proxy->WeightmapUsageMap.Find(WeightmapTexture);
-				if( Usage != NULL )
+				if (Usage != NULL)
 				{
 					Usage->ChannelUsage[WeightmapLayerAllocations[LayerIdx].WeightmapTextureChannel] = NULL;
 
-					if( Usage->FreeChannelCount()==4 )
+					if (Usage->FreeChannelCount() == 4)
 					{
 						Proxy->WeightmapUsageMap.Remove(WeightmapTexture);
 					}
@@ -671,7 +670,7 @@ FPrimitiveSceneProxy* ULandscapeComponent::CreateSceneProxy()
 	if (FeatureLevel >= ERHIFeatureLevel::SM3)
 	{
 #if WITH_EDITOR
-		if( EditToolRenderData == NULL )
+		if (EditToolRenderData == NULL)
 		{
 			EditToolRenderData = new FLandscapeEditToolRenderData(this);
 		}
@@ -683,7 +682,7 @@ FPrimitiveSceneProxy* ULandscapeComponent::CreateSceneProxy()
 	else if (FeatureLevel >= ERHIFeatureLevel::ES2)
 	{
 #if WITH_EDITOR
-		if( !PlatformData.HasValidPlatformData() ) // Deferred generation
+		if (!PlatformData.HasValidPlatformData()) // Deferred generation
 		{
 			// Try to reload the ES2 landscape data from the DDC
 			if (!PlatformData.LoadFromDDC(StateId))
@@ -700,7 +699,7 @@ FPrimitiveSceneProxy* ULandscapeComponent::CreateSceneProxy()
 
 		if (PlatformData.HasValidPlatformData())
 		{
-			if( EditToolRenderData == NULL )
+			if (EditToolRenderData == NULL)
 			{
 				EditToolRenderData = new FLandscapeEditToolRenderData(this);
 			}
@@ -759,9 +758,9 @@ void ULandscapeComponent::OnUnregister()
 #endif
 }
 
-void ULandscapeComponent::GetUsedMaterials( TArray<UMaterialInterface*>& OutMaterials ) const
+void ULandscapeComponent::GetUsedMaterials(TArray<UMaterialInterface*>& OutMaterials) const
 {
-	if( MaterialInstance != NULL )
+	if (MaterialInstance != NULL)
 	{
 		OutMaterials.Add(MaterialInstance);
 	}
@@ -773,7 +772,7 @@ void ALandscapeProxy::RegisterAllComponents()
 
 	// Landscape was added to world 
 	// We might need to update shared landscape data
-	if (GIsEditor && GetWorld() && !GetWorld()->IsPlayInEditor() )
+	if (GIsEditor && GetWorld() && !GetWorld()->IsPlayInEditor())
 	{
 		GEngine->DeferredCommands.AddUnique(TEXT("UpdateLandscapeEditorData"));
 	}
@@ -791,14 +790,14 @@ void ALandscapeProxy::UnregisterAllComponents()
 	}
 }
 
-bool ALandscapeProxy::UpdateNavigationRelevancy() 
+bool ALandscapeProxy::UpdateNavigationRelevancy()
 {
 	SetNavigationRelevancy(bUsedForNavigation);
-	
-	return bUsedForNavigation; 
+
+	return bUsedForNavigation;
 }
 
-FBox ALandscapeProxy::GetComponentsBoundingBox(bool bNonColliding) const 
+FBox ALandscapeProxy::GetComponentsBoundingBox(bool bNonColliding) const
 {
 	FBox Bounds = Super::GetComponentsBoundingBox(bNonColliding);
 
@@ -807,7 +806,7 @@ FBox ALandscapeProxy::GetComponentsBoundingBox(bool bNonColliding) const
 	// can't really tell if it breaks anything, but is needed for navmesh generation
 	if (bUsedForNavigation)
 	{
-		for (auto CollisionComponentIt( CollisionComponents.CreateConstIterator() ); CollisionComponentIt; ++CollisionComponentIt)
+		for (auto CollisionComponentIt(CollisionComponents.CreateConstIterator()); CollisionComponentIt; ++CollisionComponentIt)
 		{
 			const ULandscapeHeightfieldCollisionComponent* Component = *CollisionComponentIt;
 			if (Component != NULL)
@@ -824,12 +823,12 @@ FBox ALandscapeProxy::GetComponentsBoundingBox(bool bNonColliding) const
 }
 
 // FLandscapeWeightmapUsage serializer
-FArchive& operator<<( FArchive& Ar, FLandscapeWeightmapUsage& U )
+FArchive& operator<<(FArchive& Ar, FLandscapeWeightmapUsage& U)
 {
 	return Ar << U.ChannelUsage[0] << U.ChannelUsage[1] << U.ChannelUsage[2] << U.ChannelUsage[3];
 }
 
-FArchive& operator<<( FArchive& Ar, FLandscapeAddCollision& U )
+FArchive& operator<<(FArchive& Ar, FLandscapeAddCollision& U)
 {
 #if WITH_EDITORONLY_DATA
 	return Ar << U.Corners[0] << U.Corners[1] << U.Corners[2] << U.Corners[3];
@@ -838,7 +837,7 @@ FArchive& operator<<( FArchive& Ar, FLandscapeAddCollision& U )
 #endif // WITH_EDITORONLY_DATA
 }
 
-FArchive& operator<<( FArchive& Ar, FLandscapeLayerStruct*& L )
+FArchive& operator<<(FArchive& Ar, FLandscapeLayerStruct*& L)
 {
 	if (L)
 	{
@@ -883,7 +882,7 @@ void ULandscapeInfo::BeginDestroy()
 	Super::BeginDestroy();
 
 #if WITH_EDITOR
-	if( DataInterface )
+	if (DataInterface)
 	{
 		delete DataInterface;
 		DataInterface = NULL;
@@ -904,7 +903,7 @@ void ALandscapeProxy::Serialize(FArchive& Ar)
 	Super::Serialize(Ar);
 
 #if WITH_EDITOR
-	if( Ar.IsTransacting() )
+	if (Ar.IsTransacting())
 	{
 		Ar << WeightmapUsageMap;
 	}
@@ -1070,7 +1069,7 @@ bool ULandscapeInfo::UpdateLayerInfoMap(ALandscapeProxy* Proxy /*= NULL*/, bool
 	if (GIsEditor)
 	{
 		bIsValid = !bInvalidate;
-	
+
 		if (Proxy)
 		{
 			if (bInvalidate)
@@ -1145,7 +1144,7 @@ bool ULandscapeInfo::UpdateLayerInfoMap(ALandscapeProxy* Proxy /*= NULL*/, bool
 
 								// Valid layer infos take precedence over invalid ones
 								// Landscape Actors take precedence over Proxies
-								if ( (bValid && !LayerSettings.bValid)
+								if ((bValid && !LayerSettings.bValid)
 									|| (bValid == LayerSettings.bValid && !Proxy->bIsProxy))
 								{
 									LayerSettings.Owner = Proxy;
@@ -1171,9 +1170,9 @@ bool ULandscapeInfo::UpdateLayerInfoMap(ALandscapeProxy* Proxy /*= NULL*/, bool
 								}
 								else
 								{
-										FLandscapeInfoLayerSettings LayerSettings(LayerInfo, Proxy);
-										LayerSettings.bValid = bValid;
-										Layers.Add(LayerSettings);
+									FLandscapeInfoLayerSettings LayerSettings(LayerInfo, Proxy);
+									LayerSettings.bValid = bValid;
+									Layers.Add(LayerSettings);
 								}
 							}
 						}
@@ -1243,10 +1242,10 @@ void ALandscapeProxy::PostLoad()
 	Super::PostLoad();
 
 	// Temporary
-	if( ComponentSizeQuads == 0 && LandscapeComponents.Num() > 0 )
+	if (ComponentSizeQuads == 0 && LandscapeComponents.Num() > 0)
 	{
 		ULandscapeComponent* Comp = LandscapeComponents[0];
-		if( Comp )
+		if (Comp)
 		{
 			ComponentSizeQuads = Comp->ComponentSizeQuads;
 			SubsectionSizeQuads = Comp->SubsectionSizeQuads;
@@ -1260,7 +1259,7 @@ void ALandscapeProxy::PostLoad()
 	}
 
 #if WITH_EDITOR
-	if ( GetLinker() && (GetLinker()->UE4Ver() < VER_UE4_LANDSCAPE_COMPONENT_LAZY_REFERENCES) || LandscapeComponents.Num() != CollisionComponents.Num() )
+	if (GetLinker() && (GetLinker()->UE4Ver() < VER_UE4_LANDSCAPE_COMPONENT_LAZY_REFERENCES) || LandscapeComponents.Num() != CollisionComponents.Num())
 	{
 		// Need to clean up invalid collision components
 		RecreateCollisionComponents();
@@ -1362,7 +1361,7 @@ void ALandscapeProxy::SetAbsoluteSectionBase(FIntPoint InSectionBase)
 {
 	FIntPoint Difference = InSectionBase - LandscapeSectionOffset;
 	LandscapeSectionOffset = InSectionBase;
-	
+
 	for (int32 CompIdx = 0; CompIdx < LandscapeComponents.Num(); CompIdx++)
 	{
 		ULandscapeComponent* Comp = LandscapeComponents[CompIdx];
@@ -1372,7 +1371,7 @@ void ALandscapeProxy::SetAbsoluteSectionBase(FIntPoint InSectionBase)
 			Comp->SetSectionBase(AbsoluteSectionBase);
 		}
 	}
-	
+
 	for (int32 CompIdx = 0; CompIdx < CollisionComponents.Num(); CompIdx++)
 	{
 		ULandscapeHeightfieldCollisionComponent* Comp = CollisionComponents[CompIdx];
@@ -1391,10 +1390,10 @@ FIntPoint ALandscapeProxy::GetSectionBaseOffset() const
 
 void ALandscapeProxy::RecreateComponentsState()
 {
-	for(int32 ComponentIndex = 0; ComponentIndex < LandscapeComponents.Num(); ComponentIndex++ )
+	for (int32 ComponentIndex = 0; ComponentIndex < LandscapeComponents.Num(); ComponentIndex++)
 	{
 		ULandscapeComponent* Comp = LandscapeComponents[ComponentIndex];
-		if( Comp )
+		if (Comp)
 		{
 			Comp->UpdateComponentToWorld();
 			Comp->UpdateCachedBounds();
@@ -1403,10 +1402,10 @@ void ALandscapeProxy::RecreateComponentsState()
 		}
 	}
 
-	for(int32 ComponentIndex = 0; ComponentIndex < CollisionComponents.Num(); ComponentIndex++ )
+	for (int32 ComponentIndex = 0; ComponentIndex < CollisionComponents.Num(); ComponentIndex++)
 	{
 		ULandscapeHeightfieldCollisionComponent* Comp = CollisionComponents[ComponentIndex];
-		if( Comp )
+		if (Comp)
 		{
 			Comp->UpdateComponentToWorld();
 			Comp->RecreateCollision();
@@ -1479,25 +1478,25 @@ void ALandscape::PreSave()
 ALandscapeProxy* ULandscapeInfo::GetCurrentLevelLandscapeProxy(bool bRegistered) const
 {
 	ALandscape* Landscape = LandscapeActor.Get();
-	
+
 	if (Landscape && (!bRegistered || Landscape->GetRootComponent()->IsRegistered()))
 	{
 		UWorld* LandscapeWorld = Landscape->GetWorld();
-		if (LandscapeWorld && 
+		if (LandscapeWorld &&
 			LandscapeWorld->GetCurrentLevel() == Landscape->GetOuter())
 		{
 			return Landscape;
 		}
 	}
-		
+
 	for (auto It = Proxies.CreateConstIterator(); It; ++It)
 	{
 		ALandscapeProxy* LandscapeProxy = (*It);
-		
+
 		if (LandscapeProxy && (!bRegistered || LandscapeProxy->GetRootComponent()->IsRegistered()))
 		{
 			UWorld* ProxyWorld = LandscapeProxy->GetWorld();
-			if (ProxyWorld && 
+			if (ProxyWorld &&
 				ProxyWorld->GetCurrentLevel() == LandscapeProxy->GetOuter())
 			{
 				return LandscapeProxy;
@@ -1513,7 +1512,7 @@ ALandscapeProxy* ULandscapeInfo::GetLandscapeProxy() const
 	// Mostly this Proxy used to calculate transformations
 	// in Editor all proxies of same landscape actor have root components in same locations
 	// so it doesn't really matter which proxy we return here
-	
+
 	// prefer LandscapeActor in case it is loaded
 	ALandscape* Landscape = LandscapeActor.Get();
 	if (Landscape != NULL &&
@@ -1521,7 +1520,7 @@ ALandscapeProxy* ULandscapeInfo::GetLandscapeProxy() const
 	{
 		return Landscape;
 	}
-	
+
 	// prefer current level proxy 
 	ALandscapeProxy* Proxy = GetCurrentLevelLandscapeProxy(true);
 	if (Proxy != NULL)
@@ -1548,19 +1547,19 @@ void ULandscapeInfo::RegisterActor(ALandscapeProxy* Proxy, bool bMapCheck)
 	// do not pass here invalid actors
 	check(Proxy->GetLandscapeGuid().IsValid());
 	UWorld* OwningWorld = CastChecked<UWorld>(GetOuter());
-	
+
 	// in case this Info object is not initialized yet
 	// initialized it with properties from passed actor
 	if (LandscapeGuid.IsValid() == false ||
 		(GetLandscapeProxy() == NULL && ensure(LandscapeGuid == Proxy->GetLandscapeGuid())))
 	{
-		LandscapeGuid			= Proxy->GetLandscapeGuid();
-		ComponentSizeQuads		= Proxy->ComponentSizeQuads;
-		ComponentNumSubsections	= Proxy->NumSubsections;
-		SubsectionSizeQuads		= Proxy->SubsectionSizeQuads;
-		DrawScale				= Proxy->GetRootComponent()->RelativeScale3D;
+		LandscapeGuid = Proxy->GetLandscapeGuid();
+		ComponentSizeQuads = Proxy->ComponentSizeQuads;
+		ComponentNumSubsections = Proxy->NumSubsections;
+		SubsectionSizeQuads = Proxy->SubsectionSizeQuads;
+		DrawScale = Proxy->GetRootComponent()->RelativeScale3D;
 	}
-	
+
 	// check that passed actor matches all shared parameters
 	check(LandscapeGuid == Proxy->GetLandscapeGuid());
 	check(ComponentSizeQuads == Proxy->ComponentSizeQuads);
@@ -1575,7 +1574,7 @@ void ULandscapeInfo::RegisterActor(ALandscapeProxy* Proxy, bool bMapCheck)
 		LandscapeActor = Landscape;
 		// In world composition user is not allowed to move landscape in editor, only through WorldBrowser 
 		LandscapeActor->bLockLocation = (OwningWorld->WorldComposition != NULL);
-		
+
 		// update proxies reference actor
 		for (auto It = Proxies.CreateConstIterator(); It; ++It)
 		{
@@ -1596,7 +1595,7 @@ void ULandscapeInfo::RegisterActor(ALandscapeProxy* Proxy, bool bMapCheck)
 		RegisterActorComponent(Proxy->LandscapeComponents[CompIdx], bMapCheck);
 	}
 }
-	
+
 void ULandscapeInfo::UnregisterActor(ALandscapeProxy* Proxy)
 {
 	ALandscape* Landscape = Cast<ALandscape>(Proxy);
@@ -1622,7 +1621,7 @@ void ULandscapeInfo::UnregisterActor(ALandscapeProxy* Proxy)
 	{
 		UnregisterActorComponent(Proxy->LandscapeComponents[CompIdx]);
 	}
-			
+
 	XYtoComponentMap.Compact();
 }
 
@@ -1636,7 +1635,7 @@ void ULandscapeInfo::RegisterActorComponent(ULandscapeComponent* Component, bool
 	}
 
 	check(Component != NULL);
-	FIntPoint ComponentKey = Component->GetSectionBase()/Component->ComponentSizeQuads;
+	FIntPoint ComponentKey = Component->GetSectionBase() / Component->ComponentSizeQuads;
 	auto RegisteredComponent = XYtoComponentMap.FindRef(ComponentKey);
 
 	if (RegisteredComponent == NULL)
@@ -1652,11 +1651,11 @@ void ULandscapeInfo::RegisterActorComponent(ULandscapeComponent* Component, bool
 		Arguments.Add(TEXT("YLocation"), Component->GetSectionBase().Y);
 		FMessageLog("MapCheck").Warning()
 			->AddToken(FUObjectToken::Create(Proxy))
-			->AddToken(FTextToken::Create(FText::Format( LOCTEXT( "MapCheck_Message_LandscapeComponentPostLoad_Warning", "Landscape ({ProxyName}) has overlapping render components at location ({XLocation}, {YLocation})." ), Arguments )))
+			->AddToken(FTextToken::Create(FText::Format(LOCTEXT("MapCheck_Message_LandscapeComponentPostLoad_Warning", "Landscape ({ProxyName}) has overlapping render components at location ({XLocation}, {YLocation})."), Arguments)))
 			->AddToken(FMapErrorToken::Create(FMapErrors::LandscapeComponentPostLoad_Warning));
 
 		// Show MapCheck window
-		FMessageLog("MapCheck").Open( EMessageSeverity::Warning );
+		FMessageLog("MapCheck").Open(EMessageSeverity::Warning);
 	}
 
 	// Update Selected Components/Regions
@@ -1674,12 +1673,12 @@ void ULandscapeInfo::RegisterActorComponent(ULandscapeComponent* Component, bool
 	}
 #endif
 }
-	
+
 void ULandscapeInfo::UnregisterActorComponent(ULandscapeComponent* Component)
 {
 	check(Component != NULL);
 
-	FIntPoint ComponentKey = Component->GetSectionBase()/Component->ComponentSizeQuads;
+	FIntPoint ComponentKey = Component->GetSectionBase() / Component->ComponentSizeQuads;
 	auto RegisteredComponent = XYtoComponentMap.FindRef(ComponentKey);
 
 	if (RegisteredComponent == Component)
@@ -1694,11 +1693,11 @@ void ULandscapeInfo::UnregisterActorComponent(ULandscapeComponent* Component)
 void ULandscapeInfo::Reset()
 {
 	LandscapeActor.Reset();
-	
+
 	Proxies.Empty();
 	XYtoComponentMap.Empty();
 	XYtoAddCollisionMap.Empty();
-	
+
 	//SelectedComponents.Empty();
 	//SelectedRegionComponents.Empty();
 	//SelectedRegion.Empty();
@@ -1711,15 +1710,15 @@ void ULandscapeInfo::FixupProxiesTransform()
 	{
 		return;
 	}
-		
+
 	ALandscape* Landscape = LandscapeActor.Get();
 
-	if (Landscape == NULL || 
+	if (Landscape == NULL ||
 		Landscape->GetRootComponent()->IsRegistered() == false)
 	{
 		return;
 	}
-	
+
 	FTransform LandscapeTM = Landscape->LandscapeActorToWorld();
 	// Update transformations of all linked landscape proxies
 	for (auto It = Proxies.CreateConstIterator(); It; ++It)
@@ -1727,14 +1726,14 @@ void ULandscapeInfo::FixupProxiesTransform()
 		ALandscapeProxy* Proxy = *It;
 		FTransform ProxyRelativeTM(FVector(Proxy->LandscapeSectionOffset));
 		FTransform ProxyTransform = ProxyRelativeTM*LandscapeTM;
-		
+
 		if (!Proxy->GetTransform().Equals(ProxyTransform))
 		{
 			Proxy->SetActorTransform(ProxyTransform);
 			Proxy->RecreateComponentsState();
-			
+
 			// Let other systems know that an actor was moved
-			GEngine->BroadcastOnActorMoved( Proxy );
+			GEngine->BroadcastOnActorMoved(Proxy);
 		}
 	}
 }
@@ -1753,7 +1752,7 @@ void ULandscapeInfo::FixupProxiesWeightmaps()
 			}
 		}
 	}
-		
+
 	for (auto It = Proxies.CreateConstIterator(); It; ++It)
 	{
 		ALandscapeProxy* Proxy = (*It);
@@ -1785,40 +1784,40 @@ void ULandscapeInfo::RecreateLandscapeInfo(UWorld* InWorld, bool bMapCheck)
 	for (FActorIterator It(InWorld); It; ++It)
 	{
 		ALandscapeProxy* Proxy = Cast<ALandscapeProxy>(*It);
-		if (Proxy 
-			&& Proxy->HasAnyFlags(RF_BeginDestroyed) == false 
+		if (Proxy
+			&& Proxy->HasAnyFlags(RF_BeginDestroyed) == false
 			&& Proxy->IsPendingKill() == false
 			&& Proxy->IsPendingKillPending() == false)
 		{
 			ValidLandscapesMap.FindOrAdd(Proxy->GetLandscapeGuid()).Add(Proxy);
 		}
 	}
-	
+
 	// Register gathered landscapes in shared LandscapeInfo map
 	for (auto It = ValidLandscapesMap.CreateIterator(); It; ++It)
 	{
 		auto& LandscapeList = It.Value();
 		ALandscapeProxy* FirstVisibleLandscape = NULL;
-			
+
 		for (int32 LandscapeIdx = 0; LandscapeIdx < LandscapeList.Num(); LandscapeIdx++)
 		{
 			ALandscapeProxy* Proxy = LandscapeList[LandscapeIdx];
-			if (InWorld->WorldComposition && 
+			if (InWorld->WorldComposition &&
 				Proxy->GetRootComponent()->IsRegistered())
 			{
 				if (FirstVisibleLandscape == NULL)
 				{
 					FirstVisibleLandscape = Proxy;
 				}
-				
+
 				// Setup first landscape in the list as origin for section based grid
 				// And based on that calculate section coordinates for all landscapes in the list
 				FVector Offset = Proxy->GetActorLocation() - FirstVisibleLandscape->GetActorLocation();
 				FVector	DrawScale = FirstVisibleLandscape->GetRootComponent()->RelativeScale3D;
-				
-				FIntPoint QuadsSpaceOffset; 
-				QuadsSpaceOffset.X = FMath::RoundToInt(Offset.X/DrawScale.X);
-				QuadsSpaceOffset.Y = FMath::RoundToInt(Offset.Y/DrawScale.Y);
+
+				FIntPoint QuadsSpaceOffset;
+				QuadsSpaceOffset.X = FMath::RoundToInt(Offset.X / DrawScale.X);
+				QuadsSpaceOffset.Y = FMath::RoundToInt(Offset.Y / DrawScale.Y);
 				Proxy->SetAbsoluteSectionBase(QuadsSpaceOffset);
 			}
 
@@ -1902,10 +1901,10 @@ void FLandscapeComponentDerivedData::InitializeFromUncompressedData(const TArray
 	int32 CompressedSize = TempCompressedMemory.Num() * TempCompressedMemory.GetTypeSize();
 
 	verify(FCompression::CompressMemory(
-		(ECompressionFlags)(COMPRESS_ZLIB | COMPRESS_BiasMemory), 
-		TempCompressedMemory.GetData(), 
-		CompressedSize, 
-		UncompressedData.GetData(), 
+		(ECompressionFlags)(COMPRESS_ZLIB | COMPRESS_BiasMemory),
+		TempCompressedMemory.GetData(),
+		CompressedSize,
+		UncompressedData.GetData(),
 		UncompressedSize));
 
 	// Note: change LANDSCAPE_FULL_DERIVEDDATA_VER when modifying the serialization layout
