Remove FLandscapeToolSet

diff --git a/Engine/Source/Runtime/Engine/Private/Landscape/LandscapeCollision.cpp b/Engine/Source/Runtime/Engine/Private/Landscape/LandscapeCollision.cpp
--- a/Engine/Source/Runtime/Engine/Private/Landscape/LandscapeCollision.cpp
+++ b/Engine/Source/Runtime/Engine/Private/Landscape/LandscapeCollision.cpp
@@ -22,7 +22,7 @@ ULandscapeHeightfieldCollisionComponent::FPhysXHeightfieldRef::~FPhysXHeightfiel
 {
 #if WITH_PHYSX
 	// Free the existing heightfield data.
-	if( RBHeightfield )
+	if (RBHeightfield)
 	{
 		GPhysXPendingKillHeightfield.Add(RBHeightfield);
 		RBHeightfield = NULL;
@@ -46,7 +46,7 @@ ULandscapeMeshCollisionComponent::FPhysXMeshRef::~FPhysXMeshRef()
 {
 #if WITH_PHYSX
 	// Free the existing heightfield data.
-	if( RBTriangleMesh )
+	if (RBTriangleMesh)
 	{
 		GPhysXPendingKillTriMesh.Add(RBTriangleMesh);
 		RBTriangleMesh = NULL;
@@ -97,12 +97,12 @@ void ULandscapeHeightfieldCollisionComponent::CreatePhysicsState()
 {
 	USceneComponent::CreatePhysicsState(); // route CreatePhysicsState, skip PrimitiveComponent implementation
 
-	if( !BodyInstance.IsValidBodyInstance())
+	if (!BodyInstance.IsValidBodyInstance())
 	{
 #if WITH_PHYSX
 		// This will do nothing, because we create heightfield at component PostLoad event, unless we destroyed it explicitly
 		CreateCollisionObject();
-		
+
 		if (IsValidRef(HeightfieldRef))
 		{
 			// Make transform for this landscape component PxActor
@@ -127,7 +127,7 @@ void ULandscapeHeightfieldCollisionComponent::CreatePhysicsState()
 			LandscapeComponentMatrix.SetAxis(1, TerrainZ);
 
 			PxTransform PhysXLandscapeComponentTransform = U2PTransform(FTransform(LandscapeComponentMatrix));
-						
+
 			// Create the geometry
 			PxHeightFieldGeometry LandscapeComponentGeom(HeightfieldRef->RBHeightfield, PxMeshGeometryFlags(), LandscapeScale.Z * LANDSCAPE_ZSCALE, LandscapeScale.Y * CollisionScale, LandscapeScale.X * CollisionScale);
 
@@ -142,14 +142,14 @@ void ULandscapeHeightfieldCollisionComponent::CreatePhysicsState()
 
 				// Setup filtering
 				PxFilterData PQueryFilterData, PSimFilterData;
-				CreateShapeFilterData( GetCollisionObjectType(), GetUniqueID(), GetCollisionResponseToChannels(), 0, 0, PQueryFilterData, PSimFilterData, true, false, true);
+				CreateShapeFilterData(GetCollisionObjectType(), GetUniqueID(), GetCollisionResponseToChannels(), 0, 0, PQueryFilterData, PSimFilterData, true, false, true);
 
 				// Heightfield is used for simple and complex collision
 				PQueryFilterData.word3 |= (EPDF_SimpleCollision | EPDF_ComplexCollision);
 				HeightFieldShapeSync->setQueryFilterData(PQueryFilterData);
 				HeightFieldShapeSync->setSimulationFilterData(PSimFilterData);
-				HeightFieldShapeSync->setFlag(PxShapeFlag::eSCENE_QUERY_SHAPE, true); 
-				HeightFieldShapeSync->setFlag(PxShapeFlag::eSIMULATION_SHAPE, true); 
+				HeightFieldShapeSync->setFlag(PxShapeFlag::eSCENE_QUERY_SHAPE, true);
+				HeightFieldShapeSync->setFlag(PxShapeFlag::eSIMULATION_SHAPE, true);
 				HeightFieldShapeSync->setFlag(PxShapeFlag::eVISUALIZATION, true);
 
 #if WITH_EDITOR
@@ -162,12 +162,12 @@ void ULandscapeHeightfieldCollisionComponent::CreatePhysicsState()
 						PxMaterial* PDefaultMat = GEngine->DefaultPhysMaterial->GetPhysXMaterial();
 						PxShape* HeightFieldEdShapeSync = HeightFieldActorSync->createShape(LandscapeComponentGeomEd, &PDefaultMat, 1);
 						check(HeightFieldEdShapeSync);
-						
+
 						FCollisionResponseContainer CollisionResponse;
 						CollisionResponse.SetAllChannels(ECollisionResponse::ECR_Ignore);
 						CollisionResponse.SetResponse(ECollisionChannel::ECC_Visibility, ECR_Block);
 						PxFilterData PQueryFilterDataEd, PSimFilterDataEd;
-						CreateShapeFilterData( ECollisionChannel::ECC_Visibility, GetUniqueID(), CollisionResponse, 0, 0, PQueryFilterDataEd, PSimFilterDataEd, true, false, true);
+						CreateShapeFilterData(ECollisionChannel::ECC_Visibility, GetUniqueID(), CollisionResponse, 0, 0, PQueryFilterDataEd, PSimFilterDataEd, true, false, true);
 
 						PQueryFilterDataEd.word3 |= (EPDF_SimpleCollision | EPDF_ComplexCollision);
 						HeightFieldEdShapeSync->setQueryFilterData(PQueryFilterDataEd);
@@ -211,8 +211,8 @@ void ULandscapeHeightfieldCollisionComponent::CreatePhysicsState()
 				if (PhysScene->HasAsyncScene())
 				{
 					PxScene* AsyncScene = PhysScene->GetPhysXScene(PST_Async);
-			
-					SCOPED_SCENE_WRITE_LOCK( AsyncScene );
+
+					SCOPED_SCENE_WRITE_LOCK(AsyncScene);
 					AsyncScene->addActor(*HeightFieldActorAsync);
 				}
 			}
@@ -241,16 +241,16 @@ void ULandscapeHeightfieldCollisionComponent::CreateCollisionObject()
 
 #if WITH_PHYSX	
 	// If we have not created a heightfield yet - do it now.
-	if( !IsValidRef(HeightfieldRef) )
+	if (!IsValidRef(HeightfieldRef))
 	{
-		if( !HeightfieldGuid.IsValid() )
+		if (!HeightfieldGuid.IsValid())
 		{
 			HeightfieldGuid = FGuid::NewGuid();
 		}
 
 		// Look for a heightfield object with the current Guid (this occurs with PIE)
 		FPhysXHeightfieldRef* ExistingHeightfieldRef = GSharedHeightfieldRefs.FindRef(HeightfieldGuid);
-		if( ExistingHeightfieldRef )
+		if (ExistingHeightfieldRef)
 		{
 			HeightfieldRef = ExistingHeightfieldRef;
 		}
@@ -302,23 +302,23 @@ void ULandscapeHeightfieldCollisionComponent::CreateCollisionObject()
 bool ULandscapeHeightfieldCollisionComponent::CookCollsionData(const FName& Format, bool bUseDefMaterial, TArray<uint8>& OutCookedData, TArray<UPhysicalMaterial*>& OutMaterails) const
 {
 #if WITH_PHYSX
-	
-	ALandscapeProxy* Proxy = GetLandscapeProxy();	
+
+	ALandscapeProxy* Proxy = GetLandscapeProxy();
 	if (!Proxy || !Proxy->GetRootComponent())
 	{
 		return false;
 	}
-	
+
 	UPhysicalMaterial* DefMaterial = Proxy->DefaultPhysMaterial ? Proxy->DefaultPhysMaterial : GEngine->DefaultPhysMaterial;
 
 	// ComponentToWorld might not be initialized at this point, so use landscape transform
 	FVector LandscapeScale = Proxy->GetRootComponent()->RelativeScale3D;
 	bool bIsMirrored = (LandscapeScale.X*LandscapeScale.Y*LandscapeScale.Z) < 0.f;
-	
-	int32 CollisionSizeVerts = CollisionSizeQuads+1;
-	
+
+	int32 CollisionSizeVerts = CollisionSizeQuads + 1;
+
 	const uint16* Heights = (const uint16*)CollisionHeightData.LockReadOnly();
-	check(CollisionHeightData.GetElementCount()==FMath::Square(CollisionSizeVerts));
+	check(CollisionHeightData.GetElementCount() == FMath::Square(CollisionSizeVerts));
 
 	//
 	const uint8* DominantLayers = nullptr;
@@ -326,22 +326,22 @@ bool ULandscapeHeightfieldCollisionComponent::CookCollsionData(const FName& Form
 	{
 		DominantLayers = (const uint8*)DominantLayerData.LockReadOnly();
 	}
-	
+
 	// List of materials which is actually used by heightfield
 	OutMaterails.Empty();
 
 	TArray<PxHeightFieldSample> Samples;
 	Samples.AddZeroed(FMath::Square(CollisionSizeVerts));
 
-	for(int32 RowIndex = 0; RowIndex < CollisionSizeVerts; RowIndex++)
+	for (int32 RowIndex = 0; RowIndex < CollisionSizeVerts; RowIndex++)
 	{
-		for(int32 ColIndex = 0; ColIndex < CollisionSizeVerts; ColIndex++)
+		for (int32 ColIndex = 0; ColIndex < CollisionSizeVerts; ColIndex++)
 		{
-			int32 SrcSampleIndex = (ColIndex * CollisionSizeVerts) + (bIsMirrored ? RowIndex : (CollisionSizeVerts-RowIndex-1));
+			int32 SrcSampleIndex = (ColIndex * CollisionSizeVerts) + (bIsMirrored ? RowIndex : (CollisionSizeVerts - RowIndex - 1));
 			int32 DstSampleIndex = (RowIndex * CollisionSizeVerts) + ColIndex;
 
 			PxHeightFieldSample& Sample = Samples[DstSampleIndex];
-			Sample.height = FMath::Clamp<int32>(((int32)Heights[SrcSampleIndex]-32768), -32768, 32767);
+			Sample.height = FMath::Clamp<int32>(((int32)Heights[SrcSampleIndex] - 32768), -32768, 32767);
 
 			// Materials are not relevant on the last row/column because they are per-triangle and the last row/column don't own any
 			if (RowIndex < CollisionSizeVerts - 1 &&
@@ -354,7 +354,7 @@ bool ULandscapeHeightfieldCollisionComponent::CookCollsionData(const FName& Form
 					if (ComponentLayerInfos.IsValidIndex(DominantLayerIdx))
 					{
 						ULandscapeLayerInfoObject* Layer = ComponentLayerInfos[DominantLayerIdx];
-						if (Layer == ALandscapeProxy::DataLayer)						
+						if (Layer == ALandscapeProxy::VisibilityLayer)
 						{
 							// If it's a hole, override with the hole flag.
 							MaterialIndex = PxHeightFieldMaterial::eHOLE;
@@ -374,7 +374,7 @@ bool ULandscapeHeightfieldCollisionComponent::CookCollsionData(const FName& Form
 			// TODO: edge turning
 		}
 	}
-			
+
 	CollisionHeightData.Unlock();
 	if (DominantLayers)
 	{
@@ -386,7 +386,7 @@ bool ULandscapeHeightfieldCollisionComponent::CookCollsionData(const FName& Form
 	{
 		OutMaterails.Add(DefMaterial);
 	}
-	
+
 	//
 	FIntPoint HFSize = FIntPoint(CollisionSizeVerts, CollisionSizeVerts);
 	float HFThickness = -Proxy->CollisionThickness / (LandscapeScale.Z * LANDSCAPE_ZSCALE);
@@ -395,7 +395,7 @@ bool ULandscapeHeightfieldCollisionComponent::CookCollsionData(const FName& Form
 	ITargetPlatformManagerModule* TPM = GetTargetPlatformManager();
 	const IPhysXFormat* Cooker = TPM->FindPhysXFormat(Format);
 	bool Result = Cooker->CookHeightField(Format, HFSize, HFThickness, Samples.GetData(), Samples.GetTypeSize(), OutData);
-	
+
 	if (Result)
 	{
 		OutCookedData.SetNumUninitialized(OutData.Num());
@@ -406,7 +406,7 @@ bool ULandscapeHeightfieldCollisionComponent::CookCollsionData(const FName& Form
 		OutCookedData.Empty();
 		OutMaterails.Empty();
 	}
-	
+
 	return Result;
 #endif	// WITH_PHYSX
 
@@ -429,15 +429,15 @@ bool ULandscapeMeshCollisionComponent::CookCollsionData(const FName& Format, boo
 	TArray<FVector>			Vertices;
 	TArray<FTriIndices>		Indices;
 	TArray<uint16>			MaterialIndices;
-	
-	const int32 CollisionSizeVerts = CollisionSizeQuads+1;
+
+	const int32 CollisionSizeVerts = CollisionSizeQuads + 1;
 	const int32 NumVerts = FMath::Square(CollisionSizeVerts);
 
 	const uint16* Heights = (const uint16*)CollisionHeightData.LockReadOnly();
 	const uint16* XYOffsets = (const uint16*)CollisionXYOffsetData.LockReadOnly();
-	check(CollisionHeightData.GetElementCount()==NumVerts);
-	check(CollisionXYOffsetData.GetElementCount()==NumVerts*2);
-	
+	check(CollisionHeightData.GetElementCount() == NumVerts);
+	check(CollisionXYOffsetData.GetElementCount() == NumVerts * 2);
+
 	const uint8* DominantLayers = nullptr;
 	if (DominantLayerData.GetElementCount() > 0)
 	{
@@ -450,7 +450,7 @@ bool ULandscapeMeshCollisionComponent::CookCollsionData(const FName& Format, boo
 	{
 		int32 X = i % CollisionSizeVerts;
 		int32 Y = i / CollisionSizeVerts;
-		Vertices[i].Set(X + ((float)XYOffsets[i*2] - 32768.f) * LANDSCAPE_XYOFFSET_SCALE, Y + ((float)XYOffsets[i*2+1] - 32768.f) * LANDSCAPE_XYOFFSET_SCALE, ((float) Heights[i] - 32768.f) * LANDSCAPE_ZSCALE );
+		Vertices[i].Set(X + ((float)XYOffsets[i * 2] - 32768.f) * LANDSCAPE_XYOFFSET_SCALE, Y + ((float)XYOffsets[i * 2 + 1] - 32768.f) * LANDSCAPE_XYOFFSET_SCALE, ((float)Heights[i] - 32768.f) * LANDSCAPE_ZSCALE);
 	}
 
 	const int32 NumTris = FMath::Square(CollisionSizeQuads) * 2;
@@ -475,7 +475,7 @@ bool ULandscapeMeshCollisionComponent::CookCollsionData(const FName& Format, boo
 				if (ComponentLayerInfos.IsValidIndex(DominantLayerIdx))
 				{
 					ULandscapeLayerInfoObject* Layer = ComponentLayerInfos[DominantLayerIdx];
-					if (Layer == ALandscapeProxy::DataLayer)
+					if (Layer == ALandscapeProxy::VisibilityLayer)
 					{
 						// If it's a hole, override with the hole flag.
 						bHole = true;
@@ -491,15 +491,15 @@ bool ULandscapeMeshCollisionComponent::CookCollsionData(const FName& Format, boo
 			FTriIndices& TriIndex1 = Indices[TriangleIdx];
 			if (bHole)
 			{
-				TriIndex1.v0 = (x+0) + (y+0) * CollisionSizeVerts;
+				TriIndex1.v0 = (x + 0) + (y + 0) * CollisionSizeVerts;
 				TriIndex1.v1 = TriIndex1.v0;
 				TriIndex1.v2 = TriIndex1.v0;
 			}
 			else
 			{
-				TriIndex1.v0 = (x+0) + (y+0) * CollisionSizeVerts;
-				TriIndex1.v1 = (x+1) + (y+1) * CollisionSizeVerts;
-				TriIndex1.v2 = (x+1) + (y+0) * CollisionSizeVerts;
+				TriIndex1.v0 = (x + 0) + (y + 0) * CollisionSizeVerts;
+				TriIndex1.v1 = (x + 1) + (y + 1) * CollisionSizeVerts;
+				TriIndex1.v2 = (x + 1) + (y + 0) * CollisionSizeVerts;
 			}
 
 			if (DominantLayers)
@@ -511,15 +511,15 @@ bool ULandscapeMeshCollisionComponent::CookCollsionData(const FName& Format, boo
 			FTriIndices& TriIndex2 = Indices[TriangleIdx];
 			if (bHole)
 			{
-				TriIndex2.v0 = (x+0) + (y+0) * CollisionSizeVerts;
+				TriIndex2.v0 = (x + 0) + (y + 0) * CollisionSizeVerts;
 				TriIndex2.v1 = TriIndex2.v0;
 				TriIndex2.v2 = TriIndex2.v0;
 			}
 			else
 			{
-				TriIndex2.v0 = (x+0) + (y+0) * CollisionSizeVerts;
-				TriIndex2.v1 = (x+0) + (y+1) * CollisionSizeVerts;
-				TriIndex2.v2 = (x+1) + (y+1) * CollisionSizeVerts;
+				TriIndex2.v0 = (x + 0) + (y + 0) * CollisionSizeVerts;
+				TriIndex2.v1 = (x + 0) + (y + 1) * CollisionSizeVerts;
+				TriIndex2.v2 = (x + 1) + (y + 1) * CollisionSizeVerts;
 			}
 
 			if (DominantLayers)
@@ -542,13 +542,13 @@ bool ULandscapeMeshCollisionComponent::CookCollsionData(const FName& Format, boo
 	{
 		OutMaterails.Add(DefMaterial);
 	}
-	
+
 	bool bFlipNormals = true;
 	TArray<uint8> OutData;
 	ITargetPlatformManagerModule* TPM = GetTargetPlatformManager();
 	const IPhysXFormat* Cooker = TPM->FindPhysXFormat(Format);
 	bool Result = Cooker->CookTriMesh(Format, Vertices, Indices, MaterialIndices, bFlipNormals, OutData);
-	
+
 	if (Result)
 	{
 		OutCookedData.SetNumUninitialized(OutData.Num());
@@ -576,20 +576,20 @@ void ULandscapeMeshCollisionComponent::CreateCollisionObject()
 
 #if WITH_PHYSX	
 	// If we have not created a heightfield yet - do it now.
-	if( !IsValidRef(MeshRef) )
+	if (!IsValidRef(MeshRef))
 	{
-		if( !MeshGuid.IsValid() )
+		if (!MeshGuid.IsValid())
 		{
 			MeshGuid = FGuid::NewGuid();
 		}
 
 		// Look for a heightfield object with the current Guid (this occurs with PIE)
 		FPhysXMeshRef* ExistingMeshRef = GSharedMeshRefs.FindRef(MeshGuid);
-		if( ExistingMeshRef )
+		if (ExistingMeshRef)
 		{
 			MeshRef = ExistingMeshRef;
 		}
-		else			
+		else
 		{
 #if WITH_EDITOR
 			// Create cooked physics data
@@ -635,7 +635,7 @@ void ULandscapeMeshCollisionComponent::CreatePhysicsState()
 {
 	USceneComponent::CreatePhysicsState(); // route CreatePhysicsState, skip PrimitiveComponent implementation
 
-	if(!BodyInstance.IsValidBodyInstance())
+	if (!BodyInstance.IsValidBodyInstance())
 	{
 #if WITH_PHYSX
 		// This will do nothing, because we create trimesh at component PostLoad event, unless we destroyed it explicitly
@@ -656,7 +656,7 @@ void ULandscapeMeshCollisionComponent::CreatePhysicsState()
 			// Get the scale to give to PhysX
 			FVector LandscapeScale = LandscapeComponentMatrix.ExtractScaling();
 			PxTransform PhysXLandscapeComponentTransform = U2PTransform(FTransform(LandscapeComponentMatrix));
-						
+
 			// Create tri-mesh shape
 			PxTriangleMeshGeometry PTriMeshGeom;
 			PTriMeshGeom.triangleMesh = MeshRef->RBTriangleMesh;
@@ -667,7 +667,7 @@ void ULandscapeMeshCollisionComponent::CreatePhysicsState()
 			if (PTriMeshGeom.isValid())
 			{
 				// Creating both a sync and async actor, since this object is static
-		
+
 				// Create the sync scene actor
 				PxRigidStatic* MeshActorSync = GPhysXSDK->createRigidStatic(PhysXLandscapeComponentTransform);
 				PxShape* MeshShapeSync = MeshActorSync->createShape(PTriMeshGeom, MeshRef->UsedPhysicalMaterialArray.GetTypedData(), MeshRef->UsedPhysicalMaterialArray.Num());
@@ -718,7 +718,7 @@ void ULandscapeMeshCollisionComponent::CreatePhysicsState()
 						CollisionResponse.SetAllChannels(ECollisionResponse::ECR_Ignore);
 						CollisionResponse.SetResponse(ECollisionChannel::ECC_Visibility, ECR_Block);
 						PxFilterData PQueryFilterDataEd, PSimFilterDataEd;
-						CreateShapeFilterData( ECollisionChannel::ECC_Visibility, GetUniqueID(), CollisionResponse, 0, 0, PQueryFilterDataEd, PSimFilterDataEd, true, false, true);
+						CreateShapeFilterData(ECollisionChannel::ECC_Visibility, GetUniqueID(), CollisionResponse, 0, 0, PQueryFilterDataEd, PSimFilterDataEd, true, false, true);
 
 						PQueryFilterDataEd.word3 |= (EPDF_SimpleCollision | EPDF_ComplexCollision);
 						MeshShapeEdSync->setQueryFilterData(PQueryFilterDataEd);
@@ -738,7 +738,7 @@ void ULandscapeMeshCollisionComponent::CreatePhysicsState()
 				{
 					MeshActorAsync->userData = &BodyInstance.PhysxUserData;
 				}
-								
+
 				// Add to scenes
 				PhysScene->GetPhysXScene(PST_Sync)->addActor(*MeshActorSync);
 
@@ -762,7 +762,7 @@ void ULandscapeMeshCollisionComponent::CreatePhysicsState()
 void ULandscapeMeshCollisionComponent::ApplyWorldOffset(const FVector& InOffset, bool bWorldShift)
 {
 	Super::ApplyWorldOffset(InOffset, bWorldShift);
-	
+
 	if (!bWorldShift || !FPhysScene::SupportsOriginShifting())
 	{
 		RecreatePhysicsState();
@@ -784,10 +784,10 @@ void ULandscapeMeshCollisionComponent::DestroyComponent()
 void ULandscapeHeightfieldCollisionComponent::UpdateHeightfieldRegion(int32 ComponentX1, int32 ComponentY1, int32 ComponentX2, int32 ComponentY2)
 {
 #if WITH_PHYSX
-	if( IsValidRef(HeightfieldRef) )
+	if (IsValidRef(HeightfieldRef))
 	{
 		// If we're currently sharing this data with a PIE session, we need to make a new heightfield.
-		if( HeightfieldRef->GetRefCount() > 1 )
+		if (HeightfieldRef->GetRefCount() > 1)
 		{
 			RecreateCollision(false);
 			return;
@@ -798,61 +798,61 @@ void ULandscapeHeightfieldCollisionComponent::UpdateHeightfieldRegion(int32 Comp
 			return;
 		}
 
-		int32 CollisionSizeVerts = CollisionSizeQuads+1;
+		int32 CollisionSizeVerts = CollisionSizeQuads + 1;
 
 		bool bIsMirrored = GetComponentToWorld().GetDeterminant() < 0.f;
 
 		uint16* Heights = (uint16*)CollisionHeightData.Lock(LOCK_READ_ONLY);
-		check(CollisionHeightData.GetElementCount()==FMath::Square(CollisionSizeVerts));
+		check(CollisionHeightData.GetElementCount() == FMath::Square(CollisionSizeVerts));
 
 		// PhysX heightfield has the X and Y axis swapped, and the X component is also inverted
 		int32 HeightfieldX1 = ComponentY1;
-		int32 HeightfieldY1 = (bIsMirrored ? ComponentX1 : (CollisionSizeVerts-ComponentX2-1));
-		int32 DstVertsX = ComponentY2-ComponentY1+1;
-		int32 DstVertsY = ComponentX2-ComponentX1+1;
+		int32 HeightfieldY1 = (bIsMirrored ? ComponentX1 : (CollisionSizeVerts - ComponentX2 - 1));
+		int32 DstVertsX = ComponentY2 - ComponentY1 + 1;
+		int32 DstVertsY = ComponentX2 - ComponentX1 + 1;
 
 		TArray<PxHeightFieldSample> Samples;
 		Samples.AddZeroed(DstVertsX*DstVertsY);
 
 		// Traverse the area in destination heigthfield coordinates
-		for(int32 RowIndex = 0; RowIndex < DstVertsY; RowIndex++)
+		for (int32 RowIndex = 0; RowIndex < DstVertsY; RowIndex++)
 		{
-			for(int32 ColIndex = 0; ColIndex < DstVertsX; ColIndex++)
+			for (int32 ColIndex = 0; ColIndex < DstVertsX; ColIndex++)
 			{
 				int32 SrcX = bIsMirrored ? (RowIndex + ComponentX1) : (ComponentX2 - RowIndex);
 				int32 SrcY = ColIndex + ComponentY1;
 				int32 SrcSampleIndex = (SrcY * CollisionSizeVerts) + SrcX;
-				check( SrcSampleIndex < FMath::Square(CollisionSizeVerts) );
+				check(SrcSampleIndex < FMath::Square(CollisionSizeVerts));
 				int32 DstSampleIndex = (RowIndex * DstVertsX) + ColIndex;
-				
+
 				PxHeightFieldSample& Sample = Samples[DstSampleIndex];
-				Sample.height = FMath::Clamp<int32>(((int32)Heights[SrcSampleIndex]-32768), -32768, 32767);
+				Sample.height = FMath::Clamp<int32>(((int32)Heights[SrcSampleIndex] - 32768), -32768, 32767);
 
 				Sample.materialIndex0 = 0;
 				Sample.materialIndex1 = 0;
 			}
 		}
-		
+
 		CollisionHeightData.Unlock();
 
 		PxHeightFieldDesc SubDesc;
-		SubDesc.format           = PxHeightFieldFormat::eS16_TM;
-		SubDesc.nbColumns		= DstVertsX;
-		SubDesc.nbRows			= DstVertsY;
-		SubDesc.samples.data	= Samples.GetData();
-		SubDesc.samples.stride	= sizeof(PxU32);
-		SubDesc.flags			= PxHeightFieldFlag::eNO_BOUNDARY_EDGES;
+		SubDesc.format = PxHeightFieldFormat::eS16_TM;
+		SubDesc.nbColumns = DstVertsX;
+		SubDesc.nbRows = DstVertsY;
+		SubDesc.samples.data = Samples.GetData();
+		SubDesc.samples.stride = sizeof(PxU32);
+		SubDesc.flags = PxHeightFieldFlag::eNO_BOUNDARY_EDGES;
 
 
 		HeightfieldRef->RBHeightfieldEd->modifySamples(HeightfieldX1, HeightfieldY1, SubDesc, true);
-		
+
 		//
 		// Reset geometry of heightfield shape. Required by the modifySamples
 		//
 		FVector LandscapeScale = GetComponentToWorld().GetScale3D();
 		// Create the geometry
 		PxHeightFieldGeometry LandscapeComponentGeom(HeightfieldRef->RBHeightfieldEd, PxMeshGeometryFlags(), LandscapeScale.Z * LANDSCAPE_ZSCALE, LandscapeScale.Y * CollisionScale, LandscapeScale.X * CollisionScale);
-				
+
 		if (BodyInstance.RigidActorSync)
 		{
 			TArray<PxShape*, TInlineAllocator<8>> PShapes;
@@ -937,8 +937,8 @@ void ULandscapeHeightfieldCollisionComponent::Serialize(FArchive& Ar)
 
 	if (Ar.UE4Ver() < VER_UE4_LANDSCAPE_COLLISION_DATA_COOKING)
 	{
-		CollisionHeightData.Serialize(Ar,this);
-		DominantLayerData.Serialize(Ar,this);
+		CollisionHeightData.Serialize(Ar, this);
+		DominantLayerData.Serialize(Ar, this);
 	}
 	else
 	{
@@ -949,11 +949,11 @@ void ULandscapeHeightfieldCollisionComponent::Serialize(FArchive& Ar)
 		{
 			UE_LOG(LogPhysics, Fatal, TEXT("This platform requires cooked packages, and physX data was not cooked into %s."), *GetFullName());
 		}
-	
+
 		if (bCooked)
 		{
 #if WITH_EDITOR
-			if (Ar.IsCooking() && !HasAnyFlags(RF_ClassDefaultObject))	
+			if (Ar.IsCooking() && !HasAnyFlags(RF_ClassDefaultObject))
 			{
 				CookCollsionData(Ar.CookingTarget()->GetPhysicsFormat(NULL), false, CookedCollisionData, CookedPhysicalMaterials);
 			}
@@ -964,8 +964,8 @@ void ULandscapeHeightfieldCollisionComponent::Serialize(FArchive& Ar)
 		}
 		else
 		{
-			CollisionHeightData.Serialize(Ar,this);
-			DominantLayerData.Serialize(Ar,this);
+			CollisionHeightData.Serialize(Ar, this);
+			DominantLayerData.Serialize(Ar, this);
 		}
 	}
 }
@@ -977,7 +977,7 @@ void ULandscapeMeshCollisionComponent::Serialize(FArchive& Ar)
 	if (Ar.UE4Ver() < VER_UE4_LANDSCAPE_COLLISION_DATA_COOKING)
 	{
 		// conditional serialization in later versions
-		CollisionXYOffsetData.Serialize(Ar,this);
+		CollisionXYOffsetData.Serialize(Ar, this);
 	}
 
 	// PhysX cooking mesh data
@@ -1000,7 +1000,7 @@ void ULandscapeMeshCollisionComponent::Serialize(FArchive& Ar)
 	else if (Ar.UE4Ver() >= VER_UE4_LANDSCAPE_COLLISION_DATA_COOKING)
 	{
 		// we serialize raw collision data only with non-cooked content
-		CollisionXYOffsetData.Serialize(Ar,this);
+		CollisionXYOffsetData.Serialize(Ar, this);
 	}
 }
 
@@ -1028,13 +1028,13 @@ void ULandscapeHeightfieldCollisionComponent::PostEditUndo()
 
 bool ULandscapeHeightfieldCollisionComponent::DoCustomNavigableGeometryExport(struct FNavigableGeometryExport* GeomExport) const
 {
-	check( IsInGameThread() );
+	check(IsInGameThread());
 #if WITH_PHYSX
 	if (IsValidRef(HeightfieldRef) && HeightfieldRef->RBHeightfield != nullptr)
 	{
 		FTransform HFToW = ComponentToWorld;
 		HFToW.MultiplyScale3D(FVector(CollisionScale, CollisionScale, LANDSCAPE_ZSCALE));
-	
+
 		GeomExport->ExportPxHeightField(HeightfieldRef->RBHeightfield, HFToW);
 	}
 #endif// WITH_PHYSX
@@ -1043,7 +1043,7 @@ bool ULandscapeHeightfieldCollisionComponent::DoCustomNavigableGeometryExport(st
 
 bool ULandscapeMeshCollisionComponent::DoCustomNavigableGeometryExport(struct FNavigableGeometryExport* GeomExport) const
 {
-	check( IsInGameThread() );
+	check(IsInGameThread());
 #if WITH_PHYSX
 	if (IsValidRef(MeshRef) && MeshRef->RBTriangleMesh != nullptr)
 	{
@@ -1066,7 +1066,7 @@ bool ULandscapeMeshCollisionComponent::DoCustomNavigableGeometryExport(struct FN
 void ULandscapeHeightfieldCollisionComponent::PostLoad()
 {
 	Super::PostLoad();
-	
+
 	// Create collision object right after component was loaded
 	// so we can free cooked collision data
 	CreateCollisionObject();
@@ -1101,7 +1101,7 @@ void ULandscapeHeightfieldCollisionComponent::PostLoad()
 			{
 				if (ComponentLayers_DEPRECATED[i] == DataWeightmapName)
 				{
-					ComponentLayerInfos[i] = ALandscapeProxy::DataLayer;
+					ComponentLayerInfos[i] = ALandscapeProxy::VisibilityLayer;
 				}
 				else
 				{
@@ -1119,7 +1119,7 @@ void ULandscapeHeightfieldCollisionComponent::PostLoad()
 #if WITH_EDITOR
 void ULandscapeInfo::UpdateAllAddCollisions()
 {
-	for (auto It = XYtoComponentMap.CreateIterator(); It; ++It )
+	for (auto It = XYtoComponentMap.CreateIterator(); It; ++It)
 	{
 		ULandscapeComponent* Comp = It.Value();
 		if (Comp)
@@ -1139,18 +1139,18 @@ void ULandscapeHeightfieldCollisionComponent::UpdateAddCollisions()
 	if (Info)
 	{
 		ALandscapeProxy* Proxy = GetLandscapeProxy();
-		FIntPoint ComponentBase = GetSectionBase()/Proxy->ComponentSizeQuads;
+		FIntPoint ComponentBase = GetSectionBase() / Proxy->ComponentSizeQuads;
 
-		FIntPoint NeighborsKeys[8] = 
+		FIntPoint NeighborsKeys[8] =
 		{
-			ComponentBase + FIntPoint(-1,-1),
-			ComponentBase + FIntPoint(+0,-1),
-			ComponentBase + FIntPoint(+1,-1),
-			ComponentBase + FIntPoint(-1,+0),
-			ComponentBase + FIntPoint(+1,+0),
-			ComponentBase + FIntPoint(-1,+1),
-			ComponentBase + FIntPoint(+0,+1),
-			ComponentBase + FIntPoint(+1,+1)
+			ComponentBase + FIntPoint(-1, -1),
+			ComponentBase + FIntPoint(+0, -1),
+			ComponentBase + FIntPoint(+1, -1),
+			ComponentBase + FIntPoint(-1, +0),
+			ComponentBase + FIntPoint(+1, +0),
+			ComponentBase + FIntPoint(-1, +1),
+			ComponentBase + FIntPoint(+0, +1),
+			ComponentBase + FIntPoint(+1, +1)
 		};
 
 		// Search for Neighbors...
@@ -1183,16 +1183,16 @@ void ULandscapeInfo::UpdateAddCollision(FIntPoint LandscapeKey)
 	// 0 1 2 
 	// 3   4
 	// 5 6 7
-	FIntPoint NeighborsKeys[8] = 
-	{
-		LandscapeKey + FIntPoint(-1,-1),
-		LandscapeKey + FIntPoint(+0,-1),
-		LandscapeKey + FIntPoint(+1,-1),
-		LandscapeKey + FIntPoint(-1,+0),
-		LandscapeKey + FIntPoint(+1,+0),
-		LandscapeKey + FIntPoint(-1,+1),
-		LandscapeKey + FIntPoint(+0,+1),
-		LandscapeKey + FIntPoint(+1,+1)
+	FIntPoint NeighborsKeys[8] =
+	{
+		LandscapeKey + FIntPoint(-1, -1),
+		LandscapeKey + FIntPoint(+0, -1),
+		LandscapeKey + FIntPoint(+1, -1),
+		LandscapeKey + FIntPoint(-1, +0),
+		LandscapeKey + FIntPoint(+1, +0),
+		LandscapeKey + FIntPoint(-1, +1),
+		LandscapeKey + FIntPoint(+0, +1),
+		LandscapeKey + FIntPoint(+1, +1)
 	};
 
 	ULandscapeHeightfieldCollisionComponent* NeighborCollisions[8];
@@ -1218,7 +1218,7 @@ void ULandscapeInfo::UpdateAddCollision(FIntPoint LandscapeKey)
 	{
 		uint16* Heights = (uint16*)NeighborCollisions[0]->CollisionHeightData.Lock(LOCK_READ_ONLY);
 		int32 CollisionSizeVerts = NeighborCollisions[0]->CollisionSizeQuads + 1;
-		HeightCorner[0] = Heights[ CollisionSizeVerts-1 + (CollisionSizeVerts-1)*CollisionSizeVerts ];
+		HeightCorner[0] = Heights[CollisionSizeVerts - 1 + (CollisionSizeVerts - 1)*CollisionSizeVerts];
 		CornerSet |= 1;
 		NeighborCollisions[0]->CollisionHeightData.Unlock();
 	}
@@ -1226,7 +1226,7 @@ void ULandscapeInfo::UpdateAddCollision(FIntPoint LandscapeKey)
 	{
 		uint16* Heights = (uint16*)NeighborCollisions[2]->CollisionHeightData.Lock(LOCK_READ_ONLY);
 		int32 CollisionSizeVerts = NeighborCollisions[2]->CollisionSizeQuads + 1;
-		HeightCorner[1] = Heights[ (CollisionSizeVerts-1)*CollisionSizeVerts ];
+		HeightCorner[1] = Heights[(CollisionSizeVerts - 1)*CollisionSizeVerts];
 		CornerSet |= 1 << 1;
 		NeighborCollisions[2]->CollisionHeightData.Unlock();
 	}
@@ -1234,7 +1234,7 @@ void ULandscapeInfo::UpdateAddCollision(FIntPoint LandscapeKey)
 	{
 		uint16* Heights = (uint16*)NeighborCollisions[5]->CollisionHeightData.Lock(LOCK_READ_ONLY);
 		int32 CollisionSizeVerts = NeighborCollisions[5]->CollisionSizeQuads + 1;
-		HeightCorner[2] = Heights[ (CollisionSizeVerts-1) ];
+		HeightCorner[2] = Heights[(CollisionSizeVerts - 1)];
 		CornerSet |= 1 << 2;
 		NeighborCollisions[5]->CollisionHeightData.Unlock();
 	}
@@ -1242,7 +1242,7 @@ void ULandscapeInfo::UpdateAddCollision(FIntPoint LandscapeKey)
 	{
 		uint16* Heights = (uint16*)NeighborCollisions[7]->CollisionHeightData.Lock(LOCK_READ_ONLY);
 		int32 CollisionSizeVerts = NeighborCollisions[7]->CollisionSizeQuads + 1;
-		HeightCorner[3] = Heights[ 0 ];
+		HeightCorner[3] = Heights[0];
 		CornerSet |= 1 << 3;
 		NeighborCollisions[7]->CollisionHeightData.Unlock();
 	}
@@ -1252,9 +1252,9 @@ void ULandscapeInfo::UpdateAddCollision(FIntPoint LandscapeKey)
 	{
 		uint16* Heights = (uint16*)NeighborCollisions[1]->CollisionHeightData.Lock(LOCK_READ_ONLY);
 		int32 CollisionSizeVerts = NeighborCollisions[1]->CollisionSizeQuads + 1;
-		HeightCorner[0] = Heights[ (CollisionSizeVerts-1)*CollisionSizeVerts ];
+		HeightCorner[0] = Heights[(CollisionSizeVerts - 1)*CollisionSizeVerts];
 		CornerSet |= 1;
-		HeightCorner[1] = Heights[ CollisionSizeVerts-1 + (CollisionSizeVerts-1)*CollisionSizeVerts ];
+		HeightCorner[1] = Heights[CollisionSizeVerts - 1 + (CollisionSizeVerts - 1)*CollisionSizeVerts];
 		CornerSet |= 1 << 1;
 		NeighborCollisions[1]->CollisionHeightData.Unlock();
 	}
@@ -1262,9 +1262,9 @@ void ULandscapeInfo::UpdateAddCollision(FIntPoint LandscapeKey)
 	{
 		uint16* Heights = (uint16*)NeighborCollisions[3]->CollisionHeightData.Lock(LOCK_READ_ONLY);
 		int32 CollisionSizeVerts = NeighborCollisions[3]->CollisionSizeQuads + 1;
-		HeightCorner[0] = Heights[ (CollisionSizeVerts-1) ];
+		HeightCorner[0] = Heights[(CollisionSizeVerts - 1)];
 		CornerSet |= 1;
-		HeightCorner[2] = Heights[ CollisionSizeVerts-1 + (CollisionSizeVerts-1)*CollisionSizeVerts ];
+		HeightCorner[2] = Heights[CollisionSizeVerts - 1 + (CollisionSizeVerts - 1)*CollisionSizeVerts];
 		CornerSet |= 1 << 2;
 		NeighborCollisions[3]->CollisionHeightData.Unlock();
 	}
@@ -1272,9 +1272,9 @@ void ULandscapeInfo::UpdateAddCollision(FIntPoint LandscapeKey)
 	{
 		uint16* Heights = (uint16*)NeighborCollisions[4]->CollisionHeightData.Lock(LOCK_READ_ONLY);
 		int32 CollisionSizeVerts = NeighborCollisions[4]->CollisionSizeQuads + 1;
-		HeightCorner[1] = Heights[ 0 ];
+		HeightCorner[1] = Heights[0];
 		CornerSet |= 1 << 1;
-		HeightCorner[3] = Heights[ (CollisionSizeVerts-1)*CollisionSizeVerts ];
+		HeightCorner[3] = Heights[(CollisionSizeVerts - 1)*CollisionSizeVerts];
 		CornerSet |= 1 << 3;
 		NeighborCollisions[4]->CollisionHeightData.Unlock();
 	}
@@ -1282,9 +1282,9 @@ void ULandscapeInfo::UpdateAddCollision(FIntPoint LandscapeKey)
 	{
 		uint16* Heights = (uint16*)NeighborCollisions[6]->CollisionHeightData.Lock(LOCK_READ_ONLY);
 		int32 CollisionSizeVerts = NeighborCollisions[6]->CollisionSizeQuads + 1;
-		HeightCorner[2] = Heights[ 0 ];
+		HeightCorner[2] = Heights[0];
 		CornerSet |= 1 << 2;
-		HeightCorner[3] = Heights[ (CollisionSizeVerts-1) ];
+		HeightCorner[3] = Heights[(CollisionSizeVerts - 1)];
 		CornerSet |= 1 << 3;
 		NeighborCollisions[6]->CollisionHeightData.Unlock();
 	}
@@ -1298,11 +1298,11 @@ void ULandscapeInfo::UpdateAddCollision(FIntPoint LandscapeKey)
 	FIntPoint SectionBase = LandscapeKey*ComponentSizeQuads;
 
 	// Transform Height to Vectors...
-	FMatrix LtoW = GetLandscapeProxy()->LandscapeActorToWorld().ToMatrixWithScale(); 
-	AddCollision->Corners[0] = LtoW.TransformPosition( FVector(SectionBase.X,						SectionBase.Y,					  LandscapeDataAccess::GetLocalHeight(HeightCorner[0])) );
-	AddCollision->Corners[1] = LtoW.TransformPosition( FVector(SectionBase.X+ComponentSizeQuads,	SectionBase.Y,					  LandscapeDataAccess::GetLocalHeight(HeightCorner[1])) );
-	AddCollision->Corners[2] = LtoW.TransformPosition( FVector(SectionBase.X,						SectionBase.Y+ComponentSizeQuads, LandscapeDataAccess::GetLocalHeight(HeightCorner[2])) );
-	AddCollision->Corners[3] = LtoW.TransformPosition( FVector(SectionBase.X+ComponentSizeQuads,	SectionBase.Y+ComponentSizeQuads, LandscapeDataAccess::GetLocalHeight(HeightCorner[3])) );
+	FMatrix LtoW = GetLandscapeProxy()->LandscapeActorToWorld().ToMatrixWithScale();
+	AddCollision->Corners[0] = LtoW.TransformPosition(FVector(SectionBase.X, SectionBase.Y, LandscapeDataAccess::GetLocalHeight(HeightCorner[0])));
+	AddCollision->Corners[1] = LtoW.TransformPosition(FVector(SectionBase.X + ComponentSizeQuads, SectionBase.Y, LandscapeDataAccess::GetLocalHeight(HeightCorner[1])));
+	AddCollision->Corners[2] = LtoW.TransformPosition(FVector(SectionBase.X, SectionBase.Y + ComponentSizeQuads, LandscapeDataAccess::GetLocalHeight(HeightCorner[2])));
+	AddCollision->Corners[3] = LtoW.TransformPosition(FVector(SectionBase.X + ComponentSizeQuads, SectionBase.Y + ComponentSizeQuads, LandscapeDataAccess::GetLocalHeight(HeightCorner[3])));
 }
 
 void ULandscapeHeightfieldCollisionComponent::ExportCustomProperties(FOutputDevice& Out, uint32 Indent)
@@ -1311,43 +1311,43 @@ void ULandscapeHeightfieldCollisionComponent::ExportCustomProperties(FOutputDevi
 	{
 		return;
 	}
-	
+
 	uint16* Heights = (uint16*)CollisionHeightData.Lock(LOCK_READ_ONLY);
-	int32 NumHeights = FMath::Square(CollisionSizeQuads+1);
-	check(CollisionHeightData.GetElementCount()==NumHeights);
+	int32 NumHeights = FMath::Square(CollisionSizeQuads + 1);
+	check(CollisionHeightData.GetElementCount() == NumHeights);
 
-	Out.Logf( TEXT("%sCustomProperties CollisionHeightData "), FCString::Spc(Indent));
-	for( int32 i=0;i<NumHeights;i++ )
+	Out.Logf(TEXT("%sCustomProperties CollisionHeightData "), FCString::Spc(Indent));
+	for (int32 i = 0; i < NumHeights; i++)
 	{
-		Out.Logf( TEXT("%d "), Heights[i] );
+		Out.Logf(TEXT("%d "), Heights[i]);
 	}
 
 	CollisionHeightData.Unlock();
-	Out.Logf( TEXT("\r\n") );
+	Out.Logf(TEXT("\r\n"));
 
 	int32 NumDominantLayerSamples = DominantLayerData.GetElementCount();
-	check(NumDominantLayerSamples == 0 || NumDominantLayerSamples==NumHeights);
+	check(NumDominantLayerSamples == 0 || NumDominantLayerSamples == NumHeights);
 
-	if( NumDominantLayerSamples > 0 )
+	if (NumDominantLayerSamples > 0)
 	{
 		uint8* DominantLayerSamples = (uint8*)DominantLayerData.Lock(LOCK_READ_ONLY);
 
-		Out.Logf( TEXT("%sCustomProperties DominantLayerData "), FCString::Spc(Indent));
-		for( int32 i=0;i<NumDominantLayerSamples;i++ )
+		Out.Logf(TEXT("%sCustomProperties DominantLayerData "), FCString::Spc(Indent));
+		for (int32 i = 0; i < NumDominantLayerSamples; i++)
 		{
-			Out.Logf( TEXT("%02x"), DominantLayerSamples[i] );
+			Out.Logf(TEXT("%02x"), DominantLayerSamples[i]);
 		}
 
 		DominantLayerData.Unlock();
-		Out.Logf( TEXT("\r\n") );
+		Out.Logf(TEXT("\r\n"));
 	}
 }
 
 void ULandscapeHeightfieldCollisionComponent::ImportCustomProperties(const TCHAR* SourceText, FFeedbackContext* Warn)
 {
-	if(FParse::Command(&SourceText,TEXT("CollisionHeightData")))
+	if (FParse::Command(&SourceText, TEXT("CollisionHeightData")))
 	{
-		int32 NumHeights = FMath::Square(CollisionSizeQuads+1);
+		int32 NumHeights = FMath::Square(CollisionSizeQuads + 1);
 
 		CollisionHeightData.Lock(LOCK_READ_WRITE);
 		uint16* Heights = (uint16*)CollisionHeightData.Realloc(NumHeights);
@@ -1355,30 +1355,30 @@ void ULandscapeHeightfieldCollisionComponent::ImportCustomProperties(const TCHAR
 
 		FParse::Next(&SourceText);
 		int32 i = 0;
-		while( FChar::IsDigit(*SourceText) ) 
+		while (FChar::IsDigit(*SourceText))
 		{
-			if( i < NumHeights )
+			if (i < NumHeights)
 			{
 				Heights[i++] = FCString::Atoi(SourceText);
-				while( FChar::IsDigit(*SourceText) ) 
+				while (FChar::IsDigit(*SourceText))
 				{
 					SourceText++;
 				}
 			}
 
 			FParse::Next(&SourceText);
-		} 
+		}
 
 		CollisionHeightData.Unlock();
 
-		if( i != NumHeights )
+		if (i != NumHeights)
 		{
-			Warn->Logf( *NSLOCTEXT( "Core", "SyntaxError", "Syntax Error" ).ToString() );
+			Warn->Logf(*NSLOCTEXT("Core", "SyntaxError", "Syntax Error").ToString());
 		}
 	}
-	else if(FParse::Command(&SourceText,TEXT("DominantLayerData")))
+	else if (FParse::Command(&SourceText, TEXT("DominantLayerData")))
 	{
-		int32 NumDominantLayerSamples = FMath::Square(CollisionSizeQuads+1);
+		int32 NumDominantLayerSamples = FMath::Square(CollisionSizeQuads + 1);
 
 		DominantLayerData.Lock(LOCK_READ_WRITE);
 		uint8* DominantLayerSamples = (uint8*)DominantLayerData.Realloc(NumDominantLayerSamples);
@@ -1386,20 +1386,20 @@ void ULandscapeHeightfieldCollisionComponent::ImportCustomProperties(const TCHAR
 
 		FParse::Next(&SourceText);
 		int32 i = 0;
-		while( SourceText[0] && SourceText[1] )
+		while (SourceText[0] && SourceText[1])
 		{
-			if( i < NumDominantLayerSamples )
+			if (i < NumDominantLayerSamples)
 			{
 				DominantLayerSamples[i++] = FParse::HexDigit(SourceText[0]) * 16 + FParse::HexDigit(SourceText[1]);
 			}
 			SourceText += 2;
-		} 
+		}
 
 		DominantLayerData.Unlock();
 
-		if( i != NumDominantLayerSamples )
+		if (i != NumDominantLayerSamples)
 		{
-			Warn->Logf( *NSLOCTEXT( "Core", "SyntaxError", "Syntax Error" ).ToString() );
+			Warn->Logf(*NSLOCTEXT("Core", "SyntaxError", "Syntax Error").ToString());
 		}
 	}
 }
@@ -1414,24 +1414,24 @@ void ULandscapeMeshCollisionComponent::ExportCustomProperties(FOutputDevice& Out
 	Super::ExportCustomProperties(Out, Indent);
 
 	uint16* XYOffsets = (uint16*)CollisionXYOffsetData.Lock(LOCK_READ_ONLY);
-	int32 NumOffsets = FMath::Square(CollisionSizeQuads+1) * 2;
-	check(CollisionXYOffsetData.GetElementCount()==NumOffsets);
+	int32 NumOffsets = FMath::Square(CollisionSizeQuads + 1) * 2;
+	check(CollisionXYOffsetData.GetElementCount() == NumOffsets);
 
-	Out.Logf( TEXT("%sCustomProperties CollisionXYOffsetData "), FCString::Spc(Indent));
-	for( int32 i=0;i<NumOffsets;i++ )
+	Out.Logf(TEXT("%sCustomProperties CollisionXYOffsetData "), FCString::Spc(Indent));
+	for (int32 i = 0; i < NumOffsets; i++)
 	{
-		Out.Logf( TEXT("%d "), XYOffsets[i] );
+		Out.Logf(TEXT("%d "), XYOffsets[i]);
 	}
 
 	CollisionXYOffsetData.Unlock();
-	Out.Logf( TEXT("\r\n") );
+	Out.Logf(TEXT("\r\n"));
 }
 
 void ULandscapeMeshCollisionComponent::ImportCustomProperties(const TCHAR* SourceText, FFeedbackContext* Warn)
 {
-	if(FParse::Command(&SourceText,TEXT("CollisionHeightData")))
+	if (FParse::Command(&SourceText, TEXT("CollisionHeightData")))
 	{
-		int32 NumHeights = FMath::Square(CollisionSizeQuads+1);
+		int32 NumHeights = FMath::Square(CollisionSizeQuads + 1);
 
 		CollisionHeightData.Lock(LOCK_READ_WRITE);
 		uint16* Heights = (uint16*)CollisionHeightData.Realloc(NumHeights);
@@ -1439,30 +1439,30 @@ void ULandscapeMeshCollisionComponent::ImportCustomProperties(const TCHAR* Sourc
 
 		FParse::Next(&SourceText);
 		int32 i = 0;
-		while( FChar::IsDigit(*SourceText) ) 
+		while (FChar::IsDigit(*SourceText))
 		{
-			if( i < NumHeights )
+			if (i < NumHeights)
 			{
 				Heights[i++] = FCString::Atoi(SourceText);
-				while( FChar::IsDigit(*SourceText) ) 
+				while (FChar::IsDigit(*SourceText))
 				{
 					SourceText++;
 				}
 			}
 
 			FParse::Next(&SourceText);
-		} 
+		}
 
 		CollisionHeightData.Unlock();
 
-		if( i != NumHeights )
+		if (i != NumHeights)
 		{
-			Warn->Logf( *NSLOCTEXT( "Core", "SyntaxError", "Syntax Error" ).ToString() );
+			Warn->Logf(*NSLOCTEXT("Core", "SyntaxError", "Syntax Error").ToString());
 		}
 	}
-	else if(FParse::Command(&SourceText,TEXT("DominantLayerData")))
+	else if (FParse::Command(&SourceText, TEXT("DominantLayerData")))
 	{
-		int32 NumDominantLayerSamples = FMath::Square(CollisionSizeQuads+1);
+		int32 NumDominantLayerSamples = FMath::Square(CollisionSizeQuads + 1);
 
 		DominantLayerData.Lock(LOCK_READ_WRITE);
 		uint8* DominantLayerSamples = (uint8*)DominantLayerData.Realloc(NumDominantLayerSamples);
@@ -1470,25 +1470,25 @@ void ULandscapeMeshCollisionComponent::ImportCustomProperties(const TCHAR* Sourc
 
 		FParse::Next(&SourceText);
 		int32 i = 0;
-		while( SourceText[0] && SourceText[1] )
+		while (SourceText[0] && SourceText[1])
 		{
-			if( i < NumDominantLayerSamples )
+			if (i < NumDominantLayerSamples)
 			{
 				DominantLayerSamples[i++] = FParse::HexDigit(SourceText[0]) * 16 + FParse::HexDigit(SourceText[1]);
 			}
 			SourceText += 2;
-		} 
+		}
 
 		DominantLayerData.Unlock();
 
-		if( i != NumDominantLayerSamples )
+		if (i != NumDominantLayerSamples)
 		{
-			Warn->Logf( *NSLOCTEXT( "Core", "SyntaxError", "Syntax Error" ).ToString() );
+			Warn->Logf(*NSLOCTEXT("Core", "SyntaxError", "Syntax Error").ToString());
 		}
 	}
-	else if (FParse::Command(&SourceText,TEXT("CollisionXYOffsetData")))
+	else if (FParse::Command(&SourceText, TEXT("CollisionXYOffsetData")))
 	{
-		int32 NumOffsets = FMath::Square(CollisionSizeQuads+1) * 2;
+		int32 NumOffsets = FMath::Square(CollisionSizeQuads + 1) * 2;
 
 		CollisionXYOffsetData.Lock(LOCK_READ_WRITE);
 		uint16* Offsets = (uint16*)CollisionXYOffsetData.Realloc(NumOffsets);
@@ -1496,25 +1496,25 @@ void ULandscapeMeshCollisionComponent::ImportCustomProperties(const TCHAR* Sourc
 
 		FParse::Next(&SourceText);
 		int32 i = 0;
-		while( FChar::IsDigit(*SourceText) ) 
+		while (FChar::IsDigit(*SourceText))
 		{
-			if( i < NumOffsets )
+			if (i < NumOffsets)
 			{
 				Offsets[i++] = FCString::Atoi(SourceText);
-				while( FChar::IsDigit(*SourceText) ) 
+				while (FChar::IsDigit(*SourceText))
 				{
 					SourceText++;
 				}
 			}
 
 			FParse::Next(&SourceText);
-		} 
+		}
 
 		CollisionXYOffsetData.Unlock();
 
-		if( i != NumOffsets )
+		if (i != NumOffsets)
 		{
-			Warn->Logf( *NSLOCTEXT( "Core", "SyntaxError", "Syntax Error" ).ToString() );
+			Warn->Logf(*NSLOCTEXT("Core", "SyntaxError", "Syntax Error").ToString());
 		}
 	}
 }
