Remove FLandscapeToolSet

diff --git a/Engine/Source/Runtime/Engine/Private/Landscape/LandscapeEdit.cpp b/Engine/Source/Runtime/Engine/Private/Landscape/LandscapeEdit.cpp
--- a/Engine/Source/Runtime/Engine/Private/Landscape/LandscapeEdit.cpp
+++ b/Engine/Source/Runtime/Engine/Private/Landscape/LandscapeEdit.cpp
@@ -36,11 +36,11 @@ DEFINE_LOG_CATEGORY(LogLandscape);
 
 #define LOCTEXT_NAMESPACE "Landscape"
 
-ULandscapeLayerInfoObject* ALandscapeProxy::DataLayer = NULL;
+ULandscapeLayerInfoObject* ALandscapeProxy::VisibilityLayer = NULL;
 
 #if WITH_EDITOR
 
-void ULandscapeComponent::Init(int32 InBaseX,int32 InBaseY,int32 InComponentSizeQuads, int32 InNumSubsections,int32 InSubsectionSizeQuads)
+void ULandscapeComponent::Init(int32 InBaseX, int32 InBaseY, int32 InComponentSizeQuads, int32 InNumSubsections, int32 InSubsectionSizeQuads)
 {
 	SetSectionBase(FIntPoint(InBaseX, InBaseY));
 	FVector RelativeLocation = FVector(GetSectionBase() - GetLandscapeProxy()->LandscapeSectionOffset);
@@ -58,20 +58,20 @@ void ULandscapeComponent::UpdateCachedBounds()
 
 	// Update local-space bounding box
 	CachedLocalBox.Init();
-	for( int32 y=0;y<ComponentSizeQuads+1;y++ )
+	for (int32 y = 0; y < ComponentSizeQuads + 1; y++)
 	{
-		for( int32 x=0;x<ComponentSizeQuads+1;x++ )
+		for (int32 x = 0; x < ComponentSizeQuads + 1; x++)
 		{
-			CachedLocalBox += CDI.GetLocalVertex(x,y);
+			CachedLocalBox += CDI.GetLocalVertex(x, y);
 		}
 	}
 
 	// Update collision component bounds
 	ULandscapeHeightfieldCollisionComponent* HFCollisionComponent = CollisionComponent.Get();
-	if( HFCollisionComponent )
+	if (HFCollisionComponent)
 	{
 		HFCollisionComponent->Modify();
-		HFCollisionComponent->CachedLocalBox = CachedLocalBox;	
+		HFCollisionComponent->CachedLocalBox = CachedLocalBox;
 		HFCollisionComponent->UpdateComponentToWorld();
 	}
 }
@@ -90,7 +90,7 @@ ULandscapeMaterialInstanceConstant* ALandscapeProxy::GetLayerThumbnailMIC(UMater
 		CombinationMaterialInstance = Proxy->MaterialInstanceConstantMap.FindRef(*LayerKey);
 	}
 
-	if( CombinationMaterialInstance == NULL || CombinationMaterialInstance->Parent != LandscapeMaterial || (Proxy && Proxy->GetOutermost() != CombinationMaterialInstance->GetOutermost()) )
+	if (CombinationMaterialInstance == NULL || CombinationMaterialInstance->Parent != LandscapeMaterial || (Proxy && Proxy->GetOutermost() != CombinationMaterialInstance->GetOutermost()))
 	{
 		FlushRenderingCommands();
 		UObject* MICOuter = GetTransientPackage();
@@ -102,17 +102,17 @@ ULandscapeMaterialInstanceConstant* ALandscapeProxy::GetLayerThumbnailMIC(UMater
 		if (Proxy)
 		{
 			UE_LOG(LogLandscape, Log, TEXT("Looking for key %s, making new combination %s"), *LayerKey, *CombinationMaterialInstance->GetName());
-			Proxy->MaterialInstanceConstantMap.Add(*LayerKey,CombinationMaterialInstance);
+			Proxy->MaterialInstanceConstantMap.Add(*LayerKey, CombinationMaterialInstance);
 		}
 		CombinationMaterialInstance->SetParentEditorOnly(LandscapeMaterial);
 
 		FStaticParameterSet StaticParameters;
 		CombinationMaterialInstance->GetStaticParameterValues(StaticParameters);
 
-		for( int32 LayerParameterIdx=0;LayerParameterIdx<StaticParameters.TerrainLayerWeightParameters.Num();LayerParameterIdx++ )
+		for (int32 LayerParameterIdx = 0; LayerParameterIdx < StaticParameters.TerrainLayerWeightParameters.Num(); LayerParameterIdx++)
 		{
 			FStaticTerrainLayerWeightParameter& LayerParameter = StaticParameters.TerrainLayerWeightParameters[LayerParameterIdx];
-			if( LayerParameter.ParameterName == LayerName )
+			if (LayerParameter.ParameterName == LayerName)
 			{
 				LayerParameter.WeightmapIndex = 0;
 				LayerParameter.bOverride = true;
@@ -133,7 +133,7 @@ ULandscapeMaterialInstanceConstant* ALandscapeProxy::GetLayerThumbnailMIC(UMater
 	MaterialInstance->SetParentEditorOnly(CombinationMaterialInstance);
 	MaterialInstance->bIsLayerThumbnail = true;
 
-	FLinearColor Mask(1.f,0.f,0.f,0.f);
+	FLinearColor Mask(1.f, 0.f, 0.f, 0.f);
 	MaterialInstance->SetVectorParameterValueEditorOnly(FName(*FString::Printf(TEXT("LayerMask_%s"), *LayerName.ToString())), Mask);
 	MaterialInstance->SetTextureParameterValueEditorOnly(FName(TEXT("Weightmap0")), ThumbnailWeightmap);
 	MaterialInstance->SetTextureParameterValueEditorOnly(FName(TEXT("Heightmap")), ThumbnailHeightmap);
@@ -150,15 +150,15 @@ UMaterialInstanceConstant* ULandscapeComponent::GetCombinationMaterial(bool bMob
 
 	UMaterialInterface* LandscapeMaterial = ComponentHasVisibilityPainted() ? GetLandscapeHoleMaterial() : GetLandscapeMaterial();
 
-	if( LandscapeMaterial != NULL )
+	if (LandscapeMaterial != NULL)
 	{
 		// Ensure top level UMaterial has appropriate usage flags set.
 		bool bNeedsRecompile;
 		UMaterial* ParentUMaterial = LandscapeMaterial->GetMaterial();
-		if( ParentUMaterial && ParentUMaterial != UMaterial::GetDefaultMaterial(MD_Surface) )
+		if (ParentUMaterial && ParentUMaterial != UMaterial::GetDefaultMaterial(MD_Surface))
 		{
-			ParentUMaterial->SetMaterialUsage(bNeedsRecompile,MATUSAGE_Landscape);
-			ParentUMaterial->SetMaterialUsage(bNeedsRecompile,MATUSAGE_StaticLighting);
+			ParentUMaterial->SetMaterialUsage(bNeedsRecompile, MATUSAGE_Landscape);
+			ParentUMaterial->SetMaterialUsage(bNeedsRecompile, MATUSAGE_StaticLighting);
 		}
 
 		FString LayerKey = GetLayerAllocationKey(bMobile);
@@ -166,32 +166,32 @@ UMaterialInstanceConstant* ULandscapeComponent::GetCombinationMaterial(bool bMob
 
 		// Find or set a matching MIC in the Landscape's map.
 		UMaterialInstanceConstant* CombinationMaterialInstance = Proxy->MaterialInstanceConstantMap.FindRef(*LayerKey);
-		if( CombinationMaterialInstance == NULL || CombinationMaterialInstance->Parent != LandscapeMaterial || GetOutermost() != CombinationMaterialInstance->GetOutermost() )
+		if (CombinationMaterialInstance == NULL || CombinationMaterialInstance->Parent != LandscapeMaterial || GetOutermost() != CombinationMaterialInstance->GetOutermost())
 		{
 			FlushRenderingCommands();
 
 			CombinationMaterialInstance = ConstructObject<ULandscapeMaterialInstanceConstant>(ULandscapeMaterialInstanceConstant::StaticClass(), GetOutermost());
 			UE_LOG(LogLandscape, Log, TEXT("Looking for key %s, making new combination %s"), *LayerKey, *CombinationMaterialInstance->GetName());
-			Proxy->MaterialInstanceConstantMap.Add(*LayerKey,CombinationMaterialInstance);
+			Proxy->MaterialInstanceConstantMap.Add(*LayerKey, CombinationMaterialInstance);
 			CombinationMaterialInstance->SetParentEditorOnly(LandscapeMaterial);
 
 			FStaticParameterSet StaticParameters;
 			CombinationMaterialInstance->GetStaticParameterValues(StaticParameters);
 
 			// Find weightmap mapping for each layer parameter, or disable if the layer is not used in this component.
-			for( int32 LayerParameterIdx=0;LayerParameterIdx<StaticParameters.TerrainLayerWeightParameters.Num();LayerParameterIdx++ )
+			for (int32 LayerParameterIdx = 0; LayerParameterIdx < StaticParameters.TerrainLayerWeightParameters.Num(); LayerParameterIdx++)
 			{
 				FStaticTerrainLayerWeightParameter& LayerParameter = StaticParameters.TerrainLayerWeightParameters[LayerParameterIdx];
 				LayerParameter.WeightmapIndex = INDEX_NONE;
 
 				// Look through our allocations to see if we need this layer.
 				// If not found, this component doesn't use the layer, and WeightmapIndex remains as INDEX_NONE.
-				for( int32 AllocIdx=0;AllocIdx<WeightmapLayerAllocations.Num();AllocIdx++ )
+				for (int32 AllocIdx = 0; AllocIdx < WeightmapLayerAllocations.Num(); AllocIdx++)
 				{
 					FWeightmapLayerAllocationInfo& Allocation = WeightmapLayerAllocations[AllocIdx];
 					if (Allocation.LayerInfo != NULL)
 					{
-						FName ThisLayerName = (Allocation.LayerInfo == ALandscapeProxy::DataLayer) ? UMaterialExpressionLandscapeVisibilityMask::ParameterName : Allocation.LayerInfo->LayerName;
+						FName ThisLayerName = (Allocation.LayerInfo == ALandscapeProxy::VisibilityLayer) ? UMaterialExpressionLandscapeVisibilityMask::ParameterName : Allocation.LayerInfo->LayerName;
 						if (ThisLayerName == LayerParameter.ParameterName)
 						{
 							LayerParameter.WeightmapIndex = Allocation.WeightmapTextureIndex;
@@ -220,10 +220,10 @@ void ULandscapeComponent::UpdateMaterialInstances()
 	// Find or set a matching MIC in the Landscape's map.
 	UMaterialInstanceConstant* CombinationMaterialInstance = GetCombinationMaterial(false);
 
-	if( CombinationMaterialInstance != NULL )
+	if (CombinationMaterialInstance != NULL)
 	{
 		// Create the instance for this component, that will use the layer combination instance.
-		if( MaterialInstance == NULL || GetOutermost() != MaterialInstance->GetOutermost() )
+		if (MaterialInstance == NULL || GetOutermost() != MaterialInstance->GetOutermost())
 		{
 			MaterialInstance = ConstructObject<ULandscapeMaterialInstanceConstant>(ULandscapeMaterialInstanceConstant::StaticClass(), GetOutermost());
 		}
@@ -235,25 +235,25 @@ void ULandscapeComponent::UpdateMaterialInstances()
 		MaterialInstance->SetParentEditorOnly(CombinationMaterialInstance);
 
 		FLinearColor Masks[4];
-		Masks[0] = FLinearColor(1.f,0.f,0.f,0.f);
-		Masks[1] = FLinearColor(0.f,1.f,0.f,0.f);
-		Masks[2] = FLinearColor(0.f,0.f,1.f,0.f);
-		Masks[3] = FLinearColor(0.f,0.f,0.f,1.f);
+		Masks[0] = FLinearColor(1.f, 0.f, 0.f, 0.f);
+		Masks[1] = FLinearColor(0.f, 1.f, 0.f, 0.f);
+		Masks[2] = FLinearColor(0.f, 0.f, 1.f, 0.f);
+		Masks[3] = FLinearColor(0.f, 0.f, 0.f, 1.f);
 
 		// Set the layer mask
-		for( int32 AllocIdx=0;AllocIdx<WeightmapLayerAllocations.Num();AllocIdx++ )
+		for (int32 AllocIdx = 0; AllocIdx < WeightmapLayerAllocations.Num(); AllocIdx++)
 		{
 			FWeightmapLayerAllocationInfo& Allocation = WeightmapLayerAllocations[AllocIdx];
 
-			FName LayerName = Allocation.LayerInfo == ALandscapeProxy::DataLayer ? UMaterialExpressionLandscapeVisibilityMask::ParameterName : Allocation.LayerInfo ? Allocation.LayerInfo->LayerName : NAME_None;
-			MaterialInstance->SetVectorParameterValueEditorOnly(FName(*FString::Printf(TEXT("LayerMask_%s"),*LayerName.ToString())), Masks[Allocation.WeightmapTextureChannel]);
+			FName LayerName = Allocation.LayerInfo == ALandscapeProxy::VisibilityLayer ? UMaterialExpressionLandscapeVisibilityMask::ParameterName : Allocation.LayerInfo ? Allocation.LayerInfo->LayerName : NAME_None;
+			MaterialInstance->SetVectorParameterValueEditorOnly(FName(*FString::Printf(TEXT("LayerMask_%s"), *LayerName.ToString())), Masks[Allocation.WeightmapTextureChannel]);
 		}
 
 		// Set the weightmaps
-		for( int32 i=0;i<WeightmapTextures.Num();i++ )
+		for (int32 i = 0; i < WeightmapTextures.Num(); i++)
 		{
 			// UE_LOG(LogLandscape, Log, TEXT("Setting Weightmap%d = %s"), i, *WeightmapTextures(i)->GetName());
-			MaterialInstance->SetTextureParameterValueEditorOnly(FName(*FString::Printf(TEXT("Weightmap%d"),i)), WeightmapTextures[i]);
+			MaterialInstance->SetTextureParameterValueEditorOnly(FName(*FString::Printf(TEXT("Weightmap%d"), i)), WeightmapTextures[i]);
 		}
 		// Set the heightmap, if needed.
 
@@ -308,7 +308,7 @@ void ULandscapeComponent::PostEditUndo()
 
 void ULandscapeComponent::FixupWeightmaps()
 {
-	if( GIsEditor && !HasAnyFlags(RF_ClassDefaultObject) )
+	if (GIsEditor && !HasAnyFlags(RF_ClassDefaultObject))
 	{
 		ULandscapeInfo* Info = GetLandscapeInfo();
 		ALandscapeProxy* Proxy = GetLandscapeProxy();
@@ -321,17 +321,17 @@ void ULandscapeComponent::FixupWeightmaps()
 			if (Info->Layers.Num() && Cast<ALandscape>(Proxy))
 			{
 				// LayerName Validation check...
-				for( int32 LayerIdx=0;LayerIdx < WeightmapLayerAllocations.Num();LayerIdx++ )
+				for (int32 LayerIdx = 0; LayerIdx < WeightmapLayerAllocations.Num(); LayerIdx++)
 				{
 					if (!WeightmapLayerAllocations[LayerIdx].LayerInfo
-						|| (WeightmapLayerAllocations[LayerIdx].LayerInfo != ALandscapeProxy::DataLayer && Info->GetLayerInfoIndex(WeightmapLayerAllocations[LayerIdx].LayerInfo) == INDEX_NONE))
+						|| (WeightmapLayerAllocations[LayerIdx].LayerInfo != ALandscapeProxy::VisibilityLayer && Info->GetLayerInfoIndex(WeightmapLayerAllocations[LayerIdx].LayerInfo) == INDEX_NONE))
 					{
 						if (!bFixedLayerDeletion)
 						{
 							FFormatNamedArguments Arguments;
 							Arguments.Add(TEXT("LandscapeName"), FText::FromString(GetName()));
 							FMessageLog("MapCheck").Warning()
-								->AddToken(FTextToken::Create(FText::Format( LOCTEXT( "MapCheck_Message_FixedUpDeletedLayerWeightmap", "{LandscapeName} : Fixed up deleted layer weightmap" ), Arguments ) ))
+								->AddToken(FTextToken::Create(FText::Format(LOCTEXT("MapCheck_Message_FixedUpDeletedLayerWeightmap", "{LandscapeName} : Fixed up deleted layer weightmap"), Arguments)))
 								->AddToken(FMapErrorToken::Create(FMapErrors::FixedUpDeletedLayerWeightmap));
 						}
 
@@ -350,23 +350,23 @@ void ULandscapeComponent::FixupWeightmaps()
 				}
 			}
 
-			bool bFixedWeightmapTextureIndex=false;
+			bool bFixedWeightmapTextureIndex = false;
 
 			// Store the weightmap allocations in WeightmapUsageMap
-			for( int32 LayerIdx=0;LayerIdx < WeightmapLayerAllocations.Num();LayerIdx++ )
+			for (int32 LayerIdx = 0; LayerIdx < WeightmapLayerAllocations.Num(); LayerIdx++)
 			{
 				FWeightmapLayerAllocationInfo& Allocation = WeightmapLayerAllocations[LayerIdx];
 
 				// Fix up any problems caused by the layer deletion bug.
-				if( Allocation.WeightmapTextureIndex >= WeightmapTextures.Num() )
+				if (Allocation.WeightmapTextureIndex >= WeightmapTextures.Num())
 				{
-					Allocation.WeightmapTextureIndex = WeightmapTextures.Num()-1; 
-					if( !bFixedWeightmapTextureIndex )
+					Allocation.WeightmapTextureIndex = WeightmapTextures.Num() - 1;
+					if (!bFixedWeightmapTextureIndex)
 					{
 						FFormatNamedArguments Arguments;
 						Arguments.Add(TEXT("LandscapeName"), FText::FromString(GetName()));
 						FMessageLog("MapCheck").Warning()
-							->AddToken(FTextToken::Create(FText::Format( LOCTEXT( "MapCheck_Message_FixedUpIncorrectLayerWeightmap", "{LandscapeName} : Fixed up incorrect layer weightmap texture index" ), Arguments ) ))
+							->AddToken(FTextToken::Create(FText::Format(LOCTEXT("MapCheck_Message_FixedUpIncorrectLayerWeightmap", "{LandscapeName} : Fixed up incorrect layer weightmap texture index"), Arguments)))
 							->AddToken(FMapErrorToken::Create(FMapErrors::FixedUpIncorrectLayerWeightmap));
 					}
 					bFixedWeightmapTextureIndex = true;
@@ -376,15 +376,15 @@ void ULandscapeComponent::FixupWeightmaps()
 				FLandscapeWeightmapUsage& Usage = Proxy->WeightmapUsageMap.FindOrAdd(WeightmapTexture);
 
 				// Detect a shared layer allocation, caused by a previous undo or layer deletion bugs
-				if( Usage.ChannelUsage[Allocation.WeightmapTextureChannel] != NULL &&
-					Usage.ChannelUsage[Allocation.WeightmapTextureChannel] != this )
+				if (Usage.ChannelUsage[Allocation.WeightmapTextureChannel] != NULL &&
+					Usage.ChannelUsage[Allocation.WeightmapTextureChannel] != this)
 				{
 					FFormatNamedArguments Arguments;
 					Arguments.Add(TEXT("LayerName"), FText::FromString(Allocation.GetLayerName().ToString()));
 					Arguments.Add(TEXT("LandscapeName"), FText::FromString(GetName()));
 					Arguments.Add(TEXT("ChannelName"), FText::FromString(Usage.ChannelUsage[Allocation.WeightmapTextureChannel]->GetName()));
 					FMessageLog("MapCheck").Warning()
-						->AddToken(FTextToken::Create(FText::Format( LOCTEXT( "MapCheck_Message_FixedUpSharedLayerWeightmap", "Fixed up shared weightmap texture for layer {LayerName} in component '{LandscapeName}' (shares with '{ChannelName}')" ), Arguments ) ))
+						->AddToken(FTextToken::Create(FText::Format(LOCTEXT("MapCheck_Message_FixedUpSharedLayerWeightmap", "Fixed up shared weightmap texture for layer {LayerName} in component '{LandscapeName}' (shares with '{ChannelName}')"), Arguments)))
 						->AddToken(FMapErrorToken::Create(FMapErrors::FixedUpSharedLayerWeightmap));
 					WeightmapLayerAllocations.RemoveAt(LayerIdx);
 					LayerIdx--;
@@ -399,12 +399,12 @@ void ULandscapeComponent::FixupWeightmaps()
 			RemoveInvalidWeightmaps();
 
 			// Store the layer combination in the MaterialInstanceConstantMap
-			if( MaterialInstance != NULL )
+			if (MaterialInstance != NULL)
 			{
 				UMaterialInstanceConstant* CombinationMaterialInstance = Cast<UMaterialInstanceConstant>(MaterialInstance->Parent);
-				if( CombinationMaterialInstance )
+				if (CombinationMaterialInstance)
 				{
-					Proxy->MaterialInstanceConstantMap.Add(*GetLayerAllocationKey(),CombinationMaterialInstance);
+					Proxy->MaterialInstanceConstantMap.Add(*GetLayerAllocationKey(), CombinationMaterialInstance);
 				}
 			}
 		}
@@ -420,18 +420,18 @@ struct FLandscapeComponentAlphaInfo
 	TArray<uint8> AlphaValues;
 
 	// tor
-	FLandscapeComponentAlphaInfo( ULandscapeComponent* InOwner, int32 InLayerIndex )
-	:	LayerIndex(InLayerIndex)
+	FLandscapeComponentAlphaInfo(ULandscapeComponent* InOwner, int32 InLayerIndex)
+		: LayerIndex(InLayerIndex)
 	{
-		AlphaValues.Empty(FMath::Square(InOwner->ComponentSizeQuads+1));
-		AlphaValues.AddZeroed(FMath::Square(InOwner->ComponentSizeQuads+1));
+		AlphaValues.Empty(FMath::Square(InOwner->ComponentSizeQuads + 1));
+		AlphaValues.AddZeroed(FMath::Square(InOwner->ComponentSizeQuads + 1));
 	}
 
 	bool IsLayerAllZero() const
 	{
-		for( int32 Index=0;Index<AlphaValues.Num();Index++ )
+		for (int32 Index = 0; Index < AlphaValues.Num(); Index++)
 		{
-			if( AlphaValues[Index] != 0 )
+			if (AlphaValues[Index] != 0)
 			{
 				return false;
 			}
@@ -441,11 +441,11 @@ struct FLandscapeComponentAlphaInfo
 };
 
 
-void ULandscapeComponent::UpdateCollisionHeightData(const FColor* HeightmapTextureMipData, int32 ComponentX1, int32 ComponentY1, int32 ComponentX2, int32 ComponentY2, bool bUpdateBounds, const FColor* XYOffsetmapTextureData, bool bRebuild )
+void ULandscapeComponent::UpdateCollisionHeightData(const FColor* HeightmapTextureMipData, int32 ComponentX1, int32 ComponentY1, int32 ComponentX2, int32 ComponentY2, bool bUpdateBounds, const FColor* XYOffsetmapTextureData, bool bRebuild)
 {
 	ULandscapeInfo* Info = GetLandscapeInfo();
 	ALandscapeProxy* Proxy = GetLandscapeProxy();
-	FIntPoint ComponentKey = GetSectionBase()/ComponentSizeQuads;
+	FIntPoint ComponentKey = GetSectionBase() / ComponentSizeQuads;
 	ULandscapeHeightfieldCollisionComponent* CollisionComp = CollisionComponent.Get();
 	ULandscapeMeshCollisionComponent* MeshCollisionComponent = Cast<ULandscapeMeshCollisionComponent>(CollisionComp);
 
@@ -463,9 +463,9 @@ void ULandscapeComponent::UpdateCollisionHeightData(const FColor* HeightmapTextu
 		}
 	}
 
-	int32 CollisionSubsectionSizeVerts = ((SubsectionSizeQuads+1)>>CollisionMipLevel);
-	int32 CollisionSubsectionSizeQuads = CollisionSubsectionSizeVerts-1;
-	int32 CollisionSizeVerts = NumSubsections*CollisionSubsectionSizeQuads+1;
+	int32 CollisionSubsectionSizeVerts = ((SubsectionSizeQuads + 1) >> CollisionMipLevel);
+	int32 CollisionSubsectionSizeQuads = CollisionSubsectionSizeVerts - 1;
+	int32 CollisionSizeVerts = NumSubsections*CollisionSubsectionSizeQuads + 1;
 
 	uint16* CollisionHeightData = NULL;
 	uint16* CollisionXYOffsetData = NULL;
@@ -480,9 +480,9 @@ void ULandscapeComponent::UpdateCollisionHeightData(const FColor* HeightmapTextu
 	}
 
 	// Existing collision component is same type with collision
-	if ( CollisionComp && ((XYOffsetmapTexture == NULL) == (MeshCollisionComponent == NULL)) )
+	if (CollisionComp && ((XYOffsetmapTexture == NULL) == (MeshCollisionComponent == NULL)))
 	{
-		if( bUpdateBounds )
+		if (bUpdateBounds)
 		{
 			CollisionComp->CachedLocalBox = CachedLocalBox;
 			CollisionComp->UpdateComponentToWorld();
@@ -511,7 +511,7 @@ void ULandscapeComponent::UpdateCollisionHeightData(const FColor* HeightmapTextu
 				DominantLayerData.AddUninitialized(FMath::Square(CollisionSizeVerts));
 
 				const uint8* SrcDominantLayerData = (uint8*)CollisionComp->DominantLayerData.Lock(LOCK_READ_ONLY);
-				FMemory::Memcpy( DominantLayerData.GetTypedData(), SrcDominantLayerData, sizeof(uint8)*FMath::Square(CollisionSizeVerts) );
+				FMemory::Memcpy(DominantLayerData.GetTypedData(), SrcDominantLayerData, sizeof(uint8)*FMath::Square(CollisionSizeVerts));
 				CollisionComp->DominantLayerData.Unlock();
 			}
 
@@ -529,14 +529,14 @@ void ULandscapeComponent::UpdateCollisionHeightData(const FColor* HeightmapTextu
 			CollisionComp = NULL;
 		}
 
-		MeshCollisionComponent = XYOffsetmapTexture ? ConstructObject<ULandscapeMeshCollisionComponent>(ULandscapeMeshCollisionComponent::StaticClass(), Proxy,NAME_None,RF_Transactional) : NULL;
+		MeshCollisionComponent = XYOffsetmapTexture ? ConstructObject<ULandscapeMeshCollisionComponent>(ULandscapeMeshCollisionComponent::StaticClass(), Proxy, NAME_None, RF_Transactional) : NULL;
 		CollisionComp = XYOffsetmapTexture ? Cast<ULandscapeHeightfieldCollisionComponent>(MeshCollisionComponent) :
-			ConstructObject<ULandscapeHeightfieldCollisionComponent>(ULandscapeHeightfieldCollisionComponent::StaticClass(), Proxy,NAME_None,RF_Transactional);
+			ConstructObject<ULandscapeHeightfieldCollisionComponent>(ULandscapeHeightfieldCollisionComponent::StaticClass(), Proxy, NAME_None, RF_Transactional);
 
 		CollisionComp->SetRelativeLocation(RelativeLocation);
-		CollisionComp->AttachTo( Proxy->GetRootComponent(), NAME_None );
+		CollisionComp->AttachTo(Proxy->GetRootComponent(), NAME_None);
 		Proxy->CollisionComponents.Add(CollisionComp);
-		
+
 		CollisionComp->RenderComponent = this;
 		CollisionComp->SetSectionBase(GetSectionBase());
 		CollisionComp->CollisionSizeQuads = CollisionSubsectionSizeQuads * NumSubsections;
@@ -547,20 +547,20 @@ void ULandscapeComponent::UpdateCollisionHeightData(const FColor* HeightmapTextu
 		// Reallocate raw collision data
 		CollisionComp->CollisionHeightData.Lock(LOCK_READ_WRITE);
 		CollisionHeightData = (uint16*)CollisionComp->CollisionHeightData.Realloc(FMath::Square(CollisionSizeVerts));
-		FMemory::Memzero( CollisionHeightData,sizeof(uint16)*FMath::Square(CollisionSizeVerts));
+		FMemory::Memzero(CollisionHeightData, sizeof(uint16)*FMath::Square(CollisionSizeVerts));
 
 		if (XYOffsetmapTexture && MeshCollisionComponent)
 		{
 			// Need XYOffsetData for Collision Component
 			MeshCollisionComponent->CollisionXYOffsetData.Lock(LOCK_READ_WRITE);
 			CollisionXYOffsetData = (uint16*)MeshCollisionComponent->CollisionXYOffsetData.Realloc(FMath::Square(CollisionSizeVerts) * 2);
-			FMemory::Memzero( CollisionXYOffsetData,sizeof(uint16)*FMath::Square(CollisionSizeVerts) * 2);
+			FMemory::Memzero(CollisionXYOffsetData, sizeof(uint16)*FMath::Square(CollisionSizeVerts) * 2);
 
 			if (DominantLayerData.Num())
 			{
 				MeshCollisionComponent->DominantLayerData.Lock(LOCK_READ_WRITE);
 				uint8* DestDominantLayerData = (uint8*)MeshCollisionComponent->DominantLayerData.Realloc(FMath::Square(CollisionSizeVerts));
-				FMemory::Memcpy( DestDominantLayerData, DominantLayerData.GetTypedData(), sizeof(uint8)*FMath::Square(CollisionSizeVerts) );
+				FMemory::Memcpy(DestDominantLayerData, DominantLayerData.GetTypedData(), sizeof(uint8)*FMath::Square(CollisionSizeVerts));
 				MeshCollisionComponent->DominantLayerData.Unlock();
 			}
 
@@ -589,20 +589,20 @@ void ULandscapeComponent::UpdateCollisionHeightData(const FColor* HeightmapTextu
 	//int32 WeightmapOffsetX = FMath::RoundToInt(WeightmapScaleBias.Z * (float)XYMipSizeU) >> CollisionMipLevel;
 	//int32 WeightmapOffsetY = FMath::RoundToInt(WeightmapScaleBias.W * (float)XYMipSizeV) >> CollisionMipLevel;
 
-	for( int32 SubsectionY = 0;SubsectionY < NumSubsections;SubsectionY++ )
+	for (int32 SubsectionY = 0; SubsectionY < NumSubsections; SubsectionY++)
 	{
 		// Check if subsection is fully above or below the area we are interested in
-		if( (ComponentY2 < SubsectionSizeQuads*SubsectionY) ||		// above
-			(ComponentY1 > SubsectionSizeQuads*(SubsectionY+1)) )	// below
+		if ((ComponentY2 < SubsectionSizeQuads*SubsectionY) ||		// above
+			(ComponentY1 > SubsectionSizeQuads*(SubsectionY + 1)))	// below
 		{
 			continue;
 		}
 
-		for( int32 SubsectionX = 0;SubsectionX < NumSubsections;SubsectionX++ )
+		for (int32 SubsectionX = 0; SubsectionX < NumSubsections; SubsectionX++)
 		{
 			// Check if subsection is fully to the left or right of the area we are interested in
-			if( (ComponentX2 < SubsectionSizeQuads*SubsectionX) ||		// left
-				(ComponentX1 > SubsectionSizeQuads*(SubsectionX+1)) )	// right
+			if ((ComponentX2 < SubsectionSizeQuads*SubsectionX) ||		// left
+				(ComponentX1 > SubsectionSizeQuads*(SubsectionX + 1)))	// right
 			{
 				continue;
 			}
@@ -614,10 +614,10 @@ void ULandscapeComponent::UpdateCollisionHeightData(const FColor* HeightmapTextu
 			int32 SubY2 = ComponentY2 - SubsectionSizeQuads*SubsectionY;
 
 			// Area to update in collision mip level coords
-			int32 CollisionSubX1 = FMath::FloorToInt( (float)SubX1 * CollisionQuadRatio );
-			int32 CollisionSubY1 = FMath::FloorToInt( (float)SubY1 * CollisionQuadRatio );
-			int32 CollisionSubX2 = FMath::CeilToInt(  (float)SubX2 * CollisionQuadRatio );
-			int32 CollisionSubY2 = FMath::CeilToInt(  (float)SubY2 * CollisionQuadRatio );
+			int32 CollisionSubX1 = FMath::FloorToInt((float)SubX1 * CollisionQuadRatio);
+			int32 CollisionSubY1 = FMath::FloorToInt((float)SubY1 * CollisionQuadRatio);
+			int32 CollisionSubX2 = FMath::CeilToInt((float)SubX2 * CollisionQuadRatio);
+			int32 CollisionSubY2 = FMath::CeilToInt((float)SubY2 * CollisionQuadRatio);
 
 			// Clamp area to update
 			int32 VertX1 = FMath::Clamp<int32>(CollisionSubX1, 0, CollisionSubsectionSizeQuads);
@@ -625,9 +625,9 @@ void ULandscapeComponent::UpdateCollisionHeightData(const FColor* HeightmapTextu
 			int32 VertX2 = FMath::Clamp<int32>(CollisionSubX2, 0, CollisionSubsectionSizeQuads);
 			int32 VertY2 = FMath::Clamp<int32>(CollisionSubY2, 0, CollisionSubsectionSizeQuads);
 
-			for( int32 VertY=VertY1;VertY<=VertY2;VertY++ )
+			for (int32 VertY = VertY1; VertY <= VertY2; VertY++)
 			{
-				for( int32 VertX=VertX1;VertX<=VertX2;VertX++ )
+				for (int32 VertX = VertX1; VertX <= VertX2; VertX++)
 				{
 					{
 						// X/Y of the vertex we're looking indexed into the texture data
@@ -651,19 +651,19 @@ void ULandscapeComponent::UpdateCollisionHeightData(const FColor* HeightmapTextu
 						// X/Y of the vertex we're looking indexed into the texture data
 						int32 TexX = CollisionSubsectionSizeVerts * SubsectionX + VertX;
 						int32 TexY = CollisionSubsectionSizeVerts * SubsectionY + VertY;
-						const FColor& TexData = XYOffsetmapTextureData[ TexX + TexY * XYMipSizeU ];
+						const FColor& TexData = XYOffsetmapTextureData[TexX + TexY * XYMipSizeU];
 
 						// this uses Quads as we don't want the duplicated vertices
 						int32 CompVertX = CollisionSubsectionSizeQuads * SubsectionX + VertX;
 						int32 CompVertY = CollisionSubsectionSizeQuads * SubsectionY + VertY;
 
 						// Copy collision data
-						uint16 NewXOffset = TexData.R<<8 | TexData.G;
-						uint16 NewYOffset = TexData.B<<8 | TexData.A;
+						uint16 NewXOffset = TexData.R << 8 | TexData.G;
+						uint16 NewYOffset = TexData.B << 8 | TexData.A;
 
-						int32 XYIndex = CompVertX+CompVertY*CollisionSizeVerts;
-						CollisionXYOffsetData[XYIndex*2] = NewXOffset;
-						CollisionXYOffsetData[XYIndex*2 + 1] = NewYOffset;
+						int32 XYIndex = CompVertX + CompVertY*CollisionSizeVerts;
+						CollisionXYOffsetData[XYIndex * 2] = NewXOffset;
+						CollisionXYOffsetData[XYIndex * 2 + 1] = NewYOffset;
 					}
 				}
 			}
@@ -678,7 +678,7 @@ void ULandscapeComponent::UpdateCollisionHeightData(const FColor* HeightmapTextu
 	}
 
 	// If we updated an existing component, we need to update the PhysX copy of the data
-	if( !CreatedNew )
+	if (!CreatedNew)
 	{
 		if (MeshCollisionComponent)
 		{
@@ -691,20 +691,20 @@ void ULandscapeComponent::UpdateCollisionHeightData(const FColor* HeightmapTextu
 		}
 		else
 		{
-			int32 CollisionCompX1 = FMath::FloorToInt( (float)ComponentX1 * CollisionQuadRatio );
-			int32 CollisionCompY1 = FMath::FloorToInt( (float)ComponentY1 * CollisionQuadRatio );
-			int32 CollisionCompX2 = FMath::CeilToInt(  (float)ComponentX2 * CollisionQuadRatio );
-			int32 CollisionCompY2 = FMath::CeilToInt(  (float)ComponentY2 * CollisionQuadRatio );
+			int32 CollisionCompX1 = FMath::FloorToInt((float)ComponentX1 * CollisionQuadRatio);
+			int32 CollisionCompY1 = FMath::FloorToInt((float)ComponentY1 * CollisionQuadRatio);
+			int32 CollisionCompX2 = FMath::CeilToInt((float)ComponentX2 * CollisionQuadRatio);
+			int32 CollisionCompY2 = FMath::CeilToInt((float)ComponentY2 * CollisionQuadRatio);
 			CollisionComp->UpdateHeightfieldRegion(CollisionCompX1, CollisionCompY1, CollisionCompX2, CollisionCompY2);
 		}
 	}
 
-	if ( bRebuild && CollisionProxy )
+	if (bRebuild && CollisionProxy)
 	{
 		CollisionProxy->RegisterAllComponents();
 
 		// Move any foliage instances if we created a new collision component.
-		if( OldCollisionComponent && OldCollisionComponent != CollisionComp )
+		if (OldCollisionComponent && OldCollisionComponent != CollisionComp)
 		{
 			AInstancedFoliageActor* IFA = AInstancedFoliageActor::GetInstancedFoliageActorForCurrentLevel(OldCollisionComponent->GetWorld());
 			IFA->MoveInstancesToNewComponent(OldCollisionComponent, CollisionComp);
@@ -713,7 +713,7 @@ void ULandscapeComponent::UpdateCollisionHeightData(const FColor* HeightmapTextu
 
 	// Set new collision component to pointer
 	CollisionComponent = CollisionComp;
-	
+
 	if (ChangeType && CreatedNew)
 	{
 		Proxy->RegisterAllComponents();
@@ -721,13 +721,13 @@ void ULandscapeComponent::UpdateCollisionHeightData(const FColor* HeightmapTextu
 }
 
 
-void ULandscapeComponent::UpdateCollisionLayerData(TArray<FColor*>& WeightmapTextureMipData, int32 ComponentX1, int32 ComponentY1, int32 ComponentX2, int32 ComponentY2 )
+void ULandscapeComponent::UpdateCollisionLayerData(TArray<FColor*>& WeightmapTextureMipData, int32 ComponentX1, int32 ComponentY1, int32 ComponentX2, int32 ComponentY2)
 {
 	ULandscapeInfo* Info = GetLandscapeInfo();
 	ALandscapeProxy* Proxy = GetLandscapeProxy();
-	FIntPoint ComponentKey = GetSectionBase()/ComponentSizeQuads;
+	FIntPoint ComponentKey = GetSectionBase() / ComponentSizeQuads;
 
-	if( CollisionComponent )
+	if (CollisionComponent)
 	{
 		CollisionComponent->Modify();
 
@@ -735,36 +735,36 @@ void ULandscapeComponent::UpdateCollisionLayerData(TArray<FColor*>& WeightmapTex
 		TArray<uint8*> CandidateDataPtrs;
 
 		// Channel remapping
-		int32 ChannelOffsets[4] = {(int32)STRUCT_OFFSET(FColor,R),(int32)STRUCT_OFFSET(FColor,G),(int32)STRUCT_OFFSET(FColor,B),(int32)STRUCT_OFFSET(FColor,A)};
+		int32 ChannelOffsets[4] = { (int32)STRUCT_OFFSET(FColor, R), (int32)STRUCT_OFFSET(FColor, G), (int32)STRUCT_OFFSET(FColor, B), (int32)STRUCT_OFFSET(FColor, A) };
 
 		bool bExistingLayerMismatch = false;
 		int32 DataLayerIdx = INDEX_NONE;
 
 		// Find the layers we're interested in
-		for( int32 AllocIdx=0;AllocIdx<WeightmapLayerAllocations.Num();AllocIdx++ )
+		for (int32 AllocIdx = 0; AllocIdx < WeightmapLayerAllocations.Num(); AllocIdx++)
 		{
 			FWeightmapLayerAllocationInfo& AllocInfo = WeightmapLayerAllocations[AllocIdx];
 			ULandscapeLayerInfoObject* LayerInfo = AllocInfo.LayerInfo;
-			if( AllocInfo.LayerInfo == ALandscapeProxy::DataLayer || (LayerInfo && LayerInfo->PhysMaterial) )
+			if (AllocInfo.LayerInfo == ALandscapeProxy::VisibilityLayer || (LayerInfo && LayerInfo->PhysMaterial))
 			{
 				int32 Idx = CandidateLayers.Add(AllocInfo.LayerInfo);
 				CandidateDataPtrs.Add(((uint8*)WeightmapTextureMipData[AllocInfo.WeightmapTextureIndex]) + ChannelOffsets[AllocInfo.WeightmapTextureChannel]);
 
 				// Check if we still match the collision component.
-				if( !CollisionComponent->ComponentLayerInfos.IsValidIndex(Idx) || CollisionComponent->ComponentLayerInfos[Idx] != AllocInfo.LayerInfo )
+				if (!CollisionComponent->ComponentLayerInfos.IsValidIndex(Idx) || CollisionComponent->ComponentLayerInfos[Idx] != AllocInfo.LayerInfo)
 				{
 					bExistingLayerMismatch = true;
 				}
 
-				if (AllocInfo.LayerInfo == ALandscapeProxy::DataLayer)
+				if (AllocInfo.LayerInfo == ALandscapeProxy::VisibilityLayer)
 				{
 					DataLayerIdx = Idx;
 					bExistingLayerMismatch = true; // always rebuild whole component for hole
 				}
-			}	
+			}
 		}
 
-		if( CandidateLayers.Num() == 0 )
+		if (CandidateLayers.Num() == 0)
 		{
 			// No layers, so don't update any weights
 			CollisionComponent->DominantLayerData.RemoveBulkData();
@@ -772,13 +772,13 @@ void ULandscapeComponent::UpdateCollisionLayerData(TArray<FColor*>& WeightmapTex
 		}
 		else
 		{
-			int32 CollisionSubsectionSizeVerts = ((SubsectionSizeQuads+1)>>CollisionMipLevel);
-			int32 CollisionSubsectionSizeQuads = CollisionSubsectionSizeVerts-1;
-			int32 CollisionSizeVerts = NumSubsections*CollisionSubsectionSizeQuads+1;
+			int32 CollisionSubsectionSizeVerts = ((SubsectionSizeQuads + 1) >> CollisionMipLevel);
+			int32 CollisionSubsectionSizeQuads = CollisionSubsectionSizeVerts - 1;
+			int32 CollisionSizeVerts = NumSubsections*CollisionSubsectionSizeQuads + 1;
 			uint8* DominantLayerData = NULL;
 
 			// If there's no existing data, or the layer allocations have changed, we need to update the data for the whole component.
-			if( bExistingLayerMismatch || CollisionComponent->DominantLayerData.GetElementCount() == 0  )
+			if (bExistingLayerMismatch || CollisionComponent->DominantLayerData.GetElementCount() == 0)
 			{
 				ComponentX1 = 0;
 				ComponentY1 = 0;
@@ -787,7 +787,7 @@ void ULandscapeComponent::UpdateCollisionLayerData(TArray<FColor*>& WeightmapTex
 
 				CollisionComponent->DominantLayerData.Lock(LOCK_READ_WRITE);
 				DominantLayerData = (uint8*)CollisionComponent->DominantLayerData.Realloc(FMath::Square(CollisionSizeVerts));
-				FMemory::Memzero(DominantLayerData,FMath::Square(CollisionSizeVerts));
+				FMemory::Memzero(DominantLayerData, FMath::Square(CollisionSizeVerts));
 
 				CollisionComponent->ComponentLayerInfos = CandidateLayers;
 			}
@@ -801,20 +801,20 @@ void ULandscapeComponent::UpdateCollisionLayerData(TArray<FColor*>& WeightmapTex
 			// Ratio to convert update region coordinate to collision mip coordinates
 			float CollisionQuadRatio = (float)CollisionSubsectionSizeQuads / (float)SubsectionSizeQuads;
 
-			for( int32 SubsectionY = 0;SubsectionY < NumSubsections;SubsectionY++ )
+			for (int32 SubsectionY = 0; SubsectionY < NumSubsections; SubsectionY++)
 			{
 				// Check if subsection is fully above or below the area we are interested in
-				if( (ComponentY2 < SubsectionSizeQuads*SubsectionY) ||		// above
-					(ComponentY1 > SubsectionSizeQuads*(SubsectionY+1)) )	// below
+				if ((ComponentY2 < SubsectionSizeQuads*SubsectionY) ||		// above
+					(ComponentY1 > SubsectionSizeQuads*(SubsectionY + 1)))	// below
 				{
 					continue;
 				}
 
-				for( int32 SubsectionX = 0;SubsectionX < NumSubsections;SubsectionX++ )
+				for (int32 SubsectionX = 0; SubsectionX < NumSubsections; SubsectionX++)
 				{
 					// Check if subsection is fully to the left or right of the area we are interested in
-					if( (ComponentX2 < SubsectionSizeQuads*SubsectionX) ||		// left
-						(ComponentX1 > SubsectionSizeQuads*(SubsectionX+1)) )	// right
+					if ((ComponentX2 < SubsectionSizeQuads*SubsectionX) ||		// left
+						(ComponentX1 > SubsectionSizeQuads*(SubsectionX + 1)))	// right
 					{
 						continue;
 					}
@@ -826,10 +826,10 @@ void ULandscapeComponent::UpdateCollisionLayerData(TArray<FColor*>& WeightmapTex
 					int32 SubY2 = ComponentY2 - SubsectionSizeQuads*SubsectionY;
 
 					// Area to update in collision mip level coords
-					int32 CollisionSubX1 = FMath::FloorToInt( (float)SubX1 * CollisionQuadRatio );
-					int32 CollisionSubY1 = FMath::FloorToInt( (float)SubY1 * CollisionQuadRatio );
-					int32 CollisionSubX2 = FMath::CeilToInt(  (float)SubX2 * CollisionQuadRatio );
-					int32 CollisionSubY2 = FMath::CeilToInt(  (float)SubY2 * CollisionQuadRatio );
+					int32 CollisionSubX1 = FMath::FloorToInt((float)SubX1 * CollisionQuadRatio);
+					int32 CollisionSubY1 = FMath::FloorToInt((float)SubY1 * CollisionQuadRatio);
+					int32 CollisionSubX2 = FMath::CeilToInt((float)SubX2 * CollisionQuadRatio);
+					int32 CollisionSubY2 = FMath::CeilToInt((float)SubY2 * CollisionQuadRatio);
 
 					// Clamp area to update
 					int32 VertX1 = FMath::Clamp<int32>(CollisionSubX1, 0, CollisionSubsectionSizeQuads);
@@ -837,18 +837,18 @@ void ULandscapeComponent::UpdateCollisionLayerData(TArray<FColor*>& WeightmapTex
 					int32 VertX2 = FMath::Clamp<int32>(CollisionSubX2, 0, CollisionSubsectionSizeQuads);
 					int32 VertY2 = FMath::Clamp<int32>(CollisionSubY2, 0, CollisionSubsectionSizeQuads);
 
-					for( int32 VertY=VertY1;VertY<=VertY2;VertY++ )
+					for (int32 VertY = VertY1; VertY <= VertY2; VertY++)
 					{
-						for( int32 VertX=VertX1;VertX<=VertX2;VertX++ )
+						for (int32 VertX = VertX1; VertX <= VertX2; VertX++)
 						{
 							// X/Y of the vertex we're looking indexed into the texture data
 							int32 TexX = CollisionSubsectionSizeVerts * SubsectionX + VertX;
 							int32 TexY = CollisionSubsectionSizeVerts * SubsectionY + VertY;
 							int32 DataOffset = (TexX + TexY * MipSizeU) * sizeof(FColor);
-					
+
 							uint8 DominantLayer = 255; // 255 as invalid value
 							int32 DominantWeight = 0;
-							for( int32 LayerIdx=0;LayerIdx<CandidateDataPtrs.Num();LayerIdx++ )
+							for (int32 LayerIdx = 0; LayerIdx < CandidateDataPtrs.Num(); LayerIdx++)
 							{
 								const uint8 LayerWeight = CandidateDataPtrs[LayerIdx][DataOffset];
 
@@ -860,7 +860,7 @@ void ULandscapeComponent::UpdateCollisionLayerData(TArray<FColor*>& WeightmapTex
 										DominantWeight = INT_MAX;
 									}
 								}
-								else if( LayerWeight > DominantWeight )
+								else if (LayerWeight > DominantWeight)
 								{
 									DominantLayer = LayerIdx;
 									DominantWeight = LayerWeight;
@@ -872,7 +872,7 @@ void ULandscapeComponent::UpdateCollisionLayerData(TArray<FColor*>& WeightmapTex
 							int32 CompVertY = CollisionSubsectionSizeQuads * SubsectionY + VertY;
 
 							// Set collision data
-							DominantLayerData[CompVertX+CompVertY*CollisionSizeVerts] = DominantLayer;
+							DominantLayerData[CompVertX + CompVertY*CollisionSizeVerts] = DominantLayer;
 						}
 					}
 				}
@@ -894,9 +894,9 @@ void ULandscapeComponent::UpdateCollisionLayerData()
 
 	WeightmapTextureMipData.Empty(WeightmapTextures.Num());
 	CachedWeightmapTextureMipData.Empty(WeightmapTextures.Num());
-	for( int32 Idx=0;Idx<WeightmapTextures.Num();Idx++ )
+	for (int32 Idx = 0; Idx < WeightmapTextures.Num(); Idx++)
 	{
-		TArray<uint8>* MipData = new(CachedWeightmapTextureMipData) TArray<uint8>();
+		TArray<uint8>* MipData = new(CachedWeightmapTextureMipData)TArray<uint8>();
 		WeightmapTextures[Idx]->Source.GetMipData(*MipData, CollisionMipLevel);
 		WeightmapTextureMipData.Add((FColor*)MipData->GetTypedData());
 	}
@@ -907,9 +907,9 @@ void ULandscapeComponent::UpdateCollisionLayerData()
 
 
 
-void ULandscapeComponent::GenerateHeightmapMips( TArray<FColor*>& HeightmapTextureMipData, int32 ComponentX1/*=0*/, int32 ComponentY1/*=0*/, int32 ComponentX2/*=MAX_int32*/, int32 ComponentY2/*=MAX_int32*/, struct FLandscapeTextureDataInfo* TextureDataInfo/*=NULL*/ )
+void ULandscapeComponent::GenerateHeightmapMips(TArray<FColor*>& HeightmapTextureMipData, int32 ComponentX1/*=0*/, int32 ComponentY1/*=0*/, int32 ComponentX2/*=MAX_int32*/, int32 ComponentY2/*=MAX_int32*/, struct FLandscapeTextureDataInfo* TextureDataInfo/*=NULL*/)
 {
-	bool EndX = false; 
+	bool EndX = false;
 	bool EndY = false;
 
 	if (ComponentX1 == MAX_int32)
@@ -924,11 +924,11 @@ void ULandscapeComponent::GenerateHeightmapMips( TArray<FColor*>& HeightmapTextu
 		ComponentY1 = 0;
 	}
 
-	if( ComponentX2==MAX_int32 )
+	if (ComponentX2 == MAX_int32)
 	{
 		ComponentX2 = ComponentSizeQuads;
 	}
-	if( ComponentY2==MAX_int32 )
+	if (ComponentY2 == MAX_int32)
 	{
 		ComponentY2 = ComponentSizeQuads;
 	}
@@ -939,20 +939,20 @@ void ULandscapeComponent::GenerateHeightmapMips( TArray<FColor*>& HeightmapTextu
 	int32 HeightmapOffsetX = FMath::RoundToInt(HeightmapScaleBias.Z * (float)HeightmapSizeU);
 	int32 HeightmapOffsetY = FMath::RoundToInt(HeightmapScaleBias.W * (float)HeightmapSizeV);
 
-	for( int32 SubsectionY = 0;SubsectionY < NumSubsections;SubsectionY++ )
+	for (int32 SubsectionY = 0; SubsectionY < NumSubsections; SubsectionY++)
 	{
 		// Check if subsection is fully above or below the area we are interested in
-		if( (ComponentY2 < SubsectionSizeQuads*SubsectionY) ||		// above
-			(ComponentY1 > SubsectionSizeQuads*(SubsectionY+1)) )	// below
+		if ((ComponentY2 < SubsectionSizeQuads*SubsectionY) ||		// above
+			(ComponentY1 > SubsectionSizeQuads*(SubsectionY + 1)))	// below
 		{
 			continue;
 		}
 
-		for( int32 SubsectionX = 0;SubsectionX < NumSubsections;SubsectionX++ )
+		for (int32 SubsectionX = 0; SubsectionX < NumSubsections; SubsectionX++)
 		{
 			// Check if subsection is fully to the left or right of the area we are interested in
-			if( (ComponentX2 < SubsectionSizeQuads*SubsectionX) ||		// left
-				(ComponentX1 > SubsectionSizeQuads*(SubsectionX+1)) )	// right
+			if ((ComponentX2 < SubsectionSizeQuads*SubsectionX) ||		// left
+				(ComponentX1 > SubsectionSizeQuads*(SubsectionX + 1)))	// right
 			{
 				continue;
 			}
@@ -972,22 +972,22 @@ void ULandscapeComponent::GenerateHeightmapMips( TArray<FColor*>& HeightmapTextu
 			int32 PrevMipHeightmapOffsetX = HeightmapOffsetX;
 			int32 PrevMipHeightmapOffsetY = HeightmapOffsetY;
 
-			for( int32 Mip=1;Mip<HeightmapTextureMipData.Num();Mip++ )
+			for (int32 Mip = 1; Mip < HeightmapTextureMipData.Num(); Mip++)
 			{
 				int32 MipSizeU = HeightmapSizeU >> Mip;
 				int32 MipSizeV = HeightmapSizeV >> Mip;
 
-				int32 MipSubsectionSizeQuads = ((SubsectionSizeQuads+1)>>Mip)-1;
+				int32 MipSubsectionSizeQuads = ((SubsectionSizeQuads + 1) >> Mip) - 1;
 				float InvMipSubsectionSizeQuads = 1.f / (float)MipSubsectionSizeQuads;
 
-				int32 MipHeightmapOffsetX = HeightmapOffsetX>>Mip;
-				int32 MipHeightmapOffsetY = HeightmapOffsetY>>Mip;
+				int32 MipHeightmapOffsetX = HeightmapOffsetX >> Mip;
+				int32 MipHeightmapOffsetY = HeightmapOffsetY >> Mip;
 
 				// Area to update in current mip level coords
-				int32 MipSubX1 = FMath::FloorToInt( (float)MipSubsectionSizeQuads * (float)PrevMipSubX1 * InvPrevMipSubsectionSizeQuads );
-				int32 MipSubY1 = FMath::FloorToInt( (float)MipSubsectionSizeQuads * (float)PrevMipSubY1 * InvPrevMipSubsectionSizeQuads );
-				int32 MipSubX2 = FMath::CeilToInt(  (float)MipSubsectionSizeQuads * (float)PrevMipSubX2 * InvPrevMipSubsectionSizeQuads );
-				int32 MipSubY2 = FMath::CeilToInt(  (float)MipSubsectionSizeQuads * (float)PrevMipSubY2 * InvPrevMipSubsectionSizeQuads );
+				int32 MipSubX1 = FMath::FloorToInt((float)MipSubsectionSizeQuads * (float)PrevMipSubX1 * InvPrevMipSubsectionSizeQuads);
+				int32 MipSubY1 = FMath::FloorToInt((float)MipSubsectionSizeQuads * (float)PrevMipSubY1 * InvPrevMipSubsectionSizeQuads);
+				int32 MipSubX2 = FMath::CeilToInt((float)MipSubsectionSizeQuads * (float)PrevMipSubX2 * InvPrevMipSubsectionSizeQuads);
+				int32 MipSubY2 = FMath::CeilToInt((float)MipSubsectionSizeQuads * (float)PrevMipSubY2 * InvPrevMipSubsectionSizeQuads);
 
 				// Clamp area to update
 				int32 VertX1 = FMath::Clamp<int32>(MipSubX1, 0, MipSubsectionSizeQuads);
@@ -995,9 +995,9 @@ void ULandscapeComponent::GenerateHeightmapMips( TArray<FColor*>& HeightmapTextu
 				int32 VertX2 = FMath::Clamp<int32>(MipSubX2, 0, MipSubsectionSizeQuads);
 				int32 VertY2 = FMath::Clamp<int32>(MipSubY2, 0, MipSubsectionSizeQuads);
 
-				for( int32 VertY=VertY1;VertY<=VertY2;VertY++ )
+				for (int32 VertY = VertY1; VertY <= VertY2; VertY++)
 				{
-					for( int32 VertX=VertX1;VertX<=VertX2;VertX++ )
+					for (int32 VertX = VertX1; VertX <= VertX2; VertX++)
 					{
 						// Convert VertX/Y into previous mip's coords
 						float PrevMipVertX = (float)PrevMipSubsectionSizeQuads * (float)VertX * InvMipSubsectionSizeQuads;
@@ -1007,67 +1007,67 @@ void ULandscapeComponent::GenerateHeightmapMips( TArray<FColor*>& HeightmapTextu
 						// Validate that the vertex we skip wouldn't use the updated data in the parent mip.
 						// Note this validation is doesn't do anything unless you change the VertY/VertX loops 
 						// above to process all verts from 0 .. MipSubsectionSizeQuads.
-						if( VertX < VertX1 || VertX > VertX2 )
+						if (VertX < VertX1 || VertX > VertX2)
 						{
-							check( FMath::CeilToInt(PrevMipVertX) < PrevMipSubX1 || FMath::FloorToInt(PrevMipVertX) > PrevMipSubX2 );
+							check(FMath::CeilToInt(PrevMipVertX) < PrevMipSubX1 || FMath::FloorToInt(PrevMipVertX) > PrevMipSubX2);
 							continue;
 						}
 
-						if( VertY < VertY1 || VertY > VertY2 )
+						if (VertY < VertY1 || VertY > VertY2)
 						{
-							check( FMath::CeilToInt(PrevMipVertY) < PrevMipSubY1 || FMath::FloorToInt(PrevMipVertY) > PrevMipSubY2 );
+							check(FMath::CeilToInt(PrevMipVertY) < PrevMipSubY1 || FMath::FloorToInt(PrevMipVertY) > PrevMipSubY2);
 							continue;
 						}
 #endif
 
 						// X/Y of the vertex we're looking indexed into the texture data
-						int32 TexX = (MipHeightmapOffsetX) + (MipSubsectionSizeQuads+1) * SubsectionX + VertX;
-						int32 TexY = (MipHeightmapOffsetY) + (MipSubsectionSizeQuads+1) * SubsectionY + VertY;
+						int32 TexX = (MipHeightmapOffsetX)+(MipSubsectionSizeQuads + 1) * SubsectionX + VertX;
+						int32 TexY = (MipHeightmapOffsetY)+(MipSubsectionSizeQuads + 1) * SubsectionY + VertY;
 
-						float fPrevMipTexX = (float)(PrevMipHeightmapOffsetX) + (float)((PrevMipSubsectionSizeQuads+1) * SubsectionX) + PrevMipVertX;
-						float fPrevMipTexY = (float)(PrevMipHeightmapOffsetY) + (float)((PrevMipSubsectionSizeQuads+1) * SubsectionY) + PrevMipVertY;
+						float fPrevMipTexX = (float)(PrevMipHeightmapOffsetX)+(float)((PrevMipSubsectionSizeQuads + 1) * SubsectionX) + PrevMipVertX;
+						float fPrevMipTexY = (float)(PrevMipHeightmapOffsetY)+(float)((PrevMipSubsectionSizeQuads + 1) * SubsectionY) + PrevMipVertY;
 
 						int32 PrevMipTexX = FMath::FloorToInt(fPrevMipTexX);
 						float fPrevMipTexFracX = FMath::Fractional(fPrevMipTexX);
 						int32 PrevMipTexY = FMath::FloorToInt(fPrevMipTexY);
 						float fPrevMipTexFracY = FMath::Fractional(fPrevMipTexY);
 
-						checkSlow( TexX >= 0 && TexX < MipSizeU );
-						checkSlow( TexY >= 0 && TexY < MipSizeV );
-						checkSlow( PrevMipTexX >= 0 && PrevMipTexX < PrevMipSizeU );
-						checkSlow( PrevMipTexY >= 0 && PrevMipTexY < PrevMipSizeV );
+						checkSlow(TexX >= 0 && TexX < MipSizeU);
+						checkSlow(TexY >= 0 && TexY < MipSizeV);
+						checkSlow(PrevMipTexX >= 0 && PrevMipTexX < PrevMipSizeU);
+						checkSlow(PrevMipTexY >= 0 && PrevMipTexY < PrevMipSizeV);
 
-						int32 PrevMipTexX1 = FMath::Min<int32>( PrevMipTexX+1, PrevMipSizeU-1 );
-						int32 PrevMipTexY1 = FMath::Min<int32>( PrevMipTexY+1, PrevMipSizeV-1 );
+						int32 PrevMipTexX1 = FMath::Min<int32>(PrevMipTexX + 1, PrevMipSizeU - 1);
+						int32 PrevMipTexY1 = FMath::Min<int32>(PrevMipTexY + 1, PrevMipSizeV - 1);
 
 						// Padding for missing data for MIP 0
 						if (Mip == 1)
 						{
-							if (EndX && SubsectionX == NumSubsections-1 && VertX == VertX2)
+							if (EndX && SubsectionX == NumSubsections - 1 && VertX == VertX2)
 							{
-								for (int32 PaddingIdx = PrevMipTexX + PrevMipTexY * PrevMipSizeU; PaddingIdx+1 < PrevMipTexY1 * PrevMipSizeU; ++PaddingIdx)
+								for (int32 PaddingIdx = PrevMipTexX + PrevMipTexY * PrevMipSizeU; PaddingIdx + 1 < PrevMipTexY1 * PrevMipSizeU; ++PaddingIdx)
 								{
-									HeightmapTextureMipData[Mip-1][ PaddingIdx+1 ] = HeightmapTextureMipData[Mip-1][ PaddingIdx ];
+									HeightmapTextureMipData[Mip - 1][PaddingIdx + 1] = HeightmapTextureMipData[Mip - 1][PaddingIdx];
 								}
 							}
 
-							if (EndY && SubsectionX == NumSubsections-1 && SubsectionY == NumSubsections-1 && VertY == VertY2 && VertX == VertX2)
+							if (EndY && SubsectionX == NumSubsections - 1 && SubsectionY == NumSubsections - 1 && VertY == VertY2 && VertX == VertX2)
 							{
-								for (int32 PaddingYIdx = PrevMipTexY; PaddingYIdx+1 < PrevMipSizeV; ++PaddingYIdx)
+								for (int32 PaddingYIdx = PrevMipTexY; PaddingYIdx + 1 < PrevMipSizeV; ++PaddingYIdx)
 								{
 									for (int32 PaddingXIdx = 0; PaddingXIdx < PrevMipSizeU; ++PaddingXIdx)
 									{
-										HeightmapTextureMipData[Mip-1][ PaddingXIdx + (PaddingYIdx+1) * PrevMipSizeU ] = HeightmapTextureMipData[Mip-1][ PaddingXIdx + PaddingYIdx * PrevMipSizeU ];
+										HeightmapTextureMipData[Mip - 1][PaddingXIdx + (PaddingYIdx + 1) * PrevMipSizeU] = HeightmapTextureMipData[Mip - 1][PaddingXIdx + PaddingYIdx * PrevMipSizeU];
 									}
 								}
 							}
 						}
 
-						FColor* TexData = &(HeightmapTextureMipData[Mip])[ TexX + TexY * MipSizeU ];
-						FColor *PreMipTexData00 = &(HeightmapTextureMipData[Mip-1])[ PrevMipTexX + PrevMipTexY * PrevMipSizeU ];
-						FColor *PreMipTexData01 = &(HeightmapTextureMipData[Mip-1])[ PrevMipTexX + PrevMipTexY1 * PrevMipSizeU ];
-						FColor *PreMipTexData10 = &(HeightmapTextureMipData[Mip-1])[ PrevMipTexX1 + PrevMipTexY * PrevMipSizeU ];
-						FColor *PreMipTexData11 = &(HeightmapTextureMipData[Mip-1])[ PrevMipTexX1 + PrevMipTexY1 * PrevMipSizeU ];
+						FColor* TexData = &(HeightmapTextureMipData[Mip])[TexX + TexY * MipSizeU];
+						FColor *PreMipTexData00 = &(HeightmapTextureMipData[Mip - 1])[PrevMipTexX + PrevMipTexY * PrevMipSizeU];
+						FColor *PreMipTexData01 = &(HeightmapTextureMipData[Mip - 1])[PrevMipTexX + PrevMipTexY1 * PrevMipSizeU];
+						FColor *PreMipTexData10 = &(HeightmapTextureMipData[Mip - 1])[PrevMipTexX1 + PrevMipTexY * PrevMipSizeU];
+						FColor *PreMipTexData11 = &(HeightmapTextureMipData[Mip - 1])[PrevMipTexX1 + PrevMipTexY1 * PrevMipSizeU];
 
 						// Lerp height values
 						uint16 PrevMipHeightValue00 = PreMipTexData00->R << 8 | PreMipTexData00->G;
@@ -1076,9 +1076,9 @@ void ULandscapeComponent::GenerateHeightmapMips( TArray<FColor*>& HeightmapTextu
 						uint16 PrevMipHeightValue11 = PreMipTexData11->R << 8 | PreMipTexData11->G;
 						uint16 HeightValue = FMath::RoundToInt(
 							FMath::Lerp(
-							FMath::Lerp( (float)PrevMipHeightValue00, (float)PrevMipHeightValue10, fPrevMipTexFracX),
-							FMath::Lerp( (float)PrevMipHeightValue01, (float)PrevMipHeightValue11, fPrevMipTexFracX),
-							fPrevMipTexFracY) );
+							FMath::Lerp((float)PrevMipHeightValue00, (float)PrevMipHeightValue10, fPrevMipTexFracX),
+							FMath::Lerp((float)PrevMipHeightValue01, (float)PrevMipHeightValue11, fPrevMipTexFracX),
+							fPrevMipTexFracY));
 
 						TexData->R = HeightValue >> 8;
 						TexData->G = HeightValue & 255;
@@ -1086,32 +1086,32 @@ void ULandscapeComponent::GenerateHeightmapMips( TArray<FColor*>& HeightmapTextu
 						// Lerp tangents
 						TexData->B = FMath::RoundToInt(
 							FMath::Lerp(
-							FMath::Lerp( (float)PreMipTexData00->B, (float)PreMipTexData10->B, fPrevMipTexFracX),
-							FMath::Lerp( (float)PreMipTexData01->B, (float)PreMipTexData11->B, fPrevMipTexFracX),
-							fPrevMipTexFracY) );
+							FMath::Lerp((float)PreMipTexData00->B, (float)PreMipTexData10->B, fPrevMipTexFracX),
+							FMath::Lerp((float)PreMipTexData01->B, (float)PreMipTexData11->B, fPrevMipTexFracX),
+							fPrevMipTexFracY));
 
 						TexData->A = FMath::RoundToInt(
 							FMath::Lerp(
-							FMath::Lerp( (float)PreMipTexData00->A, (float)PreMipTexData10->A, fPrevMipTexFracX),
-							FMath::Lerp( (float)PreMipTexData01->A, (float)PreMipTexData11->A, fPrevMipTexFracX),
-							fPrevMipTexFracY) );
+							FMath::Lerp((float)PreMipTexData00->A, (float)PreMipTexData10->A, fPrevMipTexFracX),
+							FMath::Lerp((float)PreMipTexData01->A, (float)PreMipTexData11->A, fPrevMipTexFracX),
+							fPrevMipTexFracY));
 
 						// Padding for missing data
-						if (EndX && SubsectionX == NumSubsections-1 && VertX == VertX2)
+						if (EndX && SubsectionX == NumSubsections - 1 && VertX == VertX2)
 						{
-							for (int32 PaddingIdx = TexX + TexY * MipSizeU; PaddingIdx+1 < (TexY+1) * MipSizeU; ++PaddingIdx)
+							for (int32 PaddingIdx = TexX + TexY * MipSizeU; PaddingIdx + 1 < (TexY + 1) * MipSizeU; ++PaddingIdx)
 							{
-								HeightmapTextureMipData[Mip][ PaddingIdx+1 ] = HeightmapTextureMipData[Mip][ PaddingIdx ];
+								HeightmapTextureMipData[Mip][PaddingIdx + 1] = HeightmapTextureMipData[Mip][PaddingIdx];
 							}
 						}
 
-						if (EndY && SubsectionX == NumSubsections-1 && SubsectionY == NumSubsections-1 && VertY == VertY2 && VertX == VertX2)
+						if (EndY && SubsectionX == NumSubsections - 1 && SubsectionY == NumSubsections - 1 && VertY == VertY2 && VertX == VertX2)
 						{
-							for (int32 PaddingYIdx = TexY; PaddingYIdx+1 < MipSizeV; ++PaddingYIdx)
+							for (int32 PaddingYIdx = TexY; PaddingYIdx + 1 < MipSizeV; ++PaddingYIdx)
 							{
 								for (int32 PaddingXIdx = 0; PaddingXIdx < MipSizeU; ++PaddingXIdx)
 								{
-									HeightmapTextureMipData[Mip][ PaddingXIdx + (PaddingYIdx+1) * MipSizeU ] = HeightmapTextureMipData[Mip][ PaddingXIdx + PaddingYIdx * MipSizeU ];
+									HeightmapTextureMipData[Mip][PaddingXIdx + (PaddingYIdx + 1) * MipSizeU] = HeightmapTextureMipData[Mip][PaddingXIdx + PaddingYIdx * MipSizeU];
 								}
 							}
 						}
@@ -1120,13 +1120,13 @@ void ULandscapeComponent::GenerateHeightmapMips( TArray<FColor*>& HeightmapTextu
 				}
 
 				// Record the areas we updated
-				if( TextureDataInfo )
+				if (TextureDataInfo)
 				{
-					int32 TexX1 = (MipHeightmapOffsetX) + (MipSubsectionSizeQuads+1) * SubsectionX + VertX1;
-					int32 TexY1 = (MipHeightmapOffsetY) + (MipSubsectionSizeQuads+1) * SubsectionY + VertY1;
-					int32 TexX2 = (MipHeightmapOffsetX) + (MipSubsectionSizeQuads+1) * SubsectionX + VertX2;
-					int32 TexY2 = (MipHeightmapOffsetY) + (MipSubsectionSizeQuads+1) * SubsectionY + VertY2;
-					TextureDataInfo->AddMipUpdateRegion(Mip,TexX1,TexY1,TexX2,TexY2);
+					int32 TexX1 = (MipHeightmapOffsetX)+(MipSubsectionSizeQuads + 1) * SubsectionX + VertX1;
+					int32 TexY1 = (MipHeightmapOffsetY)+(MipSubsectionSizeQuads + 1) * SubsectionY + VertY1;
+					int32 TexX2 = (MipHeightmapOffsetX)+(MipSubsectionSizeQuads + 1) * SubsectionX + VertX2;
+					int32 TexY2 = (MipHeightmapOffsetY)+(MipSubsectionSizeQuads + 1) * SubsectionY + VertY2;
+					TextureDataInfo->AddMipUpdateRegion(Mip, TexX1, TexY1, TexX2, TexY2);
 				}
 
 				// Copy current mip values to prev as we move to the next mip.
@@ -1183,17 +1183,17 @@ void ULandscapeComponent::GenerateMipsTempl(int32 InNumSubsections, int32 InSubs
 {
 	// Stores pointers to the locked mip data
 	TArray<DataType*> MipData;
-	MipData.Add( BaseMipData );
+	MipData.Add(BaseMipData);
 	for (int32 MipIndex = 1; MipIndex < Texture->Source.GetNumMips(); ++MipIndex)
 	{
 		MipData.Add((DataType*)Texture->Source.LockMip(MipIndex));
 	}
 
 	// Update the newly created mips
-	UpdateMipsTempl<DataType>( InNumSubsections, InSubsectionSizeQuads, Texture, MipData );
+	UpdateMipsTempl<DataType>(InNumSubsections, InSubsectionSizeQuads, Texture, MipData);
 
 	// Unlock all the new mips, but not the base mip's data
-	for( int32 i=1;i<MipData.Num();i++ )
+	for (int32 i = 1; i < MipData.Num(); i++)
 	{
 		Texture->Source.UnlockMip(i);
 	}
@@ -1211,9 +1211,9 @@ namespace
 	{
 		*Output = FMath::RoundToInt(
 			FMath::Lerp(
-			FMath::Lerp( (float)*Data00, (float)*Data10, FracX),
-			FMath::Lerp( (float)*Data01, (float)*Data11, FracX),
-			FracY) );
+			FMath::Lerp((float)*Data00, (float)*Data10, FracX),
+			FMath::Lerp((float)*Data01, (float)*Data11, FracX),
+			FracY));
 	}
 
 	template<>
@@ -1221,24 +1221,24 @@ namespace
 	{
 		Output->R = FMath::RoundToInt(
 			FMath::Lerp(
-			FMath::Lerp( (float)Data00->R, (float)Data10->R, FracX),
-			FMath::Lerp( (float)Data01->R, (float)Data11->R, FracX),
-			FracY) );
+			FMath::Lerp((float)Data00->R, (float)Data10->R, FracX),
+			FMath::Lerp((float)Data01->R, (float)Data11->R, FracX),
+			FracY));
 		Output->G = FMath::RoundToInt(
 			FMath::Lerp(
-			FMath::Lerp( (float)Data00->G, (float)Data10->G, FracX),
-			FMath::Lerp( (float)Data01->G, (float)Data11->G, FracX),
-			FracY) );
+			FMath::Lerp((float)Data00->G, (float)Data10->G, FracX),
+			FMath::Lerp((float)Data01->G, (float)Data11->G, FracX),
+			FracY));
 		Output->B = FMath::RoundToInt(
 			FMath::Lerp(
-			FMath::Lerp( (float)Data00->B, (float)Data10->B, FracX),
-			FMath::Lerp( (float)Data01->B, (float)Data11->B, FracX),
-			FracY) );
+			FMath::Lerp((float)Data00->B, (float)Data10->B, FracX),
+			FMath::Lerp((float)Data01->B, (float)Data11->B, FracX),
+			FracY));
 		Output->A = FMath::RoundToInt(
 			FMath::Lerp(
-			FMath::Lerp( (float)Data00->A, (float)Data10->A, FracX),
-			FMath::Lerp( (float)Data01->A, (float)Data11->A, FracX),
-			FracY) );
+			FMath::Lerp((float)Data00->A, (float)Data10->A, FracX),
+			FMath::Lerp((float)Data01->A, (float)Data11->A, FracX),
+			FracY));
 	}
 
 	template<typename DataType>
@@ -1268,20 +1268,20 @@ void ULandscapeComponent::UpdateMipsTempl(int32 InNumSubsections, int32 InSubsec
 	int32 MaxWholeSubsectionMip = FMath::FloorLog2(InSubsectionSizeQuads + 1) - 1;
 
 	// Update the mip where each texel's data comes from just one subsection.
-	for( int32 SubsectionY = 0;SubsectionY < InNumSubsections;SubsectionY++ )
+	for (int32 SubsectionY = 0; SubsectionY < InNumSubsections; SubsectionY++)
 	{
 		// Check if subsection is fully above or below the area we are interested in
-		if( (ComponentY2 < InSubsectionSizeQuads*SubsectionY) ||	// above
-			(ComponentY1 > InSubsectionSizeQuads*(SubsectionY+1)) )	// below
+		if ((ComponentY2 < InSubsectionSizeQuads*SubsectionY) ||	// above
+			(ComponentY1 > InSubsectionSizeQuads*(SubsectionY + 1)))	// below
 		{
 			continue;
 		}
 
-		for( int32 SubsectionX = 0;SubsectionX < InNumSubsections;SubsectionX++ )
+		for (int32 SubsectionX = 0; SubsectionX < InNumSubsections; SubsectionX++)
 		{
 			// Check if subsection is fully to the left or right of the area we are interested in
-			if( (ComponentX2 < InSubsectionSizeQuads*SubsectionX) ||	// left
-				(ComponentX1 > InSubsectionSizeQuads*(SubsectionX+1)) )	// right
+			if ((ComponentX2 < InSubsectionSizeQuads*SubsectionX) ||	// left
+				(ComponentX1 > InSubsectionSizeQuads*(SubsectionX + 1)))	// right
 			{
 				continue;
 			}
@@ -1298,19 +1298,19 @@ void ULandscapeComponent::UpdateMipsTempl(int32 InNumSubsections, int32 InSubsec
 			int32 PrevMipSizeU = WeightmapSizeU;
 			int32 PrevMipSizeV = WeightmapSizeV;
 
-			for( int32 Mip=1;Mip<=MaxWholeSubsectionMip;Mip++ )
+			for (int32 Mip = 1; Mip <= MaxWholeSubsectionMip; Mip++)
 			{
 				int32 MipSizeU = WeightmapSizeU >> Mip;
 				int32 MipSizeV = WeightmapSizeV >> Mip;
 
-				int32 MipSubsectionSizeQuads = ((InSubsectionSizeQuads+1)>>Mip)-1;
+				int32 MipSubsectionSizeQuads = ((InSubsectionSizeQuads + 1) >> Mip) - 1;
 				float InvMipSubsectionSizeQuads = 1.f / (float)MipSubsectionSizeQuads;
 
 				// Area to update in current mip level coords
-				int32 MipSubX1 = FMath::FloorToInt( (float)MipSubsectionSizeQuads * (float)PrevMipSubX1 * InvPrevMipSubsectionSizeQuads );
-				int32 MipSubY1 = FMath::FloorToInt( (float)MipSubsectionSizeQuads * (float)PrevMipSubY1 * InvPrevMipSubsectionSizeQuads );
-				int32 MipSubX2 = FMath::CeilToInt(  (float)MipSubsectionSizeQuads * (float)PrevMipSubX2 * InvPrevMipSubsectionSizeQuads );
-				int32 MipSubY2 = FMath::CeilToInt(  (float)MipSubsectionSizeQuads * (float)PrevMipSubY2 * InvPrevMipSubsectionSizeQuads );
+				int32 MipSubX1 = FMath::FloorToInt((float)MipSubsectionSizeQuads * (float)PrevMipSubX1 * InvPrevMipSubsectionSizeQuads);
+				int32 MipSubY1 = FMath::FloorToInt((float)MipSubsectionSizeQuads * (float)PrevMipSubY1 * InvPrevMipSubsectionSizeQuads);
+				int32 MipSubX2 = FMath::CeilToInt((float)MipSubsectionSizeQuads * (float)PrevMipSubX2 * InvPrevMipSubsectionSizeQuads);
+				int32 MipSubY2 = FMath::CeilToInt((float)MipSubsectionSizeQuads * (float)PrevMipSubY2 * InvPrevMipSubsectionSizeQuads);
 
 				// Clamp area to update
 				int32 VertX1 = FMath::Clamp<int32>(MipSubX1, 0, MipSubsectionSizeQuads);
@@ -1318,39 +1318,39 @@ void ULandscapeComponent::UpdateMipsTempl(int32 InNumSubsections, int32 InSubsec
 				int32 VertX2 = FMath::Clamp<int32>(MipSubX2, 0, MipSubsectionSizeQuads);
 				int32 VertY2 = FMath::Clamp<int32>(MipSubY2, 0, MipSubsectionSizeQuads);
 
-				for( int32 VertY=VertY1;VertY<=VertY2;VertY++ )
+				for (int32 VertY = VertY1; VertY <= VertY2; VertY++)
 				{
-					for( int32 VertX=VertX1;VertX<=VertX2;VertX++ )
+					for (int32 VertX = VertX1; VertX <= VertX2; VertX++)
 					{
 						// Convert VertX/Y into previous mip's coords
 						float PrevMipVertX = (float)PrevMipSubsectionSizeQuads * (float)VertX * InvMipSubsectionSizeQuads;
 						float PrevMipVertY = (float)PrevMipSubsectionSizeQuads * (float)VertY * InvMipSubsectionSizeQuads;
 
 						// X/Y of the vertex we're looking indexed into the texture data
-						int32 TexX = (MipSubsectionSizeQuads+1) * SubsectionX + VertX;
-						int32 TexY = (MipSubsectionSizeQuads+1) * SubsectionY + VertY;
+						int32 TexX = (MipSubsectionSizeQuads + 1) * SubsectionX + VertX;
+						int32 TexY = (MipSubsectionSizeQuads + 1) * SubsectionY + VertY;
 
-						float fPrevMipTexX = (float)((PrevMipSubsectionSizeQuads+1) * SubsectionX) + PrevMipVertX;
-						float fPrevMipTexY = (float)((PrevMipSubsectionSizeQuads+1) * SubsectionY) + PrevMipVertY;
+						float fPrevMipTexX = (float)((PrevMipSubsectionSizeQuads + 1) * SubsectionX) + PrevMipVertX;
+						float fPrevMipTexY = (float)((PrevMipSubsectionSizeQuads + 1) * SubsectionY) + PrevMipVertY;
 
 						int32 PrevMipTexX = FMath::FloorToInt(fPrevMipTexX);
 						float fPrevMipTexFracX = FMath::Fractional(fPrevMipTexX);
 						int32 PrevMipTexY = FMath::FloorToInt(fPrevMipTexY);
 						float fPrevMipTexFracY = FMath::Fractional(fPrevMipTexY);
 
-						check( TexX >= 0 && TexX < MipSizeU );
-						check( TexY >= 0 && TexY < MipSizeV );
-						check( PrevMipTexX >= 0 && PrevMipTexX < PrevMipSizeU );
-						check( PrevMipTexY >= 0 && PrevMipTexY < PrevMipSizeV );
+						check(TexX >= 0 && TexX < MipSizeU);
+						check(TexY >= 0 && TexY < MipSizeV);
+						check(PrevMipTexX >= 0 && PrevMipTexX < PrevMipSizeU);
+						check(PrevMipTexY >= 0 && PrevMipTexY < PrevMipSizeV);
 
-						int32 PrevMipTexX1 = FMath::Min<int32>( PrevMipTexX+1, PrevMipSizeU-1 );
-						int32 PrevMipTexY1 = FMath::Min<int32>( PrevMipTexY+1, PrevMipSizeV-1 );
+						int32 PrevMipTexX1 = FMath::Min<int32>(PrevMipTexX + 1, PrevMipSizeU - 1);
+						int32 PrevMipTexY1 = FMath::Min<int32>(PrevMipTexY + 1, PrevMipSizeV - 1);
 
-						DataType* TexData = &(TextureMipData[Mip])[ TexX + TexY * MipSizeU ];
-						DataType *PreMipTexData00 = &(TextureMipData[Mip-1])[ PrevMipTexX + PrevMipTexY * PrevMipSizeU ];
-						DataType *PreMipTexData01 = &(TextureMipData[Mip-1])[ PrevMipTexX + PrevMipTexY1 * PrevMipSizeU ];
-						DataType *PreMipTexData10 = &(TextureMipData[Mip-1])[ PrevMipTexX1 + PrevMipTexY * PrevMipSizeU ];
-						DataType *PreMipTexData11 = &(TextureMipData[Mip-1])[ PrevMipTexX1 + PrevMipTexY1 * PrevMipSizeU ];
+						DataType* TexData = &(TextureMipData[Mip])[TexX + TexY * MipSizeU];
+						DataType *PreMipTexData00 = &(TextureMipData[Mip - 1])[PrevMipTexX + PrevMipTexY * PrevMipSizeU];
+						DataType *PreMipTexData01 = &(TextureMipData[Mip - 1])[PrevMipTexX + PrevMipTexY1 * PrevMipSizeU];
+						DataType *PreMipTexData10 = &(TextureMipData[Mip - 1])[PrevMipTexX1 + PrevMipTexY * PrevMipSizeU];
+						DataType *PreMipTexData11 = &(TextureMipData[Mip - 1])[PrevMipTexX1 + PrevMipTexY1 * PrevMipSizeU];
 
 						// Lerp weightmap data
 						BiLerpTextureData<DataType>(TexData, PreMipTexData00, PreMipTexData10, PreMipTexData01, PreMipTexData11, fPrevMipTexFracX, fPrevMipTexFracY);
@@ -1358,13 +1358,13 @@ void ULandscapeComponent::UpdateMipsTempl(int32 InNumSubsections, int32 InSubsec
 				}
 
 				// Record the areas we updated
-				if( TextureDataInfo )
+				if (TextureDataInfo)
 				{
-					int32 TexX1 = (MipSubsectionSizeQuads+1) * SubsectionX + VertX1;
-					int32 TexY1 = (MipSubsectionSizeQuads+1) * SubsectionY + VertY1;
-					int32 TexX2 = (MipSubsectionSizeQuads+1) * SubsectionX + VertX2;
-					int32 TexY2 = (MipSubsectionSizeQuads+1) * SubsectionY + VertY2;
-					TextureDataInfo->AddMipUpdateRegion(Mip,TexX1,TexY1,TexX2,TexY2);
+					int32 TexX1 = (MipSubsectionSizeQuads + 1) * SubsectionX + VertX1;
+					int32 TexY1 = (MipSubsectionSizeQuads + 1) * SubsectionY + VertY1;
+					int32 TexX2 = (MipSubsectionSizeQuads + 1) * SubsectionX + VertX2;
+					int32 TexY2 = (MipSubsectionSizeQuads + 1) * SubsectionY + VertY2;
+					TextureDataInfo->AddMipUpdateRegion(Mip, TexX1, TexY1, TexX2, TexY2);
 				}
 
 				// Copy current mip values to prev as we move to the next mip.
@@ -1434,20 +1434,20 @@ void ULandscapeComponent::UpdateDataMips(int32 InNumSubsections, int32 InSubsect
 	UpdateMipsTempl<uint8>(InNumSubsections, InSubsectionSizeQuads, Texture, TextureMipData, ComponentX1, ComponentY1, ComponentX2, ComponentY2, TextureDataInfo);
 }
 
-float ULandscapeComponent::GetLayerWeightAtLocation( const FVector& InLocation, ULandscapeLayerInfoObject* LayerInfo, TArray<uint8>* LayerCache )
+float ULandscapeComponent::GetLayerWeightAtLocation(const FVector& InLocation, ULandscapeLayerInfoObject* LayerInfo, TArray<uint8>* LayerCache)
 {
 	// Allocate and discard locally if no external cache is passed in.
 	TArray<uint8> LocalCache;
-	if( LayerCache==NULL )
+	if (LayerCache == NULL)
 	{
 		LayerCache = &LocalCache;
 	}
 
 	// Fill the cache if necessary
-	if( LayerCache->Num() == 0 )
+	if (LayerCache->Num() == 0)
 	{
 		FLandscapeComponentDataInterface CDI(this);
-		if( !CDI.GetWeightmapTextureData( LayerInfo, *LayerCache ) )
+		if (!CDI.GetWeightmapTextureData(LayerInfo, *LayerCache))
 		{
 			// no data for this layer for this component.
 			return 0.f;
@@ -1456,10 +1456,10 @@ float ULandscapeComponent::GetLayerWeightAtLocation( const FVector& InLocation,
 
 	// Find location
 	// TODO: Root landscape isn't always loaded, would Proxy suffice?
-	ALandscape* Landscape = GetLandscapeActor(); 
+	ALandscape* Landscape = GetLandscapeActor();
 	const FVector DrawScale = Landscape->GetRootComponent()->RelativeScale3D;
-	float TestX = (InLocation.X-Landscape->GetActorLocation().X) / DrawScale.X - (float)GetSectionBase().X;
-	float TestY = (InLocation.Y-Landscape->GetActorLocation().Y) / DrawScale.Y - (float)GetSectionBase().Y;
+	float TestX = (InLocation.X - Landscape->GetActorLocation().X) / DrawScale.X - (float)GetSectionBase().X;
+	float TestY = (InLocation.Y - Landscape->GetActorLocation().Y) / DrawScale.Y - (float)GetSectionBase().Y;
 
 	// Find data
 	int32 X1 = FMath::FloorToInt(TestX);
@@ -1467,13 +1467,13 @@ float ULandscapeComponent::GetLayerWeightAtLocation( const FVector& InLocation,
 	int32 X2 = FMath::CeilToInt(TestX);
 	int32 Y2 = FMath::CeilToInt(TestY);
 
-	int32 Stride = (SubsectionSizeQuads+1) * NumSubsections;
+	int32 Stride = (SubsectionSizeQuads + 1) * NumSubsections;
 
 	// Min is to prevent the sampling of the final column from overflowing
-	int32 IdxX1 = FMath::Min<int32>(((X1 / SubsectionSizeQuads) * (SubsectionSizeQuads+1)) + (X1 % SubsectionSizeQuads), Stride-1);
-	int32 IdxY1 = FMath::Min<int32>(((Y1 / SubsectionSizeQuads) * (SubsectionSizeQuads+1)) + (Y1 % SubsectionSizeQuads), Stride-1);
-	int32 IdxX2 = FMath::Min<int32>(((X2 / SubsectionSizeQuads) * (SubsectionSizeQuads+1)) + (X2 % SubsectionSizeQuads), Stride-1);
-	int32 IdxY2 = FMath::Min<int32>(((Y2 / SubsectionSizeQuads) * (SubsectionSizeQuads+1)) + (Y2 % SubsectionSizeQuads), Stride-1);
+	int32 IdxX1 = FMath::Min<int32>(((X1 / SubsectionSizeQuads) * (SubsectionSizeQuads + 1)) + (X1 % SubsectionSizeQuads), Stride - 1);
+	int32 IdxY1 = FMath::Min<int32>(((Y1 / SubsectionSizeQuads) * (SubsectionSizeQuads + 1)) + (Y1 % SubsectionSizeQuads), Stride - 1);
+	int32 IdxX2 = FMath::Min<int32>(((X2 / SubsectionSizeQuads) * (SubsectionSizeQuads + 1)) + (X2 % SubsectionSizeQuads), Stride - 1);
+	int32 IdxY2 = FMath::Min<int32>(((Y2 / SubsectionSizeQuads) * (SubsectionSizeQuads + 1)) + (Y2 % SubsectionSizeQuads), Stride - 1);
 
 	// sample
 	float Sample11 = (float)((*LayerCache)[IdxX1 + Stride*IdxY1]) / 255.f;
@@ -1486,8 +1486,8 @@ float ULandscapeComponent::GetLayerWeightAtLocation( const FVector& InLocation,
 
 	// Bilinear interpolate
 	return FMath::Lerp(
-		FMath::Lerp( Sample11, Sample21, LerpX),
-		FMath::Lerp( Sample12, Sample22, LerpX),
+		FMath::Lerp(Sample11, Sample21, LerpX),
+		FMath::Lerp(Sample12, Sample22, LerpX),
 		LerpY);
 }
 
@@ -1495,8 +1495,8 @@ void ULandscapeComponent::GetComponentExtent(int32& MinX, int32& MinY, int32& Ma
 {
 	MinX = FMath::Min(SectionBaseX, MinX);
 	MinY = FMath::Min(SectionBaseY, MinY);
-	MaxX = FMath::Max(SectionBaseX+ComponentSizeQuads, MaxX);
-	MaxY = FMath::Max(SectionBaseY+ComponentSizeQuads, MaxY);
+	MaxX = FMath::Max(SectionBaseX + ComponentSizeQuads, MaxX);
+	MaxY = FMath::Max(SectionBaseY + ComponentSizeQuads, MaxY);
 }
 
 //
@@ -1516,12 +1516,12 @@ void ULandscapeInfo::GetComponentsInRegion(int32 X1, int32 Y1, int32 X2, int32 Y
 	int32 ComponentIndexX1, ComponentIndexY1, ComponentIndexX2, ComponentIndexY2;
 	ALandscape::CalcComponentIndicesOverlap(X1, Y1, X2, Y2, ComponentSizeQuads, ComponentIndexX1, ComponentIndexY1, ComponentIndexX2, ComponentIndexY2);
 
-	for( int32 ComponentIndexY=ComponentIndexY1;ComponentIndexY<=ComponentIndexY2;ComponentIndexY++ )
+	for (int32 ComponentIndexY = ComponentIndexY1; ComponentIndexY <= ComponentIndexY2; ComponentIndexY++)
 	{
-		for( int32 ComponentIndexX=ComponentIndexX1;ComponentIndexX<=ComponentIndexX2;ComponentIndexX++ )
-		{		
-			ULandscapeComponent* Component = XYtoComponentMap.FindRef(FIntPoint(ComponentIndexX,ComponentIndexY));
-			if( Component && !FLevelUtils::IsLevelLocked(Component->GetLandscapeProxy()->GetLevel()) && FLevelUtils::IsLevelVisible(Component->GetLandscapeProxy()->GetLevel()))
+		for (int32 ComponentIndexX = ComponentIndexX1; ComponentIndexX <= ComponentIndexX2; ComponentIndexX++)
+		{
+			ULandscapeComponent* Component = XYtoComponentMap.FindRef(FIntPoint(ComponentIndexX, ComponentIndexY));
+			if (Component && !FLevelUtils::IsLevelLocked(Component->GetLandscapeProxy()->GetLevel()) && FLevelUtils::IsLevelVisible(Component->GetLandscapeProxy()->GetLevel()))
 			{
 				OutComponents.Add(Component);
 			}
@@ -1538,12 +1538,12 @@ struct FWeightmapTextureAllocation
 	UTexture2D* Texture;
 	FColor* TextureData;
 
-	FWeightmapTextureAllocation( int32 InX, int32 InY, int32 InChannels, UTexture2D* InTexture, FColor* InTextureData )
-		:	X(InX)
-		,	Y(InY)
-		,	ChannelsInUse(InChannels)
-		,	Texture(InTexture)
-		,	TextureData(InTextureData)
+	FWeightmapTextureAllocation(int32 InX, int32 InY, int32 InChannels, UTexture2D* InTexture, FColor* InTextureData)
+		: X(InX)
+		, Y(InY)
+		, ChannelsInUse(InChannels)
+		, Texture(InTexture)
+		, TextureData(InTextureData)
 	{}
 };
 
@@ -1658,29 +1658,29 @@ ULandscapeLayerInfoObject* ALandscapeProxy::CreateLayerInfo(const TCHAR* LayerNa
 }
 
 #define HEIGHTDATA(X,Y) (HeightData[ FMath::Clamp<int32>(Y,0,VertsY) * VertsX + FMath::Clamp<int32>(X,0,VertsX) ])
-void ALandscapeProxy::Import(FGuid Guid, int32 VertsX, int32 VertsY, 
-								int32 InComponentSizeQuads, int32 InNumSubsections, int32 InSubsectionSizeQuads, 
-								const uint16* HeightData, const TCHAR* HeightmapFileName, 
-								const TArray<FLandscapeImportLayerInfo>& ImportLayerInfos )
+void ALandscapeProxy::Import(FGuid Guid, int32 VertsX, int32 VertsY,
+	int32 InComponentSizeQuads, int32 InNumSubsections, int32 InSubsectionSizeQuads,
+	const uint16* HeightData, const TCHAR* HeightmapFileName,
+	const TArray<FLandscapeImportLayerInfo>& ImportLayerInfos)
 {
-	GWarn->BeginSlowTask( LOCTEXT("BeingImportingLandscapeTask", "Importing Landscape"), true);
+	GWarn->BeginSlowTask(LOCTEXT("BeingImportingLandscapeTask", "Importing Landscape"), true);
 
 	ComponentSizeQuads = InComponentSizeQuads;
 	NumSubsections = InNumSubsections;
 	SubsectionSizeQuads = InSubsectionSizeQuads;
 	LandscapeGuid = Guid;
-	
+
 	MarkPackageDirty();
 
 	// Create and initialize landscape info object
 	GetLandscapeInfo(true)->RegisterActor(this);
-	
-	int32 NumPatchesX = (VertsX-1);
-	int32 NumPatchesY = (VertsY-1);
+
+	int32 NumPatchesX = (VertsX - 1);
+	int32 NumPatchesY = (VertsY - 1);
 
 	int32 NumSectionsX = NumPatchesX / ComponentSizeQuads;
 	int32 NumSectionsY = NumPatchesY / ComponentSizeQuads;
-	
+
 	LandscapeComponents.Empty(NumSectionsX * NumSectionsY);
 
 	for (int32 Y = 0; Y < NumSectionsY; Y++)
@@ -1694,12 +1694,12 @@ void ALandscapeProxy::Import(FGuid Guid, int32 VertsX, int32 VertsY,
 			int32 BaseX = X * ComponentSizeQuads;
 			int32 BaseY = Y * ComponentSizeQuads;
 
-			ULandscapeComponent* LandscapeComponent = ConstructObject<ULandscapeComponent>(ULandscapeComponent::StaticClass(),this,NAME_None,RF_Transactional);
-			LandscapeComponent->SetRelativeLocation(FVector(BaseX,BaseY,0));
-			LandscapeComponent->AttachTo( GetRootComponent(), NAME_None );
+			ULandscapeComponent* LandscapeComponent = ConstructObject<ULandscapeComponent>(ULandscapeComponent::StaticClass(), this, NAME_None, RF_Transactional);
+			LandscapeComponent->SetRelativeLocation(FVector(BaseX, BaseY, 0));
+			LandscapeComponent->AttachTo(GetRootComponent(), NAME_None);
 			LandscapeComponents.Add(LandscapeComponent);
 			LandscapeComponent->Init(
-				BaseX,BaseY,
+				BaseX, BaseY,
 				ComponentSizeQuads,
 				NumSubsections,
 				SubsectionSizeQuads
@@ -1713,7 +1713,7 @@ void ALandscapeProxy::Import(FGuid Guid, int32 VertsX, int32 VertsY,
 
 #define MAX_HEIGHTMAP_TEXTURE_SIZE 512
 
-	int32 ComponentSizeVerts = NumSubsections * (SubsectionSizeQuads+1);
+	int32 ComponentSizeVerts = NumSubsections * (SubsectionSizeQuads + 1);
 	int32 ComponentsPerHeightmap = MAX_HEIGHTMAP_TEXTURE_SIZE / ComponentSizeVerts;
 
 	// Ensure that we don't pack so many heightmaps into a texture that their lowest LOD isn't guaranteed to be resident
@@ -1722,7 +1722,7 @@ void ALandscapeProxy::Import(FGuid Guid, int32 VertsX, int32 VertsY,
 	// Count how many heightmaps we need and the X dimension of the final heightmap
 	int32 NumHeightmapsX = 1;
 	int32 FinalComponentsX = NumSectionsX;
-	while( FinalComponentsX > ComponentsPerHeightmap )
+	while (FinalComponentsX > ComponentsPerHeightmap)
 	{
 		FinalComponentsX -= ComponentsPerHeightmap;
 		NumHeightmapsX++;
@@ -1730,7 +1730,7 @@ void ALandscapeProxy::Import(FGuid Guid, int32 VertsX, int32 VertsY,
 	// Count how many heightmaps we need and the Y dimension of the final heightmap
 	int32 NumHeightmapsY = 1;
 	int32 FinalComponentsY = NumSectionsY;
-	while( FinalComponentsY > ComponentsPerHeightmap )
+	while (FinalComponentsY > ComponentsPerHeightmap)
 	{
 		FinalComponentsY -= ComponentsPerHeightmap;
 		NumHeightmapsY++;
@@ -1738,27 +1738,27 @@ void ALandscapeProxy::Import(FGuid Guid, int32 VertsX, int32 VertsY,
 
 	TArray<FHeightmapInfo> HeightmapInfos;
 
-	for( int32 HmY=0;HmY<NumHeightmapsY;HmY++ )
+	for (int32 HmY = 0; HmY < NumHeightmapsY; HmY++)
 	{
-		for( int32 HmX=0;HmX<NumHeightmapsX;HmX++ )
+		for (int32 HmX = 0; HmX < NumHeightmapsX; HmX++)
 		{
 			FHeightmapInfo& HeightmapInfo = HeightmapInfos[HeightmapInfos.AddZeroed()];
 
 			// make sure the heightmap UVs are powers of two.
-			HeightmapInfo.HeightmapSizeU = (1<<FMath::CeilLogTwo( ((HmX==NumHeightmapsX-1) ? FinalComponentsX : ComponentsPerHeightmap)*ComponentSizeVerts ));
-			HeightmapInfo.HeightmapSizeV = (1<<FMath::CeilLogTwo( ((HmY==NumHeightmapsY-1) ? FinalComponentsY : ComponentsPerHeightmap)*ComponentSizeVerts ));
+			HeightmapInfo.HeightmapSizeU = (1 << FMath::CeilLogTwo(((HmX == NumHeightmapsX - 1) ? FinalComponentsX : ComponentsPerHeightmap)*ComponentSizeVerts));
+			HeightmapInfo.HeightmapSizeV = (1 << FMath::CeilLogTwo(((HmY == NumHeightmapsY - 1) ? FinalComponentsY : ComponentsPerHeightmap)*ComponentSizeVerts));
 
 			// Construct the heightmap textures
 			HeightmapInfo.HeightmapTexture = CreateLandscapeTexture(HeightmapInfo.HeightmapSizeU, HeightmapInfo.HeightmapSizeV, TEXTUREGROUP_Terrain_Heightmap, TSF_BGRA8);
-			
+
 			int32 MipSubsectionSizeQuads = SubsectionSizeQuads;
 			int32 MipSizeU = HeightmapInfo.HeightmapSizeU;
 			int32 MipSizeV = HeightmapInfo.HeightmapSizeV;
-			while( MipSizeU > 1 && MipSizeV > 1 && MipSubsectionSizeQuads >= 1 )
+			while (MipSizeU > 1 && MipSizeV > 1 && MipSubsectionSizeQuads >= 1)
 			{
 				int32 MipIndex = HeightmapInfo.HeightmapTextureMipData.Num();
 				FColor* HeightmapTextureData = (FColor*)HeightmapInfo.HeightmapTexture->Source.LockMip(MipIndex);
-				FMemory::Memzero( HeightmapTextureData, MipSizeU*MipSizeV*sizeof(FColor) );
+				FMemory::Memzero(HeightmapTextureData, MipSizeU*MipSizeV*sizeof(FColor));
 				HeightmapInfo.HeightmapTextureMipData.Add(HeightmapTextureData);
 
 				MipSizeU >>= 1;
@@ -1772,25 +1772,25 @@ void ALandscapeProxy::Import(FGuid Guid, int32 VertsX, int32 VertsY,
 	const FVector DrawScale3D = GetRootComponent()->RelativeScale3D;
 
 	// Calculate the normals for each of the two triangles per quad.
-	FVector* VertexNormals = new FVector[(NumPatchesX+1)*(NumPatchesY+1)];
-	FMemory::Memzero(VertexNormals, (NumPatchesX+1)*(NumPatchesY+1)*sizeof(FVector));
-	for( int32 QuadY=0;QuadY<NumPatchesY;QuadY++ )
+	FVector* VertexNormals = new FVector[(NumPatchesX + 1)*(NumPatchesY + 1)];
+	FMemory::Memzero(VertexNormals, (NumPatchesX + 1)*(NumPatchesY + 1)*sizeof(FVector));
+	for (int32 QuadY = 0; QuadY < NumPatchesY; QuadY++)
 	{
-		for( int32 QuadX=0;QuadX<NumPatchesX;QuadX++ )
+		for (int32 QuadX = 0; QuadX < NumPatchesX; QuadX++)
 		{
-			FVector Vert00 = FVector(0.f,0.f,((float)HEIGHTDATA(QuadX+0, QuadY+0) - 32768.f)*LANDSCAPE_ZSCALE) * DrawScale3D;
-			FVector Vert01 = FVector(0.f,1.f,((float)HEIGHTDATA(QuadX+0, QuadY+1) - 32768.f)*LANDSCAPE_ZSCALE) * DrawScale3D;
-			FVector Vert10 = FVector(1.f,0.f,((float)HEIGHTDATA(QuadX+1, QuadY+0) - 32768.f)*LANDSCAPE_ZSCALE) * DrawScale3D;
-			FVector Vert11 = FVector(1.f,1.f,((float)HEIGHTDATA(QuadX+1, QuadY+1) - 32768.f)*LANDSCAPE_ZSCALE) * DrawScale3D;
+			FVector Vert00 = FVector(0.f, 0.f, ((float)HEIGHTDATA(QuadX + 0, QuadY + 0) - 32768.f)*LANDSCAPE_ZSCALE) * DrawScale3D;
+			FVector Vert01 = FVector(0.f, 1.f, ((float)HEIGHTDATA(QuadX + 0, QuadY + 1) - 32768.f)*LANDSCAPE_ZSCALE) * DrawScale3D;
+			FVector Vert10 = FVector(1.f, 0.f, ((float)HEIGHTDATA(QuadX + 1, QuadY + 0) - 32768.f)*LANDSCAPE_ZSCALE) * DrawScale3D;
+			FVector Vert11 = FVector(1.f, 1.f, ((float)HEIGHTDATA(QuadX + 1, QuadY + 1) - 32768.f)*LANDSCAPE_ZSCALE) * DrawScale3D;
 
-			FVector FaceNormal1 = ((Vert00-Vert10) ^ (Vert10-Vert11)).SafeNormal();
-			FVector FaceNormal2 = ((Vert11-Vert01) ^ (Vert01-Vert00)).SafeNormal(); 
+			FVector FaceNormal1 = ((Vert00 - Vert10) ^ (Vert10 - Vert11)).SafeNormal();
+			FVector FaceNormal2 = ((Vert11 - Vert01) ^ (Vert01 - Vert00)).SafeNormal();
 
 			// contribute to the vertex normals.
-			VertexNormals[(QuadX+1 + (NumPatchesX+1)*(QuadY+0))] += FaceNormal1;
-			VertexNormals[(QuadX+0 + (NumPatchesX+1)*(QuadY+1))] += FaceNormal2;
-			VertexNormals[(QuadX+0 + (NumPatchesX+1)*(QuadY+0))] += FaceNormal1 + FaceNormal2;
-			VertexNormals[(QuadX+1 + (NumPatchesX+1)*(QuadY+1))] += FaceNormal1 + FaceNormal2;
+			VertexNormals[(QuadX + 1 + (NumPatchesX + 1)*(QuadY + 0))] += FaceNormal1;
+			VertexNormals[(QuadX + 0 + (NumPatchesX + 1)*(QuadY + 1))] += FaceNormal2;
+			VertexNormals[(QuadX + 0 + (NumPatchesX + 1)*(QuadY + 0))] += FaceNormal1 + FaceNormal2;
+			VertexNormals[(QuadX + 1 + (NumPatchesX + 1)*(QuadY + 1))] += FaceNormal1 + FaceNormal2;
 		}
 	}
 
@@ -1807,32 +1807,32 @@ void ALandscapeProxy::Import(FGuid Guid, int32 VertsX, int32 VertsY,
 
 			// Import alphamap data into local array and check for unused layers for this component.
 			TArray<FLandscapeComponentAlphaInfo> EditingAlphaLayerData;
-			for( int32 LayerIndex=0;LayerIndex<ImportLayerInfos.Num();LayerIndex++ )
+			for (int32 LayerIndex = 0; LayerIndex < ImportLayerInfos.Num(); LayerIndex++)
 			{
 				FLandscapeComponentAlphaInfo* NewAlphaInfo = new(EditingAlphaLayerData) FLandscapeComponentAlphaInfo(LandscapeComponent, LayerIndex);
 
 				if (ImportLayerInfos[LayerIndex].LayerData.Num())
 				{
-					for( int32 AlphaY=0;AlphaY<=LandscapeComponent->ComponentSizeQuads;AlphaY++ )
+					for (int32 AlphaY = 0; AlphaY <= LandscapeComponent->ComponentSizeQuads; AlphaY++)
 					{
-						const uint8* OldAlphaRowStart = &ImportLayerInfos[LayerIndex].LayerData[ (AlphaY+LandscapeComponent->GetSectionBase().Y) * VertsX + (LandscapeComponent->GetSectionBase().X) ];
-						uint8* NewAlphaRowStart = &NewAlphaInfo->AlphaValues[AlphaY * (LandscapeComponent->ComponentSizeQuads+1)];
-						FMemory::Memcpy(NewAlphaRowStart, OldAlphaRowStart, LandscapeComponent->ComponentSizeQuads+1);
+						const uint8* OldAlphaRowStart = &ImportLayerInfos[LayerIndex].LayerData[(AlphaY + LandscapeComponent->GetSectionBase().Y) * VertsX + (LandscapeComponent->GetSectionBase().X)];
+						uint8* NewAlphaRowStart = &NewAlphaInfo->AlphaValues[AlphaY * (LandscapeComponent->ComponentSizeQuads + 1)];
+						FMemory::Memcpy(NewAlphaRowStart, OldAlphaRowStart, LandscapeComponent->ComponentSizeQuads + 1);
 					}
 				}
 			}
 
-			for( int32 AlphaMapIndex=0; AlphaMapIndex<EditingAlphaLayerData.Num();AlphaMapIndex++ )
+			for (int32 AlphaMapIndex = 0; AlphaMapIndex < EditingAlphaLayerData.Num(); AlphaMapIndex++)
 			{
-				if( EditingAlphaLayerData[AlphaMapIndex].IsLayerAllZero() )
+				if (EditingAlphaLayerData[AlphaMapIndex].IsLayerAllZero())
 				{
 					EditingAlphaLayerData.RemoveAt(AlphaMapIndex);
 					AlphaMapIndex--;
 				}
 			}
 
-			
-			UE_LOG(LogLandscape, Log, TEXT("%s needs %d alphamaps"), *LandscapeComponent->GetName(),EditingAlphaLayerData.Num());
+
+			UE_LOG(LogLandscape, Log, TEXT("%s needs %d alphamaps"), *LandscapeComponent->GetName(), EditingAlphaLayerData.Num());
 
 			// Calculate weightmap weights for this component
 			WeightValues.Empty(EditingAlphaLayerData.Num());
@@ -1843,7 +1843,7 @@ void ALandscapeProxy::Import(FGuid Guid, int32 VertsX, int32 VertsY,
 			IsNoBlendArray.Empty(EditingAlphaLayerData.Num());
 			IsNoBlendArray.AddZeroed(EditingAlphaLayerData.Num());
 
-			for( int32 WeightLayerIndex=0; WeightLayerIndex<WeightValues.Num();WeightLayerIndex++ )
+			for (int32 WeightLayerIndex = 0; WeightLayerIndex < WeightValues.Num(); WeightLayerIndex++)
 			{
 				// Lookup the original layer name
 				WeightValues[WeightLayerIndex] = EditingAlphaLayerData[WeightLayerIndex].AlphaValues;
@@ -1855,10 +1855,10 @@ void ALandscapeProxy::Import(FGuid Guid, int32 VertsX, int32 VertsY,
 			EditingAlphaLayerData.Empty();
 
 			// For each layer...
-			for( int32 WeightLayerIndex=WeightValues.Num()-1; WeightLayerIndex>=0;WeightLayerIndex-- )
+			for (int32 WeightLayerIndex = WeightValues.Num() - 1; WeightLayerIndex >= 0; WeightLayerIndex--)
 			{
 				// ... multiply all lower layers'...
-				for( int32 BelowWeightLayerIndex=WeightLayerIndex-1; BelowWeightLayerIndex>=0;BelowWeightLayerIndex-- )
+				for (int32 BelowWeightLayerIndex = WeightLayerIndex - 1; BelowWeightLayerIndex >= 0; BelowWeightLayerIndex--)
 				{
 					int32 TotalWeight = 0;
 
@@ -1868,7 +1868,7 @@ void ALandscapeProxy::Import(FGuid Guid, int32 VertsX, int32 VertsY,
 					}
 
 					// ... values by...
-					for( int32 Idx=0;Idx<WeightValues[WeightLayerIndex].Num();Idx++ )
+					for (int32 Idx = 0; Idx < WeightValues[WeightLayerIndex].Num(); Idx++)
 					{
 						// ... one-minus the current layer's values
 						int32 NewValue = (int32)WeightValues[BelowWeightLayerIndex][Idx] * (int32)(255 - WeightValues[WeightLayerIndex][Idx]) / 255;
@@ -1876,7 +1876,7 @@ void ALandscapeProxy::Import(FGuid Guid, int32 VertsX, int32 VertsY,
 						TotalWeight += NewValue;
 					}
 
-					if( TotalWeight == 0 )
+					if (TotalWeight == 0)
 					{
 						// Remove the layer as it has no contribution
 						WeightValues.RemoveAt(BelowWeightLayerIndex);
@@ -1892,13 +1892,13 @@ void ALandscapeProxy::Import(FGuid Guid, int32 VertsX, int32 VertsY,
 			// Weight normalization for total should be 255...
 			if (WeightValues.Num())
 			{
-				for( int32 Idx=0;Idx<WeightValues[0].Num();Idx++ )
+				for (int32 Idx = 0; Idx < WeightValues[0].Num(); Idx++)
 				{
 					int32 TotalWeight = 0;
 					int32 MaxLayerIdx = -1;
 					int32 MaxWeight = INT_MIN;
 
-					for( int32 WeightLayerIndex = 0; WeightLayerIndex < WeightValues.Num(); WeightLayerIndex++ )
+					for (int32 WeightLayerIndex = 0; WeightLayerIndex < WeightValues.Num(); WeightLayerIndex++)
 					{
 						if (!IsNoBlendArray[WeightLayerIndex])
 						{
@@ -1922,9 +1922,9 @@ void ALandscapeProxy::Import(FGuid Guid, int32 VertsX, int32 VertsY,
 					else if (TotalWeight != 255)
 					{
 						// normalization...
-						float Factor = 255.f/TotalWeight;
+						float Factor = 255.f / TotalWeight;
 						TotalWeight = 0;
-						for( int32 WeightLayerIndex = 0; WeightLayerIndex < WeightValues.Num(); WeightLayerIndex++ )
+						for (int32 WeightLayerIndex = 0; WeightLayerIndex < WeightValues.Num(); WeightLayerIndex++)
 						{
 							if (!IsNoBlendArray[WeightLayerIndex])
 							{
@@ -1945,11 +1945,11 @@ void ALandscapeProxy::Import(FGuid Guid, int32 VertsX, int32 VertsY,
 
 	// Remember where we have spare texture channels.
 	TArray<FWeightmapTextureAllocation> TextureAllocations;
-		
+
 	for (int32 ComponentY = 0; ComponentY < NumSectionsY; ComponentY++)
 	{
 		int32 HmY = ComponentY / ComponentsPerHeightmap;
-		int32 HeightmapOffsetY = (ComponentY - ComponentsPerHeightmap*HmY) * NumSubsections * (SubsectionSizeQuads+1);
+		int32 HeightmapOffsetY = (ComponentY - ComponentsPerHeightmap*HmY) * NumSubsections * (SubsectionSizeQuads + 1);
 
 		for (int32 ComponentX = 0; ComponentX < NumSectionsX; ComponentX++)
 		{
@@ -1962,43 +1962,43 @@ void ALandscapeProxy::Import(FGuid Guid, int32 VertsX, int32 VertsY,
 			TArray<TArray<uint8> >& WeightValues = ComponentWeightValues[ComponentX + ComponentY*NumSectionsX];
 
 			// Heightmap offsets
-			int32 HeightmapOffsetX = (ComponentX - ComponentsPerHeightmap*HmX) * NumSubsections * (SubsectionSizeQuads+1);
+			int32 HeightmapOffsetX = (ComponentX - ComponentsPerHeightmap*HmX) * NumSubsections * (SubsectionSizeQuads + 1);
 
-			LandscapeComponent->HeightmapScaleBias = FVector4( 1.f / (float)HeightmapInfo.HeightmapSizeU, 1.f / (float)HeightmapInfo.HeightmapSizeV, (float)((HeightmapOffsetX)) / (float)HeightmapInfo.HeightmapSizeU, ((float)(HeightmapOffsetY)) / (float)HeightmapInfo.HeightmapSizeV );
+			LandscapeComponent->HeightmapScaleBias = FVector4(1.f / (float)HeightmapInfo.HeightmapSizeU, 1.f / (float)HeightmapInfo.HeightmapSizeV, (float)((HeightmapOffsetX)) / (float)HeightmapInfo.HeightmapSizeU, ((float)(HeightmapOffsetY)) / (float)HeightmapInfo.HeightmapSizeV);
 			LandscapeComponent->HeightmapTexture = HeightmapInfo.HeightmapTexture;
 
 			// Weightmap is sized the same as the component
-			int32 WeightmapSize = (SubsectionSizeQuads+1) * NumSubsections;
+			int32 WeightmapSize = (SubsectionSizeQuads + 1) * NumSubsections;
 			// Should be power of two
-			check(((WeightmapSize-1) & WeightmapSize) == 0);
+			check(((WeightmapSize - 1) & WeightmapSize) == 0);
 
-			LandscapeComponent->WeightmapScaleBias = FVector4( 1.f / (float)WeightmapSize, 1.f / (float)WeightmapSize, 0.5f / (float)WeightmapSize, 0.5f / (float)WeightmapSize);
-			LandscapeComponent->WeightmapSubsectionOffset =  (float)(SubsectionSizeQuads+1) / (float)WeightmapSize;
+			LandscapeComponent->WeightmapScaleBias = FVector4(1.f / (float)WeightmapSize, 1.f / (float)WeightmapSize, 0.5f / (float)WeightmapSize, 0.5f / (float)WeightmapSize);
+			LandscapeComponent->WeightmapSubsectionOffset = (float)(SubsectionSizeQuads + 1) / (float)WeightmapSize;
 
 			// Pointers to the texture data where we'll store each layer. Stride is 4 (FColor)
 			TArray<uint8*> WeightmapTextureDataPointers;
 
-			UE_LOG(LogLandscape, Log, TEXT("%s needs %d weightmap channels"), *LandscapeComponent->GetName(),WeightValues.Num());
+			UE_LOG(LogLandscape, Log, TEXT("%s needs %d weightmap channels"), *LandscapeComponent->GetName(), WeightValues.Num());
 
 			// Find texture channels to store each layer.
 			int32 LayerIndex = 0;
-			while( LayerIndex < WeightValues.Num() )
+			while (LayerIndex < WeightValues.Num())
 			{
-				int32 RemainingLayers = WeightValues.Num()-LayerIndex;
+				int32 RemainingLayers = WeightValues.Num() - LayerIndex;
 
 				int32 BestAllocationIndex = -1;
 
 				// if we need less than 4 channels, try to find them somewhere to put all of them
-				if( RemainingLayers < 4 )
+				if (RemainingLayers < 4)
 				{
 					int32 BestDistSquared = MAX_int32;
-					for( int32 TryAllocIdx=0;TryAllocIdx<TextureAllocations.Num();TryAllocIdx++ )
+					for (int32 TryAllocIdx = 0; TryAllocIdx < TextureAllocations.Num(); TryAllocIdx++)
 					{
-						if( TextureAllocations[TryAllocIdx].ChannelsInUse + RemainingLayers <= 4 )
+						if (TextureAllocations[TryAllocIdx].ChannelsInUse + RemainingLayers <= 4)
 						{
 							FWeightmapTextureAllocation& TryAllocation = TextureAllocations[TryAllocIdx];
-							int32 TryDistSquared = FMath::Square(TryAllocation.X-ComponentX) + FMath::Square(TryAllocation.Y-ComponentY);
-							if( TryDistSquared < BestDistSquared )
+							int32 TryDistSquared = FMath::Square(TryAllocation.X - ComponentX) + FMath::Square(TryAllocation.Y - ComponentY);
+							if (TryDistSquared < BestDistSquared)
 							{
 								BestDistSquared = TryDistSquared;
 								BestAllocationIndex = TryAllocIdx;
@@ -2007,19 +2007,19 @@ void ALandscapeProxy::Import(FGuid Guid, int32 VertsX, int32 VertsY,
 					}
 				}
 
-				if( BestAllocationIndex != -1 )
+				if (BestAllocationIndex != -1)
 				{
 					FWeightmapTextureAllocation& Allocation = TextureAllocations[BestAllocationIndex];
 					FLandscapeWeightmapUsage& WeightmapUsage = WeightmapUsageMap.FindChecked(Allocation.Texture);
 
-					UE_LOG(LogLandscape, Log, TEXT("  ==> Storing %d channels starting at %s[%d]"), RemainingLayers, *Allocation.Texture->GetName(), Allocation.ChannelsInUse );
+					UE_LOG(LogLandscape, Log, TEXT("  ==> Storing %d channels starting at %s[%d]"), RemainingLayers, *Allocation.Texture->GetName(), Allocation.ChannelsInUse);
 
-					for( int32 i=0;i<RemainingLayers;i++ )
+					for (int32 i = 0; i < RemainingLayers; i++)
 					{
-						LandscapeComponent->WeightmapLayerAllocations[LayerIndex+i].WeightmapTextureIndex = LandscapeComponent->WeightmapTextures.Num();
-						LandscapeComponent->WeightmapLayerAllocations[LayerIndex+i].WeightmapTextureChannel = Allocation.ChannelsInUse;
+						LandscapeComponent->WeightmapLayerAllocations[LayerIndex + i].WeightmapTextureIndex = LandscapeComponent->WeightmapTextures.Num();
+						LandscapeComponent->WeightmapLayerAllocations[LayerIndex + i].WeightmapTextureChannel = Allocation.ChannelsInUse;
 						WeightmapUsage.ChannelUsage[Allocation.ChannelsInUse] = LandscapeComponent;
-						switch( Allocation.ChannelsInUse )
+						switch (Allocation.ChannelsInUse)
 						{
 						case 1:
 							WeightmapTextureDataPointers.Add((uint8*)&Allocation.TextureData->G);
@@ -2047,36 +2047,36 @@ void ALandscapeProxy::Import(FGuid Guid, int32 VertsX, int32 VertsY,
 					UTexture2D* WeightmapTexture = CreateLandscapeTexture(WeightmapSize, WeightmapSize, TEXTUREGROUP_Terrain_Weightmap, TSF_BGRA8);
 					FColor* MipData = (FColor*)WeightmapTexture->Source.LockMip(0);
 
-					int32 ThisAllocationLayers = FMath::Min<int32>(RemainingLayers,4);
-					new(TextureAllocations) FWeightmapTextureAllocation(ComponentX,ComponentY,ThisAllocationLayers,WeightmapTexture,MipData);
+					int32 ThisAllocationLayers = FMath::Min<int32>(RemainingLayers, 4);
+					new(TextureAllocations)FWeightmapTextureAllocation(ComponentX, ComponentY, ThisAllocationLayers, WeightmapTexture, MipData);
 					FLandscapeWeightmapUsage& WeightmapUsage = WeightmapUsageMap.Add(WeightmapTexture, FLandscapeWeightmapUsage());
 
 					UE_LOG(LogLandscape, Log, TEXT("  ==> Storing %d channels in new texture %s"), ThisAllocationLayers, *WeightmapTexture->GetName());
 
 					WeightmapTextureDataPointers.Add((uint8*)&MipData->R);
-					LandscapeComponent->WeightmapLayerAllocations[LayerIndex+0].WeightmapTextureIndex = LandscapeComponent->WeightmapTextures.Num();
-					LandscapeComponent->WeightmapLayerAllocations[LayerIndex+0].WeightmapTextureChannel = 0;
+					LandscapeComponent->WeightmapLayerAllocations[LayerIndex + 0].WeightmapTextureIndex = LandscapeComponent->WeightmapTextures.Num();
+					LandscapeComponent->WeightmapLayerAllocations[LayerIndex + 0].WeightmapTextureChannel = 0;
 					WeightmapUsage.ChannelUsage[0] = LandscapeComponent;
 
-					if( ThisAllocationLayers > 1 )
+					if (ThisAllocationLayers > 1)
 					{
 						WeightmapTextureDataPointers.Add((uint8*)&MipData->G);
-						LandscapeComponent->WeightmapLayerAllocations[LayerIndex+1].WeightmapTextureIndex = LandscapeComponent->WeightmapTextures.Num();
-						LandscapeComponent->WeightmapLayerAllocations[LayerIndex+1].WeightmapTextureChannel = 1;
+						LandscapeComponent->WeightmapLayerAllocations[LayerIndex + 1].WeightmapTextureIndex = LandscapeComponent->WeightmapTextures.Num();
+						LandscapeComponent->WeightmapLayerAllocations[LayerIndex + 1].WeightmapTextureChannel = 1;
 						WeightmapUsage.ChannelUsage[1] = LandscapeComponent;
 
-						if( ThisAllocationLayers > 2 )
+						if (ThisAllocationLayers > 2)
 						{
 							WeightmapTextureDataPointers.Add((uint8*)&MipData->B);
-							LandscapeComponent->WeightmapLayerAllocations[LayerIndex+2].WeightmapTextureIndex = LandscapeComponent->WeightmapTextures.Num();
-							LandscapeComponent->WeightmapLayerAllocations[LayerIndex+2].WeightmapTextureChannel = 2;
+							LandscapeComponent->WeightmapLayerAllocations[LayerIndex + 2].WeightmapTextureIndex = LandscapeComponent->WeightmapTextures.Num();
+							LandscapeComponent->WeightmapLayerAllocations[LayerIndex + 2].WeightmapTextureChannel = 2;
 							WeightmapUsage.ChannelUsage[2] = LandscapeComponent;
 
-							if( ThisAllocationLayers > 3 )
+							if (ThisAllocationLayers > 3)
 							{
 								WeightmapTextureDataPointers.Add((uint8*)&MipData->A);
-								LandscapeComponent->WeightmapLayerAllocations[LayerIndex+3].WeightmapTextureIndex = LandscapeComponent->WeightmapTextures.Num();
-								LandscapeComponent->WeightmapLayerAllocations[LayerIndex+3].WeightmapTextureChannel = 3;
+								LandscapeComponent->WeightmapLayerAllocations[LayerIndex + 3].WeightmapTextureIndex = LandscapeComponent->WeightmapTextures.Num();
+								LandscapeComponent->WeightmapLayerAllocations[LayerIndex + 3].WeightmapTextureChannel = 3;
 								WeightmapUsage.ChannelUsage[3] = LandscapeComponent;
 							}
 						}
@@ -2088,52 +2088,52 @@ void ALandscapeProxy::Import(FGuid Guid, int32 VertsX, int32 VertsY,
 			}
 			check(WeightmapTextureDataPointers.Num() == WeightValues.Num());
 
-			FVector* LocalVerts = new FVector[FMath::Square(ComponentSizeQuads+1)];
+			FVector* LocalVerts = new FVector[FMath::Square(ComponentSizeQuads + 1)];
 
-			for( int32 SubsectionY = 0;SubsectionY < NumSubsections;SubsectionY++ )
+			for (int32 SubsectionY = 0; SubsectionY < NumSubsections; SubsectionY++)
 			{
-				for( int32 SubsectionX = 0;SubsectionX < NumSubsections;SubsectionX++ )
+				for (int32 SubsectionX = 0; SubsectionX < NumSubsections; SubsectionX++)
 				{
-					for( int32 SubY=0;SubY<=SubsectionSizeQuads;SubY++ )
+					for (int32 SubY = 0; SubY <= SubsectionSizeQuads; SubY++)
 					{
-						for( int32 SubX=0;SubX<=SubsectionSizeQuads;SubX++ )
+						for (int32 SubX = 0; SubX <= SubsectionSizeQuads; SubX++)
 						{
 							// X/Y of the vertex we're looking at in component's coordinates.
 							int32 CompX = SubsectionSizeQuads * SubsectionX + SubX;
 							int32 CompY = SubsectionSizeQuads * SubsectionY + SubY;
 
 							// X/Y of the vertex we're looking indexed into the texture data
-							int32 TexX = (SubsectionSizeQuads+1) * SubsectionX + SubX;
-							int32 TexY = (SubsectionSizeQuads+1) * SubsectionY + SubY;
+							int32 TexX = (SubsectionSizeQuads + 1) * SubsectionX + SubX;
+							int32 TexY = (SubsectionSizeQuads + 1) * SubsectionY + SubY;
 
-							int32 WeightSrcDataIdx = CompY * (ComponentSizeQuads+1) + CompX;
+							int32 WeightSrcDataIdx = CompY * (ComponentSizeQuads + 1) + CompX;
 							int32 HeightTexDataIdx = (HeightmapOffsetX + TexX) + (HeightmapOffsetY + TexY) * (HeightmapInfo.HeightmapSizeU);
 
-							int32 WeightTexDataIdx = (TexX) + (TexY) * (WeightmapSize);
+							int32 WeightTexDataIdx = (TexX)+(TexY)* (WeightmapSize);
 
 							// copy height and normal data
 							uint16 HeightValue = HEIGHTDATA(CompX + LandscapeComponent->GetSectionBase().X, CompY + LandscapeComponent->GetSectionBase().Y);
-							FVector Normal = VertexNormals[CompX+LandscapeComponent->GetSectionBase().X + (NumPatchesX+1)*(CompY+LandscapeComponent->GetSectionBase().Y)].SafeNormal();
+							FVector Normal = VertexNormals[CompX + LandscapeComponent->GetSectionBase().X + (NumPatchesX + 1)*(CompY + LandscapeComponent->GetSectionBase().Y)].SafeNormal();
 
 							HeightmapInfo.HeightmapTextureMipData[0][HeightTexDataIdx].R = HeightValue >> 8;
 							HeightmapInfo.HeightmapTextureMipData[0][HeightTexDataIdx].G = HeightValue & 255;
-							HeightmapInfo.HeightmapTextureMipData[0][HeightTexDataIdx].B = FMath::RoundToInt( 127.5f * (Normal.X + 1.f) );
-							HeightmapInfo.HeightmapTextureMipData[0][HeightTexDataIdx].A = FMath::RoundToInt( 127.5f * (Normal.Y + 1.f) );
+							HeightmapInfo.HeightmapTextureMipData[0][HeightTexDataIdx].B = FMath::RoundToInt(127.5f * (Normal.X + 1.f));
+							HeightmapInfo.HeightmapTextureMipData[0][HeightTexDataIdx].A = FMath::RoundToInt(127.5f * (Normal.Y + 1.f));
 
-							for( int32 WeightmapIndex=0;WeightmapIndex<WeightValues.Num(); WeightmapIndex++ )
+							for (int32 WeightmapIndex = 0; WeightmapIndex < WeightValues.Num(); WeightmapIndex++)
 							{
-								WeightmapTextureDataPointers[WeightmapIndex][WeightTexDataIdx*4] = WeightValues[WeightmapIndex][WeightSrcDataIdx];
+								WeightmapTextureDataPointers[WeightmapIndex][WeightTexDataIdx * 4] = WeightValues[WeightmapIndex][WeightSrcDataIdx];
 							}
 
 							// Get local space verts
-							FVector LocalVertex( CompX, CompY, LandscapeDataAccess::GetLocalHeight(HeightValue) );
-							LocalVerts[(LandscapeComponent->ComponentSizeQuads+1) * CompY + CompX] = LocalVertex;
+							FVector LocalVertex(CompX, CompY, LandscapeDataAccess::GetLocalHeight(HeightValue));
+							LocalVerts[(LandscapeComponent->ComponentSizeQuads + 1) * CompY + CompX] = LocalVertex;
 						}
 					}
 				}
 			}
 
-			LandscapeComponent->CachedLocalBox = FBox(LocalVerts, FMath::Square(ComponentSizeQuads+1));
+			LandscapeComponent->CachedLocalBox = FBox(LocalVerts, FMath::Square(ComponentSizeQuads + 1));
 			delete[] LocalVerts;
 
 			// Update MaterialInstance
@@ -2142,7 +2142,7 @@ void ALandscapeProxy::Import(FGuid Guid, int32 VertsX, int32 VertsY,
 	}
 
 	// Unlock the weightmaps' base mips
-	for( int32 AllocationIndex=0;AllocationIndex<TextureAllocations.Num();AllocationIndex++ )
+	for (int32 AllocationIndex = 0; AllocationIndex < TextureAllocations.Num(); AllocationIndex++)
 	{
 		UTexture2D* WeightmapTexture = TextureAllocations[AllocationIndex].Texture;
 		FColor* BaseMipData = TextureAllocations[AllocationIndex].TextureData;
@@ -2166,13 +2166,13 @@ void ALandscapeProxy::Import(FGuid Guid, int32 VertsX, int32 VertsY,
 			FHeightmapInfo& HeightmapInfo = HeightmapInfos[HmX + HmY * NumHeightmapsX];
 
 			ULandscapeComponent* LandscapeComponent = LandscapeComponents[ComponentX + ComponentY*NumSectionsX];
-			LandscapeComponent->GenerateHeightmapMips(HeightmapInfo.HeightmapTextureMipData, ComponentX == NumSectionsX-1 ? MAX_int32 : 0, ComponentY == NumSectionsY-1 ? MAX_int32 : 0);
+			LandscapeComponent->GenerateHeightmapMips(HeightmapInfo.HeightmapTextureMipData, ComponentX == NumSectionsX - 1 ? MAX_int32 : 0, ComponentY == NumSectionsY - 1 ? MAX_int32 : 0);
 			LandscapeComponent->UpdateCollisionHeightData(HeightmapInfo.HeightmapTextureMipData[LandscapeComponent->CollisionMipLevel]);
 			LandscapeComponent->UpdateCollisionLayerData();
 		}
 	}
 
-	for( int32 HmIdx=0;HmIdx<HeightmapInfos.Num();HmIdx++ )
+	for (int32 HmIdx = 0; HmIdx < HeightmapInfos.Num(); HmIdx++)
 	{
 		FHeightmapInfo& HeightmapInfo = HeightmapInfos[HmIdx];
 
@@ -2181,22 +2181,22 @@ void ALandscapeProxy::Import(FGuid Guid, int32 VertsX, int32 VertsY,
 		int32 Mip = HeightmapInfo.HeightmapTextureMipData.Num();
 		int32 MipSizeU = (HeightmapInfo.HeightmapTexture->Source.GetSizeX()) >> Mip;
 		int32 MipSizeV = (HeightmapInfo.HeightmapTexture->Source.GetSizeY()) >> Mip;
-		while( MipSizeU > 1 && MipSizeV > 1 )
+		while (MipSizeU > 1 && MipSizeV > 1)
 		{
 			HeightmapInfo.HeightmapTextureMipData.Add((FColor*)HeightmapInfo.HeightmapTexture->Source.LockMip(Mip));
-			int32 PrevMipSizeU = (HeightmapInfo.HeightmapTexture->Source.GetSizeX()) >> (Mip-1);
-			int32 PrevMipSizeV = (HeightmapInfo.HeightmapTexture->Source.GetSizeY()) >> (Mip-1);
+			int32 PrevMipSizeU = (HeightmapInfo.HeightmapTexture->Source.GetSizeX()) >> (Mip - 1);
+			int32 PrevMipSizeV = (HeightmapInfo.HeightmapTexture->Source.GetSizeY()) >> (Mip - 1);
 
-			for( int32 Y = 0;Y < MipSizeV;Y++ )
+			for (int32 Y = 0; Y < MipSizeV; Y++)
 			{
-				for( int32 X = 0;X < MipSizeU;X++ )
+				for (int32 X = 0; X < MipSizeU; X++)
 				{
-					FColor* TexData = &(HeightmapInfo.HeightmapTextureMipData[Mip])[ X + Y * MipSizeU ];
+					FColor* TexData = &(HeightmapInfo.HeightmapTextureMipData[Mip])[X + Y * MipSizeU];
 
-					FColor *PreMipTexData00 = &(HeightmapInfo.HeightmapTextureMipData[Mip-1])[ (X*2+0) + (Y*2+0)  * PrevMipSizeU ];
-					FColor *PreMipTexData01 = &(HeightmapInfo.HeightmapTextureMipData[Mip-1])[ (X*2+0) + (Y*2+1)  * PrevMipSizeU ];
-					FColor *PreMipTexData10 = &(HeightmapInfo.HeightmapTextureMipData[Mip-1])[ (X*2+1) + (Y*2+0)  * PrevMipSizeU ];
-					FColor *PreMipTexData11 = &(HeightmapInfo.HeightmapTextureMipData[Mip-1])[ (X*2+1) + (Y*2+1)  * PrevMipSizeU ];
+					FColor *PreMipTexData00 = &(HeightmapInfo.HeightmapTextureMipData[Mip - 1])[(X * 2 + 0) + (Y * 2 + 0)  * PrevMipSizeU];
+					FColor *PreMipTexData01 = &(HeightmapInfo.HeightmapTextureMipData[Mip - 1])[(X * 2 + 0) + (Y * 2 + 1)  * PrevMipSizeU];
+					FColor *PreMipTexData10 = &(HeightmapInfo.HeightmapTextureMipData[Mip - 1])[(X * 2 + 1) + (Y * 2 + 0)  * PrevMipSizeU];
+					FColor *PreMipTexData11 = &(HeightmapInfo.HeightmapTextureMipData[Mip - 1])[(X * 2 + 1) + (Y * 2 + 1)  * PrevMipSizeU];
 
 					TexData->R = (((int32)PreMipTexData00->R + (int32)PreMipTexData01->R + (int32)PreMipTexData10->R + (int32)PreMipTexData11->R) >> 2);
 					TexData->G = (((int32)PreMipTexData00->G + (int32)PreMipTexData01->G + (int32)PreMipTexData10->G + (int32)PreMipTexData11->G) >> 2);
@@ -2209,7 +2209,7 @@ void ALandscapeProxy::Import(FGuid Guid, int32 VertsX, int32 VertsY,
 			MipSizeV >>= 1;
 		}
 
-		for( int32 i=0;i<HeightmapInfo.HeightmapTextureMipData.Num();i++ )
+		for (int32 i = 0; i < HeightmapInfo.HeightmapTextureMipData.Num(); i++)
 		{
 			HeightmapInfo.HeightmapTexture->Source.UnlockMip(i);
 		}
@@ -2233,30 +2233,30 @@ bool ALandscapeProxy::ExportToRawMesh(int32 InExportLOD, FRawMesh& OutRawMesh) c
 {
 	TArray<ULandscapeComponent*> RegisteredLandscapeComponents;
 	GetComponents<ULandscapeComponent>(RegisteredLandscapeComponents);
-	
+
 	const FIntRect LandscapeSectionRect = GetBoundingRect();
-	const FVector2D LandscapeUVScale = FVector2D(1.f, 1.f)/FVector2D(LandscapeSectionRect.Size());
+	const FVector2D LandscapeUVScale = FVector2D(1.f, 1.f) / FVector2D(LandscapeSectionRect.Size());
 
 	// User specified LOD to export
-	int32 LandscapeLODToExport = ExportLOD; 
+	int32 LandscapeLODToExport = ExportLOD;
 	if (InExportLOD != INDEX_NONE)
 	{
-		LandscapeLODToExport = FMath::Clamp<int32>(InExportLOD, 0, FMath::CeilLogTwo(SubsectionSizeQuads+1)-1);
+		LandscapeLODToExport = FMath::Clamp<int32>(InExportLOD, 0, FMath::CeilLogTwo(SubsectionSizeQuads + 1) - 1);
 	}
-		
+
 	// Export data for each component
 	for (auto It = RegisteredLandscapeComponents.CreateConstIterator(); It; ++It)
 	{
 		ULandscapeComponent* Component = (*It);
 		FLandscapeComponentDataInterface CDI(Component, LandscapeLODToExport);
-		const int32 ComponentSizeQuadsLOD = ((Component->ComponentSizeQuads+1)>>LandscapeLODToExport)-1;
-		const int32 SubsectionSizeQuadsLOD = ((Component->SubsectionSizeQuads+1)>>LandscapeLODToExport)-1;
+		const int32 ComponentSizeQuadsLOD = ((Component->ComponentSizeQuads + 1) >> LandscapeLODToExport) - 1;
+		const int32 SubsectionSizeQuadsLOD = ((Component->SubsectionSizeQuads + 1) >> LandscapeLODToExport) - 1;
 		const FIntPoint ComponentOffsetQuads = Component->GetSectionBase() - LandscapeSectionOffset - LandscapeSectionRect.Min;
-		const FVector2D ComponentUVOffsetLOD = FVector2D(ComponentOffsetQuads)*((float)ComponentSizeQuadsLOD/ComponentSizeQuads);
-		const FVector2D ComponentUVScaleLOD = LandscapeUVScale*((float)ComponentSizeQuads/ComponentSizeQuadsLOD);
-						
-		const int32 NumFaces = FMath::Square(ComponentSizeQuadsLOD)*2;
-		const int32 NumVertices = NumFaces*3;
+		const FVector2D ComponentUVOffsetLOD = FVector2D(ComponentOffsetQuads)*((float)ComponentSizeQuadsLOD / ComponentSizeQuads);
+		const FVector2D ComponentUVScaleLOD = LandscapeUVScale*((float)ComponentSizeQuads / ComponentSizeQuadsLOD);
+
+		const int32 NumFaces = FMath::Square(ComponentSizeQuadsLOD) * 2;
+		const int32 NumVertices = NumFaces * 3;
 		const int32 VerticesOffset = OutRawMesh.VertexPositions.Num();
 		const int32 IndicesOffset = OutRawMesh.WedgeIndices.Num();
 
@@ -2270,30 +2270,30 @@ bool ALandscapeProxy::ExportToRawMesh(int32 InExportLOD, FRawMesh& OutRawMesh) c
 		OutRawMesh.WedgeTangentY.AddZeroed(NumVertices);
 		OutRawMesh.WedgeTangentZ.AddZeroed(NumVertices);
 		OutRawMesh.WedgeTexCoords[0].AddZeroed(NumVertices);
-		
-				
+
+
 		// Check if there are any holes
 		TArray<uint8> VisDataMap;
-		
+
 		for (int32 AllocIdx = 0; AllocIdx < Component->WeightmapLayerAllocations.Num(); AllocIdx++)
 		{
 			FWeightmapLayerAllocationInfo& AllocInfo = Component->WeightmapLayerAllocations[AllocIdx];
-			if (AllocInfo.LayerInfo == ALandscapeProxy::DataLayer)
+			if (AllocInfo.LayerInfo == ALandscapeProxy::VisibilityLayer)
 			{
 				CDI.GetWeightmapTextureData(AllocInfo.LayerInfo, VisDataMap);
 			}
 		}
 
-		const FIntPoint QuadPattern[6] = 
+		const FIntPoint QuadPattern[6] =
 		{
 			//face 1
-			FIntPoint(0,0),
-			FIntPoint(0,1),
-			FIntPoint(1,1),
+			FIntPoint(0, 0),
+			FIntPoint(0, 1),
+			FIntPoint(1, 1),
 			//face 2
-			FIntPoint(0,0),
-			FIntPoint(1,1),
-			FIntPoint(1,0),
+			FIntPoint(0, 0),
+			FIntPoint(1, 1),
+			FIntPoint(1, 0),
 		};
 
 		const int32 VisThreshold = 170;
@@ -2313,7 +2313,7 @@ bool ALandscapeProxy::ExportToRawMesh(int32 InExportLOD, FRawMesh& OutRawMesh) c
 
 					// Whether this vertex is in hole
 					bool bInvisible = false;
-					if (VisDataMap.Num()) 
+					if (VisDataMap.Num())
 					{
 						bInvisible = (VisDataMap[x + y * (ComponentSizeQuadsLOD + 1)] >= VisThreshold);
 					}
@@ -2322,12 +2322,12 @@ bool ALandscapeProxy::ExportToRawMesh(int32 InExportLOD, FRawMesh& OutRawMesh) c
 					Faces[0] = VertexIdx;
 					Faces[1] = bInvisible ? Faces[0] : VertexIdx + 1;
 					Faces[2] = bInvisible ? Faces[0] : VertexIdx + 2;
-					Faces+=3;
+					Faces += 3;
 
 					Faces[0] = VertexIdx + 3;
 					Faces[1] = bInvisible ? Faces[0] : VertexIdx + 4;
 					Faces[2] = bInvisible ? Faces[0] : VertexIdx + 5;
-					Faces+=3;
+					Faces += 3;
 				}
 
 				// Fill vertices
@@ -2342,9 +2342,9 @@ bool ALandscapeProxy::ExportToRawMesh(int32 InExportLOD, FRawMesh& OutRawMesh) c
 
 					OutRawMesh.VertexPositions[VertexIdx] = LocalVertexPos;
 					OutRawMesh.WedgeTangentX[VertexIdx] = LocalTangentX;
-					OutRawMesh.WedgeTangentY[VertexIdx] = LocalTangentY; 
+					OutRawMesh.WedgeTangentY[VertexIdx] = LocalTangentY;
 					OutRawMesh.WedgeTangentZ[VertexIdx] = LocalTangentZ;
-					
+
 					OutRawMesh.WedgeTexCoords[0][VertexIdx] = (ComponentUVOffsetLOD + FVector2D(VertexX, VertexY))*ComponentUVScaleLOD;
 
 					VertexIdx++;
@@ -2352,11 +2352,11 @@ bool ALandscapeProxy::ExportToRawMesh(int32 InExportLOD, FRawMesh& OutRawMesh) c
 			}
 		}
 	}
-	
+
 	// Add lightmap UVs
 	OutRawMesh.WedgeTexCoords[1].Append(OutRawMesh.WedgeTexCoords[0]);
 
-	return true;			
+	return true;
 }
 
 FIntRect ALandscapeProxy::GetBoundingRect() const
@@ -2367,17 +2367,17 @@ FIntRect ALandscapeProxy::GetBoundingRect() const
 	{
 		Rect.Include(LandscapeComponents[CompIdx]->GetSectionBase());
 	}
-	
+
 	if (Rect.Width() > 0 && Rect.Height() > 0)
 	{
-		Rect.Max+= FIntPoint(ComponentSizeQuads, ComponentSizeQuads);
-		Rect-= LandscapeSectionOffset;
+		Rect.Max += FIntPoint(ComponentSizeQuads, ComponentSizeQuads);
+		Rect -= LandscapeSectionOffset;
 	}
 	else
 	{
 		Rect = FIntRect();
 	}
-	
+
 	return Rect;
 }
 
@@ -2394,11 +2394,11 @@ bool ALandscape::HasAllComponent()
 
 bool ULandscapeInfo::GetLandscapeExtent(int32& MinX, int32& MinY, int32& MaxX, int32& MaxY)
 {
-	MinX = MAX_int32; 
-	MinY = MAX_int32; 
-	MaxX = MIN_int32; 
+	MinX = MAX_int32;
+	MinY = MAX_int32;
+	MaxX = MIN_int32;
 	MaxY = MIN_int32;
-	
+
 	// Find range of entire landscape
 	for (auto It = XYtoComponentMap.CreateIterator(); It; ++It)
 	{
@@ -2434,12 +2434,12 @@ bool ULandscapeInfo::GetSelectedExtent(int32& MinX, int32& MinY, int32& MaxX, in
 	return MinX != MAX_int32;
 }
 
-FVector ULandscapeInfo::GetLandscapeCenterPos(float& LengthZ, int32 MinX /*= MAX_INT*/, int32 MinY /*= MAX_INT*/, int32 MaxX /*= MIN_INT*/, int32 MaxY /*= MIN_INT*/ )
+FVector ULandscapeInfo::GetLandscapeCenterPos(float& LengthZ, int32 MinX /*= MAX_INT*/, int32 MinY /*= MAX_INT*/, int32 MaxX /*= MIN_INT*/, int32 MaxY /*= MIN_INT*/)
 {
 	// MinZ, MaxZ is Local coordinate
 	float MaxZ = -HALF_WORLD_MAX, MinZ = HALF_WORLD_MAX;
 	const float ScaleZ = DrawScale.Z;
-	
+
 	if (MinX == MAX_int32)
 	{
 		// Find range of entire landscape
@@ -2449,13 +2449,13 @@ FVector ULandscapeInfo::GetLandscapeCenterPos(float& LengthZ, int32 MinX /*= MAX
 			Comp->GetComponentExtent(MinX, MinY, MaxX, MaxY);
 		}
 
-		const int32 Dist = (ComponentSizeQuads+1) >> 1; // Should be same in ALandscapeGizmoActiveActor::SetTargetLandscape
-		FVector2D MidPoint(((float)(MinX+MaxX))/2.f, ((float)(MinY+MaxY))/2.f);
+		const int32 Dist = (ComponentSizeQuads + 1) >> 1; // Should be same in ALandscapeGizmoActiveActor::SetTargetLandscape
+		FVector2D MidPoint(((float)(MinX + MaxX)) / 2.f, ((float)(MinY + MaxY)) / 2.f);
 		MinX = FMath::FloorToInt(MidPoint.X) - Dist;
 		MaxX = FMath::CeilToInt(MidPoint.X) + Dist;
 		MinY = FMath::FloorToInt(MidPoint.Y) - Dist;
 		MaxY = FMath::CeilToInt(MidPoint.Y) + Dist;
-		check(MidPoint.X == ((float)(MinX+MaxX))/2.f && MidPoint.Y == ((float)(MinY+MaxY))/2.f );
+		check(MidPoint.X == ((float)(MinX + MaxX)) / 2.f && MidPoint.Y == ((float)(MinY + MaxY)) / 2.f);
 	}
 
 	check(MinX != MAX_int32);
@@ -2503,9 +2503,9 @@ FVector ULandscapeInfo::GetLandscapeCenterPos(float& LengthZ, int32 MinX /*= MAX
 		MaxZ = +MarginZ;
 		MinZ = -MarginZ;
 	}
-	LengthZ = (MaxZ - MinZ + 2*MarginZ) * ScaleZ;
-	
-	const FVector LocalPosition(((float)(MinX+MaxX))/2.f, ((float)(MinY+MaxY))/2.f, MinZ - MarginZ);
+	LengthZ = (MaxZ - MinZ + 2 * MarginZ) * ScaleZ;
+
+	const FVector LocalPosition(((float)(MinX + MaxX)) / 2.f, ((float)(MinY + MaxY)) / 2.f, MinZ - MarginZ);
 	//return GetLandscapeProxy()->TransformLandscapeLocationToWorld(LocalPosition);
 	return GetLandscapeProxy()->LandscapeActorToWorld().TransformPosition(LocalPosition);
 }
@@ -2527,39 +2527,39 @@ bool ULandscapeInfo::IsValidPosition(int32 X, int32 Y)
 
 void ULandscapeInfo::Export(const TArray<ULandscapeLayerInfoObject*>& LayerInfos, const TArray<FString>& Filenames)
 {
-	check( Filenames.Num() > 0 );
+	check(Filenames.Num() > 0);
 
-	int32 MinX=MAX_int32;
-	int32 MinY=MAX_int32;
-	int32 MaxX=-MAX_int32;
-	int32 MaxY=-MAX_int32;
+	int32 MinX = MAX_int32;
+	int32 MinY = MAX_int32;
+	int32 MaxX = -MAX_int32;
+	int32 MaxY = -MAX_int32;
 
-	if( !GetLandscapeExtent(MinX,MinY,MaxX,MaxY) )
+	if (!GetLandscapeExtent(MinX, MinY, MaxX, MaxY))
 	{
 		return;
 	}
 
-	GWarn->BeginSlowTask( LOCTEXT("BeginExportingLandscapeTask", "Exporting Landscape"), true);
+	GWarn->BeginSlowTask(LOCTEXT("BeginExportingLandscapeTask", "Exporting Landscape"), true);
 
 	FLandscapeEditDataInterface LandscapeEdit(this);
 
 	TArray<uint8> HeightData;
-	HeightData.AddZeroed((1+MaxX-MinX)*(1+MaxY-MinY)*sizeof(uint16));
-	LandscapeEdit.GetHeightDataFast(MinX,MinY,MaxX,MaxY,(uint16*)HeightData.GetTypedData(),0);
-	FFileHelper::SaveArrayToFile(HeightData,*Filenames[0]);
+	HeightData.AddZeroed((1 + MaxX - MinX)*(1 + MaxY - MinY)*sizeof(uint16));
+	LandscapeEdit.GetHeightDataFast(MinX, MinY, MaxX, MaxY, (uint16*)HeightData.GetTypedData(), 0);
+	FFileHelper::SaveArrayToFile(HeightData, *Filenames[0]);
 
-	for( int32 i=1;i<Filenames.Num();i++ )
+	for (int32 i = 1; i < Filenames.Num(); i++)
 	{
-		if ( i <= LayerInfos.Num())
+		if (i <= LayerInfos.Num())
 		{
 			TArray<uint8> WeightData;
-			WeightData.AddZeroed((1+MaxX-MinX)*(1+MaxY-MinY));
-			ULandscapeLayerInfoObject* LayerInfo = LayerInfos[i-1];
+			WeightData.AddZeroed((1 + MaxX - MinX)*(1 + MaxY - MinY));
+			ULandscapeLayerInfoObject* LayerInfo = LayerInfos[i - 1];
 			if (LayerInfo)
 			{
-				LandscapeEdit.GetWeightDataFast(LayerInfo, MinX,MinY,MaxX,MaxY,WeightData.GetTypedData(),0);
+				LandscapeEdit.GetWeightDataFast(LayerInfo, MinX, MinY, MaxX, MaxY, WeightData.GetTypedData(), 0);
 			}
-			FFileHelper::SaveArrayToFile(WeightData,*Filenames[i]);
+			FFileHelper::SaveArrayToFile(WeightData, *Filenames[i]);
 		}
 	}
 
@@ -2568,22 +2568,22 @@ void ULandscapeInfo::Export(const TArray<ULandscapeLayerInfoObject*>& LayerInfos
 
 void ULandscapeInfo::ExportHeightmap(const FString& Filename)
 {
-	int32 MinX=MAX_int32;
-	int32 MinY=MAX_int32;
-	int32 MaxX=-MAX_int32;
-	int32 MaxY=-MAX_int32;
+	int32 MinX = MAX_int32;
+	int32 MinY = MAX_int32;
+	int32 MaxX = -MAX_int32;
+	int32 MaxY = -MAX_int32;
 
-	if( !GetLandscapeExtent(MinX,MinY,MaxX,MaxY) )
+	if (!GetLandscapeExtent(MinX, MinY, MaxX, MaxY))
 	{
 		return;
 	}
 
-	GWarn->BeginSlowTask( LOCTEXT("BeginExportingLandscapeHeightmapTask", "Exporting Landscape Heightmap"), true);
+	GWarn->BeginSlowTask(LOCTEXT("BeginExportingLandscapeHeightmapTask", "Exporting Landscape Heightmap"), true);
 
 	FLandscapeEditDataInterface LandscapeEdit(this);
 
 	TArray<uint8> HeightData;
-	HeightData.AddZeroed((MaxX-MinX+1) * (MaxY-MinY+1) * sizeof(uint16));
+	HeightData.AddZeroed((MaxX - MinX + 1) * (MaxY - MinY + 1) * sizeof(uint16));
 	LandscapeEdit.GetHeightDataFast(MinX, MinY, MaxX, MaxY, (uint16*)HeightData.GetData(), 0);
 
 	if (Filename.EndsWith(".png"))
@@ -2592,7 +2592,7 @@ void ULandscapeInfo::ExportHeightmap(const FString& Filename)
 		IImageWrapperPtr ImageWrapper = ImageWrapperModule.CreateImageWrapper(EImageFormat::PNG);
 
 		const TArray<uint8>* RawData = NULL;
-		if (ImageWrapper->SetRaw(HeightData.GetData(), HeightData.Num(), (MaxX-MinX+1), (MaxY-MinY+1), ERGBFormat::Gray, 16))
+		if (ImageWrapper->SetRaw(HeightData.GetData(), HeightData.Num(), (MaxX - MinX + 1), (MaxY - MinY + 1), ERGBFormat::Gray, 16))
 		{
 			HeightData = ImageWrapper->GetCompressed();
 		}
@@ -2605,20 +2605,20 @@ void ULandscapeInfo::ExportHeightmap(const FString& Filename)
 
 void ULandscapeInfo::ExportLayer(ULandscapeLayerInfoObject* LayerInfo, const FString& Filename)
 {
-	int32 MinX=MAX_int32;
-	int32 MinY=MAX_int32;
-	int32 MaxX=-MAX_int32;
-	int32 MaxY=-MAX_int32;
+	int32 MinX = MAX_int32;
+	int32 MinY = MAX_int32;
+	int32 MaxX = -MAX_int32;
+	int32 MaxY = -MAX_int32;
 
-	if (!GetLandscapeExtent(MinX,MinY,MaxX,MaxY))
+	if (!GetLandscapeExtent(MinX, MinY, MaxX, MaxY))
 	{
 		return;
 	}
 
-	GWarn->BeginSlowTask( LOCTEXT("BeginExportingLandscapeWeightmapTask", "Exporting Landscape Layer Weightmap"), true);
+	GWarn->BeginSlowTask(LOCTEXT("BeginExportingLandscapeWeightmapTask", "Exporting Landscape Layer Weightmap"), true);
 
 	TArray<uint8> WeightData;
-	WeightData.AddZeroed((MaxX-MinX+1) * (MaxY-MinY+1));
+	WeightData.AddZeroed((MaxX - MinX + 1) * (MaxY - MinY + 1));
 	if (LayerInfo)
 	{
 		FLandscapeEditDataInterface LandscapeEdit(this);
@@ -2631,7 +2631,7 @@ void ULandscapeInfo::ExportLayer(ULandscapeLayerInfoObject* LayerInfo, const FSt
 		IImageWrapperPtr ImageWrapper = ImageWrapperModule.CreateImageWrapper(EImageFormat::PNG);
 
 		const TArray<uint8>* RawData = NULL;
-		if (ImageWrapper->SetRaw(WeightData.GetData(), WeightData.Num(), (MaxX-MinX+1), (MaxY-MinY+1), ERGBFormat::Gray, 8))
+		if (ImageWrapper->SetRaw(WeightData.GetData(), WeightData.Num(), (MaxX - MinX + 1), (MaxY - MinY + 1), ERGBFormat::Gray, 8))
 		{
 			WeightData = ImageWrapper->GetCompressed();
 		}
@@ -2654,8 +2654,8 @@ public:
 	typedef FVector InterpolantType;
 
 	/** Initialization constructor. */
-	FLandscapeSplineHeightsRasterPolicy(TArray<uint16>& InData, int32 InMinX, int32 InMinY, int32 InMaxX, int32 InMaxY, bool InbRaiseTerrain, bool InbLowerTerrain):
-	Data(InData),
+	FLandscapeSplineHeightsRasterPolicy(TArray<uint16>& InData, int32 InMinX, int32 InMinY, int32 InMaxX, int32 InMaxY, bool InbRaiseTerrain, bool InbLowerTerrain) :
+		Data(InData),
 		MinX(InMinX),
 		MinY(InMinY),
 		MaxX(InMaxX),
@@ -2676,14 +2676,14 @@ protected:
 
 	void ProcessPixel(int32 X, int32 Y, const InterpolantType& Interpolant, bool BackFacing)
 	{
-		const float CosInterpX = (Interpolant.X >= 1 ? 1 : 0.5f - 0.5f * FMath::Cos(Interpolant.X * PI) );
-		const float CosInterpY = (Interpolant.Y >= 1 ? 1 : 0.5f - 0.5f * FMath::Cos(Interpolant.Y * PI) );
+		const float CosInterpX = (Interpolant.X >= 1 ? 1 : 0.5f - 0.5f * FMath::Cos(Interpolant.X * PI));
+		const float CosInterpY = (Interpolant.Y >= 1 ? 1 : 0.5f - 0.5f * FMath::Cos(Interpolant.Y * PI));
 		const float Alpha = CosInterpX * CosInterpY;
-		uint16& Dest = Data[(Y-MinY)*(1+MaxX-MinX) + X-MinX];
+		uint16& Dest = Data[(Y - MinY)*(1 + MaxX - MinX) + X - MinX];
 		float Value = FMath::Lerp((float)Dest, Interpolant.Z, Alpha);
 		uint16 DValue = (uint32)FMath::Clamp<float>(Value, 0, LandscapeDataAccess::MaxValue);
-		if ( (bRaiseTerrain && DValue > Dest) ||
-			(bLowerTerrain && DValue < Dest) )
+		if ((bRaiseTerrain && DValue > Dest) ||
+			(bLowerTerrain && DValue < Dest))
 		{
 			Dest = DValue;
 		}
@@ -2692,7 +2692,7 @@ protected:
 private:
 	TArray<uint16>& Data;
 	int32 MinX, MinY, MaxX, MaxY;
-	uint32 bRaiseTerrain:1, bLowerTerrain:1;
+	uint32 bRaiseTerrain : 1, bLowerTerrain : 1;
 };
 
 class FLandscapeSplineBlendmaskRasterPolicy
@@ -2702,8 +2702,8 @@ public:
 	typedef FVector InterpolantType;
 
 	/** Initialization constructor. */
-	FLandscapeSplineBlendmaskRasterPolicy(TArray<uint8>& InData, int32 InMinX, int32 InMinY, int32 InMaxX, int32 InMaxY):
-	Data(InData),
+	FLandscapeSplineBlendmaskRasterPolicy(TArray<uint8>& InData, int32 InMinX, int32 InMinY, int32 InMaxX, int32 InMaxY) :
+		Data(InData),
 		MinX(InMinX),
 		MinY(InMinY),
 		MaxX(InMaxX),
@@ -2722,10 +2722,10 @@ protected:
 
 	void ProcessPixel(int32 X, int32 Y, const InterpolantType& Interpolant, bool BackFacing)
 	{
-		const float CosInterpX = (Interpolant.X >= 1 ? 1 : 0.5f - 0.5f * FMath::Cos(Interpolant.X * PI) );
-		const float CosInterpY = (Interpolant.Y >= 1 ? 1 : 0.5f - 0.5f * FMath::Cos(Interpolant.Y * PI) );
+		const float CosInterpX = (Interpolant.X >= 1 ? 1 : 0.5f - 0.5f * FMath::Cos(Interpolant.X * PI));
+		const float CosInterpY = (Interpolant.Y >= 1 ? 1 : 0.5f - 0.5f * FMath::Cos(Interpolant.Y * PI));
 		const float Alpha = CosInterpX * CosInterpY;
-		uint8& Dest = Data[(Y-MinY)*(1+MaxX-MinX) + X-MinX];
+		uint8& Dest = Data[(Y - MinY)*(1 + MaxX - MinX) + X - MinX];
 		float Value = FMath::Lerp((float)Dest, Interpolant.Z, Alpha);
 		Dest = (uint32)FMath::Clamp<float>(Value, 0, LandscapeDataAccess::MaxValue);
 	}
@@ -2750,14 +2750,14 @@ void ShrinkData(TArray<T>& Data, int32 OldMinX, int32 OldMinY, int32 OldMaxX, in
 		{
 			for (int32 DestY = 0, SrcY = NewMinY - OldMinY; DestY <= NewMaxY - NewMinY; DestY++, SrcY++)
 			{
-				UE_LOG(LogLandscape, Warning, TEXT("Dest: %d, %d = %d Src: %d, %d = %d Width = %d"), 0, DestY, DestY * (1+NewMaxX-NewMinX), NewMinX-OldMinX, SrcY, SrcY * (1+OldMaxX-OldMinX) + NewMinX-OldMinX, (1+NewMaxX-NewMinX));
-				T* DestData = &Data[DestY * (1+NewMaxX-NewMinX)];
-				const T* SrcData = &Data[SrcY * (1+OldMaxX-OldMinX) + NewMinX-OldMinX];
-				FMemory::Memmove(DestData, SrcData, (1+NewMaxX-NewMinX) * sizeof(T));
+				UE_LOG(LogLandscape, Warning, TEXT("Dest: %d, %d = %d Src: %d, %d = %d Width = %d"), 0, DestY, DestY * (1 + NewMaxX - NewMinX), NewMinX - OldMinX, SrcY, SrcY * (1 + OldMaxX - OldMinX) + NewMinX - OldMinX, (1 + NewMaxX - NewMinX));
+				T* DestData = &Data[DestY * (1 + NewMaxX - NewMinX)];
+				const T* SrcData = &Data[SrcY * (1 + OldMaxX - OldMinX) + NewMinX - OldMinX];
+				FMemory::Memmove(DestData, SrcData, (1 + NewMaxX - NewMinX) * sizeof(T));
 			}
 		}
 
-		const int32 NewSize = (1+NewMaxY-NewMinY) * (1+NewMaxX-NewMinX);
+		const int32 NewSize = (1 + NewMaxY - NewMinY) * (1 + NewMaxX - NewMinX);
 		Data.RemoveAt(NewSize, Data.Num() - NewSize);
 	}
 }
@@ -2786,7 +2786,7 @@ bool ULandscapeInfo::ApplySplinesInternal(bool bOnlySelected, ALandscapeProxy* L
 		return false;
 	}
 
-	FScopedTransaction Transaction( LOCTEXT("LandscapeSpline_ApplySplines", "Apply Splines to Landscape") );
+	FScopedTransaction Transaction(LOCTEXT("LandscapeSpline_ApplySplines", "Apply Splines to Landscape"));
 
 	const FTransform SplineToLandscape = Landscape->SplineComponent->ComponentToWorld.GetRelativeTransform(Landscape->LandscapeActorToWorld());
 
@@ -2835,7 +2835,7 @@ bool ULandscapeInfo::ApplySplinesInternal(bool bOnlySelected, ALandscapeProxy* L
 		if (ControlPoint->bRaiseTerrain || ControlPoint->bLowerTerrain)
 		{
 			TArray<uint16> Data;
-			Data.AddZeroed( (1+MaxY-MinY) * (1+MaxX-MinX) );
+			Data.AddZeroed((1 + MaxY - MinY) * (1 + MaxX - MinX));
 
 			int32 ValidMinX = MinX;
 			int32 ValidMinY = MinY;
@@ -2852,13 +2852,13 @@ bool ULandscapeInfo::ApplySplinesInternal(bool bOnlySelected, ALandscapeProxy* L
 
 			ShrinkData(Data, MinX, MinY, MaxX, MaxY, ValidMinX, ValidMinY, ValidMaxX, ValidMaxY);
 
-					MinX = ValidMinX;
-					MinY = ValidMinY;
-					MaxX = ValidMaxX;
-				MaxY = ValidMaxY;
+			MinX = ValidMinX;
+			MinY = ValidMinY;
+			MaxX = ValidMaxX;
+			MaxY = ValidMaxY;
 
 			FTriangleRasterizer<FLandscapeSplineHeightsRasterPolicy> Rasterizer(
-				FLandscapeSplineHeightsRasterPolicy(Data, MinX, MinY, MaxX, MaxY, ControlPoint->bRaiseTerrain, ControlPoint->bLowerTerrain) );
+				FLandscapeSplineHeightsRasterPolicy(Data, MinX, MinY, MaxX, MaxY, ControlPoint->bRaiseTerrain, ControlPoint->bLowerTerrain));
 
 			const FVector Center3D = SplineToLandscape.TransformPosition(ControlPoint->Location);
 			const FVector2D CenterPos = FVector2D(Center3D);
@@ -2868,10 +2868,10 @@ bool ULandscapeInfo::ApplySplinesInternal(bool bOnlySelected, ALandscapeProxy* L
 			{
 				// Solid center
 				const FVector2D Right0Pos = FVector2D(Points[i].Right);
-				const FVector2D Left1Pos  = FVector2D(Points[j].Left);
+				const FVector2D Left1Pos = FVector2D(Points[j].Left);
 				const FVector2D Right1Pos = FVector2D(Points[j].Right);
 				const FVector Right0 = FVector(1.0f, Points[i].StartEndFalloff, Points[i].Right.Z);
-				const FVector Left1  = FVector(1.0f, Points[j].StartEndFalloff, Points[j].Left.Z);
+				const FVector Left1 = FVector(1.0f, Points[j].StartEndFalloff, Points[j].Left.Z);
 				const FVector Right1 = FVector(1.0f, Points[j].StartEndFalloff, Points[j].Right.Z);
 
 				Rasterizer.DrawTriangle(Center, Right0, Left1, CenterPos, Right0Pos, Left1Pos, false);
@@ -2879,9 +2879,9 @@ bool ULandscapeInfo::ApplySplinesInternal(bool bOnlySelected, ALandscapeProxy* L
 
 				// Falloff
 				FVector2D FalloffRight0Pos = FVector2D(Points[i].FalloffRight);
-				FVector2D FalloffLeft1Pos  = FVector2D(Points[j].FalloffLeft);
+				FVector2D FalloffLeft1Pos = FVector2D(Points[j].FalloffLeft);
 				FVector FalloffRight0 = FVector(0.0f, Points[i].StartEndFalloff, Points[i].FalloffRight.Z);
-				FVector FalloffLeft1  = FVector(0.0f, Points[j].StartEndFalloff, Points[j].FalloffLeft.Z);
+				FVector FalloffLeft1 = FVector(0.0f, Points[j].StartEndFalloff, Points[j].FalloffLeft.Z);
 				Rasterizer.DrawTriangle(Right0, FalloffRight0, Left1, Right0Pos, FalloffRight0Pos, Left1Pos, false);
 				Rasterizer.DrawTriangle(FalloffRight0, Left1, FalloffLeft1, FalloffRight0Pos, Left1Pos, FalloffLeft1Pos, false);
 			}
@@ -2895,7 +2895,7 @@ bool ULandscapeInfo::ApplySplinesInternal(bool bOnlySelected, ALandscapeProxy* L
 		if (ControlPoint->LayerName != NAME_None && LayerInfo != NULL)
 		{
 			TArray<uint8> Data;
-			Data.AddZeroed( (1+MaxY-MinY) * (1+MaxX-MinX) );
+			Data.AddZeroed((1 + MaxY - MinY) * (1 + MaxX - MinX));
 
 			int32 ValidMinX = MinX;
 			int32 ValidMinY = MinY;
@@ -2917,7 +2917,7 @@ bool ULandscapeInfo::ApplySplinesInternal(bool bOnlySelected, ALandscapeProxy* L
 			MaxY = ValidMaxY;
 
 			FTriangleRasterizer<FLandscapeSplineBlendmaskRasterPolicy> Rasterizer(
-				FLandscapeSplineBlendmaskRasterPolicy( Data, MinX, MinY, MaxX, MaxY) );
+				FLandscapeSplineBlendmaskRasterPolicy(Data, MinX, MinY, MaxX, MaxY));
 
 			const float BlendValue = 255;
 
@@ -2929,10 +2929,10 @@ bool ULandscapeInfo::ApplySplinesInternal(bool bOnlySelected, ALandscapeProxy* L
 			{
 				// Solid center
 				const FVector2D Right0Pos = FVector2D(Points[i].Right);
-				const FVector2D Left1Pos  = FVector2D(Points[j].Left);
+				const FVector2D Left1Pos = FVector2D(Points[j].Left);
 				const FVector2D Right1Pos = FVector2D(Points[j].Right);
 				const FVector Right0 = FVector(1.0f, Points[i].StartEndFalloff, BlendValue);
-				const FVector Left1  = FVector(1.0f, Points[j].StartEndFalloff, BlendValue);
+				const FVector Left1 = FVector(1.0f, Points[j].StartEndFalloff, BlendValue);
 				const FVector Right1 = FVector(1.0f, Points[j].StartEndFalloff, BlendValue);
 
 				Rasterizer.DrawTriangle(Center, Right0, Left1, CenterPos, Right0Pos, Left1Pos, false);
@@ -2940,9 +2940,9 @@ bool ULandscapeInfo::ApplySplinesInternal(bool bOnlySelected, ALandscapeProxy* L
 
 				// Falloff
 				FVector2D FalloffRight0Pos = FVector2D(Points[i].FalloffRight);
-				FVector2D FalloffLeft1Pos  = FVector2D(Points[j].FalloffLeft);
+				FVector2D FalloffLeft1Pos = FVector2D(Points[j].FalloffLeft);
 				FVector FalloffRight0 = FVector(0.0f, Points[i].StartEndFalloff, BlendValue);
-				FVector FalloffLeft1  = FVector(0.0f, Points[j].StartEndFalloff, BlendValue);
+				FVector FalloffLeft1 = FVector(0.0f, Points[j].StartEndFalloff, BlendValue);
 				Rasterizer.DrawTriangle(Right0, FalloffRight0, Left1, Right0Pos, FalloffRight0Pos, Left1Pos, false);
 				Rasterizer.DrawTriangle(FalloffRight0, Left1, FalloffLeft1, FalloffRight0Pos, Left1Pos, FalloffLeft1Pos, false);
 			}
@@ -2989,7 +2989,7 @@ bool ULandscapeInfo::ApplySplinesInternal(bool bOnlySelected, ALandscapeProxy* L
 		if (Segment->bRaiseTerrain || Segment->bLowerTerrain)
 		{
 			TArray<uint16> Data;
-			Data.AddZeroed( (1+MaxY-MinY) * (1+MaxX-MinX) );
+			Data.AddZeroed((1 + MaxY - MinY) * (1 + MaxX - MinX));
 
 			int32 ValidMinX = MinX;
 			int32 ValidMinY = MinY;
@@ -3006,40 +3006,40 @@ bool ULandscapeInfo::ApplySplinesInternal(bool bOnlySelected, ALandscapeProxy* L
 
 			ShrinkData(Data, MinX, MinY, MaxX, MaxY, ValidMinX, ValidMinY, ValidMaxX, ValidMaxY);
 
-					MinX = ValidMinX;
-					MinY = ValidMinY;
-					MaxX = ValidMaxX;
-				MaxY = ValidMaxY;
+			MinX = ValidMinX;
+			MinY = ValidMinY;
+			MaxX = ValidMaxX;
+			MaxY = ValidMaxY;
 
 			FTriangleRasterizer<FLandscapeSplineHeightsRasterPolicy> Rasterizer(
-				FLandscapeSplineHeightsRasterPolicy(Data, MinX, MinY, MaxX, MaxY, Segment->bRaiseTerrain, Segment->bLowerTerrain) );
+				FLandscapeSplineHeightsRasterPolicy(Data, MinX, MinY, MaxX, MaxY, Segment->bRaiseTerrain, Segment->bLowerTerrain));
 
 			for (int32 j = 1; j < Points.Num(); j++)
 			{
 				// Middle
-				FVector2D Left0Pos  = FVector2D(Points[j-1].Left);
-				FVector2D Right0Pos = FVector2D(Points[j-1].Right);
-				FVector2D Left1Pos  = FVector2D(Points[j].Left);
+				FVector2D Left0Pos = FVector2D(Points[j - 1].Left);
+				FVector2D Right0Pos = FVector2D(Points[j - 1].Right);
+				FVector2D Left1Pos = FVector2D(Points[j].Left);
 				FVector2D Right1Pos = FVector2D(Points[j].Right);
-				FVector Left0  = FVector(1.0f, Points[j-1].StartEndFalloff, Points[j-1].Left.Z);
-				FVector Right0 = FVector(1.0f, Points[j-1].StartEndFalloff, Points[j-1].Right.Z);
-				FVector Left1  = FVector(1.0f, Points[j].StartEndFalloff, Points[j].Left.Z);
+				FVector Left0 = FVector(1.0f, Points[j - 1].StartEndFalloff, Points[j - 1].Left.Z);
+				FVector Right0 = FVector(1.0f, Points[j - 1].StartEndFalloff, Points[j - 1].Right.Z);
+				FVector Left1 = FVector(1.0f, Points[j].StartEndFalloff, Points[j].Left.Z);
 				FVector Right1 = FVector(1.0f, Points[j].StartEndFalloff, Points[j].Right.Z);
 				Rasterizer.DrawTriangle(Left0, Right0, Left1, Left0Pos, Right0Pos, Left1Pos, false);
 				Rasterizer.DrawTriangle(Right0, Left1, Right1, Right0Pos, Left1Pos, Right1Pos, false);
 
 				// Left Falloff
-				FVector2D FalloffLeft0Pos  = FVector2D(Points[j-1].FalloffLeft);
-				FVector2D FalloffLeft1Pos  = FVector2D(Points[j].FalloffLeft);
-				FVector FalloffLeft0  = FVector(0.0f, Points[j-1].StartEndFalloff, Points[j-1].FalloffLeft.Z);
-				FVector FalloffLeft1  = FVector(0.0f, Points[j].StartEndFalloff, Points[j].FalloffLeft.Z);
+				FVector2D FalloffLeft0Pos = FVector2D(Points[j - 1].FalloffLeft);
+				FVector2D FalloffLeft1Pos = FVector2D(Points[j].FalloffLeft);
+				FVector FalloffLeft0 = FVector(0.0f, Points[j - 1].StartEndFalloff, Points[j - 1].FalloffLeft.Z);
+				FVector FalloffLeft1 = FVector(0.0f, Points[j].StartEndFalloff, Points[j].FalloffLeft.Z);
 				Rasterizer.DrawTriangle(FalloffLeft0, Left0, FalloffLeft1, FalloffLeft0Pos, Left0Pos, FalloffLeft1Pos, false);
 				Rasterizer.DrawTriangle(Left0, FalloffLeft1, Left1, Left0Pos, FalloffLeft1Pos, Left1Pos, false);
 
 				// Right Falloff
-				FVector2D FalloffRight0Pos = FVector2D(Points[j-1].FalloffRight);
+				FVector2D FalloffRight0Pos = FVector2D(Points[j - 1].FalloffRight);
 				FVector2D FalloffRight1Pos = FVector2D(Points[j].FalloffRight);
-				FVector FalloffRight0 = FVector(0.0f, Points[j-1].StartEndFalloff, Points[j-1].FalloffRight.Z);
+				FVector FalloffRight0 = FVector(0.0f, Points[j - 1].StartEndFalloff, Points[j - 1].FalloffRight.Z);
 				FVector FalloffRight1 = FVector(0.0f, Points[j].StartEndFalloff, Points[j].FalloffRight.Z);
 				Rasterizer.DrawTriangle(Right0, FalloffRight0, Right1, Right0Pos, FalloffRight0Pos, Right1Pos, false);
 				Rasterizer.DrawTriangle(FalloffRight0, Right1, FalloffRight1, FalloffRight0Pos, Right1Pos, FalloffRight1Pos, false);
@@ -3054,7 +3054,7 @@ bool ULandscapeInfo::ApplySplinesInternal(bool bOnlySelected, ALandscapeProxy* L
 		if (Segment->LayerName != NAME_None && LayerInfo != NULL)
 		{
 			TArray<uint8> Data;
-			Data.AddZeroed( (1+MaxY-MinY) * (1+MaxX-MinX) );
+			Data.AddZeroed((1 + MaxY - MinY) * (1 + MaxX - MinX));
 
 			int32 ValidMinX = MinX;
 			int32 ValidMinY = MinY;
@@ -3076,36 +3076,36 @@ bool ULandscapeInfo::ApplySplinesInternal(bool bOnlySelected, ALandscapeProxy* L
 			MaxY = ValidMaxY;
 
 			FTriangleRasterizer<FLandscapeSplineBlendmaskRasterPolicy> Rasterizer(
-				FLandscapeSplineBlendmaskRasterPolicy( Data, MinX, MinY, MaxX, MaxY) );
+				FLandscapeSplineBlendmaskRasterPolicy(Data, MinX, MinY, MaxX, MaxY));
 
 			const float BlendValue = 255;
 
 			for (int32 j = 1; j < Points.Num(); j++)
 			{
 				// Middle
-				FVector2D Left0Pos  = FVector2D(Points[j-1].Left);
-				FVector2D Right0Pos = FVector2D(Points[j-1].Right);
-				FVector2D Left1Pos  = FVector2D(Points[j].Left);
+				FVector2D Left0Pos = FVector2D(Points[j - 1].Left);
+				FVector2D Right0Pos = FVector2D(Points[j - 1].Right);
+				FVector2D Left1Pos = FVector2D(Points[j].Left);
 				FVector2D Right1Pos = FVector2D(Points[j].Right);
-				FVector Left0  = FVector(1.0f, Points[j-1].StartEndFalloff, BlendValue);
-				FVector Right0 = FVector(1.0f, Points[j-1].StartEndFalloff, BlendValue);
-				FVector Left1  = FVector(1.0f, Points[j].StartEndFalloff, BlendValue);
+				FVector Left0 = FVector(1.0f, Points[j - 1].StartEndFalloff, BlendValue);
+				FVector Right0 = FVector(1.0f, Points[j - 1].StartEndFalloff, BlendValue);
+				FVector Left1 = FVector(1.0f, Points[j].StartEndFalloff, BlendValue);
 				FVector Right1 = FVector(1.0f, Points[j].StartEndFalloff, BlendValue);
 				Rasterizer.DrawTriangle(Left0, Right0, Left1, Left0Pos, Right0Pos, Left1Pos, false);
 				Rasterizer.DrawTriangle(Right0, Left1, Right1, Right0Pos, Left1Pos, Right1Pos, false);
 
 				// Left Falloff
-				FVector2D FalloffLeft0Pos  = FVector2D(Points[j-1].FalloffLeft);
-				FVector2D FalloffLeft1Pos  = FVector2D(Points[j].FalloffLeft);
-				FVector FalloffLeft0  = FVector(0.0f, Points[j-1].StartEndFalloff, BlendValue);
-				FVector FalloffLeft1  = FVector(0.0f, Points[j].StartEndFalloff, BlendValue);
+				FVector2D FalloffLeft0Pos = FVector2D(Points[j - 1].FalloffLeft);
+				FVector2D FalloffLeft1Pos = FVector2D(Points[j].FalloffLeft);
+				FVector FalloffLeft0 = FVector(0.0f, Points[j - 1].StartEndFalloff, BlendValue);
+				FVector FalloffLeft1 = FVector(0.0f, Points[j].StartEndFalloff, BlendValue);
 				Rasterizer.DrawTriangle(FalloffLeft0, Left0, FalloffLeft1, FalloffLeft0Pos, Left0Pos, FalloffLeft1Pos, false);
 				Rasterizer.DrawTriangle(Left0, FalloffLeft1, Left1, Left0Pos, FalloffLeft1Pos, Left1Pos, false);
 
 				// Right Falloff
-				FVector2D FalloffRight0Pos = FVector2D(Points[j-1].FalloffRight);
+				FVector2D FalloffRight0Pos = FVector2D(Points[j - 1].FalloffRight);
 				FVector2D FalloffRight1Pos = FVector2D(Points[j].FalloffRight);
-				FVector FalloffRight0 = FVector(0.0f, Points[j-1].StartEndFalloff, BlendValue);
+				FVector FalloffRight0 = FVector(0.0f, Points[j - 1].StartEndFalloff, BlendValue);
 				FVector FalloffRight1 = FVector(0.0f, Points[j].StartEndFalloff, BlendValue);
 				Rasterizer.DrawTriangle(Right0, FalloffRight0, Right1, Right0Pos, FalloffRight0Pos, Right1Pos, false);
 				Rasterizer.DrawTriangle(FalloffRight0, Right1, FalloffRight1, FalloffRight0Pos, Right1Pos, FalloffRight1Pos, false);
@@ -3137,9 +3137,9 @@ bool ULandscapeInfo::ApplySplinesInternal(bool bOnlySelected, ALandscapeProxy* L
 }
 
 
-void ULandscapeInfo::DeleteLayer( ULandscapeLayerInfoObject* LayerInfo )
+void ULandscapeInfo::DeleteLayer(ULandscapeLayerInfoObject* LayerInfo)
 {
-	GWarn->BeginSlowTask( LOCTEXT("BeginDeletingLayerTask", "Deleting Layer"), true);
+	GWarn->BeginSlowTask(LOCTEXT("BeginDeletingLayerTask", "Deleting Layer"), true);
 
 	// Remove data from all components
 	FLandscapeEditDataInterface LandscapeEdit(this);
@@ -3178,7 +3178,7 @@ void ULandscapeInfo::ReplaceLayer(ULandscapeLayerInfoObject* FromLayerInfo, ULan
 {
 	if (ensure(FromLayerInfo != ToLayerInfo))
 	{
-		GWarn->BeginSlowTask( LOCTEXT("BeginReplacingLayerTask", "Replacing Layer"), true);
+		GWarn->BeginSlowTask(LOCTEXT("BeginReplacingLayerTask", "Replacing Layer"), true);
 
 		// Remove data from all components
 		FLandscapeEditDataInterface LandscapeEdit(this);
@@ -3189,7 +3189,7 @@ void ULandscapeInfo::ReplaceLayer(ULandscapeLayerInfoObject* FromLayerInfo, ULan
 		{
 			if (Layers[j].LayerInfoObj && Layers[j].LayerInfoObj == FromLayerInfo)
 			{
-				 Layers[j].LayerInfoObj = ToLayerInfo;
+				Layers[j].LayerInfoObj = ToLayerInfo;
 			}
 		}
 
@@ -3260,15 +3260,15 @@ void ALandscapeProxy::EditorApplyScale(const FVector& DeltaScale, const FVector*
 
 	// Correct for attempts to scale to 0 on any axis
 	FVector CurrentScale = GetRootComponent()->RelativeScale3D;
-	if( AActor::bUsePercentageBasedScaling )
+	if (AActor::bUsePercentageBasedScaling)
 	{
 		if (ModifiedScale.X == -1)
 		{
-			ModifiedScale.X = ModifiedScale.Y = -(CurrentScale.X-1) / CurrentScale.X;
+			ModifiedScale.X = ModifiedScale.Y = -(CurrentScale.X - 1) / CurrentScale.X;
 		}
 		if (ModifiedScale.Z == -1)
 		{
-			ModifiedScale.Z = -(CurrentScale.Z-1) / CurrentScale.Z;
+			ModifiedScale.Z = -(CurrentScale.Z - 1) / CurrentScale.Z;
 		}
 	}
 	else
@@ -3310,7 +3310,7 @@ void ALandscapeProxy::PostEditMove(bool bFinished)
 void ALandscapeProxy::PostEditImport()
 {
 	Super::PostEditImport();
-	if (!bIsProxy && GetWorld() ) // For Landscape
+	if (!bIsProxy && GetWorld()) // For Landscape
 	{
 		for (FActorIterator It(GetWorld()); It; ++It)
 		{
@@ -3422,7 +3422,7 @@ void ALandscapeProxy::RecreateCollisionComponents()
 		RemoveOwnedComponent(CollisionComps[i]);
 	}
 
-	CollisionMipLevel = FMath::Clamp<int32>( CollisionMipLevel, 0, FMath::CeilLogTwo(SubsectionSizeQuads+1)-1 );
+	CollisionMipLevel = FMath::Clamp<int32>(CollisionMipLevel, 0, FMath::CeilLogTwo(SubsectionSizeQuads + 1) - 1);
 	for (int32 i = 0; i < LandscapeComponents.Num(); ++i)
 	{
 		ULandscapeComponent* Comp = LandscapeComponents[i];
@@ -3470,7 +3470,7 @@ namespace
 			// Adjust to 1/2^n
 			int32 i = 2;
 			int32 LightmapSize = (NumSubsections * (SubsectionSizeQuads + 1)) >> 1;
-			while ( StaticLightingResolution < (1.f/i) && LightmapSize > 4  )
+			while (StaticLightingResolution < (1.f / i) && LightmapSize > 4)
 			{
 				i <<= 1;
 				LightmapSize >>= 1;
@@ -3479,8 +3479,8 @@ namespace
 
 			int32 PixelPaddingX = GPixelFormats[PF_DXT1].BlockSizeX;
 
-			int32 DestSize = (int32)((2*PixelPaddingX + ComponentSizeQuads + 1) * StaticLightingResolution);
-			StaticLightingResolution = (float)DestSize / (2*PixelPaddingX + ComponentSizeQuads + 1 );
+			int32 DestSize = (int32)((2 * PixelPaddingX + ComponentSizeQuads + 1) * StaticLightingResolution);
+			StaticLightingResolution = (float)DestSize / (2 * PixelPaddingX + ComponentSizeQuads + 1);
 		}
 
 		return StaticLightingResolution;
@@ -3493,12 +3493,12 @@ void ALandscapeProxy::PostEditChangeProperty(FPropertyChangedEvent& PropertyChan
 	const FName PropertyName = PropertyChangedEvent.MemberProperty ? PropertyChangedEvent.MemberProperty->GetFName() : NAME_None;
 	if (bIsProxy)
 	{
-		if( PropertyName == FName(TEXT("LandscapeActor")) )
+		if (PropertyName == FName(TEXT("LandscapeActor")))
 		{
 			if (LandscapeActor && IsValidLandscapeActor(LandscapeActor.Get()))
 			{
 				// defer LandscapeInfo setup
-				if (GIsEditor && GetWorld() && !GetWorld()->IsPlayInEditor() )
+				if (GIsEditor && GetWorld() && !GetWorld()->IsPlayInEditor())
 				{
 					GEngine->DeferredCommands.AddUnique(TEXT("UpdateLandscapeEditorData"));
 				}
@@ -3508,13 +3508,13 @@ void ALandscapeProxy::PostEditChangeProperty(FPropertyChangedEvent& PropertyChan
 				LandscapeActor = 0;
 			}
 		}
-		else if( PropertyName == FName(TEXT("LandscapeMaterial")) || PropertyName == FName(TEXT("LandscapeHoleMaterial")) )
+		else if (PropertyName == FName(TEXT("LandscapeMaterial")) || PropertyName == FName(TEXT("LandscapeHoleMaterial")))
 		{
 			FMaterialUpdateContext MaterialUpdateContext;
 			GetLandscapeInfo()->UpdateLayerInfoMap(/*this*/);
 
 			// Clear the parents out of combination material instances
-			for( TMap<FString ,UMaterialInstanceConstant*>::TIterator It(MaterialInstanceConstantMap); It; ++It )
+			for (TMap<FString, UMaterialInstanceConstant*>::TIterator It(MaterialInstanceConstantMap); It; ++It)
 			{
 				It.Value()->SetParentEditorOnly(NULL);
 				MaterialUpdateContext.AddMaterial(It.Value()->GetMaterial());
@@ -3523,10 +3523,10 @@ void ALandscapeProxy::PostEditChangeProperty(FPropertyChangedEvent& PropertyChan
 			// Remove our references to any material instances
 			MaterialInstanceConstantMap.Empty();
 
-			for(int32 ComponentIndex = 0; ComponentIndex < LandscapeComponents.Num(); ComponentIndex++ )
+			for (int32 ComponentIndex = 0; ComponentIndex < LandscapeComponents.Num(); ComponentIndex++)
 			{
 				ULandscapeComponent* Comp = LandscapeComponents[ComponentIndex];
-				if( Comp )
+				if (Comp)
 				{
 					// Update the MIC
 					Comp->UpdateMaterialInstances();
@@ -3535,26 +3535,26 @@ void ALandscapeProxy::PostEditChangeProperty(FPropertyChangedEvent& PropertyChan
 		}
 	}
 
-	if ( GIsEditor && PropertyName == FName(TEXT("StreamingDistanceMultiplier")) )
+	if (GIsEditor && PropertyName == FName(TEXT("StreamingDistanceMultiplier")))
 	{
 		// Recalculate in a few seconds.
 		GetWorld()->TriggerStreamingDataRebuild();
 	}
-	else if ( GIsEditor && PropertyName == FName(TEXT("DefaultPhysMaterial")) )
+	else if (GIsEditor && PropertyName == FName(TEXT("DefaultPhysMaterial")))
 	{
 		ChangedPhysMaterial();
 	}
-	else if ( GIsEditor && (PropertyName == FName(TEXT("CollisionMipLevel"))) )
+	else if (GIsEditor && (PropertyName == FName(TEXT("CollisionMipLevel"))))
 	{
 		RecreateCollisionComponents();
 	}
 	else if (PropertyName == FName(TEXT("bCastStaticShadow")) || PropertyName == FName(TEXT("bCastShadowAsTwoSided")))
 	{
 		// Replicate shared properties to all components.
-		for(int32 ComponentIndex = 0; ComponentIndex < LandscapeComponents.Num(); ComponentIndex++ )
+		for (int32 ComponentIndex = 0; ComponentIndex < LandscapeComponents.Num(); ComponentIndex++)
 		{
 			ULandscapeComponent* Comp = LandscapeComponents[ComponentIndex];
-			if( Comp )
+			if (Comp)
 			{
 				Comp->bCastStaticShadow = bCastStaticShadow;
 				Comp->bCastShadowAsTwoSided = bCastShadowAsTwoSided;
@@ -3651,7 +3651,7 @@ void ALandscape::PostEditChangeProperty(FPropertyChangedEvent& PropertyChangedEv
 
 	ULandscapeInfo* Info = GetLandscapeInfo();
 
-	if( PropertyName == FName(TEXT("LandscapeMaterial")) || PropertyName == FName(TEXT("LandscapeHoleMaterial")) )
+	if (PropertyName == FName(TEXT("LandscapeMaterial")) || PropertyName == FName(TEXT("LandscapeHoleMaterial")))
 	{
 		FMaterialUpdateContext MaterialUpdateContext;
 		GetLandscapeInfo()->UpdateLayerInfoMap(/*this*/);
@@ -3659,57 +3659,57 @@ void ALandscape::PostEditChangeProperty(FPropertyChangedEvent& PropertyChangedEv
 		ChangedMaterial = true;
 
 		// Clear the parents out of combination material instances
-		for( TMap<FString ,UMaterialInstanceConstant*>::TIterator It(MaterialInstanceConstantMap); It; ++It )
+		for (TMap<FString, UMaterialInstanceConstant*>::TIterator It(MaterialInstanceConstantMap); It; ++It)
 		{
 			It.Value()->SetParentEditorOnly(NULL);
 			MaterialUpdateContext.AddMaterial(It.Value()->GetMaterial());
 		}
-		
+
 		// Remove our references to any material instances
 		MaterialInstanceConstantMap.Empty();
 	}
-	else if( PropertyName == FName(TEXT("RelativeScale3D")) ||
-			 PropertyName == FName(TEXT("RelativeLocation")) ||
-			 PropertyName == FName(TEXT("RelativeRotation")) )
+	else if (PropertyName == FName(TEXT("RelativeScale3D")) ||
+		PropertyName == FName(TEXT("RelativeLocation")) ||
+		PropertyName == FName(TEXT("RelativeRotation")))
 	{
 		// update transformations for all linked proxies 
 		Info->FixupProxiesTransform();
 		bNeedsRecalcBoundingBox = true;
 	}
-	else if ( GIsEditor && PropertyName == FName(TEXT("MaxLODLevel")) )
+	else if (GIsEditor && PropertyName == FName(TEXT("MaxLODLevel")))
 	{
-		MaxLODLevel = FMath::Min<int32>(MaxLODLevel, FMath::CeilLogTwo(SubsectionSizeQuads+1)-1);
+		MaxLODLevel = FMath::Min<int32>(MaxLODLevel, FMath::CeilLogTwo(SubsectionSizeQuads + 1) - 1);
 		bPropagateToProxies = true;
 	}
-	else if ( PropertyName == FName(TEXT("LODDistanceFactor")) )
+	else if (PropertyName == FName(TEXT("LODDistanceFactor")))
 	{
 		LODDistanceFactor = FMath::Clamp<float>(LODDistanceFactor, 0.1f, MAX_LANDSCAPE_LOD_DISTANCE_FACTOR); // limit because LOD transition became too popping...
 		bPropagateToProxies = true;
 	}
-	else if ( PropertyName == FName(TEXT("CollisionMipLevel")) )
+	else if (PropertyName == FName(TEXT("CollisionMipLevel")))
 	{
-		CollisionMipLevel = FMath::Clamp<int32>( CollisionMipLevel, 0, FMath::CeilLogTwo(SubsectionSizeQuads+1)-1 );
+		CollisionMipLevel = FMath::Clamp<int32>(CollisionMipLevel, 0, FMath::CeilLogTwo(SubsectionSizeQuads + 1) - 1);
 		bPropagateToProxies = true;
 	}
 	else if (PropertyName == FName(TEXT("LODFalloff")))
 	{
 		bPropagateToProxies = true;
 	}
-	else if ( GIsEditor && PropertyName == FName(TEXT("StaticLightingResolution")) )
+	else if (GIsEditor && PropertyName == FName(TEXT("StaticLightingResolution")))
 	{
 		StaticLightingResolution = ::AdjustStaticLightingResolution(StaticLightingResolution, NumSubsections, SubsectionSizeQuads, ComponentSizeQuads);
 		bChangedLighting = true;
 	}
 	else if (GIsEditor && PropertyName == GET_MEMBER_NAME_CHECKED(ALandscapeProxy, StaticLightingLOD))
 	{
-		StaticLightingLOD = FMath::Clamp<int32>(StaticLightingLOD, 0, FMath::CeilLogTwo(SubsectionSizeQuads+1)-1);
+		StaticLightingLOD = FMath::Clamp<int32>(StaticLightingLOD, 0, FMath::CeilLogTwo(SubsectionSizeQuads + 1) - 1);
 		bChangedLighting = true;
 	}
 	else if (GIsEditor && PropertyName == GET_MEMBER_NAME_CHECKED(ALandscapeProxy, ExportLOD))
 	{
-		ExportLOD = FMath::Clamp<int32>(ExportLOD, 0, FMath::CeilLogTwo(SubsectionSizeQuads+1)-1);
+		ExportLOD = FMath::Clamp<int32>(ExportLOD, 0, FMath::CeilLogTwo(SubsectionSizeQuads + 1) - 1);
 	}
-	else if ( GIsEditor && PropertyName == NAME_UsedForNavigation)
+	else if (GIsEditor && PropertyName == NAME_UsedForNavigation)
 	{
 		UpdateNavigationRelevancy();
 	}
@@ -3721,7 +3721,7 @@ void ALandscape::PostEditChangeProperty(FPropertyChangedEvent& PropertyChangedEv
 		if (bPropagateToProxies)
 		{
 			// Propagate Event to Proxies...
-			for (TSet<ALandscapeProxy*>::TIterator It(Info->Proxies); It; ++It )
+			for (TSet<ALandscapeProxy*>::TIterator It(Info->Proxies); It; ++It)
 			{
 				(*It)->GetSharedProperties(this);
 				(*It)->PostEditChangeProperty(PropertyChangedEvent);
@@ -3729,7 +3729,7 @@ void ALandscape::PostEditChangeProperty(FPropertyChangedEvent& PropertyChangedEv
 		}
 
 		// Update normals if DrawScale3D is changed
-		if( PropertyName == FName(TEXT("RelativeScale3D")) )
+		if (PropertyName == FName(TEXT("RelativeScale3D")))
 		{
 			FLandscapeEditDataInterface LandscapeEdit(Info);
 			LandscapeEdit.RecalculateNormals();
@@ -3742,15 +3742,15 @@ void ALandscape::PostEditChangeProperty(FPropertyChangedEvent& PropertyChangedEv
 		for (auto It = AllComponents.CreateIterator(); It; ++It)
 		{
 			ULandscapeComponent* Comp = *It;
-			if( Comp )
+			if (Comp)
 			{
-				if( bNeedsRecalcBoundingBox )
+				if (bNeedsRecalcBoundingBox)
 				{
 					Comp->UpdateCachedBounds();
 					Comp->UpdateBounds();
 				}
 
-				if( ChangedMaterial )
+				if (ChangedMaterial)
 				{
 					// Update the MIC
 					Comp->UpdateMaterialInstances();
@@ -3769,7 +3769,7 @@ void ALandscape::PostEditChangeProperty(FPropertyChangedEvent& PropertyChangedEv
 		// Need to update Gizmo scene proxy
 		if (bNeedsRecalcBoundingBox && GetWorld())
 		{
-			for ( FActorIterator It(GetWorld()); It; ++It )
+			for (FActorIterator It(GetWorld()); It; ++It)
 			{
 				ALandscapeGizmoActiveActor* Gizmo = Cast<ALandscapeGizmoActiveActor>(*It);
 				if (Gizmo)
@@ -3778,7 +3778,7 @@ void ALandscape::PostEditChangeProperty(FPropertyChangedEvent& PropertyChangedEv
 				}
 			}
 		}
-				
+
 		if (ChangedMaterial)
 		{
 			if (GIsEditor && GetWorld() && !GetWorld()->IsPlayInEditor())
@@ -3787,7 +3787,7 @@ void ALandscape::PostEditChangeProperty(FPropertyChangedEvent& PropertyChangedEv
 			}
 
 			// Update all the proxies...
-			for ( TSet<ALandscapeProxy*>::TIterator It(Info->Proxies); It; ++It )
+			for (TSet<ALandscapeProxy*>::TIterator It(Info->Proxies); It; ++It)
 			{
 				(*It)->MarkComponentsRenderStateDirty();
 			}
@@ -3799,13 +3799,13 @@ void ALandscapeProxy::ChangedPhysMaterial()
 {
 	ULandscapeInfo* LandscapeInfo = GetLandscapeInfo();
 	if (!LandscapeInfo) return;
-	for(auto It = LandscapeInfo->XYtoComponentMap.CreateIterator(); It; ++It )
+	for (auto It = LandscapeInfo->XYtoComponentMap.CreateIterator(); It; ++It)
 	{
 		ULandscapeComponent* Comp = It.Value();
-		if( Comp )
+		if (Comp)
 		{
 			ULandscapeHeightfieldCollisionComponent* CollisionComponent = Comp->CollisionComponent.Get();
-			if( CollisionComponent )
+			if (CollisionComponent)
 			{
 				Comp->UpdateCollisionLayerData();
 			}
@@ -3820,7 +3820,7 @@ void ULandscapeComponent::SetLOD(bool bForcedLODChanged, int32 InLODValue)
 		ForcedLOD = InLODValue;
 		if (ForcedLOD >= 0)
 		{
-			ForcedLOD = FMath::Clamp<int32>( ForcedLOD, 0, FMath::CeilLogTwo(SubsectionSizeQuads+1)-1 );
+			ForcedLOD = FMath::Clamp<int32>(ForcedLOD, 0, FMath::CeilLogTwo(SubsectionSizeQuads + 1) - 1);
 		}
 		else
 		{
@@ -3829,7 +3829,7 @@ void ULandscapeComponent::SetLOD(bool bForcedLODChanged, int32 InLODValue)
 	}
 	else
 	{
-		int32 MaxLOD = FMath::CeilLogTwo(SubsectionSizeQuads+1)-1;
+		int32 MaxLOD = FMath::CeilLogTwo(SubsectionSizeQuads + 1) - 1;
 		LODBias = FMath::Clamp<int32>(InLODValue, -MaxLOD, MaxLOD);
 	}
 
@@ -3839,17 +3839,17 @@ void ULandscapeComponent::SetLOD(bool bForcedLODChanged, int32 InLODValue)
 	ULandscapeInfo* Info = GetLandscapeInfo(false);
 	if (Info)
 	{
-		FIntPoint ComponentBase = GetSectionBase()/ComponentSizeQuads;
-		FIntPoint LandscapeKey[8] = 
-		{
-			ComponentBase + FIntPoint(-1,-1),
-			ComponentBase + FIntPoint(+0,-1),
-			ComponentBase + FIntPoint(+1,-1),
-			ComponentBase + FIntPoint(-1,+0),
-			ComponentBase + FIntPoint(+1,+0),
-			ComponentBase + FIntPoint(-1,+1),
-			ComponentBase + FIntPoint(+0,+1),
-			ComponentBase + FIntPoint(+1,+1)
+		FIntPoint ComponentBase = GetSectionBase() / ComponentSizeQuads;
+		FIntPoint LandscapeKey[8] =
+		{
+			ComponentBase + FIntPoint(-1, -1),
+			ComponentBase + FIntPoint(+0, -1),
+			ComponentBase + FIntPoint(+1, -1),
+			ComponentBase + FIntPoint(-1, +0),
+			ComponentBase + FIntPoint(+1, +0),
+			ComponentBase + FIntPoint(-1, +1),
+			ComponentBase + FIntPoint(+0, +1),
+			ComponentBase + FIntPoint(+1, +1)
 		};
 
 		for (int32 Idx = 0; Idx < 8; ++Idx)
@@ -3860,12 +3860,12 @@ void ULandscapeComponent::SetLOD(bool bForcedLODChanged, int32 InLODValue)
 				Comp->Modify();
 				if (bForcedLODChanged)
 				{
-					Comp->NeighborLOD[7-Idx] = ForcedLOD >= 0 ? ForcedLOD : 255; // Use 255 as unspecified value
+					Comp->NeighborLOD[7 - Idx] = ForcedLOD >= 0 ? ForcedLOD : 255; // Use 255 as unspecified value
 				}
 				else
 				{
 					// Neighbor LODBias are saved in BYTE, so need to convert to range [-128:127]
-					Comp->NeighborLODBias[7-Idx] = LODBias + 128;
+					Comp->NeighborLODBias[7 - Idx] = LODBias + 128;
 				}
 
 				Comp->InvalidateLightingCache();
@@ -3900,18 +3900,18 @@ void ULandscapeComponent::PostEditChangeProperty(FPropertyChangedEvent& Property
 {
 	Super::PostEditChangeProperty(PropertyChangedEvent);
 	const FName PropertyName = PropertyChangedEvent.Property ? PropertyChangedEvent.Property->GetFName() : NAME_None;
-	if( PropertyName == FName(TEXT("OverrideMaterial")) )
+	if (PropertyName == FName(TEXT("OverrideMaterial")))
 	{
 		UpdateMaterialInstances();
 		// Reregister all components
 		FComponentReregisterContext ReregisterContext(this);
 	}
-	else if ( GIsEditor && (PropertyName == FName(TEXT("ForcedLOD")) || PropertyName == FName(TEXT("LODBias"))) )
+	else if (GIsEditor && (PropertyName == FName(TEXT("ForcedLOD")) || PropertyName == FName(TEXT("LODBias"))))
 	{
 		bool bForcedLODChanged = PropertyName == FName(TEXT("ForcedLOD"));
 		SetLOD(bForcedLODChanged, bForcedLODChanged ? ForcedLOD : LODBias);
 	}
-	else if ( GIsEditor && PropertyName == FName(TEXT("StaticLightingResolution")) )
+	else if (GIsEditor && PropertyName == FName(TEXT("StaticLightingResolution")))
 	{
 		if (StaticLightingResolution > 0.f)
 		{
@@ -3922,9 +3922,9 @@ void ULandscapeComponent::PostEditChangeProperty(FPropertyChangedEvent& Property
 			StaticLightingResolution = 0;
 		}
 	}
-	else if ( GIsEditor && (PropertyName == FName(TEXT("CollisionMipLevel"))) )
+	else if (GIsEditor && (PropertyName == FName(TEXT("CollisionMipLevel"))))
 	{
-		CollisionMipLevel = FMath::Clamp<int32>( CollisionMipLevel, 0, FMath::CeilLogTwo(SubsectionSizeQuads+1)-1 );
+		CollisionMipLevel = FMath::Clamp<int32>(CollisionMipLevel, 0, FMath::CeilLogTwo(SubsectionSizeQuads + 1) - 1);
 		TArray<uint8> CollisionMipData;
 		HeightmapTexture->Source.GetMipData(CollisionMipData, CollisionMipLevel);
 		UpdateCollisionHeightData((FColor*)CollisionMipData.GetTypedData(), 0, 0, MAX_int32, MAX_int32, true, NULL, true); // Rebuild for new CollisionMipLevel
@@ -3947,7 +3947,7 @@ void ULandscapeInfo::UpdateSelectedComponents(TSet<ULandscapeComponent*>& NewCom
 
 	if (bIsComponentwise)
 	{
-		for( TSet<ULandscapeComponent*>::TIterator It(NewComponents); It; ++It )
+		for (TSet<ULandscapeComponent*>::TIterator It(NewComponents); It; ++It)
 		{
 			ULandscapeComponent* Comp = *It;
 			if (Comp->EditToolRenderData != NULL && (Comp->EditToolRenderData->SelectedType & InSelectType) == 0)
@@ -3961,10 +3961,10 @@ void ULandscapeInfo::UpdateSelectedComponents(TSet<ULandscapeComponent*>& NewCom
 
 		// Remove the material from any old components that are no longer in the region
 		TSet<ULandscapeComponent*> RemovedComponents = SelectedComponents.Difference(NewComponents);
-		for ( TSet<ULandscapeComponent*>::TIterator It(RemovedComponents); It; ++It )
+		for (TSet<ULandscapeComponent*>::TIterator It(RemovedComponents); It; ++It)
 		{
 			ULandscapeComponent* Comp = *It;
-			if( Comp->EditToolRenderData != NULL )
+			if (Comp->EditToolRenderData != NULL)
 			{
 				Comp->Modify();
 				int32 SelectedType = Comp->EditToolRenderData->SelectedType;
@@ -3996,10 +3996,10 @@ void ULandscapeInfo::UpdateSelectedComponents(TSet<ULandscapeComponent*>& NewCom
 		else
 		{
 			// Remove the material from any old components that are no longer in the region
-			for ( TSet<ULandscapeComponent*>::TIterator It(SelectedRegionComponents); It; ++It )
+			for (TSet<ULandscapeComponent*>::TIterator It(SelectedRegionComponents); It; ++It)
 			{
 				ULandscapeComponent* Comp = *It;
-				if( Comp->EditToolRenderData != NULL )
+				if (Comp->EditToolRenderData != NULL)
 				{
 					Comp->Modify();
 					int32 SelectedType = Comp->EditToolRenderData->SelectedType;
@@ -4016,12 +4016,12 @@ void ULandscapeInfo::SortSelectedComponents()
 {
 	struct FCompareULandscapeComponentBySectionBase
 	{
-		FORCEINLINE bool operator()( const ULandscapeComponent& A, const ULandscapeComponent & B ) const
+		FORCEINLINE bool operator()(const ULandscapeComponent& A, const ULandscapeComponent & B) const
 		{
 			return (A.GetSectionBase().X == B.GetSectionBase().X) ? (A.GetSectionBase().Y < B.GetSectionBase().Y) : (A.GetSectionBase().X < B.GetSectionBase().X);
 		}
 	};
-	SelectedComponents.Sort( FCompareULandscapeComponentBySectionBase() );
+	SelectedComponents.Sort(FCompareULandscapeComponentBySectionBase());
 }
 
 void ULandscapeInfo::ClearSelectedRegion(bool bIsComponentwise /*= true*/)
@@ -4036,8 +4036,8 @@ void ULandscapeInfo::ClearSelectedRegion(bool bIsComponentwise /*= true*/)
 
 struct FLandscapeDataInterface* ULandscapeInfo::GetDataInterface()
 {
-	if( DataInterface == NULL )
-	{ 
+	if (DataInterface == NULL)
+	{
 		DataInterface = new FLandscapeDataInterface();
 	}
 
@@ -4047,18 +4047,18 @@ struct FLandscapeDataInterface* ULandscapeInfo::GetDataInterface()
 void ULandscapeComponent::ReallocateWeightmaps(FLandscapeEditDataInterface* DataInterface)
 {
 	ALandscapeProxy* Proxy = GetLandscapeProxy();
-	
-	int32 NeededNewChannels=0;
-	for( int32 LayerIdx=0;LayerIdx < WeightmapLayerAllocations.Num();LayerIdx++ )
+
+	int32 NeededNewChannels = 0;
+	for (int32 LayerIdx = 0; LayerIdx < WeightmapLayerAllocations.Num(); LayerIdx++)
 	{
-		if( WeightmapLayerAllocations[LayerIdx].WeightmapTextureIndex == 255 )
+		if (WeightmapLayerAllocations[LayerIdx].WeightmapTextureIndex == 255)
 		{
 			NeededNewChannels++;
 		}
 	}
 
 	// All channels allocated!
-	if( NeededNewChannels == 0 )
+	if (NeededNewChannels == 0)
 	{
 		return;
 	}
@@ -4071,42 +4071,42 @@ void ULandscapeComponent::ReallocateWeightmaps(FLandscapeEditDataInterface* Data
 	// UE_LOG(LogLandscape, Log, TEXT("Component %s needs %d layers (%d new)"), *GetName(), WeightmapLayerAllocations.Num(), NeededNewChannels);
 
 	// See if our existing textures have sufficient space
-	int32 ExistingTexAvailableChannels=0;
-	for( int32 TexIdx=0;TexIdx<WeightmapTextures.Num();TexIdx++ )
+	int32 ExistingTexAvailableChannels = 0;
+	for (int32 TexIdx = 0; TexIdx < WeightmapTextures.Num(); TexIdx++)
 	{
 		FLandscapeWeightmapUsage* Usage = Proxy->WeightmapUsageMap.Find(WeightmapTextures[TexIdx]);
 		check(Usage);
 
 		ExistingTexAvailableChannels += Usage->FreeChannelCount();
 
-		if( ExistingTexAvailableChannels >= NeededNewChannels )
+		if (ExistingTexAvailableChannels >= NeededNewChannels)
 		{
 			break;
 		}
 	}
-	
-	if( ExistingTexAvailableChannels >= NeededNewChannels )
+
+	if (ExistingTexAvailableChannels >= NeededNewChannels)
 	{
 		// UE_LOG(LogLandscape, Log, TEXT("Existing texture has available channels"));
 
 		// Allocate using our existing textures' spare channels.
-		for( int32 TexIdx=0;TexIdx<WeightmapTextures.Num();TexIdx++ )
+		for (int32 TexIdx = 0; TexIdx < WeightmapTextures.Num(); TexIdx++)
 		{
 			FLandscapeWeightmapUsage* Usage = Proxy->WeightmapUsageMap.Find(WeightmapTextures[TexIdx]);
-			
-			for( int32 ChanIdx=0;ChanIdx<4;ChanIdx++ )
+
+			for (int32 ChanIdx = 0; ChanIdx < 4; ChanIdx++)
 			{
-				if( Usage->ChannelUsage[ChanIdx]==NULL )
+				if (Usage->ChannelUsage[ChanIdx] == NULL)
 				{
-					for( int32 LayerIdx=0;LayerIdx < WeightmapLayerAllocations.Num();LayerIdx++ )
+					for (int32 LayerIdx = 0; LayerIdx < WeightmapLayerAllocations.Num(); LayerIdx++)
 					{
 						FWeightmapLayerAllocationInfo& AllocInfo = WeightmapLayerAllocations[LayerIdx];
-						if( AllocInfo.WeightmapTextureIndex == 255 )
+						if (AllocInfo.WeightmapTextureIndex == 255)
 						{
 							// Zero out the data for this texture channel
-							if( DataInterface )
+							if (DataInterface)
 							{
-								DataInterface->ZeroTextureChannel( WeightmapTextures[TexIdx], ChanIdx );
+								DataInterface->ZeroTextureChannel(WeightmapTextures[TexIdx], ChanIdx);
 							}
 
 							AllocInfo.WeightmapTextureIndex = TexIdx;
@@ -4114,7 +4114,7 @@ void ULandscapeComponent::ReallocateWeightmaps(FLandscapeEditDataInterface* Data
 							Usage->ChannelUsage[ChanIdx] = this;
 							NeededNewChannels--;
 
-							if( NeededNewChannels == 0 )
+							if (NeededNewChannels == 0)
 							{
 								return;
 							}
@@ -4133,31 +4133,31 @@ void ULandscapeComponent::ReallocateWeightmaps(FLandscapeEditDataInterface* Data
 	int32 TotalNeededChannels = WeightmapLayerAllocations.Num();
 	int32 CurrentLayer = 0;
 	TArray<UTexture2D*> NewWeightmapTextures;
-	while( TotalNeededChannels > 0 )
+	while (TotalNeededChannels > 0)
 	{
 		// UE_LOG(LogLandscape, Log, TEXT("Still need %d channels"), TotalNeededChannels);
 
 		UTexture2D* CurrentWeightmapTexture = NULL;
 		FLandscapeWeightmapUsage* CurrentWeightmapUsage = NULL;
 
-		if( TotalNeededChannels < 4 )
+		if (TotalNeededChannels < 4)
 		{
 			// UE_LOG(LogLandscape, Log, TEXT("Looking for nearest"));
 
 			// see if we can find a suitable existing weightmap texture with sufficient channels
 			int32 BestDistanceSquared = MAX_int32;
-			for( TMap<UTexture2D*,struct FLandscapeWeightmapUsage>::TIterator It(Proxy->WeightmapUsageMap); It; ++It )
+			for (TMap<UTexture2D*, struct FLandscapeWeightmapUsage>::TIterator It(Proxy->WeightmapUsageMap); It; ++It)
 			{
 				FLandscapeWeightmapUsage* TryWeightmapUsage = &It.Value();
-				if( TryWeightmapUsage->FreeChannelCount() >= TotalNeededChannels )
+				if (TryWeightmapUsage->FreeChannelCount() >= TotalNeededChannels)
 				{
 					// See if this candidate is closer than any others we've found
-					for( int32 ChanIdx=0;ChanIdx<4;ChanIdx++ )
+					for (int32 ChanIdx = 0; ChanIdx < 4; ChanIdx++)
 					{
-						if( TryWeightmapUsage->ChannelUsage[ChanIdx] != NULL  )
+						if (TryWeightmapUsage->ChannelUsage[ChanIdx] != NULL)
 						{
 							int32 TryDistanceSquared = (TryWeightmapUsage->ChannelUsage[ChanIdx]->GetSectionBase() - GetSectionBase()).SizeSquared();
-							if( TryDistanceSquared < BestDistanceSquared )
+							if (TryDistanceSquared < BestDistanceSquared)
 							{
 								CurrentWeightmapTexture = It.Key();
 								CurrentWeightmapUsage = TryWeightmapUsage;
@@ -4169,14 +4169,14 @@ void ULandscapeComponent::ReallocateWeightmaps(FLandscapeEditDataInterface* Data
 			}
 		}
 
-		bool NeedsUpdateResource=false;
+		bool NeedsUpdateResource = false;
 		// No suitable weightmap texture
-		if( CurrentWeightmapTexture == NULL )
+		if (CurrentWeightmapTexture == NULL)
 		{
 			MarkPackageDirty();
 
 			// Weightmap is sized the same as the component
-			int32 WeightmapSize = (SubsectionSizeQuads+1) * NumSubsections;
+			int32 WeightmapSize = (SubsectionSizeQuads + 1) * NumSubsections;
 
 			// We need a new weightmap texture
 			CurrentWeightmapTexture = GetLandscapeProxy()->CreateLandscapeTexture(WeightmapSize, WeightmapSize, TEXTUREGROUP_Terrain_Weightmap, TSF_BGRA8);
@@ -4192,21 +4192,21 @@ void ULandscapeComponent::ReallocateWeightmaps(FLandscapeEditDataInterface* Data
 
 		NewWeightmapTextures.Add(CurrentWeightmapTexture);
 
-		for( int32 ChanIdx=0;ChanIdx<4 && TotalNeededChannels > 0;ChanIdx++ )
+		for (int32 ChanIdx = 0; ChanIdx < 4 && TotalNeededChannels > 0; ChanIdx++)
 		{
 			// UE_LOG(LogLandscape, Log, TEXT("Finding allocation for layer %d"), CurrentLayer);
 
-			if( CurrentWeightmapUsage->ChannelUsage[ChanIdx] == NULL  )
+			if (CurrentWeightmapUsage->ChannelUsage[ChanIdx] == NULL)
 			{
 				// Use this allocation
 				FWeightmapLayerAllocationInfo& AllocInfo = WeightmapLayerAllocations[CurrentLayer];
 
-				if( AllocInfo.WeightmapTextureIndex == 255 )
+				if (AllocInfo.WeightmapTextureIndex == 255)
 				{
 					// New layer - zero out the data for this texture channel
-					if( DataInterface )
+					if (DataInterface)
 					{
-						DataInterface->ZeroTextureChannel( CurrentWeightmapTexture, ChanIdx );
+						DataInterface->ZeroTextureChannel(CurrentWeightmapTexture, ChanIdx);
 						// UE_LOG(LogLandscape, Log, TEXT("Zeroing out channel %s.%d"), *CurrentWeightmapTexture->GetName(), ChanIdx);
 					}
 				}
@@ -4217,8 +4217,8 @@ void ULandscapeComponent::ReallocateWeightmaps(FLandscapeEditDataInterface* Data
 					// Copy the data
 					if (ensure(DataInterface != NULL)) // it's not safe to skip the copy
 					{
-						DataInterface->CopyTextureChannel( CurrentWeightmapTexture, ChanIdx, OldWeightmapTexture, AllocInfo.WeightmapTextureChannel );
-						DataInterface->ZeroTextureChannel( OldWeightmapTexture, AllocInfo.WeightmapTextureChannel );
+						DataInterface->CopyTextureChannel(CurrentWeightmapTexture, ChanIdx, OldWeightmapTexture, AllocInfo.WeightmapTextureChannel);
+						DataInterface->ZeroTextureChannel(OldWeightmapTexture, AllocInfo.WeightmapTextureChannel);
 						// UE_LOG(LogLandscape, Log, TEXT("Copying old channel (%s).%d to new channel (%s).%d"), *OldWeightmapTexture->GetName(), AllocInfo.WeightmapTextureChannel, *CurrentWeightmapTexture->GetName(), ChanIdx);
 					}
 
@@ -4229,7 +4229,7 @@ void ULandscapeComponent::ReallocateWeightmaps(FLandscapeEditDataInterface* Data
 
 				// Assign the new allocation
 				CurrentWeightmapUsage->ChannelUsage[ChanIdx] = this;
-				AllocInfo.WeightmapTextureIndex = NewWeightmapTextures.Num()-1;
+				AllocInfo.WeightmapTextureIndex = NewWeightmapTextures.Num() - 1;
 				AllocInfo.WeightmapTextureChannel = ChanIdx;
 				CurrentLayer++;
 				TotalNeededChannels--;
@@ -4243,7 +4243,7 @@ void ULandscapeComponent::ReallocateWeightmaps(FLandscapeEditDataInterface* Data
 	if (DataInterface)
 	{
 		// Update the mipmaps for the textures we edited
-		for( int32 Idx=0;Idx<WeightmapTextures.Num();Idx++)
+		for (int32 Idx = 0; Idx < WeightmapTextures.Num(); Idx++)
 		{
 			UTexture2D* WeightmapTexture = WeightmapTextures[Idx];
 			FLandscapeTextureDataInfo* WeightmapDataInfo = DataInterface->GetTextureDataInfo(WeightmapTexture);
@@ -4251,7 +4251,7 @@ void ULandscapeComponent::ReallocateWeightmaps(FLandscapeEditDataInterface* Data
 			int32 NumMips = WeightmapTexture->Source.GetNumMips();
 			TArray<FColor*> WeightmapTextureMipData;
 			WeightmapTextureMipData.AddUninitialized(NumMips);
-			for( int32 MipIdx=0;MipIdx<NumMips;MipIdx++ )
+			for (int32 MipIdx = 0; MipIdx < NumMips; MipIdx++)
 			{
 				WeightmapTextureMipData[MipIdx] = (FColor*)WeightmapDataInfo->GetMipData(MipIdx);
 			}
@@ -4265,7 +4265,7 @@ void ALandscapeProxy::RemoveInvalidWeightmaps()
 {
 	if (GIsEditor)
 	{
-		for( TMap< UTexture2D*,struct FLandscapeWeightmapUsage >::TIterator It(WeightmapUsageMap); It; ++It )
+		for (TMap< UTexture2D*, struct FLandscapeWeightmapUsage >::TIterator It(WeightmapUsageMap); It; ++It)
 		{
 			UTexture2D* Tex = It.Key();
 			FLandscapeWeightmapUsage& Usage = It.Value();
@@ -4283,7 +4283,7 @@ void ALandscapeProxy::RemoveInvalidWeightmaps()
 		}
 
 		// Remove Unused Weightmaps...
-		for (int32 Idx=0; Idx < LandscapeComponents.Num(); ++Idx)
+		for (int32 Idx = 0; Idx < LandscapeComponents.Num(); ++Idx)
 		{
 			ULandscapeComponent* Component = LandscapeComponents[Idx];
 			Component->RemoveInvalidWeightmaps();
@@ -4296,14 +4296,14 @@ void ULandscapeComponent::RemoveInvalidWeightmaps()
 	// Adjust WeightmapTextureIndex index for other layers
 	TSet<int32> UsedTextureIndices;
 	TSet<int32> AllTextureIndices;
-	for( int32 LayerIdx=0;LayerIdx<WeightmapLayerAllocations.Num();LayerIdx++ )
+	for (int32 LayerIdx = 0; LayerIdx < WeightmapLayerAllocations.Num(); LayerIdx++)
 	{
-		UsedTextureIndices.Add( WeightmapLayerAllocations[LayerIdx].WeightmapTextureIndex );
+		UsedTextureIndices.Add(WeightmapLayerAllocations[LayerIdx].WeightmapTextureIndex);
 	}
 
-	for ( int32 WeightIdx=0; WeightIdx < WeightmapTextures.Num(); ++WeightIdx )
+	for (int32 WeightIdx = 0; WeightIdx < WeightmapTextures.Num(); ++WeightIdx)
 	{
-		AllTextureIndices.Add( WeightIdx );
+		AllTextureIndices.Add(WeightIdx);
 	}
 
 	TSet<int32> UnUsedTextureIndices = AllTextureIndices.Difference(UsedTextureIndices);
@@ -4316,19 +4316,19 @@ void ULandscapeComponent::RemoveInvalidWeightmaps()
 		WeightmapTextures[DeleteLayerWeightmapTextureIndex]->Modify();
 		WeightmapTextures[DeleteLayerWeightmapTextureIndex]->MarkPackageDirty();
 		WeightmapTextures[DeleteLayerWeightmapTextureIndex]->ClearFlags(RF_Standalone);
-		WeightmapTextures.RemoveAt( DeleteLayerWeightmapTextureIndex );
+		WeightmapTextures.RemoveAt(DeleteLayerWeightmapTextureIndex);
 
 		// Adjust WeightmapTextureIndex index for other layers
-		for( int32 LayerIdx=0;LayerIdx<WeightmapLayerAllocations.Num();LayerIdx++ )
+		for (int32 LayerIdx = 0; LayerIdx < WeightmapLayerAllocations.Num(); LayerIdx++)
 		{
 			FWeightmapLayerAllocationInfo& Allocation = WeightmapLayerAllocations[LayerIdx];
 
-			if( Allocation.WeightmapTextureIndex > DeleteLayerWeightmapTextureIndex )
+			if (Allocation.WeightmapTextureIndex > DeleteLayerWeightmapTextureIndex)
 			{
 				Allocation.WeightmapTextureIndex--;
 			}
 
-			check( Allocation.WeightmapTextureIndex < WeightmapTextures.Num() );
+			check(Allocation.WeightmapTextureIndex < WeightmapTextures.Num());
 		}
 		DeletedLayers++;
 	}
@@ -4336,15 +4336,15 @@ void ULandscapeComponent::RemoveInvalidWeightmaps()
 
 void ULandscapeComponent::InitHeightmapData(TArray<FColor>& Heights, bool bUpdateCollision)
 {
-	int32 ComponentSizeVerts = NumSubsections * (SubsectionSizeQuads+1);
+	int32 ComponentSizeVerts = NumSubsections * (SubsectionSizeQuads + 1);
 
-	if (Heights.Num() != FMath::Square(ComponentSizeVerts) )
+	if (Heights.Num() != FMath::Square(ComponentSizeVerts))
 	{
 		return;
 	}
 
 	// Handling old Height map....
-	if (HeightmapTexture && HeightmapTexture->GetOutermost() != GetTransientPackage() 
+	if (HeightmapTexture && HeightmapTexture->GetOutermost() != GetTransientPackage()
 		&& HeightmapTexture->GetOutermost() == GetOutermost()
 		&& HeightmapTexture->Source.GetSizeX() >= ComponentSizeVerts) // if Height map is not valid...
 	{
@@ -4357,20 +4357,20 @@ void ULandscapeComponent::InitHeightmapData(TArray<FColor>& Heights, bool bUpdat
 	// New Height map
 	TArray<FColor*> HeightmapTextureMipData;
 	// make sure the heightmap UVs are powers of two.
-	int32 HeightmapSizeU = (1<<FMath::CeilLogTwo( ComponentSizeVerts ));
-	int32 HeightmapSizeV = (1<<FMath::CeilLogTwo( ComponentSizeVerts ));
+	int32 HeightmapSizeU = (1 << FMath::CeilLogTwo(ComponentSizeVerts));
+	int32 HeightmapSizeV = (1 << FMath::CeilLogTwo(ComponentSizeVerts));
 
 	// Height map construction
 	HeightmapTexture = GetLandscapeProxy()->CreateLandscapeTexture(HeightmapSizeU, HeightmapSizeV, TEXTUREGROUP_Terrain_Heightmap, TSF_BGRA8);
-	
+
 	int32 MipSubsectionSizeQuads = SubsectionSizeQuads;
 	int32 MipSizeU = HeightmapSizeU;
 	int32 MipSizeV = HeightmapSizeV;
 
-	HeightmapScaleBias = FVector4( 1.f / (float)HeightmapSizeU, 1.f / (float)HeightmapSizeV, 0.f, 0.f);
+	HeightmapScaleBias = FVector4(1.f / (float)HeightmapSizeU, 1.f / (float)HeightmapSizeV, 0.f, 0.f);
 
 	int32 Mip = 0;
-	while( MipSizeU > 1 && MipSizeV > 1 && MipSubsectionSizeQuads >= 1 )
+	while (MipSizeU > 1 && MipSizeV > 1 && MipSubsectionSizeQuads >= 1)
 	{
 		FColor* HeightmapTextureData = (FColor*)HeightmapTexture->Source.LockMip(Mip);
 		if (Mip == 0)
@@ -4389,14 +4389,14 @@ void ULandscapeComponent::InitHeightmapData(TArray<FColor>& Heights, bool bUpdat
 
 		MipSubsectionSizeQuads = ((MipSubsectionSizeQuads + 1) >> 1) - 1;
 	}
-	ULandscapeComponent::GenerateHeightmapMips( HeightmapTextureMipData );
+	ULandscapeComponent::GenerateHeightmapMips(HeightmapTextureMipData);
 
-	if( bUpdateCollision )
+	if (bUpdateCollision)
 	{
-		ULandscapeComponent::UpdateCollisionHeightData( HeightmapTextureMipData[CollisionMipLevel] );
+		ULandscapeComponent::UpdateCollisionHeightData(HeightmapTextureMipData[CollisionMipLevel]);
 	}
 
-	for( int32 i=0;i<HeightmapTextureMipData.Num();i++ )
+	for (int32 i = 0; i < HeightmapTextureMipData.Num(); i++)
 	{
 		HeightmapTexture->Source.UnlockMip(i);
 	}
@@ -4410,22 +4410,22 @@ void ULandscapeComponent::InitWeightmapData(TArray<ULandscapeLayerInfoObject*>&
 		return;
 	}
 
-	int32 ComponentSizeVerts = NumSubsections * (SubsectionSizeQuads+1);
+	int32 ComponentSizeVerts = NumSubsections * (SubsectionSizeQuads + 1);
 
 	// Validation..
 	for (int32 Idx = 0; Idx < WeightmapData.Num(); ++Idx)
 	{
-		if ( WeightmapData[Idx].Num() != FMath::Square(ComponentSizeVerts) )
+		if (WeightmapData[Idx].Num() != FMath::Square(ComponentSizeVerts))
 		{
 			return;
 		}
 	}
-	
+
 	for (int32 Idx = 0; Idx < WeightmapTextures.Num(); ++Idx)
 	{
-		if (WeightmapTextures[Idx] && WeightmapTextures[Idx]->GetOutermost() != GetTransientPackage() 
+		if (WeightmapTextures[Idx] && WeightmapTextures[Idx]->GetOutermost() != GetTransientPackage()
 			&& WeightmapTextures[Idx]->GetOutermost() == GetOutermost()
-			&& WeightmapTextures[Idx]->Source.GetSizeX() == ComponentSizeVerts) 
+			&& WeightmapTextures[Idx]->Source.GetSizeX() == ComponentSizeVerts)
 		{
 			WeightmapTextures[Idx]->SetFlags(RF_Transactional);
 			WeightmapTextures[Idx]->Modify();
@@ -4438,19 +4438,19 @@ void ULandscapeComponent::InitWeightmapData(TArray<ULandscapeLayerInfoObject*>&
 	WeightmapLayerAllocations.Empty(LayerInfos.Num());
 	for (int32 Idx = 0; Idx < LayerInfos.Num(); ++Idx)
 	{
-		new (WeightmapLayerAllocations) FWeightmapLayerAllocationInfo(LayerInfos[Idx]);
+		new (WeightmapLayerAllocations)FWeightmapLayerAllocationInfo(LayerInfos[Idx]);
 	}
 
 	ReallocateWeightmaps(NULL);
 
-	check(WeightmapLayerAllocations.Num() > 0 && WeightmapTextures.Num() > 0 );
+	check(WeightmapLayerAllocations.Num() > 0 && WeightmapTextures.Num() > 0);
 
 	int32 WeightmapSize = ComponentSizeVerts;
-	WeightmapScaleBias = FVector4( 1.f / (float)WeightmapSize, 1.f / (float)WeightmapSize, 0.5f / (float)WeightmapSize, 0.5f / (float)WeightmapSize);
-	WeightmapSubsectionOffset =  (float)(SubsectionSizeQuads+1) / (float)WeightmapSize;
+	WeightmapScaleBias = FVector4(1.f / (float)WeightmapSize, 1.f / (float)WeightmapSize, 0.5f / (float)WeightmapSize, 0.5f / (float)WeightmapSize);
+	WeightmapSubsectionOffset = (float)(SubsectionSizeQuads + 1) / (float)WeightmapSize;
 
 	// Channel remapping
-	int32 ChannelOffsets[4] = {(int32)STRUCT_OFFSET(FColor,R),(int32)STRUCT_OFFSET(FColor,G),(int32)STRUCT_OFFSET(FColor,B),(int32)STRUCT_OFFSET(FColor,A)};
+	int32 ChannelOffsets[4] = { (int32)STRUCT_OFFSET(FColor, R), (int32)STRUCT_OFFSET(FColor, G), (int32)STRUCT_OFFSET(FColor, B), (int32)STRUCT_OFFSET(FColor, A) };
 
 	TArray<void*> WeightmapDataPtrs;
 	WeightmapDataPtrs.AddUninitialized(WeightmapTextures.Num());
@@ -4462,22 +4462,22 @@ void ULandscapeComponent::InitWeightmapData(TArray<ULandscapeLayerInfoObject*>&
 	for (int32 LayerIdx = 0; LayerIdx < WeightmapLayerAllocations.Num(); ++LayerIdx)
 	{
 		void* DestDataPtr = WeightmapDataPtrs[WeightmapLayerAllocations[LayerIdx].WeightmapTextureIndex];
-		uint8* DestTextureData = (uint8*)DestDataPtr + ChannelOffsets[ WeightmapLayerAllocations[LayerIdx].WeightmapTextureChannel ];
+		uint8* DestTextureData = (uint8*)DestDataPtr + ChannelOffsets[WeightmapLayerAllocations[LayerIdx].WeightmapTextureChannel];
 		uint8* SrcTextureData = (uint8*)&WeightmapData[LayerIdx][0];
 
-		for( int32 i=0;i<WeightmapData[LayerIdx].Num();i++ )
+		for (int32 i = 0; i < WeightmapData[LayerIdx].Num(); i++)
 		{
-			DestTextureData[i*4] = SrcTextureData[i];
+			DestTextureData[i * 4] = SrcTextureData[i];
 		}
 	}
 
-	for( int32 Idx=0;Idx<WeightmapTextures.Num();Idx++)
+	for (int32 Idx = 0; Idx < WeightmapTextures.Num(); Idx++)
 	{
 		UTexture2D* WeightmapTexture = WeightmapTextures[Idx];
 		WeightmapTexture->Source.UnlockMip(0);
 	}
 
-	for( int32 Idx=0;Idx<WeightmapTextures.Num();Idx++)
+	for (int32 Idx = 0; Idx < WeightmapTextures.Num(); Idx++)
 	{
 		UTexture2D* WeightmapTexture = WeightmapTextures[Idx];
 		{
@@ -4486,7 +4486,7 @@ void ULandscapeComponent::InitWeightmapData(TArray<ULandscapeLayerInfoObject*>&
 			int32 NumMips = WeightmapTexture->Source.GetNumMips();
 			TArray<FColor*> WeightmapTextureMipData;
 			WeightmapTextureMipData.AddUninitialized(NumMips);
-			for( int32 MipIdx=0;MipIdx<NumMips;MipIdx++ )
+			for (int32 MipIdx = 0; MipIdx < NumMips; MipIdx++)
 			{
 				WeightmapTextureMipData[MipIdx] = (FColor*)WeightmapDataInfo.GetMipData(MipIdx);
 			}
@@ -4512,8 +4512,8 @@ bool ALandscapeProxy::ShouldExport()
 	if (!bIsMovingToLevel && LandscapeComponents.Num() > MAX_LANDSCAPE_EXPORT_COMPONENTS_NUM)
 	{
 		// Prompt to save startup packages
-		if( EAppReturnType::Yes == FMessageDialog::Open( EAppMsgType::YesNo, FText::Format(
-			NSLOCTEXT("UnrealEd", "LandscapeExport_Warning", "Landscape has large number({0}) of components, so it will use large amount memory to copy it to the clipboard. Do you want to proceed?"), FText::AsNumber(LandscapeComponents.Num())) ) )
+		if (EAppReturnType::Yes == FMessageDialog::Open(EAppMsgType::YesNo, FText::Format(
+			NSLOCTEXT("UnrealEd", "LandscapeExport_Warning", "Landscape has large number({0}) of components, so it will use large amount memory to copy it to the clipboard. Do you want to proceed?"), FText::AsNumber(LandscapeComponents.Num()))))
 		{
 			return true;
 		}
@@ -4531,8 +4531,8 @@ bool ALandscapeProxy::ShouldImport(FString* ActorPropString, bool IsMovingToLeve
 	if (!bIsMovingToLevel && ActorPropString && ActorPropString->Len() > MAX_LANDSCAPE_PROP_TEXT_LENGTH)
 	{
 		// Prompt to save startup packages
-		if( EAppReturnType::Yes == FMessageDialog::Open( EAppMsgType::YesNo, FText::Format(
-			NSLOCTEXT("UnrealEd", "LandscapeImport_Warning", "Landscape is about to import large amount memory ({0}MB) from the clipboard, which will take some time. Do you want to proceed?"), FText::AsNumber(ActorPropString->Len() >> 20 )) ) )
+		if (EAppReturnType::Yes == FMessageDialog::Open(EAppMsgType::YesNo, FText::Format(
+			NSLOCTEXT("UnrealEd", "LandscapeImport_Warning", "Landscape is about to import large amount memory ({0}MB) from the clipboard, which will take some time. Do you want to proceed?"), FText::AsNumber(ActorPropString->Len() >> 20))))
 		{
 			return true;
 		}
@@ -4551,42 +4551,42 @@ void ULandscapeComponent::ExportCustomProperties(FOutputDevice& Out, uint32 Inde
 		return;
 	}
 	// Height map
-	int32 NumVertices = FMath::Square( NumSubsections*(SubsectionSizeQuads+1) );
+	int32 NumVertices = FMath::Square(NumSubsections*(SubsectionSizeQuads + 1));
 	FLandscapeComponentDataInterface DataInterface(this);
 	TArray<FColor> Heightmap;
 	DataInterface.GetHeightmapTextureData(Heightmap);
 	check(Heightmap.Num() == NumVertices);
 
-	Out.Logf( TEXT("%sCustomProperties LandscapeHeightData "), FCString::Spc(Indent));
-	for( int32 i=0;i<NumVertices;i++ )
+	Out.Logf(TEXT("%sCustomProperties LandscapeHeightData "), FCString::Spc(Indent));
+	for (int32 i = 0; i < NumVertices; i++)
 	{
-		Out.Logf( TEXT("%x "), Heightmap[i].DWColor() );
+		Out.Logf(TEXT("%x "), Heightmap[i].DWColor());
 	}
 
 	TArray<uint8> Weightmap;
 	// Weight map
-	Out.Logf( TEXT("LayerNum=%d "), WeightmapLayerAllocations.Num());
-	for (int32 i=0; i < WeightmapLayerAllocations.Num(); i++)
+	Out.Logf(TEXT("LayerNum=%d "), WeightmapLayerAllocations.Num());
+	for (int32 i = 0; i < WeightmapLayerAllocations.Num(); i++)
 	{
 		if (DataInterface.GetWeightmapTextureData(WeightmapLayerAllocations[i].LayerInfo, Weightmap))
 		{
-			Out.Logf( TEXT("LayerInfo=%s "), *WeightmapLayerAllocations[i].LayerInfo->GetPathName());
-			for (int32 VertexIndex = 0; VertexIndex<NumVertices; VertexIndex++)
+			Out.Logf(TEXT("LayerInfo=%s "), *WeightmapLayerAllocations[i].LayerInfo->GetPathName());
+			for (int32 VertexIndex = 0; VertexIndex < NumVertices; VertexIndex++)
 			{
 				Out.Logf(TEXT("%x "), Weightmap[VertexIndex]);
 			}
 		}
 	}
 
-	Out.Logf( TEXT("\r\n") );
+	Out.Logf(TEXT("\r\n"));
 }
 
 
 void ULandscapeComponent::ImportCustomProperties(const TCHAR* SourceText, FFeedbackContext* Warn)
 {
-	if(FParse::Command(&SourceText,TEXT("LandscapeHeightData")))
+	if (FParse::Command(&SourceText, TEXT("LandscapeHeightData")))
 	{
-		int32 NumVertices = FMath::Square( NumSubsections*(SubsectionSizeQuads+1) );
+		int32 NumVertices = FMath::Square(NumSubsections*(SubsectionSizeQuads + 1));
 
 		TArray<FColor> Heights;
 		Heights.Empty(NumVertices);
@@ -4595,26 +4595,26 @@ void ULandscapeComponent::ImportCustomProperties(const TCHAR* SourceText, FFeedb
 		FParse::Next(&SourceText);
 		int32 i = 0;
 		TCHAR* StopStr;
-		while( FChar::IsHexDigit(*SourceText) ) 
+		while (FChar::IsHexDigit(*SourceText))
 		{
-			if( i < NumVertices )
+			if (i < NumVertices)
 			{
 				Heights[i++].DWColor() = FCString::Strtoi(SourceText, &StopStr, 16);
-				while( FChar::IsHexDigit(*SourceText) ) 
+				while (FChar::IsHexDigit(*SourceText))
 				{
 					SourceText++;
 				}
 			}
 
 			FParse::Next(&SourceText);
-		} 
+		}
 
-		if( i != NumVertices )
+		if (i != NumVertices)
 		{
-			Warn->Logf( *NSLOCTEXT( "Core", "SyntaxError", "Syntax Error" ).ToString() );
+			Warn->Logf(*NSLOCTEXT("Core", "SyntaxError", "Syntax Error").ToString());
 		}
 
-		int32 ComponentSizeVerts = NumSubsections * (SubsectionSizeQuads+1);
+		int32 ComponentSizeVerts = NumSubsections * (SubsectionSizeQuads + 1);
 
 		InitHeightmapData(Heights, false);
 
@@ -4622,7 +4622,7 @@ void ULandscapeComponent::ImportCustomProperties(const TCHAR* SourceText, FFeedb
 		int32 LayerNum = 0;
 		if (FParse::Value(SourceText, TEXT("LayerNum="), LayerNum))
 		{
-			while(*SourceText && (!FChar::IsWhitespace(*SourceText)))
+			while (*SourceText && (!FChar::IsWhitespace(*SourceText)))
 			{
 				++SourceText;
 			}
@@ -4638,7 +4638,7 @@ void ULandscapeComponent::ImportCustomProperties(const TCHAR* SourceText, FFeedb
 		TArray<ULandscapeLayerInfoObject*> LayerInfos;
 		LayerInfos.Empty(LayerNum);
 		TArray<TArray<uint8>> WeightmapData;
-		for (int32 LayerIndex=0; LayerIndex < LayerNum; ++LayerIndex)
+		for (int32 LayerIndex = 0; LayerIndex < LayerNum; ++LayerIndex)
 		{
 			TArray<uint8> Weights;
 			Weights.Empty(NumVertices);
@@ -4648,13 +4648,13 @@ void ULandscapeComponent::ImportCustomProperties(const TCHAR* SourceText, FFeedb
 
 		int32 LayerIdx = 0;
 		FString LayerInfoPath;
-		while ( *SourceText )
+		while (*SourceText)
 		{
 			if (FParse::Value(SourceText, TEXT("LayerInfo="), LayerInfoPath))
 			{
 				LayerInfos.Add(LoadObject<ULandscapeLayerInfoObject>(NULL, *LayerInfoPath));
 
-				while(*SourceText && (!FChar::IsWhitespace(*SourceText)))
+				while (*SourceText && (!FChar::IsWhitespace(*SourceText)))
 				{
 					++SourceText;
 				}
@@ -4662,22 +4662,22 @@ void ULandscapeComponent::ImportCustomProperties(const TCHAR* SourceText, FFeedb
 				check(*SourceText);
 
 				i = 0;
-				while( FChar::IsHexDigit(*SourceText) ) 
+				while (FChar::IsHexDigit(*SourceText))
 				{
-					if( i < NumVertices )
+					if (i < NumVertices)
 					{
 						(WeightmapData[LayerIdx])[i++] = (uint8)FCString::Strtoi(SourceText, &StopStr, 16);
-						while( FChar::IsHexDigit(*SourceText) ) 
+						while (FChar::IsHexDigit(*SourceText))
 						{
 							SourceText++;
 						}
 					}
 					FParse::Next(&SourceText);
-				} 
+				}
 
-				if( i != NumVertices )
+				if (i != NumVertices)
 				{
-					Warn->Logf( *NSLOCTEXT( "Core", "SyntaxError", "Syntax Error" ).ToString() );
+					Warn->Logf(*NSLOCTEXT("Core", "SyntaxError", "Syntax Error").ToString());
 				}
 				LayerIdx++;
 			}
@@ -4701,9 +4701,9 @@ bool ALandscapeProxy::IsValidLandscapeActor(ALandscape* Landscape)
 			{
 				return true; // always valid for newly created Proxy
 			}
-			if ( ( (LandscapeActor && LandscapeActor == Landscape)
-				    || (LandscapeActor.IsNull() && LandscapeGuid.IsValid() && LandscapeGuid == Landscape->LandscapeGuid) )
-				&& ComponentSizeQuads == Landscape->ComponentSizeQuads 
+			if (((LandscapeActor && LandscapeActor == Landscape)
+				|| (LandscapeActor.IsNull() && LandscapeGuid.IsValid() && LandscapeGuid == Landscape->LandscapeGuid))
+				&& ComponentSizeQuads == Landscape->ComponentSizeQuads
 				&& NumSubsections == Landscape->NumSubsections
 				&& SubsectionSizeQuads == Landscape->SubsectionSizeQuads)
 			{
@@ -4736,7 +4736,7 @@ UMaterialInstance* ULandscapeComponent::GeneratePlatformPixelData(TArray<UTextur
 		return MaterialInstance;
 	}
 
-	int32 WeightmapSize = (SubsectionSizeQuads+1) * NumSubsections;
+	int32 WeightmapSize = (SubsectionSizeQuads + 1) * NumSubsections;
 	UTexture2D* WeightmapTexture = GetLandscapeActor()->CreateLandscapeTexture(WeightmapSize, WeightmapSize, TEXTUREGROUP_Terrain_Weightmap, TSF_BGRA8);
 	CreateEmptyTextureMips(WeightmapTexture);
 
@@ -4747,8 +4747,8 @@ UMaterialInstance* ULandscapeComponent::GeneratePlatformPixelData(TArray<UTextur
 		{
 			// Reordering weight map channels
 			int32 CurrentIdx = 0;
-			int32 FromTextures[3] = {0, 0, 0};
-			int32 FromChannels[3] = {0, 0, 0};
+			int32 FromTextures[3] = { 0, 0, 0 };
+			int32 FromChannels[3] = { 0, 0, 0 };
 			for (int32 i = 0; i < WeightmapLayerAllocations.Num(); ++i)
 			{
 				// Only for valid Layers
@@ -4764,12 +4764,12 @@ UMaterialInstance* ULandscapeComponent::GeneratePlatformPixelData(TArray<UTextur
 				}
 			}
 
-			LandscapeEdit.CopyTextureChannel( WeightmapTexture, 0, InWeightmapTextures[FromTextures[0]], FromChannels[0] );
-			LandscapeEdit.CopyTextureChannel( WeightmapTexture, 1, InWeightmapTextures[FromTextures[1]], FromChannels[1] );
+			LandscapeEdit.CopyTextureChannel(WeightmapTexture, 0, InWeightmapTextures[FromTextures[0]], FromChannels[0]);
+			LandscapeEdit.CopyTextureChannel(WeightmapTexture, 1, InWeightmapTextures[FromTextures[1]], FromChannels[1]);
 		}
 
-		LandscapeEdit.CopyTextureFromHeightmap( WeightmapTexture, 2, this, 2 );
-		LandscapeEdit.CopyTextureFromHeightmap( WeightmapTexture, 3, this, 3 );
+		LandscapeEdit.CopyTextureFromHeightmap(WeightmapTexture, 2, this, 2);
+		LandscapeEdit.CopyTextureFromHeightmap(WeightmapTexture, 3, this, 3);
 	}
 
 	WeightmapTexture->PostEditChange();
@@ -4777,51 +4777,51 @@ UMaterialInstance* ULandscapeComponent::GeneratePlatformPixelData(TArray<UTextur
 	InWeightmapTextures.Add(WeightmapTexture);
 
 	FLinearColor Masks[4];
-	Masks[0] = FLinearColor(1.f,0.f,0.f,0.f);
-	Masks[1] = FLinearColor(0.f,1.f,0.f,0.f);
-	Masks[2] = FLinearColor(0.f,0.f,1.f,0.f);
-	Masks[3] = FLinearColor(0.f,0.f,0.f,1.f);
+	Masks[0] = FLinearColor(1.f, 0.f, 0.f, 0.f);
+	Masks[1] = FLinearColor(0.f, 1.f, 0.f, 0.f);
+	Masks[2] = FLinearColor(0.f, 0.f, 1.f, 0.f);
+	Masks[3] = FLinearColor(0.f, 0.f, 0.f, 1.f);
 
 	if (!bIsCooking)
 	{
-		UMaterialInstanceDynamic* MobileMaterialInstance = UMaterialInstanceDynamic::Create(MaterialInstance, GetOutermost()); 
+		UMaterialInstanceDynamic* MobileMaterialInstance = UMaterialInstanceDynamic::Create(MaterialInstance, GetOutermost());
 		// Set the layer mask
 		int32 CurrentIdx = 0;
-		for( int32 AllocIdx=0;AllocIdx<WeightmapLayerAllocations.Num();AllocIdx++ )
+		for (int32 AllocIdx = 0; AllocIdx < WeightmapLayerAllocations.Num(); AllocIdx++)
 		{
 			FWeightmapLayerAllocationInfo& Allocation = WeightmapLayerAllocations[AllocIdx];
 			FName LayerName = Allocation.LayerInfo ? Allocation.LayerInfo->LayerName : NAME_None;
 			if (Allocation.LayerInfo)
 			{
-				MobileMaterialInstance->SetVectorParameterValue(FName(*FString::Printf(TEXT("LayerMask_%s"),*LayerName.ToString())), Masks[FMath::Min(3, CurrentIdx++)]);
+				MobileMaterialInstance->SetVectorParameterValue(FName(*FString::Printf(TEXT("LayerMask_%s"), *LayerName.ToString())), Masks[FMath::Min(3, CurrentIdx++)]);
 			}
 			else
 			{
-				MobileMaterialInstance->SetVectorParameterValue(FName(*FString::Printf(TEXT("LayerMask_%s"),*LayerName.ToString())), Masks[3]);
+				MobileMaterialInstance->SetVectorParameterValue(FName(*FString::Printf(TEXT("LayerMask_%s"), *LayerName.ToString())), Masks[3]);
 			}
 		}
 		return MobileMaterialInstance;
 	}
 	else // for cooking
 	{
-		UMaterialInstanceConstant* CombinationMaterialInstance = GetCombinationMaterial(true); 
+		UMaterialInstanceConstant* CombinationMaterialInstance = GetCombinationMaterial(true);
 		UMaterialInstanceConstant* MobileMaterialInstance = ConstructObject<ULandscapeMaterialInstanceConstant>(ULandscapeMaterialInstanceConstant::StaticClass(), GetOutermost());
 
 		MobileMaterialInstance->SetParentEditorOnly(CombinationMaterialInstance);
 
 		// Set the layer mask
 		int32 CurrentIdx = 0;
-		for( int32 AllocIdx=0;AllocIdx<WeightmapLayerAllocations.Num();AllocIdx++ )
+		for (int32 AllocIdx = 0; AllocIdx < WeightmapLayerAllocations.Num(); AllocIdx++)
 		{
 			FWeightmapLayerAllocationInfo& Allocation = WeightmapLayerAllocations[AllocIdx];
 			FName LayerName = Allocation.LayerInfo ? Allocation.LayerInfo->LayerName : NAME_None;
 			if (Allocation.LayerInfo)
 			{
-				MobileMaterialInstance->SetVectorParameterValueEditorOnly(FName(*FString::Printf(TEXT("LayerMask_%s"),*LayerName.ToString())), Masks[FMath::Min(3, CurrentIdx++)]);
+				MobileMaterialInstance->SetVectorParameterValueEditorOnly(FName(*FString::Printf(TEXT("LayerMask_%s"), *LayerName.ToString())), Masks[FMath::Min(3, CurrentIdx++)]);
 			}
 			else
 			{
-				MobileMaterialInstance->SetVectorParameterValueEditorOnly(FName(*FString::Printf(TEXT("LayerMask_%s"),*LayerName.ToString())), Masks[3]);
+				MobileMaterialInstance->SetVectorParameterValueEditorOnly(FName(*FString::Printf(TEXT("LayerMask_%s"), *LayerName.ToString())), Masks[3]);
 			}
 		}
 
@@ -4857,15 +4857,15 @@ void ULandscapeComponent::GeneratePlatformVertexData()
 
 	// Get the required mip data
 	TArray<FColor*> HeightmapMipData;
-	for( int32 MipIdx=0; MipIdx < FMath::Min(LANDSCAPE_MAX_ES_LOD, HeightmapTexture->Source.GetNumMips()); MipIdx++ )
+	for (int32 MipIdx = 0; MipIdx < FMath::Min(LANDSCAPE_MAX_ES_LOD, HeightmapTexture->Source.GetNumMips()); MipIdx++)
 	{
 		int32 MipSubsectionSizeVerts = (SubsectionSizeVerts) >> MipIdx;
-		if( MipSubsectionSizeVerts > 1)
+		if (MipSubsectionSizeVerts > 1)
 		{
-			HeightmapMipData.Add( (FColor*)HeightmapTexture->Source.LockMip(MipIdx) );
+			HeightmapMipData.Add((FColor*)HeightmapTexture->Source.LockMip(MipIdx));
 		}
 	}
-	
+
 	TMap<uint64, int32> VertexMap;
 	TArray<FLandscapeVertexRef> VertexOrder;
 	VertexOrder.Empty(FMath::Square(SubsectionSizeVerts * NumSubsections));
@@ -4991,7 +4991,7 @@ void ULandscapeComponent::GeneratePlatformVertexData()
 		DstVert++;
 	}
 
-	for( int32 MipIdx=0;MipIdx<HeightmapTexture->Source.GetNumMips();MipIdx++ )
+	for (int32 MipIdx = 0; MipIdx < HeightmapTexture->Source.GetNumMips(); MipIdx++)
 	{
 		HeightmapTexture->Source.UnlockMip(MipIdx);
 	}
