Fix incorrect PIEInstance getting set on multi-player PIE world context

diff --git a/Engine/Source/Editor/UnrealEd/Private/PlayLevel.cpp b/Engine/Source/Editor/UnrealEd/Private/PlayLevel.cpp
--- a/Engine/Source/Editor/UnrealEd/Private/PlayLevel.cpp
+++ b/Engine/Source/Editor/UnrealEd/Private/PlayLevel.cpp
@@ -2273,7 +2273,7 @@ UGameInstance* UEditorEngine::CreatePIEGameInstance(int32 PIEInstance, bool bInS
 	// We need to temporarily add the GameInstance to the root because the InitPIE call can do garbage collection wiping out the GameInstance
 	GameInstance->AddToRoot();
 
-	bool bSuccess = GameInstance->InitPIE(bAnyBlueprintErrors);
+	bool bSuccess = GameInstance->InitPIE(bAnyBlueprintErrors, PIEInstance);
 	if (!bSuccess)
 	{
 		FEditorDelegates::EndPIE.Broadcast(bInSimulateInEditor);
@@ -2291,7 +2291,6 @@ UGameInstance* UEditorEngine::CreatePIEGameInstance(int32 PIEInstance, bool bInS
 	check(PieWorldContext);
 	PlayWorld = PieWorldContext->World();
 
-	PieWorldContext->PIEInstance = PIEInstance;
 	PieWorldContext->RunAsDedicated = bRunAsDedicated;
 
 	GWorld = PlayWorld;
diff --git a/Engine/Source/Runtime/Engine/Classes/Engine/GameInstance.h b/Engine/Source/Runtime/Engine/Classes/Engine/GameInstance.h
--- a/Engine/Source/Runtime/Engine/Classes/Engine/GameInstance.h
+++ b/Engine/Source/Runtime/Engine/Classes/Engine/GameInstance.h
@@ -55,7 +55,7 @@ public:
 	virtual void Init();
 
 #if WITH_EDITOR
-	virtual bool InitPIE(bool bAnyBlueprintErrors);
+	virtual bool InitPIE(bool bAnyBlueprintErrors, int32 PIEInstance);
 
 	bool StartPIEGameInstance(ULocalPlayer* LocalPlayer, bool bInSimulateInEditor, bool bAnyBlueprintErrors, bool bStartInSpectatorMode);
 #endif
diff --git a/Engine/Source/Runtime/Engine/Private/GameInstance.cpp b/Engine/Source/Runtime/Engine/Private/GameInstance.cpp
--- a/Engine/Source/Runtime/Engine/Private/GameInstance.cpp
+++ b/Engine/Source/Runtime/Engine/Private/GameInstance.cpp
@@ -39,19 +39,19 @@ void UGameInstance::Init()
 }
 
 #if WITH_EDITOR
-bool UGameInstance::InitPIE(bool bAnyBlueprintErrors)
+bool UGameInstance::InitPIE(bool bAnyBlueprintErrors, int32 PIEInstance)
 {
 	UEditorEngine* const EditorEngine = CastChecked<UEditorEngine>(GetEngine());
 
 	WorldContext = &EditorEngine->CreateNewWorldContext(EWorldType::PIE);
 	WorldContext->OwningGameInstance = this;
-	WorldContext->PIEInstance = 0;		// @fixme?
+	WorldContext->PIEInstance = PIEInstance;
 
 	const FString WorldPackageName = EditorEngine->EditorWorld->GetOutermost()->GetName();
 
 	// Establish World Context for PIE World
 	WorldContext->LastURL.Map = WorldPackageName;
-	WorldContext->PIEPrefix = UWorld::BuildPIEPackagePrefix(WorldContext->PIEInstance);
+	WorldContext->PIEPrefix = UWorld::BuildPIEPackagePrefix(PIEInstance);
 
 	const ULevelEditorPlaySettings* PlayInSettings = GetDefault<ULevelEditorPlaySettings>();
 
