Fix DOF lock getting garbage collected on skeletal mesh

diff --git a/Engine/Source/Runtime/Engine/Classes/PhysicsEngine/BodyInstance.h b/Engine/Source/Runtime/Engine/Classes/PhysicsEngine/BodyInstance.h
--- a/Engine/Source/Runtime/Engine/Classes/PhysicsEngine/BodyInstance.h
+++ b/Engine/Source/Runtime/Engine/Classes/PhysicsEngine/BodyInstance.h
@@ -198,8 +198,7 @@ public:
 	void CreateDOFLock();
 
 	/** Constraint used to allow for easy DOF setup per bodyinstance */
-	UPROPERTY(transient)
-	UPhysicsConstraintComponent * DOFConstraint;
+	FConstraintInstance * DOFConstraint;
 
 protected:
 
diff --git a/Engine/Source/Runtime/Engine/Classes/PhysicsEngine/ConstraintInstance.h b/Engine/Source/Runtime/Engine/Classes/PhysicsEngine/ConstraintInstance.h
--- a/Engine/Source/Runtime/Engine/Classes/PhysicsEngine/ConstraintInstance.h
+++ b/Engine/Source/Runtime/Engine/Classes/PhysicsEngine/ConstraintInstance.h
@@ -484,6 +484,10 @@ public:
 	void PostSerialize(const FArchive& Ar);
 
 	void OnConstraintBroken();
+
+	//Hacks to easily get zeroed memory for special case when we don't use GC
+	static void Free(FConstraintInstance * Ptr);
+	static FConstraintInstance * Alloc();
 };
 
 template<>
diff --git a/Engine/Source/Runtime/Engine/Private/PhysicsEngine/BodyInstance.cpp b/Engine/Source/Runtime/Engine/Private/PhysicsEngine/BodyInstance.cpp
--- a/Engine/Source/Runtime/Engine/Private/PhysicsEngine/BodyInstance.cpp
+++ b/Engine/Source/Runtime/Engine/Private/PhysicsEngine/BodyInstance.cpp
@@ -532,8 +532,12 @@ FVector FBodyInstance::GetLockedAxis() const
 
 void FBodyInstance::CreateDOFLock()
 {
+	if (DOFConstraint)
+	{
+		DOFConstraint->TermConstraint();
+		FConstraintInstance::Free(DOFConstraint);
+	}
 
-	DOFConstraint = NULL;
 	//setup constraint based on DOF
 	FVector LockedAxis = GetLockedAxis();
 
@@ -542,19 +546,19 @@ void FBodyInstance::CreateDOFLock()
 		return;
 	}
 
-	DOFConstraint = ConstructObject<UPhysicsConstraintComponent>(UPhysicsConstraintComponent::StaticClass());
+	DOFConstraint = FConstraintInstance::Alloc();
 	{
-		DOFConstraint->ConstraintInstance.bSwingLimitSoft = false;
-		DOFConstraint->ConstraintInstance.bTwistLimitSoft = false;
-		DOFConstraint->ConstraintInstance.bLinearLimitSoft = false;
+		DOFConstraint->bSwingLimitSoft = false;
+		DOFConstraint->bTwistLimitSoft = false;
+		DOFConstraint->bLinearLimitSoft = false;
 		//set all rotation to free
-		DOFConstraint->ConstraintInstance.AngularSwing1Motion = EAngularConstraintMotion::ACM_Locked;
-		DOFConstraint->ConstraintInstance.AngularSwing2Motion = EAngularConstraintMotion::ACM_Locked;
-		DOFConstraint->ConstraintInstance.AngularTwistMotion = EAngularConstraintMotion::ACM_Free;
+		DOFConstraint->AngularSwing1Motion = EAngularConstraintMotion::ACM_Locked;
+		DOFConstraint->AngularSwing2Motion = EAngularConstraintMotion::ACM_Locked;
+		DOFConstraint->AngularTwistMotion = EAngularConstraintMotion::ACM_Free;
 
-		DOFConstraint->ConstraintInstance.LinearXMotion = ELinearConstraintMotion::LCM_Locked;
-		DOFConstraint->ConstraintInstance.LinearYMotion = ELinearConstraintMotion::LCM_Free;
-		DOFConstraint->ConstraintInstance.LinearZMotion = ELinearConstraintMotion::LCM_Free;
+		DOFConstraint->LinearXMotion = ELinearConstraintMotion::LCM_Locked;
+		DOFConstraint->LinearYMotion = ELinearConstraintMotion::LCM_Free;
+		DOFConstraint->LinearZMotion = ELinearConstraintMotion::LCM_Free;
 
 		FVector Normal = LockedAxis.SafeNormal();
 		FVector Sec;
@@ -564,15 +568,15 @@ void FBodyInstance::CreateDOFLock()
 
 		FTransform TM = GetUnrealWorldTransform();
 
-		DOFConstraint->ConstraintInstance.PriAxis1 = TM.InverseTransformVectorNoScale(Normal);
-		DOFConstraint->ConstraintInstance.SecAxis1 = TM.InverseTransformVectorNoScale(Sec);
+		DOFConstraint->PriAxis1 = TM.InverseTransformVectorNoScale(Normal);
+		DOFConstraint->SecAxis1 = TM.InverseTransformVectorNoScale(Sec);
 
-		DOFConstraint->ConstraintInstance.PriAxis2 = Normal;
-		DOFConstraint->ConstraintInstance.SecAxis2 = Sec;
-		DOFConstraint->ConstraintInstance.Pos2 = TM.GetLocation();
+		DOFConstraint->PriAxis2 = Normal;
+		DOFConstraint->SecAxis2 = Sec;
+		DOFConstraint->Pos2 = TM.GetLocation();
 
 		// Create constraint instance based on DOF
-		DOFConstraint->ConstraintInstance.InitConstraint(OwnerComponent.Get(), this, nullptr, 1.f);
+		DOFConstraint->InitConstraint(OwnerComponent.Get(), this, nullptr, 1.f);
 	}
 }
 
@@ -1405,7 +1409,14 @@ void FBodyInstance::TermBody()
 
 	BodySetup = NULL;
 	OwnerComponent = NULL;
-	DOFConstraint = NULL;
+
+	if (DOFConstraint)
+	{
+		DOFConstraint->TermConstraint();
+		FConstraintInstance::Free(DOFConstraint);
+		DOFConstraint = NULL;
+	}
+	
 
 	
 }
diff --git a/Engine/Source/Runtime/Engine/Private/PhysicsEngine/ConstraintInstance.cpp b/Engine/Source/Runtime/Engine/Private/PhysicsEngine/ConstraintInstance.cpp
--- a/Engine/Source/Runtime/Engine/Private/PhysicsEngine/ConstraintInstance.cpp
+++ b/Engine/Source/Runtime/Engine/Private/PhysicsEngine/ConstraintInstance.cpp
@@ -1220,4 +1220,17 @@ void FConstraintInstance::SetAngularTwistLimit(EAngularConstraintMotion TwistMot
 	UpdateAngularLimit();
 }
 
+//Hacks to easily get zeroed memory for special case when we don't use GC
+void FConstraintInstance::Free(FConstraintInstance * Ptr)
+{
+	Ptr->~FConstraintInstance();
+	FMemory::Free(Ptr);
+}
+FConstraintInstance * FConstraintInstance::Alloc()
+{
+	void * Memory = FMemory::Malloc(sizeof(FConstraintInstance));
+	FMemory::Memzero(Memory, sizeof(FConstraintInstance));
+	return new (Memory)FConstraintInstance();
+}
+
 #undef LOCTEXT_NAMESPACE
\ No newline at end of file
