Fix CIS (missing forward declaration).

diff --git a/Engine/Source/Runtime/Engine/Classes/Landscape/LandscapeProxy.h b/Engine/Source/Runtime/Engine/Classes/Landscape/LandscapeProxy.h
--- a/Engine/Source/Runtime/Engine/Classes/Landscape/LandscapeProxy.h
+++ b/Engine/Source/Runtime/Engine/Classes/Landscape/LandscapeProxy.h
@@ -14,6 +14,7 @@
 class ULandscapeMaterialInstanceConstant;
 class ULandscapeLayerInfoObject;
 class ULandscapeSplinesComponent;
+class ULandscapeHeightfieldCollisionComponent;
 
 /** Structure storing channel usage for weightmap textures */
 USTRUCT()
