blueprint overlap queries also check async scene

diff --git a/Engine/Source/Runtime/Engine/Private/KismetSystemLibrary.cpp b/Engine/Source/Runtime/Engine/Private/KismetSystemLibrary.cpp
--- a/Engine/Source/Runtime/Engine/Private/KismetSystemLibrary.cpp
+++ b/Engine/Source/Runtime/Engine/Private/KismetSystemLibrary.cpp
@@ -594,6 +594,7 @@ bool UKismetSystemLibrary::SphereOverlapComponents_DEPRECATED(UObject* WorldCont
 
 	static FName SphereOverlapComponentsName(TEXT("SphereOverlapComponents"));
 	FCollisionQueryParams Params(SphereOverlapComponentsName, false);
+	Params.bTraceAsyncScene = true;
 	Params.AddIgnoredActors(ActorsToIgnore);
 	TArray<FOverlapResult> Overlaps;
 
@@ -639,6 +640,7 @@ bool UKismetSystemLibrary::BoxOverlapComponents_DEPRECATED(UObject* WorldContext
 
 	static FName BoxOverlapComponentsName(TEXT("BoxOverlapComponents"));
 	FCollisionQueryParams Params(BoxOverlapComponentsName, false);
+	Params.bTraceAsyncScene = true;
 	Params.AddIgnoredActors(ActorsToIgnore);
 
 	TArray<FOverlapResult> Overlaps;
@@ -684,6 +686,7 @@ bool UKismetSystemLibrary::CapsuleOverlapComponents_DEPRECATED(UObject* WorldCon
 
 	static FName CapsuleOverlapComponentsName(TEXT("CapsuleOverlapComponents"));
 	FCollisionQueryParams Params(CapsuleOverlapComponentsName, false);
+	Params.bTraceAsyncScene = true;
 	Params.AddIgnoredActors(ActorsToIgnore);
 
 	TArray<FOverlapResult> Overlaps;
@@ -774,6 +777,7 @@ bool UKismetSystemLibrary::SphereOverlapComponents_NEW(UObject* WorldContextObje
 	static FName SphereOverlapComponentsName(TEXT("SphereOverlapComponents"));
 	FCollisionQueryParams Params(SphereOverlapComponentsName, false);
 	Params.AddIgnoredActors(ActorsToIgnore);
+	Params.bTraceAsyncScene = true;
 	TArray<FOverlapResult> Overlaps;
 
 	FCollisionObjectQueryParams ObjectParams;
@@ -824,6 +828,7 @@ bool UKismetSystemLibrary::BoxOverlapComponents_NEW(UObject* WorldContextObject,
 
 	static FName BoxOverlapComponentsName(TEXT("BoxOverlapComponents"));
 	FCollisionQueryParams Params(BoxOverlapComponentsName, false);
+	Params.bTraceAsyncScene = true;
 	Params.AddIgnoredActors(ActorsToIgnore);
 
 	TArray<FOverlapResult> Overlaps;
@@ -874,6 +879,7 @@ bool UKismetSystemLibrary::CapsuleOverlapComponents_NEW(UObject* WorldContextObj
 
 	static FName CapsuleOverlapComponentsName(TEXT("CapsuleOverlapComponents"));
 	FCollisionQueryParams Params(CapsuleOverlapComponentsName, false);
+	Params.bTraceAsyncScene = true;
 	Params.AddIgnoredActors(ActorsToIgnore);
 
 	TArray<FOverlapResult> Overlaps;
