#ue4 - Pawns without a PawnMovementComponent can still use AddMovementInput() and read the results.

- Moved ControlInputVector from PawnMovementComponent to Pawn. There are "internal" methods on Pawn for use by the component to access/modify it.
- Added movement input accessors to Pawn: GetMovementInputVector(), ConsumeMovementInputVector(), IsMoveInputIgnored().
- Exposed PawnMovementComponent::IsMoveInputIgnored() to Blueprints.

diff --git a/Engine/Source/Runtime/Engine/Classes/GameFramework/Character.h b/Engine/Source/Runtime/Engine/Classes/GameFramework/Character.h
--- a/Engine/Source/Runtime/Engine/Classes/GameFramework/Character.h
+++ b/Engine/Source/Runtime/Engine/Classes/GameFramework/Character.h
@@ -208,14 +208,14 @@ class ENGINE_API ACharacter : public APawn
 	UPROPERTY(Category=Character, VisibleAnywhere, BlueprintReadOnly)
 	TSubobjectPtr<class UCharacterMovementComponent> CharacterMovement;
 
-	/** The CapsuleComponent being used for movement collision (if has CharacterMovement). Always treated as being vertically aligned in simple collision check functions. */
+	/** The CapsuleComponent being used for movement collision (by CharacterMovement). Always treated as being vertically aligned in simple collision check functions. */
 	UPROPERTY(Category=Character, VisibleAnywhere, BlueprintReadOnly)
 	TSubobjectPtr<class UCapsuleComponent> CapsuleComponent;
 
 	/** Name of the MeshComponent. Use this name if you want to prevent creation of the component (with PCIP.DoNotCreateDefaultSubobject). */
 	static FName MeshComponentName;
 
-	/** Name of the CharacterMovement. Use this name if you want to use a different class (with PCIP.SetDefaultSubobjectClass). */
+	/** Name of the CharacterMovement component. Use this name if you want to use a different class (with PCIP.SetDefaultSubobjectClass). */
 	static FName CharacterMovementComponentName;
 
 	/** Name of the CapsuleComponent. */
diff --git a/Engine/Source/Runtime/Engine/Classes/GameFramework/Pawn.h b/Engine/Source/Runtime/Engine/Classes/GameFramework/Pawn.h
--- a/Engine/Source/Runtime/Engine/Classes/GameFramework/Pawn.h
+++ b/Engine/Source/Runtime/Engine/Classes/GameFramework/Pawn.h
@@ -284,26 +284,57 @@ protected:
 
 public:
 	/**
-	 * Add movement input to our PawnMovementComponent, along the given world direction vector scaled by 'ScaleValue'. If ScaleValue < 0, movement will be in the opposite direction.
+	 * Add movement input along the given world direction vector (usually normalized) scaled by 'ScaleValue'. If ScaleValue < 0, movement will be in the opposite direction.
+	 * Base Pawn classes won't automatically apply movement, it's up to the user to do so in a Tick. Subclasses such as Character and DefaultPawn automatically handle this input and move.
+	 * @see GetMovementInputVector(), ConsumeMovementInputVector()
 	 *
-	 * Note that input is accumulated on the movement component (accessed via GetInputVector()) during input processing,
-	 * and actual movement of the Pawn won't occur until the acceleration is applied.
+	 * @param WorldDirection:	Direction in world space to apply input
+	 * @param ScaleValue:		Scale to apply to input. This can be used for analog input, ie a value of 0.5 applies half the normal value.
+	 * @param bForce:			If true always add the input, ignoring the result of IsMoveInputIgnored().
 	 */
-	UFUNCTION(BlueprintCallable, Category="Pawn")
-	virtual void AddMovementInput(FVector WorldDirection, float ScaleValue = 1.0f);
+	UFUNCTION(BlueprintCallable, Category="Pawn|Input")
+	virtual void AddMovementInput(FVector WorldDirection, float ScaleValue = 1.0f, bool bForce = false);
+
+	/** Return the input vector in world space. Note that the input should be consumed with ConsumeMovementInputVector() at the end of an update, to prevent accumulation of control input between frames. */
+	UFUNCTION(BlueprintCallable, Category="Pawn|Input")
+	FVector GetMovementInputVector() const;
+
+	/** Returns the input vector and resets it to zero. Should be used during an update to prevent accumulation of control input between frames. */
+	UFUNCTION(BlueprintCallable, Category="Pawn|Input")
+	virtual FVector ConsumeMovementInputVector();
 
 	/** Add input (affecting Pitch) to the Controller's ControlRotation, if it is a local PlayerController. */
-	UFUNCTION(BlueprintCallable, Category="Pawn", meta=(Keywords="up down"))
+	UFUNCTION(BlueprintCallable, Category="Pawn|Input", meta=(Keywords="up down"))
 	virtual void AddControllerPitchInput(float Val);
 
 	/** Add input (affecting Yaw) to the Controller's ControlRotation, if it is a local PlayerController. */
-	UFUNCTION(BlueprintCallable, Category="Pawn", meta=(Keywords="left right turn"))
+	UFUNCTION(BlueprintCallable, Category="Pawn|Input", meta=(Keywords="left right turn"))
 	virtual void AddControllerYawInput(float Val);
 
 	/** Add input (affecting Roll) to the Controller's ControlRotation, if it is a local PlayerController. */
-	UFUNCTION(BlueprintCallable, Category="Pawn")
+	UFUNCTION(BlueprintCallable, Category="Pawn|Input")
 	virtual void AddControllerRollInput(float Val);
 
+	/** Helper to see if move input is ignored. If our controller is a PlayerController, checks Controller->IsMoveInputIgnored(). */
+	UFUNCTION(BlueprintCallable, Category="Pawn|Input")
+	virtual bool IsMoveInputIgnored() const;
+
+private:
+
+	/** Accumulated control input vector, stored in world space. */
+	UPROPERTY(Transient)
+	FVector ControlInputVector;
+
+public:
+	/** Internal function meant for use only within Pawn or by a PawnMovementComponent. Adds movement input if not ignored, or if forced. */
+	void Internal_AddMovementInput(FVector WorldAccel, bool bForce = false);
+
+	/** Internal function meant for use only within Pawn or by a PawnMovementComponent. Returns the value of ControlInputVector. */
+	inline FVector Internal_GetMovementInputVector() const { return ControlInputVector; }
+
+	/** Internal function meant for use only within Pawn or by a PawnMovementComponent. Returns the vale of ControlInputVector and sets it to zero. */
+	FVector Internal_ConsumeMovementInputVector();
+
 public:
 	/** (DEPRECATED) Launch Character with LaunchVelocity  */
 	UFUNCTION(BlueprintCallable, Category="Pawn", meta=(DeprecatedFunction, DeprecationMessage="Use Character.LaunchCharacter instead"))
@@ -316,11 +347,3 @@ public:
 	void MoveIgnoreActorRemove(AActor* ActorToIgnore);
 };
 
-
-//////////////////////////////////////////////////////////////////////////
-// Inlines
-
-inline class UPawnMovementComponent* APawn::GetMovementComponent() const
-{
-	return FindComponentByClass<UPawnMovementComponent>();
-}
diff --git a/Engine/Source/Runtime/Engine/Classes/GameFramework/PawnMovementComponent.h b/Engine/Source/Runtime/Engine/Classes/GameFramework/PawnMovementComponent.h
--- a/Engine/Source/Runtime/Engine/Classes/GameFramework/PawnMovementComponent.h
+++ b/Engine/Source/Runtime/Engine/Classes/GameFramework/PawnMovementComponent.h
@@ -16,9 +16,15 @@ class ENGINE_API UPawnMovementComponent : public UNavMovementComponent
 	// Overridden to only allow registration with components owned by a Pawn.
 	virtual void SetUpdatedComponent(class UPrimitiveComponent* NewUpdatedComponent) override;
 
-	/** Adds the given vector to the accumulated input in world space. Input vectors are usually between 0 and 1 in magnitude. 
+	/**
+	 * Adds the given vector to the accumulated input in world space. Input vectors are usually between 0 and 1 in magnitude. 
 	 * They are accumulated during a frame then applied as acceleration during the movement update.
-	 * if bForce is true, it bypasses the IsMoveInputIgnored() check. */
+	 * @see Pawn.AddMovementInput()
+	 *
+	 * @param WorldDirection:	Direction in world space to apply input
+	 * @param ScaleValue:		Scale to apply to input. This can be used for analog input, ie a value of 0.5 applies half the normal value.
+	 * @param bForce:			If true always add the input, ignoring the result of IsMoveInputIgnored().
+	 */
 	UFUNCTION(BlueprintCallable, Category="Pawn|Components|PawnMovement")
 	virtual void AddInputVector(FVector WorldVector, bool bForce = false);
 
@@ -30,7 +36,8 @@ class ENGINE_API UPawnMovementComponent : public UNavMovementComponent
 	UFUNCTION(BlueprintCallable, Category="Pawn|Components|PawnMovement")
 	virtual FVector ConsumeInputVector();
 
-	/** Helper to see if move input is ignored. If possessed by a player-controlled Pawn, defers to the PlayerController's input ignore flags. */
+	/** Helper to see if move input is ignored. If there is no Pawn or UpdatedComponent, returns true, otherwise defers to the Pawn's implementation of IsMoveInputIgnored(). */
+	UFUNCTION(BlueprintCallable, Category="Pawn|Components|PawnMovement")
 	virtual bool IsMoveInputIgnored() const;
 
 	/** Return the Pawn that owns UpdatedComponent. */
@@ -42,25 +49,8 @@ public:
 	virtual void NotifyBumpedPawn(APawn* BumpedPawn) {}
 
 protected:
-	/** Accumulated control input vector, stored in world space.  */
-	UPROPERTY(Transient)
-	FVector ControlInputVector;
 
 	/** Pawn which owns this component. */
 	UPROPERTY()
 	class APawn* PawnOwner;
 };
-
-
-//////////////////////////////////////////////////////////////////////////
-// Inlines
-
-inline FVector UPawnMovementComponent::GetInputVector() const
-{
-	return ControlInputVector;
-}
-
-inline class APawn* UPawnMovementComponent::GetPawnOwner() const
-{
-	return PawnOwner;
-}
diff --git a/Engine/Source/Runtime/Engine/Private/Pawn.cpp b/Engine/Source/Runtime/Engine/Private/Pawn.cpp
--- a/Engine/Source/Runtime/Engine/Private/Pawn.cpp
+++ b/Engine/Source/Runtime/Engine/Private/Pawn.cpp
@@ -112,6 +112,11 @@ void APawn::PostRegisterAllComponents()
 	UpdateNavAgent();
 }
 
+UPawnMovementComponent* APawn::GetMovementComponent() const
+{
+	return FindComponentByClass<UPawnMovementComponent>();
+}
+
 void APawn::UpdateNavAgent()
 {
 	UPawnMovementComponent* MovementComponent = GetMovementComponent();
@@ -527,15 +532,70 @@ void APawn::DestroyPlayerInputComponent()
 }
 
 
-void APawn::AddMovementInput(FVector WorldDirection, float ScaleValue)
+bool APawn::IsMoveInputIgnored() const
+{
+	const APlayerController* PCOwner = Cast<const APlayerController>(Controller);
+	if (PCOwner && PCOwner->IsMoveInputIgnored())
+	{
+		return true;
+	}
+
+	// No player controller or not ignoring input. We allow non-PCs to receive input, so AI can use the movement input system.
+	return false;
+}
+
+
+void APawn::AddMovementInput(FVector WorldDirection, float ScaleValue, bool bForce /*=false*/)
 {
 	UPawnMovementComponent* MovementComponent = GetMovementComponent();
 	if (MovementComponent)
 	{
-		MovementComponent->AddInputVector(WorldDirection * ScaleValue);
+		MovementComponent->AddInputVector(WorldDirection * ScaleValue, bForce);
 	}
+	else
+	{
+		Internal_AddMovementInput(WorldDirection * ScaleValue, bForce);
+	}
+}
+
+
+FVector APawn::GetMovementInputVector() const
+{
+	// There's really no point redirecting to the MovementComponent since GetInputVector is not virtual there, and it just comes back to us.
+	return ControlInputVector;
 }
 
+
+FVector APawn::ConsumeMovementInputVector()
+{
+	UPawnMovementComponent* MovementComponent = GetMovementComponent();
+	if (MovementComponent)
+	{
+		return MovementComponent->ConsumeInputVector();
+	}
+	else
+	{
+		return Internal_ConsumeMovementInputVector();
+	}
+}
+
+
+void APawn::Internal_AddMovementInput(FVector WorldAccel, bool bForce /*=false*/)
+{
+	if (bForce || !IsMoveInputIgnored())
+	{
+		ControlInputVector += WorldAccel;
+	}
+}
+
+FVector APawn::Internal_ConsumeMovementInputVector()
+{
+	const FVector OldValue = ControlInputVector;
+	ControlInputVector = FVector::ZeroVector;
+	return OldValue;
+}
+
+
 void APawn::AddControllerPitchInput(float Val)
 {
 	if (Controller && Controller->IsLocalPlayerController())
diff --git a/Engine/Source/Runtime/Engine/Private/PawnMovementComponent.cpp b/Engine/Source/Runtime/Engine/Private/PawnMovementComponent.cpp
--- a/Engine/Source/Runtime/Engine/Private/PawnMovementComponent.cpp
+++ b/Engine/Source/Runtime/Engine/Private/PawnMovementComponent.cpp
@@ -28,6 +28,10 @@ void UPawnMovementComponent::SetUpdatedComponent(UPrimitiveComponent* NewUpdated
 	PawnOwner = NewUpdatedComponent ? CastChecked<APawn>(NewUpdatedComponent->GetOwner()) : NULL;
 }
 
+APawn* UPawnMovementComponent::GetPawnOwner() const
+{
+	return PawnOwner;
+}
 
 bool UPawnMovementComponent::IsMoveInputIgnored() const
 {
@@ -35,16 +39,7 @@ bool UPawnMovementComponent::IsMoveInputIgnored() const
 	{
 		if (PawnOwner)
 		{
-			const APlayerController* PCOwner = Cast<const APlayerController>(PawnOwner->Controller);
-			if (PCOwner)
-			{
-				return PCOwner->IsMoveInputIgnored();
-			}
-			else
-			{
-				// Allow movement by non-player controllers, if for some reason AI wants to go through simulated input.
-				return false;
-			}
+			return PawnOwner->IsMoveInputIgnored();
 		}
 	}
 
@@ -52,18 +47,20 @@ bool UPawnMovementComponent::IsMoveInputIgnored() const
 	return true;
 }
 
-
 void UPawnMovementComponent::AddInputVector(FVector WorldAccel, bool bForce /*=false*/)
 {
-	if (bForce || !IsMoveInputIgnored())
+	if (PawnOwner)
 	{
-		ControlInputVector += WorldAccel;
+		PawnOwner->Internal_AddMovementInput(WorldAccel, bForce);
 	}
 }
 
+FVector UPawnMovementComponent::GetInputVector() const
+{
+	return PawnOwner ? PawnOwner->Internal_GetMovementInputVector() : FVector::ZeroVector;
+}
+
 FVector UPawnMovementComponent::ConsumeInputVector()
 {
-	const FVector OldValue = GetInputVector();
-	ControlInputVector = FVector::ZeroVector;
-	return OldValue;
+	return PawnOwner ? PawnOwner->Internal_ConsumeMovementInputVector() : FVector::ZeroVector;
 }
