Fix spline mesh UVs - UV channels 1-3 now work, and UV channels 4+ don't crash!

diff --git a/Engine/Source/Runtime/Engine/Private/SplineMeshComponent.cpp b/Engine/Source/Runtime/Engine/Private/SplineMeshComponent.cpp
--- a/Engine/Source/Runtime/Engine/Private/SplineMeshComponent.cpp
+++ b/Engine/Source/Runtime/Engine/Private/SplineMeshComponent.cpp
@@ -132,20 +132,32 @@ void FSplineMeshSceneProxy::InitResources(USplineMeshComponent* InComponent, int
 
 		if (!RenderData->VertexBuffer.GetUseFullPrecisionUVs())
 		{
-			for (uint32 UVIndex = 0; UVIndex < RenderData->VertexBuffer.GetNumTexCoords(); UVIndex++)
+			int32 UVIndex;
+			for (UVIndex = 0; UVIndex < (int32)RenderData->VertexBuffer.GetNumTexCoords() - 1; UVIndex += 2)
 			{
 				Data.TextureCoordinates.Add(FVertexStreamComponent(
 					&RenderData->VertexBuffer,
-					STRUCT_OFFSET(TStaticMeshFullVertexFloat16UVs<MAX_TEXCOORDS>, UVs) + sizeof(FVector2DHalf)* UVIndex,
+					STRUCT_OFFSET(TStaticMeshFullVertexFloat16UVs<MAX_STATIC_TEXCOORDS>, UVs) + sizeof(FVector2DHalf) * UVIndex,
+					RenderData->VertexBuffer.GetStride(),
+					VET_Half4
+					));
+			}
+			// possible last UV channel if we have an odd number
+			if (UVIndex < (int32)RenderData->VertexBuffer.GetNumTexCoords())
+			{
+				Data.TextureCoordinates.Add(FVertexStreamComponent(
+					&RenderData->VertexBuffer,
+					STRUCT_OFFSET(TStaticMeshFullVertexFloat16UVs<MAX_STATIC_TEXCOORDS>, UVs) + sizeof(FVector2DHalf) * UVIndex,
 					RenderData->VertexBuffer.GetStride(),
 					VET_Half2
 					));
 			}
+
 			if (Parent->LightMapCoordinateIndex >= 0 && (uint32)Parent->LightMapCoordinateIndex < RenderData->VertexBuffer.GetNumTexCoords())
 			{
 				Data.LightMapCoordinateComponent = FVertexStreamComponent(
 					&RenderData->VertexBuffer,
-					STRUCT_OFFSET(TStaticMeshFullVertexFloat16UVs<MAX_TEXCOORDS>, UVs) + sizeof(FVector2DHalf)* Parent->LightMapCoordinateIndex,
+					STRUCT_OFFSET(TStaticMeshFullVertexFloat16UVs<MAX_STATIC_TEXCOORDS>, UVs) + sizeof(FVector2DHalf) * Parent->LightMapCoordinateIndex,
 					RenderData->VertexBuffer.GetStride(),
 					VET_Half2
 					);
@@ -153,11 +165,22 @@ void FSplineMeshSceneProxy::InitResources(USplineMeshComponent* InComponent, int
 		}
 		else
 		{
-			for (uint32 UVIndex = 0; UVIndex < RenderData->VertexBuffer.GetNumTexCoords(); UVIndex++)
+			int32 UVIndex;
+			for (UVIndex = 0; UVIndex < (int32)RenderData->VertexBuffer.GetNumTexCoords() - 1; UVIndex += 2)
+			{
+				Data.TextureCoordinates.Add(FVertexStreamComponent(
+					&RenderData->VertexBuffer,
+					STRUCT_OFFSET(TStaticMeshFullVertexFloat32UVs<MAX_STATIC_TEXCOORDS>, UVs) + sizeof(FVector2D) * UVIndex,
+					RenderData->VertexBuffer.GetStride(),
+					VET_Float4
+					));
+			}
+			// possible last UV channel if we have an odd number
+			if (UVIndex < (int32)RenderData->VertexBuffer.GetNumTexCoords())
 			{
 				Data.TextureCoordinates.Add(FVertexStreamComponent(
 					&RenderData->VertexBuffer,
-					STRUCT_OFFSET(TStaticMeshFullVertexFloat32UVs<MAX_TEXCOORDS>, UVs) + sizeof(FVector2D)* UVIndex,
+					STRUCT_OFFSET(TStaticMeshFullVertexFloat32UVs<MAX_STATIC_TEXCOORDS>, UVs) + sizeof(FVector2D) * UVIndex,
 					RenderData->VertexBuffer.GetStride(),
 					VET_Float2
 					));
@@ -167,7 +190,7 @@ void FSplineMeshSceneProxy::InitResources(USplineMeshComponent* InComponent, int
 			{
 				Data.LightMapCoordinateComponent = FVertexStreamComponent(
 					&RenderData->VertexBuffer,
-					STRUCT_OFFSET(TStaticMeshFullVertexFloat32UVs<MAX_TEXCOORDS>, UVs) + sizeof(FVector2D)* Parent->LightMapCoordinateIndex,
+					STRUCT_OFFSET(TStaticMeshFullVertexFloat32UVs<MAX_STATIC_TEXCOORDS>, UVs) + sizeof(FVector2D) * Parent->LightMapCoordinateIndex,
 					RenderData->VertexBuffer.GetStride(),
 					VET_Float2
 					);
