Made SimpleMoveToLocation and SimpleMoveToActor early-out for NULL pawns or controllers #UE4

diff --git a/Engine/Source/Runtime/Engine/Private/AI/Navigation/NavigationSystem.cpp b/Engine/Source/Runtime/Engine/Private/AI/Navigation/NavigationSystem.cpp
--- a/Engine/Source/Runtime/Engine/Private/AI/Navigation/NavigationSystem.cpp
+++ b/Engine/Source/Runtime/Engine/Private/AI/Navigation/NavigationSystem.cpp
@@ -833,20 +833,17 @@ bool UNavigationSystem::ProjectPointToNavigation(const FVector& Point, FNavLocat
 
 void UNavigationSystem::SimpleMoveToActor(AController* Controller, const AActor* Goal)
 {
-	if (Goal == NULL)
+	if (Goal == NULL || Controller == NULL || Controller->GetPawn() == NULL)
 	{
-		UE_LOG(LogNavigation, Log, TEXT("UNavigationSystem::SimpleMoveToActor called for %s with NULL goal actor")
-			, *GetNameSafe(Controller));
+		UE_LOG(LogNavigation, Warning, TEXT("UNavigationSystem::SimpleMoveToActor called for Controller:%s controlling Pawn:%s with goal actor %s (if any of these is None then there's your problem")
+			, *GetNameSafe(Controller), Controller ? *GetNameSafe(Controller->GetPawn()) : TEXT("NULL"), *GetNameSafe(Goal));
 		return;
 	}
 
 	UNavigationComponent* PFindComp = NULL;
 	UPathFollowingComponent* PFollowComp = NULL;
 	
-	if (Controller)
-	{
-		Controller->InitNavigationControl(PFindComp, PFollowComp);
-	}
+	Controller->InitNavigationControl(PFindComp, PFollowComp);
 
 	if (PFindComp && PFollowComp && !PFollowComp->HasReached(Goal))
 	{
@@ -860,13 +857,17 @@ void UNavigationSystem::SimpleMoveToActor(AController* Controller, const AActor*
 
 void UNavigationSystem::SimpleMoveToLocation(AController* Controller, const FVector& Goal)
 {
+	if (Controller == NULL || Controller->GetPawn() == NULL)
+	{
+		UE_LOG(LogNavigation, Warning, TEXT("UNavigationSystem::SimpleMoveToLocation called for Controller:%s controlling Pawn:%s (if any of these is None then there's your problem")
+			, *GetNameSafe(Controller), Controller ? *GetNameSafe(Controller->GetPawn()) : TEXT("NULL"));
+		return;
+	}
+
 	UNavigationComponent* PFindComp = NULL;
 	UPathFollowingComponent* PFollowComp = NULL;
 
-	if (Controller)
-	{
-		Controller->InitNavigationControl(PFindComp, PFollowComp);
-	}
+	Controller->InitNavigationControl(PFindComp, PFollowComp);
 
 	if (PFindComp && PFollowComp && !PFollowComp->HasReached(Goal))
 	{
