BodyInstance DOFLock deletion not properly nulling out

diff --git a/Engine/Source/Runtime/Engine/Private/PhysicsEngine/BodyInstance.cpp b/Engine/Source/Runtime/Engine/Private/PhysicsEngine/BodyInstance.cpp
--- a/Engine/Source/Runtime/Engine/Private/PhysicsEngine/BodyInstance.cpp
+++ b/Engine/Source/Runtime/Engine/Private/PhysicsEngine/BodyInstance.cpp
@@ -536,6 +536,7 @@ void FBodyInstance::CreateDOFLock()
 	{
 		DOFConstraint->TermConstraint();
 		FConstraintInstance::Free(DOFConstraint);
+		DOFConstraint = NULL;
 	}
 
 	//setup constraint based on DOF
