Deprecate GetControlledPawn to avoid confusion with GetPawn().  Added K2_GetPawn as more clear name that backs the blueprint functionality.

diff --git a/Engine/Config/BaseEngine.ini b/Engine/Config/BaseEngine.ini
--- a/Engine/Config/BaseEngine.ini
+++ b/Engine/Config/BaseEngine.ini
@@ -801,6 +801,7 @@ bUseStreamingPause=false
 +K2FieldRedirects=(OldFieldName="PawnMovementComponent.StopActiveMovement",NewFieldName="NavMovementComponent.StopActiveMovement")
 
 +K2FieldRedirects=(OldFieldName="Controller.K2_GetActorRotation",NewFieldName="Controller.GetControlRotation")
++K2FieldRedirects=(OldFieldName="Controller.GetControlledPawn",NewFieldName="Controller.K2_GetPawn")
 +K2FieldRedirects=(OldFieldName="PlayerController.AddTurnInput",NewFieldName="PlayerController.AddYawInput")
 +K2FieldRedirects=(OldFieldName="PlayerController.AddLookUpInput",NewFieldName="PlayerController.AddPitchInput")
 +K2FieldRedirects=(OldFieldName="Pawn.AddTurnInput",NewFieldName="Pawn.AddControllerYawInput")
diff --git a/Engine/Source/Runtime/AIModule/Private/Perception/PawnSensingComponent.cpp b/Engine/Source/Runtime/AIModule/Private/Perception/PawnSensingComponent.cpp
--- a/Engine/Source/Runtime/AIModule/Private/Perception/PawnSensingComponent.cpp
+++ b/Engine/Source/Runtime/AIModule/Private/Perception/PawnSensingComponent.cpp
@@ -142,7 +142,7 @@ AActor* UPawnSensingComponent::GetSensorActor() const
 	if (IsValid(Controller))
 	{
 		// In case the owner is a controller, use the controlled pawn as the sensing location.
-		SensorActor = Controller->GetControlledPawn();
+		SensorActor = Controller->GetPawn();
 	}
 
 	if (!IsValid(SensorActor))
diff --git a/Engine/Source/Runtime/Engine/Classes/GameFramework/Controller.h b/Engine/Source/Runtime/Engine/Classes/GameFramework/Controller.h
--- a/Engine/Source/Runtime/Engine/Classes/GameFramework/Controller.h
+++ b/Engine/Source/Runtime/Engine/Classes/GameFramework/Controller.h
@@ -161,8 +161,13 @@ public:
 	void ClientSetRotation(FRotator NewRotation, bool bResetCamera = false);
 
 	/** Return the Pawn that is currently 'controlled' by this PlayerController */
-	UFUNCTION(BlueprintCallable, Category="Pawn")
-	class APawn* GetControlledPawn() const;
+	UFUNCTION(BlueprintCallable, Category="Pawn", meta=(FriendlyName="Get Controlled Pawn"))
+	class APawn* K2_GetPawn() const;
+
+	DEPRECATED(4.4, "Use GetPawn() instead of GetControlledPawn()")
+	class APawn* GetControlledPawn() const { return GetPawn(); }
+
+public:
 
 	// Begin AActor interface
 	virtual void TickActor( float DeltaTime, enum ELevelTick TickType, FActorTickFunction& ThisTickFunction ) override;
diff --git a/Engine/Source/Runtime/Engine/Private/Controller.cpp b/Engine/Source/Runtime/Engine/Private/Controller.cpp
--- a/Engine/Source/Runtime/Engine/Private/Controller.cpp
+++ b/Engine/Source/Runtime/Engine/Private/Controller.cpp
@@ -521,9 +521,9 @@ APlayerController* AController::CastToPlayerController()
 	return Cast<APlayerController>(this);
 }
 
-APawn* AController::GetControlledPawn() const
+APawn* AController::K2_GetPawn() const
 {
-	return Pawn;
+	return GetPawn();
 }
 
 const FNavAgentProperties* AController::GetNavAgentProperties() const
diff --git a/Engine/Source/Runtime/Engine/Private/PlayerCameraManager.cpp b/Engine/Source/Runtime/Engine/Private/PlayerCameraManager.cpp
--- a/Engine/Source/Runtime/Engine/Private/PlayerCameraManager.cpp
+++ b/Engine/Source/Runtime/Engine/Private/PlayerCameraManager.cpp
@@ -1120,7 +1120,7 @@ APawn* FTViewTarget::GetTargetPawn() const
 	}
 	else if (AController* Controller = Cast<AController>(Target))
 	{
-		return Controller->GetControlledPawn();
+		return Controller->GetPawn();
 	}
 	else
 	{
diff --git a/Engine/Source/Runtime/Engine/Private/PlayerController.cpp b/Engine/Source/Runtime/Engine/Private/PlayerController.cpp
--- a/Engine/Source/Runtime/Engine/Private/PlayerController.cpp
+++ b/Engine/Source/Runtime/Engine/Private/PlayerController.cpp
@@ -663,7 +663,7 @@ void APlayerController::Possess(APawn* PawnToPossess)
 		check(GetPawn() != NULL);
 
 		GetPawn()->SetActorTickEnabled(true);
-		GetControlledPawn()->Restart();
+		GetPawn()->Restart();
 
 		INetworkPredictionInterface* NetworkPredictionInterface = GetPawn() ? InterfaceCast<INetworkPredictionInterface>(GetPawn()->GetMovementComponent()) : NULL;
 		if (NetworkPredictionInterface)
