Negative time on Android platform is a warning instead of fatal log message
#Android
#ttp 341593

diff --git a/Engine/Source/Runtime/Engine/Private/UnrealEngine.cpp b/Engine/Source/Runtime/Engine/Private/UnrealEngine.cpp
--- a/Engine/Source/Runtime/Engine/Private/UnrealEngine.cpp
+++ b/Engine/Source/Runtime/Engine/Private/UnrealEngine.cpp
@@ -974,8 +974,12 @@ void UEngine::UpdateTimeAndHandleMaxTickRate()
 		// Negative delta time means something is wrong with the system. Error out so user can address issue.
 		if( DeltaTime < 0 )
 		{
+#if PLATFORM_ANDROID
+			UE_LOG(LogEngine, Warning, TEXT("Detected negative delta time - ignoring"));
+#else
 			// AMD dual-core systems are a known issue that require AMD CPU drivers to be installed. Installer will take care of this for shipping.
 			UE_LOG(LogEngine, Fatal,TEXT("Detected negative delta time - on AMD systems please install http://files.aoaforums.com/I3199-setup.zip.html"));
+#endif
 			DeltaTime = 0.01;
 		}
 
