Add initialization values to FBodyInstance for case when its memory is not zeroed out

diff --git a/Engine/Source/Runtime/Engine/Private/PhysicsEngine/BodyInstance.cpp b/Engine/Source/Runtime/Engine/Private/PhysicsEngine/BodyInstance.cpp
--- a/Engine/Source/Runtime/Engine/Private/PhysicsEngine/BodyInstance.cpp
+++ b/Engine/Source/Runtime/Engine/Private/PhysicsEngine/BodyInstance.cpp
@@ -201,6 +201,9 @@ FBodyInstance::FBodyInstance()
 , bEnableGravity(true)
 , bUseAsyncScene(false)
 , bUpdateMassWhenScaleChanges(false)
+, LockedAxisMode(0)
+, CustomLockedAxis(FVector::ZeroVector)
+, DOFConstraint(NULL)
 , bOverrideWalkableSlopeOnInstance(false)
 , PhysMaterialOverride(NULL)
 , COMNudge(ForceInit)
