#ue4 - Pawn.ControlInputVector should be protected, not private.

diff --git a/Engine/Source/Runtime/Engine/Classes/GameFramework/Pawn.h b/Engine/Source/Runtime/Engine/Classes/GameFramework/Pawn.h
--- a/Engine/Source/Runtime/Engine/Classes/GameFramework/Pawn.h
+++ b/Engine/Source/Runtime/Engine/Classes/GameFramework/Pawn.h
@@ -319,7 +319,7 @@ public:
 	UFUNCTION(BlueprintCallable, Category="Pawn|Input")
 	virtual bool IsMoveInputIgnored() const;
 
-private:
+protected:
 
 	/** Accumulated control input vector, stored in world space. */
 	UPROPERTY(Transient)
