Remove some unused local variables

diff --git a/Engine/Source/Runtime/Engine/Private/LightComponent.cpp b/Engine/Source/Runtime/Engine/Private/LightComponent.cpp
--- a/Engine/Source/Runtime/Engine/Private/LightComponent.cpp
+++ b/Engine/Source/Runtime/Engine/Private/LightComponent.cpp
@@ -488,7 +488,6 @@ void ULightComponent::PostEditChangeProperty(FPropertyChangedEvent& PropertyChan
 {
 	UProperty* PropertyThatChanged = PropertyChangedEvent.MemberProperty;
 	const FString PropertyName = PropertyThatChanged ? PropertyThatChanged->GetName() : TEXT("");
-	const FName PropertyCategory = FObjectEditorUtils::GetCategoryFName(PropertyThatChanged);
 
 	Intensity = FMath::Max(0.0f, Intensity);
 
diff --git a/Engine/Source/Runtime/Engine/Private/SkyLightComponent.cpp b/Engine/Source/Runtime/Engine/Private/SkyLightComponent.cpp
--- a/Engine/Source/Runtime/Engine/Private/SkyLightComponent.cpp
+++ b/Engine/Source/Runtime/Engine/Private/SkyLightComponent.cpp
@@ -223,8 +223,6 @@ void USkyLightComponent::DestroyRenderState_Concurrent()
 #if WITH_EDITOR
 void USkyLightComponent::PostEditChangeProperty(FPropertyChangedEvent& PropertyChangedEvent)
 {
-	const FName CategoryName = FObjectEditorUtils::GetCategoryFName(PropertyChangedEvent.Property);
-
 	Super::PostEditChangeProperty(PropertyChangedEvent);
 
 	SetCaptureIsDirty();
