Engine: Make UPrimitiveComponent::GetUsedTextures virtual

diff --git a/Engine/Source/Runtime/Engine/Classes/Components/PrimitiveComponent.h b/Engine/Source/Runtime/Engine/Classes/Components/PrimitiveComponent.h
--- a/Engine/Source/Runtime/Engine/Classes/Components/PrimitiveComponent.h
+++ b/Engine/Source/Runtime/Engine/Classes/Components/PrimitiveComponent.h
@@ -957,7 +957,7 @@ public:
 	 *
 	 * @param OutTextures	[out] The list of used textures.
 	 */
-	void GetUsedTextures(TArray<UTexture*>& OutTextures, EMaterialQualityLevel::Type QualityLevel);
+	virtual void GetUsedTextures(TArray<UTexture*>& OutTextures, EMaterialQualityLevel::Type QualityLevel);
 
 	/** Controls if we get a post physics tick or not. If set during ticking, will take effect next frame **/
 	void SetPostPhysicsComponentTickEnabled(bool bEnable);
