Engine: Improvements to UMeshComponent as a base class
- Fix an array bounds assert if trying to expand the materials array via UMeshComponent::SetMaterial
- Update UMeshComponent::PrestreamTextures and UMeshComponent::SetTextureForceResidentFlag to work off of GetUsedTextures and made them virtual

diff --git a/Engine/Source/Runtime/Engine/Classes/Components/MeshComponent.h b/Engine/Source/Runtime/Engine/Classes/Components/MeshComponent.h
--- a/Engine/Source/Runtime/Engine/Classes/Components/MeshComponent.h
+++ b/Engine/Source/Runtime/Engine/Classes/Components/MeshComponent.h
@@ -23,6 +23,7 @@ class ENGINE_API UMeshComponent : public UPrimitiveComponent
 	virtual int32 GetNumMaterials() const override;
 	virtual UMaterialInterface* GetMaterial(int32 ElementIndex) const override;
 	virtual void SetMaterial(int32 ElementIndex, UMaterialInterface* Material) override;
+	virtual void GetUsedMaterials(TArray<UMaterialInterface*>& OutMaterials) const override;
 	// End UPrimitiveComponent Interface
 
 	/** Accesses the scene relevance information for the materials applied to the mesh. */
@@ -32,7 +33,7 @@ class ENGINE_API UMeshComponent : public UPrimitiveComponent
 	 *	Tell the streaming system whether or not all mip levels of all textures used by this component should be loaded and remain loaded.
 	 *	@param bForceMiplevelsToBeResident		Whether textures should be forced to be resident or not.
 	 */
-	void SetTextureForceResidentFlag( bool bForceMiplevelsToBeResident );
+	virtual void SetTextureForceResidentFlag( bool bForceMiplevelsToBeResident );
 
 	/**
 	 *	Tell the streaming system to start loading all textures with all mip-levels.
@@ -40,10 +41,5 @@ class ENGINE_API UMeshComponent : public UPrimitiveComponent
 	 *	@param bPrioritizeCharacterTextures		Whether character textures should be prioritized for a while by the streaming system
 	 *	@param CinematicTextureGroups			Bitfield indicating which texture groups that use extra high-resolution mips
 	 */
-	void PrestreamTextures( float Seconds, bool bPrioritizeCharacterTextures, int32 CinematicTextureGroups = 0 );
-
-
+	virtual void PrestreamTextures( float Seconds, bool bPrioritizeCharacterTextures, int32 CinematicTextureGroups = 0 );
 };
-
-
-
diff --git a/Engine/Source/Runtime/Engine/Private/MeshComponent.cpp b/Engine/Source/Runtime/Engine/Private/MeshComponent.cpp
--- a/Engine/Source/Runtime/Engine/Private/MeshComponent.cpp
+++ b/Engine/Source/Runtime/Engine/Private/MeshComponent.cpp
@@ -11,7 +11,6 @@ UMeshComponent::UMeshComponent(const class FPostConstructInitializeProperties& P
 	bCanEverAffectNavigation = true;
 }
 
-
 void UMeshComponent::BeginDestroy()
 {
 	// Notify the streaming system.
@@ -22,31 +21,40 @@ void UMeshComponent::BeginDestroy()
 
 UMaterialInterface* UMeshComponent::GetMaterial(int32 ElementIndex) const
 {
-	if(ElementIndex < Materials.Num() && Materials[ElementIndex])
+	if (Materials.IsValidIndex(ElementIndex))
 	{
 		return Materials[ElementIndex];
 	}
 	else
 	{
-		return NULL;
+		return nullptr;
 	}
 }
 
-void UMeshComponent::SetMaterial(int32 ElementIndex,UMaterialInterface* Material)
+void UMeshComponent::SetMaterial(int32 ElementIndex, UMaterialInterface* Material)
 {
-	if (ElementIndex >= 0 && (Materials.Num() <= ElementIndex || Materials[ElementIndex] != Material))
+	if (ElementIndex >= 0)
 	{
-		if (Materials.Num() <= ElementIndex)
+		if (Materials.IsValidIndex(ElementIndex) && (Materials[ElementIndex] == Material))
 		{
-			Materials.AddZeroed(ElementIndex + 1 - Materials.Num());
+			// Do nothing, the material is already set
 		}
-
-		Materials[ElementIndex] = Material;
-		MarkRenderStateDirty();
-
-		if(BodyInstance.IsValidBodyInstance())
+		else
 		{
-			BodyInstance.UpdatePhysicalMaterials();
+			// Grow the array if the new index is too large
+			if (Materials.Num() <= ElementIndex)
+			{
+				Materials.AddZeroed(ElementIndex + 1 - Materials.Num());
+			}
+
+			// Set the material and invalidate things
+			Materials[ElementIndex] = Material;
+			MarkRenderStateDirty();
+
+			if (BodyInstance.IsValidBodyInstance())
+			{
+				BodyInstance.UpdatePhysicalMaterials();
+			}
 		}
 	}
 }
@@ -72,37 +80,51 @@ int32 UMeshComponent::GetNumMaterials() const
 	return 0;
 }
 
+void UMeshComponent::GetUsedMaterials(TArray<UMaterialInterface*>& OutMaterials) const
+{
+	for (int32 ElementIndex = 0; ElementIndex < GetNumMaterials(); ElementIndex++)
+	{
+		if (UMaterialInterface* MaterialInterface = GetMaterial(ElementIndex))
+		{
+			OutMaterials.Add(MaterialInterface);
+		}
+	}
+}
+
 void UMeshComponent::PrestreamTextures( float Seconds, bool bPrioritizeCharacterTextures, int32 CinematicTextureGroups )
 {
 	// If requested, tell the streaming system to only process character textures for 30 frames.
-	if ( bPrioritizeCharacterTextures )
+	if (bPrioritizeCharacterTextures)
 	{
-		IStreamingManager::Get().SetDisregardWorldResourcesForFrames( 30 );
+		IStreamingManager::Get().SetDisregardWorldResourcesForFrames(30);
 	}
 
-	int32 NumElements = GetNumMaterials();
-	for ( int32 ElementIndex=0; ElementIndex < NumElements; ++ElementIndex )
+	TArray<UTexture*> Textures;
+	GetUsedTextures(/*out*/ Textures, EMaterialQualityLevel::Num);
+
+	for (UTexture* Texture : Textures)
 	{
-		UMaterialInterface* Material = GetMaterial( ElementIndex );
-		if ( Material )
+		if (UTexture2D* Texture2D = Cast<UTexture2D>(Texture))
 		{
-			Material->SetForceMipLevelsToBeResident( false, false, Seconds, CinematicTextureGroups );
+			Texture2D->SetForceMipLevelsToBeResident(Seconds, CinematicTextureGroups);
 		}
 	}
 }
 
-
 void UMeshComponent::SetTextureForceResidentFlag( bool bForceMiplevelsToBeResident )
 {
 	const int32 CinematicTextureGroups = 0;
+	const float Seconds = -1.0f;
+
+	TArray<UTexture*> Textures;
+	GetUsedTextures(/*out*/ Textures, EMaterialQualityLevel::Num);
 
-	int32 NumElements = GetNumMaterials();
-	for ( int32 ElementIndex=0; ElementIndex < NumElements; ++ElementIndex )
+	for (UTexture* Texture : Textures)
 	{
-		UMaterialInterface* Material = GetMaterial( ElementIndex );
-		if ( Material )
+		if (UTexture2D* Texture2D = Cast<UTexture2D>(Texture))
 		{
-			Material->SetForceMipLevelsToBeResident( true, bForceMiplevelsToBeResident, -1.0f, CinematicTextureGroups );
+			Texture2D->SetForceMipLevelsToBeResident(Seconds, CinematicTextureGroups);
+			Texture2D->bForceMiplevelsToBeResident = bForceMiplevelsToBeResident;
 		}
 	}
 }
