Engine: Update texture streaming code to treat all UMeshComponent derived classes equally, generalizing the special treatment for USkeletalMeshComponent or UStaticMeshComponent

diff --git a/Engine/Source/Runtime/Engine/Private/ContentStreaming.cpp b/Engine/Source/Runtime/Engine/Private/ContentStreaming.cpp
--- a/Engine/Source/Runtime/Engine/Private/ContentStreaming.cpp
+++ b/Engine/Source/Runtime/Engine/Private/ContentStreaming.cpp
@@ -1944,8 +1944,8 @@ void FStreamingManagerCollection::NotifyPrimitiveAttached( const UPrimitiveCompo
 {
 	STAT( double StartTime = FPlatformTime::Seconds() );
 
-	// Distance-based heuristics only handle skeletalmeshes and staticmeshes.
-	if ( Primitive->IsA( USkeletalMeshComponent::StaticClass() ) || Primitive->IsA( UStaticMeshComponent::StaticClass() ) )
+	// Distance-based heuristics only handle mesh components
+	if ( Primitive->IsA( UMeshComponent::StaticClass() ) )
 	{
 		for( int32 ManagerIndex=0; ManagerIndex<StreamingManagers.Num(); ManagerIndex++ )
 		{
@@ -1961,8 +1961,8 @@ void FStreamingManagerCollection::NotifyPrimitiveDetached( const UPrimitiveCompo
 {
 	STAT( double StartTime = FPlatformTime::Seconds() );
 
-	// Distance-based heuristics only handle skeletalmeshes and staticmeshes.
-	if ( Primitive->IsA( USkeletalMeshComponent::StaticClass() ) || Primitive->IsA( UStaticMeshComponent::StaticClass() ) )
+	// Distance-based heuristics only handle mesh components
+	if ( Primitive->IsA( UMeshComponent::StaticClass() ) )
 	{
 		// Route to streaming managers.
 		for( int32 ManagerIndex=0; ManagerIndex<StreamingManagers.Num(); ManagerIndex++ )
@@ -1983,7 +1983,7 @@ void FStreamingManagerCollection::NotifyPrimitiveUpdated( const UPrimitiveCompon
 {
 	STAT( double StartTime = FPlatformTime::Seconds() );
 
-	if ( Primitive->IsA( USkeletalMeshComponent::StaticClass() ) || Primitive->IsA( UStaticMeshComponent::StaticClass() ) )
+	if ( Primitive->IsA( UMeshComponent::StaticClass() ) )
 	{
 		// Route to streaming managers.
 		for( int32 ManagerIndex=0; ManagerIndex<StreamingManagers.Num(); ManagerIndex++ )
@@ -2598,26 +2598,25 @@ void FStreamingManagerTexture::BoostTextures( AActor* Actor, float BoostFactor )
 		Actor->GetComponents(Components);
 
 		for(int32 ComponentIndex = 0;ComponentIndex < Components.Num();ComponentIndex++)
-			{
-				// Only handle skeletalmeshes and staticmeshes.
+		{
 			UPrimitiveComponent* Primitive = Components[ComponentIndex];
-			if ( Primitive->IsRegistered() && ( Primitive->IsA(USkeletalMeshComponent::StaticClass()) || Primitive->IsA(UStaticMeshComponent::StaticClass()) ))
+			if ( Primitive->IsRegistered() && Primitive->IsA(UMeshComponent::StaticClass()) )
+			{
+				Textures.Reset();
+				Primitive->GetUsedTextures( Textures, EMaterialQualityLevel::Num );
+				for ( int32 TextureIndex=0; TextureIndex < Textures.Num(); ++TextureIndex )
 				{
-					Textures.Reset();
-					Primitive->GetUsedTextures( Textures, EMaterialQualityLevel::Num );
-					for ( int32 TextureIndex=0; TextureIndex < Textures.Num(); ++TextureIndex )
+					UTexture2D* Texture2D = Cast<UTexture2D>( Textures[ TextureIndex ] );
+					if ( Texture2D && IsManagedStreamingTexture(Texture2D) )
 					{
-						UTexture2D* Texture2D = Cast<UTexture2D>( Textures[ TextureIndex ] );
-						if ( Texture2D && IsManagedStreamingTexture(Texture2D) )
-						{
-							FStreamingTexture& StreamingTexture = GetStreamingTexture( Texture2D );
-							StreamingTexture.BoostFactor = FMath::Max( StreamingTexture.BoostFactor, BoostFactor );
-						}
+						FStreamingTexture& StreamingTexture = GetStreamingTexture( Texture2D );
+						StreamingTexture.BoostFactor = FMath::Max( StreamingTexture.BoostFactor, BoostFactor );
 					}
 				}
 			}
 		}
 	}
+}
 
 /**
  * Updates the thread-safe cache information for dynamic primitives.
@@ -2698,29 +2697,19 @@ bool FStreamingManagerTexture::AddDynamicPrimitive( const UPrimitiveComponent* P
 	}
 
 #if STREAMING_LOG_DYNAMIC
-	const USkeletalMeshComponent* SkeletalMeshComponent = Cast<const USkeletalMeshComponent>(Primitive);
-	const UStaticMeshComponent* StaticMeshComponent = Cast<const UStaticMeshComponent>(Primitive);
-	if ( SkeletalMeshComponent && SkeletalMeshComponent->SkeletalMesh )
-	{
-		UE_LOG(LogContentStreaming, Log, TEXT("AddDynamicPrimitive(0x%08x \"%s\", \"%s\"), IsRegistered=%d, NumTextures=%d"), SIZE_T(Primitive), *SkeletalMeshComponent->SkeletalMesh->GetName(), *SkeletalMeshComponent->GetFullName(), Primitive->IsRegistered(), NumTexturesAdded);
 #if 0
-//@DEBUG: To set a breakpoint on a specific skelmesh
-		if ( SkeletalMeshComponent->SkeletalMesh->GetName().Contains(TEXT("SpecificMeshName")) )
-		{
-			int32 Q=0;
-		}
-#endif
-	}
-	else if ( StaticMeshComponent && StaticMeshComponent->StaticMesh )
-	{
-		UE_LOG(LogContentStreaming, Log, TEXT("AddDynamicPrimitive(0x%08x \"%s\"), IsRegistered=%d, NumTextures=%d"), SIZE_T(Primitive), *StaticMeshComponent->StaticMesh->GetName(), Primitive->IsRegistered(), NumTexturesAdded);
-	}
-	else
+	//@DEBUG: To set a breakpoint on a specific asset (e.g., a particular skeletal or static mesh)
+	UObject* AssociatedAsset = Primitive->AdditionalStatObject();
+	if ((AssociatedAsset != nullptr) && AssociatedAsset->GetName().Contains(TEXT("SpecificMeshName")) )
 	{
-		UE_LOG(LogContentStreaming, Log, TEXT("AddDynamicPrimitive(0x%08x \"%s\"), IsRegistered=%d, NumTextures=%d"), SIZE_T(Primitive), *Primitive->GetName(), Primitive->IsRegistered(), NumTexturesAdded);
+		static int32 Q=0;
+		Q++;
 	}
 #endif
 
+	UE_LOG(LogContentStreaming, Log, TEXT("AddDynamicPrimitive(0x%08x \"%s\"), IsRegistered=%d, NumTextures=%d"), SIZE_T(Primitive), *Primitive->GetReadableName(), Primitive->IsRegistered(), NumTexturesAdded);
+#endif
+
 	if ( NumTexturesAdded > 0 )
 	{
 		return true;
@@ -2872,16 +2861,15 @@ void FStreamingManagerTexture::NotifyActorSpawned( AActor* Actor )
 		Actor->GetComponents(Components);
 
 		for(int32 ComponentIndex = 0;ComponentIndex < Components.Num();ComponentIndex++)
+		{
+			UPrimitiveComponent* Primitive = Components[ComponentIndex];
+			if (Primitive->IsRegistered() && Primitive->IsA(UMeshComponent::StaticClass()))
 			{
-				// Only handle skeletalmeshes and staticmeshes.
-				UPrimitiveComponent* Primitive = Components[ComponentIndex];
-				if ( Primitive->IsRegistered() && ( Primitive->IsA(USkeletalMeshComponent::StaticClass()) || Primitive->IsA(UStaticMeshComponent::StaticClass()) ))
-				{
-					NotifyPrimitiveAttached( Primitive, DPT_Spawned );
-				}
+				NotifyPrimitiveAttached( Primitive, DPT_Spawned );
 			}
 		}
 	}
+}
 
 /** Called when a spawned primitive is deleted. */
 void FStreamingManagerTexture::NotifyActorDestroyed( AActor* Actor )
@@ -2893,11 +2881,11 @@ void FStreamingManagerTexture::NotifyActorDestroyed( AActor* Actor )
 	{
 		UPrimitiveComponent* Primitive = Components[ComponentIndex];
 		if ( Primitive->IsRegistered() )
-			{
-				NotifyPrimitiveDetached( Primitive );
-			}
+		{
+			NotifyPrimitiveDetached( Primitive );
 		}
 	}
+}
 
 /**
  * Called when a primitive is attached to an actor or another component.
@@ -2911,20 +2899,7 @@ void FStreamingManagerTexture::NotifyPrimitiveAttached( const UPrimitiveComponen
 	if ( bUseDynamicStreaming && Primitive && Primitive->IsA(UMeshComponent::StaticClass()) )
 	{
 #if STREAMING_LOG_DYNAMIC
-		const USkeletalMeshComponent* SkeletalMeshComponent = Cast<const USkeletalMeshComponent>(Primitive);
-		const UStaticMeshComponent* StaticMeshComponent = Cast<const UStaticMeshComponent>(Primitive);
-		if ( SkeletalMeshComponent && SkeletalMeshComponent->SkeletalMesh )
-		{
-			UE_LOG(LogContentStreaming, Log, TEXT("NotifyPrimitiveAttached(0x%08x \"%s\"), IsRegistered=%d"), SIZE_T(Primitive), *SkeletalMeshComponent->SkeletalMesh->GetName(), Primitive->IsRegistered());
-		}
-		else if ( StaticMeshComponent && StaticMeshComponent->StaticMesh )
-		{
-			UE_LOG(LogContentStreaming, Log, TEXT("NotifyPrimitiveAttached(0x%08x \"%s\"), IsRegistered=%d"), SIZE_T(Primitive), *StaticMeshComponent->StaticMesh->GetName(), Primitive->IsRegistered());
-		}
-		else
-		{
-			UE_LOG(LogContentStreaming, Log, TEXT("NotifyPrimitiveAttached(0x%08x \"%s\"), IsRegistered=%d"), SIZE_T(Primitive), *Primitive->GetName(), Primitive->IsRegistered());
-		}
+		UE_LOG(LogContentStreaming, Log, TEXT("NotifyPrimitiveAttached(0x%08x \"%s\"), IsRegistered=%d"), SIZE_T(Primitive), *Primitive->GetReadableName(), Primitive->IsRegistered());
 #endif
 
 		// If we already have a pending update, keep its current dynamic type.
@@ -2957,20 +2932,7 @@ void FStreamingManagerTexture::NotifyPrimitiveDetached( const UPrimitiveComponen
 	if ( bUseDynamicStreaming && Primitive )
 	{
 #if STREAMING_LOG_DYNAMIC
-		const USkeletalMeshComponent* SkeletalMeshComponent = Cast<const USkeletalMeshComponent>(Primitive);
-		const UStaticMeshComponent* StaticMeshComponent = Cast<const UStaticMeshComponent>(Primitive);
-		if ( SkeletalMeshComponent && SkeletalMeshComponent->SkeletalMesh )
-		{
-			UE_LOG(LogContentStreaming, Log, TEXT("NotifyPrimitiveDetached(0x%08x \"%s\", \"%s\"), IsRegistered=%d"), SIZE_T(Primitive), *SkeletalMeshComponent->SkeletalMesh->GetName(), *SkeletalMeshComponent->GetFullName(), Primitive->IsRegistered());
-		}
-		else if ( StaticMeshComponent && StaticMeshComponent->StaticMesh )
-		{
-			UE_LOG(LogContentStreaming, Log, TEXT("NotifyPrimitiveDetached(0x%08x \"%s\"), IsRegistered=%d"), SIZE_T(Primitive), *StaticMeshComponent->StaticMesh->GetName(), Primitive->IsRegistered());
-		}
-		else
-		{
-			UE_LOG(LogContentStreaming, Log, TEXT("NotifyPrimitiveDetached(0x%08x \"%s\"), IsRegistered=%d"), SIZE_T(Primitive), *Primitive->GetName(), Primitive->IsRegistered());
-		}
+		UE_LOG(LogContentStreaming, Log, TEXT("NotifyPrimitiveDetached(0x%08x \"%s\"), IsRegistered=%d"), SIZE_T(Primitive), *Primitive->GetReadableName(), Primitive->IsRegistered());
 #endif
 
 		// Did we already mark this as detached?
@@ -3044,20 +3006,7 @@ void FStreamingManagerTexture::NotifyPrimitiveUpdated( const UPrimitiveComponent
 	if ( bIsCurrentlyTracked && !bIsBeingRemoved )
 	{
 #if STREAMING_LOG_DYNAMIC
-		const USkeletalMeshComponent* SkeletalMeshComponent = Cast<const USkeletalMeshComponent>(Primitive);
-		const UStaticMeshComponent* StaticMeshComponent = Cast<const UStaticMeshComponent>(Primitive);
-		if ( SkeletalMeshComponent && SkeletalMeshComponent->SkeletalMesh )
-		{
-			UE_LOG(LogContentStreaming, Log, TEXT("NotifyPrimitiveUpdated(0x%08x \"%s\"), IsRegistered=%d"), SIZE_T(Primitive), *SkeletalMeshComponent->SkeletalMesh->GetName(), Primitive->IsRegistered());
-		}
-		else if ( StaticMeshComponent && StaticMeshComponent->StaticMesh )
-		{
-			UE_LOG(LogContentStreaming, Log, TEXT("NotifyPrimitiveUpdated(0x%08x \"%s\"), IsRegistered=%d"), SIZE_T(Primitive), *StaticMeshComponent->StaticMesh->GetName(), Primitive->IsRegistered());
-		}
-		else
-		{
-			UE_LOG(LogContentStreaming, Log, TEXT("NotifyPrimitiveUpdated(0x%08x \"%s\"), IsRegistered=%d"), SIZE_T(Primitive), *Primitive->GetName(), Primitive->IsRegistered());
-		}
+		UE_LOG(LogContentStreaming, Log, TEXT("NotifyPrimitiveUpdated(0x%08x \"%s\"), IsRegistered=%d"), SIZE_T(Primitive), *Primitive->GetReadableName(), Primitive->IsRegistered());
 #endif
 
 		// Reattach the primitive, replacing previous info.
@@ -5161,17 +5110,6 @@ void FStreamingManagerTexture::DumpTextureInstances( const UPrimitiveComponent*
 			}
 			ScreenSizeFactor *= StreamingTexture.BoostFactor;
 
-			const UStaticMeshComponent* StaticMeshComponent = Cast<const UStaticMeshComponent>(Primitive);
-			const USkeletalMeshComponent* SkeletalMeshComponent = Cast<const USkeletalMeshComponent>(Primitive);
-			FString MeshName;
-			if ( StaticMeshComponent && StaticMeshComponent->StaticMesh )
-			{
-				MeshName = StaticMeshComponent->StaticMesh->GetName();
-			}
-			else if ( SkeletalMeshComponent && SkeletalMeshComponent->SkeletalMesh )
-			{
-				MeshName = SkeletalMeshComponent->SkeletalMesh->GetName();
-			}
 			float MinDistance = MAX_FLT;
 			FFloatMipLevel WantedMipCount;
 			for( int32 ViewIndex=0; ViewIndex < ThreadNumViews(); ViewIndex++ )
@@ -5202,12 +5140,12 @@ void FStreamingManagerTexture::DumpTextureInstances( const UPrimitiveComponent*
 			}
 			int32 IntWantedMipCount = WantedMipCount.ComputeMip(&StreamingTexture, ThreadSettings.MipBias, false);
 			int32 WantedMipIndex = FMath::Max(Texture2D->GetNumMips() - IntWantedMipCount, 0);
-			UE_LOG(LogContentStreaming, Log, TEXT("%s: Wanted=%dx%d, Distance=%.1f, TexelFactor=%.2f, CurrentRadius=%5.1f, OriginalRadius=%5.1f, Position=(%d,%d,%d), IsRegistered=%d, Mesh=\"%s\", Component=\"%s\""),
+			UE_LOG(LogContentStreaming, Log, TEXT("%s: Wanted=%dx%d, Distance=%.1f, TexelFactor=%.2f, CurrentRadius=%5.1f, OriginalRadius=%5.1f, Position=(%d,%d,%d), IsRegistered=%d, Component=\"%s\""),
 				TypeString, Texture2D->PlatformData->Mips[WantedMipIndex].SizeX, Texture2D->PlatformData->Mips[WantedMipIndex].SizeY,
 				MinDistance, Instance.TexelFactor, PrimitiveData.BoundingSphere.W, 1.0f / Instance.InvOriginalRadius,
 				int32(PrimitiveData.BoundingSphere.Center.X), int32(PrimitiveData.BoundingSphere.Center.Y), int32(PrimitiveData.BoundingSphere.Center.Z),
-				PrimitiveData.bAttached ? true : false, *MeshName,
-				*Primitive->GetFullName() );
+				PrimitiveData.bAttached ? true : false,
+				*Primitive->GetReadableName() );
 		}
 	}
 #endif
