Cleanup SetClassVolume

diff --git a/Engine/Source/Runtime/Engine/Private/AudioDevice.cpp b/Engine/Source/Runtime/Engine/Private/AudioDevice.cpp
--- a/Engine/Source/Runtime/Engine/Private/AudioDevice.cpp
+++ b/Engine/Source/Runtime/Engine/Private/AudioDevice.cpp
@@ -2107,11 +2107,10 @@ void FAudioDevice::RemoveActiveSound(FActiveSound* ActiveSound)
 
 void FAudioDevice::SetClassVolume( USoundClass* InSoundClass, const float Volume )
 {
-	check(InSoundClass);
-	
-	UE_LOG(LogAudio, Warning, TEXT( "ModifySoundClass Class: %s NewVolume: %f" ), *InSoundClass->GetPathName(), Volume );
-
-	InSoundClass->Properties.Volume = Volume;
+	if (InSoundClass)
+	{
+		InSoundClass->Properties.Volume = Volume;
+	}
 }
 
 bool FAudioDevice::LocationIsAudible( FVector Location, float MaxDistance )
