Added exports to various texture classes to allow subclassing in different modules

diff --git a/Engine/Source/Runtime/Engine/Classes/Engine/Texture.h b/Engine/Source/Runtime/Engine/Classes/Engine/Texture.h
--- a/Engine/Source/Runtime/Engine/Classes/Engine/Texture.h
+++ b/Engine/Source/Runtime/Engine/Classes/Engine/Texture.h
@@ -626,7 +626,7 @@ public:
 	/**
 	* Return maximum dimension for this texture type.
 	*/
-	virtual uint32 GetMaximumDimension() const;
+	ENGINE_API virtual uint32 GetMaximumDimension() const;
 #endif
 
 	/** @return the width of the surface represented by the texture. */
diff --git a/Engine/Source/Runtime/Engine/Classes/Engine/TextureRenderTarget.h b/Engine/Source/Runtime/Engine/Classes/Engine/TextureRenderTarget.h
--- a/Engine/Source/Runtime/Engine/Classes/Engine/TextureRenderTarget.h
+++ b/Engine/Source/Runtime/Engine/Classes/Engine/TextureRenderTarget.h
@@ -37,8 +37,8 @@ class UTextureRenderTarget : public UTexture
 
 
 	// Begin UTexture Interface
-	virtual class FTextureResource* CreateResource();
-	virtual EMaterialValueType GetMaterialType();
+	ENGINE_API virtual class FTextureResource* CreateResource();
+	ENGINE_API virtual EMaterialValueType GetMaterialType();
 	// End UTexture Interface
 };
 
diff --git a/Engine/Source/Runtime/Engine/Classes/Engine/TextureRenderTarget2D.h b/Engine/Source/Runtime/Engine/Classes/Engine/TextureRenderTarget2D.h
--- a/Engine/Source/Runtime/Engine/Classes/Engine/TextureRenderTarget2D.h
+++ b/Engine/Source/Runtime/Engine/Classes/Engine/TextureRenderTarget2D.h
@@ -77,17 +77,17 @@ class UTextureRenderTarget2D : public UTextureRenderTarget
 	// Begin UTexture interface.
 	virtual float GetSurfaceWidth() const override { return SizeX; }
 	virtual float GetSurfaceHeight() const override { return SizeY; }
-	virtual FTextureResource* CreateResource() override;
-	virtual EMaterialValueType GetMaterialType() override;
+	ENGINE_API virtual FTextureResource* CreateResource() override;
+	ENGINE_API virtual EMaterialValueType GetMaterialType() override;
 	// End UTexture interface.
 
 	// Begin UObject interface
 #if WITH_EDITOR
-	virtual void PostEditChangeProperty(FPropertyChangedEvent& PropertyChangedEvent) override;
+	ENGINE_API virtual void PostEditChangeProperty(FPropertyChangedEvent& PropertyChangedEvent) override;
 #endif // WITH_EDITOR
-	virtual void PostLoad() override;
-	virtual SIZE_T GetResourceSize(EResourceSizeMode::Type Mode) override;
-	virtual FString GetDesc() override;
+	ENGINE_API virtual void PostLoad() override;
+	ENGINE_API virtual SIZE_T GetResourceSize(EResourceSizeMode::Type Mode) override;
+	ENGINE_API virtual FString GetDesc() override;
 	// End UObject interface
 
 	FORCEINLINE int32 GetNumMips() const
