Shootergame user management cleanup

diff --git a/Engine/Source/Runtime/Engine/Classes/Engine/GameInstance.h b/Engine/Source/Runtime/Engine/Classes/Engine/GameInstance.h
--- a/Engine/Source/Runtime/Engine/Classes/Engine/GameInstance.h
+++ b/Engine/Source/Runtime/Engine/Classes/Engine/GameInstance.h
@@ -68,6 +68,10 @@ public:
 	/** Starts the GameInstance state machine running */
 	virtual void StartGameInstance();
 
-	APlayerController* GetFirstLocalPlayerController() const;
-	ULocalPlayer* GetFirstGamePlayer() const;
+	/** Local player access */
+	int32				GetNumLocalPlayers() const;
+	ULocalPlayer*		GetLocalPlayerByIndex(const int32 Index) const;
+	APlayerController*	GetFirstLocalPlayerController() const;
+	ULocalPlayer*		FindLocalPlayerFromUniqueNetId(TSharedPtr<FUniqueNetId> UniqueNetId) const;
+	ULocalPlayer*		GetFirstGamePlayer() const;
 };
\ No newline at end of file
diff --git a/Engine/Source/Runtime/Engine/Private/GameInstance.cpp b/Engine/Source/Runtime/Engine/Private/GameInstance.cpp
--- a/Engine/Source/Runtime/Engine/Private/GameInstance.cpp
+++ b/Engine/Source/Runtime/Engine/Private/GameInstance.cpp
@@ -325,6 +325,25 @@ bool UGameInstance::Exec(UWorld* InWorld, const TCHAR* Cmd, FOutputDevice& Ar)
 	return false;
 }
 
+int32 UGameInstance::GetNumLocalPlayers() const
+{
+	return ( WorldContext != nullptr ) ? WorldContext->GamePlayers.Num() : 0;
+}
+
+ULocalPlayer* UGameInstance::GetLocalPlayerByIndex( const int32 Index ) const
+{
+	if ( WorldContext == nullptr )
+	{
+		return nullptr;
+	}
+
+	if ( !ensure( Index >= 0 && Index < WorldContext->GamePlayers.Num() ) )
+	{
+		return nullptr;
+	}
+
+	return WorldContext->GamePlayers[Index];
+}
 
 APlayerController* UGameInstance::GetFirstLocalPlayerController() const
 {
@@ -344,6 +363,24 @@ APlayerController* UGameInstance::GetFirstLocalPlayerController() const
 	return nullptr;
 }
 
+ULocalPlayer* UGameInstance::FindLocalPlayerFromUniqueNetId( TSharedPtr<FUniqueNetId> UniqueNetId ) const
+{
+	if (WorldContext)
+	{
+		for (ULocalPlayer* Player : WorldContext->GamePlayers)
+		{
+			if (Player && Player->GetUniqueNetId() == UniqueNetId)
+			{
+				// Match
+				return Player;
+			}
+		}
+	}
+
+	// didn't find one
+	return nullptr;
+}
+
 ULocalPlayer* UGameInstance::GetFirstGamePlayer() const
 {
 	return (WorldContext && WorldContext->GamePlayers.Num() > 0) ? WorldContext->GamePlayers[0] : nullptr;
