#UE4

- Making it so FSlateGameResources accepts path like "Texture2D'/Game/UI/CharSelect/CS_Port_Selected.CS_Port_Selected'" and not just "/Game/UI/CharSelect/CS_Port_Selected".

This will eliminate a lot of Copy->Paste->Edit manual labor.

diff --git a/Engine/Source/Runtime/Engine/Private/Slate/SlateGameResources.cpp b/Engine/Source/Runtime/Engine/Private/Slate/SlateGameResources.cpp
--- a/Engine/Source/Runtime/Engine/Private/Slate/SlateGameResources.cpp
+++ b/Engine/Source/Runtime/Engine/Private/Slate/SlateGameResources.cpp
@@ -43,11 +43,12 @@ void FSlateGameResources::SetContentRoot( const FString& InContentRootDir )
 const FSlateBrush* FSlateGameResources::GetBrush( const FName PropertyName, const ANSICHAR* Specifier ) const
 {
 	ensureMsgf(Specifier == NULL, TEXT("Attempting to look up resource (%s, %s). \n Specifiers not supported by Slate Resource Sets loaded from content browser."), *PropertyName.ToString(), Specifier);
-	UObject* const * Resource = UIResources.Find(PropertyName);
+	FName CleanName = GetCleanName(PropertyName);
+	UObject* const * Resource = UIResources.Find(CleanName);
 	if(Resource)
 	{
 		const USlateBrushAsset* BrushAsset = Cast<USlateBrushAsset>(*Resource);
-		ensureMsgf(BrushAsset, TEXT("Could not find resource '%s'"), *PropertyName.ToString());
+		ensureMsgf(BrushAsset, TEXT("Could not find resource '%s'"), *CleanName.ToString());
 		return BrushAsset ? &BrushAsset->Brush : GetDefaultBrush();
 	}
 	return FSlateStyleSet::GetBrush(PropertyName, Specifier);
@@ -56,11 +57,12 @@ const FSlateBrush* FSlateGameResources::GetBrush( const FName PropertyName, cons
 const FSlateBrush* FSlateGameResources::GetOptionalBrush(const FName PropertyName, const ANSICHAR* Specifier, const FSlateBrush* const DefaultBrush) const
 {
 	ensureMsgf(Specifier == NULL, TEXT("Attempting to look up resource (%s, %s). \n Specifiers not supported by Slate Resource Sets loaded from content browser."), *PropertyName.ToString(), Specifier);
-	UObject* const * Resource = UIResources.Find(PropertyName);
+	FName CleanName = GetCleanName(PropertyName);
+	UObject* const * Resource = UIResources.Find(CleanName);
 	if(Resource)
 	{
 		const USlateBrushAsset* BrushAsset = Cast<USlateBrushAsset>(*Resource);
-		ensureMsgf(BrushAsset, TEXT("Could not find resource '%s'"), *PropertyName.ToString());
+		ensureMsgf(BrushAsset, TEXT("Could not find resource '%s'"), *CleanName.ToString());
 		return BrushAsset ? &BrushAsset->Brush : DefaultBrush;
 	}
 	return FSlateStyleSet::GetOptionalBrush(PropertyName, Specifier, DefaultBrush);
@@ -68,25 +70,28 @@ const FSlateBrush* FSlateGameResources::GetOptionalBrush(const FName PropertyNam
 
 UCurveFloat* FSlateGameResources::GetCurveFloat( const FName AssetName ) const
 {
-	UObject* const* Resource = UIResources.Find( AssetName );
+	FName CleanName = GetCleanName(AssetName);
+	UObject* const* Resource = UIResources.Find(CleanName);
 	UCurveFloat* Curve = Resource ? Cast<UCurveFloat>(*Resource) : NULL;
-	ensureMsgf(Curve, TEXT("Could not find resource '%s'"), *AssetName.ToString());
+	ensureMsgf(Curve, TEXT("Could not find resource '%s'"), *CleanName.ToString());
 	return Curve;
 }
 
 UCurveVector* FSlateGameResources::GetCurveVector( const FName AssetName ) const
 {
-	UObject* const* Resource = UIResources.Find( AssetName );
+	FName CleanName = GetCleanName(AssetName);
+	UObject* const* Resource = UIResources.Find(CleanName);
 	UCurveVector* Curve = Resource ? Cast<UCurveVector>(*Resource) : NULL;
-	ensureMsgf(Curve, TEXT("Could not find resource '%s'"), *AssetName.ToString());
+	ensureMsgf(Curve, TEXT("Could not find resource '%s'"), *CleanName.ToString());
 	return Curve;
 }
 
 UCurveLinearColor* FSlateGameResources::GetCurveLinearColor( const FName AssetName ) const
 {
-	UObject* const* Resource = UIResources.Find( AssetName );
+	FName CleanName = GetCleanName(AssetName);
+	UObject* const* Resource = UIResources.Find(CleanName);
 	UCurveLinearColor* Curve = Resource ? Cast<UCurveLinearColor>(*Resource) : NULL;
-	ensureMsgf(Curve, TEXT("Could not find resource '%s'"), *AssetName.ToString());
+	ensureMsgf(Curve, TEXT("Could not find resource '%s'"), *CleanName.ToString());
 	return Curve;
 }
 
@@ -279,6 +284,17 @@ FName FSlateGameResources::GenerateMapName( UObject* StyleObject )
 	return GenerateMapName( FAssetData( StyleObject ) );
 }
 
+FName FSlateGameResources::GetCleanName(const FName& AssetName) const
+{
+	static const FString SplitOn = TEXT("'");
+	FString RetValue = AssetName.ToString();
+	if (!FPaths::GetBaseFilename(RetValue, false).Split(SplitOn, NULL, &RetValue))
+	{
+		return AssetName;
+	}
+	return FName(*RetValue);
+}
+
 void FSlateGameResources::AddReferencedObjects( FReferenceCollector& Collector )
 {
 	//We only add references to our style assets when in Game.
diff --git a/Engine/Source/Runtime/Engine/Public/Slate/SlateGameResources.h b/Engine/Source/Runtime/Engine/Public/Slate/SlateGameResources.h
--- a/Engine/Source/Runtime/Engine/Public/Slate/SlateGameResources.h
+++ b/Engine/Source/Runtime/Engine/Public/Slate/SlateGameResources.h
@@ -73,6 +73,15 @@ protected:
 
 	FName GenerateMapName( UObject* StyleObject );
 
+	/** Takes paths from the Editor's "Copy Reference" button and turns them into paths accepted by this object.
+	 *
+	 *   Example: 
+	 *   This: "SlateBrushAsset'/Game/UI/STeamAndHeroSelection/CS_Port_Brush.CS_Port_Brush'"
+	 *	 into
+	 *	 This: "/Game/UI/STeamAndHeroSelection/CS_Port_Brush"
+	*/
+	FName GetCleanName(const FName& AssetName) const;
+
 	virtual void AddReferencedObjects( FReferenceCollector& Collector ) override;
 
 private:
