Fix issue where splitscreen doesn't collapse to a single player when the second player leaves + fix issue where splitscreen was activating inside the menus

diff --git a/Engine/Source/Runtime/Engine/Classes/Engine/GameViewportClient.h b/Engine/Source/Runtime/Engine/Classes/Engine/GameViewportClient.h
--- a/Engine/Source/Runtime/Engine/Classes/Engine/GameViewportClient.h
+++ b/Engine/Source/Runtime/Engine/Classes/Engine/GameViewportClient.h
@@ -466,6 +466,9 @@ public:
 	/** Called before rendering to allow the game viewport to allocate subregions to players. */
 	virtual void LayoutPlayers();
 
+	/** Allows game code to disable splitscreen (useful when in menus) */
+	void SetDisableSplitscreenOverride( const bool bDisabled );
+
 	/** called before rending subtitles to allow the game viewport to determine the size of the subtitle area
 	 * @param Min top left bounds of subtitle region (0 to 1)
 	 * @param Max bottom right bounds of subtitle region (0 to 1)
@@ -757,6 +760,9 @@ private:
 
 	/** Those sound stat flags which are enabled on this viewport */
 	static ESoundShowFlags::Type SoundShowFlags;
+
+	/** Disables splitscreen, useful when game code is in menus, and doesn't want splitscreen on */
+	bool bDisableSplitScreenOverride;
 };
 
 
diff --git a/Engine/Source/Runtime/Engine/Private/GameViewportClient.cpp b/Engine/Source/Runtime/Engine/Private/GameViewportClient.cpp
--- a/Engine/Source/Runtime/Engine/Private/GameViewportClient.cpp
+++ b/Engine/Source/Runtime/Engine/Private/GameViewportClient.cpp
@@ -1306,39 +1306,34 @@ ULocalPlayer* UGameViewportClient::CreatePlayer(int32 ControllerId, FString& Out
 
 bool UGameViewportClient::RemovePlayer(ULocalPlayer* ExPlayer)
 {
-	//Mapping of: index into array is new index, value is old index
-
-	// can't destroy viewports while connected to a server
-	if (ExPlayer->PlayerController->Role == ROLE_Authority)
+	if ( ExPlayer->PlayerController != NULL )
 	{
-		UE_LOG(LogPlayerManagement, Log, TEXT("Removing player %s  with ControllerId %s at index %s (%s existing players)"), ExPlayer, ExPlayer->ControllerId, GetOuterUEngine()->GetGamePlayers(this).Find(ExPlayer), GetOuterUEngine()->GetNumGamePlayers(this));
-		
-		if ( ExPlayer->PlayerController != NULL )
+		if (ExPlayer->PlayerController->Role == ROLE_Authority)
 		{
 			// Destroy the player's actors.
 			ExPlayer->PlayerController->Destroy();
 		}
-
-		// Remove the player from the global and viewport lists of players.
-		int32 OldIndex = RemoveLocalPlayer(ExPlayer);
-		if ( OldIndex != INDEX_NONE )
+		else
 		{
-			NotifyPlayerRemoved(OldIndex, ExPlayer);
+			// FIXME: Notify server we want to leave the game
 		}
-
-		// Disassociate this viewport client from the player.
-		// Do this after notifications, as some of them require the ViewportClient.
-		ExPlayer->ViewportClient = NULL;
-
-		UE_LOG(LogPlayerManagement, Log, TEXT("Finished removing player  %s  with ControllerId %i at index %i (%i remaining players)"), *ExPlayer->GetName(), ExPlayer->ControllerId, OldIndex, GetOuterUEngine()->GetNumGamePlayers(this));
-		return true;
 	}
-	else
+
+	// Remove the player from the context list
+	const int32 OldIndex = RemoveLocalPlayer(ExPlayer);
+	
+	// Notify the viewport so the viewport can do the fixups, resize, etc
+	if ( OldIndex != INDEX_NONE )
 	{
-		UEnum* NetRoleEnum = FindObject<UEnum>( NULL, TEXT( "ENetRole" ) );
-		UE_LOG(LogPlayerManagement, Log, TEXT("Not removing player %s  with ControllerId %i because UPlayer* does not have appropriate role (%s"), *ExPlayer->GetName(), ExPlayer->ControllerId, *NetRoleEnum->GetEnum(ExPlayer->PlayerController->Role).ToString());
-		return false;
+		NotifyPlayerRemoved(OldIndex, ExPlayer);
 	}
+
+	// Disassociate this viewport client from the player.
+	// Do this after notifications, as some of them require the ViewportClient.
+	ExPlayer->ViewportClient = NULL;
+
+	UE_LOG(LogPlayerManagement, Log, TEXT("UGameViewportClient::RemovePlayer: Removed player %s with ControllerId %i at index %i (%i remaining players)"), *ExPlayer->GetName(), ExPlayer->ControllerId, OldIndex, GetOuterUEngine()->GetNumGamePlayers(this));
+	return true;
 }
 
 void UGameViewportClient::DebugCreatePlayer(int32 ControllerId)
@@ -1437,7 +1432,7 @@ void UGameViewportClient::UpdateActiveSplitscreenType()
 	const int32 NumPlayers = GEngine->GetNumGamePlayers(GetWorld());
 	const UGameMapsSettings* Settings = GetDefault<UGameMapsSettings>();
 
-	if (Settings->bUseSplitscreen)
+	if (Settings->bUseSplitscreen && !bDisableSplitScreenOverride)
 	{
 		switch (NumPlayers)
 		{
@@ -1519,6 +1514,11 @@ void UGameViewportClient::LayoutPlayers()
 	}
 }
 
+void UGameViewportClient::SetDisableSplitscreenOverride( const bool bDisabled )
+{
+	bDisableSplitScreenOverride = bDisabled;
+	LayoutPlayers();
+}
 
 void UGameViewportClient::GetSubtitleRegion(FVector2D& MinPos, FVector2D& MaxPos)
 {
