#ue4 - Fix crash in CharacterMovement when Character has no Mesh.

diff --git a/Engine/Source/Runtime/Engine/Private/CharacterMovementComponent.cpp b/Engine/Source/Runtime/Engine/Private/CharacterMovementComponent.cpp
--- a/Engine/Source/Runtime/Engine/Private/CharacterMovementComponent.cpp
+++ b/Engine/Source/Runtime/Engine/Private/CharacterMovementComponent.cpp
@@ -821,8 +821,9 @@ void UCharacterMovementComponent::TickComponent(float DeltaTime, enum ELevelTick
 				float CharacterMass = 1.0f;
 #endif
 
-				float GravZ = GetGravityZ();
-				BaseComp->AddForceAtLocation(FVector(0,0, GravZ * CharacterMass), CharacterComp->GetComponentLocation());
+				const float GravZ = GetGravityZ();
+				const FVector ForceLocation = UpdatedComponent->GetComponentLocation() - CharacterOwner->CapsuleComponent->GetScaledCapsuleHalfHeight();
+				BaseComp->AddForceAtLocation(FVector(0,0, GravZ * CharacterMass), ForceLocation);
 			}
 		}
 
@@ -1259,7 +1260,7 @@ void UCharacterMovementComponent::UpdateBasedMovement(float DeltaSeconds)
 			}
 		}
 
-		if(MovementBase->IsSimulatingPhysics())
+		if (MovementBase->IsSimulatingPhysics() && CharacterOwner->Mesh)
 		{
 			CharacterOwner->Mesh->ApplyDeltaToAllPhysicsTransforms(DeltaPosition, DeltaQuat);
 		}
