New dynamic parameters now copy there parameter names from pre-existing dynamic parameters

#ttp 337106 - material editor: new instances of the dynamic parameter node do not automatically gather the existing channel names

#branch UE4

#add Added new function to scan through other dynamic parameters to copy their settings from and modified existing code to reduce dup code.

diff --git a/Engine/Source/Editor/MaterialEditor/Private/MaterialEditor.cpp b/Engine/Source/Editor/MaterialEditor/Private/MaterialEditor.cpp
--- a/Engine/Source/Editor/MaterialEditor/Private/MaterialEditor.cpp
+++ b/Engine/Source/Editor/MaterialEditor/Private/MaterialEditor.cpp
@@ -10,6 +10,7 @@
 #include "Materials/MaterialExpressionConstant2Vector.h"
 #include "Materials/MaterialExpressionConstant3Vector.h"
 #include "Materials/MaterialExpressionConstant4Vector.h"
+#include "Materials/MaterialExpressionDynamicParameter.h"
 #include "Materials/MaterialExpressionFontSampleParameter.h"
 #include "Materials/MaterialExpressionFunctionInput.h"
 #include "Materials/MaterialExpressionFunctionOutput.h"
@@ -2715,6 +2716,13 @@ UMaterialExpression* FMaterialEditor::CreateNewMaterialExpression(UClass* NewExp
 			PositionTransform->TransformType = TRANSFORMPOSSOURCE_World;
 		}
 
+		// Make sure the dynamic parameters are named based on existing ones
+		UMaterialExpressionDynamicParameter* DynamicExpression = Cast<UMaterialExpressionDynamicParameter>(NewExpression);
+		if (DynamicExpression)
+		{
+			DynamicExpression->UpdateDynamicParameterNames();
+		}
+
 		Material->AddExpressionParameter(NewExpression);
 
 		if (NewExpression)
diff --git a/Engine/Source/Runtime/Engine/Classes/Materials/Material.h b/Engine/Source/Runtime/Engine/Classes/Materials/Material.h
--- a/Engine/Source/Runtime/Engine/Classes/Materials/Material.h
+++ b/Engine/Source/Runtime/Engine/Classes/Materials/Material.h
@@ -1027,13 +1027,13 @@ public:
 	 * Return whether the provided expression parameter has duplicates.
 	 * @param	Expression	The expression parameter to check for duplicates.
 	 */
-	ENGINE_API virtual bool HasDuplicateParameters(UMaterialExpression* Expression);
+	ENGINE_API virtual bool HasDuplicateParameters(const UMaterialExpression* Expression);
 
 	/**
 	 * Return whether the provided expression dynamic parameter has duplicates.
 	 * @param	Expression	The expression dynamic parameter to check for duplicates.
 	 */
-	ENGINE_API virtual bool HasDuplicateDynamicParameters(UMaterialExpression* Expression);
+	ENGINE_API virtual bool HasDuplicateDynamicParameters(const UMaterialExpression* Expression);
 
 	/**
 	 * Iterate through all of the expression nodes and fix up changed names on
@@ -1041,7 +1041,7 @@ public:
 	 *
 	 * @param	Expression	The expression dynamic parameter.
 	 */
-	ENGINE_API virtual void UpdateExpressionDynamicParameterNames(UMaterialExpression* Expression);
+	ENGINE_API virtual void UpdateExpressionDynamicParameterNames(const UMaterialExpression* Expression);
 
 	/**
 	 * Get the name of a parameter.
@@ -1049,7 +1049,7 @@ public:
 	 * @param	OutName		The variable that will hold the parameter name.
 	 * @return	true if the expression is a parameter with a name.
 	 */
-	static bool GetExpressionParameterName(UMaterialExpression* Expression, FName& OutName);
+	static bool GetExpressionParameterName(const UMaterialExpression* Expression, FName& OutName);
 
 	/**
 	 * Copy the values of an expression parameter to another expression parameter of the same class.
@@ -1071,7 +1071,7 @@ public:
 	 *
 	 * @param	Expression	The expression node to inspect.
 	 */
-	ENGINE_API static bool IsDynamicParameter(UMaterialExpression* Expression);
+	ENGINE_API static bool IsDynamicParameter(const UMaterialExpression* Expression);
 
 	/**
 	 * Return wheter the number of parameter groups. NOTE: The number returned can be innaccurate if you have parameters of different types with the same name.
diff --git a/Engine/Source/Runtime/Engine/Classes/Materials/MaterialExpressionDynamicParameter.h b/Engine/Source/Runtime/Engine/Classes/Materials/MaterialExpressionDynamicParameter.h
--- a/Engine/Source/Runtime/Engine/Classes/Materials/MaterialExpressionDynamicParameter.h
+++ b/Engine/Source/Runtime/Engine/Classes/Materials/MaterialExpressionDynamicParameter.h
@@ -37,6 +37,20 @@ class UMaterialExpressionDynamicParameter : public UMaterialExpression
 	virtual int32 GetLabelPadding() override { return 8; }
 	virtual bool MatchesSearchQuery( const TCHAR* SearchQuery ) override;
 	// End UMaterialExpression Interface
+
+	/**
+	 * Iterate through all of the expression nodes until we find another 
+	 * dynamic parameter we can copy the parameter names from
+	 */
+	ENGINE_API void UpdateDynamicParameterNames();
+
+	/**
+	 * Copy the parameters names from the specified dynamic parameter
+	 *
+	 * @param	FromParam	The param to copy the names from
+	 * @return	true if sucessful
+	 */
+	bool CopyDynamicParameterNames(const UMaterialExpressionDynamicParameter* FromParam);
 };
 
 
diff --git a/Engine/Source/Runtime/Engine/Private/Materials/Material.cpp b/Engine/Source/Runtime/Engine/Private/Materials/Material.cpp
--- a/Engine/Source/Runtime/Engine/Private/Materials/Material.cpp
+++ b/Engine/Source/Runtime/Engine/Private/Materials/Material.cpp
@@ -2376,7 +2376,7 @@ bool UMaterial::IsParameter(const UMaterialExpression* Expression)
 }
 
 
-bool UMaterial::IsDynamicParameter(UMaterialExpression* Expression)
+bool UMaterial::IsDynamicParameter(const UMaterialExpression* Expression)
 {
 	if (Expression->IsA(UMaterialExpressionDynamicParameter::StaticClass()))
 	{
@@ -2399,7 +2399,7 @@ void UMaterial::BuildEditorParameterList()
 }
 
 
-bool UMaterial::HasDuplicateParameters(UMaterialExpression* Expression)
+bool UMaterial::HasDuplicateParameters(const UMaterialExpression* Expression)
 {
 	FName ExpressionName;
 
@@ -2424,9 +2424,9 @@ bool UMaterial::HasDuplicateParameters(UMaterialExpression* Expression)
 }
 
 
-bool UMaterial::HasDuplicateDynamicParameters(UMaterialExpression* Expression)
+bool UMaterial::HasDuplicateDynamicParameters(const UMaterialExpression* Expression)
 {
-	UMaterialExpressionDynamicParameter* DynParam = Cast<UMaterialExpressionDynamicParameter>(Expression);
+	const UMaterialExpressionDynamicParameter* DynParam = Cast<UMaterialExpressionDynamicParameter>(Expression);
 	if (DynParam)
 	{
 		for (int32 ExpIndex = 0; ExpIndex < Expressions.Num(); ExpIndex++)
@@ -2442,20 +2442,16 @@ bool UMaterial::HasDuplicateDynamicParameters(UMaterialExpression* Expression)
 }
 
 
-void UMaterial::UpdateExpressionDynamicParameterNames(UMaterialExpression* Expression)
+void UMaterial::UpdateExpressionDynamicParameterNames(const UMaterialExpression* Expression)
 {
-	UMaterialExpressionDynamicParameter* DynParam = Cast<UMaterialExpressionDynamicParameter>(Expression);
+	const UMaterialExpressionDynamicParameter* DynParam = Cast<UMaterialExpressionDynamicParameter>(Expression);
 	if (DynParam)
 	{
 		for (int32 ExpIndex = 0; ExpIndex < Expressions.Num(); ExpIndex++)
 		{
 			UMaterialExpressionDynamicParameter* CheckParam = Cast<UMaterialExpressionDynamicParameter>(Expressions[ExpIndex]);
-			if (CheckParam && (CheckParam != DynParam))
+			if (CheckParam->CopyDynamicParameterNames(DynParam))
 			{
-				for (int32 NameIndex = 0; NameIndex < 4; NameIndex++)
-				{
-					CheckParam->ParamNames[NameIndex] = DynParam->ParamNames[NameIndex];
-				}
 #if WITH_EDITORONLY_DATA
 				CheckParam->GraphNode->ReconstructNode();
 #endif // WITH_EDITORONLY_DATA
@@ -2509,7 +2505,7 @@ void UMaterial::UpdateExpressionParameterName(UMaterialExpression* Expression)
 }
 
 
-bool UMaterial::GetExpressionParameterName(UMaterialExpression* Expression, FName& OutName)
+bool UMaterial::GetExpressionParameterName(const UMaterialExpression* Expression, FName& OutName)
 {
 	bool bRet = false;
 
diff --git a/Engine/Source/Runtime/Engine/Private/Materials/MaterialExpressions.cpp b/Engine/Source/Runtime/Engine/Private/Materials/MaterialExpressions.cpp
--- a/Engine/Source/Runtime/Engine/Private/Materials/MaterialExpressions.cpp
+++ b/Engine/Source/Runtime/Engine/Private/Materials/MaterialExpressions.cpp
@@ -4363,6 +4363,32 @@ void UMaterialExpressionDynamicParameter::PostEditChangeProperty(FPropertyChange
 
 #endif // WITH_EDITOR
 
+void UMaterialExpressionDynamicParameter::UpdateDynamicParameterNames()
+{
+	check(Material);
+	for (int32 ExpIndex = 0; ExpIndex < Material->Expressions.Num(); ExpIndex++)
+	{
+		const UMaterialExpressionDynamicParameter* DynParam = Cast<UMaterialExpressionDynamicParameter>(Material->Expressions[ExpIndex]);
+		if (CopyDynamicParameterNames(DynParam))
+		{
+			break;
+		}
+	}
+}
+
+bool UMaterialExpressionDynamicParameter::CopyDynamicParameterNames(const UMaterialExpressionDynamicParameter* FromParam)
+{
+	if (FromParam && (FromParam != this))
+	{
+		for (int32 NameIndex = 0; NameIndex < 4; NameIndex++)
+		{
+			ParamNames[NameIndex] = FromParam->ParamNames[NameIndex];
+		}
+		return true;
+	}
+	return false;
+}
+
 //
 //	MaterialExpressionParticleSubUV
 //
