Fixed: Loaded static components should not always reset render state after creation
TTP #317039

diff --git a/Engine/Source/Runtime/Engine/Classes/Components/SceneComponent.h b/Engine/Source/Runtime/Engine/Classes/Components/SceneComponent.h
--- a/Engine/Source/Runtime/Engine/Classes/Components/SceneComponent.h
+++ b/Engine/Source/Runtime/Engine/Classes/Components/SceneComponent.h
@@ -437,6 +437,7 @@ public:
 	FPhysicsVolumeChanged PhysicsVolumeChangedDelegate;
 
 	// Begin ActorComponent interface
+	virtual void OnRegister() override;
 	virtual void UpdateComponentToWorld(bool bSkipPhysicsMove = false) override final;
 	virtual void DestroyComponent() override;
 	virtual void ApplyWorldOffset(const FVector& InOffset, bool bWorldShift) override;
diff --git a/Engine/Source/Runtime/Engine/Private/Actor.cpp b/Engine/Source/Runtime/Engine/Private/Actor.cpp
--- a/Engine/Source/Runtime/Engine/Private/Actor.cpp
+++ b/Engine/Source/Runtime/Engine/Private/Actor.cpp
@@ -2098,55 +2098,7 @@ void AActor::DisableComponentsSimulatePhysics()
 
 void AActor::PostRegisterAllComponents() 
 {
-	// Try and get transforms correct.
-	// Obviously a component could be attached to something in another actor that has not been processed yet, that will just move the registered components to the correct location
-
-	// We refetch the components because the registration process could have modified the list
-	TArray<USceneComponent*> SceneComponents;
-	GetComponents(SceneComponents);
-
-	for(int32 CompIdx=0; CompIdx < SceneComponents.Num(); CompIdx++)
-	{
-		USceneComponent* SceneComp = SceneComponents[CompIdx];
-#if !(UE_BUILD_SHIPPING || UE_BUILD_TEST)
-		AActor* CompOwner = SceneComp->GetOwner();
-		FString CompOwnerPathName = (CompOwner != NULL) ? CompOwner->GetFullName() : TEXT("NONE");
-		checkf(SceneComp->GetOwner() == this, TEXT("SceneComp '%s' has Owner of %x'%s', expected %x'%s' (TTP 241544)"), *SceneComp->GetFullName(), CompOwner, *CompOwnerPathName, this, *this->GetFullName());
-#endif
-		//UE_LOG(LogActorComponent, Log, TEXT("%d: %s (AP: %s)"), CompIdx, *SceneComp->GetPathName(), SceneComp->AttachParent ? *SceneComp->AttachParent->GetPathName() : TEXT("NONE"));
-		// If we need to perform a call to AttachTo to add it to AttachChildren array, do that now
-		if(SceneComp->AttachParent != NULL && !SceneComp->AttachParent->AttachChildren.Contains(SceneComp))
-		{
-			USceneComponent* PendingAttachParent = SceneComp->AttachParent;
-			SceneComp->AttachParent = NULL;
-
-			// If our attach parent is a blueprint component that is marked pending kill belonging to another actor,
-			// then rerun construction script has almost certainly been run on it and we are just going to pick its
-			// current root component instead
-			if (PendingAttachParent->bCreatedByConstructionScript && PendingAttachParent->IsPendingKill())
-			{
-				AActor* PendingAttachOwner = PendingAttachParent->GetOwner();
-				if (PendingAttachOwner != this)
-				{
-					PendingAttachParent = PendingAttachOwner->GetRootComponent();
-				}
-			}
-
-			FName PendingAttachSocketName = SceneComp->AttachSocketName;
-			SceneComp->AttachSocketName = NAME_None;
-
-			//UE_LOG(LogActorComponent, Log, TEXT("  Attached!"));
-			SceneComp->AttachTo(PendingAttachParent, PendingAttachSocketName);
-		}
-		// Otherwise just update its ComponentToWorld
-		else
-		{
-			SceneComp->UpdateComponentToWorld();
-		}
-	}
-		
-	if (bNavigationRelevant == true && GetWorld() != NULL 
-		&& GetWorld()->GetNavigationSystem() != NULL)
+	if (bNavigationRelevant == true && GetWorld()->GetNavigationSystem() != NULL)
 	{
 		// force re-adding to navoctree if already added, in case it has already 
 		// been added (possibly with not all components that could affect navigation
@@ -3082,6 +3034,30 @@ void AActor::RegisterAllComponents()
 	verify(IncrementalRegisterComponents(0));
 }
 
+// Walks through components hierarchy and returns closest to root parent component that is unregistered
+// Only for components that belong to the same owner
+static USceneComponent* GetUnregisteredParent(UActorComponent* Component)
+{
+	USceneComponent* ParentComponent = nullptr;
+	USceneComponent* SceneComponent = Cast<USceneComponent>(Component);
+	
+	while (	SceneComponent && 
+			SceneComponent->AttachParent && 
+			SceneComponent->AttachParent->GetOwner() == Component->GetOwner() && 
+			!SceneComponent->AttachParent->IsRegistered())
+	{
+		SceneComponent = SceneComponent->AttachParent;
+		if (SceneComponent->bAutoRegister && !SceneComponent->IsPendingKill())
+		{
+			// We found unregistered parent that should be registered
+			// But keep looking up the tree
+			ParentComponent = SceneComponent;
+		}
+	}
+
+	return ParentComponent;
+}
+
 bool AActor::IncrementalRegisterComponents(int32 NumComponentsToRegister)
 {
 	if (NumComponentsToRegister == 0)
@@ -3114,6 +3090,15 @@ bool AActor::IncrementalRegisterComponents(int32 NumComponentsToRegister)
 		UActorComponent* Component = Components[CompIdx];
 		if(!Component->IsRegistered() && Component->bAutoRegister && !Component->IsPendingKill())
 		{
+			// Ensure that all parent are registered first
+			USceneComponent* ParentComponent = GetUnregisteredParent(Component);
+			if (ParentComponent)
+			{
+				// Register parent first, then return to this component on a next iteration
+				Component = ParentComponent;
+				CompIdx--;
+			}
+				
 			//Before we register our component, save it to our transaction buffer so if "undone" it will return to an unregistered state.
 			//This should prevent unwanted components hanging around when undoing a copy/paste or duplication action.
 			Component->Modify(false);
diff --git a/Engine/Source/Runtime/Engine/Private/SceneComponent.cpp b/Engine/Source/Runtime/Engine/Private/SceneComponent.cpp
--- a/Engine/Source/Runtime/Engine/Private/SceneComponent.cpp
+++ b/Engine/Source/Runtime/Engine/Private/SceneComponent.cpp
@@ -105,6 +105,16 @@ void USceneComponent::UpdateComponentToWorldWithParent(USceneComponent * Parent,
 	}
 }
 
+void USceneComponent::OnRegister()
+{
+	// If we need to perform a call to AttachTo, do that now
+	// At this point scene component still has no any state (rendering, physics),
+	// so this call will just add this component to an AttachChildren array of a the Parent component
+	AttachTo(AttachParent, AttachSocketName);
+	
+	Super::OnRegister();
+}
+
 void USceneComponent::UpdateComponentToWorld(bool bSkipPhysicsMove)
 {
 	UpdateComponentToWorldWithParent(AttachParent, bSkipPhysicsMove);
