Missing file from CL:2225377 (stupid p4)

diff --git a/Engine/Source/Runtime/Engine/Private/UnrealEngine.cpp b/Engine/Source/Runtime/Engine/Private/UnrealEngine.cpp
--- a/Engine/Source/Runtime/Engine/Private/UnrealEngine.cpp
+++ b/Engine/Source/Runtime/Engine/Private/UnrealEngine.cpp
@@ -5487,7 +5487,6 @@ void UEngine::OnHardwareSurveyComplete(const FHardwareSurveyResults& SurveyResul
 		HardwareWEIAttribs.Add(FAnalyticsEventAttribute(TEXT( "Memory.WEI" ), FString::Printf( TEXT( "%.1f" ), SurveyResults.RAMPerformanceIndex )));
 
 		Analytics.RecordEvent(TEXT( "Hardware.WEI.1" ), HardwareWEIAttribs);
-		Analytics.RecordUserAttribute(HardwareWEIAttribs);
 
 		FString MainGPUName(TEXT("Unknown"));
 		float MainGPUVRAMMB = 0.0f;
@@ -5572,7 +5571,6 @@ void UEngine::OnHardwareSurveyComplete(const FHardwareSurveyResults& SurveyResul
 		HardwareStatsAttribs.Add(FAnalyticsEventAttribute(TEXT( "SynthIdx.GPU4" ), FString::Printf( TEXT( "%.1f" ), SurveyResults.SynthBenchmark.GPUStats[4].ComputePerfIndex() )));
 
 		Analytics.RecordEvent(TEXT( "HardwareStats.1" ), HardwareStatsAttribs);
-		Analytics.RecordUserAttribute(TEXT( "ResolutionClass" ), ResolutionClass);
 
 		TArray<FAnalyticsEventAttribute> HardwareStatErrorsAttribs;
 		HardwareStatErrorsAttribs.Add(FAnalyticsEventAttribute(TEXT( "ErrorCount" ), FString::Printf( TEXT( "%d" ), SurveyResults.ErrorCount )));
