Header cleanup: SoundNodeAmbientNonLoopToggle.h

diff --git a/Engine/Source/Runtime/Engine/Classes/Sound/SoundNodeAmbientNonLoopToggle.h b/Engine/Source/Runtime/Engine/Classes/Sound/SoundNodeAmbientNonLoopToggle.h
--- a/Engine/Source/Runtime/Engine/Classes/Sound/SoundNodeAmbientNonLoopToggle.h
+++ b/Engine/Source/Runtime/Engine/Classes/Sound/SoundNodeAmbientNonLoopToggle.h
@@ -6,10 +6,7 @@
  */
 
 #pragma once
-
-#include "SoundNodeAttenuation.h"
 #include "Sound/SoundNodeAmbientNonLoop.h"
-
 #include "SoundNodeAmbientNonLoopToggle.generated.h"
 
 UCLASS(hidecategories=Object, editinlinenew)
diff --git a/Engine/Source/Runtime/Engine/Private/SoundNodeAmbientNonLoopToggle.cpp b/Engine/Source/Runtime/Engine/Private/SoundNodeAmbientNonLoopToggle.cpp
--- a/Engine/Source/Runtime/Engine/Private/SoundNodeAmbientNonLoopToggle.cpp
+++ b/Engine/Source/Runtime/Engine/Private/SoundNodeAmbientNonLoopToggle.cpp
@@ -2,7 +2,6 @@
 
 
 #include "EnginePrivate.h"
-#include "SoundDefinitions.h"
 #include "Sound/SoundNodeAmbientNonLoopToggle.h"
 
 /*-----------------------------------------------------------------------------
