Header cleanup: SoundNodeAmbientNonLoop.h

diff --git a/Engine/Source/Runtime/Engine/Classes/Sound/SoundNodeAmbientNonLoop.h b/Engine/Source/Runtime/Engine/Classes/Sound/SoundNodeAmbientNonLoop.h
--- a/Engine/Source/Runtime/Engine/Classes/Sound/SoundNodeAmbientNonLoop.h
+++ b/Engine/Source/Runtime/Engine/Classes/Sound/SoundNodeAmbientNonLoop.h
@@ -6,11 +6,7 @@
  */
 
 #pragma once
-
-#include "SoundNodeAttenuation.h"
-#include "Distributions/DistributionFloat.h"
 #include "SoundNodeAmbient.h"
-
 #include "SoundNodeAmbientNonLoop.generated.h"
 
 UCLASS(editinlinenew)
diff --git a/Engine/Source/Runtime/Engine/Private/SoundNodeAmbientNonLoop.cpp b/Engine/Source/Runtime/Engine/Private/SoundNodeAmbientNonLoop.cpp
--- a/Engine/Source/Runtime/Engine/Private/SoundNodeAmbientNonLoop.cpp
+++ b/Engine/Source/Runtime/Engine/Private/SoundNodeAmbientNonLoop.cpp
@@ -2,7 +2,6 @@
 
 
 #include "EnginePrivate.h"
-#include "SoundDefinitions.h"
 #include "Sound/SoundNodeAmbientNonLoop.h"
 
 /*-----------------------------------------------------------------------------
