Header cleanup: Pawn.h

diff --git a/Engine/Source/Runtime/AIModule/Private/AIController.cpp b/Engine/Source/Runtime/AIModule/Private/AIController.cpp
--- a/Engine/Source/Runtime/AIModule/Private/AIController.cpp
+++ b/Engine/Source/Runtime/AIModule/Private/AIController.cpp
@@ -10,6 +10,7 @@
 #include "Navigation/PathFollowingComponent.h"
 #include "Navigation/NavigationComponent.h"
 #include "GameFramework/Pawn.h"
+#include "GameFramework/PawnMovementComponent.h"
 #include "Engine/Canvas.h"
 #include "GameFramework/PhysicsVolume.h"
 #include "AIController.h"
diff --git a/Engine/Source/Runtime/Engine/Classes/GameFramework/DefaultPawn.h b/Engine/Source/Runtime/Engine/Classes/GameFramework/DefaultPawn.h
--- a/Engine/Source/Runtime/Engine/Classes/GameFramework/DefaultPawn.h
+++ b/Engine/Source/Runtime/Engine/Classes/GameFramework/DefaultPawn.h
@@ -6,17 +6,21 @@
 //=============================================================================
 
 #pragma once
-#include "FloatingPawnMovement.h"
+#include "GameFramework/Pawn.h"
 #include "DefaultPawn.generated.h"
 
+class UFloatingPawnMovement;
+class USphereComponent;
+class UStaticMeshComponent;
+
 UCLASS(config=Game, Blueprintable, BlueprintType)
 class ENGINE_API ADefaultPawn : public APawn
 {
 	GENERATED_UCLASS_BODY()
 
 	// Begin Pawn overrides
-	virtual class UPawnMovementComponent* GetMovementComponent() const override { return MovementComponent; }
-	virtual void SetupPlayerInputComponent(class UInputComponent* InputComponent) override;
+	virtual UPawnMovementComponent* GetMovementComponent() const override;
+	virtual void SetupPlayerInputComponent(UInputComponent* InputComponent) override;
 	// End Pawn overrides
 
 	/** Input callback to move forward. */
@@ -68,21 +72,21 @@ public:
 
 	/** DefaultPawn movement component */
 	UPROPERTY(Category=Pawn, VisibleAnywhere, BlueprintReadOnly)
-	TSubobjectPtr<class UFloatingPawnMovement> MovementComponent;
+	TSubobjectPtr<UFloatingPawnMovement> MovementComponent;
 
 	/** Name of the CollisionComponent. */
 	static FName CollisionComponentName;
 
 	/** DefaultPawn collision component */
 	UPROPERTY(Category=Pawn, VisibleAnywhere, BlueprintReadOnly)
-	TSubobjectPtr<class USphereComponent> CollisionComponent;
+	TSubobjectPtr<USphereComponent> CollisionComponent;
 
 	/** Name of the MeshComponent. Use this name if you want to prevent creation of the component (with PCIP.DoNotCreateDefaultSubobject). */
 	static FName MeshComponentName;
 
 	/** The mesh associated with this Pawn. */
 	UPROPERTY(Category=Pawn, VisibleAnywhere, BlueprintReadOnly)
-	TSubobjectPtr<class UStaticMeshComponent> MeshComponent;
+	TSubobjectPtr<UStaticMeshComponent> MeshComponent;
 
 	/** If true, adds default bindings for movement and camera look. */
 	UPROPERTY(Category=Pawn, EditAnywhere, BlueprintReadOnly)
diff --git a/Engine/Source/Runtime/Engine/Classes/GameFramework/FloatingPawnMovement.h b/Engine/Source/Runtime/Engine/Classes/GameFramework/FloatingPawnMovement.h
--- a/Engine/Source/Runtime/Engine/Classes/GameFramework/FloatingPawnMovement.h
+++ b/Engine/Source/Runtime/Engine/Classes/GameFramework/FloatingPawnMovement.h
@@ -5,6 +5,7 @@
  */
 
 #pragma once
+#include "GameFramework/PawnMovementComponent.h"
 #include "FloatingPawnMovement.generated.h"
 
 UCLASS()
diff --git a/Engine/Source/Runtime/Engine/Classes/GameFramework/Pawn.h b/Engine/Source/Runtime/Engine/Classes/GameFramework/Pawn.h
--- a/Engine/Source/Runtime/Engine/Classes/GameFramework/Pawn.h
+++ b/Engine/Source/Runtime/Engine/Classes/GameFramework/Pawn.h
@@ -1,12 +1,25 @@
 // Copyright 1998-2014 Epic Games, Inc. All Rights Reserved.
 
 #pragma once
+#include "AI/Navigation/NavAgentInterface.h"
 #include "GameFramework/Actor.h"
-#include "PawnMovementComponent.h"
 #include "Pawn.generated.h"
 
 ENGINE_API DECLARE_LOG_CATEGORY_EXTERN(LogDamage, Warning, All);
 
+class AController;
+class APhysicsVolume;
+class APlayerController;
+class APlayerState;
+class UCanvas;
+class UDamageType;
+class UInputComponent;
+class UNetConnection;
+class UPawnMovementComponent;
+class UPawnNoiseEmitterComponent;
+class UPlayer;
+class UPrimitiveComponent;
+
 /** 
  *	Pawn is the base class of all actors that can be possessed by players or AI.
  *	Pawns are the physical representations of players and creatures in a level.
@@ -19,10 +32,10 @@ class ENGINE_API APawn : public AActor, public INavAgentInterface
 
 	/** Return our PawnMovementComponent, if we have one. By default, returns the first PawnMovementComponent found. Native classes that create their own movement component should override this method for more efficiency. */
 	UFUNCTION(BlueprintCallable, meta=(Tooltip="Return our PawnMovementComponent, if we have one."), Category="Pawn")
-	virtual class UPawnMovementComponent* GetMovementComponent() const;
+	virtual UPawnMovementComponent* GetMovementComponent() const;
 
 	/** Return Actor we are based on (standing on, attached to, and moving on). */
-	virtual class UPrimitiveComponent* GetMovementBase() const { return NULL; }
+	virtual UPrimitiveComponent* GetMovementBase() const { return NULL; }
 
 public:
 	/** If true, this Pawn's pitch will be updated to match the Controller's ControlRotation pitch, if controlled by a PlayerController. */
@@ -57,7 +70,7 @@ public:
 	 * Return our PawnNoiseEmitterComponent, if any. Default implementation returns the first PawnNoiseEmitterComponent found in the components array.
 	 * If one isn't found, then it tries to find one on the Pawn's current Controller.
 	 */
-	virtual class UPawnNoiseEmitterComponent* GetPawnNoiseEmitterComponent() const;
+	virtual UPawnNoiseEmitterComponent* GetPawnNoiseEmitterComponent() const;
 
 	/**
 	 * Inform AIControllers that you've made a noise they might hear (they are sent a HearNoise message if they have bHearNoises==true)
@@ -73,11 +86,11 @@ public:
 
 	/** default class to use when pawn is controlled by AI. */
 	UPROPERTY(EditAnywhere, BlueprintReadWrite, Category=AI)
-	TSubclassOf<class AController>  AIControllerClass;
+	TSubclassOf<AController>  AIControllerClass;
 
 	/** If Pawn is possessed by a player, points to his playerstate.  Needed for network play as controllers are not replicated to clients. */
 	UPROPERTY(editinline, replicatedUsing=OnRep_PlayerState, BlueprintReadOnly, Category="Pawn")
-	class APlayerState* PlayerState;
+	APlayerState* PlayerState;
 
 	/** Replicated so we can see where remote clients are looking. */
 	UPROPERTY(replicated)
@@ -85,11 +98,11 @@ public:
 
 	/** Controller of the last Actor that caused us damage. */
 	UPROPERTY(transient)
-	class AController* LastHitBy;
+	AController* LastHitBy;
 
 	/** Controller currently possessing this Actor */
 	UPROPERTY(editinline, replicatedUsing=OnRep_Controller)
-	class AController* Controller;
+	AController* Controller;
 
 	/** Max difference between pawn's Rotation.Yaw and GetDesiredRotation().Yaw for pawn to be considered as having reached its desired rotation */
 	float AllowedYawError;
@@ -114,7 +127,7 @@ public:
 	virtual void UnPossessed();
 
 	/** Return Physics Volume for this Pawn **/
-	virtual class APhysicsVolume* GetPawnPhysicsVolume() const;
+	virtual APhysicsVolume* GetPawnPhysicsVolume() const;
 
 private:
 	/** (DEPRECATED) @RETURN true if Pawn is currently walking (moving along the ground) */
@@ -173,29 +186,29 @@ public:
 	virtual bool IsNetRelevantFor(APlayerController* RealViewer, AActor* Viewer, const FVector& SrcLocation) override;
 	virtual void PostNetReceiveLocationAndRotation() override;
 	virtual void PostNetReceiveVelocity(const FVector& NewVelocity) override;
-	virtual void DisplayDebug(class UCanvas* Canvas, const FDebugDisplayInfo& DebugDisplay, float& YL, float& YPos) override;
+	virtual void DisplayDebug(UCanvas* Canvas, const FDebugDisplayInfo& DebugDisplay, float& YL, float& YPos) override;
 	virtual void GetActorEyesViewPoint( FVector& Location, FRotator& Rotation ) const override;
 	virtual void OutsideWorldBounds() override;
 	virtual void Destroyed() override;
 	virtual void PreInitializeComponents() override;
 	virtual void PostInitializeComponents() override;
-	virtual class UPlayer* GetNetOwningPlayer() override;
-	virtual class UNetConnection* GetNetConnection() override;
+	virtual UPlayer* GetNetOwningPlayer() override;
+	virtual UNetConnection* GetNetConnection() override;
 	virtual void PostInitProperties() override;
 	virtual void PostLoad() override;
 	virtual void PostRegisterAllComponents() override;
-	virtual float TakeDamage(float Damage, struct FDamageEvent const& DamageEvent, class AController* EventInstigator, class AActor* DamageCauser) override;
-	virtual void BecomeViewTarget(class APlayerController* PC) override;
+	virtual float TakeDamage(float Damage, struct FDamageEvent const& DamageEvent, AController* EventInstigator, AActor* DamageCauser) override;
+	virtual void BecomeViewTarget(APlayerController* PC) override;
 	virtual bool UpdateNavigationRelevancy() override { SetNavigationRelevancy(false); return false; }
-	virtual void EnableInput(class APlayerController* PlayerController) override;
-	virtual void DisableInput(class APlayerController* PlayerController) override;
+	virtual void EnableInput(APlayerController* PlayerController) override;
+	virtual void DisableInput(APlayerController* PlayerController) override;
 
 	/** Overridden to defer to the RootComponent's CanCharacterStepUpOn setting if it is explicitly Yes or No. If set to Owner, will return Super::CanBeBaseForCharacter(). */
-	virtual bool CanBeBaseForCharacter(class APawn* APawn) const override;
+	virtual bool CanBeBaseForCharacter(APawn* APawn) const override;
 	// End AActor Interface
 
 	// Begin INavAgentInterface Interface
-	virtual const struct FNavAgentProperties* GetNavAgentProperties() const override { return GetMovementComponent() ? GetMovementComponent()->GetNavAgentProperties() : NULL;}
+	virtual const struct FNavAgentProperties* GetNavAgentProperties() const override;
 	/** Basically retrieved pawn's location on navmesh */
 	UFUNCTION(BlueprintCallable, Category="Pawn")
 	virtual FVector GetNavAgentLocation() const override { return GetActorLocation() - FVector(0.f, 0.f, BaseEyeHeight); }
@@ -228,7 +241,7 @@ public:
 	 * Called when this Pawn is possessed. Only called on the server (or in standalone).
 	 *	@param C is the controller possessing this pawn
 	 */
-	virtual void PossessedBy(class AController* NewController);
+	virtual void PossessedBy(AController* NewController);
 
 	UFUNCTION(BlueprintImplementableEvent, meta=(FriendlyName = "Possessed"))
 	void ReceivePossessed(AController* NewController);
@@ -271,16 +284,16 @@ public:
 
 protected:
 	/** Get the controller instigating the damage. If the damage is caused by the world and the supplied controller is NULL or is this pawn's controller, uses LastHitBy as the instigator. */
-	virtual class AController* GetDamageInstigator(class AController* InstigatedBy, const class UDamageType& DamageType) const;
+	virtual AController* GetDamageInstigator(AController* InstigatedBy, const UDamageType& DamageType) const;
 
 	/** Creates an InputComponent that can be used for custom input bindings. Called upon possession by a PlayerController. Return null if you don't want one. */
-	virtual class UInputComponent* CreatePlayerInputComponent();
+	virtual UInputComponent* CreatePlayerInputComponent();
 
 	/** Destroys the player input component and removes any references to it. */
 	virtual void DestroyPlayerInputComponent();
 
 	/** Allows a Pawn to set up custom input bindings. Called upon possession by a PlayerController, using the InputComponent created by CreatePlayerInputComponent(). */
-	virtual void SetupPlayerInputComponent(class UInputComponent* InputComponent) { /* No bindings by default.*/ }
+	virtual void SetupPlayerInputComponent(UInputComponent* InputComponent) { /* No bindings by default.*/ }
 
 public:
 	/**
diff --git a/Engine/Source/Runtime/Engine/Classes/GameFramework/RotatingMovementComponent.h b/Engine/Source/Runtime/Engine/Classes/GameFramework/RotatingMovementComponent.h
--- a/Engine/Source/Runtime/Engine/Classes/GameFramework/RotatingMovementComponent.h
+++ b/Engine/Source/Runtime/Engine/Classes/GameFramework/RotatingMovementComponent.h
@@ -1,6 +1,7 @@
 // Copyright 1998-2014 Epic Games, Inc. All Rights Reserved.
 
 #pragma once
+#include "GameFramework/MovementComponent.h"
 #include "RotatingMovementComponent.generated.h"
 
 /**
diff --git a/Engine/Source/Runtime/Engine/Classes/Vehicles/WheeledVehicleMovementComponent.h b/Engine/Source/Runtime/Engine/Classes/Vehicles/WheeledVehicleMovementComponent.h
--- a/Engine/Source/Runtime/Engine/Classes/Vehicles/WheeledVehicleMovementComponent.h
+++ b/Engine/Source/Runtime/Engine/Classes/Vehicles/WheeledVehicleMovementComponent.h
@@ -1,10 +1,9 @@
 // Copyright 1998-2014 Epic Games, Inc. All Rights Reserved.
 
 #pragma once
-
+#include "GameFramework/PawnMovementComponent.h"
 #include "VehicleWheel.h"
 #include "Curves/CurveBase.h"
-
 #include "WheeledVehicleMovementComponent.generated.h"
 
 #if WITH_VEHICLE
diff --git a/Engine/Source/Runtime/Engine/Private/DebugCameraController.cpp b/Engine/Source/Runtime/Engine/Private/DebugCameraController.cpp
--- a/Engine/Source/Runtime/Engine/Private/DebugCameraController.cpp
+++ b/Engine/Source/Runtime/Engine/Private/DebugCameraController.cpp
@@ -6,6 +6,7 @@
 =============================================================================*/
 
 #include "EnginePrivate.h"
+#include "GameFramework/PawnMovementComponent.h"
 #include "GameFramework/SpectatorPawn.h"
 #include "GameFramework/SpectatorPawnMovement.h"
 #include "Engine/DebugCameraHUD.h"
diff --git a/Engine/Source/Runtime/Engine/Private/DebugCameraHUD.cpp b/Engine/Source/Runtime/Engine/Private/DebugCameraHUD.cpp
--- a/Engine/Source/Runtime/Engine/Private/DebugCameraHUD.cpp
+++ b/Engine/Source/Runtime/Engine/Private/DebugCameraHUD.cpp
@@ -5,6 +5,7 @@
 =============================================================================*/
 
 #include "EnginePrivate.h"
+#include "GameFramework/PawnMovementComponent.h"
 #include "GameFramework/SpectatorPawn.h"
 #include "Engine/DebugCameraHUD.h"
 
diff --git a/Engine/Source/Runtime/Engine/Private/DefaultPawn.cpp b/Engine/Source/Runtime/Engine/Private/DefaultPawn.cpp
--- a/Engine/Source/Runtime/Engine/Private/DefaultPawn.cpp
+++ b/Engine/Source/Runtime/Engine/Private/DefaultPawn.cpp
@@ -3,6 +3,7 @@
 
 #include "EnginePrivate.h"
 #include "GameFramework/DefaultPawn.h"
+#include "GameFramework/FloatingPawnMovement.h"
 
 FName ADefaultPawn::MovementComponentName(TEXT("MovementComponent0"));
 FName ADefaultPawn::CollisionComponentName(TEXT("CollisionComponent0"));
@@ -192,3 +193,8 @@ void ADefaultPawn::Turn(float Val)
 		PC->AddYawInput(Val);
 	}
 }
+
+UPawnMovementComponent* ADefaultPawn::GetMovementComponent() const
+{
+	return MovementComponent;
+}
diff --git a/Engine/Source/Runtime/Engine/Private/MovementComponent.cpp b/Engine/Source/Runtime/Engine/Private/MovementComponent.cpp
--- a/Engine/Source/Runtime/Engine/Private/MovementComponent.cpp
+++ b/Engine/Source/Runtime/Engine/Private/MovementComponent.cpp
@@ -2,6 +2,7 @@
 
 
 #include "EnginePrivate.h"
+#include "GameFramework/MovementComponent.h"
 #include "GameFramework/PhysicsVolume.h"
 #include "MessageLog.h"
 
diff --git a/Engine/Source/Runtime/Engine/Private/NavMovementComponent.cpp b/Engine/Source/Runtime/Engine/Private/NavMovementComponent.cpp
--- a/Engine/Source/Runtime/Engine/Private/NavMovementComponent.cpp
+++ b/Engine/Source/Runtime/Engine/Private/NavMovementComponent.cpp
@@ -1,6 +1,7 @@
 // Copyright 1998-2014 Epic Games, Inc. All Rights Reserved.
 
 #include "EnginePrivate.h"
+#include "GameFramework/NavMovementComponent.h"
 #include "Navigation/PathFollowingComponent.h"
 
 //----------------------------------------------------------------------//
diff --git a/Engine/Source/Runtime/Engine/Private/Pawn.cpp b/Engine/Source/Runtime/Engine/Private/Pawn.cpp
--- a/Engine/Source/Runtime/Engine/Private/Pawn.cpp
+++ b/Engine/Source/Runtime/Engine/Private/Pawn.cpp
@@ -8,6 +8,7 @@
 
 #include "EnginePrivate.h"
 #include "GameFramework/Character.h"
+#include "GameFramework/PawnMovementComponent.h"
 #include "Net/UnrealNetwork.h"
 #include "ConfigCacheIni.h"
 #include "ParticleDefinitions.h"
@@ -1048,3 +1049,8 @@ void APawn::PawnMakeNoise(float Loudness, FVector NoiseLocation, bool bUseNoiseM
 	}
 	NoiseMaker->MakeNoise(Loudness, this, bUseNoiseMakerLocation ? NoiseMaker->GetActorLocation() : NoiseLocation);
 }
+
+const struct FNavAgentProperties* APawn::GetNavAgentProperties() const
+{
+	return GetMovementComponent() ? GetMovementComponent()->GetNavAgentProperties() : NULL;
+}
diff --git a/Engine/Source/Runtime/Engine/Private/PawnMovementComponent.cpp b/Engine/Source/Runtime/Engine/Private/PawnMovementComponent.cpp
--- a/Engine/Source/Runtime/Engine/Private/PawnMovementComponent.cpp
+++ b/Engine/Source/Runtime/Engine/Private/PawnMovementComponent.cpp
@@ -1,6 +1,7 @@
 // Copyright 1998-2014 Epic Games, Inc. All Rights Reserved.
 
 #include "EnginePrivate.h"
+#include "GameFramework/PawnMovementComponent.h"
 
 DEFINE_LOG_CATEGORY_STATIC(LogPawnMovementComponent, Log, All);
 
diff --git a/Engine/Source/Runtime/Engine/Private/PhysicsEngine/RadialForceComponent.cpp b/Engine/Source/Runtime/Engine/Private/PhysicsEngine/RadialForceComponent.cpp
--- a/Engine/Source/Runtime/Engine/Private/PhysicsEngine/RadialForceComponent.cpp
+++ b/Engine/Source/Runtime/Engine/Private/PhysicsEngine/RadialForceComponent.cpp
@@ -1,6 +1,7 @@
 // Copyright 1998-2014 Epic Games, Inc. All Rights Reserved.
 
 #include "EnginePrivate.h"
+#include "GameFramework/MovementComponent.h"
 #include "Net/UnrealNetwork.h"
 
 //////////////////////////////////////////////////////////////////////////
diff --git a/Engine/Source/Runtime/Engine/Private/Vehicles/VehicleWheel.cpp b/Engine/Source/Runtime/Engine/Private/Vehicles/VehicleWheel.cpp
--- a/Engine/Source/Runtime/Engine/Private/Vehicles/VehicleWheel.cpp
+++ b/Engine/Source/Runtime/Engine/Private/Vehicles/VehicleWheel.cpp
@@ -1,6 +1,7 @@
 // Copyright 1998-2014 Epic Games, Inc. All Rights Reserved.
 
 #include "EnginePrivate.h"
+#include "GameFramework/PawnMovementComponent.h"
 #include "PhysicsPublic.h"
 #include "Vehicles/VehicleWheel.h"
 #include "Vehicles/WheeledVehicleMovementComponent.h"
diff --git a/Engine/Source/Runtime/Engine/Private/Vehicles/WheeledVehicleMovementComponent.cpp b/Engine/Source/Runtime/Engine/Private/Vehicles/WheeledVehicleMovementComponent.cpp
--- a/Engine/Source/Runtime/Engine/Private/Vehicles/WheeledVehicleMovementComponent.cpp
+++ b/Engine/Source/Runtime/Engine/Private/Vehicles/WheeledVehicleMovementComponent.cpp
@@ -2,6 +2,7 @@
 
 #include "EnginePrivate.h"
 #include "PhysicsPublic.h"
+#include "GameFramework/PawnMovementComponent.h"
 #include "Net/UnrealNetwork.h"
 #include "MessageLog.h"
 #include "Vehicles/VehicleWheel.h"
