Header cleanup: NavMeshRenderingHelpers.h

diff --git a/Engine/Source/Runtime/Engine/Private/NavMeshRenderingHelpers.h b/Engine/Source/Runtime/Engine/Private/NavMeshRenderingHelpers.h
--- a/Engine/Source/Runtime/Engine/Private/NavMeshRenderingHelpers.h
+++ b/Engine/Source/Runtime/Engine/Private/NavMeshRenderingHelpers.h
@@ -1,12 +1,9 @@
 // Copyright 1998-2014 Epic Games, Inc. All Rights Reserved.
 #pragma once
+#include "AI/Navigation/RecastHelpers.h"
 #include "DebugRenderSceneProxy.h"
 #include "DynamicMeshBuilder.h"
-#include "AI/Navigation/RecastHelpers.h"
 #include "BatchedElements.h"
-#include "AI/Navigation/NavMeshRenderingComponent.h"
-#include "Materials/Material.h"
-#include "MaterialShared.h"
 
 static const FColor NavMeshRenderColor_Recast_TriangleEdges(255,255,255);
 static const FColor NavMeshRenderColor_Recast_TileEdges(16,16,16,32);
