Header cleanup: SoundNodeAmbient.h

diff --git a/Engine/Source/Runtime/Engine/Classes/Sound/SoundNodeAmbient.h b/Engine/Source/Runtime/Engine/Classes/Sound/SoundNodeAmbient.h
--- a/Engine/Source/Runtime/Engine/Classes/Sound/SoundNodeAmbient.h
+++ b/Engine/Source/Runtime/Engine/Classes/Sound/SoundNodeAmbient.h
@@ -6,10 +6,7 @@
  */
 
 #pragma once
-
-#include "SoundNodeAttenuation.h"
 #include "Sound/SoundNodeDeprecated.h"
-
 #include "SoundNodeAmbient.generated.h"
 
 USTRUCT()
