UKismetSystemLibrary::K2_SetTimerDelegate blueprint_callable fucntion added.

FBlueprintTimerDynamicDelegate had to be introduced, because UTH was not aware about FTimerDynamicDelegate (because it is declared in other header).

diff --git a/Engine/Source/Runtime/Engine/Classes/Kismet/KismetSystemLibrary.h b/Engine/Source/Runtime/Engine/Classes/Kismet/KismetSystemLibrary.h
--- a/Engine/Source/Runtime/Engine/Classes/Kismet/KismetSystemLibrary.h
+++ b/Engine/Source/Runtime/Engine/Classes/Kismet/KismetSystemLibrary.h
@@ -5,6 +5,9 @@
 #include "Engine/LatentActionManager.h"
 #include "KismetSystemLibrary.generated.h"
 
+/** It's just a proxy for FTimerDynamicDelegate */
+DECLARE_DYNAMIC_DELEGATE(FBlueprintTimerDynamicDelegate);
+
 UENUM(BlueprintType)
 namespace EDrawDebugTrace
 {
@@ -235,6 +238,15 @@ class UKismetSystemLibrary : public UBlueprintFunctionLibrary
 	UFUNCTION(BlueprintCallable, meta=(FriendlyName = "SetTimer", DefaultToSelf = "Object"), Category="Utilities|Time")
 	static void K2_SetTimer(UObject* Object, FString FunctionName, float Time, bool bLooping);
 
+	/**
+	 * Set a timer to execute delegate. Setting an existing timer will reset that timer with updated parameters.
+	 * @param Delegate		Delegate. Can be a K2 function or a Custom Event.
+	 * @param Time			How long to wait before executing the delegate, in seconds. Setting a timer to <= 0 seconds will clear it if it is set.
+	 * @param bLooping		true to keep executing the delegate every Time seconds, false to execute delegate only once.
+	 */
+	UFUNCTION(BlueprintCallable, meta=(FriendlyName = "SetTimerDelegate"), Category="Utilities|Time")
+	static void K2_SetTimerDelegate(FBlueprintTimerDynamicDelegate Delegate, float Time, bool bLooping);
+
 	/**
 	 * Clears a set timer.
 	 * @param Object		Object that implements the delegate function. Defaults to self (this blueprint)
diff --git a/Engine/Source/Runtime/Engine/Private/KismetSystemLibrary.cpp b/Engine/Source/Runtime/Engine/Private/KismetSystemLibrary.cpp
--- a/Engine/Source/Runtime/Engine/Private/KismetSystemLibrary.cpp
+++ b/Engine/Source/Runtime/Engine/Private/KismetSystemLibrary.cpp
@@ -246,17 +246,26 @@ void UKismetSystemLibrary::K2_SetTimer(UObject* Object, FString FunctionName, fl
 		}
 	}
 
-	FTimerDynamicDelegate Delegate;
+	FBlueprintTimerDynamicDelegate Delegate;
 	Delegate.BindUFunction(Object, FunctionFName);
+	K2_SetTimerDelegate(Delegate, Time, bLooping);
+}
 
-	if (Delegate.IsBound())
+void UKismetSystemLibrary::K2_SetTimerDelegate(FBlueprintTimerDynamicDelegate Delegate, float Time, bool bLooping)
+{
+	UObject* const Object = Delegate.GetUObject();
+	FName const FunctionName = Delegate.GetFunctionName();
+
+	FTimerDynamicDelegate InnerDelegate;
+	InnerDelegate.BindUFunction(Object, FunctionName);
+	if (InnerDelegate.IsBound())
 	{
-		UWorld* const World = GEngine->GetWorldFromContextObject(Object);
-		World->GetTimerManager().SetTimer(Delegate, Time, bLooping);
+		const UWorld* const World = GEngine->GetWorldFromContextObject(Object);
+		World->GetTimerManager().SetTimer(InnerDelegate, Time, bLooping);
 	}
 	else
 	{
-		UE_LOG(LogBlueprintUserMessages, Warning, TEXT("SetTimer passed a bad function (%s) or object (%s)"), *FunctionName, *GetNameSafe(Object));
+		UE_LOG(LogBlueprintUserMessages, Warning, TEXT("SetTimer passed a bad function (%s) or object (%s)"), *FunctionName.ToString(), *GetNameSafe(Object));
 	}
 }
 
