#ue4 - Fix issues with ApplyVelocityBraking incorrectly accumulating changes when timestep is over the threshold.

diff --git a/Engine/Source/Runtime/Engine/Private/CharacterMovementComponent.cpp b/Engine/Source/Runtime/Engine/Private/CharacterMovementComponent.cpp
--- a/Engine/Source/Runtime/Engine/Private/CharacterMovementComponent.cpp
+++ b/Engine/Source/Runtime/Engine/Private/CharacterMovementComponent.cpp
@@ -2413,35 +2413,32 @@ void UCharacterMovementComponent::ApplyVelocityBraking(float DeltaTime, float Fr
 
 	Friction = FMath::Max(0.f, Friction);
 	const FVector OldVel = Velocity;
-	FVector SumVel = FVector::ZeroVector;
 
 	// subdivide braking to get reasonably consistent results at lower frame rates
 	// (important for packet loss situations w/ networking)
 	float RemainingTime = DeltaTime;
-	const float TimeStep = 0.03f;
+	const float TimeStep = (1.0f / 33.0f);
 
 	// Decelerate to brake to a stop
 	const FVector RevAccel = (-1.f * FMath::Max(0.f,BrakingDeceleration)) * SafeNormalPrecise(Velocity);
 	while( RemainingTime >= MIN_TICK_TIME )
 	{
-		float dt = ((RemainingTime > TimeStep) ? TimeStep : RemainingTime);
+		const float dt = ((RemainingTime > TimeStep) ? TimeStep : RemainingTime);
 		RemainingTime -= dt;
 
 		// apply friction and braking
 		Velocity = Velocity + ((-2.f * Friction) * Velocity + RevAccel) * dt ; 
-		if( (Velocity | OldVel) > 0.f )
-		{
-			SumVel += Velocity * (dt / DeltaTime);
-		}
-		else
+		
+		// Don't reverse direction
+		if ((Velocity | OldVel) <= 0.f)
 		{
-			break;
+			Velocity = FVector::ZeroVector;
+			return;
 		}
 	}
-	Velocity = SumVel;
 
-	// brake to a stop, not backwards
-	if( ((OldVel | Velocity) < 0.f)	|| (Velocity.SizeSquared() <= FMath::Square(BRAKE_TO_STOP_VELOCITY)) )
+	// Clamp to zero if below min threshold.
+	if ((Velocity.SizeSquared() <= FMath::Square(BRAKE_TO_STOP_VELOCITY)) )
 	{
 		Velocity = FVector::ZeroVector;
 	}
