Max depenetration velocity

diff --git a/Engine/Source/Runtime/Engine/Classes/PhysicsEngine/BodyInstance.h b/Engine/Source/Runtime/Engine/Classes/PhysicsEngine/BodyInstance.h
--- a/Engine/Source/Runtime/Engine/Classes/PhysicsEngine/BodyInstance.h
+++ b/Engine/Source/Runtime/Engine/Classes/PhysicsEngine/BodyInstance.h
@@ -396,6 +396,8 @@ public:
 	void SetAngularVelocity(const FVector& NewAngVel, bool bAddToCurrent);
 	/** Set the maximum angular velocity of this body */
 	void SetMaxAngularVelocity(float NewMaxAngVel, bool bAddToCurrent);
+	/** Set the maximum depenetration velocity the physics simulation will introduce */
+	void SetMaxDepenetrationVelocity(float MaxVelocity);
 	/** Set whether we should get a notification about physics collisions */
 	void SetInstanceNotifyRBCollision(bool bNewNotifyCollision);
 	/** Enables/disables whether this body is affected by gravity. */
diff --git a/Engine/Source/Runtime/Engine/Classes/PhysicsEngine/PhysicsSettings.h b/Engine/Source/Runtime/Engine/Classes/PhysicsEngine/PhysicsSettings.h
--- a/Engine/Source/Runtime/Engine/Classes/PhysicsEngine/PhysicsSettings.h
+++ b/Engine/Source/Runtime/Engine/Classes/PhysicsEngine/PhysicsSettings.h
@@ -99,6 +99,11 @@ class ENGINE_API UPhysicsSettings
 	UPROPERTY(config, EditAnywhere, Category=Simulation)
 	TEnumAsByte<EFrictionCombineMode::Type> FrictionCombineMode;
 
+	/** Max velocity which may be used to depenetrate simulated physics objects. 0 means no maximum. */
+	UPROPERTY(config, EditAnywhere, Category = Simulation)
+	float MaxDepenetrationVelocity;
+
+
 	/** Max Physics Delta Time to be clamped. */
 	UPROPERTY(config, EditAnywhere, meta=(ClampMin="0.0013", UIMin = "0.0013", ClampMax="1.0", UIMax="1.0"), Category=Framerate)
 	float MaxPhysicsDeltaTime;
diff --git a/Engine/Source/Runtime/Engine/Private/PhysicsEngine/BodyInstance.cpp b/Engine/Source/Runtime/Engine/Private/PhysicsEngine/BodyInstance.cpp
--- a/Engine/Source/Runtime/Engine/Private/PhysicsEngine/BodyInstance.cpp
+++ b/Engine/Source/Runtime/Engine/Private/PhysicsEngine/BodyInstance.cpp
@@ -1088,6 +1088,8 @@ void FBodyInstance::InitBody(UBodySetup* Setup, const FTransform& Transform, UPr
 
 			SetMaxAngularVelocity(MaxAngularVelocity, false);
 
+			SetMaxDepenetrationVelocity(UPhysicsSettings::Get()->MaxDepenetrationVelocity);
+
 			//@TODO: BOX2D: Determine if sleep threshold and solver settings can be configured per-body or not
 #if 0
 			// Set the parameters for determining when to put the object to sleep.
@@ -1283,6 +1285,8 @@ void FBodyInstance::InitBody(UBodySetup* Setup, const FTransform& Transform, UPr
 
 		SetMaxAngularVelocity(MaxAngularVelocity, false);
 
+		SetMaxDepenetrationVelocity(UPhysicsSettings::Get()->MaxDepenetrationVelocity);
+
 		// Set initial velocity 
 		if(bUseSimulate)
 		{
@@ -2734,6 +2738,18 @@ void FBodyInstance::SetMaxAngularVelocity(float NewMaxAngVel, bool bAddToCurrent
 	//@TODO: BOX2D: Implement SetMaxAngularVelocity
 }
 
+void FBodyInstance::SetMaxDepenetrationVelocity(float MaxVelocity)
+{
+#if WITH_PHYSX
+	float UseMaxVelocity = MaxVelocity == 0.f ? PX_MAX_F32 : MaxVelocity;
+
+	if (PxRigidDynamic* PRigidDynamic = GetPxRigidDynamic())
+	{
+		PRigidDynamic->setMaxDepenetrationVelocity(MaxVelocity);
+	}
+#endif
+}
+
 
 void FBodyInstance::AddForce(const FVector& Force, bool bAllowSubstepping)
 {
