iOS CIS fix

diff --git a/Engine/Source/Runtime/Engine/Private/UserInterface/PlayerInput.cpp b/Engine/Source/Runtime/Engine/Private/UserInterface/PlayerInput.cpp
--- a/Engine/Source/Runtime/Engine/Private/UserInterface/PlayerInput.cpp
+++ b/Engine/Source/Runtime/Engine/Private/UserInterface/PlayerInput.cpp
@@ -551,8 +551,8 @@ void UPlayerInput::GetChordsForAction(const FInputActionBinding& ActionBinding,
 				// test modifier conditions and ignore the event if they failed
 				if (	(KeyMapping.bAlt == false || IsAltPressed())
 					&&	(KeyMapping.bCtrl == false || IsCtrlPressed())
-					&&	(KeyMapping.bShift == false || IsShiftPressed()
-					&&	(KeyMapping.bCmd == false || IsCmdPressed()))
+					&&	(KeyMapping.bShift == false || IsShiftPressed())
+					&&	(KeyMapping.bCmd == false || IsCmdPressed())
 					&& 	KeyEventOccurred(KeyMapping.Key, ActionBinding.KeyEvent, EventIndices))
 				{
 					bool bAddDelegate = true;
