#ue4 - Header cleanup: Character.h, CharacterMovement.h

diff --git a/Engine/Source/Runtime/Engine/Classes/Engine/EngineTypes.h b/Engine/Source/Runtime/Engine/Classes/Engine/EngineTypes.h
--- a/Engine/Source/Runtime/Engine/Classes/Engine/EngineTypes.h
+++ b/Engine/Source/Runtime/Engine/Classes/Engine/EngineTypes.h
@@ -188,6 +188,19 @@ enum ETriangleSortAxis
 	TSA_MAX,
 };
 
+/** Movement modes for Characters.  */
+UENUM(BlueprintType)
+enum EMovementMode
+{
+	MOVE_None		UMETA(DisplayName="None"),
+	MOVE_Walking	UMETA(DisplayName="Walking"),
+	MOVE_Falling	UMETA(DisplayName="Falling"),
+	MOVE_Swimming	UMETA(DisplayName="Swimming"),
+	MOVE_Flying		UMETA(DisplayName="Flying"),
+	MOVE_Custom		UMETA(DisplayName="Custom"),
+	MOVE_MAX		UMETA(Hidden),
+};
+
 
 /** 
  * Enum indicating different type of objects for rigid-body collision purposes. 
diff --git a/Engine/Source/Runtime/Engine/Classes/GameFramework/Character.h b/Engine/Source/Runtime/Engine/Classes/GameFramework/Character.h
--- a/Engine/Source/Runtime/Engine/Classes/GameFramework/Character.h
+++ b/Engine/Source/Runtime/Engine/Classes/GameFramework/Character.h
@@ -2,10 +2,13 @@
 
 #pragma once
 #include "GameFramework/Pawn.h"
-#include "GameFramework/CharacterMovementComponent.h"
-#include "Components/PrimitiveComponent.h"
+#include "Animation/AnimationAsset.h"
 #include "Character.generated.h"
 
+class UPawnMovementComponent;
+class UCharacterMovementComponent;
+class UPrimitiveComponent;
+
 DECLARE_DYNAMIC_MULTICAST_DELEGATE_ThreeParams(FMovementModeChangedSignature, class ACharacter*, Character, EMovementMode, PrevMovementMode, uint8, PreviousCustomMode);
 DECLARE_DYNAMIC_MULTICAST_DELEGATE_ThreeParams(FCharacterMovementUpdatedSignature, float, DeltaSeconds, FVector, OldLocation, FVector, OldVelocity);
 //
@@ -82,10 +85,7 @@ struct FSimulatedRootMotionReplicatedMove
 namespace MovementBaseUtility
 {
 	/** Determine whether MovementBase can possibly move. */
-	FORCEINLINE bool IsDynamicBase(const UPrimitiveComponent* MovementBase)
-	{
-		return (MovementBase && MovementBase->Mobility == EComponentMobility::Movable);
-	}
+	ENGINE_API bool IsDynamicBase(const UPrimitiveComponent* MovementBase);
 
 	/** Determine if we should use relative positioning when based on a component (because it may move). */
 	FORCEINLINE bool UseRelativeLocation(const UPrimitiveComponent* MovementBase)
@@ -339,7 +339,7 @@ public:
 
 	// Begin APawn Interface.
 	virtual void PostInitializeComponents() override;
-	virtual class UPawnMovementComponent* GetMovementComponent() const override { return CharacterMovement; }
+	virtual UPawnMovementComponent* GetMovementComponent() const override;
 	virtual UPrimitiveComponent* GetMovementBase() const override final { return BasedMovement.MovementBase; }
 	virtual float GetDefaultHalfHeight() const override;
 	virtual void TurnOff() override;
diff --git a/Engine/Source/Runtime/Engine/Classes/GameFramework/CharacterMovementComponent.h b/Engine/Source/Runtime/Engine/Classes/GameFramework/CharacterMovementComponent.h
--- a/Engine/Source/Runtime/Engine/Classes/GameFramework/CharacterMovementComponent.h
+++ b/Engine/Source/Runtime/Engine/Classes/GameFramework/CharacterMovementComponent.h
@@ -19,19 +19,6 @@ struct FVector_NetQuantize100;
 struct FVector_NetQuantizeNormal;
 class FDebugDisplayInfo;
 
-/** Movement modes for Characters.  */
-UENUM(BlueprintType)
-enum EMovementMode
-{
-	MOVE_None		UMETA(DisplayName="None"),
-	MOVE_Walking	UMETA(DisplayName="Walking"),
-	MOVE_Falling	UMETA(DisplayName="Falling"),
-	MOVE_Swimming	UMETA(DisplayName="Swimming"),
-	MOVE_Flying		UMETA(DisplayName="Flying"),
-	MOVE_Custom		UMETA(DisplayName="Custom"),
-	MOVE_MAX		UMETA(Hidden),
-};
-
 // Data about the floor for walking movement.
 USTRUCT(BlueprintType)
 struct ENGINE_API FFindFloorResult
diff --git a/Engine/Source/Runtime/Engine/Classes/GameFramework/MovementComponent.h b/Engine/Source/Runtime/Engine/Classes/GameFramework/MovementComponent.h
--- a/Engine/Source/Runtime/Engine/Classes/GameFramework/MovementComponent.h
+++ b/Engine/Source/Runtime/Engine/Classes/GameFramework/MovementComponent.h
@@ -5,9 +5,10 @@
  */
 
 #pragma once
-#include "Components/SceneComponent.h"
+#include "Components/ActorComponent.h"
 #include "MovementComponent.generated.h"
 
+class UPrimitiveComponent;
 struct FCollisionQueryParams;
 struct FCollisionResponseParams;
 struct FCollisionShape;
@@ -19,7 +20,7 @@ class ENGINE_API UMovementComponent : public UActorComponent
 
 	/** The component we move and update. **/
 	UPROPERTY(BlueprintReadOnly, Category=MovementComponent)
-	class UPrimitiveComponent* UpdatedComponent;
+	UPrimitiveComponent* UpdatedComponent;
 
 	/**
 	 * Flags that control the behavior of calls to MoveComponent() on our UpdatedComponent.
diff --git a/Engine/Source/Runtime/Engine/Classes/GameFramework/NavMovementComponent.h b/Engine/Source/Runtime/Engine/Classes/GameFramework/NavMovementComponent.h
--- a/Engine/Source/Runtime/Engine/Classes/GameFramework/NavMovementComponent.h
+++ b/Engine/Source/Runtime/Engine/Classes/GameFramework/NavMovementComponent.h
@@ -9,7 +9,6 @@
 #include "AI/Navigation/NavigationSystem.h"
 #include "AI/Navigation/NavAgentInterface.h"
 #include "AI/Navigation/NavigationTypes.h"
-#include "Components/PrimitiveComponent.h"
 #include "GameFramework/MovementComponent.h"
 #include "NavMovementComponent.generated.h"
 
diff --git a/Engine/Source/Runtime/Engine/Classes/Kismet/GameplayStatics.h b/Engine/Source/Runtime/Engine/Classes/Kismet/GameplayStatics.h
--- a/Engine/Source/Runtime/Engine/Classes/Kismet/GameplayStatics.h
+++ b/Engine/Source/Runtime/Engine/Classes/Kismet/GameplayStatics.h
@@ -4,6 +4,7 @@
 
 #include "Sound/DialogueTypes.h"
 #include "Engine/LatentActionManager.h"
+#include "Kismet/BlueprintFunctionLibrary.h"
 #include "GameplayStatics.generated.h"
 
 //
diff --git a/Engine/Source/Runtime/Engine/Private/Character.cpp b/Engine/Source/Runtime/Engine/Private/Character.cpp
--- a/Engine/Source/Runtime/Engine/Private/Character.cpp
+++ b/Engine/Source/Runtime/Engine/Private/Character.cpp
@@ -7,6 +7,8 @@
 
 #include "EnginePrivate.h"
 #include "GameFramework/Character.h"
+#include "GameFramework/CharacterMovementComponent.h"
+#include "Components/PrimitiveComponent.h"
 #include "Net/UnrealNetwork.h"
 #include "DisplayDebugHelpers.h"
 #include "Animation/AnimInstance.h"
@@ -118,6 +120,11 @@ void ACharacter::PostInitializeComponents()
 	}
 }
 
+UPawnMovementComponent* ACharacter::GetMovementComponent() const
+{
+	return CharacterMovement;
+}
+
 
 void ACharacter::SetupPlayerInputComponent(UInputComponent* InputComponent)
 {
@@ -369,6 +376,11 @@ void ACharacter::ClearCrossLevelReferences()
 
 namespace MovementBaseUtility
 {
+	bool IsDynamicBase(const UPrimitiveComponent* MovementBase)
+	{
+		return (MovementBase && MovementBase->Mobility == EComponentMobility::Movable);
+	}
+
 	void AddTickDependency(FTickFunction& BasedObjectTick, UPrimitiveComponent* NewBase)
 	{
 		if (NewBase && MovementBaseUtility::UseRelativeLocation(NewBase))
diff --git a/Engine/Source/Runtime/Engine/Private/CharacterMovementComponent.cpp b/Engine/Source/Runtime/Engine/Private/CharacterMovementComponent.cpp
--- a/Engine/Source/Runtime/Engine/Private/CharacterMovementComponent.cpp
+++ b/Engine/Source/Runtime/Engine/Private/CharacterMovementComponent.cpp
@@ -9,6 +9,7 @@
 #include "GameFramework/PhysicsVolume.h"
 #include "GameFramework/GameNetworkManager.h"
 #include "GameFramework/Character.h"
+#include "Components/PrimitiveComponent.h"
 #include "Animation/AnimMontage.h"
 
 // @todo this is here only due to circular dependency to AIModule. To be removed
diff --git a/Engine/Source/Runtime/Engine/Private/MovementComponent.cpp b/Engine/Source/Runtime/Engine/Private/MovementComponent.cpp
--- a/Engine/Source/Runtime/Engine/Private/MovementComponent.cpp
+++ b/Engine/Source/Runtime/Engine/Private/MovementComponent.cpp
@@ -4,6 +4,7 @@
 #include "EnginePrivate.h"
 #include "GameFramework/MovementComponent.h"
 #include "GameFramework/PhysicsVolume.h"
+#include "Components/SceneComponent.h"
 #include "MessageLog.h"
 
 #define LOCTEXT_NAMESPACE "MovementComponent"
diff --git a/Engine/Source/Runtime/Engine/Private/PlayerController.cpp b/Engine/Source/Runtime/Engine/Private/PlayerController.cpp
--- a/Engine/Source/Runtime/Engine/Private/PlayerController.cpp
+++ b/Engine/Source/Runtime/Engine/Private/PlayerController.cpp
@@ -6,6 +6,8 @@
 #include "Engine/WorldComposition.h"
 #include "GameFramework/GameNetworkManager.h"
 #include "GameFramework/Character.h"
+#include "Interfaces/NetworkPredictionInterface.h"
+#include "GameFramework/CharacterMovementComponent.h"
 #include "ConfigCacheIni.h"
 #include "SoundDefinitions.h"
 #include "OnlineSubsystemUtils.h"
