341709	LIVE: CRASH: UE4: Placing a user created asset into a level is causing 4.3 to crash

diff --git a/Engine/Source/Runtime/Engine/Private/Animation/AnimSingleNodeInstance.cpp b/Engine/Source/Runtime/Engine/Private/Animation/AnimSingleNodeInstance.cpp
--- a/Engine/Source/Runtime/Engine/Private/Animation/AnimSingleNodeInstance.cpp
+++ b/Engine/Source/Runtime/Engine/Private/Animation/AnimSingleNodeInstance.cpp
@@ -57,6 +57,7 @@ void UAnimSingleNodeInstance::SetAnimationAsset(class UAnimationAsset* NewAsset,
 
 	bLooping = bIsLooping;
 	PlayRate = InPlayRate;
+	CurrentTime = 0.f;
 #if WITH_EDITORONLY_DATA
 	PreviewPoseCurrentTime = 0.0f;
 #endif
