Fix crash when apex module doesn't own PxActor of destructible component

diff --git a/Engine/Source/Runtime/Engine/Private/PhysicsEngine/PhysScene.cpp b/Engine/Source/Runtime/Engine/Private/PhysicsEngine/PhysScene.cpp
--- a/Engine/Source/Runtime/Engine/Private/PhysicsEngine/PhysScene.cpp
+++ b/Engine/Source/Runtime/Engine/Private/PhysicsEngine/PhysScene.cpp
@@ -676,7 +676,7 @@ void FPhysScene::SyncComponentsToBodies(uint32 SceneType)
 		//Special code for destructible chunk
 		if (const FDestructibleChunkInfo * DestructibleChunkInfo = FPhysxUserData::Get<FDestructibleChunkInfo>(RigidActor->userData))
 		{
-			if (DestructibleChunkInfo->OwningComponent.IsValid())
+			if (GApexModuleDestructible->owns(RigidActor) && DestructibleChunkInfo->OwningComponent.IsValid())
 			{
 				//TODO: waiting on new API to avoid duplicate updates per shape.
 				TArray<PxShape*> Shapes;
@@ -686,16 +686,18 @@ void FPhysScene::SyncComponentsToBodies(uint32 SceneType)
 				{
 					PxShape* Shape = Shapes[ShapeIdx];
 					int32 ChunkIndex;
-					NxDestructibleActor* DestructibleActor = GApexModuleDestructible->getDestructibleAndChunk(Shape, &ChunkIndex);
-					const physx::PxMat44 ChunkPoseRT = DestructibleActor->getChunkPose(ChunkIndex);
-					const physx::PxTransform Transform(ChunkPoseRT);
-					if (UDestructibleComponent * DestructibleComponent = Cast<UDestructibleComponent>(FPhysxUserData::Get<UPrimitiveComponent>(DestructibleActor->userData)))
+					if (NxDestructibleActor* DestructibleActor = GApexModuleDestructible->getDestructibleAndChunk(Shape, &ChunkIndex))
 					{
-						// if this component was already unregistered, transform data were deallocated
-						// transform data become empty when a piece of destructible actor goes out from a streaming volume
-						if (DestructibleComponent->IsRegistered())
+						const physx::PxMat44 ChunkPoseRT = DestructibleActor->getChunkPose(ChunkIndex);
+						const physx::PxTransform Transform(ChunkPoseRT);
+						if (UDestructibleComponent * DestructibleComponent = Cast<UDestructibleComponent>(FPhysxUserData::Get<UPrimitiveComponent>(DestructibleActor->userData)))
 						{
-							DestructibleComponent->SetChunkWorldRT(ChunkIndex, P2UQuat(Transform.q), P2UVector(Transform.p));
+							// if this component was already unregistered, transform data were deallocated
+							// transform data become empty when a piece of destructible actor goes out from a streaming volume
+							if (DestructibleComponent->IsRegistered())
+							{
+								DestructibleComponent->SetChunkWorldRT(ChunkIndex, P2UQuat(Transform.q), P2UVector(Transform.p));
+							}
 						}
 					}
 				}
