#ue4 - Minor optimization to use SetActorLocationAndRotation() rather than two separate calls in PlayerCameraManager (which causes multiple transform updates).

diff --git a/Engine/Source/Runtime/Engine/Private/PlayerCameraManager.cpp b/Engine/Source/Runtime/Engine/Private/PlayerCameraManager.cpp
--- a/Engine/Source/Runtime/Engine/Private/PlayerCameraManager.cpp
+++ b/Engine/Source/Runtime/Engine/Private/PlayerCameraManager.cpp
@@ -569,8 +569,7 @@ void APlayerCameraManager::UpdateViewTarget(FTViewTarget& OutVT, float DeltaTime
 	}
 
 	// Synchronize the actor with the view target results
-	SetActorLocation(OutVT.POV.Location, false);
-	SetActorRotation(OutVT.POV.Rotation);
+	SetActorLocationAndRotation(OutVT.POV.Location, OutVT.POV.Rotation, false);
 }
 
 
