- CIS build error missing #includes
- Fixed losing joystick textures when loading a bad level

diff --git a/Engine/Source/Runtime/Engine/Classes/GameFramework/PlayerController.h b/Engine/Source/Runtime/Engine/Classes/GameFramework/PlayerController.h
--- a/Engine/Source/Runtime/Engine/Classes/GameFramework/PlayerController.h
+++ b/Engine/Source/Runtime/Engine/Classes/GameFramework/PlayerController.h
@@ -940,6 +940,11 @@ protected:
 	/** The virtual touch interface */
 	TSharedPtr<class SVirtualJoystick> VirtualJoystick;
 
+	/** The currently set touch interface */
+	UPROPERTY()
+	class UTouchInterface* CurrentTouchInterface;
+
+
 public:
 	/** Adds an inputcomponent to the top of the input stack. */
 	void PushInputComponent(UInputComponent* Input);
diff --git a/Engine/Source/Runtime/Engine/Private/Controller.cpp b/Engine/Source/Runtime/Engine/Private/Controller.cpp
--- a/Engine/Source/Runtime/Engine/Private/Controller.cpp
+++ b/Engine/Source/Runtime/Engine/Private/Controller.cpp
@@ -10,6 +10,8 @@
 #include "NetworkingDistanceConstants.h"
 #include "VisualLog.h"
 #include "GameFramework/Character.h"
+#include "GameFramework/PawnMovementComponent.h"
+#include "GameFramework/CharacterMovementComponent.h"
 
 // @todo this is here only due to circular dependency to AIModule. To be removed
 #include "Navigation/PathFollowingComponent.h"
diff --git a/Engine/Source/Runtime/Engine/Private/PlayerController.cpp b/Engine/Source/Runtime/Engine/Private/PlayerController.cpp
--- a/Engine/Source/Runtime/Engine/Private/PlayerController.cpp
+++ b/Engine/Source/Runtime/Engine/Private/PlayerController.cpp
@@ -22,6 +22,7 @@
 #include "GameFramework/SpectatorPawn.h"
 #include "GameFramework/HUD.h"
 #include "ContentStreaming.h"
+#include "GameFramework/PawnMovementComponent.h"
 
 DEFINE_LOG_CATEGORY_STATIC(LogPlayerController, Log, All);
 
@@ -923,6 +924,12 @@ void APlayerController::CreateTouchInterface()
 	// do we want to show virtual joysticks?
 	if (LocalPlayer && LocalPlayer->ViewportClient && SVirtualJoystick::ShouldDisplayTouchInterface())
 	{
+		// in case we already had one, remove it
+		if (VirtualJoystick.IsValid())
+		{
+			Cast<ULocalPlayer>(Player)->ViewportClient->RemoveViewportWidgetContent(VirtualJoystick.ToSharedRef());
+		}
+
 		// load what the game wants to show at startup
 		FStringAssetReference DefaultTouchInterfaceName = GetDefault<UInputSettings>()->DefaultTouchInterface;
 
@@ -938,7 +945,7 @@ void APlayerController::CreateTouchInterface()
 			UTouchInterface* DefaultTouchInterface = LoadObject<UTouchInterface>(NULL, *DefaultTouchInterfaceName.ToString());
 			if (DefaultTouchInterface != NULL)
 			{
-				DefaultTouchInterface->Activate(VirtualJoystick);
+				ActivateTouchInterface(DefaultTouchInterface);
 			}
 		}
 	}
@@ -3948,6 +3955,7 @@ void APlayerController::ActivateTouchInterface(UTouchInterface* NewTouchInterfac
 	{
 		NewTouchInterface->Activate(VirtualJoystick);
 	}
+	CurrentTouchInterface = NewTouchInterface;
 }
 
 void APlayerController::UpdateCameraManager(float DeltaSeconds)
