Save the animation notifies to be in the order of track index if it's in the same time.

diff --git a/Engine/Source/Runtime/Engine/Private/Animation/AnimSequenceBase.cpp b/Engine/Source/Runtime/Engine/Private/Animation/AnimSequenceBase.cpp
--- a/Engine/Source/Runtime/Engine/Private/Animation/AnimSequenceBase.cpp
+++ b/Engine/Source/Runtime/Engine/Private/Animation/AnimSequenceBase.cpp
@@ -209,9 +209,28 @@ void UAnimSequenceBase::SortNotifies()
 	{
 		FORCEINLINE bool operator()( const FAnimNotifyEvent& A, const FAnimNotifyEvent& B ) const
 		{
-			return A.GetTriggerTime() < B.GetTriggerTime();
+			float ATime = A.GetTriggerTime();
+			float BTime = B.GetTriggerTime();
+
+#if WITH_EDITORONLY_DATA
+			// this sorting only works if it's saved in editor or loaded with editor data
+			// this is required for gameplay team to have reliable order of notifies
+			// but it was noted that this change will require to resave notifies. 
+			// once you resave, this order will be preserved
+			if(FMath::IsNearlyEqual(ATime, BTime))
+			{
+
+				// if the 2 anim notify events are the same display time sort based off of track index
+				return A.TrackIndex < B.TrackIndex;
+			}
+			else
+#endif // WITH_EDITORONLY_DATA
+			{
+				return ATime < BTime;
+			}
 		}
 	};
+
 	Notifies.Sort( FCompareFAnimNotifyEvent() );
 }
 
@@ -386,6 +405,8 @@ void UAnimSequenceBase::UpdateAnimNotifyTrackCache()
 		}
 	}
 
+	SortNotifies();
+
 	// notification broadcast
 	OnNotifyChanged.Broadcast();
 }
