CIS fix.

diff --git a/Engine/Source/Runtime/Engine/Private/CharacterMovementComponent.cpp b/Engine/Source/Runtime/Engine/Private/CharacterMovementComponent.cpp
--- a/Engine/Source/Runtime/Engine/Private/CharacterMovementComponent.cpp
+++ b/Engine/Source/Runtime/Engine/Private/CharacterMovementComponent.cpp
@@ -9,6 +9,7 @@
 #include "GameFramework/PhysicsVolume.h"
 #include "GameFramework/GameNetworkManager.h"
 #include "GameFramework/Character.h"
+#include "GameFramework/CharacterMovementComponent.h"
 #include "Components/PrimitiveComponent.h"
 #include "Animation/AnimMontage.h"
 
