CIS fix.

diff --git a/Engine/Source/Runtime/Engine/Private/AI/Navigation/AvoidanceManager.cpp b/Engine/Source/Runtime/Engine/Private/AI/Navigation/AvoidanceManager.cpp
--- a/Engine/Source/Runtime/Engine/Private/AI/Navigation/AvoidanceManager.cpp
+++ b/Engine/Source/Runtime/Engine/Private/AI/Navigation/AvoidanceManager.cpp
@@ -3,6 +3,7 @@
 #include "EnginePrivate.h"
 #include "AI/Navigation/AvoidanceManager.h"
 #include "GameFramework/Character.h"
+#include "GameFramework/CharacterMovementComponent.h"
 
 DEFINE_STAT(STAT_AI_ObstacleAvoidance);
 
