Fix Landscape Spline Control Point Mesh mobility to be static and collision profile to be BlockAll (worldstatic) by default.

diff --git a/Engine/Source/Runtime/Engine/Private/Landscape/LandscapeSplines.cpp b/Engine/Source/Runtime/Engine/Private/Landscape/LandscapeSplines.cpp
--- a/Engine/Source/Runtime/Engine/Private/Landscape/LandscapeSplines.cpp
+++ b/Engine/Source/Runtime/Engine/Private/Landscape/LandscapeSplines.cpp
@@ -463,6 +463,9 @@ void ULandscapeSplinesComponent::ShowSplineEditorMesh(bool bShow)
 UControlPointMeshComponent::UControlPointMeshComponent(const class FPostConstructInitializeProperties& PCIP)
 	: Super(PCIP)
 {
+	SetCollisionProfileName(UCollisionProfile::BlockAll_ProfileName);
+	Mobility = EComponentMobility:Static;
+
 #if WITH_EDITORONLY_DATA
 	bSelected = false;
 #endif
