Added 3 new Blueprint-callable functions to InstancedStaticMesh: GetInstanceTransform, UpdateInstanceTransform and GetInstanceCount

diff --git a/Engine/Source/Runtime/Engine/Classes/Components/InstancedStaticMeshComponent.h b/Engine/Source/Runtime/Engine/Classes/Components/InstancedStaticMeshComponent.h
--- a/Engine/Source/Runtime/Engine/Classes/Components/InstancedStaticMeshComponent.h
+++ b/Engine/Source/Runtime/Engine/Classes/Components/InstancedStaticMeshComponent.h
@@ -86,13 +86,25 @@ class ENGINE_API UInstancedStaticMeshComponent : public UStaticMeshComponent
 
 	/** Add an instance to this component. Transform is given in world space. */
 	UFUNCTION(BlueprintCallable, Category = "Components|InstancedStaticMesh")
-	 void AddInstanceWorldSpace(const FTransform& WorldTransform);
+	void AddInstanceWorldSpace(const FTransform& WorldTransform);
 
-	virtual bool ShouldCreatePhysicsState() const;
+	/** Get the transform for the instance specified. Instance is returned in local space of this component unless bWorldSpace is set.  Returns True on success. */
+	UFUNCTION(BlueprintCallable, Category = "Components|InstancedStaticMesh")
+	bool GetInstanceTransform(int32 InstanceIndex, FTransform& OutInstanceTransform, bool bWorldSpace = false) const;
+	
+	/** Update the transform for the instance specified. Instance is given in local space of this component unless bWorldSpace is set.  Returns True on success. */
+	UFUNCTION(BlueprintCallable, Category = "Components|InstancedStaticMesh")
+	bool UpdateInstanceTransform(int32 InstanceIndex, const FTransform& NewInstanceTransform, bool bWorldSpace=false);
 
 	/** Clear all instances being rendered by this component */
 	UFUNCTION(BlueprintCallable, Category="Components|InstancedStaticMesh")
 	void ClearInstances();
+	
+	/** Get the number of instances in this component */
+	UFUNCTION(BlueprintCallable, Category = "Components|InstancedStaticMesh")
+	int32 GetInstanceCount() const;
+	
+	virtual bool ShouldCreatePhysicsState() const;
 
 public:
 #if WITH_EDITOR
diff --git a/Engine/Source/Runtime/Engine/Private/InstancedStaticMesh.cpp b/Engine/Source/Runtime/Engine/Private/InstancedStaticMesh.cpp
--- a/Engine/Source/Runtime/Engine/Private/InstancedStaticMesh.cpp
+++ b/Engine/Source/Runtime/Engine/Private/InstancedStaticMesh.cpp
@@ -1463,12 +1463,58 @@ void UInstancedStaticMeshComponent::AddInstanceWorldSpace(const FTransform& Worl
 	AddInstance(RelativeTM);
 }
 
+bool UInstancedStaticMeshComponent::GetInstanceTransform(int32 InstanceIndex, FTransform& OutInstanceTransform, bool bWorldSpace) const
+{
+	if (!PerInstanceSMData.IsValidIndex(InstanceIndex))
+	{
+		return false;
+	}
+
+	const FInstancedStaticMeshInstanceData& InstanceData = PerInstanceSMData[InstanceIndex];
+
+	OutInstanceTransform = FTransform(InstanceData.Transform);
+	if (bWorldSpace)
+	{
+		OutInstanceTransform = OutInstanceTransform * ComponentToWorld;
+	}
+
+	return true;
+}
+
+bool UInstancedStaticMeshComponent::UpdateInstanceTransform(int32 InstanceIndex, const FTransform& NewInstanceTransform, bool bWorldSpace)
+{
+	if (!PerInstanceSMData.IsValidIndex(InstanceIndex))
+	{
+		return false;
+	}
+
+	FInstancedStaticMeshInstanceData& InstanceData = PerInstanceSMData[InstanceIndex];
+
+	// Render data uses local transform of the instance
+	FTransform LocalTransform = bWorldSpace ? NewInstanceTransform.GetRelativeTransform(ComponentToWorld) : NewInstanceTransform;
+	InstanceData.Transform = LocalTransform.ToMatrixWithScale();
+
+	if (bPhysicsStateCreated)
+	{
+		// Physics uses world transform of the instance
+		FTransform WorldTransform = bWorldSpace ? NewInstanceTransform : (LocalTransform * ComponentToWorld);
+		FBodyInstance* InstanceBodyInstance = InstanceBodies[InstanceIndex];
+#if WITH_PHYSX
+		// Update transform.
+		InstanceBodyInstance->SetBodyTransform(WorldTransform, false);
+#endif //WITH_PHYSX
+	}
+
+	MarkRenderStateDirty();
+
+	return true;
+}
+
 bool UInstancedStaticMeshComponent::ShouldCreatePhysicsState() const
 {
 	return IsRegistered() && (bAlwaysCreatePhysicsState || IsCollisionEnabled());
 }
 
-
 void UInstancedStaticMeshComponent::ClearInstances()
 {
 	// Clear all the per-instance data
@@ -1480,6 +1526,10 @@ void UInstancedStaticMeshComponent::ClearInstances()
 	MarkRenderStateDirty();
 }
 
+int32 UInstancedStaticMeshComponent::GetInstanceCount() const
+{
+	return PerInstanceSMData.Num();
+}
 
 void UInstancedStaticMeshComponent::SetupNewInstanceData(FInstancedStaticMeshInstanceData& InOutNewInstanceData, int32 InInstanceIndex, const FTransform& InInstanceTransform)
 {
