TTP 311634 - UE4: RENDERING: Warn when precomputed visibility is not enabled in ES2

Added a message warning that precomputed visibility is disabled when in ES2 mode

diff --git a/Engine/Source/Runtime/Engine/Private/UnrealEngine.cpp b/Engine/Source/Runtime/Engine/Private/UnrealEngine.cpp
--- a/Engine/Source/Runtime/Engine/Private/UnrealEngine.cpp
+++ b/Engine/Source/Runtime/Engine/Private/UnrealEngine.cpp
@@ -6392,6 +6392,14 @@ void DrawStatsHUD( UWorld* World, FViewport* Viewport, FCanvas* Canvas, UCanvas*
 				MessageY += 20;
 			}
 
+			if (World->FeatureLevel == ERHIFeatureLevel::ES2 && !World->GetWorldSettings()->bPrecomputeVisibility)
+			{
+				SmallTextItem.SetColor(FLinearColor::Red);
+				SmallTextItem.Text = LOCTEXT("ES2ButNoPrecomputedVisibility", "ES2 RENDERING IS ENABLED BUT PRECOMPUTED VISIBILITY IS DISABLED IN THE WORLD SETTINGS");
+				Canvas->DrawItem(SmallTextItem, FVector2D(MessageX, MessageY));
+				MessageY += 20;
+			}
+
 			// check navmesh
 #if WITH_EDITOR
 			const bool bIsNavigationAutoUpdateEnabled = UNavigationSystem::GetIsNavigationAutoUpdateEnabled();
