(fix stupid typo in CL 2226754)

diff --git a/Engine/Source/Runtime/Engine/Private/Landscape/LandscapeSplines.cpp b/Engine/Source/Runtime/Engine/Private/Landscape/LandscapeSplines.cpp
--- a/Engine/Source/Runtime/Engine/Private/Landscape/LandscapeSplines.cpp
+++ b/Engine/Source/Runtime/Engine/Private/Landscape/LandscapeSplines.cpp
@@ -464,7 +464,7 @@ UControlPointMeshComponent::UControlPointMeshComponent(const class FPostConstruc
 	: Super(PCIP)
 {
 	SetCollisionProfileName(UCollisionProfile::BlockAll_ProfileName);
-	Mobility = EComponentMobility:Static;
+	Mobility = EComponentMobility::Static;
 
 #if WITH_EDITORONLY_DATA
 	bSelected = false;
