Added shadow option to debug text. Fixed errors in DrawText function in HUD.

diff --git a/Engine/Source/Runtime/Engine/Classes/GameFramework/HUD.h b/Engine/Source/Runtime/Engine/Classes/GameFramework/HUD.h
--- a/Engine/Source/Runtime/Engine/Classes/GameFramework/HUD.h
+++ b/Engine/Source/Runtime/Engine/Classes/GameFramework/HUD.h
@@ -4,6 +4,8 @@
 #include "GameFramework/Actor.h"
 #include "HUD.generated.h"
 
+class FCanvasTextItem;
+
 /** List of actors and debug text to draw, @see AddDebugText(), RemoveDebugText(), and DrawDebugTextList() */
 USTRUCT()
 struct FDebugTextInfo
@@ -191,6 +193,10 @@ class ENGINE_API AHUD : public AActor
 	UPROPERTY()
 	uint32 bShowOverlays:1;
 
+	/** Put shadow on debug strings */
+	UPROPERTY()
+	uint32 bEnableDebugTextShadow : 1;	
+
 	/** Holds a list of Actors that need PostRender() calls. */
 	UPROPERTY()
 	TArray<class AActor*> PostRenderedActors;
@@ -220,10 +226,11 @@ protected:
 	UPROPERTY()
 	class UCanvas* DebugCanvas;
 
-private:
 	UPROPERTY()
 	TArray<struct FDebugTextInfo> DebugTextList;
 
+private:
+
 	/** Areas of the screen where UIBlur is enabled. */
 	TArray< FIntRect > UIBlurOverrideRectangles;
 
@@ -625,8 +632,7 @@ public:
 
 	/** Get the array of rectangles in which UIBlur enabled */
 	const TArray<FIntRect>& GetUIBlurRectangles() const;
-
-
+	
 protected:
 	bool IsCanvasValid_WarnIfNot() const;
 
diff --git a/Engine/Source/Runtime/Engine/Private/HUD.cpp b/Engine/Source/Runtime/Engine/Private/HUD.cpp
--- a/Engine/Source/Runtime/Engine/Private/HUD.cpp
+++ b/Engine/Source/Runtime/Engine/Private/HUD.cpp
@@ -31,6 +31,7 @@ AHUD::AHUD(const class FPostConstructInitializeProperties& PCIP)
 	bLostFocusPaused = false;
 
 	bCanBeDamaged = false;
+	bEnableDebugTextShadow = false;
 }
 
 void AHUD::SetCanvas(class UCanvas* InCanvas, class UCanvas* InDebugCanvas)
@@ -330,16 +331,16 @@ void AHUD::DrawHUD()
 	ReceiveDrawHUD(Canvas->SizeX, Canvas->SizeY);
 }
 
-
 void AHUD::DrawText(const FString& Text, FVector2D Position, UFont* TextFont, FVector2D FontScale, FColor TextColor)
 {
 	float XL, YL;
 	Canvas->TextSize(TextFont, Text, XL, YL);
 	const float X = Canvas->ClipX/2.0f - XL/2.0f + Position.X;
-	const float Y = Canvas->ClipY/3.0f - YL/2.0f + Position.Y;
-	FCanvasTextItem TextItem( FVector2D( X, Y ), FText::FromString( Text ), GEngine->GetTinyFont(), FLinearColor( TextColor ) );
+	const float Y = Canvas->ClipY/2.0f - YL/2.0f + Position.Y;
+	FCanvasTextItem TextItem( FVector2D( X, Y ), FText::FromString( Text ), TextFont, FLinearColor( TextColor ) );
 	TextItem.Scale = FontScale;
-	Canvas->DrawText(TextFont, Text, X, Y, FontScale.X, FontScale.Y);
+	
+	Canvas->DrawItem(TextItem);
 }
 
 UFont* AHUD::GetFontFromSizeIndex(int32 FontSizeIndex) const
@@ -378,6 +379,10 @@ void AHUD::DrawDebugTextList()
 		PlayerOwner->GetPlayerViewPoint(CameraLoc, CameraRot);
 
 		FCanvasTextItem TextItem( FVector2D::ZeroVector, FText::GetEmpty(), GEngine->GetSmallFont(), FLinearColor::White );
+		if (bEnableDebugTextShadow == true)
+		{
+			TextItem.EnableShadow(FLinearColor::Black);
+		}
 		for (int32 Idx = 0; Idx < DebugTextList.Num(); Idx++)
 		{
 			if (DebugTextList[Idx].SrcActor == NULL)
