Fix landscape spline undo/redo crashes (again)

diff --git a/Engine/Source/Runtime/Engine/Classes/Landscape/LandscapeSplinesComponent.h b/Engine/Source/Runtime/Engine/Classes/Landscape/LandscapeSplinesComponent.h
--- a/Engine/Source/Runtime/Engine/Classes/Landscape/LandscapeSplinesComponent.h
+++ b/Engine/Source/Runtime/Engine/Classes/Landscape/LandscapeSplinesComponent.h
@@ -50,6 +50,7 @@ public:
 	virtual void PostLoad() override;
 #if WITH_EDITOR
 	virtual void PostEditChangeProperty(FPropertyChangedEvent& PropertyChangedEvent) override;
+	virtual void PostEditUndo() override;
 #endif // WITH_EDITOR
 	// End UObject Interface
 
diff --git a/Engine/Source/Runtime/Engine/Private/Landscape/LandscapeSplines.cpp b/Engine/Source/Runtime/Engine/Private/Landscape/LandscapeSplines.cpp
--- a/Engine/Source/Runtime/Engine/Private/Landscape/LandscapeSplines.cpp
+++ b/Engine/Source/Runtime/Engine/Private/Landscape/LandscapeSplines.cpp
@@ -358,15 +358,6 @@ void ULandscapeSplinesComponent::OnRegister()
 
 void ULandscapeSplinesComponent::OnUnregister()
 {
-	for (ULandscapeSplineControlPoint* ControlPoint : ControlPoints)
-	{
-		ControlPoint->UnregisterComponents();
-	}
-	for (ULandscapeSplineSegment* Segment : Segments)
-	{
-		Segment->UnregisterComponents();
-	}
-
 	Super::OnUnregister();
 }
 
@@ -428,21 +419,36 @@ void ULandscapeSplinesComponent::PostLoad()
 }
 
 #if WITH_EDITOR
+static bool bHackIsUndoingSplines = false;
 void ULandscapeSplinesComponent::PostEditChangeProperty(FPropertyChangedEvent& PropertyChangedEvent)
 {
 	Super::PostEditChangeProperty(PropertyChangedEvent);
 
-	const bool bUpdateCollision = PropertyChangedEvent.ChangeType != EPropertyChangeType::Interactive;
-
-	for (ULandscapeSplineControlPoint* ControlPoint : ControlPoints)
-	{
-		ControlPoint->UpdateSplinePoints(bUpdateCollision);
-	}
-	for (ULandscapeSplineSegment* Segment : Segments)
+	// Don't update splines when undoing, not only is it unnecessary and expensive,
+	// it also causes failed asserts in debug builds when trying to register components
+	// (because the actor hasn't reset its OwnedComponents array yet)
+	if (!bHackIsUndoingSplines)
 	{
-		Segment->UpdateSplinePoints(bUpdateCollision);
+		const bool bUpdateCollision = PropertyChangedEvent.ChangeType != EPropertyChangeType::Interactive;
+
+		for (ULandscapeSplineControlPoint* ControlPoint : ControlPoints)
+		{
+			ControlPoint->UpdateSplinePoints(bUpdateCollision);
+		}
+		for (ULandscapeSplineSegment* Segment : Segments)
+		{
+			Segment->UpdateSplinePoints(bUpdateCollision);
+		}
 	}
 }
+void ULandscapeSplinesComponent::PostEditUndo()
+{
+	bHackIsUndoingSplines = true;
+	Super::PostEditUndo();
+	bHackIsUndoingSplines = false;
+
+	MarkRenderStateDirty();
+}
 
 void ULandscapeSplinesComponent::ShowSplineEditorMesh(bool bShow)
 {
@@ -750,6 +756,7 @@ void ULandscapeSplineControlPoint::UpdateSplinePoints(bool bUpdateCollision, boo
 				Modify();
 				OuterSplines->GetOuter()->Modify();
 				MeshComponent = ConstructObject<UControlPointMeshComponent>(UControlPointMeshComponent::StaticClass(), OuterSplines->GetOuter(), NAME_None, RF_Transactional | RF_TextExportTransient);
+				OuterSplines->GetOwner()->AddOwnedComponent(MeshComponent);
 				MeshComponent->bSelected = bSelected;
 				MeshComponent->AttachTo(OuterSplines);
 			}
@@ -918,8 +925,9 @@ void ULandscapeSplineControlPoint::DeleteSplinePoints()
 
 void ULandscapeSplineControlPoint::PostEditUndo()
 {
-	// Shouldn't call base, that results in calling UpdateSplinePoints() which will have already been reloaded from the transaction
-	//Super::PostEditUndo();
+	bHackIsUndoingSplines = true;
+	Super::PostEditUndo();
+	bHackIsUndoingSplines = false;
 
 	GetOuterULandscapeSplinesComponent()->MarkRenderStateDirty();
 }
@@ -958,8 +966,14 @@ void ULandscapeSplineControlPoint::PostEditChangeProperty(FPropertyChangedEvent&
 	SideFalloff = FMath::Max(SideFalloff, 0.0f);
 	EndFalloff = FMath::Max(EndFalloff, 0.0f);
 
-	const bool bUpdateCollision = PropertyChangedEvent.ChangeType != EPropertyChangeType::Interactive;
-	UpdateSplinePoints(bUpdateCollision);
+	// Don't update splines when undoing, not only is it unnecessary and expensive,
+	// it also causes failed asserts in debug builds when trying to register components
+	// (because the actor hasn't reset its OwnedComponents array yet)
+	if (!bHackIsUndoingSplines)
+	{
+		const bool bUpdateCollision = PropertyChangedEvent.ChangeType != EPropertyChangeType::Interactive;
+		UpdateSplinePoints(bUpdateCollision);
+	}
 }
 #endif // WITH_EDITOR
 
@@ -1369,6 +1383,7 @@ void ULandscapeSplineSegment::UpdateSplinePoints(bool bUpdateCollision)
 			{
 				OuterSplines->GetOuter()->Modify();
 				MeshComponent = ConstructObject<USplineMeshComponent>(USplineMeshComponent::StaticClass(), OuterSplines->GetOuter(), NAME_None, RF_Transactional|RF_TextExportTransient);
+				OuterSplines->GetOwner()->AddOwnedComponent(MeshComponent);
 				MeshComponent->bSelected = bSelected;
 				MeshComponent->AttachTo(OuterSplines);
 				MeshComponents.Add(MeshComponent);
@@ -1810,10 +1825,9 @@ bool ULandscapeSplineSegment::Modify(bool bAlwaysMarkDirty /*= true*/)
 #if WITH_EDITOR
 void ULandscapeSplineSegment::PostEditUndo()
 {
-	// Shouldn't call base, that results in calling UpdateSplinePoints() which will have already been reloaded from the transaction
-	//Super::PostEditUndo();
-
-	RegisterComponents();
+	bHackIsUndoingSplines = true;
+	Super::PostEditUndo();
+	bHackIsUndoingSplines = false;
 
 	GetOuterULandscapeSplinesComponent()->MarkRenderStateDirty();
 }
@@ -1864,8 +1878,14 @@ void ULandscapeSplineSegment::PostEditChangeProperty(FPropertyChangedEvent& Prop
 		Connections[1].TangentLen = FMath::Abs(Connections[1].TangentLen);
 	}
 
-	const bool bUpdateCollision = PropertyChangedEvent.ChangeType != EPropertyChangeType::Interactive;
-	UpdateSplinePoints(bUpdateCollision);
+	// Don't update splines when undoing, not only is it unnecessary and expensive,
+	// it also causes failed asserts in debug builds when trying to register components
+	// (because the actor hasn't reset its OwnedComponents array yet)
+	if (!bHackIsUndoingSplines)
+	{
+		const bool bUpdateCollision = PropertyChangedEvent.ChangeType != EPropertyChangeType::Interactive;
+		UpdateSplinePoints(bUpdateCollision);
+	}
 }
 #endif // WITH_EDITOR
 
