Header cleanup: SoundNode.h

diff --git a/Engine/Source/Runtime/Engine/Classes/Sound/SoundNode.h b/Engine/Source/Runtime/Engine/Classes/Sound/SoundNode.h
--- a/Engine/Source/Runtime/Engine/Classes/Sound/SoundNode.h
+++ b/Engine/Source/Runtime/Engine/Classes/Sound/SoundNode.h
@@ -2,7 +2,6 @@
 
 
 #pragma once
-#include "ActiveSound.h"
 #include "SoundNode.generated.h"
 
 /*-----------------------------------------------------------------------------
diff --git a/Engine/Source/Runtime/Engine/Private/ActiveSound.cpp b/Engine/Source/Runtime/Engine/Private/ActiveSound.cpp
--- a/Engine/Source/Runtime/Engine/Private/ActiveSound.cpp
+++ b/Engine/Source/Runtime/Engine/Private/ActiveSound.cpp
@@ -1,6 +1,7 @@
 // Copyright 1998-2014 Epic Games, Inc. All Rights Reserved.
 
 #include "EnginePrivate.h"
+#include "ActiveSound.h"
 #include "AudioDevice.h"
 #include "Sound/SoundCue.h"
 #include "Sound/SoundWave.h"
diff --git a/Engine/Source/Runtime/Engine/Private/Audio.cpp b/Engine/Source/Runtime/Engine/Private/Audio.cpp
--- a/Engine/Source/Runtime/Engine/Private/Audio.cpp
+++ b/Engine/Source/Runtime/Engine/Private/Audio.cpp
@@ -5,6 +5,7 @@
 =============================================================================*/
 
 #include "EnginePrivate.h"
+#include "ActiveSound.h"
 #include "Audio.h"
 #include "AudioDevice.h"
 #include "Sound/SoundCue.h"
diff --git a/Engine/Source/Runtime/Engine/Private/AudioDevice.cpp b/Engine/Source/Runtime/Engine/Private/AudioDevice.cpp
--- a/Engine/Source/Runtime/Engine/Private/AudioDevice.cpp
+++ b/Engine/Source/Runtime/Engine/Private/AudioDevice.cpp
@@ -1,6 +1,7 @@
 // Copyright 1998-2014 Epic Games, Inc. All Rights Reserved.
 
 #include "EnginePrivate.h"
+#include "ActiveSound.h"
 #include "Audio.h"
 #include "AudioDevice.h"
 #include "AudioEffect.h"
diff --git a/Engine/Source/Runtime/Engine/Private/SoundNodeMature.cpp b/Engine/Source/Runtime/Engine/Private/SoundNodeMature.cpp
--- a/Engine/Source/Runtime/Engine/Private/SoundNodeMature.cpp
+++ b/Engine/Source/Runtime/Engine/Private/SoundNodeMature.cpp
@@ -2,6 +2,7 @@
 
 
 #include "EnginePrivate.h"
+#include "ActiveSound.h"
 #include "Sound/SoundNodeMature.h"
 #include "Sound/SoundNodeWavePlayer.h"
 #include "Sound/SoundWave.h"
diff --git a/Engine/Source/Runtime/Engine/Private/SoundNodeWaveParam.cpp b/Engine/Source/Runtime/Engine/Private/SoundNodeWaveParam.cpp
--- a/Engine/Source/Runtime/Engine/Private/SoundNodeWaveParam.cpp
+++ b/Engine/Source/Runtime/Engine/Private/SoundNodeWaveParam.cpp
@@ -2,6 +2,7 @@
 
 
 #include "EnginePrivate.h"
+#include "ActiveSound.h"
 #include "Sound/SoundNodeWaveParam.h"
 #include "Sound/SoundWave.h"
 
