- Forward declare UStaticMesh in StaticMeshcomponent.h

diff --git a/Engine/Source/Runtime/Engine/Classes/Components/StaticMeshComponent.h b/Engine/Source/Runtime/Engine/Classes/Components/StaticMeshComponent.h
--- a/Engine/Source/Runtime/Engine/Classes/Components/StaticMeshComponent.h
+++ b/Engine/Source/Runtime/Engine/Classes/Components/StaticMeshComponent.h
@@ -10,6 +10,7 @@
 #include "StaticMeshComponent.generated.h"
 
 class FColorVertexBuffer;
+class UStaticMesh;
 
 /** Cached vertex information at the time the mesh was painted. */
 USTRUCT()
