Header cleanup: ParticleEventManager.h

diff --git a/Engine/Source/Runtime/Engine/Classes/Particles/ParticleEventManager.h b/Engine/Source/Runtime/Engine/Classes/Particles/ParticleEventManager.h
--- a/Engine/Source/Runtime/Engine/Classes/Particles/ParticleEventManager.h
+++ b/Engine/Source/Runtime/Engine/Classes/Particles/ParticleEventManager.h
@@ -4,7 +4,6 @@
 #pragma once
 
 #include "GameFramework/Actor.h"
-#include "ParticleSystemComponent.h"
 
 #include "ParticleEventManager.generated.h"
 
@@ -12,6 +11,7 @@ struct FParticleEventSpawnData;
 struct FParticleEventDeathData;
 struct FParticleEventCollideData;
 struct FParticleEventBurstData;
+class UParticleSystemComponent;
 
 UCLASS(config=Game, NotBlueprintable, notplaceable)
 class AParticleEventManager : public AActor
