StringClassReference needs to work with cooking and redirects
TTP# 338203

Refactored StringClassReference to be a subclass of StringAssetReference.

diff --git a/Engine/Config/BaseEngine.ini b/Engine/Config/BaseEngine.ini
--- a/Engine/Config/BaseEngine.ini
+++ b/Engine/Config/BaseEngine.ini
@@ -1252,6 +1252,9 @@ bUseStreamingPause=false
 +K2FieldRedirects=(OldFieldName="PointLightComponent.SetRadius",NewFieldName="PointLightComponent.SetAttenuationRadius")
 +K2FieldRedirects=(OldFieldName="LightComponent.InverseSquaredFalloff",NewFieldName="PointLightComponent.bUseInverseSquaredFalloff")
 
+; FStringClassReference as subclass of FStringAssetReference
++K2FieldRedirects=(OldFieldName="StringClassReference.ClassName",NewFieldName=" StringClassReference.AssetLongPathname")
+
 ; Vehicle rename properties
 +TaggedPropertyRedirects=(ClassName="VehicleEngineData",OldPropertyName="MaxOmega",NewPropertyName="MaxRPM")
 +TaggedPropertyRedirects=(ClassName="VehicleMovementComponentBoosted4w",OldPropertyName="BoostedMaxOmega",NewPropertyName="BoostedMaxRPM")
diff --git a/Engine/Source/Editor/BlueprintGraph/Private/EdGraphSchema_K2.cpp b/Engine/Source/Editor/BlueprintGraph/Private/EdGraphSchema_K2.cpp
--- a/Engine/Source/Editor/BlueprintGraph/Private/EdGraphSchema_K2.cpp
+++ b/Engine/Source/Editor/BlueprintGraph/Private/EdGraphSchema_K2.cpp
@@ -133,7 +133,7 @@ struct FGatherStructTypesFromAssetsHelper
 		if (AssetReference.IsValid())
 		{
 			const FAssetRegistryModule& AssetRegistryModule = FModuleManager::LoadModuleChecked<FAssetRegistryModule>(TEXT("AssetRegistry"));
-			const FAssetData AssetData = AssetRegistryModule.Get().GetAssetByObjectPath(*AssetReference.AssetLongPathname);
+			const FAssetData AssetData = AssetRegistryModule.Get().GetAssetByObjectPath(*AssetReference.ToString());
 			return AssetData.GetAsset();
 		}
 		return NULL;
diff --git a/Engine/Source/Editor/DetailCustomizations/Private/StringClassReferenceCustomization.cpp b/Engine/Source/Editor/DetailCustomizations/Private/StringClassReferenceCustomization.cpp
--- a/Engine/Source/Editor/DetailCustomizations/Private/StringClassReferenceCustomization.cpp
+++ b/Engine/Source/Editor/DetailCustomizations/Private/StringClassReferenceCustomization.cpp
@@ -7,7 +7,7 @@ void FStringClassReferenceCustomization::CustomizeHeader(TSharedRef<IPropertyHan
 {
 	StructPropertyHandle = InStructPropertyHandle;
 
-	ClassNamePropertyHandle = InStructPropertyHandle->GetChildHandle("ClassName");
+	ClassNamePropertyHandle = InStructPropertyHandle->GetChildHandle("AssetLongPathname");
 	check(ClassNamePropertyHandle.IsValid());
 
 	const FString& MetaClassName = StructPropertyHandle->GetMetaData("MetaClass");
@@ -48,7 +48,7 @@ void FStringClassReferenceCustomization::CustomizeChildren(TSharedRef<IPropertyH
 const UClass* FStringClassReferenceCustomization::OnGetClass() const
 {
 	FString ClassName;
-	ClassNamePropertyHandle->GetValue(ClassName);
+	ClassNamePropertyHandle->GetValueAsFormattedString(ClassName);
 
 	// Do we have a valid cached class pointer?
 	const UClass* Class = CachedClassPtr.Get();
diff --git a/Engine/Source/Editor/MaterialEditor/Private/MaterialEditor.cpp b/Engine/Source/Editor/MaterialEditor/Private/MaterialEditor.cpp
--- a/Engine/Source/Editor/MaterialEditor/Private/MaterialEditor.cpp
+++ b/Engine/Source/Editor/MaterialEditor/Private/MaterialEditor.cpp
@@ -404,7 +404,7 @@ void FMaterialEditor::InitMaterialEditor( const EToolkitMode::Type Mode, const T
 	LoadEditorSettings();
 	
 	// Set the preview mesh for the material.  This call must occur after the toolbar is initialized.
-	if ( !SetPreviewMesh( *Material->PreviewMesh.AssetLongPathname ) )
+	if (!SetPreviewMesh(*Material->PreviewMesh.ToString()))
 	{
 		// The material preview mesh couldn't be found or isn't loaded.  Default to the one of the primitive types.
 		Viewport->SetPreviewMesh( GUnrealEd->GetThumbnailManager()->EditorSphere, NULL );
@@ -3255,7 +3255,7 @@ void FMaterialEditor::NotifyPostChange( const FPropertyChangedEvent& PropertyCha
 		{
 			// SetPreviewMesh will return false if the material has bUsedWithSkeletalMesh and
 			// a skeleton was requested, in which case revert to a sphere static mesh.
-			if ( !SetPreviewMesh( *Material->PreviewMesh.AssetLongPathname ) )
+			if (!SetPreviewMesh(*Material->PreviewMesh.ToString()))
 			{
 				SetPreviewMesh( GUnrealEd->GetThumbnailManager()->EditorSphere, NULL );
 			}
diff --git a/Engine/Source/Editor/MaterialEditor/Private/MaterialInstanceEditor.cpp b/Engine/Source/Editor/MaterialEditor/Private/MaterialInstanceEditor.cpp
--- a/Engine/Source/Editor/MaterialEditor/Private/MaterialInstanceEditor.cpp
+++ b/Engine/Source/Editor/MaterialEditor/Private/MaterialInstanceEditor.cpp
@@ -225,7 +225,7 @@ void FMaterialInstanceEditor::InitMaterialInstanceEditor( const EToolkitMode::Ty
 	{
 		InstanceConstant->PreviewMesh = ThumbnailInfoWithPrim->PreviewMesh;
 	}
-	SetPreviewMesh( *InstanceConstant->PreviewMesh.AssetLongPathname );
+	SetPreviewMesh(*InstanceConstant->PreviewMesh.ToString());
 }
 
 FMaterialInstanceEditor::~FMaterialInstanceEditor()
diff --git a/Engine/Source/Editor/PhAT/Private/PhATSharedData.cpp b/Engine/Source/Editor/PhAT/Private/PhATSharedData.cpp
--- a/Engine/Source/Editor/PhAT/Private/PhATSharedData.cpp
+++ b/Engine/Source/Editor/PhAT/Private/PhATSharedData.cpp
@@ -75,9 +75,9 @@ void FPhATSharedData::Initialize()
 	USkeletalMesh * PreviewMesh =  NULL; 
 	FStringAssetReference PreviewMeshStringRef = PhysicsAsset->PreviewSkeletalMesh.ToStringReference();
 	// load it since now is the time to load
-	if(!PreviewMeshStringRef.AssetLongPathname.IsEmpty())
+	if (!PreviewMeshStringRef.ToString().IsEmpty())
 	{
-		PreviewMesh = Cast<USkeletalMesh>(StaticLoadObject(USkeletalMesh::StaticClass(), NULL, *PreviewMeshStringRef.AssetLongPathname, NULL, LOAD_None, NULL));
+		PreviewMesh = Cast<USkeletalMesh>(StaticLoadObject(USkeletalMesh::StaticClass(), NULL, *PreviewMeshStringRef.ToString(), NULL, LOAD_None, NULL));
 	}
 
 	if ( PreviewMesh == NULL)
diff --git a/Engine/Source/Editor/UnrealEd/Private/Editor.cpp b/Engine/Source/Editor/UnrealEd/Private/Editor.cpp
--- a/Engine/Source/Editor/UnrealEd/Private/Editor.cpp
+++ b/Engine/Source/Editor/UnrealEd/Private/Editor.cpp
@@ -6325,7 +6325,7 @@ bool UEditorEngine::LoadPreviewMesh( int32 Index )
 		}
 
 		// Load the new mesh, if not already loaded. 
-		UStaticMesh* PreviewMesh = LoadObject<UStaticMesh>( NULL, *MeshName.AssetLongPathname, NULL, LOAD_None, NULL );
+		UStaticMesh* PreviewMesh = LoadObject<UStaticMesh>(NULL, *MeshName.ToString(), NULL, LOAD_None, NULL);
 
 		// Swap out the meshes if we loaded or found the given static mesh. 
 		if( PreviewMesh )
@@ -6335,7 +6335,7 @@ bool UEditorEngine::LoadPreviewMesh( int32 Index )
 		}
 		else
 		{
-			UE_LOG(LogEditorViewport, Warning, TEXT("Couldn't load the PreviewMeshNames for the player at index, %d, with the name, %s."), Index, *MeshName.AssetLongPathname );
+			UE_LOG(LogEditorViewport, Warning, TEXT("Couldn't load the PreviewMeshNames for the player at index, %d, with the name, %s."), Index, *MeshName.ToString());
 		}
 	}
 	else
diff --git a/Engine/Source/Editor/UnrealEd/Public/GameModeInfoCustomizer.h b/Engine/Source/Editor/UnrealEd/Public/GameModeInfoCustomizer.h
--- a/Engine/Source/Editor/UnrealEd/Public/GameModeInfoCustomizer.h
+++ b/Engine/Source/Editor/UnrealEd/Public/GameModeInfoCustomizer.h
@@ -86,25 +86,12 @@ public:
 	void CustomizeGameModeSetting(IDetailLayoutBuilder& LayoutBuilder, IDetailCategoryBuilder& CategoryBuilder)
 	{
 		// Add GameMode picker widget
-		TSharedPtr<IPropertyHandle> DefaultGameModeHandle = LayoutBuilder.GetProperty(GameModePropertyName);
+		DefaultGameModeClassHandle = LayoutBuilder.GetProperty(GameModePropertyName);
+		check(DefaultGameModeClassHandle.IsValid());
 
-		UProperty* DefaultGameProperty = DefaultGameModeHandle->GetProperty();
-		// See if its a FStringClassReference property
-		if (UStructProperty* GameModeStructProp = Cast<UStructProperty>(DefaultGameProperty))
-		{
-			check(GameModeStructProp->Struct->GetName() == TEXT("StringClassReference"));
-			DefaultGameModeClassHandle = DefaultGameModeHandle->GetChildHandle("ClassName");
-			check(DefaultGameModeClassHandle.IsValid());
-		}
-		// Look for a regular Class property
-		else if (UClassProperty* GameModeClassProp = Cast<UClassProperty>(DefaultGameProperty))
-		{
-			DefaultGameModeClassHandle = DefaultGameModeHandle;
-		}
-
-		IDetailPropertyRow& DefaultGameModeRow = CategoryBuilder.AddProperty(DefaultGameModeHandle);
+		IDetailPropertyRow& DefaultGameModeRow = CategoryBuilder.AddProperty(DefaultGameModeClassHandle);
 		// SEe if we are allowed to choose 'no' GameMode
-		const bool bAllowNone = !(DefaultGameModeHandle->GetProperty()->PropertyFlags & CPF_NoClear);
+		const bool bAllowNone = !(DefaultGameModeClassHandle->GetProperty()->PropertyFlags & CPF_NoClear);
 
 		TAttribute<bool> CanBrowseAtrribute = TAttribute<bool>(this, &FGameModeInfoCustomizer::CanBrowseGameMode);
 
@@ -113,7 +100,7 @@ public:
 		.CustomWidget()
 		.NameContent()
 		[
-			DefaultGameModeHandle->CreatePropertyNameWidget()
+			DefaultGameModeClassHandle->CreatePropertyNameWidget()
 		]
 		.ValueContent()
 		.MaxDesiredWidth(0)
diff --git a/Engine/Source/Runtime/AIModule/Private/AIModule.cpp b/Engine/Source/Runtime/AIModule/Private/AIModule.cpp
--- a/Engine/Source/Runtime/AIModule/Private/AIModule.cpp
+++ b/Engine/Source/Runtime/AIModule/Private/AIModule.cpp
@@ -33,6 +33,6 @@ UAISystemBase* FAIModule::CreateAISystemInstance(UWorld* World)
 {
 	UE_LOG(LogAIModule, Log, TEXT("Creating AISystem for world %s"), *GetNameSafe(World));
 	
-	TSubclassOf<class UAISystemBase> AISystemClass = LoadClass<UAISystemBase>(NULL, *UAISystem::GetAISystemClassName().ClassName, NULL, LOAD_None, NULL);
+	TSubclassOf<class UAISystemBase> AISystemClass = LoadClass<UAISystemBase>(NULL, *UAISystem::GetAISystemClassName().ToString(), NULL, LOAD_None, NULL);
 	return ConstructObject<UAISystemBase>(AISystemClass, World);
 }
\ No newline at end of file
diff --git a/Engine/Source/Runtime/CoreUObject/Classes/Object.h b/Engine/Source/Runtime/CoreUObject/Classes/Object.h
--- a/Engine/Source/Runtime/CoreUObject/Classes/Object.h
+++ b/Engine/Source/Runtime/CoreUObject/Classes/Object.h
@@ -666,10 +666,8 @@ struct FStringAssetReference
 // A string class reference
 
 USTRUCT(noexport)
-struct FStringClassReference
+struct FStringClassReference : public FStringAssetReference
 {
-	UPROPERTY(EditAnywhere, BlueprintReadWrite, Category=StringClassReference)
-	FString ClassName;
 };
 
 // A struct used as stub for deleted ones.
diff --git a/Engine/Source/Runtime/CoreUObject/Private/Misc/StringAssetReference.cpp b/Engine/Source/Runtime/CoreUObject/Private/Misc/StringAssetReference.cpp
--- a/Engine/Source/Runtime/CoreUObject/Private/Misc/StringAssetReference.cpp
+++ b/Engine/Source/Runtime/CoreUObject/Private/Misc/StringAssetReference.cpp
@@ -89,44 +89,17 @@ bool FStringAssetReference::ImportTextItem( const TCHAR*& Buffer, int32 PortFlag
 	return true;
 }
 
+#include "StringReferenceTemplates.h"
 
 bool FStringAssetReference::SerializeFromMismatchedTag(struct FPropertyTag const& Tag, FArchive& Ar)
 {
-	if (Tag.Type == NAME_ObjectProperty)
+	struct UObjectTypePolicy
 	{
-		UObject* Object = NULL;
-		Ar << Object;
-		if (Object)
-		{
-			AssetLongPathname = Object->GetPathName();
-		}
-		else
-		{
-			AssetLongPathname = FString();
-		}
-#if WITH_EDITOR
-		if (Ar.IsLoading() && Ar.IsPersistent() && FCoreDelegates::StringAssetReferenceLoaded.IsBound())
-		{
-			FCoreDelegates::StringAssetReferenceLoaded.Execute(AssetLongPathname);
-		}
-#endif // WITH_EDITOR
-		return true;
-	}
-	else if( Tag.Type == NAME_StrProperty )
-	{
-		FString String;
-		Ar << String;
+		typedef UObject Type;
+		static const FName FORCEINLINE GetTypeName() { return NAME_ObjectProperty; }
+	};
 
-		AssetLongPathname = String;
-#if WITH_EDITOR
-		if (Ar.IsLoading() && Ar.IsPersistent() && FCoreDelegates::StringAssetReferenceLoaded.IsBound())
-		{
-			FCoreDelegates::StringAssetReferenceLoaded.Execute(AssetLongPathname);
-		}
-#endif // WITH_EDITOR
-		return true;
-	}
-	return false;
+	return SerializeFromMismatchedTagTemplate<UObjectTypePolicy>(AssetLongPathname, Tag, Ar);
 }
 
 UObject *FStringAssetReference::ResolveObject() const
diff --git a/Engine/Source/Runtime/CoreUObject/Private/Misc/StringClassReference.cpp b/Engine/Source/Runtime/CoreUObject/Private/Misc/StringClassReference.cpp
--- a/Engine/Source/Runtime/CoreUObject/Private/Misc/StringClassReference.cpp
+++ b/Engine/Source/Runtime/CoreUObject/Private/Misc/StringClassReference.cpp
@@ -5,125 +5,25 @@
 #include "StringClassReference.h"
 #include "PropertyTag.h"
 
-FStringClassReference::FStringClassReference(const class UClass* InClass)
-{
-	if (InClass)
-	{
-		ClassName = InClass->GetPathName();
-	}
-}
-
-FStringClassReference::FStringClassReference(const FString& PathString)
-	: ClassName(PathString)
-{
-	if (ClassName == TEXT("None"))
-	{
-		ClassName = TEXT("");
-	}
-}
-
-bool FStringClassReference::Serialize(FArchive& Ar)
-{
-	Ar << ClassName;
-	return true;
-}
-
-bool FStringClassReference::operator==(FStringClassReference const& Other) const
-{
-	return ClassName == Other.ClassName;
-}
-
-void FStringClassReference::operator=(FStringClassReference const& Other)
-{
-	ClassName = Other.ClassName;
-}
-
-bool FStringClassReference::ExportTextItem(FString& ValueStr, FStringClassReference const& DefaultValue, UObject* Parent, int32 PortFlags, UObject* ExportRootScope) const
-{
-	if (!ClassName.IsEmpty())
-	{
-		ValueStr += ClassName;
-	}
-	else
-	{
-		ValueStr += TEXT("None");
-	}
-	return true;
-}
-
-bool FStringClassReference::ImportTextItem( const TCHAR*& Buffer, int32 PortFlags, UObject* Parent, FOutputDevice* ErrorText )
-{
-	ClassName = TEXT("");
-	const TCHAR* NewBuffer = UPropertyHelpers::ReadToken( Buffer, ClassName, 1 );
-	if( !NewBuffer )
-	{
-		return false;
-	}
-	Buffer = NewBuffer;
-	if( ClassName==TEXT("None") )
-	{
-		ClassName = TEXT("");
-	}
-	else
-	{
-		if( *Buffer == TCHAR('\'') )
-		{
-			NewBuffer = UPropertyHelpers::ReadToken( Buffer, ClassName, 1 );
-			if( !NewBuffer )
-			{
-				return false;
-			}
-			Buffer = NewBuffer;
-			if( *Buffer++ != TCHAR('\'') )
-			{
-				return false;
-			}
-		}
-	}
-	return true;
-}
+#include "StringReferenceTemplates.h"
 
 bool FStringClassReference::SerializeFromMismatchedTag(struct FPropertyTag const& Tag, FArchive& Ar)
 {
-	if (Tag.Type == NAME_ClassProperty)
+	struct UClassTypePolicy
 	{
-		UClass* Class = NULL;
-		Ar << Class;
-		if (Class)
-		{
-			ClassName = Class->GetPathName();
-		}
-		else
-		{
-			ClassName = FString();
-		}
-		return true;
-	}
-	else if( Tag.Type == NAME_StrProperty )
-	{
-		FString String;
-		Ar << String;
+		typedef UClass Type;
+		static const FName FORCEINLINE GetTypeName() { return NAME_ClassProperty; }
+	};
 
-		ClassName = String;
-		return true;
-	}
-	return false;
+	return SerializeFromMismatchedTagTemplate<UClassTypePolicy>(AssetLongPathname, Tag, Ar);
 }
 
 UClass *FStringClassReference::ResolveClass() const
 {
-	if (!IsValid())
-	{
-		return NULL;
-	}
-
-	// construct full name
-	FString FullPath = ToString();
-	UClass* FoundClass = FindObject<UClass>(ANY_PACKAGE, *FullPath);
-	return FoundClass;
+	return Cast<UClass>(ResolveObject());
 }
 
-FStringClassReference FStringClassReference::GetOrCreateIDForClass(const class UClass *InClass)
+FStringClassReference FStringClassReference::GetOrCreateIDForClass(const UClass *InClass)
 {
 	check(InClass);
 	return FStringClassReference(InClass);
diff --git a/Engine/Source/Runtime/CoreUObject/Private/Misc/StringReferenceTemplates.h b/Engine/Source/Runtime/CoreUObject/Private/Misc/StringReferenceTemplates.h
new file mode 100644
--- /dev/null
+++ b/Engine/Source/Runtime/CoreUObject/Private/Misc/StringReferenceTemplates.h
@@ -0,0 +1,52 @@
+// Copyright 1998-2014 Epic Games, Inc. All Rights Reserved.
+
+#pragma once
+
+/**
+ * Serializes from mismatched tag.
+ *
+ * @template_param TypePolicy The policy should provide two things:
+ *	- GetTypeName() method that returns registered name for this property type,
+ *	- typedef Type, which is a C++ type to serialize if property matched type name.
+ * @param Tag Property tag to match type.
+ * @param Ar Archive to serialize from.
+ */
+template <class TypePolicy>
+bool SerializeFromMismatchedTagTemplate(FString& Output, const FPropertyTag& Tag, FArchive& Ar)
+{
+	if (Tag.Type == TypePolicy::GetTypeName())
+	{
+		typename TypePolicy::Type* ObjPtr = NULL;
+		Ar << ObjPtr;
+		if (ObjPtr)
+		{
+			Output = ObjPtr->GetPathName();
+		}
+		else
+		{
+			Output = FString();
+		}
+#if WITH_EDITOR
+		if (Ar.IsLoading() && Ar.IsPersistent() && FCoreDelegates::StringAssetReferenceLoaded.IsBound())
+		{
+			FCoreDelegates::StringAssetReferenceLoaded.Execute(Output);
+		}
+#endif // WITH_EDITOR
+		return true;
+	}
+	else if (Tag.Type == NAME_StrProperty)
+	{
+		FString String;
+		Ar << String;
+
+		Output = String;
+#if WITH_EDITOR
+		if (Ar.IsLoading() && Ar.IsPersistent() && FCoreDelegates::StringAssetReferenceLoaded.IsBound())
+		{
+			FCoreDelegates::StringAssetReferenceLoaded.Execute(Output);
+		}
+#endif // WITH_EDITOR
+		return true;
+	}
+	return false;
+}
\ No newline at end of file
diff --git a/Engine/Source/Runtime/CoreUObject/Public/Misc/StringAssetReference.h b/Engine/Source/Runtime/CoreUObject/Public/Misc/StringAssetReference.h
--- a/Engine/Source/Runtime/CoreUObject/Public/Misc/StringAssetReference.h
+++ b/Engine/Source/Runtime/CoreUObject/Public/Misc/StringAssetReference.h
@@ -2,18 +2,22 @@
 
 #pragma once
 
+class UObject;
+struct FPropertyTag;
+
 /**
  * A struct that contains a string reference to an asset, can be used to make soft references to assets
  */
 struct COREUOBJECT_API FStringAssetReference
 {
-	// Caution, this is mirrored in Object.h as a non export struct.
+	/** Asset path */
 	FString AssetLongPathname;
-
+	
 	FStringAssetReference()
 	{
 	}
-	FStringAssetReference(FStringAssetReference const& Other)
+
+	FStringAssetReference(const FStringAssetReference& Other)
 		: AssetLongPathname(Other.AssetLongPathname)
 	{
 	}
@@ -33,7 +37,7 @@ struct COREUOBJECT_API FStringAssetReference
 	/**
 	 * Construct from an existing object, will do some string processing
 	 */
-	FStringAssetReference(const class UObject* InObject);
+	FStringAssetReference(const UObject* InObject);
 
 	/**
 	 * Converts in to a string
@@ -47,7 +51,7 @@ struct COREUOBJECT_API FStringAssetReference
 	 * Attempts to find a currently loaded object that matches this object ID
 	 * @return Found UObject, or NULL if not currently loaded
 	 */
-	class UObject *ResolveObject() const;
+	UObject *ResolveObject() const;
 
 	/**
 	 * Resets reference to point to NULL
@@ -98,13 +102,14 @@ struct COREUOBJECT_API FStringAssetReference
 		return CurrentTag++;
 	}
 
-	static FStringAssetReference GetOrCreateIDForObject(const class UObject *Object);
+	static FStringAssetReference GetOrCreateIDForObject(const UObject *Object);
 
 	static void SetPackageNamesBeingDuplicatedForPIE(const TArray<FString>& InPackageNamesBeingDuplicatedForPIE);
 	static void ClearPackageNamesBeingDuplicatedForPIE();
 
-private:
+protected:
 
+private:
 	/**
 	 * Fixes up this StringAssetReference to add or remove the PIE prefix depending on what is currently active
 	 */
@@ -113,7 +118,6 @@ private:
 	/** Global counter that determines when we need to re-search for GUIDs because more objects have been loaded **/
 	static int32 CurrentTag;
 	static TArray<FString> PackageNamesBeingDuplicatedForPIE;
-
 };
 
 template<>
diff --git a/Engine/Source/Runtime/CoreUObject/Public/Misc/StringClassReference.h b/Engine/Source/Runtime/CoreUObject/Public/Misc/StringClassReference.h
--- a/Engine/Source/Runtime/CoreUObject/Public/Misc/StringClassReference.h
+++ b/Engine/Source/Runtime/CoreUObject/Public/Misc/StringClassReference.h
@@ -2,86 +2,50 @@
 
 #pragma once
 
+#include "StringAssetReference.h"
+
 /**
  * A struct that contains a string reference to a class, can be used to make soft references to classes
  */
-struct COREUOBJECT_API FStringClassReference
+struct COREUOBJECT_API FStringClassReference : public FStringAssetReference
 {
-	// Caution, this is mirrored in Object.h as a non export struct.
-	FString ClassName;
-
 	FStringClassReference()
-	{
-	}
+	{ }
 
 	FStringClassReference(FStringClassReference const& Other)
-		: ClassName(Other.ClassName)
-	{
-	}
+		: FStringAssetReference(Other)
+	{ }
 
 	/**
 	 * Construct from a path string
 	 */
-	FStringClassReference(const FString& PathString);
+	FStringClassReference(const FString& PathString)
+		: FStringAssetReference(PathString)
+	{ }
 
 	/**
 	 * Construct from an existing class, will do some string processing
 	 */
-	FStringClassReference(const class UClass* InClass);
-
-	/**
-	 * Converts in to a string
-	 */
-	const FString& ToString() const 
-	{
-		return ClassName;
-	}
+	FStringClassReference(const UClass* InClass)
+		: FStringAssetReference(InClass)
+	{ }
 
 	/**
 	 * Attempts to find a currently loaded object that matches this object ID
 	 * @return Found UClass, or NULL if not currently loaded
 	 */
-	class UClass *ResolveClass() const;
+	UClass *ResolveClass() const;
 
-	/**
-	 * Resets reference to point to NULL
-	 */
-	void Reset()
-	{		
-		ClassName = TEXT("");
-	}
-	
-	/**
-	 * Check if this could possibly refer to a real class, or was initialized to NULL
-	 */
-	bool IsValid() const
-	{
-		return ClassName.Len() > 0;
-	}
+	bool SerializeFromMismatchedTag(const FPropertyTag& Tag, FArchive& Ar);
 
-	bool Serialize(FArchive& Ar);
-	bool operator==(FStringClassReference const& Other) const;
-	bool operator!=(FStringClassReference const& Other) const
-	{
-		return !(*this == Other);
-	}
-	void operator=(FStringClassReference const& Other);
-	bool ExportTextItem(FString& ValueStr, FStringClassReference const& DefaultValue, UObject* Parent, int32 PortFlags, UObject* ExportRootScope) const;
-	bool ImportTextItem( const TCHAR*& Buffer, int32 PortFlags, UObject* Parent, FOutputDevice* ErrorText );
-	bool SerializeFromMismatchedTag(struct FPropertyTag const& Tag, FArchive& Ar);
-
-	friend uint32 GetTypeHash(FStringClassReference const& This)
-	{
-		return GetTypeHash(This.ClassName);
-	}
+	static FStringClassReference GetOrCreateIDForClass(const UClass *InClass);
 
-	friend FArchive& operator<<(FArchive& Ar,FStringClassReference& P)
-	{
-		P.Serialize(Ar);
-		return Ar;
-	}
+private:
+	/* Forbid creation for UObject. This class is for UClass only. Use FStringAssetReference instead. */
+	FStringClassReference(const UObject* InObject) { }
 
-	static FStringClassReference GetOrCreateIDForClass(const class UClass *InClass);
+	/* Forbidden. This class is for UClass only. Use FStringAssetReference instead. */
+	static FStringAssetReference GetOrCreateIDForObject(const UObject *Object);
 };
 
 template<>
diff --git a/Engine/Source/Runtime/Engine/Classes/Animation/PreviewAssetAttachComponent.h b/Engine/Source/Runtime/Engine/Classes/Animation/PreviewAssetAttachComponent.h
--- a/Engine/Source/Runtime/Engine/Classes/Animation/PreviewAssetAttachComponent.h
+++ b/Engine/Source/Runtime/Engine/Classes/Animation/PreviewAssetAttachComponent.h
@@ -49,9 +49,9 @@ public:
 			// if preview mesh isn't loaded, see if we have set
 			FStringAssetReference AttachedObjectStringRef = AttachedObject.ToStringReference();
 			// load it since now is the time to load
-			if (!AttachedObjectStringRef.AssetLongPathname.IsEmpty())
+			if (!AttachedObjectStringRef.ToString().IsEmpty())
 			{
-				AttachedObjectPtr = Cast<UObject>(StaticLoadObject(UObject::StaticClass(), NULL, *AttachedObjectStringRef.AssetLongPathname, NULL, LOAD_None, NULL));
+				AttachedObjectPtr = Cast<UObject>(StaticLoadObject(UObject::StaticClass(), NULL, *AttachedObjectStringRef.ToString(), NULL, LOAD_None, NULL));
 			}
 		}
 		return AttachedObjectPtr;
diff --git a/Engine/Source/Runtime/Engine/Private/AI/Navigation/NavigationSystem.cpp b/Engine/Source/Runtime/Engine/Private/AI/Navigation/NavigationSystem.cpp
--- a/Engine/Source/Runtime/Engine/Private/AI/Navigation/NavigationSystem.cpp
+++ b/Engine/Source/Runtime/Engine/Private/AI/Navigation/NavigationSystem.cpp
@@ -247,7 +247,7 @@ void UNavigationSystem::DoInitialSetup()
 	NavDataClasses.Empty(RequiredNavigationDataClassNames.Num());
 	for (int32 Index = 0; Index < RequiredNavigationDataClassNames.Num(); ++Index)
 	{
-		TSubclassOf<class ANavigationData> NavDataClass = LoadClass<ANavigationData>(NULL, *RequiredNavigationDataClassNames[Index].ClassName, NULL, LOAD_None, NULL);
+		TSubclassOf<class ANavigationData> NavDataClass = LoadClass<ANavigationData>(NULL, *RequiredNavigationDataClassNames[Index].ToString(), NULL, LOAD_None, NULL);
 		if (NavDataClass)
 		{
 			NavDataClasses.AddUnique(NavDataClass);
@@ -255,7 +255,7 @@ void UNavigationSystem::DoInitialSetup()
 		else
 		{
 			UE_LOG(LogNavigation, Warning, TEXT("Unable to find navigation data class \'%s\' while setting up require navigation types")
-				, *RequiredNavigationDataClassNames[Index].ClassName);
+				, *RequiredNavigationDataClassNames[Index].ToString());
 		}
 	}
 
diff --git a/Engine/Source/Runtime/Engine/Private/Animation/Skeleton.cpp b/Engine/Source/Runtime/Engine/Private/Animation/Skeleton.cpp
--- a/Engine/Source/Runtime/Engine/Private/Animation/Skeleton.cpp
+++ b/Engine/Source/Runtime/Engine/Private/Animation/Skeleton.cpp
@@ -713,9 +713,9 @@ USkeletalMesh* USkeleton::GetPreviewMesh(bool bFindIfNotSet)
 		// if preview mesh isn't loaded, see if we have set
 		FStringAssetReference PreviewMeshStringRef = PreviewSkeletalMesh.ToStringReference();
 		// load it since now is the time to load
-		if(!PreviewMeshStringRef.AssetLongPathname.IsEmpty())
+		if (!PreviewMeshStringRef.ToString().IsEmpty())
 		{
-			PreviewMesh = Cast<USkeletalMesh>(StaticLoadObject(USkeletalMesh::StaticClass(), NULL, *PreviewMeshStringRef.AssetLongPathname, NULL, LOAD_None, NULL));
+			PreviewMesh = Cast<USkeletalMesh>(StaticLoadObject(USkeletalMesh::StaticClass(), NULL, *PreviewMeshStringRef.ToString(), NULL, LOAD_None, NULL));
 		}
 		// if not existing, and if bFindIfNotExisting is true, then try find one
 		else if (bFindIfNotSet)
diff --git a/Engine/Source/Runtime/Engine/Private/StreamableManager.cpp b/Engine/Source/Runtime/Engine/Private/StreamableManager.cpp
--- a/Engine/Source/Runtime/Engine/Private/StreamableManager.cpp
+++ b/Engine/Source/Runtime/Engine/Private/StreamableManager.cpp
@@ -100,7 +100,7 @@ void FStreamableManager::OnPostGarbageCollect()
 			if (Existing->WeakTarget.Get())
 			{
 				// This is reasonable if someone else has made a hard reference to it
-				UE_LOG(LogStreamableManager, Verbose, TEXT("Failed to unload %s."), *It.Key().AssetLongPathname);
+				UE_LOG(LogStreamableManager, Verbose, TEXT("Failed to unload %s."), *It.Key().ToString());
 			}
 			delete Existing;
 			It.RemoveCurrent();
@@ -110,11 +110,11 @@ void FStreamableManager::OnPostGarbageCollect()
 
 UObject* FStreamableManager::SynchronousLoad(FStringAssetReference const& InTargetName)
 {
-	UE_LOG(LogStreamableManager, Verbose, TEXT("Synchronous load %s"), *InTargetName.AssetLongPathname);
+	UE_LOG(LogStreamableManager, Verbose, TEXT("Synchronous load %s"), *InTargetName.ToString());
 
-	if (FPackageName::IsShortPackageName(InTargetName.AssetLongPathname))
+	if (FPackageName::IsShortPackageName(InTargetName.ToString()))
 	{
-		UE_LOG(LogStreamableManager, Error, TEXT("     Can't load invalid package name %s"), *InTargetName.AssetLongPathname);
+		UE_LOG(LogStreamableManager, Error, TEXT("     Can't load invalid package name %s"), *InTargetName.ToString());
 		return NULL;
 	}
 
@@ -122,7 +122,7 @@ UObject* FStreamableManager::SynchronousLoad(FStringAssetReference const& InTarg
 	FStreamable* Existing = StreamableItems.FindRef(TargetName);
 	while (Existing && Existing->bAsyncLoadRequestOutstanding)
 	{
-		UE_LOG(LogStreamableManager, Verbose, TEXT("     Flushing async load for %s"), *TargetName.AssetLongPathname);
+		UE_LOG(LogStreamableManager, Verbose, TEXT("     Flushing async load for %s"), *TargetName.ToString());
 		check(!Existing->bAsyncUnloadRequestOutstanding); // we should not be both loading and unloading
 		FlushAsyncLoading(); 
 		// the async request might have found a redirect and retried
@@ -138,15 +138,15 @@ UObject* FStreamableManager::SynchronousLoad(FStringAssetReference const& InTarg
 	if (Existing->bAsyncUnloadRequestOutstanding)
 	{
 		check(!Existing->bAsyncLoadRequestOutstanding); // we should not be both loading and unloading
-		UE_LOG(LogStreamableManager, Verbose, TEXT("     Aborted unload for %s"), *TargetName.AssetLongPathname);
+		UE_LOG(LogStreamableManager, Verbose, TEXT("     Aborted unload for %s"), *TargetName.ToString());
 		Existing->bAsyncUnloadRequestOutstanding = false;
 	}
 	check(!Existing->bAsyncUnloadRequestOutstanding); // should have already dealt with this
 	check(!Existing->WeakTarget.Get()); // weak target is only valid during GC
 	if (!Existing->Target)
 	{
-		UE_LOG(LogStreamableManager, Verbose, TEXT("     Static loading %s"), *TargetName.AssetLongPathname);
-		Existing->Target = StaticLoadObject(UObject::StaticClass(), NULL, *TargetName.AssetLongPathname);
+		UE_LOG(LogStreamableManager, Verbose, TEXT("     Static loading %s"), *TargetName.ToString());
+		Existing->Target = StaticLoadObject(UObject::StaticClass(), NULL, *TargetName.ToString());
 		// need to manually detect redirectors because the above call only expects to load a UObject::StaticClass() type
 		while (UObjectRedirector* Redirector = Cast<UObjectRedirector>(Existing->Target))
 		{
@@ -155,11 +155,10 @@ UObject* FStreamableManager::SynchronousLoad(FStringAssetReference const& InTarg
 		if (Existing->Target)
 		{
 			UE_LOG(LogStreamableManager, Verbose, TEXT("     Static loaded %s"), *Existing->Target->GetFullName());
-			FStringAssetReference PossiblyNewName;
-			PossiblyNewName.AssetLongPathname = Existing->Target->GetPathName();
+			FStringAssetReference PossiblyNewName(Existing->Target->GetPathName());
 			if (PossiblyNewName != TargetName)
 			{
-				UE_LOG(LogStreamableManager, Verbose, TEXT("     Which redirected to %s"), *PossiblyNewName.AssetLongPathname);
+				UE_LOG(LogStreamableManager, Verbose, TEXT("     Which redirected to %s"), *PossiblyNewName.ToString());
 				StreamableRedirects.Add(TargetName, PossiblyNewName);
 				StreamableItems.Add(PossiblyNewName, Existing);
 				StreamableItems.Remove(TargetName);
@@ -169,7 +168,7 @@ UObject* FStreamableManager::SynchronousLoad(FStringAssetReference const& InTarg
 		else
 		{
 			Existing->bLoadFailed = true;
-			UE_LOG(LogStreamableManager, Log, TEXT("Failed attempt to load %s"), *TargetName.AssetLongPathname);
+			UE_LOG(LogStreamableManager, Log, TEXT("Failed attempt to load %s"), *TargetName.ToString());
 		}
 	}
 	else
@@ -182,11 +181,11 @@ UObject* FStreamableManager::SynchronousLoad(FStringAssetReference const& InTarg
 struct FStreamable* FStreamableManager::StreamInternal(FStringAssetReference const& InTargetName)
 {
 	check(IsInGameThread());
-	UE_LOG(LogStreamableManager, Verbose, TEXT("Asynchronous load %s"), *InTargetName.AssetLongPathname);
+	UE_LOG(LogStreamableManager, Verbose, TEXT("Asynchronous load %s"), *InTargetName.ToString());
 
-	if (FPackageName::IsShortPackageName(InTargetName.AssetLongPathname))
+	if (FPackageName::IsShortPackageName(InTargetName.ToString()))
 	{
-		UE_LOG(LogStreamableManager, Error, TEXT("     Can't load invalid package name %s"), *InTargetName.AssetLongPathname);
+		UE_LOG(LogStreamableManager, Error, TEXT("     Can't load invalid package name %s"), *InTargetName.ToString());
 		return NULL;
 	}
 
@@ -198,21 +197,21 @@ struct FStreamable* FStreamableManager::StreamInternal(FStringAssetReference con
 		{
 			// It's valid to have a live pointer if an async loaded object was hard referenced later
 			check(!Existing->bAsyncLoadRequestOutstanding); // we should not be both loading and unloading
-			UE_LOG(LogStreamableManager, Verbose, TEXT("     Aborted unload for %s"), *TargetName.AssetLongPathname);
+			UE_LOG(LogStreamableManager, Verbose, TEXT("     Aborted unload for %s"), *TargetName.ToString());
 			Existing->bAsyncUnloadRequestOutstanding = false;
 			check(Existing->Target || Existing->bLoadFailed); // should not be an unload request unless the target is valid
 			return Existing;
 		}
 		if (Existing->bAsyncLoadRequestOutstanding)
 		{
-			UE_LOG(LogStreamableManager, Verbose, TEXT("     Already in progress %s"), *TargetName.AssetLongPathname);
+			UE_LOG(LogStreamableManager, Verbose, TEXT("     Already in progress %s"), *TargetName.ToString());
 			check(!Existing->bAsyncUnloadRequestOutstanding); // we should not be both loading and unloading
 			check(!Existing->Target); // should not be an load request unless the target is invalid
 			return Existing; // already have one in process
 		}
 		if (Existing->Target)
 		{
-			UE_LOG(LogStreamableManager, Verbose, TEXT("     Already Loaded %s"), *TargetName.AssetLongPathname);
+			UE_LOG(LogStreamableManager, Verbose, TEXT("     Already Loaded %s"), *TargetName.ToString());
 			return Existing; 
 		}
 	}
@@ -225,7 +224,7 @@ struct FStreamable* FStreamableManager::StreamInternal(FStringAssetReference con
 
 	if (!Existing->Target)
 	{
-		FString Package = TargetName.AssetLongPathname;
+		FString Package = TargetName.ToString();
 		int32 FirstDot = Package.Find(TEXT("."));
 		if (FirstDot != INDEX_NONE)
 		{
@@ -282,8 +281,8 @@ void FStreamableManager::FindInMemory(FStringAssetReference& InOutTargetName, st
 	check(Existing);
 	check(!Existing->bAsyncUnloadRequestOutstanding);
 	check(!Existing->bAsyncLoadRequestOutstanding);
-	UE_LOG(LogStreamableManager, Verbose, TEXT("     Searching in memory for %s"), *InOutTargetName.AssetLongPathname);
-	Existing->Target = StaticFindObject(UObject::StaticClass(), NULL, *InOutTargetName.AssetLongPathname);
+	UE_LOG(LogStreamableManager, Verbose, TEXT("     Searching in memory for %s"), *InOutTargetName.ToString());
+	Existing->Target = StaticFindObject(UObject::StaticClass(), NULL, *InOutTargetName.ToString());
 
 	UObjectRedirector* Redir = Cast<UObjectRedirector>(Existing->Target);
 	if (Redir)
@@ -293,7 +292,7 @@ void FStreamableManager::FindInMemory(FStringAssetReference& InOutTargetName, st
 		if (!Existing->Target)
 		{
 			Existing->bLoadFailed = true;
-			UE_LOG(LogStreamableManager, Warning, TEXT("Destination of redirector was not found %s -> %s."), *InOutTargetName.AssetLongPathname, *Redir->GetFullName());
+			UE_LOG(LogStreamableManager, Warning, TEXT("Destination of redirector was not found %s -> %s."), *InOutTargetName.ToString(), *Redir->GetFullName());
 		}
 		else
 		{
@@ -302,11 +301,10 @@ void FStreamableManager::FindInMemory(FStringAssetReference& InOutTargetName, st
 	}
 	if (Existing->Target)
 	{
-		FStringAssetReference PossiblyNewName;
-		PossiblyNewName.AssetLongPathname = Existing->Target->GetPathName();
+		FStringAssetReference PossiblyNewName(Existing->Target->GetPathName());
 		if (InOutTargetName != PossiblyNewName)
 		{
-			UE_LOG(LogStreamableManager, Verbose, TEXT("     Name changed to %s"), *PossiblyNewName.AssetLongPathname);
+			UE_LOG(LogStreamableManager, Verbose, TEXT("     Name changed to %s"), *PossiblyNewName.ToString());
 			StreamableRedirects.Add(InOutTargetName, PossiblyNewName);
 			StreamableItems.Add(PossiblyNewName, Existing);
 			StreamableItems.Remove(InOutTargetName);
@@ -323,7 +321,7 @@ void FStreamableManager::AsyncLoadCallback(FStringAssetReference Request)
 	FStringAssetReference TargetName = Request;
 	FStreamable* Existing = StreamableItems.FindRef(TargetName);
 
-	UE_LOG(LogStreamableManager, Verbose, TEXT("Stream Complete callback %s"), *TargetName.AssetLongPathname);
+	UE_LOG(LogStreamableManager, Verbose, TEXT("Stream Complete callback %s"), *TargetName.ToString());
 	if (!Existing)
 	{
 		// hmm, maybe it was redirected by a consolidate
@@ -348,7 +346,7 @@ void FStreamableManager::AsyncLoadCallback(FStringAssetReference Request)
 	{
 		// Async load failed to find the object
 		Existing->bLoadFailed = true;
-		UE_LOG(LogStreamableManager, Verbose, TEXT("    Failed async load."), *TargetName.AssetLongPathname);
+		UE_LOG(LogStreamableManager, Verbose, TEXT("    Failed async load."), *TargetName.ToString());
 	}
 }
 
@@ -382,7 +380,7 @@ bool FStreamableManager::IsAsyncLoadComplete(FStringAssetReference const& InTarg
 	check(IsInGameThread());
 	FStringAssetReference TargetName = ResolveRedirects(InTargetName);
 	FStreamable* Existing = StreamableItems.FindRef(TargetName);
-	UE_LOG(LogStreamableManager, Verbose, TEXT("IsStreamComplete %s  -> %d"), *TargetName.AssetLongPathname, !Existing || !Existing->bAsyncLoadRequestOutstanding);
+	UE_LOG(LogStreamableManager, Verbose, TEXT("IsStreamComplete %s  -> %d"), *TargetName.ToString(), !Existing || !Existing->bAsyncLoadRequestOutstanding);
 	return !Existing || !Existing->bAsyncLoadRequestOutstanding;
 }
 
@@ -393,10 +391,10 @@ UObject* FStreamableManager::GetStreamed(FStringAssetReference const& InTargetNa
 	FStreamable* Existing = StreamableItems.FindRef(TargetName);
 	if (Existing && Existing->Target)
 	{
-		UE_LOG(LogStreamableManager, Verbose, TEXT("GetStreamed %s  -> %s"), *TargetName.AssetLongPathname, *Existing->Target->GetFullName());
+		UE_LOG(LogStreamableManager, Verbose, TEXT("GetStreamed %s  -> %s"), *TargetName.ToString(), *Existing->Target->GetFullName());
 		return Existing->Target;
 	}
-	UE_LOG(LogStreamableManager, Verbose, TEXT("GetStreamed %s  -> NULL"), *TargetName.AssetLongPathname);
+	UE_LOG(LogStreamableManager, Verbose, TEXT("GetStreamed %s  -> NULL"), *TargetName.ToString());
 	return NULL;
 }
 
@@ -407,13 +405,13 @@ void FStreamableManager::Unload(FStringAssetReference const& InTargetName)
 	FStreamable* Existing = StreamableItems.FindRef(TargetName);
 	if (Existing)
 	{
-		UE_LOG(LogStreamableManager, Verbose, TEXT("Unload %s"), *TargetName.AssetLongPathname);
+		UE_LOG(LogStreamableManager, Verbose, TEXT("Unload %s"), *TargetName.ToString());
 		Existing->bAsyncLoadRequestOutstanding = false;
 		Existing->bAsyncUnloadRequestOutstanding = true;
 	}
 	else
 	{
-		UE_LOG(LogStreamableManager, Verbose, TEXT("Attempt to unload %s, but it isn't loaded"), *TargetName.AssetLongPathname);
+		UE_LOG(LogStreamableManager, Verbose, TEXT("Attempt to unload %s, but it isn't loaded"), *TargetName.ToString());
 	}
 }
 
@@ -435,7 +433,7 @@ FStringAssetReference FStreamableManager::ResolveRedirects(FStringAssetReference
 	if (Redir)
 	{
 		check(Target != *Redir);
-		UE_LOG(LogStreamableManager, Verbose, TEXT("Redirected %s -> %s"), *Target.AssetLongPathname, *Redir->AssetLongPathname);
+		UE_LOG(LogStreamableManager, Verbose, TEXT("Redirected %s -> %s"), *Target.ToString(), *Redir->ToString());
 		return *Redir;
 	}
 	return Target;
diff --git a/Engine/Source/Runtime/Engine/Private/UnrealEngine.cpp b/Engine/Source/Runtime/Engine/Private/UnrealEngine.cpp
--- a/Engine/Source/Runtime/Engine/Private/UnrealEngine.cpp
+++ b/Engine/Source/Runtime/Engine/Private/UnrealEngine.cpp
@@ -1216,40 +1216,40 @@ void UEngine::InitializeObjectReferences()
 	if (AllowDebugViewmodes())
 	{
 		// Materials that are needed in-game if debug viewmodes are allowed
-		LoadSpecialMaterial(WireframeMaterialName.AssetLongPathname, WireframeMaterial, true);
-		LoadSpecialMaterial(LevelColorationLitMaterialName.AssetLongPathname, LevelColorationLitMaterial, true);
-		LoadSpecialMaterial(LevelColorationUnlitMaterialName.AssetLongPathname, LevelColorationUnlitMaterial, true);
-		LoadSpecialMaterial(LightingTexelDensityName.AssetLongPathname, LightingTexelDensityMaterial, false);
-		LoadSpecialMaterial(ShadedLevelColorationLitMaterialName.AssetLongPathname, ShadedLevelColorationLitMaterial, true);
-		LoadSpecialMaterial(ShadedLevelColorationUnlitMaterialName.AssetLongPathname, ShadedLevelColorationUnlitMaterial, true);
-		LoadSpecialMaterial(VertexColorMaterialName.AssetLongPathname, VertexColorMaterial, false);
-		LoadSpecialMaterial(VertexColorViewModeMaterialName_ColorOnly.AssetLongPathname, VertexColorViewModeMaterial_ColorOnly, false);
-		LoadSpecialMaterial(VertexColorViewModeMaterialName_AlphaAsColor.AssetLongPathname, VertexColorViewModeMaterial_AlphaAsColor, false);
-		LoadSpecialMaterial(VertexColorViewModeMaterialName_RedOnly.AssetLongPathname, VertexColorViewModeMaterial_RedOnly, false);
-		LoadSpecialMaterial(VertexColorViewModeMaterialName_GreenOnly.AssetLongPathname, VertexColorViewModeMaterial_GreenOnly, false);
-		LoadSpecialMaterial(VertexColorViewModeMaterialName_BlueOnly.AssetLongPathname, VertexColorViewModeMaterial_BlueOnly, false);
+		LoadSpecialMaterial(WireframeMaterialName.ToString(), WireframeMaterial, true);
+		LoadSpecialMaterial(LevelColorationLitMaterialName.ToString(), LevelColorationLitMaterial, true);
+		LoadSpecialMaterial(LevelColorationUnlitMaterialName.ToString(), LevelColorationUnlitMaterial, true);
+		LoadSpecialMaterial(LightingTexelDensityName.ToString(), LightingTexelDensityMaterial, false);
+		LoadSpecialMaterial(ShadedLevelColorationLitMaterialName.ToString(), ShadedLevelColorationLitMaterial, true);
+		LoadSpecialMaterial(ShadedLevelColorationUnlitMaterialName.ToString(), ShadedLevelColorationUnlitMaterial, true);
+		LoadSpecialMaterial(VertexColorMaterialName.ToString(), VertexColorMaterial, false);
+		LoadSpecialMaterial(VertexColorViewModeMaterialName_ColorOnly.ToString(), VertexColorViewModeMaterial_ColorOnly, false);
+		LoadSpecialMaterial(VertexColorViewModeMaterialName_AlphaAsColor.ToString(), VertexColorViewModeMaterial_AlphaAsColor, false);
+		LoadSpecialMaterial(VertexColorViewModeMaterialName_RedOnly.ToString(), VertexColorViewModeMaterial_RedOnly, false);
+		LoadSpecialMaterial(VertexColorViewModeMaterialName_GreenOnly.ToString(), VertexColorViewModeMaterial_GreenOnly, false);
+		LoadSpecialMaterial(VertexColorViewModeMaterialName_BlueOnly.ToString(), VertexColorViewModeMaterial_BlueOnly, false);
 	}
 
 	// Materials that may or may not be needed when debug viewmodes are disabled but haven't been fixed up yet
-	LoadSpecialMaterial(RemoveSurfaceMaterialName.AssetLongPathname, RemoveSurfaceMaterial, false);	
+	LoadSpecialMaterial(RemoveSurfaceMaterialName.ToString(), RemoveSurfaceMaterial, false);
 
 	// these one's are needed both editor and standalone 
-	LoadSpecialMaterial(DebugMeshMaterialName.AssetLongPathname, DebugMeshMaterial, false);
-	LoadSpecialMaterial(InvalidLightmapSettingsMaterialName.AssetLongPathname, InvalidLightmapSettingsMaterial, false);
-	LoadSpecialMaterial(ArrowMaterialName.AssetLongPathname, ArrowMaterial, false);
+	LoadSpecialMaterial(DebugMeshMaterialName.ToString(), DebugMeshMaterial, false);
+	LoadSpecialMaterial(InvalidLightmapSettingsMaterialName.ToString(), InvalidLightmapSettingsMaterial, false);
+	LoadSpecialMaterial(ArrowMaterialName.ToString(), ArrowMaterial, false);
 
 
 	if (GIsEditor && !IsRunningCommandlet())
 	{
 		// Materials that are only needed in the interactive editor
 #if WITH_EDITORONLY_DATA
-		LoadSpecialMaterial(GeomMaterialName.AssetLongPathname, GeomMaterial, false);
-		LoadSpecialMaterial(EditorBrushMaterialName.AssetLongPathname, EditorBrushMaterial, false);
-		LoadSpecialMaterial(BoneWeightMaterialName.AssetLongPathname, BoneWeightMaterial, false);
+		LoadSpecialMaterial(GeomMaterialName.ToString(), GeomMaterial, false);
+		LoadSpecialMaterial(EditorBrushMaterialName.ToString(), EditorBrushMaterial, false);
+		LoadSpecialMaterial(BoneWeightMaterialName.ToString(), BoneWeightMaterial, false);
 #endif
 
-		LoadSpecialMaterial(PreviewShadowsIndicatorMaterialName.AssetLongPathname, PreviewShadowsIndicatorMaterial, false);
-		LoadSpecialMaterial(ConstraintLimitMaterialName.AssetLongPathname, ConstraintLimitMaterial, false);
+		LoadSpecialMaterial(PreviewShadowsIndicatorMaterialName.ToString(), PreviewShadowsIndicatorMaterial, false);
+		LoadSpecialMaterial(ConstraintLimitMaterialName.ToString(), ConstraintLimitMaterial, false);
 
 		ConstraintLimitMaterialX = UMaterialInstanceDynamic::Create(ConstraintLimitMaterial, NULL);
 		ConstraintLimitMaterialX->SetVectorParameterValue(FName("Color"), FLinearColor::Red);
@@ -1261,124 +1261,124 @@ void UEngine::InitializeObjectReferences()
 		//@TODO: This should move into the editor (used in editor modes exclusively)
 		if (DefaultBSPVertexTexture == NULL)
 		{
-			DefaultBSPVertexTexture = LoadObject<UTexture2D>(NULL, *DefaultBSPVertexTextureName.AssetLongPathname, NULL, LOAD_None, NULL);
+			DefaultBSPVertexTexture = LoadObject<UTexture2D>(NULL, *DefaultBSPVertexTextureName.ToString(), NULL, LOAD_None, NULL);
 		}
 	}
 
 	if( DefaultTexture == NULL )
 	{
-		DefaultTexture = LoadObject<UTexture2D>(NULL, *DefaultTextureName.AssetLongPathname, NULL, LOAD_None, NULL);	
+		DefaultTexture = LoadObject<UTexture2D>(NULL, *DefaultTextureName.ToString(), NULL, LOAD_None, NULL);
 	}
 
 	if( DefaultDiffuseTexture == NULL )
 	{
-		DefaultDiffuseTexture = LoadObject<UTexture2D>(NULL, *DefaultDiffuseTextureName.AssetLongPathname, NULL, LOAD_None, NULL);	
+		DefaultDiffuseTexture = LoadObject<UTexture2D>(NULL, *DefaultDiffuseTextureName.ToString(), NULL, LOAD_None, NULL);
 	}
 
 	if( HighFrequencyNoiseTexture == NULL )
 	{
-		HighFrequencyNoiseTexture = LoadObject<UTexture2D>(NULL, *HighFrequencyNoiseTextureName.AssetLongPathname, NULL, LOAD_None, NULL);	
+		HighFrequencyNoiseTexture = LoadObject<UTexture2D>(NULL, *HighFrequencyNoiseTextureName.ToString(), NULL, LOAD_None, NULL);
 	}
 
 	if( DefaultBokehTexture == NULL )
 	{
-		DefaultBokehTexture = LoadObject<UTexture2D>(NULL, *DefaultBokehTextureName.AssetLongPathname, NULL, LOAD_None, NULL);	
+		DefaultBokehTexture = LoadObject<UTexture2D>(NULL, *DefaultBokehTextureName.ToString(), NULL, LOAD_None, NULL);
 	}
 
 	if( PreIntegratedSkinBRDFTexture == NULL )
 	{
-		PreIntegratedSkinBRDFTexture = LoadObject<UTexture2D>(NULL, *PreIntegratedSkinBRDFTextureName.AssetLongPathname, NULL, LOAD_None, NULL);	
+		PreIntegratedSkinBRDFTexture = LoadObject<UTexture2D>(NULL, *PreIntegratedSkinBRDFTextureName.ToString(), NULL, LOAD_None, NULL);
 	}
 
 	if( MiniFontTexture == NULL )
 	{
-		MiniFontTexture = LoadObject<UTexture2D>(NULL, *MiniFontTextureName.AssetLongPathname, NULL, LOAD_None, NULL);	
+		MiniFontTexture = LoadObject<UTexture2D>(NULL, *MiniFontTextureName.ToString(), NULL, LOAD_None, NULL);
 	}
 
 	if( WeightMapPlaceholderTexture == NULL )
 	{
-		WeightMapPlaceholderTexture = LoadObject<UTexture2D>(NULL, *WeightMapPlaceholderTextureName.AssetLongPathname, NULL, LOAD_None, NULL);	
+		WeightMapPlaceholderTexture = LoadObject<UTexture2D>(NULL, *WeightMapPlaceholderTextureName.ToString(), NULL, LOAD_None, NULL);
 	}
 
 	if (LightMapDensityTexture == NULL)
 	{
-		LightMapDensityTexture = LoadObject<UTexture2D>(NULL, *LightMapDensityTextureName.AssetLongPathname, NULL, LOAD_None, NULL);
+		LightMapDensityTexture = LoadObject<UTexture2D>(NULL, *LightMapDensityTextureName.ToString(), NULL, LOAD_None, NULL);
 	}
 
 	if ( DefaultPhysMaterial == NULL )
 	{
-		DefaultPhysMaterial = LoadObject<UPhysicalMaterial>(NULL, *DefaultPhysMaterialName.AssetLongPathname, NULL, LOAD_None, NULL);	
+		DefaultPhysMaterial = LoadObject<UPhysicalMaterial>(NULL, *DefaultPhysMaterialName.ToString(), NULL, LOAD_None, NULL);
 
-		checkf(DefaultPhysMaterial != NULL, TEXT("The default material (%s) is not found. Please make sure you have default material set up correctly."), *DefaultPhysMaterialName.AssetLongPathname);
+		checkf(DefaultPhysMaterial != NULL, TEXT("The default material (%s) is not found. Please make sure you have default material set up correctly."), *DefaultPhysMaterialName.ToString());
 	}
 
 	if ( ConsoleClass == NULL )
 	{
-		ConsoleClass = LoadClass<UConsole>(NULL, *ConsoleClassName.ClassName, NULL, LOAD_None, NULL);
+		ConsoleClass = LoadClass<UConsole>(NULL, *ConsoleClassName.ToString(), NULL, LOAD_None, NULL);
 	}
 
 	if ( GameViewportClientClass == NULL )
 	{
-		GameViewportClientClass = LoadClass<UGameViewportClient>(NULL, *GameViewportClientClassName.ClassName, NULL, LOAD_None, NULL);
+		GameViewportClientClass = LoadClass<UGameViewportClient>(NULL, *GameViewportClientClassName.ToString(), NULL, LOAD_None, NULL);
 
 		checkf(GameViewportClientClass != NULL, TEXT("Engine config value GameViewportClientClassName is not a valid class name."));
 	}
 
 	if ( LocalPlayerClass == NULL )
 	{
-		LocalPlayerClass = LoadClass<ULocalPlayer>(NULL, *LocalPlayerClassName.ClassName, NULL, LOAD_None, NULL);
+		LocalPlayerClass = LoadClass<ULocalPlayer>(NULL, *LocalPlayerClassName.ToString(), NULL, LOAD_None, NULL);
 	}
 
 	if ( WorldSettingsClass == NULL )
 	{
-		WorldSettingsClass = LoadClass<AWorldSettings>(NULL, *WorldSettingsClassName.ClassName, NULL, LOAD_None, NULL);
+		WorldSettingsClass = LoadClass<AWorldSettings>(NULL, *WorldSettingsClassName.ToString(), NULL, LOAD_None, NULL);
 	}
 
 	if ( NavigationSystemClass == NULL )
 	{
-		NavigationSystemClass = LoadClass<UNavigationSystem>(NULL, *NavigationSystemClassName.ClassName, NULL, LOAD_None, NULL);
+		NavigationSystemClass = LoadClass<UNavigationSystem>(NULL, *NavigationSystemClassName.ToString(), NULL, LOAD_None, NULL);
 	}
 
 	if ( AvoidanceManagerClass == NULL )
 	{
-		AvoidanceManagerClass = LoadClass<UAvoidanceManager>(NULL, *AvoidanceManagerClassName.ClassName, NULL, LOAD_None, NULL);
+		AvoidanceManagerClass = LoadClass<UAvoidanceManager>(NULL, *AvoidanceManagerClassName.ToString(), NULL, LOAD_None, NULL);
 	}
 
 	if ( PhysicsCollisionHandlerClass == NULL )
 	{
-		PhysicsCollisionHandlerClass = LoadClass<UPhysicsCollisionHandler>(NULL, *PhysicsCollisionHandlerClassName.ClassName, NULL, LOAD_None, NULL);
+		PhysicsCollisionHandlerClass = LoadClass<UPhysicsCollisionHandler>(NULL, *PhysicsCollisionHandlerClassName.ToString(), NULL, LOAD_None, NULL);
 	}
 
 	if ( GameUserSettingsClass == NULL )
 	{
-		GameUserSettingsClass = LoadClass<UGameUserSettings>(NULL, *GameUserSettingsClassName.ClassName, NULL, LOAD_None, NULL);
+		GameUserSettingsClass = LoadClass<UGameUserSettings>(NULL, *GameUserSettingsClassName.ToString(), NULL, LOAD_None, NULL);
 	}
 
 	if ( LevelScriptActorClass == NULL )
 	{
-		LevelScriptActorClass = LoadClass<ALevelScriptActor>(NULL, *LevelScriptActorClassName.ClassName, NULL, LOAD_None, NULL);
+		LevelScriptActorClass = LoadClass<ALevelScriptActor>(NULL, *LevelScriptActorClassName.ToString(), NULL, LOAD_None, NULL);
 	}
 
 	// set the font object pointers
-	if( TinyFont == NULL && TinyFontName.AssetLongPathname.Len() )
+	if (TinyFont == NULL && TinyFontName.ToString().Len())
 	{
-		TinyFont = LoadObject<UFont>(NULL,*TinyFontName.AssetLongPathname,NULL,LOAD_None,NULL);
+		TinyFont = LoadObject<UFont>(NULL, *TinyFontName.ToString(), NULL, LOAD_None, NULL);
 	}
-	if( SmallFont == NULL && SmallFontName.AssetLongPathname.Len() )
+	if (SmallFont == NULL && SmallFontName.ToString().Len())
 	{
-		SmallFont = LoadObject<UFont>(NULL,*SmallFontName.AssetLongPathname,NULL,LOAD_None,NULL);
+		SmallFont = LoadObject<UFont>(NULL, *SmallFontName.ToString(), NULL, LOAD_None, NULL);
 	}
-	if( MediumFont == NULL && MediumFontName.AssetLongPathname.Len() )
+	if (MediumFont == NULL && MediumFontName.ToString().Len())
 	{
-		MediumFont = LoadObject<UFont>(NULL,*MediumFontName.AssetLongPathname,NULL,LOAD_None,NULL);
+		MediumFont = LoadObject<UFont>(NULL, *MediumFontName.ToString(), NULL, LOAD_None, NULL);
 	}
-	if( LargeFont == NULL && LargeFontName.AssetLongPathname.Len() )
+	if (LargeFont == NULL && LargeFontName.ToString().Len())
 	{
-		LargeFont = LoadObject<UFont>(NULL,*LargeFontName.AssetLongPathname,NULL,LOAD_None,NULL);
+		LargeFont = LoadObject<UFont>(NULL, *LargeFontName.ToString(), NULL, LOAD_None, NULL);
 	}
-	if( SubtitleFont == NULL && SubtitleFontName.AssetLongPathname.Len() )
+	if (SubtitleFont == NULL && SubtitleFontName.ToString().Len())
 	{
-		SubtitleFont = LoadObject<UFont>(NULL,*SubtitleFontName.AssetLongPathname,NULL,LOAD_None,NULL);
+		SubtitleFont = LoadObject<UFont>(NULL, *SubtitleFontName.ToString(), NULL, LOAD_None, NULL);
 	}
 
 	// Additional fonts.
@@ -1394,23 +1394,23 @@ void UEngine::InitializeObjectReferences()
 		AdditionalFonts.Add( NewFont );
 	}
 
-	if ( GameSingleton == NULL && GameSingletonClassName.ClassName.Len() > 0)
+	if (GameSingleton == NULL && GameSingletonClassName.ToString().Len() > 0)
 	{
-		UClass *SingletonClass = LoadClass<UObject>(NULL, *GameSingletonClassName.ClassName, NULL, LOAD_None, NULL);
+		UClass *SingletonClass = LoadClass<UObject>(NULL, *GameSingletonClassName.ToString(), NULL, LOAD_None, NULL);
 
 		checkf(SingletonClass != NULL, TEXT("Engine config value GameSingletonClassName is not a valid class name."));
 
 		GameSingleton = ConstructObject<UObject>(SingletonClass, this);
 	}
 
-	if ( DefaultTireType == NULL && DefaultTireTypeName.AssetLongPathname.Len() )
+	if (DefaultTireType == NULL && DefaultTireTypeName.ToString().Len())
 	{
-		DefaultTireType = LoadObject<UTireType>(NULL,*DefaultTireTypeName.AssetLongPathname,NULL,LOAD_None,NULL);
+		DefaultTireType = LoadObject<UTireType>(NULL, *DefaultTireTypeName.ToString(), NULL, LOAD_None, NULL);
 	}
 
-	if( DefaultPreviewPawnClass == NULL && DefaultPreviewPawnClassName.ClassName.Len() )
+	if (DefaultPreviewPawnClass == NULL && DefaultPreviewPawnClassName.ToString().Len())
 	{
-		DefaultPreviewPawnClass = LoadClass<APawn>(NULL, *DefaultPreviewPawnClassName.ClassName, NULL, LOAD_None, NULL);
+		DefaultPreviewPawnClass = LoadClass<APawn>(NULL, *DefaultPreviewPawnClassName.ToString(), NULL, LOAD_None, NULL);
 
 		checkf(DefaultPreviewPawnClass != NULL, TEXT("Engine config value DefaultPreviewPawnClass is not a valid class name."));
 	}
diff --git a/Engine/Source/Runtime/Engine/Private/World.cpp b/Engine/Source/Runtime/Engine/Private/World.cpp
--- a/Engine/Source/Runtime/Engine/Private/World.cpp
+++ b/Engine/Source/Runtime/Engine/Private/World.cpp
@@ -771,7 +771,7 @@ UAISystemBase* UWorld::CreateAISystem()
 				AISystem = AISystemModule->CreateAISystemInstance(this);
 				if (AISystem == NULL)
 				{
-					UE_LOG(LogWorld, Error, TEXT("Failed to create AISystem instance of class %s!"), *UAISystemBase::GetAISystemClassName().ClassName);
+					UE_LOG(LogWorld, Error, TEXT("Failed to create AISystem instance of class %s!"), *UAISystemBase::GetAISystemClassName().ToString());
 				}
 			}
 		}
diff --git a/Engine/Source/Runtime/EngineSettings/Private/EngineSettingsModule.cpp b/Engine/Source/Runtime/EngineSettings/Private/EngineSettingsModule.cpp
--- a/Engine/Source/Runtime/EngineSettings/Private/EngineSettingsModule.cpp
+++ b/Engine/Source/Runtime/EngineSettings/Private/EngineSettingsModule.cpp
@@ -80,8 +80,8 @@ const FString& UGameMapsSettings::GetGlobalDefaultGameMode( )
 	UGameMapsSettings* GameMapsSettings = Cast<UGameMapsSettings>(UGameMapsSettings::StaticClass()->GetDefaultObject());
 
 	return IsRunningDedicatedServer() && GameMapsSettings->GlobalDefaultServerGameMode.IsValid()
-		? GameMapsSettings->GlobalDefaultServerGameMode.ClassName
-		: GameMapsSettings->GlobalDefaultGameMode.ClassName;
+		? GameMapsSettings->GlobalDefaultServerGameMode.ToString()
+		: GameMapsSettings->GlobalDefaultGameMode.ToString();
 }
 
 
