Header cleanup: ParticleSystemComponent.h

diff --git a/Engine/Source/Runtime/Engine/Classes/Particles/ParticleSystemComponent.h b/Engine/Source/Runtime/Engine/Classes/Particles/ParticleSystemComponent.h
--- a/Engine/Source/Runtime/Engine/Classes/Particles/ParticleSystemComponent.h
+++ b/Engine/Source/Runtime/Engine/Classes/Particles/ParticleSystemComponent.h
@@ -3,10 +3,7 @@
 
 #pragma once
 #include "Components/PrimitiveComponent.h"
-#include "Materials/MaterialInterface.h"
-#include "ParticleHelper.h"
 #include "Particles/Emitter.h"
-#include "Particles/ParticleEmitter.h"
 #include "Particles/ParticleSystem.h"
 #include "ParticleSystemComponent.generated.h"
 
diff --git a/Engine/Source/Runtime/Engine/Private/Particles/ParticleModules_Event.cpp b/Engine/Source/Runtime/Engine/Private/Particles/ParticleModules_Event.cpp
--- a/Engine/Source/Runtime/Engine/Private/Particles/ParticleModules_Event.cpp
+++ b/Engine/Source/Runtime/Engine/Private/Particles/ParticleModules_Event.cpp
@@ -13,6 +13,7 @@
 #include "Particles/Event/ParticleModuleEventReceiverSpawn.h"
 #include "Particles/Event/ParticleModuleEventSendToGame.h"
 #include "Particles/ParticleLODLevel.h"
+#include "Particles/ParticleEmitter.h"
 
 /*-----------------------------------------------------------------------------
 	Abstract base modules used for categorization.
