- Forward declare UTexture in MAterialShared.h and AController in PrimitiveComponent.h
- Include CameraTypes.h and SceneComponent.h in CameraComponent.h

diff --git a/Engine/Source/Runtime/Engine/Classes/Camera/CameraComponent.h b/Engine/Source/Runtime/Engine/Classes/Camera/CameraComponent.h
--- a/Engine/Source/Runtime/Engine/Classes/Camera/CameraComponent.h
+++ b/Engine/Source/Runtime/Engine/Classes/Camera/CameraComponent.h
@@ -1,6 +1,8 @@
 // Copyright 1998-2014 Epic Games, Inc. All Rights Reserved.
 
 #pragma once
+#include "Camera/CameraTypes.h"
+#include "Components/SceneComponent.h"
 #include "CameraComponent.generated.h"
 
 /**
diff --git a/Engine/Source/Runtime/Engine/Classes/Components/PrimitiveComponent.h b/Engine/Source/Runtime/Engine/Classes/Components/PrimitiveComponent.h
--- a/Engine/Source/Runtime/Engine/Classes/Components/PrimitiveComponent.h
+++ b/Engine/Source/Runtime/Engine/Classes/Components/PrimitiveComponent.h
@@ -16,6 +16,7 @@
 #include "PrimitiveComponent.generated.h"
 
 class FPrimitiveSceneProxy;
+class AController; 
 
 /** Information about a vertex of a primitive's triangle. */
 struct FPrimitiveTriangleVertex
diff --git a/Engine/Source/Runtime/Engine/Public/MaterialShared.h b/Engine/Source/Runtime/Engine/Public/MaterialShared.h
--- a/Engine/Source/Runtime/Engine/Public/MaterialShared.h
+++ b/Engine/Source/Runtime/Engine/Public/MaterialShared.h
@@ -25,6 +25,7 @@ class UMaterial;
 class UMaterialInstance;
 class UMaterialExpression;
 class UMaterialInterface;
+class UTexture;
 struct FExpressionInput;
 
 #define ME_CAPTION_HEIGHT		18
