Header cleanup: Brush.h

diff --git a/Engine/Source/Runtime/Engine/Classes/Engine/Brush.h b/Engine/Source/Runtime/Engine/Classes/Engine/Brush.h
--- a/Engine/Source/Runtime/Engine/Classes/Engine/Brush.h
+++ b/Engine/Source/Runtime/Engine/Classes/Engine/Brush.h
@@ -5,8 +5,6 @@
 //=============================================================================
 
 #pragma once
-#include "Model.h"
-#include "Engine/Level.h"
 #include "Brush.generated.h"
 
 //-----------------------------------------------------------------------------
