Back out changelist 2227030

diff --git a/Engine/Source/Runtime/Engine/Classes/Engine/Brush.h b/Engine/Source/Runtime/Engine/Classes/Engine/Brush.h
--- a/Engine/Source/Runtime/Engine/Classes/Engine/Brush.h
+++ b/Engine/Source/Runtime/Engine/Classes/Engine/Brush.h
@@ -5,6 +5,8 @@
 //=============================================================================
 
 #pragma once
+#include "Model.h"
+#include "Engine/Level.h"
 #include "Brush.generated.h"
 
 //-----------------------------------------------------------------------------
