* removed stuff that was not supposed to be checked in

diff --git a/Engine/Source/Runtime/Engine/Private/Materials/MaterialShared.cpp b/Engine/Source/Runtime/Engine/Private/Materials/MaterialShared.cpp
--- a/Engine/Source/Runtime/Engine/Private/Materials/MaterialShared.cpp
+++ b/Engine/Source/Runtime/Engine/Private/Materials/MaterialShared.cpp
@@ -2062,14 +2062,8 @@ FMaterialUpdateContext::~FMaterialUpdateContext()
 		);
 }
 
-uint32 Global1 = 0;
-uint32 Global2 = 0;
-
 EMaterialProperty GetMaterialPropertyFromInputOutputIndex(int32 Index)
 {
-	++Global1;
-	FPlatformMisc::LowLevelOutputDebugStringf(TEXT(">>>> %d %d\n"), Global1, Global2);
-
 	switch(Index)
 	{
 	case 0: return MP_BaseColor;
@@ -2108,9 +2102,6 @@ EMaterialProperty GetMaterialPropertyFromInputOutputIndex(int32 Index)
 
 int32 GetInputOutputIndexFromMaterialProperty(EMaterialProperty Property)
 {
-	++Global2;
-	FPlatformMisc::LowLevelOutputDebugStringf(TEXT(">>>> %d %d\n"), Global1, Global2);
-
 	switch(Property)
 	{
 	case MP_BaseColor: return 0;
