Touch events are enabled by default on PlayerControllers

diff --git a/Engine/Source/Runtime/Engine/Private/PlayerController.cpp b/Engine/Source/Runtime/Engine/Private/PlayerController.cpp
--- a/Engine/Source/Runtime/Engine/Private/PlayerController.cpp
+++ b/Engine/Source/Runtime/Engine/Private/PlayerController.cpp
@@ -52,6 +52,7 @@ APlayerController::APlayerController(const class FPostConstructInitializePropert
 	bCinemaDisableInputLook = false;
 
 	bInputEnabled = true;
+	bEnableTouchEvents = true;
 
 	bAutoManageActiveCameraTarget = true;
 
