UE4 - Enable per instance random/fade on materials

diff --git a/Engine/Source/Runtime/Engine/Private/Materials/HLSLMaterialTranslator.h b/Engine/Source/Runtime/Engine/Private/Materials/HLSLMaterialTranslator.h
--- a/Engine/Source/Runtime/Engine/Private/Materials/HLSLMaterialTranslator.h
+++ b/Engine/Source/Runtime/Engine/Private/Materials/HLSLMaterialTranslator.h
@@ -3449,11 +3449,6 @@ protected:
 	 */
 	virtual int32 PerInstanceRandom() override
 	{
-		if (ErrorUnlessFeatureLevelSupported(ERHIFeatureLevel::SM3) == INDEX_NONE)
-		{
-			return INDEX_NONE;
-		}
-
 		if (ShaderFrequency != SF_Pixel && ShaderFrequency != SF_Vertex)
 		{
 			return NonVertexOrPixelShaderExpressionError();
@@ -3471,11 +3466,6 @@ protected:
 	 */
 	virtual int32 PerInstanceFadeAmount() override
 	{
-		if (ErrorUnlessFeatureLevelSupported(ERHIFeatureLevel::SM3) == INDEX_NONE)
-		{
-			return INDEX_NONE;
-		}
-
 		if (ShaderFrequency != SF_Pixel && ShaderFrequency != SF_Vertex)
 		{
 			return NonVertexOrPixelShaderExpressionError();
