Local player management cleanup

* Rename GamePlayers to LocalPlayers
* Move LocalPlayers list into GameInstance
* Move local player API's from viewport to GameInstance
* Make GameInstance the choke point for local user access
* Started moving towards more usage of LocalPlayer nomenclature

diff --git a/Engine/Source/Editor/UnrealEd/Private/PlayLevel.cpp b/Engine/Source/Editor/UnrealEd/Private/PlayLevel.cpp
--- a/Engine/Source/Editor/UnrealEd/Private/PlayLevel.cpp
+++ b/Engine/Source/Editor/UnrealEd/Private/PlayLevel.cpp
@@ -2381,7 +2381,7 @@ UGameInstance* UEditorEngine::CreatePIEGameInstance(int32 PIEInstance, bool bInS
 		RequestedDestinationSlateViewport = NULL;
 
 		FString Error;
-		NewLocalPlayer = ViewportClient->Init(Error);
+		NewLocalPlayer = ViewportClient->SetupInitialLocalPlayer(Error);
 		if(!NewLocalPlayer)
 		{
 			FMessageDialog::Open( EAppMsgType::Ok, FText::Format(NSLOCTEXT("UnrealEd", "Error_CouldntSpawnPlayer", "Couldn't spawn player: {0}"), FText::FromString(Error)) );
diff --git a/Engine/Source/Runtime/Engine/Classes/Engine/Engine.h b/Engine/Source/Runtime/Engine/Classes/Engine/Engine.h
--- a/Engine/Source/Runtime/Engine/Classes/Engine/Engine.h
+++ b/Engine/Source/Runtime/Engine/Classes/Engine/Engine.h
@@ -259,9 +259,6 @@ struct FWorldContext
 	UPROPERTY()
 	class UGameInstance* OwningGameInstance;
 
-	UPROPERTY()
-	TArray<ULocalPlayer*> GamePlayers;
-
 	/** A list of active net drivers */
 	UPROPERTY(transient)
 	TArray<FNamedNetDriver> ActiveNetDrivers;
@@ -1282,10 +1279,6 @@ public:
 	float MinDesiredFrameRate;
 
 private:
-	/** Viewports for all players in all game instances (all PIE windows, for example) */
-	//UPROPERTY()
-	//TArray<class ULocalPlayer*> GamePlayers;
-
 	/** Default color of selected objects in the level viewport (additive) */
 	UPROPERTY(globalconfig)
 	FLinearColor DefaultSelectedMaterialColor;
@@ -1769,10 +1762,10 @@ public:
 	const TArray<class ULocalPlayer*>& GetGamePlayers(const UGameViewportClient *Viewport);
 
 	/**
-	 *	returns ULocalPlayer from a given Voice chat index. This should only be used by the online subsystems module
-	 *  and never by game code to get quick access to a local player - please use FLocalPlayerIterator instead!
+	 *	Returns the first ULocalPlayer that matches the given ControllerId. 
+	 *  This will search across all world contexts.
 	 */
-	class ULocalPlayer* LocalPlayerFromVoiceIndex(int32 index) const;
+	class ULocalPlayer* FindFirstLocalPlayerFromControllerId(int32 ControllerId) const;
 
 	/**
 	 * return the number of entries in the GamePlayers array
@@ -1809,22 +1802,6 @@ public:
 	 */
 	ULocalPlayer* GetDebugLocalPlayer();
 
-	/**
-	 * Add a ULocalPlayer to the GamePlayers array
-	 *
-	 * @param	InPlayer		Player pointer to store at the given index	 
-	 */
-	void AddGamePlayer( UWorld *InWorld, ULocalPlayer* InPlayer );
-	void AddGamePlayer( const UGameViewportClient *InViewport, ULocalPlayer* InPlayer );
-
-	/**
-	 * Remove the player with the given index from the player array if valid
-	 *
-	 * @param	InPlayerIndex		Index of the player to remove.
-	 */
-	bool RemoveGamePlayer( UWorld *InWorld, int32 InPlayerIndex );
-	bool RemoveGamePlayer( const UGameViewportClient *InViewport, int32 InPlayerIndex );
-
 	/** Clean up the GameViewport */
 	void CleanupGameViewport();
 
diff --git a/Engine/Source/Runtime/Engine/Classes/Engine/GameInstance.h b/Engine/Source/Runtime/Engine/Classes/Engine/GameInstance.h
--- a/Engine/Source/Runtime/Engine/Classes/Engine/GameInstance.h
+++ b/Engine/Source/Runtime/Engine/Classes/Engine/GameInstance.h
@@ -38,6 +38,8 @@ protected:
 
 	virtual bool HandleOpenCommand(const TCHAR* Cmd, FOutputDevice& Ar, UWorld* InWorld);
 
+	UPROPERTY()
+	TArray<ULocalPlayer*> LocalPlayers;		// List of locally participating players in this game instance
 public:
 
 
@@ -69,9 +71,59 @@ public:
 	virtual void StartGameInstance();
 
 	/** Local player access */
-	int32				GetNumLocalPlayers() const;
-	ULocalPlayer*		GetLocalPlayerByIndex(const int32 Index) const;
-	APlayerController*	GetFirstLocalPlayerController() const;
-	ULocalPlayer*		FindLocalPlayerFromUniqueNetId(TSharedPtr<FUniqueNetId> UniqueNetId) const;
-	ULocalPlayer*		GetFirstGamePlayer() const;
+
+	/**
+	 * Debug console command to create a player.
+	 * @param ControllerId - The controller ID the player should accept input from.
+	 */
+	UFUNCTION(exec)
+	virtual void			DebugCreatePlayer(int32 ControllerId);
+
+	/**
+	 * Debug console command to remove the player with a given controller ID.
+	 * @param ControllerId - The controller ID to search for.
+	 */
+	UFUNCTION(exec)
+	virtual void			DebugRemovePlayer(int32 ControllerId);
+
+	virtual ULocalPlayer*	CreateInitialPlayer(FString& OutError);
+
+	/**
+	 * Adds a new player.
+	 * @param ControllerId - The controller ID the player should accept input from.
+	 * @param OutError - If no player is returned, OutError will contain a string describing the reason.
+	 * @param SpawnActor - True if an actor should be spawned for the new player.
+	 * @return The player which was created.
+	 */
+	class ULocalPlayer*		CreateLocalPlayer(int32 ControllerId, FString& OutError, bool bSpawnActor);
+
+	/**
+	 * Adds a LocalPlayer to the local and global list of Players.
+	 *
+	 * @param	NewPlayer	the player to add
+	 * @param	ControllerId id of the controller associated with the player
+	 */
+	int32					AddLocalPlayer(ULocalPlayer * NewPlayer, int32 ControllerId);
+
+
+	/**
+	 * Removes a player.
+	 * @param Player - The player to remove.
+	 * @return whether the player was successfully removed. Removal is not allowed while connected to a server.
+	 */
+	bool					RemoveLocalPlayer(ULocalPlayer * ExistingPlayer);
+
+	int32					GetNumLocalPlayers() const;
+	ULocalPlayer*			GetLocalPlayerByIndex(const int32 Index) const;
+	APlayerController*		GetFirstLocalPlayerController() const;
+	ULocalPlayer*			FindLocalPlayerFromControllerId(int32 ControllerId) const;
+	ULocalPlayer*			FindLocalPlayerFromUniqueNetId(TSharedPtr< FUniqueNetId > UniqueNetId) const;
+	ULocalPlayer*			GetFirstGamePlayer() const;
+
+	void					CleanupGameViewport();
+
+	TArray<class ULocalPlayer*>::TConstIterator	GetLocalPlayerIterator();
+	const TArray<class ULocalPlayer*> &			GetLocalPlayers();
+
+	static void				AddReferencedObjects(UObject* InThis, FReferenceCollector& Collector);
 };
\ No newline at end of file
diff --git a/Engine/Source/Runtime/Engine/Classes/Engine/GameViewportClient.h b/Engine/Source/Runtime/Engine/Classes/Engine/GameViewportClient.h
--- a/Engine/Source/Runtime/Engine/Classes/Engine/GameViewportClient.h
+++ b/Engine/Source/Runtime/Engine/Classes/Engine/GameViewportClient.h
@@ -150,6 +150,8 @@ struct FDebugDisplayProperty
 
 };
 
+class ULocalPlayer;
+
 /**
  * Delegate type for when a screenshot has been captured
  *
@@ -225,24 +227,10 @@ public:
 	/** see enum EViewModeIndex */
 	int32 ViewModeIndex;
 
-	/**
-	 * Debug console command to create a player.
-	 * @param ControllerId - The controller ID the player should accept input from.
-	 */
-	UFUNCTION(exec)
-	virtual void DebugCreatePlayer(int32 ControllerId);
-
 	/** Rotates controller ids among gameplayers, useful for testing splitscreen with only one controller. */
 	UFUNCTION(exec)
 	virtual void SSSwapControllers();
 
-	/**
-	 * Debug console command to remove the player with a given controller ID.
-	 * @param ControllerId - The controller ID to search for.
-	 */
-	UFUNCTION(exec)
-	virtual void DebugRemovePlayer(int32 ControllerId);
-
 	/** Exec for toggling the display of the title safe area */
 	UFUNCTION(exec)
 	virtual void ShowTitleSafeArea();
@@ -418,38 +406,18 @@ public:
 	 */
 	bool ShouldForceFullscreenViewport() const;
 
-	/**
-	 * Adds a new player.
-	 * @param ControllerId - The controller ID the player should accept input from.
-	 * @param OutError - If no player is returned, OutError will contain a string describing the reason.
-	 * @param SpawnActor - True if an actor should be spawned for the new player.
-	 * @return The player which was created.
-	 */
-	virtual class ULocalPlayer* CreatePlayer(int32 ControllerId, FString& OutError, bool bSpawnActor);
-
-	/**
-	 * Removes a player.
-	 * @param Player - The player to remove.
-	 * @return whether the player was successfully removed. Removal is not allowed while connected to a server.
-	 */
-	virtual bool RemovePlayer(class ULocalPlayer* ExPlayer);
-
 	/**
 	 * Initialize the game viewport.
 	 * @param OutError - If an error occurs, returns the error description.
 	 * @return False if an error occurred, true if the viewport was initialized successfully.
 	 */
-	virtual ULocalPlayer* Init(FString& OutError);
+	virtual ULocalPlayer* SetupInitialLocalPlayer(FString& OutError);
 
-	/**
-	 * Create the game's initial player at startup.  
-	 * Creates a player with a ControllerId of 0.
-	 *
-	 * @param	OutError	receives the error string if an error occurs while creating the player.
-	 *
-	 * @return	player that was created (NULL if failed).
-	 */
-	virtual ULocalPlayer* CreateInitialPlayer( FString& OutError );
+	DEPRECATED(4.4, "CreatePlayer is deprecated UGameInstance::CreateLocalPlayer instead.")
+	virtual ULocalPlayer* CreatePlayer(int32 ControllerId, FString& OutError, bool bSpawnActor);
+
+	DEPRECATED(4.4, "RemovePlayer is deprecated UGameInstance::RemoveLocalPlayer instead.")
+	virtual bool RemovePlayer(class ULocalPlayer* ExPlayer);
 
 	/** @return Returns the splitscreen type that is currently being used */
 	FORCEINLINE ESplitScreenType::Type GetCurrentSplitscreenConfiguration() const
@@ -543,21 +511,6 @@ public:
 	 */
 	void NotifyPlayerRemoved( int32 PlayerIndex, class ULocalPlayer* RemovedPlayer );
 
-	/**
-	 * Adds a LocalPlayer to the local and global list of Players.
-	 *
-	 * @param	NewPlayer	the player to add
-	 * @param	ControllerId id of the controller associated with the player
-	 */
-	int32 AddLocalPlayer( class ULocalPlayer* NewPlayer, int32 ControllerId );
-
-	/**
-	 * Removes a LocalPlayer from the local and global list of Players.
-	 *
-	 * @param	ExistingPlayer	the player to remove
-	 */
-	int32 RemoveLocalPlayer( class ULocalPlayer* ExistingPlayer );
-
 	/**
 	 * Notification of server travel error messages, generally network connection related (package verification, client server handshaking, etc) 
 	 * generally not expected to handle the failure here, but provide notification to the user
@@ -577,16 +530,6 @@ public:
 	 */
 	virtual void PeekNetworkFailureMessages(UWorld *World, UNetDriver *NetDriver, enum ENetworkFailure::Type FailureType, const FString& ErrorString);
 
-	/**
-	 * Retrieves a reference to a LocalPlayer.
-	 *
-	 * @param	PlayerIndex		if specified, returns the player at this index in the GamePlayers array.  Otherwise, returns
-	 *							the player associated with the owner scene.
-	 * @return	the player that owns this scene or is located in the specified index of the GamePlayers array.
-	 */
-	virtual void OnPrimaryPlayerSwitch(class ULocalPlayer* OldPrimaryPlayer, class ULocalPlayer* NewPrimaryPlayer);
-
-
 	/** Make sure all navigation objects have appropriate path rendering components set.  Called when EngineShowFlags.Navigation is set. */
 	virtual void VerifyPathRenderingComponents();
 
diff --git a/Engine/Source/Runtime/Engine/Private/GameEngine.cpp b/Engine/Source/Runtime/Engine/Private/GameEngine.cpp
--- a/Engine/Source/Runtime/Engine/Private/GameEngine.cpp
+++ b/Engine/Source/Runtime/Engine/Private/GameEngine.cpp
@@ -466,7 +466,7 @@ void UGameEngine::Init(IEngineLoop* InEngineLoop)
 			SwitchGameWindowToUseGameViewport();
 		}
 		FString Error;
-		if(!ViewportClient->Init(Error))
+		if(ViewportClient->SetupInitialLocalPlayer(Error) == NULL)
 		{
 			UE_LOG(LogEngine, Fatal,TEXT("%s"),*Error);
 		}
@@ -875,7 +875,7 @@ void UGameEngine::Tick( float DeltaSeconds, bool bIdleMode )
 			Context.World()->bWorldWasLoadedThisTick = false;
 			
 			const TCHAR* InitialExec = Context.LastURL.GetOption(TEXT("causeevent="),NULL);
-			ULocalPlayer* GamePlayer = Context.GamePlayers.Num() > 0 ? Context.GamePlayers[0] : NULL;
+			ULocalPlayer* GamePlayer = Context.OwningGameInstance ? Context.OwningGameInstance->GetFirstGamePlayer() : NULL;
 			if( InitialExec && GamePlayer )
 			{
 				UE_LOG(LogEngine, Log, TEXT("Issuing initial cause event passed from URL: %s"), InitialExec);
diff --git a/Engine/Source/Runtime/Engine/Private/GameInstance.cpp b/Engine/Source/Runtime/Engine/Private/GameInstance.cpp
--- a/Engine/Source/Runtime/Engine/Private/GameInstance.cpp
+++ b/Engine/Source/Runtime/Engine/Private/GameInstance.cpp
@@ -187,7 +187,7 @@ bool UGameInstance::StartPIEGameInstance(ULocalPlayer* LocalPlayer, bool bInSimu
 			PlayWorld->NavigateTo(PlayWorld->GlobalOriginOffset);
 		}
 
-		UNavigationSystem::InitializeForWorld(PlayWorld, WorldContext->GamePlayers.Num() > 0 ? FNavigationSystem::PIEMode : FNavigationSystem::SimulationMode);
+		UNavigationSystem::InitializeForWorld(PlayWorld, LocalPlayers.Num() > 0 ? FNavigationSystem::PIEMode : FNavigationSystem::SimulationMode);
 		PlayWorld->CreateAISystem();
 
 		PlayWorld->InitializeActorsForPlay(URL);
@@ -325,37 +325,176 @@ bool UGameInstance::Exec(UWorld* InWorld, const TCHAR* Cmd, FOutputDevice& Ar)
 	return false;
 }
 
-int32 UGameInstance::GetNumLocalPlayers() const
+ULocalPlayer* UGameInstance::CreateInitialPlayer(FString& OutError)
+{
+	return CreateLocalPlayer( 0, OutError, false );
+}
+
+ULocalPlayer* UGameInstance::CreateLocalPlayer(int32 ControllerId, FString& OutError, bool bSpawnActor)
+{
+	checkf(GetEngine()->LocalPlayerClass != NULL);
+
+	ULocalPlayer* NewPlayer = NULL;
+	int32 InsertIndex = INDEX_NONE;
+
+	const int32 MaxSplitscreenPlayers = (GetGameViewportClient() != NULL) ? GetGameViewportClient()->MaxSplitscreenPlayers : 1;
+
+	if (FindLocalPlayerFromControllerId( ControllerId ) != NULL)
+	{
+		OutError = FString::Printf(TEXT("A local player already exists for controller ID %d,"), ControllerId);
+	}
+	else if (LocalPlayers.Num() < MaxSplitscreenPlayers)
+	{
+		// If the controller ID is not specified then find the first available
+		if (ControllerId < 0)
+		{
+			for (ControllerId = 0; ControllerId < MaxSplitscreenPlayers; ++ControllerId)
+			{
+				if (FindLocalPlayerFromControllerId( ControllerId ) == NULL)
+				{
+					break;
+				}
+			}
+			check(ControllerId < MaxSplitscreenPlayers);
+		}
+		else if (ControllerId >= MaxSplitscreenPlayers)
+		{
+			UE_LOG(LogPlayerManagement, Warning, TEXT("Controller ID (%d) is unlikely to map to any physical device, so this player will not receive input"), ControllerId);
+		}
+
+		NewPlayer = CastChecked<ULocalPlayer>(StaticConstructObject(GetEngine()->LocalPlayerClass, GetEngine()));
+		InsertIndex = AddLocalPlayer(NewPlayer, ControllerId);
+		if (bSpawnActor && InsertIndex != INDEX_NONE && GetWorld() != NULL)
+		{
+			if (GetWorld()->GetNetMode() != NM_Client)
+			{
+				// server; spawn a new PlayerController immediately
+				if (!NewPlayer->SpawnPlayActor("", OutError, GetWorld()))
+				{
+					RemoveLocalPlayer(NewPlayer);
+					NewPlayer = NULL;
+				}
+			}
+			else
+			{
+				// client; ask the server to let the new player join
+				NewPlayer->SendSplitJoin();
+			}
+		}
+	}
+	else
+	{
+		OutError = FString::Printf(TEXT( "Maximum number of players (%d) already created.  Unable to create more."), MaxSplitscreenPlayers);
+	}
+
+	if (OutError != TEXT(""))
+	{
+		UE_LOG(LogPlayerManagement, Log, TEXT("UPlayer* creation failed with error: %s"), *OutError);
+	}
+
+	return NewPlayer;
+}
+
+int32 UGameInstance::AddLocalPlayer(ULocalPlayer* NewLocalPlayer, int32 ControllerId)
 {
-	return ( WorldContext != nullptr ) ? WorldContext->GamePlayers.Num() : 0;
+	if (NewLocalPlayer == NULL)
+	{
+		return INDEX_NONE;
+	}
+
+	const int32 InsertIndex = LocalPlayers.Num();
+
+	// Add to list
+	LocalPlayers.AddUnique(NewLocalPlayer);
+
+	// Notify the player he/she was added
+	NewLocalPlayer->PlayerAdded(GetGameViewportClient(), ControllerId);
+
+	// Notify the viewport that we added a player (so it can update splitscreen settings, etc)
+	if ( GetGameViewportClient() != NULL )
+	{
+		GetGameViewportClient()->NotifyPlayerAdded(InsertIndex, NewLocalPlayer);
+	}
+
+	return InsertIndex;
 }
 
-ULocalPlayer* UGameInstance::GetLocalPlayerByIndex( const int32 Index ) const
+bool UGameInstance::RemoveLocalPlayer(ULocalPlayer* ExistingPlayer)
 {
-	if ( WorldContext == nullptr )
+	// FIXME: Notify server we want to leave the game if this is an online game
+	if (ExistingPlayer->PlayerController != NULL)
 	{
-		return nullptr;
+		// FIXME: Do this all inside PlayerRemoved?
+		ExistingPlayer->PlayerController->CleanupGameViewport();
+		// Destroy the player's actors.
+		ExistingPlayer->PlayerController->Destroy();
 	}
 
-	if ( !ensure( Index >= 0 && Index < WorldContext->GamePlayers.Num() ) )
+	// Remove the player from the context list
+	const int32 OldIndex = LocalPlayers.Find(ExistingPlayer);
+
+	if (OldIndex != INDEX_NONE)
 	{
-		return nullptr;
+		ExistingPlayer->PlayerRemoved();
+		LocalPlayers.RemoveAt(OldIndex);
+
+		// Notify the viewport so the viewport can do the fixups, resize, etc
+		if (GetGameViewportClient() != NULL)
+		{
+			GetGameViewportClient()->NotifyPlayerRemoved(OldIndex, ExistingPlayer);
+		}
 	}
 
-	return WorldContext->GamePlayers[Index];
+	// Disassociate this viewport client from the player.
+	// Do this after notifications, as some of them require the ViewportClient.
+	ExistingPlayer->ViewportClient = NULL;
+
+	UE_LOG(LogPlayerManagement, Log, TEXT("UGameInstance::RemovePlayer: Removed player %s with ControllerId %i at index %i (%i remaining players)"), *ExistingPlayer->GetName(), ExistingPlayer->ControllerId, OldIndex, LocalPlayers.Num());
+
+	return true;
+}
+
+void UGameInstance::DebugCreatePlayer(int32 ControllerId)
+{
+#if !UE_BUILD_SHIPPING
+	FString Error;
+	CreateLocalPlayer(ControllerId, Error, true);
+	if (Error.Len() > 0)
+	{
+		UE_LOG(LogPlayerManagement, Error, TEXT("Failed to DebugCreatePlayer: %s"), *Error);
+	}
+#endif
+}
+
+void UGameInstance::DebugRemovePlayer(int32 ControllerId)
+{
+#if !UE_BUILD_SHIPPING
+	ULocalPlayer* const ExistingPlayer = FindLocalPlayerFromControllerId(ControllerId);
+	if (ExistingPlayer != NULL)
+	{
+		RemoveLocalPlayer(ExistingPlayer);
+	}
+#endif
+}
+
+int32 UGameInstance::GetNumLocalPlayers() const
+{
+	return LocalPlayers.Num();
+}
+
+ULocalPlayer* UGameInstance::GetLocalPlayerByIndex(const int32 Index) const
+{
+	return LocalPlayers[Index];
 }
 
 APlayerController* UGameInstance::GetFirstLocalPlayerController() const
 {
-	if (WorldContext)
+	for (ULocalPlayer* Player : LocalPlayers)
 	{
-		for (ULocalPlayer* Player : WorldContext->GamePlayers)
+		if (Player && Player->PlayerController)
 		{
-			if (Player && Player->PlayerController)
-			{
-				// return first non-null entry
-				return Player->PlayerController;
-			}
+			// return first non-null entry
+			return Player->PlayerController;
 		}
 	}
 
@@ -363,17 +502,27 @@ APlayerController* UGameInstance::GetFirstLocalPlayerController() const
 	return nullptr;
 }
 
-ULocalPlayer* UGameInstance::FindLocalPlayerFromUniqueNetId( TSharedPtr<FUniqueNetId> UniqueNetId ) const
+ULocalPlayer* UGameInstance::FindLocalPlayerFromControllerId(int32 ControllerId) const
 {
-	if (WorldContext)
+	for (ULocalPlayer * LP : LocalPlayers)
 	{
-		for (ULocalPlayer* Player : WorldContext->GamePlayers)
+		if (LP && (LP->ControllerId == ControllerId))
 		{
-			if (Player && Player->GetUniqueNetId() == UniqueNetId)
-			{
-				// Match
-				return Player;
-			}
+			return LP;
+		}
+	}
+
+	return nullptr;
+}
+
+ULocalPlayer* UGameInstance::FindLocalPlayerFromUniqueNetId(TSharedPtr<FUniqueNetId> UniqueNetId) const
+{
+	for (ULocalPlayer* Player : LocalPlayers)
+	{
+		if (Player && Player->GetUniqueNetId() == UniqueNetId)
+		{
+			// Match
+			return Player;
 		}
 	}
 
@@ -383,6 +532,41 @@ ULocalPlayer* UGameInstance::FindLocalPlayerFromUniqueNetId( TSharedPtr<FUniqueN
 
 ULocalPlayer* UGameInstance::GetFirstGamePlayer() const
 {
-	return (WorldContext && WorldContext->GamePlayers.Num() > 0) ? WorldContext->GamePlayers[0] : nullptr;
+	return (LocalPlayers.Num() > 0) ? LocalPlayers[0] : nullptr;
+}
+
+void UGameInstance::CleanupGameViewport()
+{
+	// Clean up the viewports that have been closed.
+	for(int32 idx = LocalPlayers.Num()-1; idx >= 0; --idx)
+	{
+		ULocalPlayer *Player = LocalPlayers[idx];
+
+		if(Player && Player->ViewportClient && !Player->ViewportClient->Viewport)
+		{
+			RemoveLocalPlayer( Player );
+		}
+	}
 }
 
+TArray<class ULocalPlayer*>::TConstIterator	UGameInstance::GetLocalPlayerIterator()
+{
+	return LocalPlayers.CreateConstIterator();
+}
+
+const TArray<class ULocalPlayer*>& UGameInstance::GetLocalPlayers()
+{
+	return LocalPlayers;
+}
+
+void UGameInstance::AddReferencedObjects(UObject* InThis, FReferenceCollector& Collector)
+{
+	UGameInstance* This = CastChecked<UGameInstance>(InThis);
+
+	for (ULocalPlayer* LocalPlayer : This->LocalPlayers)
+	{
+		Collector.AddReferencedObject(LocalPlayer, This);
+	}
+
+	Super::AddReferencedObjects(This, Collector);
+}
diff --git a/Engine/Source/Runtime/Engine/Private/GameViewportClient.cpp b/Engine/Source/Runtime/Engine/Private/GameViewportClient.cpp
--- a/Engine/Source/Runtime/Engine/Private/GameViewportClient.cpp
+++ b/Engine/Source/Runtime/Engine/Private/GameViewportClient.cpp
@@ -1234,120 +1234,6 @@ void UGameViewportClient::PeekNetworkFailureMessages(UWorld *InWorld, UNetDriver
 	UE_LOG(LogNet, Warning, TEXT("Network Failure: %s[%s]: %s"), NetDriver ? *NetDriver->NetDriverName.ToString() : TEXT("NULL"), ENetworkFailure::ToString(FailureType), *ErrorString);
 }
 
-ULocalPlayer* UGameViewportClient::CreatePlayer(int32 ControllerId, FString& OutError, bool bSpawnActor)
-{
-	checkf(GetOuterUEngine()->LocalPlayerClass != NULL);
-
-	ULocalPlayer* NewPlayer = NULL;
-	int32 InsertIndex = INDEX_NONE;
-
-	if (GetOuterUEngine()->GetLocalPlayerFromControllerId(this, ControllerId) != NULL)
-	{
-		OutError = FString::Printf(TEXT("A local player already exists for controller ID %d,"), ControllerId);
-	}
-	else if (GetOuterUEngine()->GetNumGamePlayers(this) < MaxSplitscreenPlayers)
-	{
-		// If the controller ID is not specified then find the first available
-		if (ControllerId < 0)
-		{
-			for (ControllerId = 0; ControllerId < MaxSplitscreenPlayers; ++ControllerId)
-			{
-				if (GetOuterUEngine()->GetLocalPlayerFromControllerId(this, ControllerId) == NULL)
-				{
-					break;
-				}
-			}
-			check(ControllerId < MaxSplitscreenPlayers);
-		}
-		else if (ControllerId >= MaxSplitscreenPlayers)
-		{
-			UE_LOG(LogPlayerManagement, Warning, TEXT("Controller ID (%d) is unlikely to map to any physical device, so this player will not receive input"), ControllerId);
-		}
-
-		NewPlayer = CastChecked<ULocalPlayer>(StaticConstructObject(GetOuterUEngine()->LocalPlayerClass, GetOuter()));
-		InsertIndex = AddLocalPlayer(NewPlayer, ControllerId);
-		if ( bSpawnActor && InsertIndex != INDEX_NONE )
-		{
-			if (World->GetNetMode() != NM_Client)
-			{
-				// server; spawn a new PlayerController immediately
-				if (!NewPlayer->SpawnPlayActor("", OutError, World ))
-				{
-					RemoveLocalPlayer(NewPlayer);
-					NewPlayer = NULL;
-				}
-			}
-			else
-			{
-				// client; ask the server to let the new player join
-				NewPlayer->SendSplitJoin();
-			}
-		}
-	}
-	else
-	{
-		OutError = FString::Printf(TEXT("Maximum number of players (%d) already created.  Unable to create more."), MaxSplitscreenPlayers);
-	}
-
-	if (OutError != TEXT(""))
-	{
-		UE_LOG(LogPlayerManagement, Log, TEXT("UPlayer* creation failed with error: %s"), *OutError);
-	}
-	else
-	{
-		if ( NewPlayer != NULL && InsertIndex != INDEX_NONE )
-		{
-			NotifyPlayerAdded(InsertIndex, NewPlayer);
-		}
-	}
-	return NewPlayer;
-}
-
-
-bool UGameViewportClient::RemovePlayer(ULocalPlayer* ExPlayer)
-{
-	if ( ExPlayer->PlayerController != NULL )
-	{
-		if (ExPlayer->PlayerController->Role == ROLE_Authority)
-		{
-			// Destroy the player's actors.
-			ExPlayer->PlayerController->Destroy();
-		}
-		else
-		{
-			// FIXME: Notify server we want to leave the game
-		}
-	}
-
-	// Remove the player from the context list
-	const int32 OldIndex = RemoveLocalPlayer(ExPlayer);
-	
-	// Notify the viewport so the viewport can do the fixups, resize, etc
-	if ( OldIndex != INDEX_NONE )
-	{
-		NotifyPlayerRemoved(OldIndex, ExPlayer);
-	}
-
-	// Disassociate this viewport client from the player.
-	// Do this after notifications, as some of them require the ViewportClient.
-	ExPlayer->ViewportClient = NULL;
-
-	UE_LOG(LogPlayerManagement, Log, TEXT("UGameViewportClient::RemovePlayer: Removed player %s with ControllerId %i at index %i (%i remaining players)"), *ExPlayer->GetName(), ExPlayer->ControllerId, OldIndex, GetOuterUEngine()->GetNumGamePlayers(this));
-	return true;
-}
-
-void UGameViewportClient::DebugCreatePlayer(int32 ControllerId)
-{
-#if !UE_BUILD_SHIPPING
-	FString Error;
-	CreatePlayer(ControllerId, Error, true);
-	if (Error.Len() > 0)
-	{
-		UE_LOG(LogPlayerManagement, Error, TEXT("Failed to DebugCreatePlayer: %s"), *Error);
-	}
-#endif
-}
-
 void UGameViewportClient::SSSwapControllers()
 {
 #if !UE_BUILD_SHIPPING
@@ -1361,18 +1247,6 @@ void UGameViewportClient::SSSwapControllers()
 #endif
 }
 
-
-void UGameViewportClient::DebugRemovePlayer(int32 ControllerId)
-{
-#if !UE_BUILD_SHIPPING
-	ULocalPlayer* const ExPlayer = GEngine->GetLocalPlayerFromControllerId(this, ControllerId);
-	if (ExPlayer != NULL)
-	{
-		RemovePlayer(ExPlayer);
-	}
-#endif
-}
-
 void UGameViewportClient::ShowTitleSafeArea()
 {
 #if !UE_BUILD_SHIPPING
@@ -1380,7 +1254,6 @@ void UGameViewportClient::ShowTitleSafeArea()
 #endif
 }
 
-
 void UGameViewportClient::SetConsoleTarget(int32 PlayerIndex)
 {
 #if !UE_BUILD_SHIPPING
@@ -1399,7 +1272,7 @@ void UGameViewportClient::SetConsoleTarget(int32 PlayerIndex)
 }
 
 
-ULocalPlayer* UGameViewportClient::Init(FString& OutError)
+ULocalPlayer* UGameViewportClient::SetupInitialLocalPlayer(FString& OutError)
 {
 	checkf(GetOuterUEngine()->ConsoleClass != NULL);
 
@@ -1416,14 +1289,27 @@ ULocalPlayer* UGameViewportClient::Init(FString& OutError)
 	GEngine->OnTravelFailure().AddUObject(this, &UGameViewportClient::PeekTravelFailureMessages);
 	GEngine->OnNetworkFailure().AddUObject(this, &UGameViewportClient::PeekNetworkFailureMessages);
 
-	// create the initial player - this is necessary or we can't render anything in-game.
-	return CreateInitialPlayer(OutError);
+	UGameInstance * GameInstance = GEngine->GetWorldContextFromGameViewportChecked(this).OwningGameInstance;
+
+	if ( !ensure( GameInstance != NULL ) )
+	{
+		return NULL;
+	}
+
+	// Create the initial player - this is necessary or we can't render anything in-game.
+	return GameInstance->CreateInitialPlayer(OutError);
 }
 
+ULocalPlayer* UGameViewportClient::CreatePlayer(int32 ControllerId, FString& OutError, bool bSpawnActor)
+{
+	UGameInstance * GameInstance = GEngine->GetWorldContextFromGameViewportChecked(this).OwningGameInstance;
+	return (GameInstance != NULL) ? GameInstance->CreateLocalPlayer(ControllerId, OutError, bSpawnActor) : NULL;
+}
 
-ULocalPlayer* UGameViewportClient::CreateInitialPlayer( FString& OutError )
+bool UGameViewportClient::RemovePlayer(class ULocalPlayer* ExPlayer)
 {
-	return CreatePlayer(0, OutError, false);
+	UGameInstance * GameInstance = GEngine->GetWorldContextFromGameViewportChecked(this).OwningGameInstance;
+	return (GameInstance != NULL) ? GameInstance->RemoveLocalPlayer(ExPlayer) : false;
 }
 
 void UGameViewportClient::UpdateActiveSplitscreenType()
@@ -1809,32 +1695,6 @@ void UGameViewportClient::NotifyPlayerRemoved( int32 PlayerIndex, ULocalPlayer*
 	LayoutPlayers();
 }
 
-
-int32 UGameViewportClient::AddLocalPlayer( ULocalPlayer* NewPlayer, int32 ControllerId )
-{
-	int32 InsertIndex = INDEX_NONE;
-	if ( NewPlayer != NULL )
-	{
-		// add to list
-		NewPlayer->PlayerAdded(this, ControllerId);
-		InsertIndex = GetOuterUEngine()->GetNumGamePlayers(this);
-		GetOuterUEngine()->AddGamePlayer(this, NewPlayer);
-	}
-	return InsertIndex;
-}
-
-int32 UGameViewportClient::RemoveLocalPlayer( ULocalPlayer* ExistingPlayer )
-{
-	int32 Index = GetOuterUEngine()->GetGamePlayers(this).Find(ExistingPlayer);
-	if ( Index != INDEX_NONE )
-	{
-		ExistingPlayer->PlayerRemoved();
-		GetOuterUEngine()->RemoveGamePlayer(this, Index);
-	}
-
-	return Index;
-}
-
 void UGameViewportClient::AddViewportWidgetContent( TSharedRef<SWidget> ViewportContent, const int32 ZOrder )
 {
 	TSharedPtr< SOverlay > PinnedViewportOverlayWidget( ViewportOverlayWidget.Pin() );
@@ -1869,8 +1729,6 @@ void UGameViewportClient::RemoveAllViewportWidgets()
 	}
 }
 
-void UGameViewportClient::OnPrimaryPlayerSwitch(ULocalPlayer* OldPrimaryPlayer, ULocalPlayer* NewPrimaryPlayer) {}
-
 void UGameViewportClient::VerifyPathRenderingComponents()
 {
 	const bool bShowPaths = !!EngineShowFlags.Navigation;
diff --git a/Engine/Source/Runtime/Engine/Private/GameplayStatics.cpp b/Engine/Source/Runtime/Engine/Private/GameplayStatics.cpp
--- a/Engine/Source/Runtime/Engine/Private/GameplayStatics.cpp
+++ b/Engine/Source/Runtime/Engine/Private/GameplayStatics.cpp
@@ -72,7 +72,7 @@ APlayerController* UGameplayStatics::CreatePlayer(UObject* WorldContextObject, i
 	UWorld* World = GEngine->GetWorldFromContextObject(WorldContextObject);
 	FString Error;
 
-	ULocalPlayer* LocalPlayer = World->GetGameViewport()->CreatePlayer(ControllerId, Error, bSpawnPawn);
+	ULocalPlayer* LocalPlayer = World->GetGameInstance()->CreateLocalPlayer(ControllerId, Error, bSpawnPawn);
 
 	if (Error.Len() > 0)
 	{
diff --git a/Engine/Source/Runtime/Engine/Private/UnrealEngine.cpp b/Engine/Source/Runtime/Engine/Private/UnrealEngine.cpp
--- a/Engine/Source/Runtime/Engine/Private/UnrealEngine.cpp
+++ b/Engine/Source/Runtime/Engine/Private/UnrealEngine.cpp
@@ -79,6 +79,8 @@
 #include "ComponentReregisterContext.h"
 #include "ContentStreaming.h"
 
+#include "Engine/GameInstance.h"
+
 DEFINE_LOG_CATEGORY_STATIC(LogEngine, Log, All);
 
 IMPLEMENT_MODULE( FEngineModule, Engine );
@@ -559,10 +561,6 @@ void FWorldContext::AddReferencedObjects(FReferenceCollector& Collector, const U
 	}
 	Collector.AddReferencedObject(GameViewport, ReferencingObject);
 	Collector.AddReferencedObject(OwningGameInstance, ReferencingObject);
-	for (ULocalPlayer* GamePlayer : GamePlayers)
-	{
-		Collector.AddReferencedObject(GamePlayer, ReferencingObject);
-	}
 	for (FNamedNetDriver& ActiveNetDriver : ActiveNetDrivers)
 	{
 		Collector.AddReferencedObject(ActiveNetDriver.NetDriver, ReferencingObject);
@@ -1049,10 +1047,10 @@ void UEngine::UpdateTimeAndHandleMaxTickRate()
 			int32 NumGamePlayers = 0;
 			for (int32 WorldIndex = 0; WorldIndex < WorldList.Num(); ++WorldIndex)
 			{
-				if (WorldList[WorldIndex].WorldType == EWorldType::Game)
+				if (WorldList[WorldIndex].WorldType == EWorldType::Game && WorldList[WorldIndex].OwningGameInstance != NULL)
 				{
 					World = WorldList[WorldIndex].World();
-					NumGamePlayers = WorldList[WorldIndex].GamePlayers.Num();
+					NumGamePlayers = WorldList[WorldIndex].OwningGameInstance->GetNumLocalPlayers();
 					break;
 				}
 			}
@@ -1479,21 +1477,10 @@ void UEngine::CleanupGameViewport()
 	for (auto WorldIt = WorldList.CreateIterator(); WorldIt; ++WorldIt)
 	{
 		FWorldContext &Context = *WorldIt;
-		// Clean up the viewports that have been closed.
-		for(int32 idx = Context.GamePlayers.Num()-1; idx >= 0; --idx)
-		{
-			ULocalPlayer *Player = Context.GamePlayers[idx];
 
-			if(Player && Player->ViewportClient && !Player->ViewportClient->Viewport)
-			{
-				if (Player->PlayerController)
-				{
-					Player->PlayerController->CleanupGameViewport();
-				}
-				Player->ViewportClient = NULL;
-				Player->PlayerRemoved();
-				Context.GamePlayers.RemoveAt(idx);
-			}
+		if ( Context.OwningGameInstance != NULL )
+		{
+			Context.OwningGameInstance->CleanupGameViewport();
 		}
 
 		if ( Context.GameViewport != NULL && Context.GameViewport->Viewport == NULL )
@@ -1753,7 +1740,7 @@ bool UEngine::IsSplitScreen(UWorld *InWorld)
 		// If no specified world, return true if any world context has multiple local players
 		for (auto It = WorldList.CreateIterator(); It; ++It)
 		{
-			if (It->GamePlayers.Num() > 1)
+			if (It->OwningGameInstance != NULL && It->OwningGameInstance->GetNumLocalPlayers() > 1)
 			{
 				return true;
 			}
@@ -1804,15 +1791,22 @@ ULocalPlayer* UEngine::GetLocalPlayerFromControllerId( UWorld * InWorld, int32 C
 void UEngine::SwapControllerId(ULocalPlayer *NewPlayer, int32 CurrentControllerId, int32 NewControllerID)
 {
 	for (auto It = WorldList.CreateIterator(); It; ++It)
-{
-		if (It->GamePlayers.Contains(NewPlayer))
+	{
+		if (It->OwningGameInstance == NULL)
+		{
+			continue;
+		}
+
+		const TArray<class ULocalPlayer*> & LocalPlayers = It->OwningGameInstance->GetLocalPlayers();
+
+		if (LocalPlayers.Contains(NewPlayer))
 		{
 			// This is the world context that NewPlayer belongs to, see if anyone is using his CurrentControllerId
-			for (int32 i=0; i < It->GamePlayers.Num(); ++i)
+			for (int32 i=0; i < LocalPlayers.Num(); ++i)
 			{
-				if(It->GamePlayers[i] && It->GamePlayers[i]->ControllerId == NewControllerID)
+				if(LocalPlayers[i] && LocalPlayers[i]->ControllerId == NewControllerID)
 				{
-					It->GamePlayers[i]->ControllerId = CurrentControllerId;
+					LocalPlayers[i]->ControllerId = CurrentControllerId;
 					return;
 				}
 			}
@@ -1822,24 +1816,20 @@ void UEngine::SwapControllerId(ULocalPlayer *NewPlayer, int32 CurrentControllerI
 
 APlayerController* UEngine::GetFirstLocalPlayerController(UWorld *InWorld)
 {
-	const TArray<class ULocalPlayer*>& GamePlayers = GetGamePlayers(InWorld);
-
-	for( int32 iPlayers=0; iPlayers < GamePlayers.Num(); iPlayers++ )
-	{
-		if( GamePlayers[iPlayers] && GamePlayers[iPlayers]->PlayerController )
-		{
-			return GamePlayers[iPlayers]->PlayerController;
-		}
-	}
-
-	return NULL;
+	const FWorldContext &Context = GetWorldContextFromWorldChecked(InWorld);
+	
+	return ( Context.OwningGameInstance != NULL ) ? Context.OwningGameInstance->GetFirstLocalPlayerController() : NULL;
 }
 
 void UEngine::GetAllLocalPlayerControllers(TArray<APlayerController*> & PlayerList)
 {
 	for (auto It = WorldList.CreateIterator(); It; ++It)
 	{
-		for (auto PlayerIt = It->GamePlayers.CreateIterator(); PlayerIt; ++PlayerIt)
+		if ( It->OwningGameInstance == NULL )
+		{
+			continue;
+		}
+		for (auto PlayerIt = It->OwningGameInstance->GetLocalPlayerIterator(); PlayerIt; ++PlayerIt)
 		{
 			ULocalPlayer *Player = *PlayerIt;
 			PlayerList.Add( Player->PlayerController );
@@ -5261,19 +5251,26 @@ void UEngine::OnLostFocusPause(bool EnablePause)
 		{
 			FWorldContext &Context = *It;
 
-		// Iterate over all players and pause / unpause them
-		// Note: pausing / unpausing the player is done via their HUD pausing / unpausing
-			for (int32 PlayerIndex = 0; PlayerIndex < Context.GamePlayers.Num(); ++PlayerIndex)
-		{
-				APlayerController* PlayerController = Context.GamePlayers[PlayerIndex]->PlayerController;
-			if(PlayerController && PlayerController->MyHUD)
+			if ( Context.OwningGameInstance == NULL )
 			{
-				PlayerController->MyHUD->OnLostFocusPause(EnablePause);
+				continue;
+			}
+
+			const TArray<class ULocalPlayer*> & LocalPlayers = Context.OwningGameInstance->GetLocalPlayers();
+
+			// Iterate over all players and pause / unpause them
+			// Note: pausing / unpausing the player is done via their HUD pausing / unpausing
+			for (int32 PlayerIndex = 0; PlayerIndex < LocalPlayers.Num(); ++PlayerIndex)
+			{
+				APlayerController* PlayerController = LocalPlayers[PlayerIndex]->PlayerController;
+				if(PlayerController && PlayerController->MyHUD)
+				{
+					PlayerController->MyHUD->OnLostFocusPause(EnablePause);
+				}
 			}
 		}
 	}
 }
-}
 
 void UEngine::InitHardwareSurvey()
 {
@@ -6946,32 +6943,56 @@ TArray<class ULocalPlayer*>::TConstIterator UEngine::GetLocalPlayerIterator(cons
 	return GetGamePlayers(Viewport).CreateConstIterator();
 }
 
+static TArray<ULocalPlayer*> FakeEmptyLocalPlayers;
+
+const TArray<class ULocalPlayer*>& HandleFakeLocalPlayersList()
+{
+#if 0
+	if (!IsRunningCommandlet())
+	{
+		UE_LOG(LogLoad, Error, TEXT("WorldContext requested with NULL game instance object.") );
+		check(false);
+	}
+#endif
+	check(FakeEmptyLocalPlayers.Num() == 0);
+	return FakeEmptyLocalPlayers;
+}
+
 const TArray<class ULocalPlayer*>& UEngine::GetGamePlayers(UWorld *World)
 {
 	const FWorldContext &Context = GetWorldContextFromWorldChecked(World);
-	return Context.GamePlayers;
+	if ( Context.OwningGameInstance == NULL )
+	{
+		return HandleFakeLocalPlayersList();
+	}
+	return Context.OwningGameInstance->GetLocalPlayers();
 }
 	
 const TArray<class ULocalPlayer*>& UEngine::GetGamePlayers(const UGameViewportClient *Viewport)
 {
-	return GetWorldContextFromGameViewportChecked(Viewport).GamePlayers;
+	const FWorldContext &Context = GetWorldContextFromGameViewportChecked(Viewport);
+	if ( Context.OwningGameInstance == NULL )
+	{
+		return HandleFakeLocalPlayersList();
+	}
+	return Context.OwningGameInstance->GetLocalPlayers();
 }
 
-ULocalPlayer* UEngine::LocalPlayerFromVoiceIndex(int32 VoiceId) const
+ULocalPlayer* UEngine::FindFirstLocalPlayerFromControllerId(int32 ControllerId) const
 {
-	// Search for the first Game or PIE instance. This is imperfect and means we cannot support voice chat properly for
-	// multiple UWorlds (but thats ok for the time being).
 	for (auto It=WorldList.CreateConstIterator(); It; ++It)
 	{
 		const FWorldContext &Context = *It;
-		if (Context.World() && (Context.WorldType == EWorldType::Game || Context.WorldType == EWorldType::PIE))
+		if (Context.World() && Context.OwningGameInstance && (Context.WorldType == EWorldType::Game || Context.WorldType == EWorldType::PIE))
 		{
+			const TArray<class ULocalPlayer*> & LocalPlayers = Context.OwningGameInstance->GetLocalPlayers();
+
 			// Use this world context, look for the ULocalPlayer with this ControllerId
-			for (int32 i=0; i < Context.GamePlayers.Num(); ++i)
+			for (int32 i=0; i < LocalPlayers.Num(); ++i)
 			{
-				if (Context.GamePlayers[i] && Context.GamePlayers[i]->ControllerId == VoiceId)
+				if (LocalPlayers[i] && LocalPlayers[i]->ControllerId == ControllerId)
 				{
-					return Context.GamePlayers[i];
+					return LocalPlayers[i];
 				}
 			}
 		}
@@ -7016,7 +7037,7 @@ ULocalPlayer* UEngine::GetFirstGamePlayer( UPendingNetGame *PendingNetGame )
 	{
 		if (It->PendingNetGame == PendingNetGame)
 		{
-			return It->GamePlayers.Num() > 0 ? It->GamePlayers[0] : NULL;
+			return ( It->OwningGameInstance != NULL ) ? It->OwningGameInstance->GetFirstGamePlayer() : NULL;
 		}
 	}
 	return NULL;
@@ -7028,7 +7049,7 @@ ULocalPlayer* UEngine::GetFirstGamePlayer(const UGameViewportClient *InViewport
 	{
 		if (It->GameViewport == InViewport)
 		{
-			return It->GamePlayers.Num() > 0 ? It->GamePlayers[0] : NULL;
+			return ( It->OwningGameInstance != NULL ) ? It->OwningGameInstance->GetFirstGamePlayer() : NULL;
 		}
 	}
 	return NULL;
@@ -7038,50 +7059,14 @@ ULocalPlayer* UEngine::GetDebugLocalPlayer()
 {
 	for (auto It = WorldList.CreateConstIterator(); It; ++It)
 	{
-		if (It->World() && It->GamePlayers.Num() > 0 )
+		if (It->OwningGameInstance != NULL && It->OwningGameInstance->GetFirstGamePlayer() != NULL )
 		{
-			return It->GamePlayers[0];
+			return It->OwningGameInstance->GetFirstGamePlayer();
 		}
 	}
 	return NULL;
 }
 
-void UEngine::AddGamePlayer( UWorld *InWorld, ULocalPlayer* InPlayer )
-{
-	GetWorldContextFromWorldChecked(InWorld).GamePlayers.AddUnique(InPlayer);
-}
-
-void UEngine::AddGamePlayer( const UGameViewportClient *InViewport, ULocalPlayer* InPlayer )
-{
-	GetWorldContextFromGameViewportChecked(InViewport).GamePlayers.AddUnique(InPlayer);
-}
-
-bool RemoveGamePlayer_Local(TArray<class ULocalPlayer*>& PlayerList, int32 InPlayerIndex)
-{
-	bool bResult = true;
-	if( PlayerList.IsValidIndex( InPlayerIndex ) )
-	{
-		PlayerList.RemoveAt(InPlayerIndex);		
-	}
-	else
-	{
-		bResult = false;
-	}
-	return bResult;
-}
-
-bool UEngine::RemoveGamePlayer( UWorld *InWorld, int32 InPlayerIndex )
-{
-	TArray<class ULocalPlayer*>& PlayerList = const_cast<TArray<class ULocalPlayer*> & >(GetGamePlayers(InWorld));
-	return RemoveGamePlayer_Local(PlayerList, InPlayerIndex);
-}
-
-bool UEngine::RemoveGamePlayer( const UGameViewportClient *InViewport, int32 InPlayerIndex )
-{
-	TArray<class ULocalPlayer*>& PlayerList = const_cast<TArray<class ULocalPlayer*> & >(GetGamePlayers(InViewport));
-	return RemoveGamePlayer_Local(PlayerList, InPlayerIndex);
-}
-
 #if !UE_BUILD_SHIPPING
 // Moved this class from Class.cpp because FExportObjectInnerContext is defined in Engine
 static class FCDODump : private FSelfRegisteringExec
@@ -7318,9 +7303,9 @@ static inline void CallHandleDisconnectForFailure(UWorld* InWorld, UNetDriver* N
 	{
 		// The only disconnect case without a valid InWorld, should be in a travel case where there is a pending game net driver.
 		FWorldContext &Context = GEngine->GetWorldContextFromPendingNetGameNetDriverChecked(NetDriver);
-		check(Context.GamePlayers.Num() > 0);
+		check(Context.OwningGameInstance != NULL && Context.OwningGameInstance->GetFirstGamePlayer() != NULL);
 
-		ULocalPlayer* const LP = Context.GamePlayers[0];
+		ULocalPlayer* const LP = Context.OwningGameInstance->GetFirstGamePlayer();
 		check(LP);
 		LP->HandleDisconnect(InWorld, NetDriver);
 	}
@@ -8242,20 +8227,23 @@ bool UEngine::LoadMap( FWorldContext& WorldContext, FURL URL, class UPendingNetG
 		FWorldDelegates::LevelRemovedFromWorld.Broadcast(NULL, WorldContext.World());
 
 		// Disassociate the players from their PlayerControllers.
-		for(auto It = WorldContext.GamePlayers.CreateIterator(); It; ++It)
+		if (WorldContext.OwningGameInstance != NULL)
 		{
-			ULocalPlayer *Player = *It;
-			if(Player->PlayerController)
+			for(auto It = WorldContext.OwningGameInstance->GetLocalPlayerIterator(); It; ++It)
 			{
-				if(Player->PlayerController->GetPawn())
+				ULocalPlayer *Player = *It;
+				if(Player->PlayerController)
 				{
-					WorldContext.World()->DestroyActor(Player->PlayerController->GetPawn(), true);
+					if(Player->PlayerController->GetPawn())
+					{
+						WorldContext.World()->DestroyActor(Player->PlayerController->GetPawn(), true);
+					}
+					WorldContext.World()->DestroyActor(Player->PlayerController, true);
+					Player->PlayerController = NULL;
 				}
-				WorldContext.World()->DestroyActor(Player->PlayerController, true);
-				Player->PlayerController = NULL;
+				// reset split join info so we'll send one after loading the new map if necessary
+				Player->bSentSplitJoin = false;
 			}
-			// reset split join info so we'll send one after loading the new map if necessary
-			Player->bSentSplitJoin = false;
 		}
 
 		for (FActorIterator ActorIt(WorldContext.World()); ActorIt; ++ActorIt)
@@ -8619,13 +8607,16 @@ bool UEngine::LoadMap( FWorldContext& WorldContext, FURL URL, class UPendingNetG
 		WorldContext.LastRemoteURL = URL;
 	}
 
-	// Client init.
-	for(auto It = WorldContext.GamePlayers.CreateIterator(); It; ++It)
+	// Spawn play actors for all active local players
+	if (WorldContext.OwningGameInstance != NULL)
 	{
-		FString Error2;
-		if(!(*It)->SpawnPlayActor(URL.ToString(1),Error2,WorldContext.World()))
+		for(auto It = WorldContext.OwningGameInstance->GetLocalPlayerIterator(); It; ++It)
 		{
-			UE_LOG(LogEngine, Fatal, TEXT("Couldn't spawn player: %s"), *Error2);
+			FString Error2;
+			if(!(*It)->SpawnPlayActor(URL.ToString(1),Error2,WorldContext.World()))
+			{
+				UE_LOG(LogEngine, Fatal, TEXT("Couldn't spawn player: %s"), *Error2);
+			}
 		}
 	}
 
@@ -8832,13 +8823,16 @@ void UEngine::UpdateTransitionType(UWorld *CurrentWorld)
 		TransitionType = TT_None;
 
 		FWorldContext &Context = GetWorldContextFromWorldChecked(CurrentWorld);
-		for (auto It = Context.GamePlayers.CreateIterator(); It; ++It)
+		if (Context.OwningGameInstance != NULL)
 		{
-			if(!(*It)->PlayerController)
+			for(auto It = Context.OwningGameInstance->GetLocalPlayerIterator(); It; ++It)
 			{
-				// This player has not received a PlayerController from the server yet, so leave the connecting screen up.
-				TransitionType = TT_Connecting;
-				break;
+				if(!(*It)->PlayerController)
+				{
+					// This player has not received a PlayerController from the server yet, so leave the connecting screen up.
+					TransitionType = TT_Connecting;
+					break;
+				}
 			}
 		}
 	}
diff --git a/Engine/Source/Runtime/Online/OnlineSubsystemSteam/Private/VoiceInterfaceSteam.cpp b/Engine/Source/Runtime/Online/OnlineSubsystemSteam/Private/VoiceInterfaceSteam.cpp
--- a/Engine/Source/Runtime/Online/OnlineSubsystemSteam/Private/VoiceInterfaceSteam.cpp
+++ b/Engine/Source/Runtime/Online/OnlineSubsystemSteam/Private/VoiceInterfaceSteam.cpp
@@ -498,7 +498,10 @@ void FOnlineVoiceSteam::ProcessMuteChangeNotification()
 			// For each local user with voice
 			for (int32 Index = 0; Index < MaxLocalTalkers; Index++)
 			{
-				ULocalPlayer* LP = GEngine->LocalPlayerFromVoiceIndex(Index);
+				// Find the very first ULocalPlayer for this ControllerId. 
+				// This is imperfect and means we cannot support voice chat properly for
+				// multiple UWorlds (but thats ok for the time being).
+				ULocalPlayer* LP = GEngine->FindFirstLocalPlayerFromControllerId(Index);
 				if (LP && LP->PlayerController)
 				{
 					// If there is a player controller, we can mute/unmute people
diff --git a/Engine/Source/Runtime/Online/OnlineSubsystemUtils/Private/VoiceInterfaceImpl.cpp b/Engine/Source/Runtime/Online/OnlineSubsystemUtils/Private/VoiceInterfaceImpl.cpp
--- a/Engine/Source/Runtime/Online/OnlineSubsystemUtils/Private/VoiceInterfaceImpl.cpp
+++ b/Engine/Source/Runtime/Online/OnlineSubsystemUtils/Private/VoiceInterfaceImpl.cpp
@@ -497,7 +497,10 @@ void FOnlineVoiceImpl::ProcessMuteChangeNotification()
 			// For each local user with voice
 			for (int32 Index = 0; Index < MaxLocalTalkers; Index++)
 			{
-				ULocalPlayer* LP = GEngine->LocalPlayerFromVoiceIndex(Index);
+				// Find the very first ULocalPlayer for this ControllerId. 
+				// This is imperfect and means we cannot support voice chat properly for
+				// multiple UWorlds (but thats ok for the time being).
+				ULocalPlayer* LP = GEngine->FindFirstLocalPlayerFromControllerId(Index);
 				if (LP && LP->PlayerController)
 				{
 					// If there is a player controller, we can mute/unmute people
