UE4 - Instanced Static Meshes on Mobile (> 64 instances)
* Splits into multiple MeshBatches
* Had to change GetMeshElement() to return a TArray... We might need to reconsider this if perf is an issue

diff --git a/Engine/Source/Runtime/Engine/Private/InstancedStaticMesh.cpp b/Engine/Source/Runtime/Engine/Private/InstancedStaticMesh.cpp
--- a/Engine/Source/Runtime/Engine/Private/InstancedStaticMesh.cpp
+++ b/Engine/Source/Runtime/Engine/Private/InstancedStaticMesh.cpp
@@ -393,13 +393,19 @@ public:
 
 	static FVertexFactoryShaderParameters* ConstructShaderParameters(EShaderFrequency ShaderFrequency);
 
+	/** Make sure we account for changes in the signature of GetStaticBatchElementVisibility() */
+	static CONSTEXPR uint32 NumBitsForVisibilityMask()
+	{
+		decltype(((FInstancedStaticMeshVertexFactory*)nullptr)->GetStaticBatchElementVisibility(FSceneView(FSceneViewInitOptions()), nullptr)) Mask = 0;
+		return 8 * sizeof(Mask);
+	}
+
 	/**
 	* Get a bitmask representing the visibility of each FMeshBatch element.
 	*/
 	virtual uint64 GetStaticBatchElementVisibility(const class FSceneView& View, const struct FMeshBatch* Batch) const override
 	{
-		//@todo-rco: Only draw the first 64 instances
-		int32 NumElements = FMath::Min(Batch->Elements.Num(), 64);
+		uint32 NumElements = FMath::Min((uint32)Batch->Elements.Num(), NumBitsForVisibilityMask());
 		return (1ULL << (uint64)NumElements) - 1ULL;
 	}
 
@@ -842,16 +848,6 @@ void FInstancedStaticMeshRenderData::InitStaticMeshVertexFactories(
 	FInstancedStaticMeshSceneProxy
 -----------------------------------------------------------------------------*/
 
-
-class FInstancedStaticMeshSceneProxyInstanceData
-{
-
-public: 
-
-	/** Cached instance -> World transform */
-	FMatrix InstanceToWorld;
-};
-
 class FInstancedStaticMeshSceneProxy : public FStaticMeshSceneProxy
 {
 public:
@@ -924,19 +920,19 @@ public:
 	}
 
 	/** Common path for the Get*MeshElement functions */
-	inline void SetupInstancedMeshBatch(int32 LODIndex, FMeshBatch& OutMeshBatch) const;
+	inline void SetupInstancedMeshBatch(int32 LODIndex, TArray<FMeshBatch, TInlineAllocator<1>>& OutMeshBatches) const;
 
 	/** Draw the scene proxy as a dynamic element */
 	virtual void DrawDynamicElements(FPrimitiveDrawInterface* PDI,const FSceneView* View) override;
 		
 	/** Sets up a shadow FMeshBatch for a specific LOD. */
-	virtual bool GetShadowMeshElement(int32 LODIndex, uint8 InDepthPriorityGroup, FMeshBatch& OutMeshElement) const override;
+	virtual bool GetShadowMeshElements(int32 LODIndex, uint8 InDepthPriorityGroup, TArray<FMeshBatch, TInlineAllocator<1>>& OutMeshBatches) const override;
 
 	/** Sets up a FMeshBatch for a specific LOD and element. */
-	virtual bool GetMeshElement(int32 LODIndex,int32 ElementIndex,uint8 InDepthPriorityGroup,FMeshBatch& OutMeshElement, const bool bUseSelectedMaterial, const bool bUseHoveredMaterial) const override;
+	virtual bool GetMeshElements(int32 LODIndex,int32 ElementIndex,uint8 InDepthPriorityGroup, TArray<FMeshBatch, TInlineAllocator<1>>& OutMeshBatches, const bool bUseSelectedMaterial, const bool bUseHoveredMaterial) const override;
 
 	/** Sets up a wireframe FMeshBatch for a specific LOD. */
-	virtual bool GetWireframeMeshElement(int32 LODIndex, const FMaterialRenderProxy* WireframeRenderProxy, uint8 InDepthPriorityGroup, FMeshBatch& OutMeshElement) const override;
+	virtual bool GetWireframeMeshElements(int32 LODIndex, const FMaterialRenderProxy* WireframeRenderProxy, uint8 InDepthPriorityGroup, TArray<FMeshBatch, TInlineAllocator<1>>& OutMeshBatches) const override;
 
 	/**
 	 * Creates the hit proxies are used when DrawDynamicElements is called.
@@ -1019,22 +1015,26 @@ void FInstancedStaticMeshSceneProxy::DrawDynamicElements(FPrimitiveDrawInterface
 
 			for(int32 SectionIndex = 0;SectionIndex < LODModel.Sections.Num();SectionIndex++)
 			{
-				FMeshBatch MeshElement;
-				if(GetMeshElement(LODIndex,SectionIndex,GetDepthPriorityGroup(View), MeshElement, PassRenderSelection[Pass], IsHovered()))
+				TArray<FMeshBatch, TInlineAllocator<1>> MeshBatches;
+				new (MeshBatches) FMeshBatch();
+				if (GetMeshElements(LODIndex,SectionIndex,GetDepthPriorityGroup(View), MeshBatches, PassRenderSelection[Pass], IsHovered()))
 				{
-					MeshElement.Elements[0].UserData = PassUserData[Pass];
-
-					const int32 NumCalls = DrawRichMesh(
-						PDI,
-						MeshElement,
-						WireframeColor,
-						UtilColor,
-						PropertyColor,
-						this,
-						PassRenderSelection[Pass],
-						bIsWireframe
-						);
-					INC_DWORD_STAT_BY(STAT_StaticMeshTriangles,MeshElement.GetNumPrimitives() * NumCalls);
+					for (int32 Index = 0; Index < MeshBatches.Num(); ++Index)
+					{
+						MeshBatches[Index].Elements[0].UserData = PassUserData[Pass];
+
+						const int32 NumCalls = DrawRichMesh(
+							PDI,
+							MeshBatches[Index],
+							WireframeColor,
+							UtilColor,
+							PropertyColor,
+							this,
+							PassRenderSelection[Pass],
+							bIsWireframe
+							);
+						INC_DWORD_STAT_BY(STAT_StaticMeshTriangles,MeshBatches[Index].GetNumPrimitives() * NumCalls);
+					}
 				}
 			}
 		}
@@ -1042,12 +1042,12 @@ void FInstancedStaticMeshSceneProxy::DrawDynamicElements(FPrimitiveDrawInterface
 #endif
 }
 
-inline void FInstancedStaticMeshSceneProxy::SetupInstancedMeshBatch(int32 LODIndex, FMeshBatch& OutMeshBatch) const
+inline void FInstancedStaticMeshSceneProxy::SetupInstancedMeshBatch(int32 LODIndex, TArray<FMeshBatch, TInlineAllocator<1>>& OutMeshBatches) const
 {
 	const bool bInstanced = RHISupportsInstancing(GRHIShaderPlatform);
-	OutMeshBatch.VertexFactory = &InstancedRenderData.VertexFactories[LODIndex];
+	OutMeshBatches[0].VertexFactory = &InstancedRenderData.VertexFactories[LODIndex];
 	uint32 NumInstances = InstancedRenderData.InstanceBuffer.GetNumInstances();
-	auto* OutBatchElement0 = OutMeshBatch.Elements.GetTypedData();
+	auto* OutBatchElement0 = OutMeshBatches[0].Elements.GetTypedData();
 	OutBatchElement0->UserData = (void*)&UserData_AllInstances;
 	OutBatchElement0->UserIndex = 0;
 	if (bInstanced && NumInstances > 1)
@@ -1056,45 +1056,74 @@ inline void FInstancedStaticMeshSceneProxy::SetupInstancedMeshBatch(int32 LODInd
 	}
 	else
 	{
-		OutMeshBatch.Elements.Reserve(NumInstances);
-		OutBatchElement0->NumInstances = 1;
-		for (uint32 Instance = 1; Instance < NumInstances; ++Instance)
+		// Create a BatchElement per Instance, and if they overflow the size per mask, create additional MeshBatches
+
+		const uint32 TotalInstances = NumInstances;
+		const uint32 MaxInstancesPerBatch = FInstancedStaticMeshVertexFactory::NumBitsForVisibilityMask();
+		uint32 NumBatches = (TotalInstances + MaxInstancesPerBatch - 1) / MaxInstancesPerBatch;
+
+		// Make a copy of the Original MeshBatch, as it only has one Element at this point (cheaper to memcpy)
+		FMeshBatch OriginalBatch = OutMeshBatches[0];
+
+		// Add more Batches if we can't fit all Instances within the visibility mask
+		OutMeshBatches.Reserve(NumBatches);
+
+		uint32 RemainingInstances = TotalInstances;
+		int32 InstanceIndex = 0;
+		for (uint32 Batch = 0; Batch < NumBatches; ++Batch)
 		{
-			auto* NewBatchElement = new(OutMeshBatch.Elements) FMeshBatchElement();
-			*NewBatchElement = *OutBatchElement0;
-			NewBatchElement->UserIndex = Instance;
+			NumInstances = FMath::Min(RemainingInstances, MaxInstancesPerBatch);
+			if (Batch > 0)
+			{
+				auto* NewBatch = new (OutMeshBatches) FMeshBatch();
+				*NewBatch = OriginalBatch;
+			}
+
+			OutMeshBatches[Batch].Elements.Reserve(NumInstances);
+			for (uint32 Instance = 0; Instance < NumInstances; ++Instance)
+			{
+				auto* NewBatchElement = (Instance == 0)
+					? OutMeshBatches[Batch].Elements.GetTypedData()
+					: new(OutMeshBatches[Batch].Elements) FMeshBatchElement();
+				*NewBatchElement = *OutBatchElement0;
+				NewBatchElement->UserIndex = InstanceIndex++;
+			}
+
+			RemainingInstances -= NumInstances;
 		}
+
+		check(RemainingInstances == 0);
 	}
 }
 
 
-bool FInstancedStaticMeshSceneProxy::GetShadowMeshElement(int32 LODIndex, uint8 InDepthPriorityGroup, FMeshBatch& OutMeshElement) const
+bool FInstancedStaticMeshSceneProxy::GetShadowMeshElements(int32 LODIndex, uint8 InDepthPriorityGroup, TArray<FMeshBatch, TInlineAllocator<1>>& OutMeshBatches) const
 {
-	if (LODIndex < InstancedRenderData.VertexFactories.Num() && FStaticMeshSceneProxy::GetShadowMeshElement(LODIndex, InDepthPriorityGroup, OutMeshElement))
+	if (LODIndex < InstancedRenderData.VertexFactories.Num() && FStaticMeshSceneProxy::GetShadowMeshElements(LODIndex, InDepthPriorityGroup, OutMeshBatches))
 	{
-		SetupInstancedMeshBatch(LODIndex, OutMeshElement);
+		SetupInstancedMeshBatch(LODIndex, OutMeshBatches);
 		return true;
 	}
 	return false;
 }
 
 /** Sets up a FMeshBatch for a specific LOD and element. */
-bool FInstancedStaticMeshSceneProxy::GetMeshElement(int32 LODIndex,int32 ElementIndex,uint8 InDepthPriorityGroup,FMeshBatch& OutMeshElement, const bool bUseSelectedMaterial, const bool bUseHoveredMaterial) const
+bool FInstancedStaticMeshSceneProxy::GetMeshElements(int32 LODIndex,int32 ElementIndex,uint8 InDepthPriorityGroup, TArray<FMeshBatch, TInlineAllocator<1>>& OutMeshBatches, const bool bUseSelectedMaterial, const bool bUseHoveredMaterial) const
 {
-	if (LODIndex < InstancedRenderData.VertexFactories.Num() && FStaticMeshSceneProxy::GetMeshElement(LODIndex, ElementIndex, InDepthPriorityGroup, OutMeshElement, bUseSelectedMaterial, bUseHoveredMaterial))
+	if (LODIndex < InstancedRenderData.VertexFactories.Num() && FStaticMeshSceneProxy::GetMeshElements(LODIndex, ElementIndex, InDepthPriorityGroup, OutMeshBatches, bUseSelectedMaterial, bUseHoveredMaterial))
 	{
-		SetupInstancedMeshBatch(LODIndex, OutMeshElement);
+		SetupInstancedMeshBatch(LODIndex, OutMeshBatches);
 		return true;
 	}
 	return false;
 };
 
 /** Sets up a wireframe FMeshBatch for a specific LOD. */
-bool FInstancedStaticMeshSceneProxy::GetWireframeMeshElement(int32 LODIndex, const FMaterialRenderProxy* WireframeRenderProxy, uint8 InDepthPriorityGroup, FMeshBatch& OutMeshElement) const
+bool FInstancedStaticMeshSceneProxy::GetWireframeMeshElements(int32 LODIndex, const FMaterialRenderProxy* WireframeRenderProxy, uint8 InDepthPriorityGroup, TArray<FMeshBatch, TInlineAllocator<1>>& OutMeshBatches) const
 {
-	if (LODIndex < InstancedRenderData.VertexFactories.Num() && FStaticMeshSceneProxy::GetWireframeMeshElement(LODIndex, WireframeRenderProxy, InDepthPriorityGroup, OutMeshElement))
+	if (LODIndex < InstancedRenderData.VertexFactories.Num() && FStaticMeshSceneProxy::GetWireframeMeshElements(LODIndex, WireframeRenderProxy, InDepthPriorityGroup, OutMeshBatches))
 	{
-		SetupInstancedMeshBatch(LODIndex, OutMeshElement);
+		SetupInstancedMeshBatch(LODIndex, OutMeshBatches);
 		return true;
 	}
 	return false;
@@ -1676,7 +1705,6 @@ void UInstancedStaticMeshComponent::SelectInstance(bool bInSelected, int32 InIns
 
 void FInstancedStaticMeshVertexFactoryShaderParameters::SetMesh( FRHICommandList& RHICmdList, FShader* VertexShader,const class FVertexFactory* VertexFactory,const class FSceneView& View,const struct FMeshBatchElement& BatchElement,uint32 DataFlags ) const 
 {
-	SCOPED_DRAW_EVENT(STATIC_INSTANCED, DEC_SCENE_ITEMS);
 	FRHIVertexShader* VS = VertexShader->GetVertexShader();
 	if( InstancingFadeOutParamsParameter.IsBound() )
 	{
diff --git a/Engine/Source/Runtime/Engine/Private/SplineMeshSceneProxy.cpp b/Engine/Source/Runtime/Engine/Private/SplineMeshSceneProxy.cpp
--- a/Engine/Source/Runtime/Engine/Private/SplineMeshSceneProxy.cpp
+++ b/Engine/Source/Runtime/Engine/Private/SplineMeshSceneProxy.cpp
@@ -50,40 +50,40 @@ FSplineMeshSceneProxy::~FSplineMeshSceneProxy()
 	LODResources.Empty();
 }
 
-bool FSplineMeshSceneProxy::GetShadowMeshElement(int32 LODIndex, uint8 InDepthPriorityGroup, FMeshBatch& OutMeshElement) const
+bool FSplineMeshSceneProxy::GetShadowMeshElements(int32 LODIndex, uint8 InDepthPriorityGroup, TArray<FMeshBatch, TInlineAllocator<1>>& OutMeshBatches) const
 {
 	//checkf(LODIndex == 0, TEXT("Getting spline static mesh element with invalid LOD [%d]"), LODIndex);
 
-	if (FStaticMeshSceneProxy::GetShadowMeshElement(LODIndex, InDepthPriorityGroup, OutMeshElement))
+	if (FStaticMeshSceneProxy::GetShadowMeshElements(LODIndex, InDepthPriorityGroup, OutMeshBatches))
 	{
-		OutMeshElement.VertexFactory = LODResources[LODIndex].VertexFactory;
-		OutMeshElement.ReverseCulling ^= (SplineParams.StartScale.X < 0) ^ (SplineParams.StartScale.Y < 0);
+		OutMeshBatches[0].VertexFactory = LODResources[LODIndex].VertexFactory;
+		OutMeshBatches[0].ReverseCulling ^= (SplineParams.StartScale.X < 0) ^ (SplineParams.StartScale.Y < 0);
 		return true;
 	}
 	return false;
 }
 
-bool FSplineMeshSceneProxy::GetMeshElement(int32 LODIndex, int32 SectionIndex, uint8 InDepthPriorityGroup, FMeshBatch& OutMeshElement, const bool bUseSelectedMaterial, const bool bUseHoveredMaterial) const
+bool FSplineMeshSceneProxy::GetMeshElements(int32 LODIndex, int32 SectionIndex, uint8 InDepthPriorityGroup, TArray<FMeshBatch, TInlineAllocator<1>>& OutMeshBatches, const bool bUseSelectedMaterial, const bool bUseHoveredMaterial) const
 {
 	//checkf(LODIndex == 0 /*&& SectionIndex == 0*/, TEXT("Getting spline static mesh element with invalid params [%d, %d]"), LODIndex, SectionIndex);
 
-	if (FStaticMeshSceneProxy::GetMeshElement(LODIndex, SectionIndex, InDepthPriorityGroup, OutMeshElement, bUseSelectedMaterial, bUseHoveredMaterial))
+	if (FStaticMeshSceneProxy::GetMeshElements(LODIndex, SectionIndex, InDepthPriorityGroup, OutMeshBatches, bUseSelectedMaterial, bUseHoveredMaterial))
 	{
-		OutMeshElement.VertexFactory = LODResources[LODIndex].VertexFactory;
-		OutMeshElement.ReverseCulling ^= (SplineParams.StartScale.X < 0) ^ (SplineParams.StartScale.Y < 0);
+		OutMeshBatches[0].VertexFactory = LODResources[LODIndex].VertexFactory;
+		OutMeshBatches[0].ReverseCulling ^= (SplineParams.StartScale.X < 0) ^ (SplineParams.StartScale.Y < 0);
 		return true;
 	}
 	return false;
 }
 
-bool FSplineMeshSceneProxy::GetWireframeMeshElement(int32 LODIndex, const FMaterialRenderProxy* WireframeRenderProxy, uint8 InDepthPriorityGroup, FMeshBatch& OutMeshElement) const
+bool FSplineMeshSceneProxy::GetWireframeMeshElements(int32 LODIndex, const FMaterialRenderProxy* WireframeRenderProxy, uint8 InDepthPriorityGroup, TArray<FMeshBatch, TInlineAllocator<1>>& OutMeshBatches) const
 {
 	//checkf(LODIndex == 0, TEXT("Getting spline static mesh element with invalid LOD [%d]"), LODIndex);
 
-	if (FStaticMeshSceneProxy::GetWireframeMeshElement(LODIndex, WireframeRenderProxy, InDepthPriorityGroup, OutMeshElement))
+	if (FStaticMeshSceneProxy::GetWireframeMeshElements(LODIndex, WireframeRenderProxy, InDepthPriorityGroup, OutMeshBatches))
 	{
-		OutMeshElement.VertexFactory = LODResources[LODIndex].VertexFactory;
-		OutMeshElement.ReverseCulling ^= (SplineParams.StartScale.X < 0) ^ (SplineParams.StartScale.Y < 0);
+		OutMeshBatches[0].VertexFactory = LODResources[LODIndex].VertexFactory;
+		OutMeshBatches[0].ReverseCulling ^= (SplineParams.StartScale.X < 0) ^ (SplineParams.StartScale.Y < 0);
 		return true;
 	}
 	return false;
diff --git a/Engine/Source/Runtime/Engine/Private/SplineMeshSceneProxy.h b/Engine/Source/Runtime/Engine/Private/SplineMeshSceneProxy.h
--- a/Engine/Source/Runtime/Engine/Private/SplineMeshSceneProxy.h
+++ b/Engine/Source/Runtime/Engine/Private/SplineMeshSceneProxy.h
@@ -155,13 +155,13 @@ public:
 	void ReleaseResources();
 
 	/** Sets up a shadow FMeshBatch for a specific LOD. */
-	virtual bool GetShadowMeshElement(int32 LODIndex, uint8 InDepthPriorityGroup, FMeshBatch& OutMeshElement) const override;
+	virtual bool GetShadowMeshElements(int32 LODIndex, uint8 InDepthPriorityGroup, TArray<FMeshBatch, TInlineAllocator<1>>& OutMeshBatches) const override;
 
 	/** Sets up a FMeshBatch for a specific LOD and element. */
-	virtual bool GetMeshElement(int32 LODIndex, int32 SectionIndex, uint8 InDepthPriorityGroup, FMeshBatch& OutMeshElement, const bool bUseSelectedMaterial, const bool bUseHoveredMaterial) const override;
+	virtual bool GetMeshElements(int32 LODIndex, int32 SectionIndex, uint8 InDepthPriorityGroup, TArray<FMeshBatch, TInlineAllocator<1>>& OutMeshBatches, const bool bUseSelectedMaterial, const bool bUseHoveredMaterial) const override;
 
 	/** Sets up a wireframe FMeshBatch for a specific LOD. */
-	virtual bool GetWireframeMeshElement(int32 LODIndex, const FMaterialRenderProxy* WireframeRenderProxy, uint8 InDepthPriorityGroup, FMeshBatch& OutMeshElement) const override;
+	virtual bool GetWireframeMeshElements(int32 LODIndex, const FMaterialRenderProxy* WireframeRenderProxy, uint8 InDepthPriorityGroup, TArray<FMeshBatch, TInlineAllocator<1>>& OutMeshBatches) const override;
 
 	virtual FPrimitiveViewRelevance GetViewRelevance(const FSceneView* View) override
 	{
diff --git a/Engine/Source/Runtime/Engine/Private/StaticMeshRender.cpp b/Engine/Source/Runtime/Engine/Private/StaticMeshRender.cpp
--- a/Engine/Source/Runtime/Engine/Private/StaticMeshRender.cpp
+++ b/Engine/Source/Runtime/Engine/Private/StaticMeshRender.cpp
@@ -149,26 +149,26 @@ void UStaticMeshComponent::SetLODDataCount( const uint32 MinSize, const uint32 M
 	}
 }
 
-bool FStaticMeshSceneProxy::GetShadowMeshElement(int32 LODIndex, uint8 InDepthPriorityGroup, FMeshBatch& OutMeshElement) const
+bool FStaticMeshSceneProxy::GetShadowMeshElements(int32 LODIndex, uint8 InDepthPriorityGroup, TArray<FMeshBatch, TInlineAllocator<1>>& OutMeshBatches) const
 {
 	const FStaticMeshLODResources& LOD = RenderData->LODResources[LODIndex];
 	const FLODInfo& ProxyLODInfo = LODs[LODIndex];
 
-	FMeshBatchElement& OutBatchElement = OutMeshElement.Elements[0];
-	OutMeshElement.MaterialRenderProxy = UMaterial::GetDefaultMaterial(MD_Surface)->GetRenderProxy(false, false);
-	OutMeshElement.VertexFactory = &LOD.VertexFactory;
+	FMeshBatchElement& OutBatchElement = OutMeshBatches[0].Elements[0];
+	OutMeshBatches[0].MaterialRenderProxy = UMaterial::GetDefaultMaterial(MD_Surface)->GetRenderProxy(false, false);
+	OutMeshBatches[0].VertexFactory = &LOD.VertexFactory;
 	OutBatchElement.IndexBuffer = &LOD.DepthOnlyIndexBuffer;
-	OutMeshElement.Type = PT_TriangleList;
+	OutMeshBatches[0].Type = PT_TriangleList;
 	OutBatchElement.FirstIndex = 0;
 	OutBatchElement.NumPrimitives = LOD.DepthOnlyIndexBuffer.GetNumIndices() / 3;
 	OutBatchElement.PrimitiveUniformBufferResource = &GetUniformBuffer();
 	OutBatchElement.MinVertexIndex = 0;
 	OutBatchElement.MaxVertexIndex = LOD.PositionVertexBuffer.GetNumVertices() - 1;
-	OutMeshElement.DepthPriorityGroup = InDepthPriorityGroup;
-	OutMeshElement.ReverseCulling = IsLocalToWorldDeterminantNegative();
-	OutMeshElement.CastShadow = true;
-	OutMeshElement.LODIndex = LODIndex;
-	OutMeshElement.LCI = &ProxyLODInfo;
+	OutMeshBatches[0].DepthPriorityGroup = InDepthPriorityGroup;
+	OutMeshBatches[0].ReverseCulling = IsLocalToWorldDeterminantNegative();
+	OutMeshBatches[0].CastShadow = true;
+	OutMeshBatches[0].LODIndex = LODIndex;
+	OutMeshBatches[0].LCI = &ProxyLODInfo;
 	if (ForcedLodModel > 0) 
 	{
 		OutBatchElement.MaxScreenSize = 0.0f;
@@ -188,15 +188,15 @@ bool FStaticMeshSceneProxy::GetShadowMeshElement(int32 LODIndex, uint8 InDepthPr
 }
 
 /** Sets up a FMeshBatch for a specific LOD and element. */
-bool FStaticMeshSceneProxy::GetMeshElement(int32 LODIndex,int32 SectionIndex,uint8 InDepthPriorityGroup,FMeshBatch& OutMeshElement, const bool bUseSelectedMaterial, const bool bUseHoveredMaterial) const
+bool FStaticMeshSceneProxy::GetMeshElements(int32 LODIndex,int32 SectionIndex,uint8 InDepthPriorityGroup, TArray<FMeshBatch, TInlineAllocator<1>>& OutMeshBatches, const bool bUseSelectedMaterial, const bool bUseHoveredMaterial) const
 {
 	const FStaticMeshLODResources& LOD = RenderData->LODResources[LODIndex];
 	const FStaticMeshSection& Section = LOD.Sections[SectionIndex];
 
 	const FLODInfo& ProxyLODInfo = LODs[LODIndex];
 	UMaterialInterface* Material = ProxyLODInfo.Sections[SectionIndex].Material;
-	OutMeshElement.MaterialRenderProxy = Material->GetRenderProxy(bUseSelectedMaterial,bUseHoveredMaterial);
-	OutMeshElement.VertexFactory = &LOD.VertexFactory;
+	OutMeshBatches[0].MaterialRenderProxy = Material->GetRenderProxy(bUseSelectedMaterial,bUseHoveredMaterial);
+	OutMeshBatches[0].VertexFactory = &LOD.VertexFactory;
 
 	// Has the mesh component overridden the vertex color stream for this mesh LOD?
 	if( ProxyLODInfo.OverrideColorVertexBuffer != NULL )
@@ -208,29 +208,29 @@ bool FStaticMeshSceneProxy::GetMeshElement(int32 LODIndex,int32 SectionIndex,uin
 		{
 			// Switch out the stock mesh vertex factory with own own vertex factory that points to
 			// our overridden color data
-			OutMeshElement.VertexFactory = ProxyLODInfo.OverrideColorVertexFactory.GetOwnedPointer();
+			OutMeshBatches[0].VertexFactory = ProxyLODInfo.OverrideColorVertexFactory.GetOwnedPointer();
 		}
 	}
 
 	const bool bWireframe = false;
-	const bool bRequiresAdjacencyInformation = RequiresAdjacencyInformation( Material, OutMeshElement.VertexFactory->GetType(), GetScene()->GetFeatureLevel() );
+	const bool bRequiresAdjacencyInformation = RequiresAdjacencyInformation( Material, OutMeshBatches[0].VertexFactory->GetType(), GetScene()->GetFeatureLevel() );
 
-	SetIndexSource(LODIndex, SectionIndex, OutMeshElement, bWireframe, bRequiresAdjacencyInformation );
+	SetIndexSource(LODIndex, SectionIndex, OutMeshBatches[0], bWireframe, bRequiresAdjacencyInformation );
 
-	FMeshBatchElement& OutBatchElement = OutMeshElement.Elements[0];
+	FMeshBatchElement& OutBatchElement = OutMeshBatches[0].Elements[0];
 
 	if(OutBatchElement.NumPrimitives > 0)
 	{
-		OutMeshElement.DynamicVertexData = NULL;
-		OutMeshElement.LCI = &ProxyLODInfo;
+		OutMeshBatches[0].DynamicVertexData = NULL;
+		OutMeshBatches[0].LCI = &ProxyLODInfo;
 		OutBatchElement.PrimitiveUniformBufferResource = &GetUniformBuffer();
 		OutBatchElement.MinVertexIndex = Section.MinVertexIndex;
 		OutBatchElement.MaxVertexIndex = Section.MaxVertexIndex;
-		OutMeshElement.LODIndex = LODIndex;
-		OutMeshElement.UseDynamicData = false;
-		OutMeshElement.ReverseCulling = IsLocalToWorldDeterminantNegative();
-		OutMeshElement.CastShadow = bCastShadow && Section.bCastShadow;
-		OutMeshElement.DepthPriorityGroup = (ESceneDepthPriorityGroup)InDepthPriorityGroup;
+		OutMeshBatches[0].LODIndex = LODIndex;
+		OutMeshBatches[0].UseDynamicData = false;
+		OutMeshBatches[0].ReverseCulling = IsLocalToWorldDeterminantNegative();
+		OutMeshBatches[0].CastShadow = bCastShadow && Section.bCastShadow;
+		OutMeshBatches[0].DepthPriorityGroup = (ESceneDepthPriorityGroup)InDepthPriorityGroup;
 		if (ForcedLodModel > 0) 
 		{
 			OutBatchElement.MaxScreenSize = 0.0f;
@@ -255,20 +255,20 @@ bool FStaticMeshSceneProxy::GetMeshElement(int32 LODIndex,int32 SectionIndex,uin
 }
 
 /** Sets up a wireframe FMeshBatch for a specific LOD. */
-bool FStaticMeshSceneProxy::GetWireframeMeshElement(int32 LODIndex, const FMaterialRenderProxy* WireframeRenderProxy, uint8 InDepthPriorityGroup, FMeshBatch& OutMeshElement) const
+bool FStaticMeshSceneProxy::GetWireframeMeshElements(int32 LODIndex, const FMaterialRenderProxy* WireframeRenderProxy, uint8 InDepthPriorityGroup, TArray<FMeshBatch, TInlineAllocator<1>>& OutMeshBatches) const
 {
 	const FStaticMeshLODResources& LODModel = RenderData->LODResources[LODIndex];
 	
-	FMeshBatchElement& OutBatchElement = OutMeshElement.Elements[0];
+	FMeshBatchElement& OutBatchElement = OutMeshBatches[0].Elements[0];
 
-	OutMeshElement.VertexFactory = &LODModel.VertexFactory;
-	OutMeshElement.MaterialRenderProxy = WireframeRenderProxy;
+	OutMeshBatches[0].VertexFactory = &LODModel.VertexFactory;
+	OutMeshBatches[0].MaterialRenderProxy = WireframeRenderProxy;
 	OutBatchElement.PrimitiveUniformBufferResource = &GetUniformBuffer();
 	OutBatchElement.MinVertexIndex = 0;
 	OutBatchElement.MaxVertexIndex = LODModel.GetNumVertices() - 1;
-	OutMeshElement.ReverseCulling = IsLocalToWorldDeterminantNegative();
-	OutMeshElement.CastShadow = bCastShadow;
-	OutMeshElement.DepthPriorityGroup = (ESceneDepthPriorityGroup)InDepthPriorityGroup;
+	OutMeshBatches[0].ReverseCulling = IsLocalToWorldDeterminantNegative();
+	OutMeshBatches[0].CastShadow = bCastShadow;
+	OutMeshBatches[0].DepthPriorityGroup = (ESceneDepthPriorityGroup)InDepthPriorityGroup;
 	if (ForcedLodModel > 0) 
 	{
 		OutBatchElement.MaxScreenSize = 0.0f;
@@ -287,7 +287,7 @@ bool FStaticMeshSceneProxy::GetWireframeMeshElement(int32 LODIndex, const FMater
 	const bool bWireframe = true;
 	const bool bRequiresAdjacencyInformation = false;
 
-	SetIndexSource(LODIndex, 0, OutMeshElement, bWireframe, bRequiresAdjacencyInformation);
+	SetIndexSource(LODIndex, 0, OutMeshBatches[0], bWireframe, bRequiresAdjacencyInformation);
 
 	return OutBatchElement.NumPrimitives > 0;
 }
@@ -438,10 +438,14 @@ void FStaticMeshSceneProxy::DrawStaticElements(FStaticPrimitiveDrawInterface* PD
 				}
 #endif // WITH_EDITOR
 
-				FMeshBatch MeshElement;
-				if(GetMeshElement(LODIndex, SectionIndex, PrimitiveDPG, MeshElement, bUseSelectedMaterial, bUseHoveredMaterial))
+				TArray<FMeshBatch, TInlineAllocator<1>> MeshBatches;
+				new (MeshBatches) FMeshBatch();
+				if(GetMeshElements(LODIndex, SectionIndex, PrimitiveDPG, MeshBatches, bUseSelectedMaterial, bUseHoveredMaterial))
 				{
-					PDI->DrawMesh(MeshElement, FLT_MAX);
+					for (int32 Index = 0; Index < MeshBatches.Num(); ++Index)
+					{
+						PDI->DrawMesh(MeshBatches[Index], FLT_MAX);
+					}
 				}
 			}
 		} 
@@ -478,11 +482,15 @@ void FStaticMeshSceneProxy::DrawStaticElements(FStaticPrimitiveDrawInterface* PD
 
 					if (bAnySectionCastsShadow && bSafeToUseShadowOnlyMesh)
 					{
-						FMeshBatch MeshElement;
-						if (GetShadowMeshElement(LODIndex, PrimitiveDPG, MeshElement))
+						TArray<FMeshBatch, TInlineAllocator<1>> MeshBatches;
+						new (MeshBatches) FMeshBatch();
+						if (GetShadowMeshElements(LODIndex, PrimitiveDPG, MeshBatches))
 						{
 							bHaveShadowOnlyMesh = true;
-							PDI->DrawMesh(MeshElement, ScreenSize, /*bShadowOnly=*/true);
+							for (int32 Index = 0; Index < MeshBatches.Num(); ++Index)
+							{
+								PDI->DrawMesh(MeshBatches[Index], ScreenSize, /*bShadowOnly=*/true);
+							}
 						}
 					}
 				}
@@ -500,12 +508,16 @@ void FStaticMeshSceneProxy::DrawStaticElements(FStaticPrimitiveDrawInterface* PD
 					}
 #endif // WITH_EDITOR
 
-					FMeshBatch MeshElement;
-					if(GetMeshElement(LODIndex, SectionIndex, PrimitiveDPG, MeshElement, bUseSelectedMaterial, bUseHoveredMaterial))
+					TArray<FMeshBatch, TInlineAllocator<1>> MeshBatches;
+					new (MeshBatches) FMeshBatch();
+					if(GetMeshElements(LODIndex, SectionIndex, PrimitiveDPG, MeshBatches, bUseSelectedMaterial, bUseHoveredMaterial))
 					{
 						// If we have submitted an optimized shadow-only mesh, remaining mesh elements must not cast shadows.
-						MeshElement.CastShadow = MeshElement.CastShadow && !bHaveShadowOnlyMesh;
-						PDI->DrawMesh(MeshElement, ScreenSize);
+						for (int32 Index = 0; Index < MeshBatches.Num(); ++Index)
+						{
+							MeshBatches[Index].CastShadow = MeshBatches[Index].CastShadow && !bHaveShadowOnlyMesh;
+							PDI->DrawMesh(MeshBatches[Index], ScreenSize);
+						}
 					}
 				}
 			}
@@ -603,13 +615,15 @@ void FStaticMeshSceneProxy::DrawDynamicElements(FPrimitiveDrawInterface* PDI,con
 					GetSelectionColor(ViewWireframeColor,!(GIsEditor && View->Family->EngineShowFlags.Selection) || bProxyIsSelected, IsHovered(), /*bUseOverlayIntensity=*/false)
 					);
 
-				FMeshBatch Mesh;
-
-				if (GetWireframeMeshElement(LODIndex, &WireframeMaterialInstance, GetDepthPriorityGroup(View), Mesh))
+				TArray<FMeshBatch, TInlineAllocator<1>> MeshBatches;
+				new (MeshBatches) FMeshBatch();
+				if (GetWireframeMeshElements(LODIndex, &WireframeMaterialInstance, GetDepthPriorityGroup(View), MeshBatches))
 				{
-					const int32 NumPasses = PDI->DrawMesh(Mesh);
-
-					INC_DWORD_STAT_BY(STAT_StaticMeshTriangles,Mesh.GetNumPrimitives() * NumPasses);
+					for (int32 Index = 0; Index < MeshBatches.Num(); ++Index)
+					{
+						const int32 NumPasses = PDI->DrawMesh(MeshBatches[Index]);
+						INC_DWORD_STAT_BY(STAT_StaticMeshTriangles,MeshBatches[Index].GetNumPrimitives() * NumPasses);
+					}
 				}
 			}
 			else
@@ -631,8 +645,9 @@ void FStaticMeshSceneProxy::DrawDynamicElements(FPrimitiveDrawInterface* PDI,con
 					}
 #endif // WITH_EDITOR
 
-					FMeshBatch MeshElement;
-					if(GetMeshElement(LODIndex,SectionIndex,GetDepthPriorityGroup(View),MeshElement,bSectionIsSelected,IsHovered()))
+					TArray<FMeshBatch, TInlineAllocator<1>> MeshBatches;
+					new (MeshBatches) FMeshBatch();
+					if(GetMeshElements(LODIndex,SectionIndex,GetDepthPriorityGroup(View), MeshBatches,bSectionIsSelected,IsHovered()))
 					{
 #if !(UE_BUILD_SHIPPING || UE_BUILD_TEST)
 						if( View->Family->EngineShowFlags.VertexColors && AllowDebugViewmodes() )
@@ -663,37 +678,14 @@ void FStaticMeshSceneProxy::DrawDynamicElements(FPrimitiveDrawInterface* PDI,con
 							}
 							check( VertexColorVisualizationMaterial != NULL );
 
-							const FColoredMaterialRenderProxy VertexColorVisualizationMaterialInstance(
-								VertexColorVisualizationMaterial->GetRenderProxy( MeshElement.MaterialRenderProxy->IsSelected(),MeshElement.MaterialRenderProxy->IsHovered() ),
-								GetSelectionColor( FLinearColor::White, bSectionIsSelected, IsHovered() )
-								);
-							FMeshBatch ModifiedMeshElement = MeshElement;
-							ModifiedMeshElement.MaterialRenderProxy = &VertexColorVisualizationMaterialInstance;
-
-							const int32 NumPasses = DrawRichMesh(
-								PDI,
-								ModifiedMeshElement,
-								WireframeColor,
-								UtilColor,
-								PropertyColor,
-								this,
-								bSectionIsSelected,
-								bIsWireframeView
-								);
-							INC_DWORD_STAT_BY(STAT_StaticMeshTriangles,MeshElement.GetNumPrimitives() * NumPasses);
-						}
-						else if( bInCollisionView && bDrawComplexCollision && AllowDebugViewmodes() )
-						{
-							if (LODModel.Sections[SectionIndex].bEnableCollision)
+							for (int32 Index = 0; Index < MeshBatches.Num(); ++Index)
 							{
-								// Override the mesh's material with our material that draws the collision color
-								const FColoredMaterialRenderProxy CollisionMaterialInstance(
-									GEngine->ShadedLevelColorationUnlitMaterial->GetRenderProxy(bProxyIsSelected, IsHovered()),
-									WireframeColor
+								const FColoredMaterialRenderProxy VertexColorVisualizationMaterialInstance(
+									VertexColorVisualizationMaterial->GetRenderProxy( MeshBatches[Index].MaterialRenderProxy->IsSelected(),MeshBatches[Index].MaterialRenderProxy->IsHovered() ),
+									GetSelectionColor( FLinearColor::White, bSectionIsSelected, IsHovered() )
 									);
-
-								FMeshBatch ModifiedMeshElement = MeshElement;
-								ModifiedMeshElement.MaterialRenderProxy = &CollisionMaterialInstance;
+								FMeshBatch ModifiedMeshElement = MeshBatches[Index];
+								ModifiedMeshElement.MaterialRenderProxy = &VertexColorVisualizationMaterialInstance;
 
 								const int32 NumPasses = DrawRichMesh(
 									PDI,
@@ -705,7 +697,36 @@ void FStaticMeshSceneProxy::DrawDynamicElements(FPrimitiveDrawInterface* PDI,con
 									bSectionIsSelected,
 									bIsWireframeView
 									);
-								INC_DWORD_STAT_BY(STAT_StaticMeshTriangles,MeshElement.GetNumPrimitives() * NumPasses);
+								INC_DWORD_STAT_BY(STAT_StaticMeshTriangles,MeshBatches[Index].GetNumPrimitives() * NumPasses);
+							}
+						}
+						else if( bInCollisionView && bDrawComplexCollision && AllowDebugViewmodes() )
+						{
+							if (LODModel.Sections[SectionIndex].bEnableCollision)
+							{
+								// Override the mesh's material with our material that draws the collision color
+								const FColoredMaterialRenderProxy CollisionMaterialInstance(
+									GEngine->ShadedLevelColorationUnlitMaterial->GetRenderProxy(bProxyIsSelected, IsHovered()),
+									WireframeColor
+									);
+
+								for (int32 Index = 0; Index < MeshBatches.Num(); ++Index)
+								{
+									FMeshBatch ModifiedMeshElement = MeshBatches[Index];
+									ModifiedMeshElement.MaterialRenderProxy = &CollisionMaterialInstance;
+
+									const int32 NumPasses = DrawRichMesh(
+										PDI,
+										ModifiedMeshElement,
+										WireframeColor,
+										UtilColor,
+										PropertyColor,
+										this,
+										bSectionIsSelected,
+										bIsWireframeView
+										);
+									INC_DWORD_STAT_BY(STAT_StaticMeshTriangles,MeshBatches[Index].GetNumPrimitives() * NumPasses);
+								}
 							}
 						}
 						else
@@ -719,35 +740,41 @@ void FStaticMeshSceneProxy::DrawDynamicElements(FPrimitiveDrawInterface* PDI,con
 								GetSelectionColor(GEngine->GetSelectedMaterialColor(), bSectionIsSelected, IsHovered())
 								);
 
-							FMeshBatch ModifiedMeshElement(MeshElement);
-							ModifiedMeshElement.MaterialRenderProxy = &SelectedMaterialInstance;
-
-							const int32 NumPasses = DrawRichMesh(
-								PDI,
-								ModifiedMeshElement,
-								WireframeColor,
-								UtilColor,
-								PropertyColor,
-								this,
-								bSectionIsSelected,
-								bIsWireframeView
-								);
-							INC_DWORD_STAT_BY(STAT_StaticMeshTriangles, MeshElement.GetNumPrimitives() * NumPasses);
+							for (int32 Index = 0; Index < MeshBatches.Num(); ++Index)
+							{
+								FMeshBatch ModifiedMeshElement(MeshBatches[Index]);
+								ModifiedMeshElement.MaterialRenderProxy = &SelectedMaterialInstance;
+
+								const int32 NumPasses = DrawRichMesh(
+									PDI,
+									ModifiedMeshElement,
+									WireframeColor,
+									UtilColor,
+									PropertyColor,
+									this,
+									bSectionIsSelected,
+									bIsWireframeView
+									);
+								INC_DWORD_STAT_BY(STAT_StaticMeshTriangles, MeshBatches[Index].GetNumPrimitives() * NumPasses);
+							}
 						}
 						else
 #endif
 						{
-							const int32 NumPasses = DrawRichMesh(
-								PDI,
-								MeshElement,
-								WireframeColor,
-								UtilColor,
-								PropertyColor,
-								this,
-								bSectionIsSelected,
-								bIsWireframeView
-								);
-							INC_DWORD_STAT_BY(STAT_StaticMeshTriangles,MeshElement.GetNumPrimitives() * NumPasses);
+							for (int32 Index = 0; Index < MeshBatches.Num(); ++Index)
+							{
+								const int32 NumPasses = DrawRichMesh(
+									PDI,
+									MeshBatches[Index],
+									WireframeColor,
+									UtilColor,
+									PropertyColor,
+									this,
+									bSectionIsSelected,
+									bIsWireframeView
+									);
+								INC_DWORD_STAT_BY(STAT_StaticMeshTriangles,MeshBatches[Index].GetNumPrimitives() * NumPasses);
+							}
 						}
 
 					}
diff --git a/Engine/Source/Runtime/Engine/Public/StaticMeshResources.h b/Engine/Source/Runtime/Engine/Public/StaticMeshResources.h
--- a/Engine/Source/Runtime/Engine/Public/StaticMeshResources.h
+++ b/Engine/Source/Runtime/Engine/Public/StaticMeshResources.h
@@ -857,13 +857,13 @@ public:
 	virtual ~FStaticMeshSceneProxy() {}
 
 	/** Sets up a shadow FMeshBatch for a specific LOD. */
-	virtual bool GetShadowMeshElement(int32 LODIndex, uint8 InDepthPriorityGroup, FMeshBatch& OutMeshElement) const;
+	virtual bool GetShadowMeshElements(int32 LODIndex, uint8 InDepthPriorityGroup, TArray<FMeshBatch, TInlineAllocator<1>>& OutMeshBatches) const;
 
 	/** Sets up a FMeshBatch for a specific LOD and element. */
-	virtual bool GetMeshElement(int32 LODIndex,int32 ElementIndex,uint8 InDepthPriorityGroup,FMeshBatch& OutMeshElement, const bool bUseSelectedMaterial, const bool bUseHoveredMaterial) const;
+	virtual bool GetMeshElements(int32 LODIndex,int32 ElementIndex,uint8 InDepthPriorityGroup, TArray<FMeshBatch, TInlineAllocator<1>>& OutMeshBatches, const bool bUseSelectedMaterial, const bool bUseHoveredMaterial) const;
 
 	/** Sets up a wireframe FMeshBatch for a specific LOD. */
-	virtual bool GetWireframeMeshElement(int32 LODIndex, const FMaterialRenderProxy* WireframeRenderProxy, uint8 InDepthPriorityGroup, FMeshBatch& OutMeshElement) const;
+	virtual bool GetWireframeMeshElements(int32 LODIndex, const FMaterialRenderProxy* WireframeRenderProxy, uint8 InDepthPriorityGroup, TArray<FMeshBatch, TInlineAllocator<1>>& OutMeshBatches) const;
 
 protected:
 	/**
