Provide mechanism for disabling force feedback at the PlayerController level
Add a new latent function that allows dynamic management of force feedback

diff --git a/Engine/Source/Runtime/Engine/Classes/GameFramework/PlayerController.h b/Engine/Source/Runtime/Engine/Classes/GameFramework/PlayerController.h
--- a/Engine/Source/Runtime/Engine/Classes/GameFramework/PlayerController.h
+++ b/Engine/Source/Runtime/Engine/Classes/GameFramework/PlayerController.h
@@ -9,6 +9,7 @@
 #include "GameFramework/ForceFeedbackEffect.h"
 #include "GameFramework/OnlineReplStructs.h"
 #include "GameFramework/Controller.h"
+#include "Engine/LatentActionManager.h"
 #include "PlayerController.generated.h"
 
 class FPrimitiveComponentId;
@@ -16,6 +17,55 @@ class FPrimitiveComponentId;
 /** delegate used to override default viewport audio listener position calculated from camera */
 DECLARE_DELEGATE_ThreeParams(FGetAudioListenerPos, FVector& /*Location*/, FVector& /*ProjFront*/, FVector& /*ProjRight*/);
 
+UENUM()
+namespace EDynamicForceFeedbackAction
+{
+	enum Type
+	{
+		Start,
+		Update,
+		Stop,
+	};
+}
+
+struct FDynamicForceFeedbackDetails
+{
+	uint32 bAffectsLeftLarge:1;
+	uint32 bAffectsLeftSmall:1;
+	uint32 bAffectsRightLarge:1;
+	uint32 bAffectsRightSmall:1;
+
+	float Intensity;
+
+	FDynamicForceFeedbackDetails()
+		: bAffectsLeftLarge(true)
+		, bAffectsLeftSmall(true)
+		, bAffectsRightLarge(true)
+		, bAffectsRightSmall(true)
+		, Intensity(0.f)
+	{}
+
+	void Update(FForceFeedbackValues& Values) const
+	{
+		if (bAffectsLeftLarge)
+		{
+			Values.LeftLarge = FMath::Clamp(Intensity, Values.LeftLarge, 1.f);
+		}
+		if (bAffectsLeftSmall)
+		{
+			Values.LeftSmall = FMath::Clamp(Intensity, Values.LeftSmall, 1.f);
+		}
+		if (bAffectsRightLarge)
+		{
+			Values.RightLarge = FMath::Clamp(Intensity, Values.RightLarge, 1.f);
+		}
+		if (bAffectsRightSmall)
+		{
+			Values.RightSmall = FMath::Clamp(Intensity, Values.RightSmall, 1.f);
+		}
+	}
+};
+
 //=============================================================================
 /**
  * PlayerControllers are used by human players to control Pawns.
@@ -111,6 +161,8 @@ class ENGINE_API APlayerController : public AController
 	UPROPERTY(transient)
 	TArray<FActiveForceFeedbackEffect> ActiveForceFeedbackEffects;
 
+	TMap<int32, FDynamicForceFeedbackDetails> DynamicForceFeedbacks;
+
 	/** list of names of levels the server is in the middle of sending us for a PrepareMapChange() call */
 	TArray<FName> PendingMapChangeLevelNames;
 
@@ -180,6 +232,9 @@ class ENGINE_API APlayerController : public AController
 	UPROPERTY(EditAnywhere, BlueprintReadWrite, Category=MouseInterface)
 	uint32 bEnableTouchOverEvents:1;
 
+	UPROPERTY(EditAnywhere, BlueprintReadWrite, Category="Game|Feedback")
+	uint32 bForceFeedbackEnabled:1;
+
 	UPROPERTY(EditAnywhere, BlueprintReadOnly, Category=MouseInterface)
 	TEnumAsByte<EMouseCursor::Type> DefaultMouseCursor;
 
@@ -684,6 +739,21 @@ public:
 	UFUNCTION(unreliable, client, BlueprintCallable, Category="Game|Feedback")
 	void ClientStopForceFeedback(class UForceFeedbackEffect* ForceFeedbackEffect, FName Tag);
 
+	/** 
+	 * Latent action that controls the playing of force feedback 
+	 * Begins playing when Start is called.  Calling Update or Stop if the feedback is not active will have no effect.
+	 * Completed will execute when Stop is called or the duration ends.
+	 * When Update is called the Intensity, Duration, and affect values will be updated with the current inputs
+	 * @param	Intensity				How strong the feedback should be.  Valid values are between 0.0 and 1.0
+	 * @param	Duration				How long the feedback should play for.  If the value is negative it will play until stopped
+	 * @param   bAffectsLeftLarge		Whether the intensity should be applied to the large left servo
+	 * @param   bAffectsLeftSmall		Whether the intensity should be applied to the small left servo
+	 * @param   bAffectsRightLarge		Whether the intensity should be applied to the large right servo
+	 * @param   bAffectsRightSmall		Whether the intensity should be applied to the small right servo
+	 */
+	UFUNCTION(BlueprintCallable, meta=(Latent, LatentInfo="LatentInfo", ExpandEnumAsExecs="Action", Duration="-1", bAffectsLeftLarge="true", bAffectsLeftSmall="true", bAffectsRightLarge="true", bAffectsRightSmall="true", AdvancedDisplay="bAffectsLeftLarge,bAffectsLeftSmall,bAffectsRightLarge,bAffectsRightSmall"), Category="Game|Feedback")
+	void PlayDynamicForceFeedback(float Intensity, float Duration, bool bAffectsLeftLarge, bool bAffectsLeftSmall, bool bAffectsRightLarge, bool bAffectsRightSmall, TEnumAsByte<EDynamicForceFeedbackAction::Type> Action, FLatentActionInfo LatentInfo);
+
 	/**
 	 * Travel to a different map or IP address. Calls the PreClientTravel event before doing anything.
 	 * NOTE: This is implemented as a locally executed wrapper for ClientTravelInternal, to avoid API compatability breakage
@@ -986,6 +1056,7 @@ public:
 	virtual void Possess(APawn* aPawn) override;
 	virtual void UnPossess() override;
 	virtual void CleanupPlayerState() override;
+	virtual void EndPlay(const EEndPlayReason::Type EndPlayReason) override;
 	virtual void Destroyed() override;
 	virtual void OnActorChannelOpen(class FInBunch& InBunch, class UNetConnection* Connection) override;
 	virtual void OnSerializeNewActor(class FOutBunch& OutBunch) override;
diff --git a/Engine/Source/Runtime/Engine/Private/PlayerController.cpp b/Engine/Source/Runtime/Engine/Private/PlayerController.cpp
--- a/Engine/Source/Runtime/Engine/Private/PlayerController.cpp
+++ b/Engine/Source/Runtime/Engine/Private/PlayerController.cpp
@@ -53,6 +53,7 @@ APlayerController::APlayerController(const class FPostConstructInitializePropert
 
 	bInputEnabled = true;
 	bEnableTouchEvents = true;
+	bForceFeedbackEnabled = true;
 
 	bAutoManageActiveCameraTarget = true;
 
@@ -1218,6 +1219,22 @@ void APlayerController::PawnLeavingGame()
 	}
 }
 
+void APlayerController::EndPlay(const EEndPlayReason::Type EndPlayReason)
+{
+	if (VirtualJoystick.IsValid())
+	{
+		CastChecked<ULocalPlayer>(Player)->ViewportClient->RemoveViewportWidgetContent(VirtualJoystick.ToSharedRef());
+	}
+
+	// Stop any force feedback effects that may be active
+	IForceFeedbackSystem* ForceFeedbackSystem = FSlateApplication::Get().GetForceFeedbackSystem();
+	if (ForceFeedbackSystem)
+	{
+		ForceFeedbackSystem->SetChannelValues(CastChecked<ULocalPlayer>(Player)->ControllerId, FForceFeedbackValues());
+	}
+
+	Super::EndPlay(EndPlayReason);
+}
 
 void APlayerController::Destroyed()
 {
@@ -1261,11 +1278,6 @@ void APlayerController::Destroyed()
 	PlayerInput = NULL;
 	CheatManager = NULL;
 
-	if (VirtualJoystick.IsValid())
-	{
-		Cast<ULocalPlayer>(Player)->ViewportClient->RemoveViewportWidgetContent(VirtualJoystick.ToSharedRef());
-	}
-
 	Super::Destroyed();
 }
 
@@ -2667,7 +2679,7 @@ void APlayerController::DisplayDebug(class UCanvas* Canvas, const FDebugDisplayI
 	if ( DebugDisplay.IsDisplayOn("ForceFeedback"))
 	{
 		Canvas->SetDrawColor(255, 255, 255);
-		Canvas->DrawText(RenderFont, FString::Printf(TEXT("Force Feedback - LL: %.2f LS: %.2f RL: %.2f RS: %.2f"), ForceFeedbackValues.LeftLarge, ForceFeedbackValues.LeftSmall, ForceFeedbackValues.RightLarge, ForceFeedbackValues.RightSmall), 4.0f, YPos);
+		Canvas->DrawText(RenderFont, FString::Printf(TEXT("Force Feedback - Enabled: %s LL: %.2f LS: %.2f RL: %.2f RS: %.2f"), (bForceFeedbackEnabled ? TEXT("true") : TEXT("false")), ForceFeedbackValues.LeftLarge, ForceFeedbackValues.LeftSmall, ForceFeedbackValues.RightLarge, ForceFeedbackValues.RightSmall), 4.0f, YPos);
 		YPos += YL;
 	}
 }
@@ -3219,6 +3231,109 @@ void APlayerController::ClientStopForceFeedback_Implementation( UForceFeedbackEf
 	}
 }
 
+/** Action that interpolates a component over time to a desired position */
+class FDynamicForceFeedbackAction : public FPendingLatentAction
+{
+public:
+	/** Time over which interpolation should happen */
+	float TotalTime;
+	/** Time so far elapsed for the interpolation */
+	float TimeElapsed;
+	/** If we are currently running. If false, update will complete */
+	uint32 bRunning:1;
+	
+	TWeakObjectPtr<APlayerController> PlayerController;
+
+	FDynamicForceFeedbackDetails ForceFeedbackDetails;
+
+	FLatentActionInfo LatentInfo;
+	/** Function to execute on completion */
+	FName ExecutionFunction;
+	/** Link to fire on completion */
+	int32 OutputLink;
+	/** Latent action ID */
+	int32 LatentUUID;
+	/** Object to call callback on upon completion */
+	FWeakObjectPtr CallbackTarget;
+
+	FDynamicForceFeedbackAction(APlayerController* InPlayerController, const float InDuration, const FLatentActionInfo& LatentInfo)
+		: TotalTime(InDuration)
+		, TimeElapsed(0.f)
+		, bRunning(true)
+		, PlayerController(InPlayerController)
+		, ExecutionFunction(LatentInfo.ExecutionFunction)
+		, OutputLink(LatentInfo.Linkage)
+		, LatentUUID(LatentInfo.UUID)
+		, CallbackTarget(LatentInfo.CallbackTarget)
+	{
+	}
+
+	virtual void UpdateOperation(FLatentResponse& Response)
+	{
+		// Update elapsed time
+		TimeElapsed += Response.ElapsedTime();
+
+		const bool bComplete = (!bRunning || (TotalTime >= 0.f && TimeElapsed >= TotalTime) || !PlayerController.IsValid());
+
+		APlayerController* PC = PlayerController.Get();
+		if (PC)
+		{
+			if (bComplete)
+			{
+				PC->DynamicForceFeedbacks.Remove(LatentUUID);
+			}
+			else
+			{
+				PC->DynamicForceFeedbacks.Add(LatentUUID, ForceFeedbackDetails);
+			}
+		}
+
+
+		Response.FinishAndTriggerIf(bComplete, ExecutionFunction, OutputLink, CallbackTarget);
+	}
+};
+
+void APlayerController::PlayDynamicForceFeedback(float Intensity, float Duration, bool bAffectsLeftLarge, bool bAffectsLeftSmall, bool bAffectsRightLarge, bool bAffectsRightSmall, TEnumAsByte<EDynamicForceFeedbackAction::Type> Action, FLatentActionInfo LatentInfo)
+{
+	FLatentActionManager& LatentActionManager = GetWorld()->GetLatentActionManager();
+	FDynamicForceFeedbackAction* LatentAction = LatentActionManager.FindExistingAction<FDynamicForceFeedbackAction>(LatentInfo.CallbackTarget, LatentInfo.UUID);
+
+	if (LatentAction)
+	{
+		if (Action == EDynamicForceFeedbackAction::Stop)
+		{
+			LatentAction->bRunning = false;
+		}
+		else
+		{
+			if (Action == EDynamicForceFeedbackAction::Start)
+			{
+				LatentAction->TotalTime = Duration;
+				LatentAction->TimeElapsed = 0.f;
+				LatentAction->bRunning = true;
+			}
+
+			LatentAction->ForceFeedbackDetails.Intensity = Intensity;
+			LatentAction->ForceFeedbackDetails.bAffectsLeftLarge = bAffectsLeftLarge;
+			LatentAction->ForceFeedbackDetails.bAffectsLeftSmall = bAffectsLeftSmall;
+			LatentAction->ForceFeedbackDetails.bAffectsRightLarge = bAffectsRightLarge;
+			LatentAction->ForceFeedbackDetails.bAffectsRightSmall = bAffectsRightSmall;
+		}
+	}
+	else if (Action == EDynamicForceFeedbackAction::Start)
+	{
+		LatentAction = new FDynamicForceFeedbackAction(this, Duration, LatentInfo);
+
+		LatentAction->ForceFeedbackDetails.Intensity = Intensity;
+		LatentAction->ForceFeedbackDetails.bAffectsLeftLarge = bAffectsLeftLarge;
+		LatentAction->ForceFeedbackDetails.bAffectsLeftSmall = bAffectsLeftSmall;
+		LatentAction->ForceFeedbackDetails.bAffectsRightLarge = bAffectsRightLarge;
+		LatentAction->ForceFeedbackDetails.bAffectsRightSmall = bAffectsRightSmall;
+
+		LatentActionManager.AddNewAction(LatentInfo.CallbackTarget, LatentInfo.UUID, LatentAction);
+	}
+}
+
 void APlayerController::ProcessForceFeedback(const float DeltaTime, const bool bGamePaused)
 {
 	if (Player == NULL)
@@ -3237,13 +3352,16 @@ void APlayerController::ProcessForceFeedback(const float DeltaTime, const bool b
 				ActiveForceFeedbackEffects.RemoveAtSwap(Index);
 			}
 		}
+		for (const auto& DynamicEntry : DynamicForceFeedbacks)
+		{
+			DynamicEntry.Value.Update(ForceFeedbackValues);
+		}
 	}
 
-	// Get the IForceFeedbackSystem pointer from the global application, returning if NULL
 	IForceFeedbackSystem* ForceFeedbackSystem = FSlateApplication::Get().GetForceFeedbackSystem();
 	if (ForceFeedbackSystem)
 	{
-		ForceFeedbackSystem->SetChannelValues(CastChecked<ULocalPlayer>(Player)->ControllerId, ForceFeedbackValues);
+		ForceFeedbackSystem->SetChannelValues(CastChecked<ULocalPlayer>(Player)->ControllerId, (bForceFeedbackEnabled ? ForceFeedbackValues : FForceFeedbackValues()));
 	}
 }
 
