Fix PS4 compile.

diff --git a/Engine/Source/Runtime/Engine/Private/InstancedStaticMesh.cpp b/Engine/Source/Runtime/Engine/Private/InstancedStaticMesh.cpp
--- a/Engine/Source/Runtime/Engine/Private/InstancedStaticMesh.cpp
+++ b/Engine/Source/Runtime/Engine/Private/InstancedStaticMesh.cpp
@@ -395,9 +395,8 @@ public:
 
 	/** Make sure we account for changes in the signature of GetStaticBatchElementVisibility() */
 	static CONSTEXPR uint32 NumBitsForVisibilityMask()
-	{
-		decltype(((FInstancedStaticMeshVertexFactory*)nullptr)->GetStaticBatchElementVisibility(FSceneView(FSceneViewInitOptions()), nullptr)) Mask = 0;
-		return 8 * sizeof(Mask);
+	{		
+		return 8 * sizeof(decltype(((FInstancedStaticMeshVertexFactory*)nullptr)->GetStaticBatchElementVisibility(FSceneView(FSceneViewInitOptions()), nullptr)));
 	}
 
 	/**
