Whoops! Compile fix for unity.

diff --git a/Engine/Source/Runtime/Engine/Private/GameplayStatics.cpp b/Engine/Source/Runtime/Engine/Private/GameplayStatics.cpp
--- a/Engine/Source/Runtime/Engine/Private/GameplayStatics.cpp
+++ b/Engine/Source/Runtime/Engine/Private/GameplayStatics.cpp
@@ -11,6 +11,7 @@
 #include "GameFramework/Character.h"
 #include "Sound/SoundBase.h"
 #include "Sound/SoundCue.h"
+#include "Engine/GameInstance.h"
 
 //////////////////////////////////////////////////////////////////////////
 // UGameplayStatics
