fixed crash introduced in CL# 2228130

diff --git a/Engine/Source/Runtime/Engine/Private/Level.cpp b/Engine/Source/Runtime/Engine/Private/Level.cpp
--- a/Engine/Source/Runtime/Engine/Private/Level.cpp
+++ b/Engine/Source/Runtime/Engine/Private/Level.cpp
@@ -753,7 +753,10 @@ void ULevel::IncrementalUpdateComponents(int32 NumComponentsToUpdate, bool bReru
 		
 		auto CalcAttachDepth = [](AActor& InActor) -> int32 {
 			int32 Depth = 0;
-			for (const USceneComponent* Test = InActor.GetRootComponent()->AttachParent; Test != nullptr; Test = Test->AttachParent, Depth++);
+			if (InActor.GetRootComponent())
+			{
+				for (const USceneComponent* Test = InActor.GetRootComponent()->AttachParent; Test != nullptr; Test = Test->AttachParent, Depth++);
+			}
 			return Depth;
 		};
 
