Landscape : Fixed TTP #332780 : Fixed problem when you add component near component with XYOffset

diff --git a/Engine/Source/Editor/LandscapeEditor/Private/LandscapeEdModeComponentTools.cpp b/Engine/Source/Editor/LandscapeEditor/Private/LandscapeEdModeComponentTools.cpp
--- a/Engine/Source/Editor/LandscapeEditor/Private/LandscapeEdModeComponentTools.cpp
+++ b/Engine/Source/Editor/LandscapeEditor/Private/LandscapeEdModeComponentTools.cpp
@@ -743,9 +743,10 @@ class FLandscapeToolStrokeAddComponent
 {
 public:
 	FLandscapeToolStrokeAddComponent(FEdModeLandscape* InEdMode, const FLandscapeToolTarget& InTarget)
-		: EdMode(InEdMode)
-		, LandscapeInfo(InTarget.LandscapeInfo.Get())
-		, HeightCache(InTarget)
+		:	EdMode(InEdMode)
+		,	LandscapeInfo(InTarget.LandscapeInfo.Get())
+		,	HeightCache(InTarget)
+		,	XYOffsetCache(InTarget)
 	{}
 
 	virtual void Apply(FEditorViewportClient* ViewportClient, FLandscapeBrush* Brush, const ULandscapeEditorObject* UISettings, const TArray<FLandscapeToolMousePosition>& MousePositions)
@@ -769,7 +770,10 @@ public:
 
 			HeightCache.CacheData(X1, Y1, X2, Y2);
 			TArray<uint16> Data;
-			HeightCache.GetCachedData(X1, Y1, X2, Y2, Data);
+			HeightCache.GetCachedData(X1,Y1,X2,Y2,Data);
+			XYOffsetCache.CacheData(X1, Y1, X2, Y2);
+			TArray<FVector> XYOffsetData;
+			bool bHasXYOffset = XYOffsetCache.GetCachedData(X1, Y1, X2, Y2, XYOffsetData);
 
 			// Find component range for this block of data, non shared vertices
 			int32 ComponentIndexX1 = (X1 + 1 >= 0) ? (X1 + 1) / Landscape->ComponentSizeQuads : (X1 + 2) / Landscape->ComponentSizeQuads - 1;
@@ -820,8 +824,16 @@ public:
 
 			Landscape->RegisterAllComponents();
 
-			HeightCache.SetCachedData(X1, Y1, X2, Y2, Data);
-			HeightCache.Flush();
+			if (bHasXYOffset)
+			{
+				XYOffsetCache.SetCachedData(X1, Y1, X2, Y2, XYOffsetData);
+				XYOffsetCache.Flush();
+			}
+			else
+			{
+				HeightCache.SetCachedData(X1, Y1, X2, Y2, Data);
+				HeightCache.Flush();
+			}
 
 			for (int32 Idx = 0; Idx < NewComponents.Num(); Idx++)
 			{
@@ -854,7 +866,7 @@ public:
 				NewComponents[Idx]->UpdateBounds();
 				NewComponents[Idx]->MarkRenderStateDirty();
 				ULandscapeHeightfieldCollisionComponent* CollisionComp = NewComponents[Idx]->CollisionComponent.Get();
-				if (CollisionComp)
+				if (CollisionComp && !bHasXYOffset)
 				{
 					CollisionComp->MarkRenderStateDirty();
 					CollisionComp->RecreateCollision();
@@ -871,6 +883,7 @@ protected:
 	FEdModeLandscape* EdMode;
 	ULandscapeInfo* LandscapeInfo;
 	FLandscapeHeightCache HeightCache;
+	FLandscapeXYOffsetCache<true> XYOffsetCache;
 };
 
 // 
diff --git a/Engine/Source/Editor/LandscapeEditor/Private/LandscapeEdModeTools.h b/Engine/Source/Editor/LandscapeEditor/Private/LandscapeEdModeTools.h
--- a/Engine/Source/Editor/LandscapeEditor/Private/LandscapeEdModeTools.h
+++ b/Engine/Source/Editor/LandscapeEditor/Private/LandscapeEdModeTools.h
@@ -425,11 +425,12 @@ struct TLandscapeEditCache
 		CachedData.Add(ALandscape::MakeKey(LandscapeX, LandscapeY), Value);
 	}
 
-	void GetCachedData(int32 X1, int32 Y1, int32 X2, int32 Y2, TArray<AccessorType>& OutData)
+	bool GetCachedData(int32 X1, int32 Y1, int32 X2, int32 Y2, TArray<AccessorType>& OutData)
 	{
 		int32 NumSamples = (1 + X2 - X1)*(1 + Y2 - Y1);
 		OutData.Empty(NumSamples);
 		OutData.AddUninitialized(NumSamples);
+		bool bHasNonZero = false;
 
 		for (int32 Y = Y1; Y <= Y2; Y++)
 		{
@@ -438,10 +439,16 @@ struct TLandscapeEditCache
 				AccessorType* Ptr = GetValueRef(X, Y);
 				if (Ptr)
 				{
-					OutData[(X - X1) + (Y - Y1)*(1 + X2 - X1)] = *Ptr;
+					OutData[(X-X1) + (Y-Y1)*(1+X2-X1)] = *Ptr;
+					if (*Ptr != (AccessorType)0)
+					{
+						bHasNonZero = true;
+					}
 				}
 			}
 		}
+
+		return bHasNonZero;
 	}
 
 	void SetCachedData(int32 X1, int32 Y1, int32 X2, int32 Y2, TArray<AccessorType>& Data, ELandscapeLayerPaintingRestriction::Type PaintingRestriction = ELandscapeLayerPaintingRestriction::None)
@@ -692,7 +699,7 @@ struct FXYOffsetmapAccessor
 
 	void GetDataFast(int32 X1, int32 Y1, int32 X2, int32 Y2, TMap<FIntPoint, FVector>& Data)
 	{
-		LandscapeEdit->GetXYOffsetData(X1, Y1, X2, Y2, Data);
+		LandscapeEdit->GetXYOffsetDataFast(X1, Y1, X2, Y2, Data);
 
 		TMap<FIntPoint, uint16> NewHeights;
 		LandscapeEdit->GetHeightDataFast(X1, Y1, X2, Y2, NewHeights);
@@ -803,11 +810,12 @@ private:
 	TSet<ULandscapeComponent*> ChangedComponents;
 };
 
-struct FLandscapeXYOffsetCache : public TLandscapeEditCache < FXYOffsetmapAccessor<true>, FVector >
+template<bool bInUseInterp>
+struct FLandscapeXYOffsetCache : public TLandscapeEditCache < FXYOffsetmapAccessor<bInUseInterp>, FVector >
 {
 	typedef FVector DataType;
 
-	FXYOffsetmapAccessor<true> XYOffsetmapAccessor;
+	FXYOffsetmapAccessor<bInUseInterp> XYOffsetmapAccessor;
 
 	FLandscapeXYOffsetCache(const FLandscapeToolTarget& InTarget)
 		: XYOffsetmapAccessor(InTarget.LandscapeInfo.Get())
diff --git a/Engine/Source/Editor/LandscapeEditor/Private/LandscapeEdModeXYOffset.cpp b/Engine/Source/Editor/LandscapeEditor/Private/LandscapeEdModeXYOffset.cpp
--- a/Engine/Source/Editor/LandscapeEditor/Private/LandscapeEdModeXYOffset.cpp
+++ b/Engine/Source/Editor/LandscapeEditor/Private/LandscapeEdModeXYOffset.cpp
@@ -543,7 +543,7 @@ public:
 
 protected:
 	ULandscapeInfo* LandscapeInfo;
-	FLandscapeXYOffsetCache Cache;
+	FLandscapeXYOffsetCache<false> Cache;
 };
 
 class FLandscapeToolRetopologize : public FLandscapeToolBase<FLandscapeToolStrokeRetopologize>
diff --git a/Engine/Source/Runtime/Engine/Private/Landscape/LandscapeEditInterface.cpp b/Engine/Source/Runtime/Engine/Private/Landscape/LandscapeEditInterface.cpp
--- a/Engine/Source/Runtime/Engine/Private/Landscape/LandscapeEditInterface.cpp
+++ b/Engine/Source/Runtime/Engine/Private/Landscape/LandscapeEditInterface.cpp
@@ -93,12 +93,12 @@ void ALandscape::CalcComponentIndices(const int32 X1, const int32 Y1, const int3
 namespace
 {
 	// Ugly helper function, all arrays should be only size 4
-	template<typename T>
-	FORCEINLINE void CalcInterpValue( const int32* Dist, const bool* Exist, const T* Value, float& ValueX, float& ValueY)
+	template<typename T, typename F>
+	FORCEINLINE void CalcInterpValue(const int32* Dist, const bool* Exist, const T* Value, F& ValueX, F& ValueY)
 	{
 		if (Exist[0] && Exist[1])
 		{
-			ValueX = (float)(Dist[1] * Value[0] + Dist[0] * Value[1]) / (Dist[0] + Dist[1]);
+			ValueX = (F)(Dist[1] * Value[0] + Dist[0] * Value[1]) / (Dist[0] + Dist[1]);
 		}
 		else
 		{
@@ -114,7 +114,7 @@ namespace
 
 		if (Exist[2] && Exist[3])
 		{
-			ValueY = (float)(Dist[3] * Value[2] + Dist[2] * Value[3]) / (Dist[2] + Dist[3]);
+			ValueY = (F)(Dist[3] * Value[2] + Dist[2] * Value[3]) / (Dist[2] + Dist[3]);
 		}
 		else
 		{
@@ -222,7 +222,7 @@ void FLandscapeEditDataInterface::SetHeightData(int32 X1, int32 Y1, int32 X2, in
 		// Need to consider XYOffset for XY displacemented map
 		FVector2D* XYOffsets = new FVector2D[NumVertsX*NumVertsY];
 		FMemory::Memzero(XYOffsets, NumVertsX*NumVertsY*sizeof(FVector2D));
-		GetXYOffsetData(X1, Y1, X2, Y2, XYOffsets, 0);
+		GetXYOffsetDataFast(X1, Y1, X2, Y2, XYOffsets, 0);
 
 		for( int32 Y=0;Y<NumVertsY-1;Y++ )
 		{
@@ -443,7 +443,7 @@ void FLandscapeEditDataInterface::RecalculateNormals()
 		FMemory::Memzero(XYOffsets, FMath::Square(Stride)*sizeof(FVector2D));
 
 		// Get XY offset
-		GetXYOffsetData(X1,Y1,X2,Y2,XYOffsets,0);
+		GetXYOffsetDataFast(X1,Y1,X2,Y2,XYOffsets,0);
 		// Get the vertex positions for entire quad
 		GetHeightData(X1,Y1,X2,Y2,HeightData,0);
 
@@ -608,7 +608,7 @@ void FLandscapeEditDataInterface::GetHeightDataTemplFast(const int32 X1, const i
 	}
 }
 
-template<typename TData, typename TStoreData>
+template<typename TData, typename TStoreData, typename FType>
 void FLandscapeEditDataInterface::CalcMissingValues(const int32& X1, const int32& X2, const int32& Y1, const int32& Y2, 
 								   const int32& ComponentIndexX1, const int32& ComponentIndexX2, const int32& ComponentIndexY1, const int32& ComponentIndexY2, 
 								   const int32& ComponentSizeX, const int32& ComponentSizeY, TData* CornerValues, 
@@ -714,22 +714,22 @@ void FLandscapeEditDataInterface::CalcMissingValues(const int32& X1, const int32
 				if (((ComponentIndexX == ComponentIndexX1) || (ComponentIndexY == ComponentIndexY1)) ? false : ComponentDataExist[ComponentSizeX*(ComponentIndexY-1-ComponentIndexY1) + ComponentIndexX-1-ComponentIndexX1])
 				{
 					CornerSet |= 1;
-					CornerValues[0] = StoreData.Load( ComponentIndexX*ComponentSizeQuads, ComponentIndexY*ComponentSizeQuads);
+					CornerValues[0] = TData(StoreData.Load(ComponentIndexX*ComponentSizeQuads, ComponentIndexY*ComponentSizeQuads));
 				}
 				if (((ComponentIndexX == ComponentIndexX2) || (ComponentIndexY == ComponentIndexY1)) ? false : ComponentDataExist[ComponentSizeX*(ComponentIndexY-1-ComponentIndexY1) + ComponentIndexX+1-ComponentIndexX1])
 				{
 					CornerSet |= 1 << 1;
-					CornerValues[1] = StoreData.Load( (ComponentIndexX+1)*ComponentSizeQuads, ComponentIndexY*ComponentSizeQuads);
+					CornerValues[1] = TData(StoreData.Load((ComponentIndexX + 1)*ComponentSizeQuads, ComponentIndexY*ComponentSizeQuads));
 				}
 				if (((ComponentIndexX == ComponentIndexX1) || (ComponentIndexY == ComponentIndexY2)) ? false : ComponentDataExist[ComponentSizeX*(ComponentIndexY+1-ComponentIndexY1) + ComponentIndexX-1-ComponentIndexX1])
 				{
 					CornerSet |= 1 << 2;
-					CornerValues[2] = StoreData.Load( ComponentIndexX*ComponentSizeQuads, (ComponentIndexY+1)*ComponentSizeQuads);
+					CornerValues[2] = TData(StoreData.Load(ComponentIndexX*ComponentSizeQuads, (ComponentIndexY + 1)*ComponentSizeQuads));
 				}
 				if (((ComponentIndexX == ComponentIndexX2) || (ComponentIndexY == ComponentIndexY2)) ? false : ComponentDataExist[ComponentSizeX*(ComponentIndexY+1-ComponentIndexY1) + ComponentIndexX+1-ComponentIndexX1])
 				{
 					CornerSet |= 1 << 3;
-					CornerValues[3] = StoreData.Load((ComponentIndexX+1)*ComponentSizeQuads, (ComponentIndexY+1)*ComponentSizeQuads);
+					CornerValues[3] = TData(StoreData.Load((ComponentIndexX + 1)*ComponentSizeQuads, (ComponentIndexY + 1)*ComponentSizeQuads));
 				}
 
 				FillCornerValues(CornerSet, CornerValues);
@@ -765,21 +765,24 @@ void FLandscapeEditDataInterface::CalcMissingValues(const int32& X1, const int32
 								int32 LandscapeY = SubIndexY*SubsectionSizeQuads + ComponentIndexY*ComponentSizeQuads + SubY;
 
 								// Find the texture data corresponding to this vertex
-								TData Value[4] = {0, 0, 0, 0};
+								TData Value[4];
+								FMemory::Memzero(Value, sizeof(TData)* 4);
 								int32 Dist[4] = {INT_MAX, INT_MAX, INT_MAX, INT_MAX};
-								float ValueX = 0.f, ValueY = 0.f;
+								FType ValueX, ValueY;
+								FMemory::Memzero(&ValueX, sizeof(FType));
+								FMemory::Memzero(&ValueY, sizeof(FType));
 								bool Exist[4] = {false, false, false, false};
 
 								if (ExistLeft)
 								{
-									Value[0] = StoreData.Load( ComponentIndexX*ComponentSizeQuads, LandscapeY);
+									Value[0] = TData(StoreData.Load(ComponentIndexX*ComponentSizeQuads, LandscapeY));
 									Dist[0] = LandscapeX - (ComponentIndexX*ComponentSizeQuads);
 									Exist[0] = true;
 								}
 								else if ( BorderX1 != INT_MAX )
 								{
 									int32 BorderIdxX = (BorderX1+1)*ComponentSizeQuads;
-									Value[0] = StoreData.Load(BorderIdxX, LandscapeY);
+									Value[0] = TData(StoreData.Load(BorderIdxX, LandscapeY));
 									Dist[0] = LandscapeX - (BorderIdxX-1);
 									Exist[0] = true;
 								}
@@ -789,7 +792,7 @@ void FLandscapeEditDataInterface::CalcMissingValues(const int32& X1, const int32
 									{
 										int32 Dist1 = LandscapeY - (ComponentIndexY*ComponentSizeQuads);
 										int32 Dist2 = ((ComponentIndexY+1)*ComponentSizeQuads) - LandscapeY;
-										Value[0] = (float)(Dist2 * CornerValues[0] + Dist1 * CornerValues[2]) / (Dist1 + Dist2);
+										Value[0] = (FType)(Dist2 * CornerValues[0] + Dist1 * CornerValues[2]) / (Dist1 + Dist2);
 										Dist[0] = LandscapeX - (ComponentIndexX*ComponentSizeQuads);
 										Exist[0] = true;
 									}
@@ -798,7 +801,7 @@ void FLandscapeEditDataInterface::CalcMissingValues(const int32& X1, const int32
 								if ( BorderX2 != INT_MIN )
 								{
 									int32 BorderIdxX = BorderX2*ComponentSizeQuads;
-									Value[1] = StoreData.Load(BorderIdxX, LandscapeY);
+									Value[1] = TData(StoreData.Load(BorderIdxX, LandscapeY));
 									Dist[1] = (BorderIdxX+1) - LandscapeX;
 									Exist[1] = true;
 								}
@@ -808,7 +811,7 @@ void FLandscapeEditDataInterface::CalcMissingValues(const int32& X1, const int32
 									{
 										int32 Dist1 = LandscapeY - (ComponentIndexY*ComponentSizeQuads);
 										int32 Dist2 = ((ComponentIndexY+1)*ComponentSizeQuads) - LandscapeY;
-										Value[1] = (float)(Dist2 * CornerValues[1] + Dist1 * CornerValues[3]) / (Dist1 + Dist2);
+										Value[1] = (FType)(Dist2 * CornerValues[1] + Dist1 * CornerValues[3]) / (Dist1 + Dist2);
 										Dist[1] = (ComponentIndexX+1)*ComponentSizeQuads - LandscapeX;
 										Exist[1] = true;
 									}
@@ -816,14 +819,14 @@ void FLandscapeEditDataInterface::CalcMissingValues(const int32& X1, const int32
 
 								if (ExistUp)
 								{
-									Value[2] = StoreData.Load( LandscapeX, ComponentIndexY*ComponentSizeQuads);
+									Value[2] = TData(StoreData.Load(LandscapeX, ComponentIndexY*ComponentSizeQuads));
 									Dist[2] = LandscapeY - (ComponentIndexY*ComponentSizeQuads);
 									Exist[2] = true;
 								}
 								else if ( BorderY1[ComponentIndexXX] != INT_MAX )
 								{
 									int32 BorderIdxY = (BorderY1[ComponentIndexXX]+1)*ComponentSizeQuads;
-									Value[2] = StoreData.Load(LandscapeX, BorderIdxY);
+									Value[2] = TData(StoreData.Load(LandscapeX, BorderIdxY));
 									Dist[2] = LandscapeY - BorderIdxY;
 									Exist[2] = true;
 								}
@@ -833,7 +836,7 @@ void FLandscapeEditDataInterface::CalcMissingValues(const int32& X1, const int32
 									{
 										int32 Dist1 = LandscapeX - (ComponentIndexX*ComponentSizeQuads);
 										int32 Dist2 = (ComponentIndexX+1)*ComponentSizeQuads - LandscapeX;
-										Value[2] = (float)(Dist2 * CornerValues[0] + Dist1 * CornerValues[1]) / (Dist1 + Dist2);
+										Value[2] = (FType)(Dist2 * CornerValues[0] + Dist1 * CornerValues[1]) / (Dist1 + Dist2);
 										Dist[2] = LandscapeY - (ComponentIndexY*ComponentSizeQuads);
 										Exist[2] = true;
 									}
@@ -842,7 +845,7 @@ void FLandscapeEditDataInterface::CalcMissingValues(const int32& X1, const int32
 								if ( BorderY2[ComponentIndexXX] != INT_MIN )
 								{
 									int32 BorderIdxY = BorderY2[ComponentIndexXX]*ComponentSizeQuads;
-									Value[3] = StoreData.Load(LandscapeX, BorderIdxY);
+									Value[3] = TData(StoreData.Load(LandscapeX, BorderIdxY));
 									Dist[3] = BorderIdxY - LandscapeY;
 									Exist[3] = true;
 								}
@@ -852,15 +855,16 @@ void FLandscapeEditDataInterface::CalcMissingValues(const int32& X1, const int32
 									{
 										int32 Dist1 = LandscapeX - (ComponentIndexX*ComponentSizeQuads);
 										int32 Dist2 = (ComponentIndexX+1)*ComponentSizeQuads - LandscapeX;
-										Value[3] = (float)(Dist2 * CornerValues[2] + Dist1 * CornerValues[3]) / (Dist1 + Dist2);
+										Value[3] = (FType)(Dist2 * CornerValues[2] + Dist1 * CornerValues[3]) / (Dist1 + Dist2);
 										Dist[3] = (ComponentIndexY+1)*ComponentSizeQuads - LandscapeY;
 										Exist[3] = true;
 									}
 								}
 
-								CalcInterpValue<TData>(Dist, Exist, Value, ValueX, ValueY);
+								CalcInterpValue<TData, FType>(Dist, Exist, Value, ValueX, ValueY);
 
-								uint16 FinalValue = 0; // Default Value
+								TData FinalValue; // Default Value
+								FMemory::Memzero(&FinalValue, sizeof(TData));
 								if ( (Exist[0] || Exist[1]) && (Exist[2] || Exist[3]) )
 								{
 									FinalValue = CalcValueFromValueXY<TData>(Dist, ValueX, ValueY, CornerSet, CornerValues);
@@ -1398,7 +1402,7 @@ void FLandscapeEditDataInterface::GetHeightDataTempl(int32& ValidX1, int32& Vali
 
 	if (bHasMissingValue)
 	{
-		CalcMissingValues<uint16, TStoreData>( X1, X2, Y1, Y2,
+		CalcMissingValues<uint16, TStoreData, float>( X1, X2, Y1, Y2,
 			ComponentIndexX1, ComponentIndexX2, ComponentIndexY1, ComponentIndexY2, 
 			ComponentSizeX, ComponentSizeY, CornerValues,
 			NoBorderY1, NoBorderY2, ComponentDataExist, StoreData );
@@ -1971,7 +1975,7 @@ namespace
 		inline void Store(int32 LandscapeX, int32 LandscapeY, uint8 Weight) {}
 		inline void Store(int32 LandscapeX, int32 LandscapeY, uint8 Weight, int32 LayerIdx) {}
 		inline void Store(int32 LandscapeX, int32 LandscapeY, FVector2D Offset) {}
-		inline uint8 Load(int32 LandscapeX, int32 LandscapeY) { return 0; }
+		inline TDataType Load(int32 LandscapeX, int32 LandscapeY) { return 0; }
 		inline void PreInit(int32 InArraySize) { ArraySize = InArraySize; }
 	};
 
@@ -1985,6 +1989,16 @@ namespace
 		return Data[ (LandscapeY-Y1) * Stride + (LandscapeX-X1) ];
 	}
 
+	template<> FVector2D TArrayStoreData<FVector2D>::Load(int32 LandscapeX, int32 LandscapeY)
+	{
+		return Data[(LandscapeY - Y1) * Stride + (LandscapeX - X1)];
+	}
+
+	template<> FVector TArrayStoreData<FVector>::Load(int32 LandscapeX, int32 LandscapeY)
+	{
+		return Data[(LandscapeY - Y1) * Stride + (LandscapeX - X1)];
+	}
+
 	template<> void TArrayStoreData<FVector2D>::Store(int32 LandscapeX, int32 LandscapeY, FVector2D Offset)
 	{
 		Data[ (LandscapeY-Y1) * Stride + (LandscapeX-X1) ] = Offset;
@@ -2021,7 +2035,7 @@ namespace
 		inline void Store(int32 LandscapeX, int32 LandscapeY, uint8 Weight) {}
 		inline void Store(int32 LandscapeX, int32 LandscapeY, uint8 Weight, int32 LayerIdx) {}
 		inline void Store(int32 LandscapeX, int32 LandscapeY, FVector2D Offset) {}
-		inline uint8 Load(int32 LandscapeX, int32 LandscapeY) { return 0; }
+		inline TDataType Load(int32 LandscapeX, int32 LandscapeY) { return 0; }
 		inline void PreInit(int32 InArraySize) { ArraySize = InArraySize; }
 	};
 
@@ -2035,6 +2049,16 @@ namespace
 		return SparseData.FindRef(ALandscape::MakeKey(LandscapeX,LandscapeY));
 	}
 
+	template<> FVector2D TSparseStoreData<FVector2D>::Load(int32 LandscapeX, int32 LandscapeY)
+	{
+		return SparseData.FindRef(ALandscape::MakeKey(LandscapeX, LandscapeY));
+	}
+
+	template<> FVector TSparseStoreData<FVector>::Load(int32 LandscapeX, int32 LandscapeY)
+	{
+		return SparseData.FindRef(ALandscape::MakeKey(LandscapeX, LandscapeY));
+	}
+
 	template<> void TSparseStoreData<TArray<uint8>>::Store(int32 LandscapeX, int32 LandscapeY, uint8 Weight, int32 LayerIdx)
 	{
 		TArray<uint8>* Value = SparseData.Find(ALandscape::MakeKey(LandscapeX,LandscapeY));
@@ -3432,7 +3456,7 @@ void FLandscapeEditDataInterface::GetWeightDataTempl(ULandscapeLayerInfoObject*
 
 	if (bHasMissingValue)
 	{
-		CalcMissingValues<uint8, TStoreData>( X1, X2, Y1, Y2,
+		CalcMissingValues<uint8, TStoreData, float>( X1, X2, Y1, Y2,
 			ComponentIndexX1, ComponentIndexX2, ComponentIndexY1, ComponentIndexY2, 
 			ComponentSizeX, ComponentSizeY, CornerValues,
 			NoBorderY1, NoBorderY2, ComponentDataExist, StoreData );
@@ -4024,7 +4048,7 @@ void FLandscapeEditDataInterface::SetXYOffsetData(int32 X1, int32 Y1, int32 X2,
 FVector2D FLandscapeEditDataInterface::GetXYOffsetmapData(const ULandscapeComponent* Component, int32 TexU, int32 TexV, FColor* TextureData/* = NULL*/)
 {
 	check(Component);
-	if (!TextureData)
+	if (!TextureData && Component->XYOffsetmapTexture)
 	{
 		FLandscapeTextureDataInfo* TexDataInfo = GetTextureDataInfo(Component->XYOffsetmapTexture);
 		TextureData = (FColor*)TexDataInfo->GetMipData(0);	
@@ -4032,21 +4056,555 @@ FVector2D FLandscapeEditDataInterface::GetXYOffsetmapData(const ULandscapeCompon
 
 	if (TextureData)
 	{
-		int32 SizeU = Component->XYOffsetmapTexture->Source.GetSizeX();
-		int32 SizeV = Component->XYOffsetmapTexture->Source.GetSizeY();
-		int32 WeightmapOffsetX = Component->WeightmapScaleBias.Z * (float)SizeU;
-		int32 WeightmapOffsetY = Component->WeightmapScaleBias.W * (float)SizeV;
+		int32 SizeU = Component->NumSubsections * (Component->SubsectionSizeQuads + 1);
+		int32 SizeV = Component->NumSubsections * (Component->SubsectionSizeQuads + 1);
 
-		int32 TexX = WeightmapOffsetX + TexU;
-		int32 TexY = WeightmapOffsetY + TexV;
+		int32 TexX = TexU;
+		int32 TexY = TexV;
 		FColor& TexData = TextureData[ TexX + TexY * SizeU ];
 		return FVector2D(((TexData.R * 256.0 + TexData.G) - 32768.0) * LANDSCAPE_XYOFFSET_SCALE, ((TexData.B * 256.0 + TexData.A) - 32768.0) * LANDSCAPE_XYOFFSET_SCALE );
 	}
-	return FVector2D(0.f, 0.f);
+	return FVector2D::ZeroVector;
 }
 
+// XYOffset Interpolation version
 template<typename TStoreData>
-void FLandscapeEditDataInterface::GetXYOffsetDataTempl(const int32 X1, const int32 Y1, const int32 X2, const int32 Y2, TStoreData& StoreData)
+void FLandscapeEditDataInterface::GetXYOffsetDataTempl(int32& ValidX1, int32& ValidY1, int32& ValidX2, int32& ValidY2, TStoreData& StoreData)
+{
+	// Copy variables
+	int32 X1 = ValidX1, X2 = ValidX2, Y1 = ValidY1, Y2 = ValidY2;
+	ValidX1 = INT_MAX; ValidX2 = INT_MIN; ValidY1 = INT_MAX; ValidY2 = INT_MIN;
+
+	int32 ComponentIndexX1, ComponentIndexY1, ComponentIndexX2, ComponentIndexY2;
+	ALandscape::CalcComponentIndicesOverlap(X1, Y1, X2, Y2, ComponentSizeQuads, ComponentIndexX1, ComponentIndexY1, ComponentIndexX2, ComponentIndexY2);
+
+	int32 ComponentSizeX = ComponentIndexX2 - ComponentIndexX1 + 1;
+	int32 ComponentSizeY = ComponentIndexY2 - ComponentIndexY1 + 1;
+
+	// Neighbor Components
+	ULandscapeComponent* BorderComponent[4] = { 0, 0, 0, 0 };
+	ULandscapeComponent* CornerComponent[4] = { 0, 0, 0, 0 };
+	bool NoBorderX1 = false, NoBorderX2 = false;
+	TArray<bool> NoBorderY1, NoBorderY2, ComponentDataExist;
+	TArray<ULandscapeComponent*> BorderComponentY1, BorderComponentY2;
+	ComponentDataExist.Empty(ComponentSizeX*ComponentSizeY);
+	ComponentDataExist.AddZeroed(ComponentSizeX*ComponentSizeY);
+	bool bHasMissingValue = false;
+
+	FLandscapeTextureDataInfo* NeighborTexDataInfo[4] = { 0, 0, 0, 0 };
+	FColor* NeighborXYOffsetmapTextureData[4] = { 0, 0, 0, 0 };
+	FVector2D CornerValues[4] = { FVector2D::ZeroVector, FVector2D::ZeroVector, FVector2D::ZeroVector, FVector2D::ZeroVector };
+
+	int32 EdgeCoord = (SubsectionSizeQuads + 1) * ComponentNumSubsections - 1; //ComponentSizeQuads;
+
+	TArray<FColor> EmptyXYOffset;
+	int32 XYOffsetSize = (LandscapeInfo->SubsectionSizeQuads + 1) * LandscapeInfo->ComponentNumSubsections;
+	XYOffsetSize = XYOffsetSize * XYOffsetSize;
+	EmptyXYOffset.Empty(XYOffsetSize);
+	for (int32 i = 0; i < XYOffsetSize; ++i)
+	{
+		EmptyXYOffset.Add(FColor(128, 0, 128, 0));
+	}
+
+	// initial loop....
+	for (int32 ComponentIndexY = ComponentIndexY1; ComponentIndexY <= ComponentIndexY2; ComponentIndexY++)
+	{
+		NoBorderX1 = false;
+		NoBorderX2 = false;
+		BorderComponent[0] = BorderComponent[1] = NULL;
+		for (int32 ComponentIndexX = ComponentIndexX1; ComponentIndexX <= ComponentIndexX2; ComponentIndexX++)
+		{
+			BorderComponent[2] = BorderComponent[3] = NULL;
+			int32 ComponentIndexXY = ComponentSizeX*(ComponentIndexY - ComponentIndexY1) + ComponentIndexX - ComponentIndexX1;
+			int32 ComponentIndexXX = ComponentIndexX - ComponentIndexX1;
+			int32 ComponentIndexYY = ComponentIndexY - ComponentIndexY1;
+			ComponentDataExist[ComponentIndexXY] = false;
+			ULandscapeComponent* Component = LandscapeInfo->XYtoComponentMap.FindRef(FIntPoint(ComponentIndexX, ComponentIndexY));
+
+			FLandscapeTextureDataInfo* TexDataInfo = NULL;
+			FColor* XYOffsetmapTextureData = NULL;
+			uint8 CornerSet = 0;
+			bool ExistLeft = ComponentIndexXX > 0 && ComponentDataExist[ComponentIndexXX - 1 + ComponentIndexYY * ComponentSizeX];
+			bool ExistUp = ComponentIndexYY > 0 && ComponentDataExist[ComponentIndexXX + (ComponentIndexYY - 1) * ComponentSizeX];
+
+			if (Component)
+			{
+				if (Component->XYOffsetmapTexture)
+				{
+					TexDataInfo = GetTextureDataInfo(Component->XYOffsetmapTexture);
+					XYOffsetmapTextureData = (FColor*)TexDataInfo->GetMipData(0);
+				}
+				else
+				{
+					XYOffsetmapTextureData = EmptyXYOffset.GetData();
+				}
+				ComponentDataExist[ComponentIndexXY] = true;
+				// Update valid region
+				ValidX1 = FMath::Min<int32>(Component->GetSectionBase().X, ValidX1);
+				ValidX2 = FMath::Max<int32>(Component->GetSectionBase().X + ComponentSizeQuads, ValidX2);
+				ValidY1 = FMath::Min<int32>(Component->GetSectionBase().Y, ValidY1);
+				ValidY2 = FMath::Max<int32>(Component->GetSectionBase().Y + ComponentSizeQuads, ValidY2);
+			}
+			else
+			{
+				if (!bHasMissingValue)
+				{
+					NoBorderY1.Empty(ComponentSizeX);
+					NoBorderY2.Empty(ComponentSizeX);
+					NoBorderY1.AddZeroed(ComponentSizeX);
+					NoBorderY2.AddZeroed(ComponentSizeX);
+					BorderComponentY1.Empty(ComponentSizeX);
+					BorderComponentY2.Empty(ComponentSizeX);
+					BorderComponentY1.AddZeroed(ComponentSizeX);
+					BorderComponentY2.AddZeroed(ComponentSizeX);
+					bHasMissingValue = true;
+				}
+
+				// Search for neighbor component for interpolation
+				bool bShouldSearchX = (BorderComponent[1] && BorderComponent[1]->GetSectionBase().X / ComponentSizeQuads <= ComponentIndexX);
+				bool bShouldSearchY = (BorderComponentY2[ComponentIndexXX] && BorderComponentY2[ComponentIndexXX]->GetSectionBase().Y / ComponentSizeQuads <= ComponentIndexY);
+				// Search for left-closest component
+				if (bShouldSearchX || (!NoBorderX1 && !BorderComponent[0]))
+				{
+					NoBorderX1 = true;
+					for (int32 X = ComponentIndexX - 1; X >= ComponentIndexX1; X--)
+					{
+						BorderComponent[0] = LandscapeInfo->XYtoComponentMap.FindRef(FIntPoint(X, ComponentIndexY));
+						if (BorderComponent[0])
+						{
+							NoBorderX1 = false;
+							if (BorderComponent[0]->XYOffsetmapTexture)
+							{
+								NeighborTexDataInfo[0] = GetTextureDataInfo(BorderComponent[0]->XYOffsetmapTexture);
+								NeighborXYOffsetmapTextureData[0] = (FColor*)NeighborTexDataInfo[0]->GetMipData(0);
+							}
+							else
+							{
+								NeighborXYOffsetmapTextureData[0] = EmptyXYOffset.GetData();
+							}
+							break;
+						}
+					}
+				}
+				// Search for right-closest component
+				if (bShouldSearchX || (!NoBorderX2 && !BorderComponent[1]))
+				{
+					NoBorderX2 = true;
+					for (int32 X = ComponentIndexX + 1; X <= ComponentIndexX2; X++)
+					{
+						BorderComponent[1] = LandscapeInfo->XYtoComponentMap.FindRef(FIntPoint(X, ComponentIndexY));
+						if (BorderComponent[1])
+						{
+							NoBorderX2 = false;
+							if (BorderComponent[1]->XYOffsetmapTexture)
+							{
+								NeighborTexDataInfo[1] = GetTextureDataInfo(BorderComponent[1]->XYOffsetmapTexture);
+								NeighborXYOffsetmapTextureData[1] = (FColor*)NeighborTexDataInfo[1]->GetMipData(0);
+							}
+							else
+							{
+								NeighborXYOffsetmapTextureData[1] = EmptyXYOffset.GetData();
+							}
+							break;
+						}
+					}
+				}
+				// Search for up-closest component
+				if (bShouldSearchY || (!NoBorderY1[ComponentIndexXX] && !BorderComponentY1[ComponentIndexXX]))
+				{
+					NoBorderY1[ComponentIndexXX] = true;
+					for (int32 Y = ComponentIndexY - 1; Y >= ComponentIndexY1; Y--)
+					{
+						BorderComponentY1[ComponentIndexXX] = BorderComponent[2] = LandscapeInfo->XYtoComponentMap.FindRef(FIntPoint(ComponentIndexX, Y));
+						if (BorderComponent[2])
+						{
+							NoBorderY1[ComponentIndexXX] = false;
+							if (BorderComponent[2]->XYOffsetmapTexture)
+							{
+								NeighborTexDataInfo[2] = GetTextureDataInfo(BorderComponent[2]->XYOffsetmapTexture);
+								NeighborXYOffsetmapTextureData[2] = (FColor*)NeighborTexDataInfo[2]->GetMipData(0);
+							}
+							else
+							{
+								NeighborXYOffsetmapTextureData[2] = EmptyXYOffset.GetData();
+							}
+							break;
+						}
+					}
+				}
+				else
+				{
+					BorderComponent[2] = BorderComponentY1[ComponentIndexXX];
+					if (BorderComponent[2])
+					{
+						if (BorderComponent[2]->XYOffsetmapTexture)
+						{
+							NeighborTexDataInfo[2] = GetTextureDataInfo(BorderComponent[2]->XYOffsetmapTexture);
+							NeighborXYOffsetmapTextureData[2] = (FColor*)NeighborTexDataInfo[2]->GetMipData(0);
+						}
+						else
+						{
+							NeighborXYOffsetmapTextureData[2] = EmptyXYOffset.GetData();
+						}
+
+					}
+				}
+				// Search for bottom-closest component
+				if (bShouldSearchY || (!NoBorderY2[ComponentIndexXX] && !BorderComponentY2[ComponentIndexXX]))
+				{
+					NoBorderY2[ComponentIndexXX] = true;
+					for (int32 Y = ComponentIndexY + 1; Y <= ComponentIndexY2; Y++)
+					{
+						BorderComponentY2[ComponentIndexXX] = BorderComponent[3] = LandscapeInfo->XYtoComponentMap.FindRef(FIntPoint(ComponentIndexX, Y));
+						if (BorderComponent[3])
+						{
+							NoBorderY2[ComponentIndexXX] = false;
+							if (BorderComponent[3]->XYOffsetmapTexture)
+							{
+								NeighborTexDataInfo[3] = GetTextureDataInfo(BorderComponent[3]->XYOffsetmapTexture);
+								NeighborXYOffsetmapTextureData[3] = (FColor*)NeighborTexDataInfo[3]->GetMipData(0);
+							}
+							else
+							{
+								NeighborXYOffsetmapTextureData[3] = EmptyXYOffset.GetData();
+							}
+							break;
+						}
+					}
+				}
+				else
+				{
+					BorderComponent[3] = BorderComponentY2[ComponentIndexXX];
+					if (BorderComponent[3])
+					{
+						if (BorderComponent[3]->XYOffsetmapTexture)
+						{
+							NeighborTexDataInfo[3] = GetTextureDataInfo(BorderComponent[3]->XYOffsetmapTexture);
+							NeighborXYOffsetmapTextureData[3] = (FColor*)NeighborTexDataInfo[3]->GetMipData(0);
+						}
+						else
+						{
+							NeighborXYOffsetmapTextureData[3] = EmptyXYOffset.GetData();
+						}
+					}
+				}
+
+				CornerComponent[0] = ComponentIndexX >= ComponentIndexX1 && ComponentIndexY >= ComponentIndexY1 ?
+					LandscapeInfo->XYtoComponentMap.FindRef(FIntPoint((ComponentIndexX - 1), (ComponentIndexY - 1))) : NULL;
+				CornerComponent[1] = ComponentIndexX <= ComponentIndexX2 && ComponentIndexY >= ComponentIndexY1 ?
+					LandscapeInfo->XYtoComponentMap.FindRef(FIntPoint((ComponentIndexX + 1), (ComponentIndexY - 1))) : NULL;
+				CornerComponent[2] = ComponentIndexX >= ComponentIndexX1 && ComponentIndexY <= ComponentIndexY2 ?
+					LandscapeInfo->XYtoComponentMap.FindRef(FIntPoint((ComponentIndexX - 1), (ComponentIndexY + 1))) : NULL;
+				CornerComponent[3] = ComponentIndexX <= ComponentIndexX2 && ComponentIndexY <= ComponentIndexY2 ?
+					LandscapeInfo->XYtoComponentMap.FindRef(FIntPoint((ComponentIndexX + 1), (ComponentIndexY + 1))) : NULL;
+
+				if (CornerComponent[0])
+				{
+					CornerSet |= 1;
+					CornerValues[0] = GetXYOffsetmapData(CornerComponent[0], EdgeCoord, EdgeCoord);
+				}
+				else if ((ExistLeft || ExistUp) && X1 <= ComponentIndexX*ComponentSizeQuads && Y1 <= ComponentIndexY*ComponentSizeQuads)
+				{
+					CornerSet |= 1;
+					CornerValues[0] = FVector2D(StoreData.Load(ComponentIndexX*ComponentSizeQuads, ComponentIndexY*ComponentSizeQuads));
+				}
+				else if (BorderComponent[0])
+				{
+					CornerSet |= 1;
+					CornerValues[0] = GetXYOffsetmapData(BorderComponent[0], EdgeCoord, 0, NeighborXYOffsetmapTextureData[0]);
+				}
+				else if (BorderComponent[2])
+				{
+					CornerSet |= 1;
+					CornerValues[0] = GetXYOffsetmapData(BorderComponent[2], 0, EdgeCoord, NeighborXYOffsetmapTextureData[2]);
+				}
+
+				if (CornerComponent[1])
+				{
+					CornerSet |= 1 << 1;
+					CornerValues[1] = GetXYOffsetmapData(CornerComponent[1], 0, EdgeCoord);
+				}
+				else if (ExistUp && X2 >= (ComponentIndexX + 1)*ComponentSizeQuads)
+				{
+					CornerSet |= 1 << 1;
+					CornerValues[1] = FVector2D(StoreData.Load((ComponentIndexX + 1)*ComponentSizeQuads, ComponentIndexY*ComponentSizeQuads));
+				}
+				else if (BorderComponent[1])
+				{
+					CornerSet |= 1 << 1;
+					CornerValues[1] = GetXYOffsetmapData(BorderComponent[1], 0, 0, NeighborXYOffsetmapTextureData[1]);
+				}
+				else if (BorderComponent[2])
+				{
+					CornerSet |= 1 << 1;
+					CornerValues[1] = GetXYOffsetmapData(BorderComponent[2], EdgeCoord, EdgeCoord, NeighborXYOffsetmapTextureData[2]);
+				}
+
+				if (CornerComponent[2])
+				{
+					CornerSet |= 1 << 2;
+					CornerValues[2] = GetXYOffsetmapData(CornerComponent[2], EdgeCoord, 0);
+				}
+				else if (ExistLeft && Y2 >= (ComponentIndexY + 1)*ComponentSizeQuads) // Use data already stored for 0, 2
+				{
+					CornerSet |= 1 << 2;
+					CornerValues[2] = FVector2D(StoreData.Load(ComponentIndexX*ComponentSizeQuads, (ComponentIndexY + 1)*ComponentSizeQuads));
+				}
+				else if (BorderComponent[0])
+				{
+					CornerSet |= 1 << 2;
+					CornerValues[2] = GetXYOffsetmapData(BorderComponent[0], EdgeCoord, EdgeCoord, NeighborXYOffsetmapTextureData[0]);
+				}
+				else if (BorderComponent[3])
+				{
+					CornerSet |= 1 << 2;
+					CornerValues[2] = GetXYOffsetmapData(BorderComponent[3], 0, 0, NeighborXYOffsetmapTextureData[3]);
+				}
+
+				if (CornerComponent[3])
+				{
+					CornerSet |= 1 << 3;
+					CornerValues[3] = GetXYOffsetmapData(CornerComponent[3], 0, 0);
+				}
+				else if (BorderComponent[1])
+				{
+					CornerSet |= 1 << 3;
+					CornerValues[3] = GetXYOffsetmapData(BorderComponent[1], 0, EdgeCoord, NeighborXYOffsetmapTextureData[1]);
+				}
+				else if (BorderComponent[3])
+				{
+					CornerSet |= 1 << 3;
+					CornerValues[3] = GetXYOffsetmapData(BorderComponent[3], EdgeCoord, 0, NeighborXYOffsetmapTextureData[3]);
+				}
+
+				FillCornerValues(CornerSet, CornerValues);
+				ComponentDataExist[ComponentIndexXY] = ExistLeft || ExistUp || (BorderComponent[0] || BorderComponent[1] || BorderComponent[2] || BorderComponent[3]) || CornerSet;
+			}
+
+			if (!ComponentDataExist[ComponentIndexXY])
+			{
+				continue;
+			}
+
+			// Find coordinates of box that lies inside component
+			int32 ComponentX1 = FMath::Clamp<int32>(X1 - ComponentIndexX*ComponentSizeQuads, 0, ComponentSizeQuads);
+			int32 ComponentY1 = FMath::Clamp<int32>(Y1 - ComponentIndexY*ComponentSizeQuads, 0, ComponentSizeQuads);
+			int32 ComponentX2 = FMath::Clamp<int32>(X2 - ComponentIndexX*ComponentSizeQuads, 0, ComponentSizeQuads);
+			int32 ComponentY2 = FMath::Clamp<int32>(Y2 - ComponentIndexY*ComponentSizeQuads, 0, ComponentSizeQuads);
+
+			// Find subsection range for this box
+			int32 SubIndexX1 = FMath::Clamp<int32>((ComponentX1 - 1) / SubsectionSizeQuads, 0, ComponentNumSubsections - 1);	// -1 because we need to pick up vertices shared between subsections
+			int32 SubIndexY1 = FMath::Clamp<int32>((ComponentY1 - 1) / SubsectionSizeQuads, 0, ComponentNumSubsections - 1);
+			int32 SubIndexX2 = FMath::Clamp<int32>(ComponentX2 / SubsectionSizeQuads, 0, ComponentNumSubsections - 1);
+			int32 SubIndexY2 = FMath::Clamp<int32>(ComponentY2 / SubsectionSizeQuads, 0, ComponentNumSubsections - 1);
+
+			for (int32 SubIndexY = SubIndexY1; SubIndexY <= SubIndexY2; SubIndexY++)
+			{
+				for (int32 SubIndexX = SubIndexX1; SubIndexX <= SubIndexX2; SubIndexX++)
+				{
+					// Find coordinates of box that lies inside subsection
+					int32 SubX1 = FMath::Clamp<int32>(ComponentX1 - SubsectionSizeQuads*SubIndexX, 0, SubsectionSizeQuads);
+					int32 SubY1 = FMath::Clamp<int32>(ComponentY1 - SubsectionSizeQuads*SubIndexY, 0, SubsectionSizeQuads);
+					int32 SubX2 = FMath::Clamp<int32>(ComponentX2 - SubsectionSizeQuads*SubIndexX, 0, SubsectionSizeQuads);
+					int32 SubY2 = FMath::Clamp<int32>(ComponentY2 - SubsectionSizeQuads*SubIndexY, 0, SubsectionSizeQuads);
+
+					// Update texture data for the box that lies inside subsection
+					for (int32 SubY = SubY1; SubY <= SubY2; SubY++)
+					{
+						for (int32 SubX = SubX1; SubX <= SubX2; SubX++)
+						{
+							int32 LandscapeX = SubIndexX*SubsectionSizeQuads + ComponentIndexX*ComponentSizeQuads + SubX;
+							int32 LandscapeY = SubIndexY*SubsectionSizeQuads + ComponentIndexY*ComponentSizeQuads + SubY;
+
+							// Find the input data corresponding to this vertex
+							if (Component)
+							{
+								// Find the texture data corresponding to this vertex
+								FVector2D XYOffset = GetXYOffsetmapData(Component, (SubsectionSizeQuads + 1) * SubIndexX + SubX, (SubsectionSizeQuads + 1) * SubIndexY + SubY, XYOffsetmapTextureData);
+								StoreData.Store(LandscapeX, LandscapeY, XYOffset);
+							}
+							else
+							{
+								// Find the texture data corresponding to this vertex
+								FVector2D Value[4] = { FVector2D::ZeroVector, FVector2D::ZeroVector, FVector2D::ZeroVector, FVector2D::ZeroVector };
+								int32 Dist[4] = { INT_MAX, INT_MAX, INT_MAX, INT_MAX };
+								FVector2D ValueX = FVector2D::ZeroVector, ValueY = FVector2D::ZeroVector;
+								bool Exist[4] = { false, false, false, false };
+
+								// Use data already stored for 0, 2
+								if (ExistLeft)
+								{
+									Value[0] = FVector2D(StoreData.Load(ComponentIndexX*ComponentSizeQuads, LandscapeY));
+									Dist[0] = LandscapeX - (ComponentIndexX*ComponentSizeQuads);
+									Exist[0] = true;
+								}
+								else if (BorderComponent[0])
+								{
+									Value[0] = GetXYOffsetmapData(BorderComponent[0], EdgeCoord, (SubsectionSizeQuads + 1) * SubIndexY + SubY, NeighborXYOffsetmapTextureData[0]);
+									Dist[0] = LandscapeX - (BorderComponent[0]->GetSectionBase().X + ComponentSizeQuads);
+									Exist[0] = true;
+								}
+								else
+								{
+									if ((CornerSet & 1) && (CornerSet & 1 << 2))
+									{
+										int32 Dist1 = LandscapeY - (ComponentIndexY*ComponentSizeQuads);
+										int32 Dist2 = ((ComponentIndexY + 1)*ComponentSizeQuads) - LandscapeY;
+										Value[0] = (Dist2 * CornerValues[0] + Dist1 * CornerValues[2]) / (Dist1 + Dist2);
+										Dist[0] = LandscapeX - (ComponentIndexX*ComponentSizeQuads);
+										Exist[0] = true;
+									}
+								}
+
+								if (BorderComponent[1])
+								{
+									Value[1] = GetXYOffsetmapData(BorderComponent[1], 0, (SubsectionSizeQuads + 1) * SubIndexY + SubY, NeighborXYOffsetmapTextureData[1]);
+									Dist[1] = (BorderComponent[1]->GetSectionBase().X) - LandscapeX;
+									Exist[1] = true;
+								}
+								else
+								{
+									if ((CornerSet & 1 << 1) && (CornerSet & 1 << 3))
+									{
+										int32 Dist1 = LandscapeY - (ComponentIndexY*ComponentSizeQuads);
+										int32 Dist2 = ((ComponentIndexY + 1)*ComponentSizeQuads) - LandscapeY;
+										Value[1] = (Dist2 * CornerValues[1] + Dist1 * CornerValues[3]) / (Dist1 + Dist2);
+										Dist[1] = (ComponentIndexX + 1)*ComponentSizeQuads - LandscapeX;
+										Exist[1] = true;
+									}
+								}
+
+								if (ExistUp)
+								{
+									Value[2] = FVector2D(StoreData.Load(LandscapeX, ComponentIndexY*ComponentSizeQuads));
+									Dist[2] = LandscapeY - (ComponentIndexY*ComponentSizeQuads);
+									Exist[2] = true;
+								}
+								else if (BorderComponent[2])
+								{
+									Value[2] = GetXYOffsetmapData(BorderComponent[2], (SubsectionSizeQuads + 1) * SubIndexX + SubX, EdgeCoord, NeighborXYOffsetmapTextureData[2]);
+									Dist[2] = LandscapeY - (BorderComponent[2]->GetSectionBase().Y + ComponentSizeQuads);
+									Exist[2] = true;
+								}
+								else
+								{
+									if ((CornerSet & 1) && (CornerSet & 1 << 1))
+									{
+										int32 Dist1 = LandscapeX - (ComponentIndexX*ComponentSizeQuads);
+										int32 Dist2 = (ComponentIndexX + 1)*ComponentSizeQuads - LandscapeX;
+										Value[2] = (Dist2 * CornerValues[0] + Dist1 * CornerValues[1]) / (Dist1 + Dist2);
+										Dist[2] = LandscapeY - (ComponentIndexY*ComponentSizeQuads);
+										Exist[2] = true;
+									}
+								}
+
+								if (BorderComponent[3])
+								{
+									Value[3] = GetXYOffsetmapData(BorderComponent[3], (SubsectionSizeQuads + 1) * SubIndexX + SubX, 0, NeighborXYOffsetmapTextureData[3]);
+									Dist[3] = (BorderComponent[3]->GetSectionBase().Y) - LandscapeY;
+									Exist[3] = true;
+								}
+								else
+								{
+									if ((CornerSet & 1 << 2) && (CornerSet & 1 << 3))
+									{
+										int32 Dist1 = LandscapeX - (ComponentIndexX*ComponentSizeQuads);
+										int32 Dist2 = (ComponentIndexX + 1)*ComponentSizeQuads - LandscapeX;
+										Value[3] = (Dist2 * CornerValues[2] + Dist1 * CornerValues[3]) / (Dist1 + Dist2);
+										Dist[3] = (ComponentIndexY + 1)*ComponentSizeQuads - LandscapeY;
+										Exist[3] = true;
+									}
+								}
+
+								CalcInterpValue<FVector2D>(Dist, Exist, Value, ValueX, ValueY);
+
+								FVector2D FinalValue = FVector2D::ZeroVector; // Default Value
+								if ((Exist[0] || Exist[1]) && (Exist[2] || Exist[3]))
+								{
+									FinalValue = CalcValueFromValueXY<FVector2D>(Dist, ValueX, ValueY, CornerSet, CornerValues);
+								}
+								else if ((BorderComponent[0] || BorderComponent[1]))
+								{
+									FinalValue = ValueX;
+								}
+								else if ((BorderComponent[2] || BorderComponent[3]))
+								{
+									FinalValue = ValueY;
+								}
+								else if ((Exist[0] || Exist[1]))
+								{
+									FinalValue = ValueX;
+								}
+								else if ((Exist[2] || Exist[3]))
+								{
+									FinalValue = ValueY;
+								}
+
+								StoreData.Store(LandscapeX, LandscapeY, FinalValue);
+								//StoreData.StoreDefault(LandscapeX, LandscapeY);
+							}
+						}
+					}
+				}
+			}
+		}
+	}
+
+	if (bHasMissingValue)
+	{
+		CalcMissingValues<FVector2D, TStoreData, FVector2D>(X1, X2, Y1, Y2,
+			ComponentIndexX1, ComponentIndexX2, ComponentIndexY1, ComponentIndexY2,
+			ComponentSizeX, ComponentSizeY, CornerValues,
+			NoBorderY1, NoBorderY2, ComponentDataExist, StoreData);
+		// Update valid region
+		ValidX1 = FMath::Max<int32>(X1, ValidX1);
+		ValidX2 = FMath::Min<int32>(X2, ValidX2);
+		ValidY1 = FMath::Max<int32>(Y1, ValidY1);
+		ValidY2 = FMath::Min<int32>(Y2, ValidY2);
+	}
+	else
+	{
+		ValidX1 = X1;
+		ValidX2 = X2;
+		ValidY1 = Y1;
+		ValidY2 = Y2;
+	}
+}
+
+void FLandscapeEditDataInterface::GetXYOffsetData(int32& X1, int32& Y1, int32& X2, int32& Y2, FVector2D* Data, int32 Stride)
+{
+	if (Stride == 0)
+	{
+		Stride = (1 + X2 - X1);
+	}
+	TArrayStoreData<FVector2D> ArrayStoreData(X1, Y1, Data, Stride);
+	GetXYOffsetDataTempl(X1, Y1, X2, Y2, ArrayStoreData);
+}
+
+void FLandscapeEditDataInterface::GetXYOffsetData(int32& X1, int32& Y1, int32& X2, int32& Y2, TMap<FIntPoint, FVector2D>& SparseData)
+{
+	TSparseStoreData<FVector2D> SparseStoreData(SparseData);
+	GetXYOffsetDataTempl(X1, Y1, X2, Y2, SparseStoreData);
+}
+
+void FLandscapeEditDataInterface::GetXYOffsetData(int32& X1, int32& Y1, int32& X2, int32& Y2, FVector* Data, int32 Stride)
+{
+	if (Stride == 0)
+	{
+		Stride = (1 + X2 - X1);
+	}
+	TArrayStoreData<FVector> ArrayStoreData(X1, Y1, Data, Stride);
+	GetXYOffsetDataTempl(X1, Y1, X2, Y2, ArrayStoreData);
+}
+
+void FLandscapeEditDataInterface::GetXYOffsetData(int32& X1, int32& Y1, int32& X2, int32& Y2, TMap<FIntPoint, FVector>& SparseData)
+{
+	TSparseStoreData<FVector> SparseStoreData(SparseData);
+	GetXYOffsetDataTempl(X1, Y1, X2, Y2, SparseStoreData);
+}
+
+template<typename TStoreData>
+void FLandscapeEditDataInterface::GetXYOffsetDataTemplFast(const int32 X1, const int32 Y1, const int32 X2, const int32 Y2, TStoreData& StoreData)
 {
 	int32 ComponentIndexX1, ComponentIndexY1, ComponentIndexX2, ComponentIndexY2;
 	ALandscape::CalcComponentIndices(X1, Y1, X2, Y2, ComponentSizeQuads, ComponentIndexX1, ComponentIndexY1, ComponentIndexX2, ComponentIndexY2);
@@ -4113,36 +4671,36 @@ void FLandscapeEditDataInterface::GetXYOffsetDataTempl(const int32 X1, const int
 	}
 }
 
-void FLandscapeEditDataInterface::GetXYOffsetData(const int32 X1, const int32 Y1, const int32 X2, const int32 Y2, FVector2D* Data, int32 Stride)
+void FLandscapeEditDataInterface::GetXYOffsetDataFast(const int32 X1, const int32 Y1, const int32 X2, const int32 Y2, FVector2D* Data, int32 Stride)
 {
 	if( Stride==0 )
 	{
 		Stride = (1+X2-X1);
 	}
 	TArrayStoreData<FVector2D> ArrayStoreData(X1, Y1, Data, Stride);
-	GetXYOffsetDataTempl(X1, Y1, X2, Y2, ArrayStoreData);
+	GetXYOffsetDataTemplFast(X1, Y1, X2, Y2, ArrayStoreData);
 }
 
-void FLandscapeEditDataInterface::GetXYOffsetData(const int32 X1, const int32 Y1, const int32 X2, const int32 Y2, TMap<FIntPoint, FVector2D>& SparseData)
+void FLandscapeEditDataInterface::GetXYOffsetDataFast(const int32 X1, const int32 Y1, const int32 X2, const int32 Y2, TMap<FIntPoint, FVector2D>& SparseData)
 {
 	TSparseStoreData<FVector2D> SparseStoreData(SparseData);
-	GetXYOffsetDataTempl(X1, Y1, X2, Y2, SparseStoreData);
+	GetXYOffsetDataTemplFast(X1, Y1, X2, Y2, SparseStoreData);
 }
 
-void FLandscapeEditDataInterface::GetXYOffsetData(const int32 X1, const int32 Y1, const int32 X2, const int32 Y2, FVector* Data, int32 Stride)
+void FLandscapeEditDataInterface::GetXYOffsetDataFast(const int32 X1, const int32 Y1, const int32 X2, const int32 Y2, FVector* Data, int32 Stride)
 {
 	if( Stride==0 )
 	{
 		Stride = (1+X2-X1);
 	}
 	TArrayStoreData<FVector> ArrayStoreData(X1, Y1, Data, Stride);
-	GetXYOffsetDataTempl(X1, Y1, X2, Y2, ArrayStoreData);
+	GetXYOffsetDataTemplFast(X1, Y1, X2, Y2, ArrayStoreData);
 }
 
-void FLandscapeEditDataInterface::GetXYOffsetData(const int32 X1, const int32 Y1, const int32 X2, const int32 Y2, TMap<FIntPoint, FVector>& SparseData)
+void FLandscapeEditDataInterface::GetXYOffsetDataFast(const int32 X1, const int32 Y1, const int32 X2, const int32 Y2, TMap<FIntPoint, FVector>& SparseData)
 {
 	TSparseStoreData<FVector> SparseStoreData(SparseData);
-	GetXYOffsetDataTempl(X1, Y1, X2, Y2, SparseStoreData);
+	GetXYOffsetDataTemplFast(X1, Y1, X2, Y2, SparseStoreData);
 }
 
 //
diff --git a/Engine/Source/Runtime/Engine/Public/Landscape/LandscapeEdit.h b/Engine/Source/Runtime/Engine/Public/Landscape/LandscapeEdit.h
--- a/Engine/Source/Runtime/Engine/Public/Landscape/LandscapeEdit.h
+++ b/Engine/Source/Runtime/Engine/Public/Landscape/LandscapeEdit.h
@@ -137,15 +137,21 @@ struct ENGINE_API FLandscapeEditDataInterface
 	void SetXYOffsetData(int32 X1, int32 Y1, int32 X2, int32 Y2, const FVector* Data, int32 Stride);
 	// Helper accessor
 	FORCEINLINE FVector2D GetXYOffsetmapData(const ULandscapeComponent* Component, int32 TexU, int32 TexV, FColor* TextureData = NULL);
-	
-	// No interpolation supported for XY Offset due to collision problem...
+
+	template<typename TStoreData>
+	void GetXYOffsetDataTempl(int32& X1, int32& Y1, int32& X2, int32& Y2, TStoreData& StoreData);
+	void GetXYOffsetData(int32& X1, int32& Y1, int32& X2, int32& Y2, FVector2D* Data, int32 Stride);
+	void GetXYOffsetData(int32& X1, int32& Y1, int32& X2, int32& Y2, TMap<FIntPoint, FVector2D>& SparseData);
+	void GetXYOffsetData(int32& X1, int32& Y1, int32& X2, int32& Y2, FVector* Data, int32 Stride);
+	void GetXYOffsetData(int32& X1, int32& Y1, int32& X2, int32& Y2, TMap<FIntPoint, FVector>& SparseData);
+	// Without data interpolation
 	template<typename TStoreData>
-	void GetXYOffsetDataTempl(const int32 X1, const int32 Y1, const int32 X2, const int32 Y2, TStoreData& StoreData);
+	void GetXYOffsetDataTemplFast(const int32 X1, const int32 Y1, const int32 X2, const int32 Y2, TStoreData& StoreData);
 	// Without data interpolation, able to get normal data
-	void GetXYOffsetData(const int32 X1, const int32 Y1, const int32 X2, const int32 Y2, FVector2D* Data, int32 Stride);
-	void GetXYOffsetData(const int32 X1, const int32 Y1, const int32 X2, const int32 Y2, TMap<FIntPoint, FVector2D>& SparseData);
-	void GetXYOffsetData(const int32 X1, const int32 Y1, const int32 X2, const int32 Y2, FVector* Data, int32 Stride);
-	void GetXYOffsetData(const int32 X1, const int32 Y1, const int32 X2, const int32 Y2, TMap<FIntPoint, FVector>& SparseData);
+	void GetXYOffsetDataFast(const int32 X1, const int32 Y1, const int32 X2, const int32 Y2, FVector2D* Data, int32 Stride);
+	void GetXYOffsetDataFast(const int32 X1, const int32 Y1, const int32 X2, const int32 Y2, TMap<FIntPoint, FVector2D>& SparseData);
+	void GetXYOffsetDataFast(const int32 X1, const int32 Y1, const int32 X2, const int32 Y2, FVector* Data, int32 Stride);
+	void GetXYOffsetDataFast(const int32 X1, const int32 Y1, const int32 X2, const int32 Y2, TMap<FIntPoint, FVector>& SparseData);
 
 	// Texture data access
 	FLandscapeTextureDataInfo* GetTextureDataInfo(UTexture2D* Texture);
@@ -177,7 +183,7 @@ private:
 	ULandscapeInfo* LandscapeInfo;
 
 	// Only for Missing Data interpolation... only internal usage
-	template<typename TData, typename TStoreData>
+	template<typename TData, typename TStoreData, typename FType>
 	FORCEINLINE void CalcMissingValues(const int32& X1, const int32& X2, const int32& Y1, const int32& Y2, 
 		const int32& ComponentIndexX1, const int32& ComponentIndexX2, const int32& ComponentIndexY1, const int32& ComponentIndexY2, 
 		const int32& ComponentSizeX, const int32& ComponentSizeY, TData* CornerValues, 
