Fixed spline duplication issue, added new BP callable method to API
#ttp 341711 - PAPER2D: 2D terrain splines fold back on themselves upon duplication
#branch UE4
#proj Runtime.Engine
Removed USplineComponent::PostDuplicate (which is not called for components) and replaced it with USplineComponent::PostEditImport, so that the spline reparam table can be rebuilt after duplication.
Added USplineComponent::GetWorldTangentAtDistanceAlongSpline, returning a non unit vector of the tangent at the given location.

diff --git a/Engine/Source/Runtime/Engine/Classes/Components/SplineComponent.h b/Engine/Source/Runtime/Engine/Classes/Components/SplineComponent.h
--- a/Engine/Source/Runtime/Engine/Classes/Components/SplineComponent.h
+++ b/Engine/Source/Runtime/Engine/Classes/Components/SplineComponent.h
@@ -96,6 +96,10 @@ class ENGINE_API USplineComponent : public USceneComponent
 	UFUNCTION(BlueprintCallable, Category=Spline)
 	FVector GetWorldDirectionAtDistanceAlongSpline(float Distance) const;
 
+	/** Given a distance along the length of this spline, return the tangent vector of the spline there. */
+	UFUNCTION(BlueprintCallable, Category = Spline)
+	FVector GetWorldTangentAtDistanceAlongSpline(float Distance) const;
+
 	/** Given a distance along the length of this spline, return a rotation corresponding to the spline's position and direction there. */
 	UFUNCTION(BlueprintCallable, Category = Spline)
 	FRotator GetWorldRotationAtDistanceAlongSpline(float Distance) const;
@@ -114,7 +118,7 @@ class ENGINE_API USplineComponent : public USceneComponent
 
 	// UObject interface
 	virtual void PostLoad() override;
-	virtual void PostDuplicate(bool bDuplicateForPIE) override;
+	virtual void PostEditImport() override;
 #if WITH_EDITORONLY_DATA
 	virtual void PostEditChangeProperty(FPropertyChangedEvent& PropertyChangedEvent) override;
 	// End of UObject interface
diff --git a/Engine/Source/Runtime/Engine/Private/SplineComponent.cpp b/Engine/Source/Runtime/Engine/Private/SplineComponent.cpp
--- a/Engine/Source/Runtime/Engine/Private/SplineComponent.cpp
+++ b/Engine/Source/Runtime/Engine/Private/SplineComponent.cpp
@@ -36,9 +36,9 @@ void USplineComponent::PostLoad()
 }
 
 
-void USplineComponent::PostDuplicate(bool bDuplicateForPIE)
+void USplineComponent::PostEditImport()
 {
-	Super::PostDuplicate(bDuplicateForPIE);
+	Super::PostEditImport();
 	UpdateSpline();
 }
 
@@ -302,6 +302,14 @@ FVector USplineComponent::GetWorldLocationAtDistanceAlongSpline(float Distance)
 }
 
 
+FVector USplineComponent::GetWorldTangentAtDistanceAlongSpline(float Distance) const
+{
+	const float Param = SplineReparamTable.Eval(Distance, 0.f);
+	const FVector Tangent = SplineInfo.EvalDerivative(Param, FVector::ZeroVector);
+	return ComponentToWorld.TransformVector(Tangent);
+}
+
+
 FVector USplineComponent::GetWorldDirectionAtDistanceAlongSpline(float Distance) const
 {
 	const float Param = SplineReparamTable.Eval(Distance, 0.f);
