Prevent spiking artifacts when using tessellation with landscape.

#ttp 341604 - UE4: Rendering: Adding a Tessellated Material to a landscape is casuing heavy spiking

diff --git a/Engine/Shaders/LandscapeVertexFactory.usf b/Engine/Shaders/LandscapeVertexFactory.usf
--- a/Engine/Shaders/LandscapeVertexFactory.usf
+++ b/Engine/Shaders/LandscapeVertexFactory.usf
@@ -702,7 +702,7 @@ FVertexFactoryInterpolantsVSToDS VertexFactoryGetInterpolantsVSToDS(FVertexFacto
 #if USING_TESSELLATION
 	// Calculate LocalToTangent directly from normal map texture.
 	float3x3 TangentToLocal = CalcTangentBasisFromWorldNormal( Intermediates.WorldNormal );
-	float3x3 TangentToWorld = mul(TangentToLocal, GetLocalToWorld3x3());
+	float3x3 TangentToWorld = mul(TangentToLocal, (float3x3)LandscapeParameters.LocalToWorldNoScaling);
 	
 	Interpolants.TangentToWorld0 = TangentToWorld[0];
 	Interpolants.TangentToWorld2 = float4(TangentToWorld[2], 1);
diff --git a/Engine/Source/Runtime/Engine/Private/Landscape/LandscapeRender.cpp b/Engine/Source/Runtime/Engine/Private/Landscape/LandscapeRender.cpp
--- a/Engine/Source/Runtime/Engine/Private/Landscape/LandscapeRender.cpp
+++ b/Engine/Source/Runtime/Engine/Private/Landscape/LandscapeRender.cpp
@@ -985,6 +985,7 @@ void FLandscapeComponentSceneProxy::OnTransformChanged()
 	FLandscapeUniformShaderParameters LandscapeParams;
 	LandscapeParams.HeightmapUVScaleBias = HeightmapScaleBias;
 	LandscapeParams.WeightmapUVScaleBias = WeightmapScaleBias;
+	LandscapeParams.LocalToWorldNoScaling = LocalToWorldNoScaling;
 
 	LandscapeParams.LandscapeLightmapScaleBias = FVector4(
 		LightmapScaleX,
diff --git a/Engine/Source/Runtime/Engine/Public/Landscape/LandscapeRender.h b/Engine/Source/Runtime/Engine/Public/Landscape/LandscapeRender.h
--- a/Engine/Source/Runtime/Engine/Public/Landscape/LandscapeRender.h
+++ b/Engine/Source/Runtime/Engine/Public/Landscape/LandscapeRender.h
@@ -80,6 +80,7 @@ BEGIN_UNIFORM_BUFFER_STRUCT(FLandscapeUniformShaderParameters,ENGINE_API)
 	DECLARE_UNIFORM_BUFFER_STRUCT_MEMBER(FVector4,SubsectionSizeVertsLayerUVPan)
 	DECLARE_UNIFORM_BUFFER_STRUCT_MEMBER(FVector4,SubsectionOffsetParams)
 	DECLARE_UNIFORM_BUFFER_STRUCT_MEMBER(FVector4,LightmapSubsectionOffsetParams)
+	DECLARE_UNIFORM_BUFFER_STRUCT_MEMBER(FMatrix,LocalToWorldNoScaling)
 END_UNIFORM_BUFFER_STRUCT(FLandscapeUniformShaderParameters)
 
 /* Data needed for the landscape vertex factory to set the render state for an individual batch element */
