Missing file from previous check-in

diff --git a/Engine/Source/Runtime/Engine/Private/Actor.cpp b/Engine/Source/Runtime/Engine/Private/Actor.cpp
--- a/Engine/Source/Runtime/Engine/Private/Actor.cpp
+++ b/Engine/Source/Runtime/Engine/Private/Actor.cpp
@@ -310,7 +310,7 @@ void AActor::PostInitProperties()
 void AActor::AddReferencedObjects(UObject* InThis, FReferenceCollector& Collector)
 {
 	AActor* This = CastChecked<AActor>(InThis);
-	Collector.AddReferencedObject(This->OwnedComponents);
+	Collector.AddReferencedObjects(This->OwnedComponents);
 	Super::AddReferencedObjects(InThis, Collector);
 }
 
