Header cleanup: PendingNetGame.h, SkeletalRenderCPUSkin.h.

diff --git a/Engine/Source/Runtime/Engine/Classes/Engine/PendingNetGame.h b/Engine/Source/Runtime/Engine/Classes/Engine/PendingNetGame.h
--- a/Engine/Source/Runtime/Engine/Classes/Engine/PendingNetGame.h
+++ b/Engine/Source/Runtime/Engine/Classes/Engine/PendingNetGame.h
@@ -1,9 +1,6 @@
 // Copyright 1998-2014 Epic Games, Inc. All Rights Reserved.
 
 #pragma once
-
-#include "GameViewportClient.h"
-
 #include "PendingNetGame.generated.h"
 
 /**
diff --git a/Engine/Source/Runtime/Engine/Private/SkeletalRenderCPUSkin.h b/Engine/Source/Runtime/Engine/Private/SkeletalRenderCPUSkin.h
--- a/Engine/Source/Runtime/Engine/Private/SkeletalRenderCPUSkin.h
+++ b/Engine/Source/Runtime/Engine/Private/SkeletalRenderCPUSkin.h
@@ -8,7 +8,6 @@
 
 #include "SkeletalRender.h"
 #include "SkeletalRenderPublic.h"
-#include "GPUSkinVertexFactory.h"
 #include "LocalVertexFactory.h"
 
 /** data for a single skinned skeletal mesh vertex */
