Added ENGINE_API to FComponentInstanceDataBase so it could be used in other modules

https://udn.unrealengine.com/questions/215060/using-fcomponentinstancedatabase-outside-engine.html

diff --git a/Engine/Source/Runtime/Engine/Public/ComponentInstanceDataCache.h b/Engine/Source/Runtime/Engine/Public/ComponentInstanceDataCache.h
--- a/Engine/Source/Runtime/Engine/Public/ComponentInstanceDataCache.h
+++ b/Engine/Source/Runtime/Engine/Public/ComponentInstanceDataCache.h
@@ -6,7 +6,7 @@ class UActorComponent;
 class AActor;
 
 /** Base class for component instance cached data of a particular type. */
-class FComponentInstanceDataBase
+class ENGINE_API FComponentInstanceDataBase
 {
 public:
 	FComponentInstanceDataBase()
