Rename Public/Canvas.h to CanvasTypes.h to avoid duplicate/ambiguous header name
Allow CanvasItem.h to be included directly
- Move ESimpleElementBlendMode to SceneTypes.h
- Move FCanvasUVTri, FFontRenderInfo & FDepthFieldGlowInfo into EngineTypes.h so they can be shared by Canvas.h and CanvasTypes.h
- Move several FCanvasTileItem ctors into cpp file, to prevent FTexture needing to be fully declared
Move FBoneContainer into its own header
- Move a couple of ctors to avoid some classes needing full decl
Add some forward decls to vehicle classes to allow them to be included directly

diff --git a/Engine/Source/Runtime/Engine/Classes/Engine/Canvas.h b/Engine/Source/Runtime/Engine/Classes/Engine/Canvas.h
--- a/Engine/Source/Runtime/Engine/Classes/Engine/Canvas.h
+++ b/Engine/Source/Runtime/Engine/Classes/Engine/Canvas.h
@@ -6,10 +6,12 @@
 
 #pragma once
 #include "WordWrapper.h"
+#include "Engine/EngineTypes.h"
 #include "Canvas.generated.h"
 
 class FSceneView;
 class FCanvas;
+class FCanvasItem;
 
 /**
  * Holds texture information with UV coordinates as well.
@@ -48,146 +50,6 @@ struct FCanvasIcon
 
 };
 
-/** info for glow when using depth field rendering */
-USTRUCT()
-struct FDepthFieldGlowInfo
-{
-	GENERATED_USTRUCT_BODY()
-
-	/** whether to turn on the outline glow (depth field fonts only) */
-	UPROPERTY()
-	uint32 bEnableGlow:1;
-
-	/** base color to use for the glow */
-	UPROPERTY()
-	FLinearColor GlowColor;
-
-	/** if bEnableGlow, outline glow outer radius (0 to 1, 0.5 is edge of character silhouette)
-	 * glow influence will be 0 at GlowOuterRadius.X and 1 at GlowOuterRadius.Y
-	*/
-	UPROPERTY()
-	FVector2D GlowOuterRadius;
-
-	/** if bEnableGlow, outline glow inner radius (0 to 1, 0.5 is edge of character silhouette)
-	 * glow influence will be 1 at GlowInnerRadius.X and 0 at GlowInnerRadius.Y
-	 */
-	UPROPERTY()
-	FVector2D GlowInnerRadius;
-
-
-	FDepthFieldGlowInfo()
-		: bEnableGlow(false)
-		, GlowColor(ForceInit)
-		, GlowOuterRadius(ForceInit)
-		, GlowInnerRadius(ForceInit)
-	{
-	}
-
-
-		bool operator==(const FDepthFieldGlowInfo& Other) const
-		{
-			if (Other.bEnableGlow != bEnableGlow)
-			{
-				return false;
-			}
-			else if (!bEnableGlow)
-			{
-				// if the glow is disabled on both, the other values don't matter
-				return true;
-			}
-			else
-			{
-				return (Other.GlowColor == GlowColor && Other.GlowOuterRadius == GlowOuterRadius && Other.GlowInnerRadius == GlowInnerRadius);
-			}
-		}
-		bool operator!=(const FDepthFieldGlowInfo& Other) const
-		{
-			return !(*this == Other);
-		}
-	
-};
-
-/** information used in font rendering */
-USTRUCT()
-struct FFontRenderInfo
-{
-	GENERATED_USTRUCT_BODY()
-
-	/** whether to clip text */
-	UPROPERTY()
-	uint32 bClipText:1;
-
-	/** whether to turn on shadowing */
-	UPROPERTY()
-	uint32 bEnableShadow:1;
-
-	/** depth field glow parameters (only usable if font was imported with a depth field) */
-	UPROPERTY()
-	struct FDepthFieldGlowInfo GlowInfo;
-
-	FFontRenderInfo()
-		: bClipText(false), bEnableShadow(false)
-	{}
-};
-
-/** Simple 2d triangle with UVs */
-USTRUCT()
-struct FCanvasUVTri
-{
-	GENERATED_USTRUCT_BODY()
-
-	/** Position of first vertex */
-	UPROPERTY(EditAnywhere, BlueprintReadWrite, Category=CanvasUVTri)
-	FVector2D V0_Pos;
-
-	/** UV of first vertex */
-	UPROPERTY(EditAnywhere, BlueprintReadWrite, Category=CanvasUVTri)
-	FVector2D V0_UV;
-
-	/** Color of first vertex */
-	UPROPERTY(EditAnywhere, BlueprintReadWrite, Category=CanvasUVTri)
-	FLinearColor V0_Color;
-
-	/** Position of second vertex */
-	UPROPERTY(EditAnywhere, BlueprintReadWrite, Category=CanvasUVTri)
-	FVector2D V1_Pos;
-
-	/** UV of second vertex */
-	UPROPERTY(EditAnywhere, BlueprintReadWrite, Category=CanvasUVTri)
-	FVector2D V1_UV;
-
-	/** Color of second vertex */
-	UPROPERTY(EditAnywhere, BlueprintReadWrite, Category=CanvasUVTri)
-	FLinearColor V1_Color;
-	
-	/** Position of third vertex */
-	UPROPERTY(EditAnywhere, BlueprintReadWrite, Category=CanvasUVTri)
-	FVector2D V2_Pos;
-
-	/** UV of third vertex */
-	UPROPERTY(EditAnywhere, BlueprintReadWrite, Category=CanvasUVTri)
-	FVector2D V2_UV;
-
-	/** Color of third vertex */
-	UPROPERTY(EditAnywhere, BlueprintReadWrite, Category=CanvasUVTri)
-	FLinearColor V2_Color;
-
-	FCanvasUVTri()
-		: V0_Pos(ForceInit)
-		, V0_UV(ForceInit)
-		, V0_Color(ForceInit)
-		, V1_Pos(ForceInit)
-		, V1_UV(ForceInit)
-		, V1_Color(ForceInit)
-		, V2_Pos(ForceInit)
-		, V2_UV(ForceInit)
-		, V2_Color(ForceInit)
-	{
-	}
-
-};
-
-template <> struct TIsZeroConstructType<FCanvasUVTri> { enum { Value = true }; };
 
 /**
  * General purpose data structure for grouping all parameters needed when sizing or wrapping a string
@@ -567,7 +429,7 @@ public:
 	 *
 	 * @param Item			Item to draw
 	 */
-	void DrawItem( class FCanvasItem& Item );
+	void DrawItem( FCanvasItem& Item );
 
 	/* 
 	 * Draw a CanvasItem at the given coordinates
@@ -575,7 +437,7 @@ public:
 	 * @param Item			Item to draw
 	 * @param InPostion		Position to draw item
 	 */
-	void DrawItem( class FCanvasItem& Item, const FVector2D& InPosition );
+	void DrawItem( FCanvasItem& Item, const FVector2D& InPosition );
 	
 	/* 
 	 * Draw a CanvasItem at the given coordinates
@@ -584,7 +446,7 @@ public:
 	 * @param X				X Position to draw item
 	 * @param Y				Y Position to draw item
 	 */
-	void DrawItem( class FCanvasItem& Item, float X, float Y );
+	void DrawItem( FCanvasItem& Item, float X, float Y );
 
 	/** Creates if necessary and returns ReporterGraph instance for 2d graph canvas drawing */
 	TWeakObjectPtr<class UReporterGraph>  GetReporterGraph();
diff --git a/Engine/Source/Runtime/Engine/Classes/Engine/CanvasRenderTarget2D.h b/Engine/Source/Runtime/Engine/Classes/Engine/CanvasRenderTarget2D.h
--- a/Engine/Source/Runtime/Engine/Classes/Engine/CanvasRenderTarget2D.h
+++ b/Engine/Source/Runtime/Engine/Classes/Engine/CanvasRenderTarget2D.h
@@ -2,9 +2,10 @@
 
 #pragma once
 
-#include "Canvas.h"
 #include "CanvasRenderTarget2D.generated.h"
 
+class UCanvas;
+
 /** This delegate is assignable through Blueprint and has similar functionality to the above. */
 DECLARE_DYNAMIC_MULTICAST_DELEGATE_ThreeParams(FOnCanvasRenderTargetUpdate, UCanvas*, Canvas, int32, Width, int32, Height);
 
@@ -54,7 +55,7 @@ public:
 	 * @param	Height				Height of the render target.
 	 */
 	UFUNCTION(BlueprintImplementableEvent, Category="Canvas Render Target 2D")
-	virtual void ReceiveUpdate(class UCanvas* Canvas, int32 Width, int32 Height);
+	virtual void ReceiveUpdate(UCanvas* Canvas, int32 Width, int32 Height);
 
 	/**
 	 * Gets a specific render target's size from the global map of canvas render targets.
diff --git a/Engine/Source/Runtime/Engine/Classes/Engine/EngineTypes.h b/Engine/Source/Runtime/Engine/Classes/Engine/EngineTypes.h
--- a/Engine/Source/Runtime/Engine/Classes/Engine/EngineTypes.h
+++ b/Engine/Source/Runtime/Engine/Classes/Engine/EngineTypes.h
@@ -2694,3 +2694,142 @@ public:
 	}
 };
 
+/** info for glow when using depth field rendering */
+USTRUCT()
+struct FDepthFieldGlowInfo
+{
+	GENERATED_USTRUCT_BODY()
+
+	/** whether to turn on the outline glow (depth field fonts only) */
+	UPROPERTY()
+	uint32 bEnableGlow:1;
+
+	/** base color to use for the glow */
+	UPROPERTY()
+	FLinearColor GlowColor;
+
+	/** if bEnableGlow, outline glow outer radius (0 to 1, 0.5 is edge of character silhouette)
+	 * glow influence will be 0 at GlowOuterRadius.X and 1 at GlowOuterRadius.Y
+	*/
+	UPROPERTY()
+	FVector2D GlowOuterRadius;
+
+	/** if bEnableGlow, outline glow inner radius (0 to 1, 0.5 is edge of character silhouette)
+	 * glow influence will be 1 at GlowInnerRadius.X and 0 at GlowInnerRadius.Y
+	 */
+	UPROPERTY()
+	FVector2D GlowInnerRadius;
+
+
+	FDepthFieldGlowInfo()
+		: bEnableGlow(false)
+		, GlowColor(ForceInit)
+		, GlowOuterRadius(ForceInit)
+		, GlowInnerRadius(ForceInit)
+	{
+	}
+
+
+		bool operator==(const FDepthFieldGlowInfo& Other) const
+		{
+			if (Other.bEnableGlow != bEnableGlow)
+			{
+				return false;
+			}
+			else if (!bEnableGlow)
+			{
+				// if the glow is disabled on both, the other values don't matter
+				return true;
+			}
+			else
+			{
+				return (Other.GlowColor == GlowColor && Other.GlowOuterRadius == GlowOuterRadius && Other.GlowInnerRadius == GlowInnerRadius);
+			}
+		}
+		bool operator!=(const FDepthFieldGlowInfo& Other) const
+		{
+			return !(*this == Other);
+		}
+	
+};
+
+/** information used in font rendering */
+USTRUCT()
+struct FFontRenderInfo
+{
+	GENERATED_USTRUCT_BODY()
+
+	/** whether to clip text */
+	UPROPERTY()
+	uint32 bClipText:1;
+
+	/** whether to turn on shadowing */
+	UPROPERTY()
+	uint32 bEnableShadow:1;
+
+	/** depth field glow parameters (only usable if font was imported with a depth field) */
+	UPROPERTY()
+	struct FDepthFieldGlowInfo GlowInfo;
+
+	FFontRenderInfo()
+		: bClipText(false), bEnableShadow(false)
+	{}
+};
+
+/** Simple 2d triangle with UVs */
+USTRUCT()
+struct FCanvasUVTri
+{
+	GENERATED_USTRUCT_BODY()
+
+	/** Position of first vertex */
+	UPROPERTY(EditAnywhere, BlueprintReadWrite, Category = CanvasUVTri)
+	FVector2D V0_Pos;
+
+	/** UV of first vertex */
+	UPROPERTY(EditAnywhere, BlueprintReadWrite, Category = CanvasUVTri)
+		FVector2D V0_UV;
+
+	/** Color of first vertex */
+	UPROPERTY(EditAnywhere, BlueprintReadWrite, Category = CanvasUVTri)
+		FLinearColor V0_Color;
+
+	/** Position of second vertex */
+	UPROPERTY(EditAnywhere, BlueprintReadWrite, Category = CanvasUVTri)
+		FVector2D V1_Pos;
+
+	/** UV of second vertex */
+	UPROPERTY(EditAnywhere, BlueprintReadWrite, Category = CanvasUVTri)
+		FVector2D V1_UV;
+
+	/** Color of second vertex */
+	UPROPERTY(EditAnywhere, BlueprintReadWrite, Category = CanvasUVTri)
+		FLinearColor V1_Color;
+
+	/** Position of third vertex */
+	UPROPERTY(EditAnywhere, BlueprintReadWrite, Category = CanvasUVTri)
+		FVector2D V2_Pos;
+
+	/** UV of third vertex */
+	UPROPERTY(EditAnywhere, BlueprintReadWrite, Category = CanvasUVTri)
+		FVector2D V2_UV;
+
+	/** Color of third vertex */
+	UPROPERTY(EditAnywhere, BlueprintReadWrite, Category = CanvasUVTri)
+		FLinearColor V2_Color;
+
+	FCanvasUVTri()
+		: V0_Pos(ForceInit)
+		, V0_UV(ForceInit)
+		, V0_Color(ForceInit)
+		, V1_Pos(ForceInit)
+		, V1_UV(ForceInit)
+		, V1_Color(ForceInit)
+		, V2_Pos(ForceInit)
+		, V2_UV(ForceInit)
+		, V2_Color(ForceInit)
+	{
+	}
+};
+
+template <> struct TIsZeroConstructType<FCanvasUVTri> { enum { Value = true }; };
diff --git a/Engine/Source/Runtime/Engine/Classes/Engine/SkeletalMesh.h b/Engine/Source/Runtime/Engine/Classes/Engine/SkeletalMesh.h
--- a/Engine/Source/Runtime/Engine/Classes/Engine/SkeletalMesh.h
+++ b/Engine/Source/Runtime/Engine/Classes/Engine/SkeletalMesh.h
@@ -8,12 +8,14 @@
 
 #include "SkeletalMeshTypes.h"
 #include "Animation/PreviewAssetAttachComponent.h"
+#include "BoneContainer.h"
 #include "SkeletalMesh.generated.h"
 
 /** The maximum number of skeletal mesh LODs allowed. */
 #define MAX_SKELETAL_MESH_LODS 5
 
 class UMorphTarget;
+class USkeleton;
 
 #if WITH_APEX_CLOTHING
 struct FApexClothCollisionVolumeData;
@@ -117,221 +119,6 @@ struct FTriangleSortSettings
 
 };
 
-/**
-* This is a native transient structure.
-* Contains:
-* - BoneIndicesArray: Array of RequiredBoneIndices for Current Asset. In increasing order. Mapping to current Array of Transforms (Pose).
-* - BoneSwitchArray: Size of current Skeleton. true if Bone is contained in RequiredBones array, false otherwise.
-**/
-struct FBoneContainer
-{
-private:
-	/** Array of RequiredBonesIndices. In increasing order. */
-	TArray<FBoneIndexType>	BoneIndicesArray;
-	/** Array sized by Current RefPose. true if Bone is contained in RequiredBones array, false otherwise. */
-	TBitArray<>				BoneSwitchArray;
-
-	/** Asset BoneIndicesArray was made for. Typically a SkeletalMesh. */
-	TWeakObjectPtr<UObject>	Asset;
-	/** If Asset is a SkeletalMesh, this will be a pointer to it. Can be NULL if Asset is a USkeleton. */
-	TWeakObjectPtr<USkeletalMesh> AssetSkeletalMesh;
-	/** If Asset is a Skeleton that will be it. If Asset is a SkeletalMesh, that will be its Skeleton. */
-	TWeakObjectPtr<USkeleton> AssetSkeleton;
-
-	/** Pointer to RefSkeleton of Asset. */
-	const FReferenceSkeleton* RefSkeleton;
-
-	/** Mapping table between Skeleton Bone Indices and Pose Bone Indices. */
-	TArray<int32> SkeletonToPoseBoneIndexArray;
-
-	/** Mapping table between Pose Bone Indices and Skeleton Bone Indices. */
-	TArray<int32> PoseToSkeletonBoneIndexArray;
-
-	/** For debugging. */
-	/** Disable Retargeting. Extract animation, but do not retarget it. */
-	bool bDisableRetargeting;
-	/** Disable animation compression, use RAW data instead. */
-	bool bUseRAWData;
-
-public:
-
-	FBoneContainer()
-		: Asset(NULL)
-		, AssetSkeletalMesh(NULL)
-		, AssetSkeleton(NULL)
-		, RefSkeleton(NULL)
-		, bDisableRetargeting(false)
-		, bUseRAWData(false)
-	{
-		BoneIndicesArray.Empty();
-		BoneSwitchArray.Empty();
-		SkeletonToPoseBoneIndexArray.Empty();
-		PoseToSkeletonBoneIndexArray.Empty();
-	}
-
-	FBoneContainer(const TArray<FBoneIndexType>& InRequiredBoneIndexArray, UObject& InAsset)
-		: BoneIndicesArray(InRequiredBoneIndexArray)
-		, Asset(&InAsset)
-		, AssetSkeletalMesh(NULL)
-		, AssetSkeleton(NULL)
-		, RefSkeleton(NULL)
-		, bDisableRetargeting(false)
-		, bUseRAWData(false)
-	{
-		Initialize();
-	}
-
-	/** Initialize BoneContainer to a new Asset, RequiredBonesArray and RefPoseArray. */
-	void InitializeTo(const TArray<FBoneIndexType>& InRequiredBoneIndexArray, UObject& InAsset);
-
-	/** Returns true if FBoneContainer is Valid. Needs an Asset, a RefPoseArray, and a RequiredBonesArray. */
-	const bool IsValid() const
-	{
-		return (Asset.IsValid() && (RefSkeleton != NULL) && (BoneIndicesArray.Num() > 0));
-	}
-
-	/** Get Asset this BoneContainer was made for. Typically a SkeletalMesh, but could also be a USkeleton. */
-	UObject* GetAsset() const
-	{
-		return Asset.Get();
-	}
-
-	/** Get SkeletalMesh Asset this BoneContainer was made for. Could be NULL if Asset is a Skeleton. */
-	USkeletalMesh* GetSkeletalMeshAsset() const
-	{
-		return AssetSkeletalMesh.Get();
-	}
-
-	/** Get Skeleton Asset. Could either be the SkeletalMesh's Skeleton, or the Skeleton this BoneContainer was made for. Is non NULL is BoneContainer is valid. */
-	USkeleton * GetSkeletonAsset() const
-	{
-		return AssetSkeleton.Get();
-	}
-
-	/** Disable Retargeting for debugging. */
-	void SetDisableRetargeting(bool InbDisableRetargeting)
-	{
-		bDisableRetargeting = InbDisableRetargeting;
-	}
-
-	/** True if retargeting is disabled for debugging. */
-	bool GetDisableRetargeting() const
-	{
-		return bDisableRetargeting;
-	}
-
-	/** Ignore compressed data and use RAW data instead, for debugging. */
-	void SetUseRAWData(bool InbUseRAWData)
-	{
-		bUseRAWData = InbUseRAWData;
-	}
-
-	/** True if we're requesting RAW data instead of compressed data. For debugging. */
-	bool ShouldUseRawData() const
-	{
-		return bUseRAWData;
-	}
-
-	/**
-	* returns Required Bone Indices Array
-	*/
-	const TArray<FBoneIndexType>& GetBoneIndicesArray() const
-	{
-		return BoneIndicesArray;
-	}
-
-	/**
-	* returns Bone Switch Array. BitMask for RequiredBoneIndex array.
-	*/
-	const TBitArray<>& GetBoneSwitchArray() const
-	{
-		return BoneSwitchArray;
-	}
-
-	/** Pointer to RefPoseArray for current Asset. */
-	const TArray<FTransform>& GetRefPoseArray() const
-	{
-		return RefSkeleton->GetRefBonePose();
-	}
-
-	/** Access to Asset's RefSkeleton. */
-	const FReferenceSkeleton& GetReferenceSkeleton() const
-	{
-		return *RefSkeleton;
-	}
-
-	/** Number of Bones in RefPose for current asset. This is NOT the number of bones in RequiredBonesArray, but the TOTAL number of bones in the RefPose of the current Asset! */
-	const int32 GetNumBones() const
-	{
-		return RefSkeleton->GetNum();
-	}
-
-	/** Get BoneIndex for BoneName for current Asset. */
-	ENGINE_API int32 GetPoseBoneIndexForBoneName(const FName& BoneName) const;
-
-	/** Get ParentBoneIndex for current Asset. */
-	ENGINE_API int32 GetParentBoneIndex(const int32& BoneIndex) const;
-
-	/** Get Depth between bones for current asset. */
-	int32 GetDepthBetweenBones(const int32& BoneIndex, const int32& ParentBoneIndex) const;
-
-	/** Returns true if bone is child of for current asset. */
-	bool BoneIsChildOf(const int32& BoneIndex, const int32& ParentBoneIndex) const;
-
-	/**
-	* Serializes the bones
-	*
-	* @param Ar - The archive to serialize into.
-	* @param Rect - The bone container to serialize.
-	*
-	* @return Reference to the Archive after serialization.
-	*/
-	friend FArchive& operator<<(FArchive& Ar, FBoneContainer& B)
-	{
-		Ar
-			<< B.BoneIndicesArray
-			<< B.BoneSwitchArray
-			<< B.Asset
-			<< B.AssetSkeletalMesh
-			<< B.AssetSkeleton
-			<< B.SkeletonToPoseBoneIndexArray
-			<< B.PoseToSkeletonBoneIndexArray
-			<< B.bDisableRetargeting
-			<< B.bUseRAWData
-			;
-		return Ar;
-	}
-
-	/**
-	* Returns true of RequiredBonesArray contains this bone index.
-	*/
-	bool Contains(FBoneIndexType NewIndex) const
-	{
-		return BoneSwitchArray[NewIndex];
-	}
-
-	/** Const accessor to GetSkeletonToPoseBoneIndexArray(). */
-	TArray<int32> const & GetSkeletonToPoseBoneIndexArray() const
-	{
-		return SkeletonToPoseBoneIndexArray;
-	}
-
-	/** Const accessor to GetSkeletonToPoseBoneIndexArray(). */
-	TArray<int32> const & GetPoseToSkeletonBoneIndexArray() const
-	{
-		return PoseToSkeletonBoneIndexArray;
-	}
-
-private:
-	/** Initialize FBoneContainer. */
-	ENGINE_API void Initialize();
-
-	/** Cache remapping data if current Asset is a SkeletalMesh, with all compatible Skeletons. */
-	void RemapFromSkelMesh(USkeletalMesh const & SourceSkeletalMesh, USkeleton & TargetSkeleton);
-
-	/** Cache remapping data if current Asset is a Skeleton, with all compatible Skeletons. */
-	void RemapFromSkeleton(USkeleton const & SourceSkeleton);
-};
 
 USTRUCT()
 struct FBoneReference
@@ -785,7 +572,7 @@ public:
 
 	/** Skeleton of this skeletal mesh **/
 	UPROPERTY(Category=Mesh, AssetRegistrySearchable, VisibleAnywhere, BlueprintReadOnly)
-	class USkeleton* Skeleton;
+	USkeleton* Skeleton;
 
 	UPROPERTY(transient, duplicatetransient)
 	FBoxSphereBounds Bounds;
diff --git a/Engine/Source/Runtime/Engine/Classes/GameFramework/NavMovementComponent.h b/Engine/Source/Runtime/Engine/Classes/GameFramework/NavMovementComponent.h
--- a/Engine/Source/Runtime/Engine/Classes/GameFramework/NavMovementComponent.h
+++ b/Engine/Source/Runtime/Engine/Classes/GameFramework/NavMovementComponent.h
@@ -10,6 +10,7 @@
 #include "AI/Navigation/NavAgentInterface.h"
 #include "AI/Navigation/NavigationTypes.h"
 #include "GameFramework/MovementComponent.h"
+#include "Components/PrimitiveComponent.h"
 #include "NavMovementComponent.generated.h"
 
 UCLASS(abstract)
diff --git a/Engine/Source/Runtime/Engine/Classes/GameFramework/WheeledVehicle.h b/Engine/Source/Runtime/Engine/Classes/GameFramework/WheeledVehicle.h
--- a/Engine/Source/Runtime/Engine/Classes/GameFramework/WheeledVehicle.h
+++ b/Engine/Source/Runtime/Engine/Classes/GameFramework/WheeledVehicle.h
@@ -1,11 +1,11 @@
 // Copyright 1998-2014 Epic Games, Inc. All Rights Reserved.
 
 #pragma once
-
-#include "Controller.h"
-
+#include "GameFramework/Pawn.h"
 #include "WheeledVehicle.generated.h"
 
+class FDebugDisplayInfo;
+
 UCLASS(abstract, config=Game, BlueprintType)
 class ENGINE_API AWheeledVehicle : public APawn
 {
diff --git a/Engine/Source/Runtime/Engine/Classes/Vehicles/VehicleWheel.h b/Engine/Source/Runtime/Engine/Classes/Vehicles/VehicleWheel.h
--- a/Engine/Source/Runtime/Engine/Classes/Vehicles/VehicleWheel.h
+++ b/Engine/Source/Runtime/Engine/Classes/Vehicles/VehicleWheel.h
@@ -7,6 +7,15 @@
 #pragma once
 #include "VehicleWheel.generated.h"
 
+class UPhysicalMaterial;
+
+#if WITH_PHYSX
+namespace physx
+{
+	class PxShape;
+}
+#endif // WITH_PHYSX
+
 UCLASS(BlueprintType, Blueprintable)
 class ENGINE_API UVehicleWheel : public UObject
 {
diff --git a/Engine/Source/Runtime/Engine/Classes/Vehicles/WheeledVehicleMovementComponent4W.h b/Engine/Source/Runtime/Engine/Classes/Vehicles/WheeledVehicleMovementComponent4W.h
--- a/Engine/Source/Runtime/Engine/Classes/Vehicles/WheeledVehicleMovementComponent4W.h
+++ b/Engine/Source/Runtime/Engine/Classes/Vehicles/WheeledVehicleMovementComponent4W.h
@@ -5,6 +5,7 @@
  */
 #pragma once
 #include "Vehicles/WheeledVehicleMovementComponent.h"
+#include "Curves/CurveFloat.h"
 #include "WheeledVehicleMovementComponent4W.generated.h"
 
 UENUM()
diff --git a/Engine/Source/Runtime/Engine/Private/AnimationAsset.cpp b/Engine/Source/Runtime/Engine/Private/AnimationAsset.cpp
--- a/Engine/Source/Runtime/Engine/Private/AnimationAsset.cpp
+++ b/Engine/Source/Runtime/Engine/Private/AnimationAsset.cpp
@@ -103,6 +103,35 @@ void UAnimationAsset::ValidateSkeleton()
 	}
 }
 
+//////////////////////////////////////////////////////////////////////////
+// FBoneContainer
+
+FBoneContainer::FBoneContainer()
+: Asset(NULL)
+, AssetSkeletalMesh(NULL)
+, AssetSkeleton(NULL)
+, RefSkeleton(NULL)
+, bDisableRetargeting(false)
+, bUseRAWData(false)
+{
+	BoneIndicesArray.Empty();
+	BoneSwitchArray.Empty();
+	SkeletonToPoseBoneIndexArray.Empty();
+	PoseToSkeletonBoneIndexArray.Empty();
+}
+
+FBoneContainer::FBoneContainer(const TArray<FBoneIndexType>& InRequiredBoneIndexArray, UObject& InAsset)
+: BoneIndicesArray(InRequiredBoneIndexArray)
+, Asset(&InAsset)
+, AssetSkeletalMesh(NULL)
+, AssetSkeleton(NULL)
+, RefSkeleton(NULL)
+, bDisableRetargeting(false)
+, bUseRAWData(false)
+{
+	Initialize();
+}
+
 void FBoneContainer::InitializeTo(const TArray<FBoneIndexType>& InRequiredBoneIndexArray, UObject& InAsset)
 {
 	BoneIndicesArray = InRequiredBoneIndexArray;
diff --git a/Engine/Source/Runtime/Engine/Private/EnginePrivate.h b/Engine/Source/Runtime/Engine/Private/EnginePrivate.h
--- a/Engine/Source/Runtime/Engine/Private/EnginePrivate.h
+++ b/Engine/Source/Runtime/Engine/Private/EnginePrivate.h
@@ -312,7 +312,7 @@
 #include "ComponentReregisterContext.h"	
 #include "DrawDebugHelpers.h"
 #include "UnrealEngine.h"					// Unreal engine helpers.
-#include "Canvas.h"							// Canvas.
+#include "CanvasTypes.h"							// Canvas.
 #include "EngineUtils.h"
 #include "InstancedFoliage.h"				// Instanced foliage.
 #include "UnrealExporter.h"					// Exporter definition.
diff --git a/Engine/Source/Runtime/Engine/Private/UserInterface/CanvasItem.cpp b/Engine/Source/Runtime/Engine/Private/UserInterface/CanvasItem.cpp
--- a/Engine/Source/Runtime/Engine/Private/UserInterface/CanvasItem.cpp
+++ b/Engine/Source/Runtime/Engine/Private/UserInterface/CanvasItem.cpp
@@ -241,6 +241,39 @@ void FCanvasItemTestbed::Draw( class FViewport* Viewport, class FCanvas* Canvas
 #endif // WITH_EDITOR
 
 
+FCanvasTileItem::FCanvasTileItem(const FVector2D& InPosition, const FTexture* InTexture, const FLinearColor& InColor)
+	: FCanvasItem(InPosition)
+	, Z(1.0f)
+	, UV0(0.0f, 0.0f)
+	, UV1(1.0f, 1.0f)
+	, Texture(InTexture)
+	, MaterialRenderProxy(NULL)
+	, Rotation(ForceInitToZero)
+	, PivotPoint(FVector2D::ZeroVector)
+{
+	SetColor(InColor);
+	// Ensure texture is valid.
+	check(InTexture);
+	Size.X = InTexture->GetSizeX();
+	Size.Y = InTexture->GetSizeY();
+}
+
+FCanvasTileItem::FCanvasTileItem(const FVector2D& InPosition, const FTexture* InTexture, const FVector2D& InSize, const FLinearColor& InColor)
+	: FCanvasItem(InPosition)
+	, Size(InSize)
+	, Z(1.0f)
+	, UV0(0.0f, 0.0f)
+	, UV1(1.0f, 1.0f)
+	, Texture(InTexture)
+	, MaterialRenderProxy(NULL)
+	, Rotation(ForceInitToZero)
+	, PivotPoint(FVector2D::ZeroVector)
+{
+	SetColor(InColor);
+	// Ensure texture is valid
+	check(InTexture);
+}
+
 FCanvasTileItem::FCanvasTileItem( const FVector2D& InPosition, const FVector2D& InSize, const FLinearColor& InColor )
 	: FCanvasItem( InPosition )
 	, Size( InSize )
@@ -255,6 +288,69 @@ FCanvasTileItem::FCanvasTileItem( const FVector2D& InPosition, const FVector2D&
 	SetColor( InColor );
 }
 
+FCanvasTileItem::FCanvasTileItem(const FVector2D& InPosition, const FTexture* InTexture, const FVector2D& InUV0, const FVector2D& InUV1, const FLinearColor& InColor)
+	: FCanvasItem(InPosition)
+	, Z(1.0f)
+	, UV0(InUV0)
+	, UV1(InUV1)
+	, MaterialRenderProxy(NULL)
+	, Rotation(ForceInitToZero)
+	, PivotPoint(FVector2D::ZeroVector)
+{
+	SetColor(InColor);
+	// Ensure texture is valid.
+	check(InTexture);
+
+	Size.X = InTexture->GetSizeX();
+	Size.Y = InTexture->GetSizeY();
+}
+
+FCanvasTileItem::FCanvasTileItem(const FVector2D& InPosition, const FTexture* InTexture, const FVector2D& InSize, const FVector2D& InUV0, const FVector2D& InUV1, const FLinearColor& InColor)
+	: FCanvasItem(InPosition)
+	, Size(InSize)
+	, Z(1.0f)
+	, UV0(InUV0)
+	, UV1(InUV1)
+	, Texture(InTexture)
+	, MaterialRenderProxy(NULL)
+	, Rotation(ForceInitToZero)
+	, PivotPoint(FVector2D::ZeroVector)
+{
+	SetColor(InColor);
+	// Ensure texture is valid.
+	check(InTexture != NULL);
+}
+
+FCanvasTileItem::FCanvasTileItem(const FVector2D& InPosition, const FMaterialRenderProxy* InMaterialRenderProxy, const FVector2D& InSize)
+	: FCanvasItem(InPosition)
+	, Size(InSize)
+	, Z(1.0f)
+	, UV0(0.0f, 0.0f)
+	, UV1(1.0f, 1.0f)
+	, Texture(NULL)
+	, MaterialRenderProxy(InMaterialRenderProxy)
+	, Rotation(ForceInitToZero)
+	, PivotPoint(FVector2D::ZeroVector)
+{
+	// Ensure specify Texture or material, but not both.
+	check(InMaterialRenderProxy);
+}
+
+FCanvasTileItem::FCanvasTileItem(const FVector2D& InPosition, const FMaterialRenderProxy* InMaterialRenderProxy, const FVector2D& InSize, const FVector2D& InUV0, const FVector2D& InUV1)
+	: FCanvasItem(InPosition)
+	, Size(InSize)
+	, Z(1.0f)
+	, UV0(InUV0)
+	, UV1(InUV1)
+	, Texture(NULL)
+	, MaterialRenderProxy(InMaterialRenderProxy)
+	, Rotation(ForceInitToZero)
+	, PivotPoint(FVector2D::ZeroVector)
+{
+	// Ensure material proxy is valid.
+	check(InMaterialRenderProxy);
+}
+
 void FCanvasTileItem::Draw( class FCanvas* InCanvas )
 {		
 	// Rotate the canvas if the item has rotation
diff --git a/Engine/Source/Runtime/Engine/Public/BatchedElements.h b/Engine/Source/Runtime/Engine/Public/BatchedElements.h
--- a/Engine/Source/Runtime/Engine/Public/BatchedElements.h
+++ b/Engine/Source/Runtime/Engine/Public/BatchedElements.h
@@ -9,28 +9,8 @@
 
 #include "StaticBoundShaderState.h"
 #include "HitProxies.h"
+#include "SceneTypes.h"
 
-/** Blend modes supported for simple element rendering */
-enum ESimpleElementBlendMode
-{
-	SE_BLEND_Opaque=0,
-	SE_BLEND_Masked,
-	SE_BLEND_Translucent,
-	SE_BLEND_Additive,
-	SE_BLEND_Modulate,
-	SE_BLEND_MaskedDistanceField,
-	SE_BLEND_MaskedDistanceFieldShadowed,
-	SE_BLEND_TranslucentDistanceField,
-	SE_BLEND_TranslucentDistanceFieldShadowed,
-	SE_BLEND_AlphaComposite,
-	// Like SE_BLEND_Translucent, but modifies destination alpha
-	SE_BLEND_AlphaBlend,
-
-	SE_BLEND_RGBA_MASK_START,
-	SE_BLEND_RGBA_MASK_END = SE_BLEND_RGBA_MASK_START+31, //Using 5bit bit-field for red, green, blue, alpha and desaturation
-
-	SE_BLEND_MAX
-};
 
 /** The type used to store batched line vertices. */
 struct FSimpleElementVertex
diff --git a/Engine/Source/Runtime/Engine/Public/BoneContainer.h b/Engine/Source/Runtime/Engine/Public/BoneContainer.h
new file mode 100644
--- /dev/null
+++ b/Engine/Source/Runtime/Engine/Public/BoneContainer.h
@@ -0,0 +1,199 @@
+// Copyright 1998-2014 Epic Games, Inc. All Rights Reserved.
+
+#pragma once
+
+class USkeleton;
+
+/**
+* This is a native transient structure.
+* Contains:
+* - BoneIndicesArray: Array of RequiredBoneIndices for Current Asset. In increasing order. Mapping to current Array of Transforms (Pose).
+* - BoneSwitchArray: Size of current Skeleton. true if Bone is contained in RequiredBones array, false otherwise.
+**/
+struct FBoneContainer
+{
+private:
+	/** Array of RequiredBonesIndices. In increasing order. */
+	TArray<FBoneIndexType>	BoneIndicesArray;
+	/** Array sized by Current RefPose. true if Bone is contained in RequiredBones array, false otherwise. */
+	TBitArray<>				BoneSwitchArray;
+
+	/** Asset BoneIndicesArray was made for. Typically a SkeletalMesh. */
+	TWeakObjectPtr<UObject>	Asset;
+	/** If Asset is a SkeletalMesh, this will be a pointer to it. Can be NULL if Asset is a USkeleton. */
+	TWeakObjectPtr<USkeletalMesh> AssetSkeletalMesh;
+	/** If Asset is a Skeleton that will be it. If Asset is a SkeletalMesh, that will be its Skeleton. */
+	TWeakObjectPtr<USkeleton> AssetSkeleton;
+
+	/** Pointer to RefSkeleton of Asset. */
+	const FReferenceSkeleton* RefSkeleton;
+
+	/** Mapping table between Skeleton Bone Indices and Pose Bone Indices. */
+	TArray<int32> SkeletonToPoseBoneIndexArray;
+
+	/** Mapping table between Pose Bone Indices and Skeleton Bone Indices. */
+	TArray<int32> PoseToSkeletonBoneIndexArray;
+
+	/** For debugging. */
+	/** Disable Retargeting. Extract animation, but do not retarget it. */
+	bool bDisableRetargeting;
+	/** Disable animation compression, use RAW data instead. */
+	bool bUseRAWData;
+
+public:
+
+	FBoneContainer();
+
+	FBoneContainer(const TArray<FBoneIndexType>& InRequiredBoneIndexArray, UObject& InAsset);
+
+	/** Initialize BoneContainer to a new Asset, RequiredBonesArray and RefPoseArray. */
+	void InitializeTo(const TArray<FBoneIndexType>& InRequiredBoneIndexArray, UObject& InAsset);
+
+	/** Returns true if FBoneContainer is Valid. Needs an Asset, a RefPoseArray, and a RequiredBonesArray. */
+	const bool IsValid() const
+	{
+		return (Asset.IsValid() && (RefSkeleton != NULL) && (BoneIndicesArray.Num() > 0));
+	}
+
+	/** Get Asset this BoneContainer was made for. Typically a SkeletalMesh, but could also be a USkeleton. */
+	UObject* GetAsset() const
+	{
+		return Asset.Get();
+	}
+
+	/** Get SkeletalMesh Asset this BoneContainer was made for. Could be NULL if Asset is a Skeleton. */
+	USkeletalMesh* GetSkeletalMeshAsset() const
+	{
+		return AssetSkeletalMesh.Get();
+	}
+
+	/** Get Skeleton Asset. Could either be the SkeletalMesh's Skeleton, or the Skeleton this BoneContainer was made for. Is non NULL is BoneContainer is valid. */
+	USkeleton * GetSkeletonAsset() const
+	{
+		return AssetSkeleton.Get();
+	}
+
+	/** Disable Retargeting for debugging. */
+	void SetDisableRetargeting(bool InbDisableRetargeting)
+	{
+		bDisableRetargeting = InbDisableRetargeting;
+	}
+
+	/** True if retargeting is disabled for debugging. */
+	bool GetDisableRetargeting() const
+	{
+		return bDisableRetargeting;
+	}
+
+	/** Ignore compressed data and use RAW data instead, for debugging. */
+	void SetUseRAWData(bool InbUseRAWData)
+	{
+		bUseRAWData = InbUseRAWData;
+	}
+
+	/** True if we're requesting RAW data instead of compressed data. For debugging. */
+	bool ShouldUseRawData() const
+	{
+		return bUseRAWData;
+	}
+
+	/**
+	* returns Required Bone Indices Array
+	*/
+	const TArray<FBoneIndexType>& GetBoneIndicesArray() const
+	{
+		return BoneIndicesArray;
+	}
+
+	/**
+	* returns Bone Switch Array. BitMask for RequiredBoneIndex array.
+	*/
+	const TBitArray<>& GetBoneSwitchArray() const
+	{
+		return BoneSwitchArray;
+	}
+
+	/** Pointer to RefPoseArray for current Asset. */
+	const TArray<FTransform>& GetRefPoseArray() const
+	{
+		return RefSkeleton->GetRefBonePose();
+	}
+
+	/** Access to Asset's RefSkeleton. */
+	const FReferenceSkeleton& GetReferenceSkeleton() const
+	{
+		return *RefSkeleton;
+	}
+
+	/** Number of Bones in RefPose for current asset. This is NOT the number of bones in RequiredBonesArray, but the TOTAL number of bones in the RefPose of the current Asset! */
+	const int32 GetNumBones() const
+	{
+		return RefSkeleton->GetNum();
+	}
+
+	/** Get BoneIndex for BoneName for current Asset. */
+	ENGINE_API int32 GetPoseBoneIndexForBoneName(const FName& BoneName) const;
+
+	/** Get ParentBoneIndex for current Asset. */
+	ENGINE_API int32 GetParentBoneIndex(const int32& BoneIndex) const;
+
+	/** Get Depth between bones for current asset. */
+	int32 GetDepthBetweenBones(const int32& BoneIndex, const int32& ParentBoneIndex) const;
+
+	/** Returns true if bone is child of for current asset. */
+	bool BoneIsChildOf(const int32& BoneIndex, const int32& ParentBoneIndex) const;
+
+	/**
+	* Serializes the bones
+	*
+	* @param Ar - The archive to serialize into.
+	* @param Rect - The bone container to serialize.
+	*
+	* @return Reference to the Archive after serialization.
+	*/
+	friend FArchive& operator<<(FArchive& Ar, FBoneContainer& B)
+	{
+		Ar
+			<< B.BoneIndicesArray
+			<< B.BoneSwitchArray
+			<< B.Asset
+			<< B.AssetSkeletalMesh
+			<< B.AssetSkeleton
+			<< B.SkeletonToPoseBoneIndexArray
+			<< B.PoseToSkeletonBoneIndexArray
+			<< B.bDisableRetargeting
+			<< B.bUseRAWData
+			;
+		return Ar;
+	}
+
+	/**
+	* Returns true of RequiredBonesArray contains this bone index.
+	*/
+	bool Contains(FBoneIndexType NewIndex) const
+	{
+		return BoneSwitchArray[NewIndex];
+	}
+
+	/** Const accessor to GetSkeletonToPoseBoneIndexArray(). */
+	TArray<int32> const & GetSkeletonToPoseBoneIndexArray() const
+	{
+		return SkeletonToPoseBoneIndexArray;
+	}
+
+	/** Const accessor to GetSkeletonToPoseBoneIndexArray(). */
+	TArray<int32> const & GetPoseToSkeletonBoneIndexArray() const
+	{
+		return PoseToSkeletonBoneIndexArray;
+	}
+
+private:
+	/** Initialize FBoneContainer. */
+	ENGINE_API void Initialize();
+
+	/** Cache remapping data if current Asset is a SkeletalMesh, with all compatible Skeletons. */
+	void RemapFromSkelMesh(USkeletalMesh const & SourceSkeletalMesh, USkeleton & TargetSkeleton);
+
+	/** Cache remapping data if current Asset is a Skeleton, with all compatible Skeletons. */
+	void RemapFromSkeleton(USkeleton const & SourceSkeleton);
+};
\ No newline at end of file
diff --git a/Engine/Source/Runtime/Engine/Public/Canvas.h b/Engine/Source/Runtime/Engine/Public/Canvas.h
deleted file mode 100644
diff --git a/Engine/Source/Runtime/Engine/Public/CanvasItem.h b/Engine/Source/Runtime/Engine/Public/CanvasItem.h
--- a/Engine/Source/Runtime/Engine/Public/CanvasItem.h
+++ b/Engine/Source/Runtime/Engine/Public/CanvasItem.h
@@ -1,14 +1,19 @@
 // Copyright 1998-2014 Epic Games, Inc. All Rights Reserved.
 
 /*=============================================================================
-	FCanvasItem.h: Unreal canvas item definitions
+	CanvasItem.h: Unreal canvas item definitions
 =============================================================================*/
 
 #pragma once
 
-#include "BatchedElements.h"
+#include "Engine/EngineTypes.h"
+#include "SceneTypes.h"
 
+class FCanvas;
+class FBatchedElementParameters;
+class FBatchedElements;
 class FMaterialRenderProxy;
+class FTexture;
 
 class FCanvasItem
 {
@@ -25,7 +30,7 @@ public:
 		, BatchedElementParameters( NULL )
 		, Color( FLinearColor::White ) {};
 
-	virtual void Draw( class FCanvas* InCanvas ) = 0;
+	virtual void Draw( FCanvas* InCanvas ) = 0;
 
 	/* 
 	 * Draw this item (this will affect the items position for future draw calls that do no specify a position)
@@ -33,7 +38,7 @@ public:
 	 * @param	InCanvas		Canvas on which to draw
 	 * @param	InPosition		Draw position - this will not preserve the items position
 	 */
-	virtual void Draw( class FCanvas* InCanvas, const FVector2D& InPosition )
+	virtual void Draw( FCanvas* InCanvas, const FVector2D& InPosition )
 	{
 		Position = InPosition;
 		Draw( InCanvas );
@@ -46,7 +51,7 @@ public:
 	 * @param	X				X Draw position 
 	 * @param	Y				Y Draw position
 	 */
-	virtual void Draw( class FCanvas* InCanvas, float X, float Y )
+	virtual void Draw( FCanvas* InCanvas, float X, float Y )
 	{
 		Position.X = X;
 		Position.Y = Y;
@@ -85,22 +90,7 @@ public:
 	 * @param	InPosition		Draw position
 	 * @param	InTexture		The texture
 	 */
-	FCanvasTileItem( const FVector2D& InPosition, const FTexture* InTexture, const FLinearColor& InColor )
-		: FCanvasItem( InPosition )
-		, Z( 1.0f )
-		, UV0( 0.0f, 0.0f )
-		, UV1( 1.0f, 1.0f )
-		, Texture( InTexture )
-		, MaterialRenderProxy( NULL )
-		, Rotation( ForceInitToZero )
-		, PivotPoint( FVector2D::ZeroVector )
-	{
-		SetColor( InColor );
-		// Ensure texture is valid.
-		check( InTexture );
-		Size.X = InTexture->GetSizeX();
-		Size.Y = InTexture->GetSizeY();
-	}
+	FCanvasTileItem( const FVector2D& InPosition, const FTexture* InTexture, const FLinearColor& InColor );
 
 	/* 
 	 * Tile item with texture using given size. 
@@ -109,21 +99,7 @@ public:
 	 * @param	InTexture		The texture
 	 * @param	InSize			The size to render
 	 */
-	FCanvasTileItem( const FVector2D& InPosition, const FTexture* InTexture, const FVector2D& InSize, const FLinearColor& InColor )
-		: FCanvasItem( InPosition )
-		, Size( InSize )
-		, Z( 1.0f )
-		, UV0( 0.0f, 0.0f )
-		, UV1( 1.0f, 1.0f )
-		, Texture( InTexture )
-		, MaterialRenderProxy( NULL )
-		, Rotation( ForceInitToZero )
-		, PivotPoint( FVector2D::ZeroVector )
-	{
-		SetColor( InColor );
-		// Ensure texture is valid
-		check( InTexture );
-	}
+	FCanvasTileItem( const FVector2D& InPosition, const FTexture* InTexture, const FVector2D& InSize, const FLinearColor& InColor );
 
 	/* 
 	 * Tile item which uses the default white texture using given size. 
@@ -141,22 +117,7 @@ public:
 	 * @param	InUV0			UV coordinates (Normalized Top/Left)
 	 * @param	InUV1			UV coordinates (Normalized Bottom/Right)
 	 */
-	FCanvasTileItem( const FVector2D& InPosition, const FTexture* InTexture, const FVector2D& InUV0, const FVector2D& InUV1, const FLinearColor& InColor )
-		: FCanvasItem( InPosition )	 
-		, Z( 1.0f )
-		, UV0( InUV0 )
-		, UV1( InUV1 )
-		, MaterialRenderProxy( NULL )
-		, Rotation( ForceInitToZero )
-		, PivotPoint( FVector2D::ZeroVector )
-	{
-		SetColor( InColor );
-		// Ensure texture is valid.
-		check( InTexture );
-		
-		Size.X = InTexture->GetSizeX();
-		Size.Y = InTexture->GetSizeY();
-	}
+	FCanvasTileItem( const FVector2D& InPosition, const FTexture* InTexture, const FVector2D& InUV0, const FVector2D& InUV1, const FLinearColor& InColor );
 
 	/* 
 	 * Tile item with texture using given size and specific UVs.
@@ -167,21 +128,7 @@ public:
 	 * @param	InUV0			UV coordinates (Normalized Top/Left)
 	 * @param	InUV1			UV coordinates (Normalized Bottom/Right)
 	 */
-	FCanvasTileItem( const FVector2D& InPosition, const FTexture* InTexture, const FVector2D& InSize, const FVector2D& InUV0, const FVector2D& InUV1, const FLinearColor& InColor )
-		: FCanvasItem( InPosition )
-		, Size( InSize )
-		, Z( 1.0f )
-		, UV0( InUV0 )
-		, UV1( InUV1 )
-		, Texture( InTexture )
-		, MaterialRenderProxy( NULL )
-		, Rotation( ForceInitToZero )
-		, PivotPoint( FVector2D::ZeroVector )
-	{
-		SetColor( InColor );
-		// Ensure texture is valid.
-		check( InTexture != NULL);
-	}
+	FCanvasTileItem( const FVector2D& InPosition, const FTexture* InTexture, const FVector2D& InSize, const FVector2D& InUV0, const FVector2D& InUV1, const FLinearColor& InColor );
 
 	/* 
 	 * Tile item with FMaterialRenderProxy using given size. 
@@ -190,21 +137,7 @@ public:
 	 * @param	InMaterialRenderProxy	Material proxy for rendering
 	 * @param	InSize					The size to render
 	 */
-	FCanvasTileItem( const FVector2D& InPosition, const FMaterialRenderProxy* InMaterialRenderProxy, const FVector2D& InSize )
-		: FCanvasItem( InPosition )
-		, Size( InSize )
-		, Z( 1.0f )
-		, UV0( 0.0f, 0.0f )
-		, UV1( 1.0f, 1.0f )
-		, Texture( NULL )
-		, MaterialRenderProxy( InMaterialRenderProxy )
-		, Rotation( ForceInitToZero )
-		, PivotPoint( FVector2D::ZeroVector )
-	{
-		// Ensure specify Texture or material, but not both.
-		check( InMaterialRenderProxy );
-	}
-
+	FCanvasTileItem( const FVector2D& InPosition, const FMaterialRenderProxy* InMaterialRenderProxy, const FVector2D& InSize );
 
 	/* 
 	 * Tile item with FMaterialRenderProxy using given size and UVs.
@@ -215,28 +148,14 @@ public:
 	 * @param	InUV0					UV coordinates (Normalized Top/Left)
 	 * @param	InUV1					UV coordinates (Normalized Bottom/Right)
 	 */
-	FCanvasTileItem( const FVector2D& InPosition, const FMaterialRenderProxy* InMaterialRenderProxy, const FVector2D& InSize, const FVector2D& InUV0, const FVector2D& InUV1 )
-		: FCanvasItem( InPosition )
-		, Size( InSize )
-		, Z( 1.0f )
-		, UV0( InUV0 )
-		, UV1( InUV1 )
-		, Texture( NULL )
-		, MaterialRenderProxy( InMaterialRenderProxy )
-		, Rotation( ForceInitToZero )
-		, PivotPoint( FVector2D::ZeroVector )
-	{
-		// Ensure material proxy is valid.
-		check( InMaterialRenderProxy );
-	}
-
+	FCanvasTileItem( const FVector2D& InPosition, const FMaterialRenderProxy* InMaterialRenderProxy, const FVector2D& InSize, const FVector2D& InUV0, const FVector2D& InUV1 );
 
 	/* 
 	 * Draw the item at the given coordinates.
 	 *
 	 * @param	InPosition		Draw position.
 	 */
-	virtual void Draw( class FCanvas* InCanvas ) override;
+	virtual void Draw( FCanvas* InCanvas ) override;
 
 	/* Expose the functions defined in the base class. */
 	using FCanvasItem::Draw;
@@ -267,7 +186,7 @@ public:
 
 private:
 	/* Render when we have a material proxy. */
-	void RenderMaterialTile( class FCanvas* InCanvas, const FVector2D& InPosition );
+	void RenderMaterialTile( FCanvas* InCanvas, const FVector2D& InPosition );
 };
 
 /* Resizable 3x3 border item. */
@@ -314,7 +233,7 @@ public:
 	 *
 	 * @param	InPosition		Draw position.
 	 */
-	virtual void Draw( class FCanvas* InCanvas ) override;
+	virtual void Draw( FCanvas* InCanvas ) override;
 
 	/* Expose the functions defined in the base class. */
 	using FCanvasItem::Draw;
@@ -423,7 +342,7 @@ public:
 	 *
 	 * @param	InCanvas		Canvas on which to draw
 	 */
-	virtual void Draw( class FCanvas* InCanvas ) override;
+	virtual void Draw( FCanvas* InCanvas ) override;
 
 	/* Expose the functions defined in the base class. */
 	using FCanvasItem::Draw;
@@ -481,7 +400,7 @@ protected:
 	 *
 	 * In a method to make it simpler to do effects like shadow, outline
 	 */
-	void DrawStringInternal( class FCanvas* InCanvas, const FVector2D& DrawPos, const FLinearColor& DrawColor );
+	void DrawStringInternal( FCanvas* InCanvas, const FVector2D& DrawPos, const FLinearColor& DrawColor );
 
 	/* 
 	 * These are used bye the DrawStringInternal function. 
@@ -560,7 +479,7 @@ public:
 	 *
 	 * @param	InCanvas		Canvas on which to draw
 	 */
-	virtual void Draw( class FCanvas* InCanvas ) override;
+	virtual void Draw( FCanvas* InCanvas ) override;
 	
 	/* 
 	 * Draw line at the given coordinates.
@@ -568,7 +487,7 @@ public:
 	 * @param	InCanvas		Canvas on which to draw
 	 * @param	InPosition		Draw Start position
 	 */
-	virtual void Draw( class FCanvas* InCanvas, const FVector2D& InPosition ) override
+	virtual void Draw( FCanvas* InCanvas, const FVector2D& InPosition ) override
 	{
 		Origin.X = InPosition.X;
 		Origin.Y = InPosition.Y;
@@ -581,7 +500,7 @@ public:
 	 * @param	InStartPos		Line start position
 	 * @param	InEndPos		Line end position
 	 */
-	virtual void Draw( class FCanvas* InCanvas, const FVector2D& InStartPos, const FVector2D& InEndPos )
+	virtual void Draw( FCanvas* InCanvas, const FVector2D& InStartPos, const FVector2D& InEndPos )
 	{
 		Origin.X = InStartPos.X;
 		Origin.Y = InStartPos.Y;
@@ -596,7 +515,7 @@ public:
 	 * @param	X				X Draw position 
 	 * @param	Y				Y Draw position
 	 */
-	virtual void Draw( class FCanvas* InCanvas, float InX, float InY ) override
+	virtual void Draw( FCanvas* InCanvas, float InX, float InY ) override
 	{
 		Origin.X = InX;
 		Origin.Y = InY;
@@ -609,7 +528,7 @@ public:
 	 * @param	InCanvas		Canvas on which to draw
 	 * @param	InPosition		Draw position
 	 */
-	virtual void Draw( class FCanvas* InCanvas, const FVector& InPosition )
+	virtual void Draw( FCanvas* InCanvas, const FVector& InPosition )
 	{
 		Origin = InPosition;
 		Draw( InCanvas );
@@ -623,7 +542,7 @@ public:
 	 * @param	Y				Y Draw position
 	 * @param	Z				Z Draw position
 	 */
-	virtual void Draw( class FCanvas* InCanvas, float X, float Y, float Z )
+	virtual void Draw( FCanvas* InCanvas, float X, float Y, float Z )
 	{
 		Origin.X = X;
 		Origin.Y = Y;
@@ -660,7 +579,7 @@ public:
 		, Size( InSize )
 		, LineThickness( 0.0f )	{};
 
-	virtual void Draw( class FCanvas* InCanvas ) override;
+	virtual void Draw( FCanvas* InCanvas ) override;
 
 	/* Expose the functions defined in the base class. */
 	using FCanvasItem::Draw;
@@ -780,7 +699,7 @@ public:
 		TriangleList[0].V2_Pos = InPointC;
 	}
 
-	virtual void Draw( class FCanvas* InCanvas ) override;
+	virtual void Draw( FCanvas* InCanvas ) override;
 		
 	/* Expose the functions defined in the base class. */
 	using FCanvasItem::Draw;
@@ -835,7 +754,7 @@ public:
 	 {
 		 delete TriListItem;
 	 }
-	 virtual void Draw( class FCanvas* InCanvas ) override;
+	 virtual void Draw( FCanvas* InCanvas ) override;
 
 	 /* 	 
 	 * Regenerates the tri list for the object with a new central point and radius
@@ -883,7 +802,7 @@ class ENGINE_API FCanvasItemTestbed
 {
 public:
 	FCanvasItemTestbed(){};
-	void  Draw( class FViewport* Viewport, class FCanvas* Canvas );
+	void  Draw( class FViewport* Viewport, FCanvas* Canvas );
 	struct LineVars
 	{
 		LineVars()
diff --git a/Engine/Source/Runtime/Engine/Public/CanvasTypes.h b/Engine/Source/Runtime/Engine/Public/CanvasTypes.h
new file mode 100644
--- /dev/null
+++ b/Engine/Source/Runtime/Engine/Public/CanvasTypes.h
@@ -0,0 +1,825 @@
+// Copyright 1998-2014 Epic Games, Inc. All Rights Reserved.
+
+/*=============================================================================
+	Canvas.h: Unreal canvas definition.
+=============================================================================*/
+
+#pragma once
+
+#include "Engine/EngineTypes.h"
+#include "BatchedElements.h"
+#include "CanvasItem.h"
+
+class FMaterialRenderProxy;
+class FTexture;
+class FRenderTarget;
+class FHitProxyConsumer;
+class FRHICommandListImmediate;
+
+/**
+ * Encapsulates the canvas state.
+ */
+class FCanvas
+{
+public:	
+
+	/** 
+	 * Enum that describes what type of element we are currently batching.
+	 */
+	enum EElementType
+	{
+		ET_Line,
+		ET_Triangle,
+		ET_MAX
+	};
+
+	/**
+	 * Enum for canvas features that are allowed
+	 **/
+	enum ECanvasAllowModes
+	{
+		// flushing and rendering
+		Allow_Flush			= 1<<0,
+		// delete the render batches when rendering
+		Allow_DeleteOnRender= 1<<1
+	};
+
+	/** 
+	* Constructor.
+	*/
+	ENGINE_API FCanvas(FRenderTarget* InRenderTarget,FHitProxyConsumer* InHitProxyConsumer, UWorld* InWorld);
+
+	/** 
+	* Constructor. For situations where a world is not available, but time information is
+	*/
+	ENGINE_API FCanvas(FRenderTarget* InRenderTarget,FHitProxyConsumer* InHitProxyConsumer, float InRealTime, float InWorldTime, float InWorldDeltaTime);
+
+	/** 
+	* Destructor.
+	*/
+	ENGINE_API ~FCanvas();
+
+	
+	ENGINE_API static ESimpleElementBlendMode BlendToSimpleElementBlend(EBlendMode BlendMode);
+	/**
+	* Returns a FBatchedElements pointer to be used for adding vertices and primitives for rendering.
+	* Adds a new render item to the sort element entry based on the current sort key.
+	*
+	* @param InElementType - Type of element we are going to draw.
+	* @param InBatchedElementParameters - Parameters for this element
+	* @param InTexture - New texture that will be set.
+	* @param InBlendMode - New blendmode that will be set.
+	* @param GlowInfo - info for optional glow effect when using depth field rendering
+	* @return Returns a pointer to a FBatchedElements object.
+	*/
+	ENGINE_API FBatchedElements* GetBatchedElements(EElementType InElementType, FBatchedElementParameters* InBatchedElementParameters=NULL, const FTexture* Texture=NULL, ESimpleElementBlendMode BlendMode=SE_BLEND_MAX,const FDepthFieldGlowInfo& GlowInfo = FDepthFieldGlowInfo());
+	
+	/**
+	* Generates a new FCanvasTileRendererItem for the current sortkey and adds it to the sortelement list of itmes to render
+	*/
+	ENGINE_API void AddTileRenderItem(float X,float Y,float SizeX,float SizeY,float U,float V,float SizeU,float SizeV,const FMaterialRenderProxy* MaterialRenderProxy,FHitProxyId HitProxyId,bool bFreezeTime);
+
+	/** 
+	* Sends a message to the rendering thread to draw the batched elements. 
+	* @param RHICmdList - command list to use
+	* @param bForce - force the flush even if Allow_Flush is not enabled
+	*/
+	ENGINE_API void Flush_RenderThread(FRHICommandListImmediate& RHICmdList, bool bForce = false);
+
+	/**
+	* Sends a message to the rendering thread to draw the batched elements.
+	* @param bForce - force the flush even if Allow_Flush is not enabled
+	*/
+	ENGINE_API void Flush_GameThread(bool bForce = false);
+
+	/**
+	 * Pushes a transform onto the canvas's transform stack, multiplying it with the current top of the stack.
+	 * @param Transform - The transform to push onto the stack.
+	 */
+	ENGINE_API void PushRelativeTransform(const FMatrix& Transform);
+
+	/**
+	 * Pushes a transform onto the canvas's transform stack.
+	 * @param Transform - The transform to push onto the stack.
+	 */
+	ENGINE_API void PushAbsoluteTransform(const FMatrix& Transform);
+
+	/**
+	 * Removes the top transform from the canvas's transform stack.
+	 */
+	ENGINE_API void PopTransform();
+
+	/**
+	* Replace the base (ie. TransformStack(0)) transform for the canvas with the given matrix
+	*
+	* @param Transform - The transform to use for the base
+	*/
+	ENGINE_API void SetBaseTransform(const FMatrix& Transform);
+
+	/**
+	* Generate a 2D projection for the canvas. Use this if you only want to transform in 2D on the XY plane
+	*
+	* @param ViewSizeX - Viewport width
+	* @param ViewSizeY - Viewport height
+	* @return Matrix for canvas projection
+	*/
+	ENGINE_API static FMatrix CalcBaseTransform2D(uint32 ViewSizeX, uint32 ViewSizeY);
+
+	/**
+	* Generate a 3D projection for the canvas. Use this if you want to transform in 3D 
+	*
+	* @param ViewSizeX - Viewport width
+	* @param ViewSizeY - Viewport height
+	* @param fFOV - Field of view for the projection
+	* @param NearPlane - Distance to the near clip plane
+	* @return Matrix for canvas projection
+	*/
+	ENGINE_API static FMatrix CalcBaseTransform3D(uint32 ViewSizeX, uint32 ViewSizeY, float fFOV, float NearPlane);
+	
+	/**
+	* Generate a view matrix for the canvas. Used for CalcBaseTransform3D
+	*
+	* @param ViewSizeX - Viewport width
+	* @param ViewSizeY - Viewport height
+	* @param fFOV - Field of view for the projection
+	* @return Matrix for canvas view orientation
+	*/
+	ENGINE_API static FMatrix CalcViewMatrix(uint32 ViewSizeX, uint32 ViewSizeY, float fFOV);
+	
+	/**
+	* Generate a projection matrix for the canvas. Used for CalcBaseTransform3D
+	*
+	* @param ViewSizeX - Viewport width
+	* @param ViewSizeY - Viewport height
+	* @param fFOV - Field of view for the projection
+	* @param NearPlane - Distance to the near clip plane
+	* @return Matrix for canvas projection
+	*/
+	ENGINE_API static FMatrix CalcProjectionMatrix(uint32 ViewSizeX, uint32 ViewSizeY, float fFOV, float NearPlane);
+
+	/**
+	* Get the current top-most transform entry without the canvas projection
+	* @return matrix from transform stack. 
+	*/
+	ENGINE_API FMatrix GetTransform() const
+	{ 
+		return TransformStack.Top().GetMatrix() * TransformStack[0].GetMatrix().Inverse(); 
+	}
+
+	/** 
+	* Get the bottom-most element of the transform stack. 
+	* @return matrix from transform stack. 
+	*/
+	ENGINE_API const FMatrix& GetBottomTransform() const
+	{ 
+		return TransformStack[0].GetMatrix(); 
+	}
+
+	/**
+	* Get the current top-most transform entry 
+	* @return matrix from transform stack. 
+	*/
+	ENGINE_API const FMatrix& GetFullTransform() const
+	{ 
+		return TransformStack.Top().GetMatrix(); 
+	}
+
+	/**
+	 * Copy the conents of the TransformStack from an existing canvas
+	 *
+	 * @param Copy	canvas to copy from
+	 **/
+	ENGINE_API void CopyTransformStack(const FCanvas& Copy);
+
+	/**
+	 * Sets the render target which will be used for subsequent canvas primitives.
+	 */
+	ENGINE_API void SetRenderTarget_GameThread(FRenderTarget* NewRenderTarget);
+
+	/**
+	* Get the current render target for the canvas
+	*/	
+	ENGINE_API FORCEINLINE FRenderTarget* GetRenderTarget() const
+	{ 
+		return RenderTarget; 
+	}
+
+	/**
+	 * Sets a rect that should be used to offset rendering into the viewport render target
+	 * If not set the canvas will render to the full target
+	 *
+ 	 * @param ViewRect The rect to use
+	 */
+	ENGINE_API void SetRenderTargetRect( const FIntRect& ViewRect );
+
+	/**
+	* Marks render target as dirty so that it will be resolved to texture
+	*/
+	void SetRenderTargetDirty(bool bDirty) 
+	{ 
+		bRenderTargetDirty = bDirty; 
+	}
+
+	/**
+	* Sets the hit proxy which will be used for subsequent canvas primitives.
+	*/ 
+	ENGINE_API void SetHitProxy(HHitProxy* HitProxy);
+
+	// HitProxy Accessors.	
+
+	FHitProxyId GetHitProxyId() const { return CurrentHitProxy ? CurrentHitProxy->Id : FHitProxyId(); }
+	FHitProxyConsumer* GetHitProxyConsumer() const { return HitProxyConsumer; }
+	bool IsHitTesting() const { return HitProxyConsumer != NULL; }
+	
+	/**
+	* Push sort key onto the stack. Rendering is done with the current sort key stack entry.
+	*
+	* @param InSortKey - key value to push onto the stack
+	*/
+	void PushDepthSortKey(int32 InSortKey)
+	{
+		DepthSortKeyStack.Push(InSortKey);
+	};
+
+	/**
+	* Pop sort key off of the stack.
+	*
+	* @return top entry of the sort key stack
+	*/
+	int32 PopDepthSortKey()
+	{
+		int32 Result = 0;
+		if( DepthSortKeyStack.Num() > 0 )
+		{
+			Result = DepthSortKeyStack.Pop();
+		}
+		else
+		{
+			// should always have one entry
+			PushDepthSortKey(0);
+		}
+		return Result;		
+	};
+
+	/**
+	* Return top sort key of the stack.
+	*
+	* @return top entry of the sort key stack
+	*/
+	int32 TopDepthSortKey()
+	{
+		checkSlow(DepthSortKeyStack.Num() > 0);
+		return DepthSortKeyStack.Top();
+	}
+
+	/**
+	 * Toggle allowed canvas modes
+	 *
+	 * @param InAllowedModes	New modes to set
+	 */
+	void SetAllowedModes(uint32 InAllowedModes)
+	{
+		AllowedModes = InAllowedModes;
+	}
+	/**
+	 * Accessor for allowed canvas modes
+	 *
+	 * @return current allowed modes
+	 */
+	uint32 GetAllowedModes() const
+	{
+		return AllowedModes;
+	}
+
+	/**
+	 * Determine if the canvas has dirty batches that need to be rendered
+	 *
+	 * @return true if the canvas has any element to render
+	 */
+	bool HasBatchesToRender() const;
+
+public:
+	float AlphaModulate;
+
+	/** Entry for the transform stack which stores a matrix and its CRC for faster comparisons */
+	class FTransformEntry
+	{
+	public:
+		FTransformEntry(const FMatrix& InMatrix)
+			:	Matrix(InMatrix)
+		{			
+			MatrixCRC = FCrc::MemCrc_DEPRECATED(&Matrix,sizeof(FMatrix));
+		}
+		FORCEINLINE void SetMatrix(const FMatrix& InMatrix)
+		{
+			Matrix = InMatrix;
+			MatrixCRC = FCrc::MemCrc_DEPRECATED(&Matrix,sizeof(FMatrix));
+		}
+		FORCEINLINE const FMatrix& GetMatrix() const
+		{
+			return Matrix;
+		}
+		FORCEINLINE uint32 GetMatrixCRC() const
+		{
+			return MatrixCRC;
+		}
+	private:
+		FMatrix Matrix;
+		uint32 MatrixCRC;
+	};
+
+	/** returns the transform stack */	
+	FORCEINLINE const TArray<FTransformEntry>& GetTransformStack() const
+	{
+		return TransformStack;
+	}
+	FORCEINLINE const FIntRect& GetViewRect() const
+	{
+		return ViewRect;
+	}
+
+	FORCEINLINE void SetScaledToRenderTarget(bool scale = true)
+	{
+		bScaledToRenderTarget = scale;
+	}
+	FORCEINLINE bool IsScaledToRenderTarget() const { return bScaledToRenderTarget; }
+private:
+	/** Stack of SortKeys. All rendering is done using the top most sort key */
+	TArray<int32> DepthSortKeyStack;	
+	/** Stack of matrices. Bottom most entry is the canvas projection */
+	TArray<FTransformEntry> TransformStack;	
+	/** View rect for the render target */
+	FIntRect ViewRect;
+	/** Current render target used by the canvas */
+	FRenderTarget* RenderTarget;
+	/** Current hit proxy consumer */
+	FHitProxyConsumer* HitProxyConsumer;
+	/** Current hit proxy object */
+	TRefCountPtr<HHitProxy> CurrentHitProxy;
+	/** Toggles for various canvas rendering functionality **/
+	uint32 AllowedModes;
+	/** true if the render target has been rendered to since last calling SetRenderTarget() */
+	bool bRenderTargetDirty;	
+	/** Current real time in seconds */
+	float CurrentRealTime;
+	/** Current world time in seconds */
+	float CurrentWorldTime;
+	/** Current world time in seconds */
+	float CurrentDeltaWorldTime;
+	/** true, if Canvas should be scaled to whole render target */
+	bool bScaledToRenderTarget;
+
+	/** 
+	* Shared construction function
+	*/
+	void Construct();
+
+public:	
+
+	/*
+	 * Access current real time 
+	 */
+	float GetCurrentRealTime() const { return CurrentRealTime; }
+
+	/*
+	 * Access current world time 
+	 */
+	float GetCurrentWorldTime() const { return CurrentWorldTime; }
+
+	/*
+	 * Access current delta time 
+	 */
+	float GetCurrentDeltaWorldTime() const { return CurrentDeltaWorldTime; }
+
+	/* 
+	 * Draw a CanvasItem
+	 *
+	 * @param Item			Item to draw
+	 */
+	FORCEINLINE void DrawItem( FCanvasItem& Item )
+	{
+		Item.Draw( this );
+	}
+	/* 
+	 * Draw a CanvasItem at the given coordinates
+	 *
+	 * @param Item			Item to draw
+	 * @param InPosition	Position to draw item
+	 */
+	FORCEINLINE void DrawItem( FCanvasItem& Item, const FVector2D& InPosition )
+	{
+		Item.Draw( this, InPosition );
+	}
+	
+	/* 
+	 * Draw a CanvasItem at the given coordinates
+	 *
+	 * @param Item			Item to draw
+	 * @param X				X Position to draw item
+	 * @param Y				Y Position to draw item
+	 */
+	FORCEINLINE void DrawItem( FCanvasItem& Item, float X, float Y  )
+	{
+		Item.Draw( this, X, Y );
+	}
+
+	/**
+	* Clear the canvas
+	*
+	* @param	Color		Color to clear with.
+	*/
+	ENGINE_API void Clear(const FLinearColor& Color);
+
+
+
+	/** 
+	* Draw arbitrary aligned rectangle.
+	*
+	* @param X - X position to draw tile at
+	* @param Y - Y position to draw tile at
+	* @param SizeX - Width of tile
+	* @param SizeY - Height of tile
+	* @param U - Horizontal position of the upper left corner of the portion of the texture to be shown(texels)
+	* @param V - Vertical position of the upper left corner of the portion of the texture to be shown(texels)
+	* @param SizeU - The width of the portion of the texture to be drawn. This value is in texels. 
+	* @param SizeV - The height of the portion of the texture to be drawn. This value is in texels. 
+	* @param Color - tint applied to tile
+	* @param Texture - Texture to draw
+	* @param AlphaBlend - true to alphablend
+	*/
+	ENGINE_API void DrawTile( float X, float Y, float SizeX, float SizeY, float U, float V,  float SizeU, float SizeV, const FLinearColor& Color, const FTexture* Texture = NULL, bool AlphaBlend = true );
+
+	/** 
+	* Draw an string centered on given location. 
+	* This function is being deprecated. a FCanvasTextItem should be used instead.
+	* 
+	* @param StartX - X point
+	* @param StartY - Y point
+	* @param Text - Text to draw
+	* @param Font - Font to use
+	* @param Color - Color of the text
+	* @param ShadowColor - Shadow color to draw underneath the text (ignored for distance field fonts)
+	* @return total size in pixels of text drawn
+	*/
+	ENGINE_API int32 DrawShadowedString( float StartX, float StartY, const TCHAR* Text, class UFont* Font, const FLinearColor& Color, const FLinearColor& ShadowColor = FLinearColor::Black );
+	
+	ENGINE_API int32 DrawShadowedText( float StartX, float StartY, const FText& Text, class UFont* Font, const FLinearColor& Color, const FLinearColor& ShadowColor = FLinearColor::Black );
+
+	ENGINE_API void DrawNGon(const FVector2D& Center, const FColor& Color, int32 NumSides, float Radius);
+
+	/** 
+	* Contains all of the batched elements that need to be rendered at a certain depth sort key
+	*/
+	class FCanvasSortElement
+	{
+	public:
+		/** 
+		* Init constructor 
+		*/
+		FCanvasSortElement(int32 InDepthSortKey=0)
+			:	DepthSortKey(InDepthSortKey)
+		{}
+
+		/** 
+		* Equality is based on sort key 
+		*
+		* @param Other - instance to compare against
+		* @return true if equal
+		*/
+		bool operator==(const FCanvasSortElement& Other) const
+		{
+			return DepthSortKey == Other.DepthSortKey;
+		}
+
+		/** sort key for this set of render batch elements */
+		int32 DepthSortKey;
+		/** list of batches that should be rendered at this sort key level */
+		TArray<class FCanvasBaseRenderItem*> RenderBatchArray;
+	};
+	
+	/** Batched canvas elements to be sorted for rendering. Sort order is back-to-front */
+	TArray<FCanvasSortElement> SortedElements;
+	/** Map from sortkey to array index of SortedElements for faster lookup of existing entries */
+	TMap<int32,int32> SortedElementLookupMap;
+
+	/** Store index of last Element off to avoid semi expensive Find() */
+	int32 LastElementIndex;
+	
+	/**
+	* Get the sort element for the given sort key. Allocates a new entry if one does not exist
+	*
+	* @param DepthSortKey - the key used to find the sort element entry
+	* @return sort element entry
+	*/
+	FCanvasSortElement& GetSortElement(int32 DepthSortKey);
+
+};
+
+
+
+/**
+* Base interface for canvas items which can be batched for rendering
+*/
+class FCanvasBaseRenderItem
+{
+public:
+	virtual ~FCanvasBaseRenderItem()
+	{}
+
+	/**
+	* Renders the canvas item
+	*
+	* @param Canvas - canvas currently being rendered
+	* @param RHICmdList - command list to use
+	* @return true if anything rendered
+	*/
+	virtual bool Render_RenderThread(FRHICommandListImmediate& RHICmdList, const FCanvas* Canvas) = 0;
+	/**
+	* Renders the canvas item
+	*
+	* @param Canvas - canvas currently being rendered
+	* @return true if anything rendered
+	*/
+	virtual bool Render_GameThread(const FCanvas* Canvas) = 0;
+	/**
+	* FCanvasBatchedElementRenderItem instance accessor
+	*
+	* @return FCanvasBatchedElementRenderItem instance
+	*/
+	virtual class FCanvasBatchedElementRenderItem* GetCanvasBatchedElementRenderItem() { return NULL; }
+	/**
+	* FCanvasTileRendererItem instance accessor
+	*
+	* @return FCanvasTileRendererItem instance
+	*/
+	virtual class FCanvasTileRendererItem* GetCanvasTileRendererItem() { return NULL; }
+};
+
+
+/**
+* Info needed to render a batched element set
+*/
+class FCanvasBatchedElementRenderItem : public FCanvasBaseRenderItem
+{
+public:
+	/** 
+	* Init constructor 
+	*/
+	FCanvasBatchedElementRenderItem(
+		FBatchedElementParameters* InBatchedElementParameters=NULL,
+		const FTexture* InTexture=NULL,
+		ESimpleElementBlendMode InBlendMode=SE_BLEND_MAX,
+		FCanvas::EElementType InElementType=FCanvas::ET_MAX,
+		const FCanvas::FTransformEntry& InTransform=FCanvas::FTransformEntry(FMatrix::Identity),
+		const FDepthFieldGlowInfo& InGlowInfo=FDepthFieldGlowInfo() )
+		// this data is deleted after rendering has completed
+		: Data(new FRenderData(InBatchedElementParameters, InTexture, InBlendMode, InElementType, InTransform, InGlowInfo))
+	{}
+
+	/**
+	* Destructor to delete data in case nothing rendered
+	*/
+	virtual ~FCanvasBatchedElementRenderItem()
+	{
+		delete Data;
+	}
+
+	/**
+	* FCanvasBatchedElementRenderItem instance accessor
+	*
+	* @return this instance
+	*/
+	virtual class FCanvasBatchedElementRenderItem* GetCanvasBatchedElementRenderItem() 
+	{ 
+		return this; 
+	}
+
+	/**
+	* Renders the canvas item. 
+	* Iterates over all batched elements and draws them with their own transforms
+	*
+	* @param Canvas - canvas currently being rendered
+	* @param RHICmdList - command list to use
+	* @return true if anything rendered
+	*/
+	virtual bool Render_RenderThread(FRHICommandListImmediate& RHICmdList, const FCanvas* Canvas) override;
+	/**
+	* Renders the canvas item.
+	* Iterates over all batched elements and draws them with their own transforms
+	*
+	* @param Canvas - canvas currently being rendered
+	* @return true if anything rendered
+	*/
+	virtual bool Render_GameThread(const FCanvas* Canvas) override;
+
+	/**
+	* Determine if this is a matching set by comparing texture,blendmode,elementype,transform. All must match
+	*
+	* @param BatchedElementParameters - parameters for this batched element
+	* @param InTexture - texture resource for the item being rendered
+	* @param InBlendMode - current alpha blend mode 
+	* @param InElementType - type of item being rendered: triangle,line,etc
+	* @param InTransform - the transform for the item being rendered
+	* @param InGlowInfo - the depth field glow of the item being rendered
+	* @return true if the parameters match this render item
+	*/
+	bool IsMatch(FBatchedElementParameters* BatchedElementParameters, const FTexture* InTexture, ESimpleElementBlendMode InBlendMode, FCanvas::EElementType InElementType, const FCanvas::FTransformEntry& InTransform, const FDepthFieldGlowInfo& InGlowInfo)
+	{
+		return(	Data->BatchedElementParameters.GetReference() == BatchedElementParameters &&
+				Data->Texture == InTexture &&
+				Data->BlendMode == InBlendMode &&
+				Data->ElementType == InElementType &&
+				Data->Transform.GetMatrixCRC() == InTransform.GetMatrixCRC() &&
+				Data->GlowInfo == InGlowInfo );
+	}
+
+	/**
+	* Accessor for the batched elements. This can be used for adding triangles and primitives to the batched elements
+	*
+	* @return pointer to batched elements struct
+	*/
+	FORCEINLINE FBatchedElements* GetBatchedElements()
+	{
+		return &Data->BatchedElements;
+	}
+
+private:
+	class FRenderData
+	{
+	public:
+		/**
+		* Init constructor
+		*/
+		FRenderData(
+			FBatchedElementParameters* InBatchedElementParameters=NULL,
+			const FTexture* InTexture=NULL,
+			ESimpleElementBlendMode InBlendMode=SE_BLEND_MAX,
+			FCanvas::EElementType InElementType=FCanvas::ET_MAX,
+			const FCanvas::FTransformEntry& InTransform=FCanvas::FTransformEntry(FMatrix::Identity),
+			const FDepthFieldGlowInfo& InGlowInfo=FDepthFieldGlowInfo() )
+			:	BatchedElementParameters(InBatchedElementParameters)
+			,	Texture(InTexture)
+			,	BlendMode(InBlendMode)
+			,	ElementType(InElementType)
+			,	Transform(InTransform)
+			,	GlowInfo(InGlowInfo)
+		{}
+		/** Current batched elements, destroyed once rendering completes. */
+		FBatchedElements BatchedElements;
+		/** Batched element parameters */
+		TRefCountPtr<FBatchedElementParameters> BatchedElementParameters;
+		/** Current texture being used for batching, set to NULL if it hasn't been used yet. */
+		const FTexture* Texture;
+		/** Current blend mode being used for batching, set to BLEND_MAX if it hasn't been used yet. */
+		ESimpleElementBlendMode BlendMode;
+		/** Current element type being used for batching, set to ET_MAX if it hasn't been used yet. */
+		FCanvas::EElementType ElementType;
+		/** Transform used to render including projection */
+		FCanvas::FTransformEntry Transform;
+		/** info for optional glow effect when using depth field rendering */
+		FDepthFieldGlowInfo GlowInfo;
+	};
+	/**
+	* Render data which is allocated when a new FCanvasBatchedElementRenderItem is added for rendering.
+	* This data is only freed on the rendering thread once the item has finished rendering
+	*/
+	FRenderData* Data;		
+};
+
+
+/**
+* Info needed to render a single FTileRenderer
+*/
+class FCanvasTileRendererItem : public FCanvasBaseRenderItem
+{
+public:
+	/** 
+	* Init constructor 
+	*/
+	FCanvasTileRendererItem( 
+		const FMaterialRenderProxy* InMaterialRenderProxy=NULL,
+		const FCanvas::FTransformEntry& InTransform=FCanvas::FTransformEntry(FMatrix::Identity),
+		bool bInFreezeTime=false)
+		// this data is deleted after rendering has completed
+		:	Data(new FRenderData(InMaterialRenderProxy,InTransform))
+		,	bFreezeTime(bInFreezeTime)
+	{}
+
+	/**
+	* Destructor to delete data in case nothing rendered
+	*/
+	virtual ~FCanvasTileRendererItem()
+	{
+		delete Data;
+	}
+
+	/**
+	* FCanvasTileRendererItem instance accessor
+	*
+	* @return this instance
+	*/
+	virtual class FCanvasTileRendererItem* GetCanvasTileRendererItem() 
+	{ 
+		return this; 
+	}
+
+	/**
+	* Renders the canvas item. 
+	* Iterates over each tile to be rendered and draws it with its own transforms
+	*
+	* @param Canvas - canvas currently being rendered
+	* @param RHICmdList - command list to use
+	* @return true if anything rendered
+	*/
+	virtual bool Render_RenderThread(FRHICommandListImmediate& RHICmdList, const FCanvas* Canvas) override;
+
+	/**
+	* Renders the canvas item.
+	* Iterates over each tile to be rendered and draws it with its own transforms
+	*
+	* @param Canvas - canvas currently being rendered
+	* @return true if anything rendered
+	*/
+	virtual bool Render_GameThread(const FCanvas* Canvas) override;
+
+	/**
+	* Determine if this is a matching set by comparing material,transform. All must match
+	*
+	* @param IInMaterialRenderProxy - material proxy resource for the item being rendered
+	* @param InTransform - the transform for the item being rendered
+	* @return true if the parameters match this render item
+	*/
+	bool IsMatch( const FMaterialRenderProxy* InMaterialRenderProxy, const FCanvas::FTransformEntry& InTransform )
+	{
+		return( Data->MaterialRenderProxy == InMaterialRenderProxy && 
+				Data->Transform.GetMatrixCRC() == InTransform.GetMatrixCRC() );
+	};
+
+	/**
+	* Add a new tile to the render data. These tiles all use the same transform and material proxy
+	*
+	* @param X - tile X offset
+	* @param Y - tile Y offset
+	* @param SizeX - tile X size
+	* @param SizeY - tile Y size
+	* @param U - tile U offset
+	* @param V - tile V offset
+	* @param SizeU - tile U size
+	* @param SizeV - tile V size
+	* @param return number of tiles added
+	*/
+	FORCEINLINE int32 AddTile(float X,float Y,float SizeX,float SizeY,float U,float V,float SizeU,float SizeV,FHitProxyId HitProxyId)
+	{
+		return Data->AddTile(X,Y,SizeX,SizeY,U,V,SizeU,SizeV,HitProxyId);
+	};
+
+private:
+	class FRenderData
+	{
+	public:
+		FRenderData(
+			const FMaterialRenderProxy* InMaterialRenderProxy=NULL,
+			const FCanvas::FTransformEntry& InTransform=FCanvas::FTransformEntry(FMatrix::Identity) )
+			:	MaterialRenderProxy(InMaterialRenderProxy)
+			,	Transform(InTransform)
+		{}
+		const FMaterialRenderProxy* MaterialRenderProxy;
+		FCanvas::FTransformEntry Transform;
+
+		struct FTileInst
+		{
+			float X,Y;
+			float SizeX,SizeY;
+			float U,V;
+			float SizeU,SizeV;
+			FHitProxyId HitProxyId;
+		};
+		TArray<FTileInst> Tiles;
+
+		FORCEINLINE int32 AddTile(float X,float Y,float SizeX,float SizeY,float U,float V,float SizeU,float SizeV,FHitProxyId HitProxyId)
+		{
+			FTileInst NewTile = {X,Y,SizeX,SizeY,U,V,SizeU,SizeV,HitProxyId};
+			return Tiles.Add(NewTile);
+		};
+	};
+	/**
+	* Render data which is allocated when a new FCanvasTileRendererItem is added for rendering.
+	* This data is only freed on the rendering thread once the item has finished rendering
+	*/
+	FRenderData* Data;	
+
+	const bool bFreezeTime;
+};
+
+/**
+* Render string using both a font and a material. The material should have a font exposed as a 
+* parameter so that the correct font page can be set based on the character being drawn.
+*
+* @param Font - font containing texture pages of character glyphs
+* @param XL - out width
+* @param YL - out height
+* @param Text - string of text to be measured
+*/
+extern ENGINE_API void StringSize( UFont* Font, int32& XL, int32& YL, const TCHAR* Text );
+
+
diff --git a/Engine/Source/Runtime/Engine/Public/Engine.h b/Engine/Source/Runtime/Engine/Public/Engine.h
--- a/Engine/Source/Runtime/Engine/Public/Engine.h
+++ b/Engine/Source/Runtime/Engine/Public/Engine.h
@@ -378,6 +378,6 @@
 
 #include "DrawDebugHelpers.h"
 #include "UnrealEngine.h"					// Unreal engine helpers.
-#include "Canvas.h"							// Canvas.
+#include "CanvasTypes.h"							// Canvas.
 #include "EngineUtils.h"
 #include "TimerManager.h"					// Game play timers
diff --git a/Engine/Source/Runtime/Engine/Public/SceneTypes.h b/Engine/Source/Runtime/Engine/Public/SceneTypes.h
--- a/Engine/Source/Runtime/Engine/Public/SceneTypes.h
+++ b/Engine/Source/Runtime/Engine/Public/SceneTypes.h
@@ -154,4 +154,26 @@ enum EMaterialProperty
 	MP_CustomizedUVs7,
 	MP_MaterialAttributes,
 	MP_MAX,
+};
+
+/** Blend modes supported for simple element rendering */
+enum ESimpleElementBlendMode
+{
+	SE_BLEND_Opaque = 0,
+	SE_BLEND_Masked,
+	SE_BLEND_Translucent,
+	SE_BLEND_Additive,
+	SE_BLEND_Modulate,
+	SE_BLEND_MaskedDistanceField,
+	SE_BLEND_MaskedDistanceFieldShadowed,
+	SE_BLEND_TranslucentDistanceField,
+	SE_BLEND_TranslucentDistanceFieldShadowed,
+	SE_BLEND_AlphaComposite,
+	// Like SE_BLEND_Translucent, but modifies destination alpha
+	SE_BLEND_AlphaBlend,
+
+	SE_BLEND_RGBA_MASK_START,
+	SE_BLEND_RGBA_MASK_END = SE_BLEND_RGBA_MASK_START + 31, //Using 5bit bit-field for red, green, blue, alpha and desaturation
+
+	SE_BLEND_MAX
 };
\ No newline at end of file
