Added BP-callable accessors for spline mesh component values.
#ttp 341187 - EDITOR: BP: Unable to get Start and End values of a spline mesh component
#branch UE4
#proj Runtime.Engine

diff --git a/Engine/Source/Runtime/Engine/Classes/Components/SplineMeshComponent.h b/Engine/Source/Runtime/Engine/Classes/Components/SplineMeshComponent.h
--- a/Engine/Source/Runtime/Engine/Classes/Components/SplineMeshComponent.h
+++ b/Engine/Source/Runtime/Engine/Classes/Components/SplineMeshComponent.h
@@ -150,18 +150,34 @@ class ENGINE_API USplineMeshComponent : public UStaticMeshComponent, public IInt
 	/** Called when spline params are changed, to notify render thread and possibly collision */
 	void MarkSplineParamsDirty();
 
+	/** Get the start position of spline in local space */
+	UFUNCTION(BlueprintCallable, Category = SplineMesh)
+	FVector GetStartPosition() const;
+
 	/** Set the start position of spline in local space */
 	UFUNCTION(BlueprintCallable, Category=SplineMesh)
 	void SetStartPosition(FVector StartPos);
 
+	/** Get the start tangent vector of spline in local space */
+	UFUNCTION(BlueprintCallable, Category = SplineMesh)
+	FVector GetStartTangent() const;
+
 	/** Set the start tangent vector of spline in local space */
 	UFUNCTION(BlueprintCallable, Category = SplineMesh)
 	void SetStartTangent(FVector StartTangent);
 
+	/** Get the end position of spline in local space */
+	UFUNCTION(BlueprintCallable, Category = SplineMesh)
+	FVector GetEndPosition() const;
+
 	/** Set the end position of spline in local space */
 	UFUNCTION(BlueprintCallable, Category = SplineMesh)
 	void SetEndPosition(FVector EndPos);
 
+	/** Get the end tangent vector of spline in local space */
+	UFUNCTION(BlueprintCallable, Category = SplineMesh)
+	FVector GetEndTangent() const;
+
 	/** Set the end tangent vector of spline in local space */
 	UFUNCTION(BlueprintCallable, Category = SplineMesh)
 	void SetEndTangent(FVector EndTangent);
@@ -170,26 +186,50 @@ class ENGINE_API USplineMeshComponent : public UStaticMeshComponent, public IInt
 	UFUNCTION(BlueprintCallable, Category = SplineMesh)
 	void SetStartAndEnd(FVector StartPos, FVector StartTangent, FVector EndPos, FVector EndTangent);
 
+	/** Get the start scaling */
+	UFUNCTION(BlueprintCallable, Category = SplineMesh)
+	FVector2D GetStartScale() const;
+
 	/** Set the start scaling */
 	UFUNCTION(BlueprintCallable, Category = SplineMesh)
 	void SetStartScale(FVector2D StartScale = FVector2D(1,1));
 
+	/** Get the start roll */
+	UFUNCTION(BlueprintCallable, Category = SplineMesh)
+	float GetStartRoll() const;
+
 	/** Set the start roll */
 	UFUNCTION(BlueprintCallable, Category = SplineMesh)
 	void SetStartRoll(float StartRoll);
 
+	/** Get the start offset */
+	UFUNCTION(BlueprintCallable, Category = SplineMesh)
+	FVector2D GetStartOffset() const;
+
 	/** Set the start offset */
 	UFUNCTION(BlueprintCallable, Category = SplineMesh)
 	void SetStartOffset(FVector2D StartOffset);
 
+	/** Get the end scaling */
+	UFUNCTION(BlueprintCallable, Category = SplineMesh)
+	FVector2D GetEndScale() const;
+
 	/** Set the end scaling */
 	UFUNCTION(BlueprintCallable, Category = SplineMesh)
 	void SetEndScale(FVector2D EndScale = FVector2D(1,1));
 
+	/** Get the end roll */
+	UFUNCTION(BlueprintCallable, Category = SplineMesh)
+	float GetEndRoll() const;
+
 	/** Set the end roll */
 	UFUNCTION(BlueprintCallable, Category = SplineMesh)
 	void SetEndRoll(float EndRoll);
 
+	/** Get the end offset */
+	UFUNCTION(BlueprintCallable, Category = SplineMesh)
+	FVector2D GetEndOffset() const;
+
 	/** Set the end offset */
 	UFUNCTION(BlueprintCallable, Category = SplineMesh)
 	void SetEndOffset(FVector2D EndOffset);
diff --git a/Engine/Source/Runtime/Engine/Private/SplineMeshComponent.cpp b/Engine/Source/Runtime/Engine/Private/SplineMeshComponent.cpp
--- a/Engine/Source/Runtime/Engine/Private/SplineMeshComponent.cpp
+++ b/Engine/Source/Runtime/Engine/Private/SplineMeshComponent.cpp
@@ -235,24 +235,44 @@ USplineMeshComponent::USplineMeshComponent(const class FPostConstructInitializeP
 	SplineParams.EndScale = FVector2D(1.f, 1.f);
 }
 
+FVector USplineMeshComponent::GetStartPosition() const
+{
+	return SplineParams.StartPos;
+}
+
 void USplineMeshComponent::SetStartPosition(FVector StartPos)
 {
 	SplineParams.StartPos = StartPos;
 	MarkSplineParamsDirty();
 }
 
+FVector USplineMeshComponent::GetStartTangent() const
+{
+	return SplineParams.StartTangent;
+}
+
 void USplineMeshComponent::SetStartTangent(FVector StartTangent)
 {
 	SplineParams.StartTangent = StartTangent;
 	MarkSplineParamsDirty();
 }
 
+FVector USplineMeshComponent::GetEndPosition() const
+{
+	return SplineParams.EndPos;
+}
+
 void USplineMeshComponent::SetEndPosition(FVector EndPos)
 {
 	SplineParams.EndPos = EndPos;
 	MarkSplineParamsDirty();
 }
 
+FVector USplineMeshComponent::GetEndTangent() const
+{
+	return SplineParams.EndTangent;
+}
+
 void USplineMeshComponent::SetEndTangent(FVector EndTangent)
 {
 	SplineParams.EndTangent = EndTangent;
@@ -268,36 +288,66 @@ void USplineMeshComponent::SetStartAndEnd(FVector StartPos, FVector StartTangent
 	MarkSplineParamsDirty();
 }
 
+FVector2D USplineMeshComponent::GetStartScale() const
+{
+	return SplineParams.StartScale;
+}
+
 void USplineMeshComponent::SetStartScale(FVector2D StartScale)
 {
 	SplineParams.StartScale = StartScale;
 	MarkSplineParamsDirty();
 }
 
+float USplineMeshComponent::GetStartRoll() const
+{
+	return SplineParams.StartRoll;
+}
+
 void USplineMeshComponent::SetStartRoll(float StartRoll)
 {
 	SplineParams.StartRoll = StartRoll;
 	MarkSplineParamsDirty();
 }
 
+FVector2D USplineMeshComponent::GetStartOffset() const
+{
+	return SplineParams.StartOffset;
+}
+
 void USplineMeshComponent::SetStartOffset(FVector2D StartOffset)
 {
 	SplineParams.StartOffset = StartOffset;
 	MarkSplineParamsDirty();
 }
 
+FVector2D USplineMeshComponent::GetEndScale() const
+{
+	return SplineParams.EndScale;
+}
+
 void USplineMeshComponent::SetEndScale(FVector2D EndScale)
 {
 	SplineParams.EndScale = EndScale;
 	MarkSplineParamsDirty();
 }
 
+float USplineMeshComponent::GetEndRoll() const
+{
+	return SplineParams.EndRoll;
+}
+
 void USplineMeshComponent::SetEndRoll(float EndRoll)
 {
 	SplineParams.EndRoll = EndRoll;
 	MarkSplineParamsDirty();
 }
 
+FVector2D USplineMeshComponent::GetEndOffset() const
+{
+	return SplineParams.EndOffset;
+}
+
 void USplineMeshComponent::SetEndOffset(FVector2D EndOffset)
 {
 	SplineParams.EndOffset = EndOffset;
