UDS: Default value propagation fix.

diff --git a/Engine/Source/Runtime/Engine/Private/UserDefinedStruct.cpp b/Engine/Source/Runtime/Engine/Private/UserDefinedStruct.cpp
--- a/Engine/Source/Runtime/Engine/Private/UserDefinedStruct.cpp
+++ b/Engine/Source/Runtime/Engine/Private/UserDefinedStruct.cpp
@@ -62,18 +62,21 @@ void UUserDefinedStruct::SerializeTaggedProperties(FArchive& Ar, uint8* Data, US
 
 	auto UDDefaultsStruct = Cast<UUserDefinedStruct>(DefaultsStruct);
 
+	const bool bDuplicate = (0 != (Ar.GetPortFlags() & PPF_Duplicate));
+
 	/*	When saving delta, we want the difference between current data and true structure's default values. 
 		When Defaults is NULL then zeroed data will be used for comparison.*/
 	const bool bUseNewDefaults = !Defaults
 		&& UDDefaultsStruct
 		&& Ar.DoDelta()
 		&& Ar.IsSaving()
-		&& (0 == (Ar.GetPortFlags() & PPF_Duplicate))
+		&& !bDuplicate
 		&& !Ar.IsCooking();
 
 	/*	Object serialized from delta will have missing properties filled with zeroed data, 
 		we want structure's default data instead */
 	const bool bLoadDefaultFirst = UDDefaultsStruct 
+		&& !bDuplicate
 		&& Ar.IsLoading() 
 		&& Ar.IsPersistent();
 
