Don't pass mouse events to viewports that don't have mouse focus

diff --git a/Engine/Source/Runtime/Engine/Private/GameViewportClient.cpp b/Engine/Source/Runtime/Engine/Private/GameViewportClient.cpp
--- a/Engine/Source/Runtime/Engine/Private/GameViewportClient.cpp
+++ b/Engine/Source/Runtime/Engine/Private/GameViewportClient.cpp
@@ -217,7 +217,7 @@ bool UGameViewportClient::InputKey(FViewport* InViewport, int32 ControllerId, FK
 	if (!bResult)
 	{
 		ULocalPlayer* const TargetPlayer = GEngine->GetLocalPlayerFromControllerId(this, ControllerId);
-		if (TargetPlayer && TargetPlayer->PlayerController)
+		if (TargetPlayer && TargetPlayer->PlayerController && InViewport->HasMouseCapture())
 		{
 			bResult = TargetPlayer->PlayerController->InputKey(Key, EventType, AmountDepressed, bGamepad);
 		}
diff --git a/Engine/Source/Runtime/Engine/Private/Slate/SceneViewport.cpp b/Engine/Source/Runtime/Engine/Private/Slate/SceneViewport.cpp
--- a/Engine/Source/Runtime/Engine/Private/Slate/SceneViewport.cpp
+++ b/Engine/Source/Runtime/Engine/Private/Slate/SceneViewport.cpp
@@ -310,7 +310,11 @@ FReply FSceneViewport::OnMouseButtonDown( const FGeometry& InGeometry, const FPo
 	// Prevent throttling when interacting with the viewport so we can move around in it
 	CurrentReplyState = FReply::Handled().PreventThrottling();
 	
-	KeyStateMap.Add( InMouseEvent.GetEffectingButton(), true );
+	if (!IsPlayInEditorViewport() || HasMouseCapture())
+	{
+		KeyStateMap.Add(InMouseEvent.GetEffectingButton(), true);
+	}
+
 	UpdateModifierKeys( InMouseEvent );
 	UpdateCachedMousePos( InGeometry, InMouseEvent );
 	UpdateCachedGeometry(InGeometry);
