#ue4 - Lower log level of collision failure to resolve MTD from "warning" to "verbose", as users are reporting too much log spam.

Related to:
https://answers.unrealengine.com/questions/70846/log-error-what-does-this-mean-logphysicswarning-co.html

diff --git a/Engine/Source/Runtime/Engine/Private/Collision/CollisionConversions.cpp b/Engine/Source/Runtime/Engine/Private/Collision/CollisionConversions.cpp
--- a/Engine/Source/Runtime/Engine/Private/Collision/CollisionConversions.cpp
+++ b/Engine/Source/Runtime/Engine/Private/Collision/CollisionConversions.cpp
@@ -639,7 +639,7 @@ static bool ConvertOverlappedShapeToImpactHit(const PxLocationHit& PHit, const F
 
 				if (Distance < KINDA_SMALL_NUMBER)
 				{
-					UE_LOG(LogCollision, Warning, TEXT("ConvertOverlappedShapeToImpactHit: Query origin inside shape, giving poor MTD."));
+					UE_LOG(LogCollision, Verbose, TEXT("Warning: ConvertOverlappedShapeToImpactHit: Query origin inside shape, giving poor MTD."));
 					PClosestPoint = PxShapeExt::getWorldBounds(*PShape, *PActor).getCenter(); 
 				}
 
