Remove unused metadata from properties

diff --git a/Engine/Source/Runtime/Engine/Classes/Animation/AnimNode_BlendSpaceEvaluator.h b/Engine/Source/Runtime/Engine/Classes/Animation/AnimNode_BlendSpaceEvaluator.h
--- a/Engine/Source/Runtime/Engine/Classes/Animation/AnimNode_BlendSpaceEvaluator.h
+++ b/Engine/Source/Runtime/Engine/Classes/Animation/AnimNode_BlendSpaceEvaluator.h
@@ -15,7 +15,7 @@ struct ENGINE_API FAnimNode_BlendSpaceEvaluator : public FAnimNode_BlendSpacePla
 
 public:
 	/** Normalized time between [0,1]. The actual length of a blendspace is dynamic based on the coordinate, so it is exposed as a normalized value. */
-	UPROPERTY(EditAnywhere, BlueprintReadWrite, Category=Settings, meta=(DefaultValue="0.0", PinShownByDefault))
+	UPROPERTY(EditAnywhere, BlueprintReadWrite, Category=Settings, meta=(PinShownByDefault))
 	mutable float NormalizedTime;
 
 public:	
diff --git a/Engine/Source/Runtime/Engine/Classes/Animation/AnimNode_SequencePlayer.h b/Engine/Source/Runtime/Engine/Classes/Animation/AnimNode_SequencePlayer.h
--- a/Engine/Source/Runtime/Engine/Classes/Animation/AnimNode_SequencePlayer.h
+++ b/Engine/Source/Runtime/Engine/Classes/Animation/AnimNode_SequencePlayer.h
@@ -21,11 +21,11 @@ public:
 	float InternalTimeAccumulator;
 
 	// Should the animation continue looping when it reaches the end?
-	UPROPERTY(EditAnywhere, BlueprintReadWrite, Category=Settings, meta=(DefaultValue="true", PinHiddenByDefault))
+	UPROPERTY(EditAnywhere, BlueprintReadWrite, Category=Settings, meta=(PinHiddenByDefault))
 	mutable bool bLoopAnimation;
 
 	// The play rate multiplier. Can be negative, which will cause the animation to play in reverse.
-	UPROPERTY(EditAnywhere, BlueprintReadWrite, Category=Settings, meta=(DefaultValue="1.0", PinHiddenByDefault))
+	UPROPERTY(EditAnywhere, BlueprintReadWrite, Category=Settings, meta=(PinHiddenByDefault))
 	mutable float PlayRate;
 
 	// The group index, assigned at compile time based on the editoronly GroupName (or INDEX_NONE if it is not part of any group)
