Removing MediumDetailSpawnRateScale, replacing with QualityLevelSpawnRateScale dependent on effects quality level; fixing several problems with initial rotation / mesh orientation not properly applying, leading to flipping particles in the first frame after spawn

diff --git a/Engine/Source/Runtime/Engine/Classes/Particles/ParticleEmitter.h b/Engine/Source/Runtime/Engine/Classes/Particles/ParticleEmitter.h
--- a/Engine/Source/Runtime/Engine/Classes/Particles/ParticleEmitter.h
+++ b/Engine/Source/Runtime/Engine/Classes/Particles/ParticleEmitter.h
@@ -137,8 +137,11 @@ class UParticleEmitter : public UObject
 	 * A value of 0 effectively disables this emitter outside of high detail mode,
 	 * And this does not affect spawn per unit, unless the value is 0.
 	 */
-	UPROPERTY(EditAnywhere, Category=Particle)
-	float MediumDetailSpawnRateScale;
+	UPROPERTY()
+	float MediumDetailSpawnRateScale_DEPRECATED;
+
+	UPROPERTY(EditAnywhere, Category = Particle)
+	float QualityLevelSpawnRateScale;
 
 	/** If detail mode is >= system detail mode, primitive won't be rendered. */
 	UPROPERTY(EditAnywhere, Category=Particle)
@@ -248,6 +251,7 @@ class UParticleEmitter : public UObject
 		}
 	}
 
+
 	/**
 	 * This will update the LOD of the particle in the editor.
 	 *
@@ -298,6 +302,21 @@ class UParticleEmitter : public UObject
 	 * Builds data needed for simulation by the emitter from all modules.
 	 */
 	void Build();
+
+	/*
+	 *   Calculate spawn rate multiplier based on global effects quality level and emitter's quality scale
+ 	 */
+	FORCEINLINE float GetQualityLevelSpawnRateMult()
+	{
+		float Level = (1 - Scalability::GetQualityLevels().EffectsQuality);
+		float Q = 1;
+		for (int i = 0; i < Level + 1; i++)
+		{
+
+			Q = Q*QualityLevelSpawnRateScale;
+		}
+		return Q;
+	}
 };
 
 
diff --git a/Engine/Source/Runtime/Engine/Private/Particles/ParticleComponents.cpp b/Engine/Source/Runtime/Engine/Private/Particles/ParticleComponents.cpp
--- a/Engine/Source/Runtime/Engine/Private/Particles/ParticleComponents.cpp
+++ b/Engine/Source/Runtime/Engine/Private/Particles/ParticleComponents.cpp
@@ -976,7 +976,8 @@ bool UParticleLODLevel::IsModuleEditable(UParticleModule* InModule)
 	UParticleEmitter implementation.
 -----------------------------------------------------------------------------*/
 UParticleEmitter::UParticleEmitter(const class FPostConstructInitializeProperties& PCIP)
-	: Super(PCIP)
+	: Super(PCIP),
+	QualityLevelSpawnRateScale(1.0f)
 {
 	// Structure to hold one-time initialization
 	struct FConstructorStatics
@@ -992,7 +993,6 @@ UParticleEmitter::UParticleEmitter(const class FPostConstructInitializePropertie
 	EmitterName = ConstructorStatics.NAME_Particle_Emitter;
 	ConvertedModules = true;
 	PeakActiveParticles = 0;
-	MediumDetailSpawnRateScale = 1.0;
 #if WITH_EDITORONLY_DATA
 	EmitterEditorColor = FColor(0, 150, 150, 255);
 #endif // WITH_EDITORONLY_DATA
@@ -1275,7 +1275,7 @@ void UParticleEmitter::PostLoad()
 	}
 
 	// Clamp the detail spawn rate scale...
-	MediumDetailSpawnRateScale = FMath::Clamp<float>(MediumDetailSpawnRateScale, 0.0f, 1.0f);
+	QualityLevelSpawnRateScale = FMath::Clamp<float>(QualityLevelSpawnRateScale, 0.0f, 1.0f);
 
 	UpdateModuleLists();
 }
@@ -1324,7 +1324,7 @@ void UParticleEmitter::PostEditChangeProperty(FPropertyChangedEvent& PropertyCha
 	}
 
 	// Clamp the detail spawn rate scale...
-	MediumDetailSpawnRateScale = FMath::Clamp<float>(MediumDetailSpawnRateScale, 0.0f, 1.0f);
+	QualityLevelSpawnRateScale = FMath::Clamp<float>(QualityLevelSpawnRateScale, 0.0f, 1.0f);
 }
 #endif // WITH_EDITOR
 
diff --git a/Engine/Source/Runtime/Engine/Private/Particles/ParticleEmitterInstances.cpp b/Engine/Source/Runtime/Engine/Private/Particles/ParticleEmitterInstances.cpp
--- a/Engine/Source/Runtime/Engine/Private/Particles/ParticleEmitterInstances.cpp
+++ b/Engine/Source/Runtime/Engine/Private/Particles/ParticleEmitterInstances.cpp
@@ -440,7 +440,7 @@ void FParticleEmitterInstance::Init()
 	// Resize to sensible default.
 	if (Component->GetWorld()->IsGameWorld() == true &&
 		// Only presize if any particles will be spawned 
-		(Component->GetCurrentDetailMode() == DM_High || SpriteTemplate->MediumDetailSpawnRateScale > 0))
+		SpriteTemplate->QualityLevelSpawnRateScale > 0)
 	{
 		if ((LODLevel->PeakActiveParticles > 0) || (SpriteTemplate->InitialAllocationCount > 0))
 		{
@@ -1720,7 +1720,7 @@ float FParticleEmitterInstance::Spawn(float DeltaTime)
 	bool bProcessBurstList = true;
 	int32 DetailMode = Component->GetCurrentDetailMode();
 
-	if ((DetailMode == DM_High) || (SpriteTemplate->MediumDetailSpawnRateScale > 0.0f))
+	if (SpriteTemplate->QualityLevelSpawnRateScale > 0.0f)
 	{
 		// Process all Spawning modules that are present in the emitter.
 		for (int32 SpawnModIndex = 0; SpawnModIndex < LODLevel->SpawningModules.Num(); SpawnModIndex++)
@@ -1771,11 +1771,9 @@ float FParticleEmitterInstance::Spawn(float DeltaTime)
 			BurstCount += Burst;
 		}
 
-		if (DetailMode != DM_High)
-		{
-			SpawnRate = FMath::Max<float>(0.0f, SpawnRate * SpriteTemplate->MediumDetailSpawnRateScale);
-			BurstCount = FMath::CeilToInt(BurstCount * SpriteTemplate->MediumDetailSpawnRateScale);
-		}
+		float QualityMult = SpriteTemplate->GetQualityLevelSpawnRateMult();
+		SpawnRate = FMath::Max<float>(0.0f, SpawnRate * QualityMult);
+		BurstCount = FMath::CeilToInt(BurstCount * QualityMult);
 	}
 	else
 	{
@@ -2934,16 +2932,14 @@ void FParticleMeshEmitterInstance::Tick(float DeltaTime, bool bSuppressSpawning)
 
 				FQuat Rotation	= FQuat::FindBetween(OldDirection, NewDirection);
 				FVector Euler	= Rotation.Euler();
-				PayloadData->Rotation.X	= PayloadData->InitialOrientation.X + Euler.X;
-				PayloadData->Rotation.Y = PayloadData->InitialOrientation.Y + Euler.Y;
-				PayloadData->Rotation.Z = PayloadData->InitialOrientation.Z + Euler.Z;
+				PayloadData->Rotation = PayloadData->InitialOrientation + PayloadData->InitRotation + Euler;
 				PayloadData->Rotation += PayloadData->CurContinuousRotation;
 			}
 			else // not PSA_Velocity or PSA_AwayfromCenter, so rotation is not reset every tick
 			{
 				if ((Particle.Flags & STATE_Particle_FreezeRotation) == 0)
 				{
-					PayloadData->Rotation = PayloadData->CurContinuousRotation;
+					PayloadData->Rotation = PayloadData->InitialOrientation + PayloadData->InitRotation + PayloadData->CurContinuousRotation;
 				}
 			}
 
diff --git a/Engine/Source/Runtime/Engine/Private/Particles/ParticleModules.cpp b/Engine/Source/Runtime/Engine/Private/Particles/ParticleModules.cpp
--- a/Engine/Source/Runtime/Engine/Private/Particles/ParticleModules.cpp
+++ b/Engine/Source/Runtime/Engine/Private/Particles/ParticleModules.cpp
@@ -1280,9 +1280,11 @@ void UParticleModuleMeshRotation::SpawnEx(FParticleEmitterInstance* Owner, int32
 				Rotation.Z	+= ParentAffectedRotation.Z / 360.0f;
 			}
 			FMeshRotationPayloadData* PayloadData = (FMeshRotationPayloadData*)((uint8*)&Particle + MeshRotationOffset);
-			PayloadData->Rotation.X	+= Rotation.X * 360.0f;
-			PayloadData->Rotation.Y	+= Rotation.Y * 360.0f;
-			PayloadData->Rotation.Z	+= Rotation.Z * 360.0f;
+			//PayloadData->Rotation.X	+= Rotation.X * 360.0f;
+			//PayloadData->Rotation.Y	+= Rotation.Y * 360.0f;
+			//PayloadData->Rotation.Z	+= Rotation.Z * 360.0f;
+			PayloadData->InitRotation = Rotation * 360.0f;
+			PayloadData->Rotation += PayloadData->InitRotation;
 		}
 	}
 }
diff --git a/Engine/Source/Runtime/Engine/Private/Particles/ParticleTrail2EmitterInstance.cpp b/Engine/Source/Runtime/Engine/Private/Particles/ParticleTrail2EmitterInstance.cpp
--- a/Engine/Source/Runtime/Engine/Private/Particles/ParticleTrail2EmitterInstance.cpp
+++ b/Engine/Source/Runtime/Engine/Private/Particles/ParticleTrail2EmitterInstance.cpp
@@ -1217,11 +1217,8 @@ float FParticleRibbonEmitterInstance::Spawn(float DeltaTime)
 	if (bProcessSpawnRate)
 	{
 		float RateScale = LODLevel->SpawnModule->RateScale.GetValue(EmitterTime, Component);
-		if (Component->GetCurrentDetailMode() != DM_High)
-		{
-			RateScale *= SpriteTemplate->MediumDetailSpawnRateScale;
-		}
-		SpawnRate += LODLevel->SpawnModule->Rate.GetValue(EmitterTime, Component) * FMath::Clamp<float>(RateScale, 0.0f, RateScale);
+		float QualityMult = SpriteTemplate->GetQualityLevelSpawnRateMult();
+		SpawnRate += LODLevel->SpawnModule->Rate.GetValue(EmitterTime, Component) * FMath::Clamp<float>(QualityMult, 0.0f, 1.0);
 	}
 
 	// Take Bursts into account as well...
@@ -3296,10 +3293,8 @@ float FParticleAnimTrailEmitterInstance::Spawn(float DeltaTime)
 	if (bProcessSpawnRate)
 	{
 		float RateScale = LODLevel->SpawnModule->RateScale.GetValue(EmitterTime, Component);
-		if (Component->GetCurrentDetailMode() != DM_High)
-		{
-			RateScale *= SpriteTemplate->MediumDetailSpawnRateScale;
-		}
+		float QualityMult = 0.25f * (1 << Scalability::GetQualityLevels().EffectsQuality);
+		RateScale *= SpriteTemplate->QualityLevelSpawnRateScale*QualityMult;
 		SpawnRate += LODLevel->SpawnModule->Rate.GetValue(EmitterTime, Component) * FMath::Clamp<float>(RateScale, 0.0f, RateScale);
 	}
 
diff --git a/Engine/Source/Runtime/Engine/Public/ParticleHelper.h b/Engine/Source/Runtime/Engine/Public/ParticleHelper.h
--- a/Engine/Source/Runtime/Engine/Public/ParticleHelper.h
+++ b/Engine/Source/Runtime/Engine/Public/ParticleHelper.h
@@ -750,7 +750,8 @@ struct FAnimTrailTypeDataPayload : public FTrailsBaseTypeDataPayload
 /** Mesh rotation data payload										*/
 struct FMeshRotationPayloadData
 {
-	FVector	 InitialOrientation;
+	FVector	 InitialOrientation;		// from mesh data module
+	FVector  InitRotation;				// from init rotation module
 	FVector  Rotation;
 	FVector	 CurContinuousRotation;
 	FVector  RotationRate;
