Changing orientation to be applied before initial size for particles, to have the desired order of operations of Orientation->InitialSize->InitRotation->RotationRate; it was using one rotation matrix with all rotations concatenated, which makes Orientation from the mesh data module and the InitRotation module do exactly the same thing, which is a bug.

diff --git a/Engine/Source/Runtime/Engine/Private/Particles/ParticleSystemRender.cpp b/Engine/Source/Runtime/Engine/Private/Particles/ParticleSystemRender.cpp
--- a/Engine/Source/Runtime/Engine/Private/Particles/ParticleSystemRender.cpp
+++ b/Engine/Source/Runtime/Engine/Private/Particles/ParticleSystemRender.cpp
@@ -2249,7 +2249,12 @@ void FDynamicMeshEmitterData::GetParticleTransform(FBaseParticle& InParticle, co
 		}
 		else
 		{
-			OutTransformMat = kScaleMat * kRotMat * kTransMat;
+			FMeshRotationPayloadData* PayloadData = (FMeshRotationPayloadData*)((uint8*)&InParticle + Source.MeshRotationOffset);
+			FRotator MeshOrient = FRotator::MakeFromEuler(PayloadData->InitialOrientation);
+			FRotator OtherRot = FRotator::MakeFromEuler(PayloadData->InitRotation+PayloadData->CurContinuousRotation);
+			FRotationMatrix OrientMat(MeshOrient);
+			FRotationMatrix RotMat(OtherRot);
+			OutTransformMat = (OrientMat*kScaleMat) * RotMat * kTransMat;
 		}
 	}
 	else if ((bUseCameraFacing == true) || (bUseMeshLockedAxis == true))
