Reseting previous frame transforms stored in ViewState when world orign is rebased

diff --git a/Engine/Source/Runtime/Engine/Classes/Engine/LocalPlayer.h b/Engine/Source/Runtime/Engine/Classes/Engine/LocalPlayer.h
--- a/Engine/Source/Runtime/Engine/Classes/Engine/LocalPlayer.h
+++ b/Engine/Source/Runtime/Engine/Classes/Engine/LocalPlayer.h
@@ -270,14 +270,6 @@ public:
 	 */
 	bool GetProjectionData(FViewport* Viewport, EStereoscopicPass StereoPass, FSceneViewProjectionData& ProjectionData);
 
-	/**
-	 * Called on world origin changes
-	 *
-	 * @param InOffset		Offset applied to current world origin
-	 */
-	void ApplyWorldOffset(FVector InOffset);
-
-
 	/**
 	 * Determines whether this player is the first and primary player on their machine.
 	 * @return	true if this player is not using splitscreen, or is the first player in the split-screen layout.
diff --git a/Engine/Source/Runtime/Engine/Classes/Engine/World.h b/Engine/Source/Runtime/Engine/Classes/Engine/World.h
--- a/Engine/Source/Runtime/Engine/Classes/Engine/World.h
+++ b/Engine/Source/Runtime/Engine/Classes/Engine/World.h
@@ -860,7 +860,10 @@ public:
 
 	/** Requested new world origin offset */
 	FIntPoint RequestedGlobalOriginOffset;
-	
+
+	/** Whether world origin was rebased this frame */
+	bool bOriginOffsetThisFrame;
+		
 	/** All levels information from which our world is composed */
 	UPROPERTY()
 	class UWorldComposition* WorldComposition;
diff --git a/Engine/Source/Runtime/Engine/Private/LevelTick.cpp b/Engine/Source/Runtime/Engine/Private/LevelTick.cpp
--- a/Engine/Source/Runtime/Engine/Private/LevelTick.cpp
+++ b/Engine/Source/Runtime/Engine/Private/LevelTick.cpp
@@ -1064,6 +1064,11 @@ void UWorld::Tick( ELevelTick TickType, float DeltaSeconds )
 	if (GlobalOriginOffset != RequestedGlobalOriginOffset)
 	{
 		SetNewWorldOrigin(RequestedGlobalOriginOffset);
+		bOriginOffsetThisFrame = true;
+	}
+	else
+	{
+		bOriginOffsetThisFrame = false;
 	}
 	
 	// update world's subsystems (NavigationSystem for now)
diff --git a/Engine/Source/Runtime/Engine/Private/LocalPlayer.cpp b/Engine/Source/Runtime/Engine/Private/LocalPlayer.cpp
--- a/Engine/Source/Runtime/Engine/Private/LocalPlayer.cpp
+++ b/Engine/Source/Runtime/Engine/Private/LocalPlayer.cpp
@@ -679,6 +679,8 @@ FSceneView* ULocalPlayer::CalcSceneView( class FSceneViewFamily* ViewFamily,
 		// Was there a camera cut this frame?
 		ViewInitOptions.bInCameraCut = PlayerController->PlayerCameraManager->bGameCameraCutThisFrame;
 	}
+	
+	check( PlayerController->GetWorld() );
 
 	// Fill out the rest of the view init options
 	ViewInitOptions.ViewFamily = ViewFamily;
@@ -690,6 +692,7 @@ FSceneView* ULocalPlayer::CalcSceneView( class FSceneViewFamily* ViewFamily,
 	ViewInitOptions.StereoPass = StereoPass;
 	ViewInitOptions.WorldToMetersScale = PlayerController->GetWorldSettings()->WorldToMeters;
 	ViewInitOptions.CursorPos = Viewport->HasMouseCapture() ? FIntPoint(-1, -1) : FIntPoint(Viewport->GetMouseX(), Viewport->GetMouseY());
+	ViewInitOptions.bOriginOffsetThisFrame = PlayerController->GetWorld()->bOriginOffsetThisFrame;
 	PlayerController->BuildHiddenComponentList(OutViewLocation, /*out*/ ViewInitOptions.HiddenPrimitives);
 
 	FSceneView* const View = new FSceneView(ViewInitOptions);
@@ -699,9 +702,7 @@ FSceneView* ULocalPlayer::CalcSceneView( class FSceneViewFamily* ViewFamily,
 
 	//@TODO: SPLITSCREEN: This call will have an issue with splitscreen, as the show flags are shared across the view family
 	EngineShowFlagOrthographicOverride(View->IsPerspectiveProjection(), ViewFamily->EngineShowFlags);
-
-	check( PlayerController->GetWorld() );
-
+		
 	ViewFamily->Views.Add(View);
 
 	{
@@ -1071,11 +1072,6 @@ bool ULocalPlayer::GetProjectionData(FViewport* Viewport, EStereoscopicPass Ster
 	return true;
 }
 
-void ULocalPlayer::ApplyWorldOffset(FVector InOffset)
-{
-	ViewState.GetReference()->ApplyWorldOffset(InOffset);
-}
-
 bool ULocalPlayer::HandleDNCommand( const TCHAR* Cmd, FOutputDevice& Ar )
 {
 	// Create a pending Note actor (only in PIE)
diff --git a/Engine/Source/Runtime/Engine/Private/SceneView.cpp b/Engine/Source/Runtime/Engine/Private/SceneView.cpp
--- a/Engine/Source/Runtime/Engine/Private/SceneView.cpp
+++ b/Engine/Source/Runtime/Engine/Private/SceneView.cpp
@@ -187,6 +187,7 @@ FSceneView::FSceneView(const FSceneViewInitOptions& InitOptions)
 	, HiddenPrimitives(InitOptions.HiddenPrimitives)
 	, LODDistanceFactor(InitOptions.LODDistanceFactor)
 	, bCameraCut(InitOptions.bInCameraCut)
+	, bOriginOffsetThisFrame(InitOptions.bOriginOffsetThisFrame)
 	, CursorPos(InitOptions.CursorPos)
 	, bIsGameView(false)
 	, bForceShowMaterials(false)
diff --git a/Engine/Source/Runtime/Engine/Private/World.cpp b/Engine/Source/Runtime/Engine/Private/World.cpp
--- a/Engine/Source/Runtime/Engine/Private/World.cpp
+++ b/Engine/Source/Runtime/Engine/Private/World.cpp
@@ -3904,20 +3904,6 @@ bool UWorld::SetNewWorldOrigin(const FIntPoint& InNewOrigin)
 
 	FVector Offset(GlobalOriginOffset - InNewOrigin, 0);
 
-	// Local players
-	for (auto It = GetPlayerControllerIterator(); It; ++It)
-	{
-		APlayerController* PlayerController = *It;
-		if (PlayerController)
-		{
-			ULocalPlayer* LocalPlayer = Cast<ULocalPlayer>(PlayerController->Player);
-			if (LocalPlayer)
-			{
-				LocalPlayer->ApplyWorldOffset(Offset);
-			}
-		}
-	}
-	
 	// Send offset command to rendering thread
 	Scene->ApplyWorldOffset(Offset);
 
diff --git a/Engine/Source/Runtime/Engine/Public/SceneManagement.h b/Engine/Source/Runtime/Engine/Public/SceneManagement.h
--- a/Engine/Source/Runtime/Engine/Public/SceneManagement.h
+++ b/Engine/Source/Runtime/Engine/Public/SceneManagement.h
@@ -116,9 +116,6 @@ public:
 	virtual void OnStartPostProcessing(FSceneView& CurrentView) = 0;
 	/** Allows MIDs being created and released during view rendering without the overhead of creating and relasing objects */
 	virtual UMaterialInstanceDynamic* GetReusableMID(class UMaterialInterface* ParentMaterial) = 0;
-	/** Called on world origin chnages */
-	virtual void ApplyWorldOffset(FVector InOffset) = 0;
-
 protected:
 	// Don't allow direct deletion of the view state, Destroy should be called instead.
 	virtual ~FSceneViewStateInterface() {}
diff --git a/Engine/Source/Runtime/Engine/Public/SceneView.h b/Engine/Source/Runtime/Engine/Public/SceneView.h
--- a/Engine/Source/Runtime/Engine/Public/SceneView.h
+++ b/Engine/Source/Runtime/Engine/Public/SceneView.h
@@ -84,6 +84,9 @@ struct FSceneViewInitOptions : public FSceneViewProjectionData
 	// Was there a camera cut this frame?
 	bool bInCameraCut;
 
+	// Whether world origin was rebased this frame?
+	bool bOriginOffsetThisFrame;
+
 #if WITH_EDITOR
 	// default to 0'th view index, which is a bitfield of 1
 	uint64 EditorViewBitflag;
@@ -109,6 +112,7 @@ struct FSceneViewInitOptions : public FSceneViewProjectionData
 		, LODDistanceFactor(1.0f)
 		, OverrideFarClippingPlaneDistance(-1.0f)
 		, bInCameraCut(false)
+		, bOriginOffsetThisFrame(false)
 #if WITH_EDITOR
 		, EditorViewBitflag(1)
 		, OverrideLODViewOrigin(ForceInitToZero)
@@ -297,6 +301,9 @@ public:
 
 	/** Whether we did a camera cut for this view this frame. */
 	bool bCameraCut;
+	
+	/** Whether world origin was rebased this frame. */
+	bool bOriginOffsetThisFrame;
 
 	// -1,-1 if not setup
 	FIntPoint CursorPos;
diff --git a/Engine/Source/Runtime/Renderer/Private/SceneOcclusion.cpp b/Engine/Source/Runtime/Renderer/Private/SceneOcclusion.cpp
--- a/Engine/Source/Runtime/Renderer/Private/SceneOcclusion.cpp
+++ b/Engine/Source/Runtime/Renderer/Private/SceneOcclusion.cpp
@@ -274,32 +274,6 @@ SIZE_T FSceneViewState::GetSizeBytes() const
 		+ PrimitiveOcclusionHistorySet.GetAllocatedSize();
 }
 
-void FSceneViewState::ApplyWorldOffset(FVector InOffset)
-{
-	// shift PrevView matrices
-	PrevViewMatrices.ViewMatrix.SetOrigin(
-		PrevViewMatrices.ViewMatrix.GetOrigin() - 
-		PrevViewMatrices.ViewMatrix.TransformVector(InOffset)
-		);
-	PrevViewMatrices.PreViewTranslation-= InOffset;
-	PrevViewMatrices.ViewOrigin+= InOffset;
-
-	// shift PendingPrevView matrices
-	PendingPrevViewMatrices.ViewMatrix.SetOrigin(
-		PendingPrevViewMatrices.ViewMatrix.GetOrigin() - 
-		PendingPrevViewMatrices.ViewMatrix.TransformVector(InOffset)
-		);
-	PendingPrevViewMatrices.PreViewTranslation-= InOffset;
-	PendingPrevViewMatrices.ViewOrigin+= InOffset;
-	
-	// shift PrevViewOclussionQuery matrix
-	PrevViewMatrixForOcclusionQuery.SetOrigin(
-		PrevViewMatrixForOcclusionQuery.GetOrigin() - 
-		PrevViewMatrixForOcclusionQuery.TransformVector(InOffset)
-		);
-	PrevViewOriginForOcclusionQuery+= InOffset;
-}
-
 class FOcclusionQueryIndexBuffer : public FIndexBuffer
 {
 public:
diff --git a/Engine/Source/Runtime/Renderer/Private/ScenePrivate.h b/Engine/Source/Runtime/Renderer/Private/ScenePrivate.h
--- a/Engine/Source/Runtime/Renderer/Private/ScenePrivate.h
+++ b/Engine/Source/Runtime/Renderer/Private/ScenePrivate.h
@@ -708,8 +708,6 @@ public:
 		return Ret;
 	}
 
-	virtual void ApplyWorldOffset(FVector InOffset) override;
-
 	// FDeferredCleanupInterface
 	virtual void FinishCleanup()
 	{
diff --git a/Engine/Source/Runtime/Renderer/Private/SceneVisibility.cpp b/Engine/Source/Runtime/Renderer/Private/SceneVisibility.cpp
--- a/Engine/Source/Runtime/Renderer/Private/SceneVisibility.cpp
+++ b/Engine/Source/Runtime/Renderer/Private/SceneVisibility.cpp
@@ -1119,6 +1119,13 @@ void FSceneRenderer::PreVisibilityFrameSetup(FRHICommandListImmediate& RHICmdLis
 
 		if ( ViewState )
 		{
+			// In case world origin was rebased, reset previous view transformations
+			if (View.bOriginOffsetThisFrame)
+			{
+				ViewState->PrevViewMatrices = View.ViewMatrices;
+				ViewState->PendingPrevViewMatrices = View.ViewMatrices;
+			}
+			
 			// determine if we are initializing or we should reset the persistent state
 			const float DeltaTime = View.Family->CurrentRealTime - ViewState->LastRenderTime;
 			const bool bFirstFrameOrTimeWasReset = DeltaTime < -0.0001f || ViewState->LastRenderTime < 0.0001f;
