Fixes crash when project module fails to load

#ttp 341587 - CRASH: Users will crash when attempting to open any code based project that has been manually associated with version 4.3

#branch UE4

#change If VisualLog is first called when the engine is shutting down, GEngine is NULL, guard against it

#change Also, checked bToggleAILogging is true before calling .Get() to save creating it needlessly

diff --git a/Engine/Source/Runtime/Engine/Private/CheatManager.cpp b/Engine/Source/Runtime/Engine/Private/CheatManager.cpp
--- a/Engine/Source/Runtime/Engine/Private/CheatManager.cpp
+++ b/Engine/Source/Runtime/Engine/Private/CheatManager.cpp
@@ -561,7 +561,7 @@ void UCheatManager::InitCheatManager()
 void UCheatManager::BeginDestroy()
 {
 #if ENABLE_VISUAL_LOG
-	if (FVisualLog::Get().IsRecording() && bToggleAILogging)
+	if (bToggleAILogging && FVisualLog::Get().IsRecording())
 	{
 		UWorld *World = GetWorld();
 		if (World)
diff --git a/Engine/Source/Runtime/Engine/Private/VisualLog.cpp b/Engine/Source/Runtime/Engine/Private/VisualLog.cpp
--- a/Engine/Source/Runtime/Engine/Private/VisualLog.cpp
+++ b/Engine/Source/Runtime/Engine/Private/VisualLog.cpp
@@ -348,18 +348,14 @@ TSharedPtr<FJsonValue> FActorsVisLog::ToJson() const
 // FVisualLog
 //----------------------------------------------------------------------//
 FVisualLog::FVisualLog()
+	: FileAr(NULL)
+	, StartRecordingTime(0.0f)
+	, bIsRecording(GEngine ? GEngine->bEnableVisualLogRecordingOnStart : false)
+	, bIsRecordingOnServer(false)
+	, bIsRecordingToFile(false)
+	, bIsAllBlocked(false)
 {
-	// just in case someone tries to VLog something when the engine is being destroyed
-	if (GEngine)
-	{
-		bIsRecording = GEngine->bEnableVisualLogRecordingOnStart;
-		bIsRecordingOnServer = false;
-
-		FileAr = NULL;
-		bIsRecordingToFile = false;
-		bIsAllBlocked = false;
-		Whitelist.Reserve(10);
-	}
+	Whitelist.Reserve(10);
 }
 
 FVisualLog::~FVisualLog()
@@ -472,7 +468,7 @@ void FVisualLog::SetIsRecording(bool NewRecording, bool bRecordToFile)
 	if (bIsRecording)
 	{
 		bIsRecordingToFile = bRecordToFile;
-		StartRecordingTime = GWorld ? GWorld->TimeSeconds : 0;
+		StartRecordingTime = GWorld ? GWorld->TimeSeconds : 0.0f;
 	}
 }
 
