Move ERadialImpulseFalloff to EngineTypes.h (should fix build)

diff --git a/Engine/Source/Runtime/Engine/Classes/Components/PrimitiveComponent.h b/Engine/Source/Runtime/Engine/Classes/Components/PrimitiveComponent.h
--- a/Engine/Source/Runtime/Engine/Classes/Components/PrimitiveComponent.h
+++ b/Engine/Source/Runtime/Engine/Classes/Components/PrimitiveComponent.h
@@ -44,16 +44,7 @@ struct FStreamingTexturePrimitiveInfo
 	float TexelFactor;
 };
 
-/** Enum for controlling the falloff of strength of a radial impulse as a function of distance from Origin. */
-UENUM()
-enum ERadialImpulseFalloff
-{
-	/** Impulse is a constant strength, up to the limit of its range. */
-	RIF_Constant,
-	/** Impulse should get linearly weaker the further from origin. */
-	RIF_Linear,
-	RIF_MAX,
-};
+
 
 UENUM()
 enum ECanBeCharacterBase
diff --git a/Engine/Source/Runtime/Engine/Classes/Engine/EngineTypes.h b/Engine/Source/Runtime/Engine/Classes/Engine/EngineTypes.h
--- a/Engine/Source/Runtime/Engine/Classes/Engine/EngineTypes.h
+++ b/Engine/Source/Runtime/Engine/Classes/Engine/EngineTypes.h
@@ -666,6 +666,16 @@ private:
 	friend class UCollisionProfile;
 };
 
+/** Enum for controlling the falloff of strength of a radial impulse as a function of distance from Origin. */
+UENUM()
+enum ERadialImpulseFalloff
+{
+	/** Impulse is a constant strength, up to the limit of its range. */
+	RIF_Constant,
+	/** Impulse should get linearly weaker the further from origin. */
+	RIF_Linear,
+	RIF_MAX,
+};
 
 /** Presets of values used in considering when put this body to sleep. */
 UENUM()
@@ -869,6 +879,7 @@ struct FCollisionImpactData
 	void SwapContactOrders();
 };
 
+
 /** Struct used to hold effects for destructible damage events */
 USTRUCT()
 struct FFractureEffect
