Added a class and struct forward decl. to soundbase.h

diff --git a/Engine/Source/Runtime/Engine/Classes/Sound/SoundBase.h b/Engine/Source/Runtime/Engine/Classes/Sound/SoundBase.h
--- a/Engine/Source/Runtime/Engine/Classes/Sound/SoundBase.h
+++ b/Engine/Source/Runtime/Engine/Classes/Sound/SoundBase.h
@@ -11,7 +11,10 @@
 struct FActiveSound;
 struct FSoundParseParameters;
 struct FWaveInstance;
-class  USoundClass;
+struct FAttenuationSettings; 
+class USoundClass;
+class USoundAttenuation;
+
 
 UENUM()
 namespace EMaxConcurrentResolutionRule
