CIS fix - Removed fwd-declared enum & added include to MovementComponent

diff --git a/Engine/Source/Runtime/Engine/Classes/GameFramework/MovementComponent.h b/Engine/Source/Runtime/Engine/Classes/GameFramework/MovementComponent.h
--- a/Engine/Source/Runtime/Engine/Classes/GameFramework/MovementComponent.h
+++ b/Engine/Source/Runtime/Engine/Classes/GameFramework/MovementComponent.h
@@ -6,9 +6,9 @@
 
 #pragma once
 #include "Components/ActorComponent.h"
+#include "Components/PrimitiveComponent.h"
 #include "MovementComponent.generated.h"
 
-class UPrimitiveComponent;
 struct FCollisionQueryParams;
 struct FCollisionResponseParams;
 struct FCollisionShape;
@@ -234,7 +234,7 @@ public:
 	 * @param	Strength	The strength of the force
 	 * @param	Falloff		The falloff from the force's origin
 	 */
-	virtual void AddRadialForce(const FVector& Origin, float Radius, float Strength, enum ERadialImpulseFalloff Falloff);
+	virtual void AddRadialForce(const FVector& Origin, float Radius, float Strength, ERadialImpulseFalloff Falloff);
 
 	/**
 	 * Adds impulse from radial force components.
@@ -244,7 +244,7 @@ public:
 	 * @param	Falloff		The falloff from the force's origin
 	 * @param	bVelChange	If true, the Strength is taken as a change in velocity instead of an impulse (ie. mass will have no affect).
 	 */
-	virtual void AddRadialImpulse(const FVector& Origin, float Radius, float Strength, enum ERadialImpulseFalloff Falloff, bool bVelChange);
+	virtual void AddRadialImpulse(const FVector& Origin, float Radius, float Strength, ERadialImpulseFalloff Falloff, bool bVelChange);
 
 	/**
 	 * Sets the normal of the plane that constrains movement, if plane constraint is enabled.
