Added new InstancedStaticMeshComponent functions RemoveInstance and SetCullDistances and exposed FHitResult.Item to blueprints

diff --git a/Engine/Source/Runtime/Engine/Classes/Components/InstancedStaticMeshComponent.h b/Engine/Source/Runtime/Engine/Classes/Components/InstancedStaticMeshComponent.h
--- a/Engine/Source/Runtime/Engine/Classes/Components/InstancedStaticMeshComponent.h
+++ b/Engine/Source/Runtime/Engine/Classes/Components/InstancedStaticMeshComponent.h
@@ -96,6 +96,10 @@ class ENGINE_API UInstancedStaticMeshComponent : public UStaticMeshComponent
 	UFUNCTION(BlueprintCallable, Category = "Components|InstancedStaticMesh")
 	bool UpdateInstanceTransform(int32 InstanceIndex, const FTransform& NewInstanceTransform, bool bWorldSpace=false);
 
+	/** Remove the instance specified. Returns True on success. */
+	UFUNCTION(BlueprintCallable, Category = "Components|InstancedStaticMesh")
+	bool RemoveInstance(int32 InstanceIndex);
+	
 	/** Clear all instances being rendered by this component */
 	UFUNCTION(BlueprintCallable, Category="Components|InstancedStaticMesh")
 	void ClearInstances();
@@ -103,7 +107,11 @@ class ENGINE_API UInstancedStaticMeshComponent : public UStaticMeshComponent
 	/** Get the number of instances in this component */
 	UFUNCTION(BlueprintCallable, Category = "Components|InstancedStaticMesh")
 	int32 GetInstanceCount() const;
-	
+
+	/** Sets the fading start and culling end distances for this component. */
+	UFUNCTION(BlueprintCallable, Category = "Components|InstancedStaticMesh")
+	void SetCullDistances(int32 StartCullDistance, int32 EndCullDistance);
+
 	virtual bool ShouldCreatePhysicsState() const;
 
 public:
@@ -159,6 +167,9 @@ private:
 	/** Initializes the body instance for the specified instance of the static mesh*/
 	void InitInstanceBody(int32 InstanceIdx, FBodyInstance* BodyInstance);
 
+	/** Creates body instances for all instances owned by this component */
+	void CreateAllInstanceBodies();
+
 	/** Terminate all body instances owned by this component */
 	void ClearAllInstanceBodies();
 
diff --git a/Engine/Source/Runtime/Engine/Classes/Kismet/GameplayStatics.h b/Engine/Source/Runtime/Engine/Classes/Kismet/GameplayStatics.h
--- a/Engine/Source/Runtime/Engine/Classes/Kismet/GameplayStatics.h
+++ b/Engine/Source/Runtime/Engine/Classes/Kismet/GameplayStatics.h
@@ -384,9 +384,10 @@ class ENGINE_API UGameplayStatics : public UBlueprintFunctionLibrary
 	 * @param HitActor
 	 * @param HitComponent
 	 * @param HitBoneName
+	 * @param HitItem		Primitive-specific data recording which item in the primitive was hit
 	 */
 	UFUNCTION(BlueprintPure, Category = "Collision", meta=(NativeBreakFunc))
-	static void BreakHitResult(const struct FHitResult& Hit, FVector& Location, FVector& Normal, FVector& ImpactPoint, FVector& ImpactNormal, class UPhysicalMaterial*& PhysMat, class AActor*& HitActor, class UPrimitiveComponent*& HitComponent, FName& HitBoneName);
+	static void BreakHitResult(const struct FHitResult& Hit, FVector& Location, FVector& Normal, FVector& ImpactPoint, FVector& ImpactNormal, class UPhysicalMaterial*& PhysMat, class AActor*& HitActor, class UPrimitiveComponent*& HitComponent, FName& HitBoneName, int32& HitItem);
 
 	/** Returns the EPhysicalSurface type of the given Hit. 
 	 * To edit surface type for your project, use ProjectSettings/Physics/PhysicalSurface section
diff --git a/Engine/Source/Runtime/Engine/Private/GameplayStatics.cpp b/Engine/Source/Runtime/Engine/Private/GameplayStatics.cpp
--- a/Engine/Source/Runtime/Engine/Private/GameplayStatics.cpp
+++ b/Engine/Source/Runtime/Engine/Private/GameplayStatics.cpp
@@ -557,7 +557,7 @@ UParticleSystemComponent* UGameplayStatics::SpawnEmitterAttached(UParticleSystem
 	return PSC;
 }
 
-void UGameplayStatics::BreakHitResult(const FHitResult& Hit, FVector& Location, FVector& Normal, FVector& ImpactPoint, FVector& ImpactNormal, UPhysicalMaterial*& PhysMat, AActor*& HitActor, UPrimitiveComponent*& HitComponent, FName& HitBoneName)
+void UGameplayStatics::BreakHitResult(const FHitResult& Hit, FVector& Location, FVector& Normal, FVector& ImpactPoint, FVector& ImpactNormal, UPhysicalMaterial*& PhysMat, AActor*& HitActor, UPrimitiveComponent*& HitComponent, FName& HitBoneName, int32& HitItem)
 {
 	Location = Hit.Location;
 	Normal = Hit.Normal;
@@ -567,6 +567,7 @@ void UGameplayStatics::BreakHitResult(const FHitResult& Hit, FVector& Location,
 	HitActor = Hit.GetActor();
 	HitComponent = Hit.Component.Get();
 	HitBoneName = Hit.BoneName;
+	HitItem = Hit.Item;
 }
 
 EPhysicalSurface UGameplayStatics::GetSurfaceType(const struct FHitResult& Hit)
diff --git a/Engine/Source/Runtime/Engine/Private/InstancedStaticMesh.cpp b/Engine/Source/Runtime/Engine/Private/InstancedStaticMesh.cpp
--- a/Engine/Source/Runtime/Engine/Private/InstancedStaticMesh.cpp
+++ b/Engine/Source/Runtime/Engine/Private/InstancedStaticMesh.cpp
@@ -1242,6 +1242,18 @@ void UInstancedStaticMeshComponent::InitInstanceBody(int32 InstanceIdx, FBodyIns
 #endif //WITH_PHYSX
 }
 
+void UInstancedStaticMeshComponent::CreateAllInstanceBodies()
+{
+	int32 NumBodies = PerInstanceSMData.Num();
+	InstanceBodies.Init(NumBodies);
+
+	for (int32 i = 0; i < NumBodies; ++i)
+	{
+		InstanceBodies[i] = new FBodyInstance;
+		InitInstanceBody(i, InstanceBodies[i]);
+	}
+}
+
 void UInstancedStaticMeshComponent::ClearAllInstanceBodies()
 {
 	for (int32 i = 0; i < InstanceBodies.Num(); i++)
@@ -1273,14 +1285,7 @@ void UInstancedStaticMeshComponent::CreatePhysicsState()
 #endif
 
 	// Create all the bodies.
-	int32 NumBodies = PerInstanceSMData.Num();
-	InstanceBodies.Init(NumBodies);
-
-	for(int32 i=0; i < NumBodies; ++i)
-	{
-		InstanceBodies[i] = new FBodyInstance;
-		InitInstanceBody(i, InstanceBodies[i]);
-	}
+	CreateAllInstanceBodies();
 
 	USceneComponent::CreatePhysicsState();
 }
@@ -1559,6 +1564,38 @@ bool UInstancedStaticMeshComponent::ShouldCreatePhysicsState() const
 	return IsRegistered() && (bAlwaysCreatePhysicsState || IsCollisionEnabled());
 }
 
+bool UInstancedStaticMeshComponent::RemoveInstance(int32 InstanceIndex)
+{
+	if (!PerInstanceSMData.IsValidIndex(InstanceIndex))
+	{
+		return false;
+	}
+
+	// remove instance
+	PerInstanceSMData.RemoveAt(InstanceIndex);
+
+#if WITH_EDITOR
+	// remove selection flag if array is filled in
+	if (SelectedInstances.IsValidIndex(InstanceIndex))
+	{
+		SelectedInstances.RemoveAt(InstanceIndex);
+	}
+#endif
+
+	// update the physics state
+	if (bPhysicsStateCreated)
+	{
+		// TODO: it may be possible to instead just update the BodyInstanceIndex for all bodies after the removed instance. 
+		ClearAllInstanceBodies();
+		CreateAllInstanceBodies();
+	}
+
+	// Indicate we need to update render state to reflect changes
+	MarkRenderStateDirty();
+
+	return true;
+}
+
 void UInstancedStaticMeshComponent::ClearInstances()
 {
 	// Clear all the per-instance data
@@ -1575,6 +1612,13 @@ int32 UInstancedStaticMeshComponent::GetInstanceCount() const
 	return PerInstanceSMData.Num();
 }
 
+void UInstancedStaticMeshComponent::SetCullDistances(int32 StartCullDistance, int32 EndCullDistance)
+{
+	InstanceStartCullDistance = StartCullDistance;
+	InstanceEndCullDistance = EndCullDistance;
+	MarkRenderStateDirty();
+}
+
 void UInstancedStaticMeshComponent::SetupNewInstanceData(FInstancedStaticMeshInstanceData& InOutNewInstanceData, int32 InInstanceIndex, const FTransform& InInstanceTransform)
 {
 	InOutNewInstanceData.Transform = InInstanceTransform.ToMatrixWithScale();
