Fix tooltip on the materials array on MeshComponent

diff --git a/Engine/Source/Runtime/Engine/Classes/Components/MeshComponent.h b/Engine/Source/Runtime/Engine/Classes/Components/MeshComponent.h
--- a/Engine/Source/Runtime/Engine/Classes/Components/MeshComponent.h
+++ b/Engine/Source/Runtime/Engine/Classes/Components/MeshComponent.h
@@ -11,7 +11,7 @@ class ENGINE_API UMeshComponent : public UPrimitiveComponent
 	GENERATED_UCLASS_BODY()
 
 	/** Per-Component material overrides.  These must NOT be set directly or a race condition can occur between GC and the rendering thread. */
-	UPROPERTY(EditAnywhere, BlueprintReadOnly, Category=Rendering)
+	UPROPERTY(EditAnywhere, BlueprintReadOnly, Category=Rendering, Meta=(ToolTip="Material overrides."))
 	TArray<class UMaterialInterface*> Materials;
 
 
