Fix for assert when duplicating timelines
#TTP 342038 - EDITOR: CRASH: TappyChicken: Duplicating Timeline_LayEgg node crashes the editor.
#Branch UE4
#Proj Engine

#Change Updated the curve name collision detection to be more agressive to catch redirect classes.

#info Assert was caused by MakeUniqueCurveName passing in the class to detect name collisions. In this instance it failed to recognise a collision with another redirected CurveFloat because it couldn't associate the redirector class with the CurveFloat class.

diff --git a/Engine/Source/Runtime/Engine/Private/TimelineTemplate.cpp b/Engine/Source/Runtime/Engine/Private/TimelineTemplate.cpp
--- a/Engine/Source/Runtime/Engine/Private/TimelineTemplate.cpp
+++ b/Engine/Source/Runtime/Engine/Private/TimelineTemplate.cpp
@@ -183,9 +183,8 @@ void UTimelineTemplate::RemoveMetaData(const FName& Key)
 FString UTimelineTemplate::MakeUniqueCurveName(UObject* Obj, UObject* InOuter)
 {
 	FString OriginalName = Obj->GetName();
-	UClass* Class = Obj->GetClass();
 	FName TestName(*OriginalName);
-	while(StaticFindObjectFast(Class, InOuter, TestName))
+	while(StaticFindObjectFast(NULL, InOuter, TestName))
 	{
 		TestName = FName(*OriginalName, TestName.GetNumber()+1);
 	}
