Allow exec functions on GameInstance to work from the console

diff --git a/Engine/Source/Runtime/Engine/Private/Player.cpp b/Engine/Source/Runtime/Engine/Private/Player.cpp
--- a/Engine/Source/Runtime/Engine/Private/Player.cpp
+++ b/Engine/Source/Runtime/Engine/Private/Player.cpp
@@ -9,6 +9,7 @@
 #include "SubtitleManager.h"
 #include "Net/UnrealNetwork.h"
 #include "GameFramework/HUD.h"
+#include "Engine/GameInstance.h"
 
 #include "RenderCore.h"
 #include "ColorList.h"
@@ -53,6 +54,10 @@ bool UPlayer::Exec( UWorld* InWorld, const TCHAR* Cmd,FOutputDevice& Ar)
 		{
 			return true;
 		}
+		else if( World->GetGameInstance() && World->GetGameInstance()->ProcessConsoleExec(Cmd, Ar, PCPawn) )
+		{
+			return true;
+		}
 		else if( World->GetAuthGameMode() && World->GetAuthGameMode()->ProcessConsoleExec(Cmd,Ar,PCPawn) )
 		{
 			return true;
