[341253] Blueprints - Apply custom scale to non-native root components at construction time.

This matches what's done for native root components when Blueprint actor instances are constructed.

diff --git a/Engine/Source/Runtime/Engine/Private/SCS_Node.cpp b/Engine/Source/Runtime/Engine/Private/SCS_Node.cpp
--- a/Engine/Source/Runtime/Engine/Private/SCS_Node.cpp
+++ b/Engine/Source/Runtime/Engine/Private/SCS_Node.cpp
@@ -55,7 +55,8 @@ void USCS_Node::ExecuteNodeOnActor(AActor* Actor, USceneComponent* ParentCompone
 			// If NULL is passed in, we are the root, so set transform and assign as RootComponent on Actor
 			if (ParentComponent == NULL || (ParentComponent && ParentComponent->IsPendingKill()))
 			{
-				NewSceneComp->SetWorldTransform(*RootTransform);
+				// Note: We use the scale vector from the component template rather than the passed in scale in this case (to match what happens with a native root)
+				NewSceneComp->SetWorldTransform(FTransform(RootTransform->GetRotation(), RootTransform->GetLocation(), NewSceneComp->RelativeScale3D));
 				Actor->SetRootComponent(NewSceneComp);
 			}
 			// Otherwise, attach to parent component passed in
