Make PlayerController log category usable from other files

diff --git a/Engine/Source/Runtime/Engine/Classes/GameFramework/PlayerController.h b/Engine/Source/Runtime/Engine/Classes/GameFramework/PlayerController.h
--- a/Engine/Source/Runtime/Engine/Classes/GameFramework/PlayerController.h
+++ b/Engine/Source/Runtime/Engine/Classes/GameFramework/PlayerController.h
@@ -17,6 +17,8 @@ class FPrimitiveComponentId;
 /** delegate used to override default viewport audio listener position calculated from camera */
 DECLARE_DELEGATE_ThreeParams(FGetAudioListenerPos, FVector& /*Location*/, FVector& /*ProjFront*/, FVector& /*ProjRight*/);
 
+DECLARE_LOG_CATEGORY_EXTERN(LogPlayerController, Log, All);
+
 UENUM()
 namespace EDynamicForceFeedbackAction
 {
diff --git a/Engine/Source/Runtime/Engine/Private/PlayerController.cpp b/Engine/Source/Runtime/Engine/Private/PlayerController.cpp
--- a/Engine/Source/Runtime/Engine/Private/PlayerController.cpp
+++ b/Engine/Source/Runtime/Engine/Private/PlayerController.cpp
@@ -24,7 +24,7 @@
 #include "ContentStreaming.h"
 #include "GameFramework/PawnMovementComponent.h"
 
-DEFINE_LOG_CATEGORY_STATIC(LogPlayerController, Log, All);
+DEFINE_LOG_CATEGORY(LogPlayerController);
 
 const float RetryClientRestartThrottleTime = 0.5f;
 const float RetryServerAcknowledgeThrottleTime = 0.25f;
