Document InputComponent
Add Bind functions for Axis, AxisKey, and VectorAxis that do not bind to a delegate (for binding with the intention of using the appropriate GetValue function)
Add warnings in the Axis GetValue functions when asking for an Axis that is not bound in the input component

diff --git a/Engine/Source/Runtime/Engine/Classes/Components/InputComponent.h b/Engine/Source/Runtime/Engine/Classes/Components/InputComponent.h
--- a/Engine/Source/Runtime/Engine/Classes/Components/InputComponent.h
+++ b/Engine/Source/Runtime/Engine/Classes/Components/InputComponent.h
@@ -6,6 +6,7 @@
 #include "Runtime/InputCore/Classes/InputCoreTypes.h"
 #include "InputComponent.generated.h"
 
+// Utility delegate class to allow binding to either a C++ function or a blueprint script delegate
 template<class DelegateType, class DynamicDelegateType>
 struct TInputUnifiedDelegate
 {
@@ -69,6 +70,7 @@ protected:
 	DynamicDelegateType FuncDynDelegate;
 };
 
+// An Input Chord is a key and the modifier keys that are to be held with it
 USTRUCT()
 struct FInputChord
 {
@@ -97,6 +99,8 @@ struct FInputChord
 		Masks
 	};
 
+	// Returns the relationship between this chord and another.  A chord is considered masking
+	// when it has all the same modifier keys as another chord plus more
 	RelationshipType GetRelationship(const FInputChord& OtherChord) const;
 
 	FInputChord()
@@ -193,7 +197,7 @@ struct FInputKeyBinding : public FInputBinding
 	/** Raw key to bind to */
 	FInputChord Chord;
 
-	/** Key event to bind it to, e.g. pressed, released, dblclick */
+	/** Key event to bind it to, e.g. pressed, released, double click */
 	TEnumAsByte<EInputEvent> KeyEvent;
 
 	FInputActionUnifiedDelegate KeyDelegate;
@@ -236,7 +240,7 @@ struct FInputTouchUnifiedDelegate : public TInputUnifiedDelegate<FInputTouchHand
 
 struct FInputTouchBinding : public FInputBinding
 {
-	/** Key event to bind it to, e.g. pressed, released, dblclick */
+	/** Key event to bind it to, e.g. pressed, released, double click */
 	TEnumAsByte<EInputEvent> KeyEvent;
 
 	FInputTouchUnifiedDelegate TouchDelegate;
@@ -415,18 +419,22 @@ struct FInputGestureBinding : public FInputBinding
 
 /** Creates and adds a new axis binding. */
 #define BIND_AXIS(Comp, AxisName, Delegate)	\
+	EMIT_DEPRECATED_WARNING_MESSAGE("BIND_AXIS macro is deprecated. Please use UInputComponent::BindAxis function instead.") \
 	Comp->BindAxis(AxisName, this, Delegate)
 
 /** Creates and adds a new action binding. */
 #define BIND_ACTION(Comp, ActionName, KeyEvent, Delegate)	\
+	EMIT_DEPRECATED_WARNING_MESSAGE("BIND_ACTION macro is deprecated. Please use UInputComponent::BindAction function instead.") \
 	Comp->BindAction(ActionName, KeyEvent, this, Delegate);
 
 /** Creates and adds a new axis binding. */
 #define BIND_AXIS_OBJ(Comp, AxisName, Delegate, DelegateObj)	\
+	EMIT_DEPRECATED_WARNING_MESSAGE("BIND_AXIS_OBJ macro is deprecated. Please use UInputComponent::BindAxis function instead.") \
 	Comp->BindAxis(AxisName, DelegateObj, Delegate)
 
 /** Same as BIND_ACTION, except can be used to bind to a delegate in a specific object. */
 #define BIND_ACTION_OBJ(Comp, ActionName, KeyEvent, Delegate, DelegateObj)	\
+	EMIT_DEPRECATED_WARNING_MESSAGE("BIND_ACTION_OBJ macro is deprecated. Please use UInputComponent::BindAction function instead.") \
 	Comp->BindAction(ActionName, KeyEvent, DelegateObj, Delegate);
 
 
@@ -441,6 +449,9 @@ namespace EControllerAnalogStick
 	};
 }
 
+// Input Component is a component that can be added to an Actor that allows the Actor to bind various kinds of input events
+// to delegate functions.  Input components are processed from a stack managed by the PlayerController and each bind can 
+// consume the input event preventing other components on the input stack from processing the input.
 UCLASS(transient, config=Input, hidecategories=(Activation, "Components|Activation"))
 class ENGINE_API UInputComponent : public UActorComponent
 {
@@ -460,12 +471,16 @@ public:
 	// Whether any components lower on the input stack should be allowed to receive input
 	uint32 bBlockInput:1;
 
+	// Returns the current value of an Axis defined in the project settings bound in this input component
 	float GetAxisValue(const FName AxisName) const;
+
+	// Returns the 
 	float GetAxisKeyValue(const FKey AxisKey) const;
 	FVector GetVectorAxisValue(const FKey AxisKey) const;
 
 	bool HasBindings() const;
 
+	// Binds a delegate function to an Action defined in the project settings
 	template<class UserClass>
 	FInputActionBinding& BindAction(const FName ActionName, const EInputEvent KeyEvent, UserClass* Object, typename FInputActionHandlerSignature::TUObjectMethodDelegate< UserClass >::FMethodPtr Func)
 	{
@@ -474,6 +489,7 @@ public:
 		return AddActionBinding(AB);
 	}
 
+	// Binds a delegate function an Axis defined in the project settings
 	template<class UserClass>
 	FInputAxisBinding& BindAxis(const FName AxisName, UserClass* Object, typename FInputAxisHandlerSignature::TUObjectMethodDelegate< UserClass >::FMethodPtr Func)
 	{
@@ -483,6 +499,15 @@ public:
 		return AxisBindings.Last();
 	}
 
+	// Indicates that the InputComponent is interested in knowing the Axis value (via GetAxisValue) but does not want a delegate function called each frame
+	FInputAxisBinding& BindAxis(const FName AxisName)
+	{
+		FInputAxisBinding AB( AxisName );
+		AxisBindings.Add(AB);
+		return AxisBindings.Last();
+	}
+
+	// Binds a delegate function for an axis key (e.g. Mouse X)
 	template<class UserClass>
 	FInputAxisKeyBinding& BindAxisKey(const FKey AxisKey, UserClass* Object, typename FInputAxisHandlerSignature::TUObjectMethodDelegate< UserClass >::FMethodPtr Func)
 	{
@@ -492,6 +517,15 @@ public:
 		return AxisKeyBindings.Last();
 	}
 
+	// Indicates that the InputComponent is interested in knowing/consuming an axis key's value (via GetAxisKeyValue) but does not want a delegate function called each frame
+	FInputAxisKeyBinding& BindAxisKey(const FKey AxisKey)
+	{
+		FInputAxisKeyBinding AB(AxisKey);
+		AxisKeyBindings.Add(AB);
+		return AxisKeyBindings.Last();
+	}
+
+	// Binds a delegate function to a vector axis key (e.g. Tilt)
 	template<class UserClass>
 	FInputVectorAxisBinding& BindVectorAxis(const FKey AxisKey, UserClass* Object, typename FInputVectorAxisHandlerSignature::TUObjectMethodDelegate< UserClass >::FMethodPtr Func)
 	{
@@ -501,6 +535,15 @@ public:
 		return VectorAxisBindings.Last();
 	}
 
+	// Indicates that the InputComponent is interested in knowing/consuming a vector axis key's value (via GetVectorAxisKeyValue) but does not want a delegate function called each frame
+	FInputVectorAxisBinding& BindVectorAxis(const FKey AxisKey)
+	{
+		FInputVectorAxisBinding AB(AxisKey);
+		VectorAxisBindings.Add(AB);
+		return VectorAxisBindings.Last();
+	}
+
+	// Binds a chord event to a delegate function
 	template<class UserClass>
 	FInputKeyBinding& BindKey(const FInputChord Chord, const EInputEvent KeyEvent, UserClass* Object, typename FInputActionHandlerSignature::TUObjectMethodDelegate< UserClass >::FMethodPtr Func)
 	{
@@ -510,12 +553,14 @@ public:
 		return KeyBindings.Last();
 	}
 
+	// Binds a key event to a delegate function
 	template<class UserClass>
 	FInputKeyBinding& BindKey(const FKey Key, const EInputEvent KeyEvent, UserClass* Object, typename FInputActionHandlerSignature::TUObjectMethodDelegate< UserClass >::FMethodPtr Func)
 	{
 		return BindKey(FInputChord(Key, false, false, false, false), KeyEvent, Object, Func);
 	}
 
+	// Binds this input component to touch events
 	template<class UserClass>
 	FInputTouchBinding& BindTouch(const EInputEvent KeyEvent, UserClass* Object, typename FInputTouchHandlerSignature::TUObjectMethodDelegate< UserClass >::FMethodPtr Func)
 	{
@@ -525,6 +570,7 @@ public:
 		return TouchBindings.Last();
 	}
 
+	// Binds a gesture event to a delegate function
 	template<class UserClass>
 	FInputGestureBinding& BindGesture(const FKey GestureKey, UserClass* Object, typename FInputGestureHandlerSignature::TUObjectMethodDelegate< UserClass >::FMethodPtr Func)
 	{
@@ -534,6 +580,7 @@ public:
 		return GestureBindings.Last();
 	}
 
+	// Utility functions for managing action bindings array as it needs to maintain knowledge of the pairing state
 	FInputActionBinding& AddActionBinding(const FInputActionBinding& Binding);
 	void RemoveActionBinding(const int32 BindingIndex);
 	int32 GetNumActionBindings() const;
diff --git a/Engine/Source/Runtime/Engine/Private/InputComponent.cpp b/Engine/Source/Runtime/Engine/Private/InputComponent.cpp
--- a/Engine/Source/Runtime/Engine/Private/InputComponent.cpp
+++ b/Engine/Source/Runtime/Engine/Private/InputComponent.cpp
@@ -40,47 +40,69 @@ UInputComponent::UInputComponent(const class FPostConstructInitializeProperties&
 float UInputComponent::GetAxisValue(const FName AxisName) const
 {
 	float AxisValue = 0.f;
+	bool bFound = false;
 
 	for (const FInputAxisBinding& AxisBinding : AxisBindings)
 	{
 		if (AxisBinding.AxisName == AxisName)
 		{
+			bFound = true;
 			AxisValue = AxisBinding.AxisValue;
 			break;
 		}
 	}
 
+	if (!bFound)
+	{
+		UE_LOG(LogPlayerController, Warning, TEXT("Request for value of axis '%s' returning 0 as it is not bound on this input component."), *AxisName.ToString());
+	}
+
 	return AxisValue;
 }
 
 float UInputComponent::GetAxisKeyValue(const FKey AxisKey) const
 {
 	float AxisValue = 0.f;
+	bool bFound = false;
 
 	for (const FInputAxisKeyBinding& AxisBinding : AxisKeyBindings)
 	{
 		if (AxisBinding.AxisKey == AxisKey)
 		{
+			bFound = true;
 			AxisValue = AxisBinding.AxisValue;
 			break;
 		}
 	}
 
+	if (!bFound)
+	{
+		UE_LOG(LogPlayerController, Warning, TEXT("Request for value of axis key '%s' returning 0 as it is not bound on this input component."), *AxisKey.ToString());
+	}
+
 	return AxisValue;
 }
 
 FVector UInputComponent::GetVectorAxisValue(const FKey AxisKey) const
 {
 	FVector AxisValue;
+	bool bFound = false;
+
 	for (const FInputVectorAxisBinding& AxisBinding : VectorAxisBindings)
 	{
 		if (AxisBinding.AxisKey == AxisKey)
 		{
 			AxisValue = AxisBinding.AxisValue;
+			bFound = true;
 			break;
 		}
 	}
 
+	if (!bFound)
+	{
+		UE_LOG(LogPlayerController, Warning, TEXT("Request for value of vector axis key '%s' returning 0 as it is not bound on this input component."), *AxisKey.ToString());
+	}
+
 	return AxisValue;
 }
 
