Fixed up Persona tooltips and interactive tutorial after NUX changes. Added tooltips to skeletal control nodes.

diff --git a/Engine/Source/Runtime/Engine/Classes/Animation/AnimSequence.h b/Engine/Source/Runtime/Engine/Classes/Animation/AnimSequence.h
--- a/Engine/Source/Runtime/Engine/Classes/Animation/AnimSequence.h
+++ b/Engine/Source/Runtime/Engine/Classes/Animation/AnimSequence.h
@@ -453,6 +453,7 @@ class UAnimSequence : public UAnimSequenceBase
 	UPROPERTY()
 	int32 EncodingPkgVersion;
 
+	/** Base pose to use when retargeting */
 	UPROPERTY(EditAnywhere, AssetRegistrySearchable, Category=Animation)
 	FName RetargetSource;
 
