Fix multiplayer crash in APlayerController::EndPlay

diff --git a/Engine/Source/Runtime/Engine/Private/PlayerController.cpp b/Engine/Source/Runtime/Engine/Private/PlayerController.cpp
--- a/Engine/Source/Runtime/Engine/Private/PlayerController.cpp
+++ b/Engine/Source/Runtime/Engine/Private/PlayerController.cpp
@@ -1221,16 +1221,20 @@ void APlayerController::PawnLeavingGame()
 
 void APlayerController::EndPlay(const EEndPlayReason::Type EndPlayReason)
 {
-	if (VirtualJoystick.IsValid())
+	ULocalPlayer* LocalPlayer = Cast<ULocalPlayer>(Player);
+	if (LocalPlayer)
 	{
-		CastChecked<ULocalPlayer>(Player)->ViewportClient->RemoveViewportWidgetContent(VirtualJoystick.ToSharedRef());
-	}
+		if (VirtualJoystick.IsValid())
+		{
+			LocalPlayer->ViewportClient->RemoveViewportWidgetContent(VirtualJoystick.ToSharedRef());
+		}
 
-	// Stop any force feedback effects that may be active
-	IForceFeedbackSystem* ForceFeedbackSystem = FSlateApplication::Get().GetForceFeedbackSystem();
-	if (ForceFeedbackSystem)
-	{
-		ForceFeedbackSystem->SetChannelValues(CastChecked<ULocalPlayer>(Player)->ControllerId, FForceFeedbackValues());
+		// Stop any force feedback effects that may be active
+		IForceFeedbackSystem* ForceFeedbackSystem = FSlateApplication::Get().GetForceFeedbackSystem();
+		if (ForceFeedbackSystem)
+		{
+			ForceFeedbackSystem->SetChannelValues(LocalPlayer->ControllerId, FForceFeedbackValues());
+		}
 	}
 
 	Super::EndPlay(EndPlayReason);
