Modified the blueprint VER_UE4_EDITORONLY_BLUEPRINTS migration to log warnings about blueprints ( in the blueprint log ) that have been migrated to make it easier for the user to determine and re-save which blueprints change to fix import curve errors in an attempt to clarify the issues. Additionally added a loading/saving editor preference to dirty the migrated blueprints to enable rapid fix up of migration issues.
#TTP 340959 - CRASH: UE4: BLUEPRINTS: Editor crashing when compiling CurveFloat in a Timeline
#Branch UE4
#Proj  Runtime, UnrealEd

diff --git a/Engine/Config/BaseEditorUserSettings.ini b/Engine/Config/BaseEditorUserSettings.ini
--- a/Engine/Config/BaseEditorUserSettings.ini
+++ b/Engine/Config/BaseEditorUserSettings.ini
@@ -135,6 +135,8 @@ bLoadDefaultLevelAtStartup=True
 ; True if we should automatically reimport textures when a change to source content is detected
 bAutoReimportTextures=False
 bAutoReimportCSV=False
+; Whether to mark blueprints dirty if they are automatically migrated during loads
+bDirtyMigratedBlueprints=False
 ; Whether to automatically save after a time interval
 bAutoSaveEnable=True
 ; Whether to automatically save maps during an autosave
diff --git a/Engine/Source/Editor/UnrealEd/Classes/Settings/EditorLoadingSavingSettings.h b/Engine/Source/Editor/UnrealEd/Classes/Settings/EditorLoadingSavingSettings.h
--- a/Engine/Source/Editor/UnrealEd/Classes/Settings/EditorLoadingSavingSettings.h
+++ b/Engine/Source/Editor/UnrealEd/Classes/Settings/EditorLoadingSavingSettings.h
@@ -13,7 +13,7 @@
 /**
  * Implements the Level Editor's loading and saving settings.
  */
-UCLASS(config=EditorUserSettings, autoexpandcategories=(AutoSave, AutoReimport))
+UCLASS(config=EditorUserSettings, autoexpandcategories=(AutoSave, AutoReimport, Blueprints))
 class UNREALED_API UEditorLoadingSavingSettings
 	: public UObject
 {
@@ -43,6 +43,12 @@ public:
 	UPROPERTY(EditAnywhere, config, Category=AutoReimport, meta=(DisplayName="Auto Reimport CSVData"))
 	bool bAutoReimportCSV;
 
+public:
+
+	/** Whether to mark blueprints dirty if they are automatically migrated during loads */
+	UPROPERTY(EditAnywhere, config, Category=Blueprints, meta=(DisplayName="Dirty Migrated Blueprints"))
+	bool bDirtyMigratedBlueprints;
+
 public:
 
 	/** Whether to automatically save after a time interval */
diff --git a/Engine/Source/Runtime/Engine/Private/Blueprint.cpp b/Engine/Source/Runtime/Engine/Private/Blueprint.cpp
--- a/Engine/Source/Runtime/Engine/Private/Blueprint.cpp
+++ b/Engine/Source/Runtime/Engine/Private/Blueprint.cpp
@@ -13,6 +13,8 @@
 #include "Editor/UnrealEd/Public/Kismet2/StructureEditorUtils.h"
 #include "Editor/UnrealEd/Public/Editor.h"
 #include "Crc.h"
+#include "MessageLog.h"
+#include "Editor/UnrealEd/Classes/Settings/EditorLoadingSavingSettings.h"
 #endif
 
 DEFINE_LOG_CATEGORY(LogBlueprint);
@@ -1032,6 +1034,19 @@ bool UBlueprint::ChangeOwnerOfTemplates()
 
 		if (bMigratedOwner)
 		{
+			// alert the user that blueprints gave been migrated and require re-saving to enable them to locate and fix them without nagging them.
+			FMessageLog("BlueprintLog").Warning( FText::Format( NSLOCTEXT( "Blueprint", "MigrationWarning", "Blueprint {0} has been migrated and requires re-saving to avoid import errors" ), FText::FromString( *GetName() )));
+
+			if( GetDefault<UEditorLoadingSavingSettings>()->bDirtyMigratedBlueprints )
+			{
+				UPackage* BPPackage = GetOutermost();
+
+				if( BPPackage )
+				{
+					BPPackage->SetDirtyFlag( true );
+				}
+			}
+
 			BPGClass->ComponentTemplates = ComponentTemplates;
 			BPGClass->Timelines = Timelines;
 			BPGClass->SimpleConstructionScript = SimpleConstructionScript;
