Move defaults out of metadata and in to default parameters
Move tooltips out of metadata and as standard comments

diff --git a/Engine/Source/Runtime/Engine/Classes/Animation/AnimInstance.h b/Engine/Source/Runtime/Engine/Classes/Animation/AnimInstance.h
--- a/Engine/Source/Runtime/Engine/Classes/Animation/AnimInstance.h
+++ b/Engine/Source/Runtime/Engine/Classes/Animation/AnimInstance.h
@@ -336,12 +336,12 @@ public:
 public:
 
 	/** Play normal animation asset on the slot node. You can only play one asset (whether montage or animsequence) at a time. */
-	UFUNCTION(BlueprintCallable, Category="Animation", meta=(BlendInTime="0.25", BlendOutTime="0.25", InPlayRate="1.0"))
-	float PlaySlotAnimation(UAnimSequenceBase* Asset, FName SlotNodeName, float BlendInTime, float BlendOutTime, float InPlayRate);
+	UFUNCTION(BlueprintCallable, Category="Animation")
+	float PlaySlotAnimation(UAnimSequenceBase* Asset, FName SlotNodeName, float BlendInTime = 0.25f, float BlendOutTime = 0.25f, float InPlayRate = 1.f);
 
 	/** Stops currently playing slot animation */
-	UFUNCTION(BlueprintCallable, Category="Animation", meta=(InBlendOutTime="0.25"))
-	void StopSlotAnimation(float InBlendOutTime);
+	UFUNCTION(BlueprintCallable, Category="Animation")
+	void StopSlotAnimation(float InBlendOutTime = 0.25f);
 
 	/** Return true if it's playing the slot animation */
 	UFUNCTION(BlueprintCallable, Category="Animation")
@@ -368,16 +368,16 @@ public:
 	void Montage_SetNextSection(FName SectionNameToChange, FName NextSection);
 
 	/** Change AnimMontage play rate. NewPlayRate = 1.0 is the default playback rate. */
-	UFUNCTION(BlueprintCallable, Category="Animation", meta=(NewPlayRate="1.0"))
-	void Montage_SetPlayRate(UAnimMontage* Montage, float NewPlayRate);
+	UFUNCTION(BlueprintCallable, Category="Animation")
+	void Montage_SetPlayRate(UAnimMontage* Montage, float NewPlayRate = 1.f);
 
 	/** Plays an animation montage. Returns the length of the animation montage in seconds. Returns 0.f if failed to play. */
-	UFUNCTION(BlueprintCallable, Category="Animation", meta=(InPlayRate="1.0"))
-	float Montage_Play(UAnimMontage* MontageToPlay, float InPlayRate);
+	UFUNCTION(BlueprintCallable, Category="Animation")
+	float Montage_Play(UAnimMontage* MontageToPlay, float InPlayRate = 1.f);
 
 	/** Stops the animation montage. */
-	UFUNCTION(BlueprintCallable, Category="Animation", meta=(InBlendOutTime="0.25"))
-	void Montage_Stop(float InBlendOutTime);
+	UFUNCTION(BlueprintCallable, Category="Animation")
+	void Montage_Stop(float InBlendOutTime = 0.25f);
 
 	/** Returns true if the animation montage is active. */
 	UFUNCTION(BlueprintCallable, Category="Animation")
diff --git a/Engine/Source/Runtime/Engine/Classes/Components/AudioComponent.h b/Engine/Source/Runtime/Engine/Classes/Components/AudioComponent.h
--- a/Engine/Source/Runtime/Engine/Classes/Components/AudioComponent.h
+++ b/Engine/Source/Runtime/Engine/Classes/Components/AudioComponent.h
@@ -189,8 +189,8 @@ class ENGINE_API UAudioComponent : public USceneComponent
 	 * @param FadeInDuration how long it should take to reach the FadeVolumeLevel
 	 * @param FadeVolumeLevel the percentage of the AudioComponents's calculated volume to fade to
 	 */
-	UFUNCTION(BlueprintCallable, Category="Audio|Components|Audio", meta=(FadeVolumeLevel="1.0"))
-	void FadeIn(float FadeInDuration, float FadeVolumeLevel, float StartTime = 0.f);
+	UFUNCTION(BlueprintCallable, Category="Audio|Components|Audio")
+	void FadeIn(float FadeInDuration, float FadeVolumeLevel = 1.f, float StartTime = 0.f);
 
 	/**
 	 * This is used in place of "stop" when it is desired to fade the volume of the sound before stopping.
diff --git a/Engine/Source/Runtime/Engine/Classes/Engine/LevelScriptActor.h b/Engine/Source/Runtime/Engine/Classes/Engine/LevelScriptActor.h
--- a/Engine/Source/Runtime/Engine/Classes/Engine/LevelScriptActor.h
+++ b/Engine/Source/Runtime/Engine/Classes/Engine/LevelScriptActor.h
@@ -24,7 +24,7 @@ class ENGINE_API ALevelScriptActor : public AActor
 	 * @param	bAffectsMovement	specify true to disable movement in cinematic mode, enable it when leaving
 	 * @param	bAffectsTurning		specify true to disable turning in cinematic mode or enable it when leaving
 	 */
-	UFUNCTION(BlueprintCallable, meta=(BlueprintProtected = "true", bHidePlayer = "true", bAffectsHUD = "true"), Category="Game|Cinematic")
+	UFUNCTION(BlueprintCallable, meta=(BlueprintProtected = "true"), Category="Game|Cinematic")
 	virtual void SetCinematicMode(bool bCinematicMode, bool bHidePlayer = true, bool bAffectsHUD = true, bool bAffectsMovement = false, bool bAffectsTurning = false);
 
 	// --- Level State Functions ------------------------
diff --git a/Engine/Source/Runtime/Engine/Classes/GameFramework/Actor.h b/Engine/Source/Runtime/Engine/Classes/GameFramework/Actor.h
--- a/Engine/Source/Runtime/Engine/Classes/GameFramework/Actor.h
+++ b/Engine/Source/Runtime/Engine/Classes/GameFramework/Actor.h
@@ -520,13 +520,16 @@ public:
 	UFUNCTION(BlueprintCallable, Category="Input")
 	virtual void DisableInput(class APlayerController* PlayerController);
 
-	UFUNCTION(BlueprintCallable, meta=(BlueprintInternalUseOnly="true", BlueprintProtected = "true", HidePin="InputAxisName", ToolTip="Gets the value of the input axis if input is enabled for this actor"))
+	/** Gets the value of the input axis if input is enabled for this actor */
+	UFUNCTION(BlueprintCallable, meta=(BlueprintInternalUseOnly="true", BlueprintProtected = "true", HidePin="InputAxisName"))
 	float GetInputAxisValue(const FName InputAxisName) const;
 
-	UFUNCTION(BlueprintCallable, meta=(BlueprintInternalUseOnly="true", BlueprintProtected = "true", HidePin="InputAxisKey", ToolTip="Gets the value of the input axis key if input is enabled for this actor"))
+	/** Gets the value of the input axis key if input is enabled for this actor */
+	UFUNCTION(BlueprintCallable, meta=(BlueprintInternalUseOnly="true", BlueprintProtected = "true", HidePin="InputAxisKey"))
 	float GetInputAxisKeyValue(const FKey InputAxisKey) const;
 
-	UFUNCTION(BlueprintCallable, meta=(BlueprintInternalUseOnly="true", BlueprintProtected = "true", HidePin="InputAxisKey", ToolTip="Gets the value of the input axis key if input is enabled for this actor"))
+	/** Gets the value of the input axis key if input is enabled for this actor */
+	UFUNCTION(BlueprintCallable, meta=(BlueprintInternalUseOnly="true", BlueprintProtected = "true", HidePin="InputAxisKey"))
 	FVector GetInputVectorAxisValue(const FKey InputAxisKey) const;
 
 	/** Returns the instigator for this actor, or NULL if there is none. */
@@ -819,11 +822,11 @@ public:
 	// Sound functions.
 	
 	/* DEPRECATED - Use UGameplayStatics::PlaySoundAttached */
-	UFUNCTION(BlueprintCallable, Category="Audio", meta=(DeprecatedFunction, VolumeMultiplier="1.0", PitchMultiplier="1.0"))
+	UFUNCTION(BlueprintCallable, Category="Audio", meta=(DeprecatedFunction))
 	void PlaySoundOnActor(class USoundCue* InSoundCue, float VolumeMultiplier=1.f, float PitchMultiplier=1.f);
 
 	/* DEPRECATED - Use UGameplayStatics::PlaySoundAtLocation */
-	UFUNCTION(BlueprintCallable, Category="Audio", meta=(DeprecatedFunction, VolumeMultiplier="1.0", PitchMultiplier="1.0"))
+	UFUNCTION(BlueprintCallable, Category="Audio", meta=(DeprecatedFunction))
 	void PlaySoundAtLocation(class USoundCue* InSoundCue, FVector SoundLocation, float VolumeMultiplier=1.f, float PitchMultiplier=1.f);
 
 	//=============================================================================
diff --git a/Engine/Source/Runtime/Engine/Classes/GameFramework/HUD.h b/Engine/Source/Runtime/Engine/Classes/GameFramework/HUD.h
--- a/Engine/Source/Runtime/Engine/Classes/GameFramework/HUD.h
+++ b/Engine/Source/Runtime/Engine/Classes/GameFramework/HUD.h
@@ -295,7 +295,7 @@ public:
 	 * @param Font				Font to draw text.  If NULL, default font is chosen.
 	 * @param Scale				Scale multiplier to control size of the text.
 	 */
-	UFUNCTION(BlueprintCallable, Category=HUD, meta=(Scale = "1"))
+	UFUNCTION(BlueprintCallable, Category=HUD)
 	void GetTextSize(const FString& Text, float& OutWidth, float& OutHeight, class UFont* Font=NULL, float Scale=1.f) const;
 
 	/**
@@ -308,7 +308,7 @@ public:
 	 * @param Scale				Scale multiplier to control size of the text.
 	 * @param bScalePosition	Whether the "Scale" parameter should also scale the position of this draw call.
 	 */
-	UFUNCTION(BlueprintCallable, Category=HUD, meta=( Scale = "1", bScalePosition = "false"))
+	UFUNCTION(BlueprintCallable, Category=HUD)
 	void DrawText(const FString& Text, FLinearColor TextColor, float ScreenX, float ScreenY, class UFont* Font=NULL, float Scale=1.f, bool bScalePosition=false);
 
 	/**
@@ -351,7 +351,7 @@ public:
 	 * @param Rotation			Amount to rotate this quad
 	 * @param RotPivot			Location (as proportion of quad, 0-1) to rotate about
 	 */
-	UFUNCTION(BlueprintCallable, Category=HUD, meta=(Scale = "1", bScalePosition = "false", AdvancedDisplay = "9"))
+	UFUNCTION(BlueprintCallable, Category=HUD, meta=(AdvancedDisplay = "9"))
 	void DrawTexture(UTexture* Texture, float ScreenX, float ScreenY, float ScreenW, float ScreenH, float TextureU, float TextureV, float TextureUWidth, float TextureVHeight, FLinearColor TintColor=FLinearColor::White, EBlendMode BlendMode=BLEND_Translucent, float Scale=1.f, bool bScalePosition=false, float Rotation=0.f, FVector2D RotPivot=FVector2D::ZeroVector);
 
 	/**
@@ -362,7 +362,7 @@ public:
 	 * @param Scale				Scale multiplier to control size of the text.
 	 * @param bScalePosition	Whether the "Scale" parameter should also scale the position of this draw call.
 	 */
-	UFUNCTION(BlueprintCallable, Category=HUD, meta=(Scale = "1", bScalePosition = "false"))
+	UFUNCTION(BlueprintCallable, Category=HUD)
 	void DrawTextureSimple(UTexture* Texture, float ScreenX, float ScreenY, float Scale=1.f, bool bScalePosition=false);
 
 	/**
@@ -381,7 +381,7 @@ public:
 	 * @param Rotation			Amount to rotate this quad
 	 * @param RotPivot			Location (as proportion of quad, 0-1) to rotate about
 	 */
-	UFUNCTION(BlueprintCallable, Category=HUD, meta=(Scale = "1", bScalePosition = "false", AdvancedDisplay = "9"))
+	UFUNCTION(BlueprintCallable, Category=HUD, meta=(AdvancedDisplay = "9"))
 	void DrawMaterial(UMaterialInterface* Material, float ScreenX, float ScreenY, float ScreenW, float ScreenH, float MaterialU, float MaterialV, float MaterialUWidth, float MaterialVHeight, float Scale=1.f, bool bScalePosition=false, float Rotation=0.f, FVector2D RotPivot=FVector2D::ZeroVector);
 
 	/**
@@ -394,7 +394,7 @@ public:
 	 * @param Scale				Amount to scale the entire texture (horizontally and vertically)
 	 * @param bScalePosition	Whether the "Scale" parameter should also scale the position of this draw call.
 	 */
-	UFUNCTION(BlueprintCallable, Category=HUD, meta=(Scale = "1", bScalePosition = "false"))
+	UFUNCTION(BlueprintCallable, Category=HUD)
 	void DrawMaterialSimple(UMaterialInterface* Material, float ScreenX, float ScreenY, float ScreenW, float ScreenH, float Scale=1.f, bool bScalePosition=false);
 
 	/** Transforms a 3D world-space vector into 2D screen coordinates */
@@ -471,7 +471,7 @@ public:
 	 * @param bConsumesInput	Whether click processing should continue if this hit box is clicked.
 	 * @param Priority			The priority of the box used for layering. Larger values are considered first.  Equal values are considered in the order they were added.
 	 */
-	UFUNCTION(BlueprintCallable, Category=HUD, meta=(InPriority="0"))
+	UFUNCTION(BlueprintCallable, Category=HUD)
 	void AddHitBox(FVector2D Position, FVector2D Size, FName Name, bool bConsumesInput, int32 Priority = 0);
 
 protected:
diff --git a/Engine/Source/Runtime/Engine/Classes/GameFramework/Pawn.h b/Engine/Source/Runtime/Engine/Classes/GameFramework/Pawn.h
--- a/Engine/Source/Runtime/Engine/Classes/GameFramework/Pawn.h
+++ b/Engine/Source/Runtime/Engine/Classes/GameFramework/Pawn.h
@@ -143,7 +143,8 @@ private:
 	virtual bool IsCrouched() const;
 
 public:
-	UFUNCTION(BlueprintPure, Category="Pawn", meta=(ToolTip="Gets the owning actor of the Movement Base Component on which the pawn is standing."))
+	/** Gets the owning actor of the Movement Base Component on which the pawn is standing. */
+	UFUNCTION(BlueprintPure, Category="Pawn")
 	static AActor* GetMovementBaseActor(const APawn* Pawn);
 
 	virtual bool IsBasedOnActor(const AActor * Other) const override;
diff --git a/Engine/Source/Runtime/Engine/Classes/Kismet/GameplayStatics.h b/Engine/Source/Runtime/Engine/Classes/Kismet/GameplayStatics.h
--- a/Engine/Source/Runtime/Engine/Classes/Kismet/GameplayStatics.h
+++ b/Engine/Source/Runtime/Engine/Classes/Kismet/GameplayStatics.h
@@ -99,8 +99,8 @@ class ENGINE_API UGameplayStatics : public UBlueprintFunctionLibrary
 	 *  @param ControllerId		The ID of the controller that the should control the newly created player.  A value of -1 specifies to use the next available ID
 	 *  @param bSpawnPawn		Whether a pawn should be spawned immediately. If false a pawn will not be created until transition to the next map.
 	 */
-	UFUNCTION(BlueprintCallable, Category="Game", meta=(HidePin="WorldContextObject", DefaultToSelf="WorldContextObject", ControllerId="-1", bSpawnPawn="true", AdvancedDisplay="2"))
-	static class APlayerController* CreatePlayer(UObject* WorldContextObject, int32 ControllerId, bool bSpawnPawn = true);
+	UFUNCTION(BlueprintCallable, Category="Game", meta=(HidePin="WorldContextObject", DefaultToSelf="WorldContextObject", AdvancedDisplay="2"))
+	static class APlayerController* CreatePlayer(UObject* WorldContextObject, int32 ControllerId = -1, bool bSpawnPawn = true);
 
 	// --- Level Streaming functions ------------------------
 	
@@ -123,8 +123,8 @@ class ENGINE_API UGameplayStatics : public UBlueprintFunctionLibrary
 	 * @param	bAbsolute			if true options are reset, if false options are carried over from current level
 	 * @param	Options				a string of options to use for the travel URL
 	 */
-	UFUNCTION(BlueprintCallable, meta=(HidePin="WorldContextObject", DefaultToSelf="WorldContextObject", bAbsolute = "true", AdvancedDisplay = "2"), Category="Game")
-	static void OpenLevel(UObject* WorldContextObject, FName LevelName, bool bAbsolute, FString Options);
+	UFUNCTION(BlueprintCallable, meta=(HidePin="WorldContextObject", DefaultToSelf="WorldContextObject", AdvancedDisplay = "2"), Category="Game")
+	static void OpenLevel(UObject* WorldContextObject, FName LevelName, bool bAbsolute = true, FString Options = FString(TEXT("")));
 
 	// --- Global functions ------------------------------
 
@@ -218,8 +218,8 @@ class ENGINE_API UGameplayStatics : public UBlueprintFunctionLibrary
 	 * @param Falloff - Affects falloff of effect as it nears OuterRadius
 	 * @param bOrientShakeTowardsEpicenter - Changes the rotation of shake to point towards epicenter instead of forward
 	 */
-	UFUNCTION(BlueprintCallable, Category="Game|Feedback", meta=(HidePin="WorldContextObject", DefaultToSelf="WorldContextObject", Falloff="1.0", UnsafeDuringActorConstruction = "true"))
-	static void PlayWorldCameraShake(UObject* WorldContextObject, TSubclassOf<class UCameraShake> Shake, FVector Epicenter, float InnerRadius, float OuterRadius, float Falloff, bool bOrientShakeTowardsEpicenter = false);
+	UFUNCTION(BlueprintCallable, Category="Game|Feedback", meta=(HidePin="WorldContextObject", DefaultToSelf="WorldContextObject", UnsafeDuringActorConstruction = "true"))
+	static void PlayWorldCameraShake(UObject* WorldContextObject, TSubclassOf<class UCameraShake> Shake, FVector Epicenter, float InnerRadius, float OuterRadius, float Falloff = 1.f, bool bOrientShakeTowardsEpicenter = false);
 
 	// --- Particle functions ------------------------------
 
@@ -230,7 +230,7 @@ class ENGINE_API UGameplayStatics : public UBlueprintFunctionLibrary
 	 * @param Rotation - rotation to place the effect in world space	
 	 * @param bAutoDestroy - Whether the component will automatically be destroyed when the particle system completes playing or whether it can be reactivated
 	 */
-	UFUNCTION(BlueprintCallable, Category="Effects|Components|ParticleSystem", meta=(Keywords = "particle system", HidePin="WorldContextObject", DefaultToSelf="WorldContextObject", bAutoDestroy="true", UnsafeDuringActorConstruction = "true"))
+	UFUNCTION(BlueprintCallable, Category="Effects|Components|ParticleSystem", meta=(Keywords = "particle system", HidePin="WorldContextObject", DefaultToSelf="WorldContextObject", UnsafeDuringActorConstruction = "true"))
 	static UParticleSystemComponent* SpawnEmitterAtLocation(UObject* WorldContextObject, class UParticleSystem* EmitterTemplate, FVector Location, FRotator Rotation = FRotator::ZeroRotator, bool bAutoDestroy = true);
 
 	/** Plays the specified effect attached to and following the specified component. The system will go away when the effect is complete. Does not replicate.
@@ -242,7 +242,7 @@ class ENGINE_API UGameplayStatics : public UBlueprintFunctionLibrary
 	 * @param LocationType - Specifies whether Location is a relative offset or an absolute world position
 	 * @param bAutoDestroy - Whether the component will automatically be destroyed when the particle system completes playing or whether it can be reactivated
 	 */
-	UFUNCTION(BlueprintCallable, Category="Effects|Components|ParticleSystem", meta=(Keywords = "particle system", bAutoDestroy="true", UnsafeDuringActorConstruction = "true"))
+	UFUNCTION(BlueprintCallable, Category="Effects|Components|ParticleSystem", meta=(Keywords = "particle system", UnsafeDuringActorConstruction = "true"))
 	static UParticleSystemComponent* SpawnEmitterAttached(class UParticleSystem* EmitterTemplate, class USceneComponent* AttachToComponent, FName AttachPointName = NAME_None, FVector Location = FVector(ForceInit), FRotator Rotation = FRotator::ZeroRotator, EAttachLocation::Type LocationType = EAttachLocation::KeepRelativeOffset, bool bAutoDestroy = true);
 
 	// --- Sound functions ------------------------------
@@ -264,7 +264,7 @@ class ENGINE_API UGameplayStatics : public UBlueprintFunctionLibrary
 	 * @param PitchMultiplier - PitchMultiplier
 	 * @param AttenuationSettings - Override attenuation settings package to play sound with
 	 */
-	UFUNCTION(BlueprintCallable, Category="Audio", meta=(HidePin="WorldContextObject", DefaultToSelf="WorldContextObject", VolumeMultiplier="1.0", PitchMultiplier="1.0", AdvancedDisplay = "3", UnsafeDuringActorConstruction = "true"))
+	UFUNCTION(BlueprintCallable, Category="Audio", meta=(HidePin="WorldContextObject", DefaultToSelf="WorldContextObject", AdvancedDisplay = "3", UnsafeDuringActorConstruction = "true"))
 	static void PlaySoundAtLocation(UObject* WorldContextObject, class USoundBase* Sound, FVector Location, float VolumeMultiplier = 1.f, float PitchMultiplier = 1.f, float StartTime = 0.f, class USoundAttenuation* AttenuationSettings = NULL);
 
 	/** Plays a dialogue at the given location. This is a fire and forget sound and does not travel with any actor. Replication is also not handled at this point.
@@ -275,7 +275,7 @@ class ENGINE_API UGameplayStatics : public UBlueprintFunctionLibrary
 	 * @param PitchMultiplier - PitchMultiplier
 	 * @param AttenuationSettings - Override attenuation settings package to play sound with
 	 */
-	UFUNCTION(BlueprintCallable, Category="Audio", meta=(HidePin="WorldContextObject", DefaultToSelf="WorldContextObject", VolumeMultiplier="1.0", PitchMultiplier="1.0", AdvancedDisplay = "3", UnsafeDuringActorConstruction = "true"))
+	UFUNCTION(BlueprintCallable, Category="Audio", meta=(HidePin="WorldContextObject", DefaultToSelf="WorldContextObject", AdvancedDisplay = "3", UnsafeDuringActorConstruction = "true"))
 	static void PlayDialogueAtLocation(UObject* WorldContextObject, class UDialogueWave* Dialogue, const struct FDialogueContext& Context, FVector Location, float VolumeMultiplier = 1.f, float PitchMultiplier = 1.f, float StartTime = 0.f, class USoundAttenuation* AttenuationSettings = NULL);
 
 	/** Plays a sound attached to and following the specified component. This is a fire and forget sound. Replication is also not handled at this point.
@@ -289,7 +289,7 @@ class ENGINE_API UGameplayStatics : public UBlueprintFunctionLibrary
 	 * @param PitchMultiplier - PitchMultiplier	 
 	 * @param AttenuationSettings - Override attenuation settings package to play sound with
 	 */
-	UFUNCTION(BlueprintCallable, Category="Audio", meta=(VolumeMultiplier="1.0", PitchMultiplier="1.0", AdvancedDisplay = "2", UnsafeDuringActorConstruction = "true"))
+	UFUNCTION(BlueprintCallable, Category="Audio", meta=(AdvancedDisplay = "2", UnsafeDuringActorConstruction = "true"))
 	static class UAudioComponent* PlaySoundAttached(class USoundBase* Sound, class USceneComponent* AttachToComponent, FName AttachPointName = NAME_None, FVector Location = FVector(ForceInit), EAttachLocation::Type LocationType = EAttachLocation::KeepRelativeOffset, bool bStopWhenAttachedToDestroyed = false, float VolumeMultiplier = 1.f, float PitchMultiplier = 1.f, float StartTime = 0.f, class USoundAttenuation* AttenuationSettings = NULL);
 
 	/** Plays a dialogue attached to and following the specified component. This is a fire and forget sound. Replication is also not handled at this point.
@@ -303,7 +303,7 @@ class ENGINE_API UGameplayStatics : public UBlueprintFunctionLibrary
 	 * @param PitchMultiplier - PitchMultiplier	 
 	 * @param AttenuationSettings - Override attenuation settings package to play sound with
 	 */
-	UFUNCTION(BlueprintCallable, Category="Audio", meta=(VolumeMultiplier="1.0", PitchMultiplier="1.0", AdvancedDisplay = "2", UnsafeDuringActorConstruction = "true"))
+	UFUNCTION(BlueprintCallable, Category="Audio", meta=(AdvancedDisplay = "2", UnsafeDuringActorConstruction = "true"))
 	static class UAudioComponent* PlayDialogueAttached(class UDialogueWave* Dialogue, const struct FDialogueContext& Context, class USceneComponent* AttachToComponent, FName AttachPointName = NAME_None, FVector Location = FVector(ForceInit), EAttachLocation::Type LocationType = EAttachLocation::KeepRelativeOffset, bool bStopWhenAttachedToDestroyed = false, float VolumeMultiplier = 1.f, float PitchMultiplier = 1.f, float StartTime = 0.f, class USoundAttenuation* AttenuationSettings = NULL);
 
 	/** DEPRECATED - Use GameplayStatics.PlaySoundAtLocation or GameplayStatics.PlaySoundAttached instead. */
diff --git a/Engine/Source/Runtime/Engine/Classes/Kismet/KismetStringLibrary.h b/Engine/Source/Runtime/Engine/Classes/Kismet/KismetStringLibrary.h
--- a/Engine/Source/Runtime/Engine/Classes/Kismet/KismetStringLibrary.h
+++ b/Engine/Source/Runtime/Engine/Classes/Kismet/KismetStringLibrary.h
@@ -229,8 +229,8 @@ class UKismetStringLibrary : public UBlueprintFunctionLibrary
 	 *
 	 * @return The requested substring
 	 */
-	UFUNCTION(BlueprintPure, meta=(StartIndex="0", Length="1"), Category="Utilities|String")
-	static FString GetSubstring(const FString& SourceString, int32 StartIndex, int32 Length);
+	UFUNCTION(BlueprintPure, Category="Utilities|String")
+	static FString GetSubstring(const FString& SourceString, int32 StartIndex = 0, int32 Length = 1);
 
 	/**
 	 * Finds the starting index of a substring in the a specified string
@@ -250,8 +250,8 @@ class UKismetStringLibrary : public UBlueprintFunctionLibrary
 	 * @param Index - Location of the character whose value is required
 	 * @return The integer value of the character or 0 if index is out of range
 	 */
-	UFUNCTION(BlueprintPure, meta=(Index="0"), Category="Utilities|String")
-	static int32 GetCharacterAsNumber(const FString& SourceString, int32 Index);
+	UFUNCTION(BlueprintPure, Category="Utilities|String")
+	static int32 GetCharacterAsNumber(const FString& SourceString, int32 Index = 0);
 
 	/** 
 	 * Gets an array of strings from a source string divided up by a separator and empty strings can optionally be culled.
@@ -260,8 +260,8 @@ class UKismetStringLibrary : public UBlueprintFunctionLibrary
 	 * @param CullEmptyStrings = true - Cull (true) empty strings or add them to the array (false)
 	 * @return The array of string that have been separated
 	 */
-	UFUNCTION(BlueprintPure, meta=(Delimiter=" "), Category="Utilities|String")
-	static TArray<FString> ParseIntoArray(const FString& SourceString, const FString& Delimiter, const bool CullEmptyStrings = true);
+	UFUNCTION(BlueprintPure, Category="Utilities|String")
+	static TArray<FString> ParseIntoArray(const FString& SourceString, const FString& Delimiter = FString(TEXT(" ")), const bool CullEmptyStrings = true);
 
 	/**
 	 * Returns an array that contains one entry for each character in SourceString
diff --git a/Engine/Source/Runtime/Engine/Classes/Kismet/KismetSystemLibrary.h b/Engine/Source/Runtime/Engine/Classes/Kismet/KismetSystemLibrary.h
--- a/Engine/Source/Runtime/Engine/Classes/Kismet/KismetSystemLibrary.h
+++ b/Engine/Source/Runtime/Engine/Classes/Kismet/KismetSystemLibrary.h
@@ -176,8 +176,8 @@ class UKismetSystemLibrary : public UBlueprintFunctionLibrary
 	 * @param	bPrintToConsole	Whether or not to print the output to the console
 	 * @param	TextColor		Whether or not to print the output to the console
 	 */
-	UFUNCTION(BlueprintCallable, meta=(HidePin="WorldContextObject", DefaultToSelf="WorldContextObject", InString = "Hello", bPrintToScreen = "true", bPrintToLog = "true", Keywords = "log", AdvancedDisplay = "2"), Category="Utilities|String")
-	static void PrintString(UObject* WorldContextObject, const FString& InString, bool bPrintToScreen, bool bPrintToLog, FLinearColor TextColor = FLinearColor(0.0,0.66,1.0));
+	UFUNCTION(BlueprintCallable, meta=(HidePin="WorldContextObject", DefaultToSelf="WorldContextObject", Keywords = "log", AdvancedDisplay = "2"), Category="Utilities|String")
+	static void PrintString(UObject* WorldContextObject, const FString& InString = FString(TEXT("Hello")), bool bPrintToScreen = true, bool bPrintToLog = true, FLinearColor TextColor = FLinearColor(0.0,0.66,1.0));
 
 	/**
 	 * Prints a warning string to the log and the screen. Meant to be used as a way to inform the user that they misused the node.
diff --git a/Engine/Source/Runtime/Engine/Classes/Sound/AmbientSound.h b/Engine/Source/Runtime/Engine/Classes/Sound/AmbientSound.h
--- a/Engine/Source/Runtime/Engine/Classes/Sound/AmbientSound.h
+++ b/Engine/Source/Runtime/Engine/Classes/Sound/AmbientSound.h
@@ -37,8 +37,8 @@ class AAmbientSound : public AActor
 	// End AActor interface.
 
 	// BEGIN DEPRECATED (use component functions now in level script)
-	UFUNCTION(BlueprintCallable, Category="Audio", meta=(FadeVolumeLevel="1.0",DeprecatedFunction))
-	void FadeIn(float FadeInDuration, float FadeVolumeLevel);
+	UFUNCTION(BlueprintCallable, Category="Audio", meta=(DeprecatedFunction))
+	void FadeIn(float FadeInDuration, float FadeVolumeLevel = 1.f);
 	UFUNCTION(BlueprintCallable, Category="Audio", meta=(DeprecatedFunction))
 	void FadeOut(float FadeOutDuration, float FadeVolumeLevel);
 	UFUNCTION(BlueprintCallable, Category="Audio", meta=(DeprecatedFunction))
diff --git a/Engine/Source/Runtime/GameLiveStreaming/GameLiveStreamingFunctionLibrary.h b/Engine/Source/Runtime/GameLiveStreaming/GameLiveStreamingFunctionLibrary.h
--- a/Engine/Source/Runtime/GameLiveStreaming/GameLiveStreamingFunctionLibrary.h
+++ b/Engine/Source/Runtime/GameLiveStreaming/GameLiveStreamingFunctionLibrary.h
@@ -29,17 +29,17 @@ class UGameLiveStreamingFunctionLibrary : public UBlueprintFunctionLibrary
 	 * @param	bCaptureAudioFromMicrophone	Enables broadcast of audio from your default microphone recording device
 	 * @param	bDrawSimpleWebCamVideo		If enabled, the engine will draw a simple web cam image on top of the game viewport.  If you turn this off, it's up to you to draw the web cam image yourself.  You can access the web cam texture by calling IGameLiveStreaming::Get().GetWebCamTexture(). 
 	 */
-	UFUNCTION( BlueprintCallable, Category="LiveStreaming", meta=(FrameRate="30", ScreenScaling="1.0", bEnableWebCam="true", DesiredWebCamWidth="320", DesiredWebCamHeight="240", bMirrorWebCamImage="false", bCaptureAudioFromComputer="true", bCaptureAudioFromMicrophone="true", bDrawSimpleWebCamVideo="true") )
+	UFUNCTION( BlueprintCallable, Category="LiveStreaming")
 	static void StartBroadcastingGame(
-		int32 FrameRate,
-		float ScreenScaling,
-		bool bEnableWebCam,
-		int32 DesiredWebCamWidth,
-		int32 DesiredWebCamHeight,
-		bool bMirrorWebCamImage,
-		bool bCaptureAudioFromComputer,
-		bool bCaptureAudioFromMicrophone,
-		bool bDrawSimpleWebCamVideo );
+		int32 FrameRate = 30,
+		float ScreenScaling = 1.f,
+		bool bEnableWebCam = true,
+		int32 DesiredWebCamWidth = 320,
+		int32 DesiredWebCamHeight = 240,
+		bool bMirrorWebCamImage = false,
+		bool bCaptureAudioFromComputer = true,
+		bool bCaptureAudioFromMicrophone = true,
+		bool bDrawSimpleWebCamVideo = true);
 
 	/** Stops broadcasting the game */
 	UFUNCTION( BlueprintCallable, Category="LiveStreaming" )
