#ue4 - Add CameraActor 'AutoActivateForPlayer' option, so cameras placed in a level can automatically be linked to a PlayerController.

- Added Blueprint events for Actors when they start and stop being the viewtarget for a PlayerController (OnBecomeViewTarget / OnEndViewTarget).
   - Native C++ events for these (BecomeViewTarget / EndViewTarget) already existed and now trigger the BP events.

- Changed the order of calls to EndViewTarget / BecomeViewTarget from the PlayerCameraManager, so that the End event occurs first on the old target, followed by the Begin event for the new one (previously it was reversed).

diff --git a/Engine/Source/Runtime/Engine/Classes/Camera/CameraActor.h b/Engine/Source/Runtime/Engine/Classes/Camera/CameraActor.h
--- a/Engine/Source/Runtime/Engine/Classes/Camera/CameraActor.h
+++ b/Engine/Source/Runtime/Engine/Classes/Camera/CameraActor.h
@@ -12,6 +12,12 @@ class ACameraActor : public AActor
 {
 	GENERATED_UCLASS_BODY()
 
+private:
+
+	/* Specifies which player controller, if any, should automatically use this Camera when the controller is active. */
+	UPROPERTY(Category="AutoPlayerActivation", EditAnywhere)
+	TEnumAsByte<EAutoReceiveInput::Type> AutoActivateForPlayer;
+
 public:
 	// The camera component for this camera
 	UPROPERTY(Category=CameraActor, VisibleAnywhere, BlueprintReadOnly)
@@ -20,7 +26,12 @@ public:
 	/** If this CameraActor is being used to preview a CameraAnim in the editor, this is the anim being previewed. */
 	TWeakObjectPtr<class UCameraAnim> PreviewedCameraAnim;
 
+	/** Returns index of the player for whom we auto-activate, or INDEX_NONE (-1) if disabled. */
+	UFUNCTION(BlueprintCallable, Category="AutoPlayerActivation")
+	ENGINE_API int32 GetAutoActivatePlayerIndex() const;
+
 private:
+
 	UPROPERTY()
 	uint32 bConstrainAspectRatio_DEPRECATED:1;
 
@@ -46,4 +57,7 @@ public:
 #endif
 	// End UObject interface
 	
+	// Begin AActor interface
+	ENGINE_API virtual void BeginPlay() override;
+	// End AActor interface
 };
diff --git a/Engine/Source/Runtime/Engine/Classes/Engine/World.h b/Engine/Source/Runtime/Engine/Classes/Engine/World.h
--- a/Engine/Source/Runtime/Engine/Classes/Engine/World.h
+++ b/Engine/Source/Runtime/Engine/Classes/Engine/World.h
@@ -26,6 +26,7 @@ class UTexture2D;
 class APlayerController;
 class AMatineeActor;
 class AWorldSettings;
+class ACameraActor;
 
 template<typename,typename> class TOctree;
 
@@ -36,6 +37,7 @@ template<typename,typename> class TOctree;
 typedef TArray<TAutoWeakObjectPtr<AController> >::TConstIterator FConstControllerIterator;
 typedef TArray<TAutoWeakObjectPtr<APlayerController> >::TConstIterator FConstPlayerControllerIterator;
 typedef TArray<TAutoWeakObjectPtr<APawn> >::TConstIterator FConstPawnIterator;	
+typedef TArray<TAutoWeakObjectPtr<ACameraActor> >::TConstIterator FConstCameraActorIterator;
 typedef TArray<class ULevel*>::TConstIterator FConstLevelIterator;
 
 DECLARE_LOG_CATEGORY_EXTERN(LogSpawn, Warning, All);
@@ -648,6 +650,9 @@ private:
 	/** List of all the pawns in the world. */
 	TArray<TAutoWeakObjectPtr<class APawn> >			PawnList;	
 
+	/** List of all the cameras in the world that auto-activate for players. */
+	TArray<TAutoWeakObjectPtr<ACameraActor> > AutoCameraActorList;
+
 	/** Physics scene for this world. */
 	FPhysScene*									PhysicsScene;
 
@@ -1431,6 +1436,12 @@ public:
 	 *  @return Pointer to the first valid ULocalPlayer, or NULL if there is not one.
 	 */	
 	ULocalPlayer* GetFirstLocalPlayerFromController() const;
+
+	/** Register a CameraActor that auto-activates for a PlayerController. */
+	void RegisterAutoActivateCamera(ACameraActor* CameraActor, int32 PlayerIndex);
+
+	/** Get an iterator for the list of CameraActors that auto-activate for PlayerControllers. */
+	FConstCameraActorIterator GetAutoActivateCameraIterator() const;
 	
 	UGameViewportClient* GetGameViewport() const;
 
diff --git a/Engine/Source/Runtime/Engine/Classes/GameFramework/Actor.h b/Engine/Source/Runtime/Engine/Classes/GameFramework/Actor.h
--- a/Engine/Source/Runtime/Engine/Classes/GameFramework/Actor.h
+++ b/Engine/Source/Runtime/Engine/Classes/GameFramework/Actor.h
@@ -1861,12 +1861,20 @@ protected:
 	virtual float InternalTakePointDamage(float Damage, struct FPointDamageEvent const& RadialDamageEvent, class AController* EventInstigator, class AActor* DamageCauser);
 public:
 
-	/* Called by Camera when this actor becomes its ViewTarget */
+	/* Called when this actor becomes the given PlayerController's ViewTarget. Triggers the Blueprint event K2_OnBecomeViewTarget. */
 	virtual void BecomeViewTarget( class APlayerController* PC );
 
-	/* Called by Camera when this actor is no longer its ViewTarget */
+	/* Called when this actor is no longer the given PlayerController's ViewTarget. Also triggers the Blueprint event K2_OnEndViewTarget. */
 	virtual void EndViewTarget( class APlayerController* PC );
 
+	/** Event called when this Actor becomes the view target for the given PlayerController. */
+	UFUNCTION(BlueprintImplementableEvent, meta=(FriendlyName="OnBecomeViewTarget", Keywords="Activate Camera"), Category=Actor)
+	virtual void K2_OnBecomeViewTarget( class APlayerController* PC );
+
+	/** Event called when this Actor is no longer the view target for the given PlayerController. */
+	UFUNCTION(BlueprintImplementableEvent, meta=(FriendlyName="OnEndViewTarget", Keywords="Deactivate Camera"), Category=Actor)
+	virtual void K2_OnEndViewTarget( class APlayerController* PC );
+
 	/**
 	 *	Calculate camera view point, when viewing this actor.
 	 *
diff --git a/Engine/Source/Runtime/Engine/Classes/GameFramework/PlayerController.h b/Engine/Source/Runtime/Engine/Classes/GameFramework/PlayerController.h
--- a/Engine/Source/Runtime/Engine/Classes/GameFramework/PlayerController.h
+++ b/Engine/Source/Runtime/Engine/Classes/GameFramework/PlayerController.h
@@ -1137,7 +1137,6 @@ protected:
 	virtual void BuildInputStack(TArray<UInputComponent*>& InputStack);
 	void ProcessForceFeedback(const float DeltaTime, const bool bGamePaused);
 
-
 	/** Allows the PlayerController to set up custom input bindings. */
 	virtual void SetupInputComponent();
 
@@ -1284,6 +1283,18 @@ public:
 	 */
 	virtual void SetViewTarget(class AActor* NewViewTarget, FViewTargetTransitionParams TransitionParams = FViewTargetTransitionParams());
 
+	/**
+	 * If bAutoManageActiveCameraTarget is true, then automatically manage the active camera target.
+	 * If there a CameraActor placed in the level with an auto-activate player assigned to it, that will be preferred, otherwise SuggestedTarget will be used.
+	 */
+	virtual void AutoManageActiveCameraTarget(AActor* SuggestedTarget);
+
+protected:
+
+	virtual ACameraActor* GetAutoActivateCameraForPlayer() const;
+
+public:
+
 	virtual bool ShouldShowMouseCursor() const;
 	virtual EMouseCursor::Type GetMouseCursor() const;
 
diff --git a/Engine/Source/Runtime/Engine/Private/Actor.cpp b/Engine/Source/Runtime/Engine/Private/Actor.cpp
--- a/Engine/Source/Runtime/Engine/Private/Actor.cpp
+++ b/Engine/Source/Runtime/Engine/Private/Actor.cpp
@@ -1931,9 +1931,15 @@ bool AActor::CanBeBaseForCharacter(APawn* APawn) const
 	return true;
 }
 
-void AActor::BecomeViewTarget( APlayerController* PC ) {}
+void AActor::BecomeViewTarget( APlayerController* PC )
+{
+	K2_OnBecomeViewTarget(PC);
+}
 
-void AActor::EndViewTarget( APlayerController* PC ) {}
+void AActor::EndViewTarget( APlayerController* PC )
+{
+	K2_OnEndViewTarget(PC);
+}
 
 APawn* AActor::GetInstigator() const
 {
diff --git a/Engine/Source/Runtime/Engine/Private/Camera/CameraActor.cpp b/Engine/Source/Runtime/Engine/Private/Camera/CameraActor.cpp
--- a/Engine/Source/Runtime/Engine/Private/Camera/CameraActor.cpp
+++ b/Engine/Source/Runtime/Engine/Private/Camera/CameraActor.cpp
@@ -80,4 +80,36 @@ void ACameraActor::PostEditChangeProperty(FPropertyChangedEvent& PropertyChanged
 }
 #endif
 
+
+int32 ACameraActor::GetAutoActivatePlayerIndex() const
+{
+	if (AutoActivateForPlayer != EAutoReceiveInput::Disabled)
+	{
+		const int32 PlayerIndex = int32(AutoActivateForPlayer.GetValue()) - 1;
+		return PlayerIndex;
+	}
+	else
+	{
+		return INDEX_NONE;
+	}
+}
+
+void ACameraActor::BeginPlay()
+{
+	if (AutoActivateForPlayer != EAutoReceiveInput::Disabled && GetNetMode() != NM_Client)
+	{
+		const int32 PlayerIndex = GetAutoActivatePlayerIndex();
+		
+		// Always put it in the pool of available auto-activate cameras.
+		GetWorld()->RegisterAutoActivateCamera(this, PlayerIndex);
+
+		// If we find a matching PC, bind to it immediately.
+		APlayerController* PC = UGameplayStatics::GetPlayerController(this, PlayerIndex);
+		if (PC)
+		{
+			PC->SetViewTarget(this);
+		}
+	}
+}
+
 #undef LOCTEXT_NAMESPACE
diff --git a/Engine/Source/Runtime/Engine/Private/Pawn.cpp b/Engine/Source/Runtime/Engine/Private/Pawn.cpp
--- a/Engine/Source/Runtime/Engine/Private/Pawn.cpp
+++ b/Engine/Source/Runtime/Engine/Private/Pawn.cpp
@@ -321,6 +321,8 @@ void APawn::TurnOff()
 
 void APawn::BecomeViewTarget(APlayerController* PC)
 {
+	Super::BecomeViewTarget(PC);
+
 	if (GetNetMode() != NM_Client)
 	{
 		PC->ForceSingleNetUpdateFor(this);
@@ -337,7 +339,7 @@ void APawn::PawnClientRestart()
 		// Handle camera possession
 		if (PC->bAutoManageActiveCameraTarget)
 		{
-			PC->SetViewTarget(this);
+			PC->AutoManageActiveCameraTarget(this);
 		}
 
 		// Set up player input component, if there isn't one already.
@@ -431,7 +433,7 @@ void APawn::OnRep_Controller()
 		APlayerController* const PC = Cast<APlayerController>(Controller);
 		if ( (PC != NULL) && PC->bAutoManageActiveCameraTarget && (PC->PlayerCameraManager->ViewTarget.Target == Controller) )
 		{
-			PC->SetViewTarget(this);
+			PC->AutoManageActiveCameraTarget(this);
 		}
 	}
 }
diff --git a/Engine/Source/Runtime/Engine/Private/PlayerController.cpp b/Engine/Source/Runtime/Engine/Private/PlayerController.cpp
--- a/Engine/Source/Runtime/Engine/Private/PlayerController.cpp
+++ b/Engine/Source/Runtime/Engine/Private/PlayerController.cpp
@@ -393,6 +393,94 @@ void APlayerController::SetViewTarget(class AActor* NewViewTarget, struct FViewT
 }
 
 
+void APlayerController::AutoManageActiveCameraTarget(AActor* SuggestedTarget)
+{
+	if (bAutoManageActiveCameraTarget)
+	{
+		// See if there is a CameraActor with an auto-activate index that matches us.
+		if (GetNetMode() == NM_Client)
+		{
+			// Clients don't know their own index on the server, so they have to trust that if they use a camera with an auto-activate index, that's their own index.
+			ACameraActor* CurrentCameraActor = Cast<ACameraActor>(GetViewTarget());
+			if (CurrentCameraActor)
+			{
+				const int32 CameraAutoIndex = CurrentCameraActor->GetAutoActivatePlayerIndex();
+				if (CameraAutoIndex != INDEX_NONE)
+				{					
+					return;
+				}
+			}
+		}
+		else
+		{
+			// See if there is a CameraActor in the level that auto-activates for this PC.
+			ACameraActor* AutoCameraTarget = GetAutoActivateCameraForPlayer();
+			if (AutoCameraTarget)
+			{
+				SetViewTarget(AutoCameraTarget);
+				return;
+			}
+		}
+
+		// No auto-activate CameraActor, so use the suggested target.
+		SetViewTarget(SuggestedTarget);
+	}
+}
+
+
+
+ACameraActor* APlayerController::GetAutoActivateCameraForPlayer() const
+{
+	if (GetNetMode() == NM_Client)
+	{
+		// Clients get their view target replicated, they don't use placed cameras because they don't know their own index.
+		return NULL;
+	}
+
+	UWorld* CurWorld = GetWorld();
+	if (!CurWorld)
+	{
+		return NULL;
+	}
+
+	// Only bother if there are any registered cameras.
+	FConstCameraActorIterator CameraIterator = CurWorld->GetAutoActivateCameraIterator();
+	if (!CameraIterator)
+	{
+		return NULL;
+	}
+
+	// Find our player index
+	int32 IterIndex = 0;
+	int32 PlayerIndex = INDEX_NONE;
+	for( FConstPlayerControllerIterator Iterator = CurWorld->GetPlayerControllerIterator(); Iterator; ++Iterator, ++IterIndex )
+	{
+		const APlayerController* PlayerController = *Iterator;
+		if (PlayerController == this)
+		{
+			PlayerIndex = IterIndex;
+			break;
+		}
+	}
+
+	if (PlayerIndex != INDEX_NONE)
+	{
+		// Find the matching camera
+		for( /*CameraIterater initialized above*/; CameraIterator; ++CameraIterator)
+		{
+			ACameraActor* CameraActor = *CameraIterator;
+			if (CameraActor && CameraActor->GetAutoActivatePlayerIndex() == PlayerIndex)
+			{
+				return CameraActor;
+			}
+		}
+	}
+
+	return NULL;
+}
+
+
+
 void APlayerController::SetControllingDirector(UInterpTrackInstDirector* NewControllingDirector, bool bClientSimulatingViewTarget)
 {
 	ControllingDirTrackInst = NewControllingDirector;
@@ -636,7 +724,7 @@ void APlayerController::ClientRestart_Implementation(APawn* NewPawn)
 	{
 		if (bAutoManageActiveCameraTarget)
 		{
-			SetViewTarget(GetPawn());
+			AutoManageActiveCameraTarget(GetPawn());
 			ResetCameraMode();
 		}
 		
@@ -679,9 +767,10 @@ void APlayerController::Possess(APawn* PawnToPossess)
 		ChangeState( NAME_Playing );
 		AcknowledgedPawn = NULL;
 		ClientRestart(GetPawn());
+		
 		if (bAutoManageActiveCameraTarget)
 		{
-			SetViewTarget(GetPawn());
+			AutoManageActiveCameraTarget(GetPawn());
 			ResetCameraMode();
 		}
 		UpdateNavigationComponents();
@@ -2285,9 +2374,6 @@ bool APlayerController::IsLookInputIgnored() const
 }
 
 
-/** returns whether this Controller is a locally controlled PlayerController
- * @note not valid until the Controller is completely spawned (i.e, unusable in Pre/PostInitializeComponents())
- */
 void APlayerController::SetViewTargetWithBlend(AActor* NewViewTarget, float BlendTime, EViewTargetBlendFunction BlendFunc, float BlendExp, bool bLockOutgoing)
 {
 	FViewTargetTransitionParams TransitionParams;
@@ -2791,11 +2877,11 @@ void APlayerController::ClientCommitMapChange_Implementation()
 		{
 			if (GetPawnOrSpectator() != NULL)
 			{
-				SetViewTarget(GetPawnOrSpectator());
+				AutoManageActiveCameraTarget(GetPawnOrSpectator());
 			}
 			else
 			{
-				SetViewTarget(this);
+				AutoManageActiveCameraTarget(this);
 			}
 		}
 		GetWorld()->CommitMapChange();
@@ -3447,7 +3533,7 @@ void APlayerController::SetPawn(APawn* InPawn)
 
 		if (bAutoManageActiveCameraTarget)
 		{
-			SetViewTarget(this);
+			AutoManageActiveCameraTarget(this);
 		}
 	}
 
diff --git a/Engine/Source/Runtime/Engine/Private/World.cpp b/Engine/Source/Runtime/Engine/Private/World.cpp
--- a/Engine/Source/Runtime/Engine/Private/World.cpp
+++ b/Engine/Source/Runtime/Engine/Private/World.cpp
@@ -3046,6 +3046,20 @@ void UWorld::RemovePawn( APawn* Pawn )
 	PawnList.Remove( Pawn );	
 }
 
+
+void UWorld::RegisterAutoActivateCamera(ACameraActor* CameraActor, int32 PlayerIndex)
+{
+	check(CameraActor);
+	check(PlayerIndex >= 0);
+	AutoCameraActorList.AddUnique(CameraActor);
+}
+
+FConstCameraActorIterator UWorld::GetAutoActivateCameraIterator() const
+{
+	return AutoCameraActorList.CreateConstIterator();
+}
+
+
 void UWorld::AddNetworkActor( AActor * Actor )
 {
 	if ( Actor == NULL )
