#ue4 - Code that sets Characters to "flying" if there is no engine physics code should not just check for PhsyX, but UE_WITH_PHYSICS instead.

diff --git a/Engine/Source/Runtime/Engine/Private/GameMode.cpp b/Engine/Source/Runtime/Engine/Private/GameMode.cpp
--- a/Engine/Source/Runtime/Engine/Private/GameMode.cpp
+++ b/Engine/Source/Runtime/Engine/Private/GameMode.cpp
@@ -468,7 +468,7 @@ void AGameMode::RestartPlayer(AController* NewPlayer)
 		SetPlayerDefaults(NewPlayer->GetPawn());
 	}
 
-#if !WITH_PHYSX
+#if !UE_WITH_PHYSICS
 	if (NewPlayer->GetPawn() != NULL)
 	{
 		UCharacterMovementComponent* CharacterMovement = Cast<UCharacterMovementComponent>(NewPlayer->GetPawn()->GetMovementComponent());
@@ -478,7 +478,7 @@ void AGameMode::RestartPlayer(AController* NewPlayer)
 			CharacterMovement->SetMovementMode(MOVE_Flying);
 		}
 	}
-#endif	//!WITH_PHYSX
+#endif	//!UE_WITH_PHYSICS
 }
 
 void AGameMode::InitStartSpot(AActor* StartSpot, AController* NewPlayer)
