Fixed so that it doesn't spam invalid montage log when no animation is added.
- Not having animation isn't good in montage, but it's not invalid. It is still valid and it can happen when you create first time or you never added animation to it.

diff --git a/Engine/Source/Runtime/Engine/Private/Animation/AnimInstance.cpp b/Engine/Source/Runtime/Engine/Private/Animation/AnimInstance.cpp
--- a/Engine/Source/Runtime/Engine/Private/Animation/AnimInstance.cpp
+++ b/Engine/Source/Runtime/Engine/Private/Animation/AnimInstance.cpp
@@ -1778,10 +1778,6 @@ float UAnimInstance::Montage_Play(UAnimMontage * MontageToPlay, float InPlayRate
 				*GetNameSafe(MontageToPlay), *GetNameSafe(MontageToPlay->GetSkeleton()), *GetNameSafe(CurrentSkeleton));
 		}
 	}
-	else
-	{
-		UE_LOG(LogAnimation, Warning, TEXT("Trying to play invalid Montage (%s)"), *GetNameSafe(MontageToPlay));
-	}
 
 	return 0.f;
 }
