* disabled log spam

diff --git a/Engine/Source/Runtime/Engine/Private/Slate/SceneViewport.cpp b/Engine/Source/Runtime/Engine/Private/Slate/SceneViewport.cpp
--- a/Engine/Source/Runtime/Engine/Private/Slate/SceneViewport.cpp
+++ b/Engine/Source/Runtime/Engine/Private/Slate/SceneViewport.cpp
@@ -1146,11 +1146,11 @@ void FSceneViewport::InitDynamicRHI()
 		if( !SlateRenderTargetHandle )
 		{
 			SlateRenderTargetHandle = new FSlateRenderTargetRHI( ShaderResourceTextureRHI, TexSizeX, TexSizeY );
-			UE_LOG(LogSlate, Log, TEXT("SRTH: %p, %d x %d"), ShaderResourceTextureRHI.GetReference(), TexSizeX, TexSizeY);
+//			UE_LOG(LogSlate, Log, TEXT("SRTH: %p, %d x %d"), ShaderResourceTextureRHI.GetReference(), TexSizeX, TexSizeY);
 		}
 		else
 		{
-			UE_LOG(LogSlate, Log, TEXT("SRTH: %p, %d x %d, prev %p"), ShaderResourceTextureRHI.GetReference(), TexSizeX, TexSizeY, SlateRenderTargetHandle->GetRHIRef().GetReference());
+//			UE_LOG(LogSlate, Log, TEXT("SRTH: %p, %d x %d, prev %p"), ShaderResourceTextureRHI.GetReference(), TexSizeX, TexSizeY, SlateRenderTargetHandle->GetRHIRef().GetReference());
 			SlateRenderTargetHandle->SetRHIRef( ShaderResourceTextureRHI, TexSizeX, TexSizeY );
 		}
 
