Fix for build issue caused by cl 2230263

diff --git a/Engine/Source/Runtime/Engine/Private/Actor.cpp b/Engine/Source/Runtime/Engine/Private/Actor.cpp
--- a/Engine/Source/Runtime/Engine/Private/Actor.cpp
+++ b/Engine/Source/Runtime/Engine/Private/Actor.cpp
@@ -609,7 +609,11 @@ void AActor::PostLoadSubobjects(FObjectInstancingGraph* OuterInstanceGraph)
 
 void AActor::ProcessEvent(UFunction* Function, void* Parameters)
 {
+	#if WITH_EDITOR
 	const bool bAllowScriptExecution = GAllowActorScriptExecutionInEditor || Function->GetBoolMetaData( "CallInEditor" );
+	#else
+	const bool bAllowScriptExecution = GAllowActorScriptExecutionInEditor;
+	#endif
 	if( ((GetWorld() && (GetWorld()->AreActorsInitialized() || bAllowScriptExecution)) || HasAnyFlags(RF_ClassDefaultObject)) && !GIsGarbageCollecting )
 	{
 		Super::ProcessEvent(Function, Parameters);
