Remove warning log happening with empty animation because those happen to morph target animations, where they dont' have key frames, but curves. Changed the verbose level to Log

339494  CONTENTEXAMPLES: FBX_Import_Options Compression Requested for Empty Animation
339671  CONTENTEXAMPLES: MorphTargets warning Compression Requested for Empty Animation
 339676  MORPHTARGETS: LogAnimation Warning Compression Requested for Empty Animation

diff --git a/Engine/Source/Runtime/Engine/Private/Animation/AnimationUtils.cpp b/Engine/Source/Runtime/Engine/Private/Animation/AnimationUtils.cpp
--- a/Engine/Source/Runtime/Engine/Private/Animation/AnimationUtils.cpp
+++ b/Engine/Source/Runtime/Engine/Private/Animation/AnimationUtils.cpp
@@ -1376,7 +1376,8 @@ void FAnimationUtils::CompressAnimSequenceExplicit(
 	}
 	else
 	{
-		UE_LOG(LogAnimation, Warning, TEXT("Compression Requested for Empty Animation %s"), *AnimSeq->GetName() );
+		// this can happen if the animation only contains curve - i.e. blendshape curves
+		UE_LOG(LogAnimation, Log, TEXT("Compression Requested for Empty Animation %s"), *AnimSeq->GetName() );
 	}
 #endif // WITH_EDITORONLY_DATA
 }
