Add generalized dialog box support for shootergame

* Unify dialog into a single manager
* Viewport can trap all widget show/hide calls so it can intelligently override focus, and visibility when showing dialogs
* Store dialog management on viewport
* The FShooterMessageMenu now utilizes the shared dialog manager
* You should be able to invoke the dialog from anywhere (welcome, main, ingame), safely, and rely on focus and visibility working properly

diff --git a/Engine/Source/Runtime/Engine/Classes/Engine/GameViewportClient.h b/Engine/Source/Runtime/Engine/Classes/Engine/GameViewportClient.h
--- a/Engine/Source/Runtime/Engine/Classes/Engine/GameViewportClient.h
+++ b/Engine/Source/Runtime/Engine/Classes/Engine/GameViewportClient.h
@@ -318,14 +318,14 @@ public:
 	 * @param  ViewportContent	The widget to add.  Must be valid.
 	 * @param  ZOrder  The Z-order index for this widget.  Larger values will cause the widget to appear on top of widgets with lower values.
 	 */
-	void AddViewportWidgetContent( TSharedRef<class SWidget> ViewportContent, const int32 ZOrder = 0 );
+	virtual void AddViewportWidgetContent( TSharedRef<class SWidget> ViewportContent, const int32 ZOrder = 0 );
 
 	/**
 	 * Removes a previously-added widget from the Slate viewport
 	 *
 	 * @param	ViewportContent  The widget to remove.  Must be valid.
 	 */
-	void RemoveViewportWidgetContent( TSharedRef<class SWidget> ViewportContent );
+	virtual void RemoveViewportWidgetContent( TSharedRef<class SWidget> ViewportContent );
 
 
 	/**
