[GitHub] - Added PlayerController::ProjectWorldLocationToScreen

Pull request: https://github.com/EpicGames/UnrealEngine/pull/301

diff --git a/Engine/Source/Runtime/Engine/Classes/GameFramework/PlayerController.h b/Engine/Source/Runtime/Engine/Classes/GameFramework/PlayerController.h
--- a/Engine/Source/Runtime/Engine/Classes/GameFramework/PlayerController.h
+++ b/Engine/Source/Runtime/Engine/Classes/GameFramework/PlayerController.h
@@ -381,6 +381,10 @@ public:
 	UFUNCTION(BlueprintCallable, Category = "Game|Player", meta = (FriendlyName = "ConvertScreenLocationToWorldSpace"))
 	void DeprojectScreenPositionToWorld(float ScreenX, float ScreenY, FVector & WorldLocation, FVector & WorldDirection) const;
 
+	/** Convert a World Space 3D position into a 2D Screen Space position. **/
+	UFUNCTION(BlueprintCallable, Category = "Game|Player", meta = (FriendlyName = "ConvertWorldLocationToScreenLocation"))
+	void ProjectWorldLocationToScreen(FVector WorldLocation, FVector2D& ScreenLocation) const;
+	
 	/**
 	  * Updates the rotation of player, based on ControlRotation after RotationInput has been applied.
 	  * This may then be modified by the PlayerCamera, and is passed to Pawn->FaceRotation().
diff --git a/Engine/Source/Runtime/Engine/Private/PlayerController.cpp b/Engine/Source/Runtime/Engine/Private/PlayerController.cpp
--- a/Engine/Source/Runtime/Engine/Private/PlayerController.cpp
+++ b/Engine/Source/Runtime/Engine/Private/PlayerController.cpp
@@ -1841,6 +1841,31 @@ void APlayerController::DeprojectScreenPositionToWorld(float ScreenX, float Scre
 }
 
 
+void APlayerController::ProjectWorldLocationToScreen(FVector WorldLocation, FVector2D& ScreenLocation) const
+{
+	ULocalPlayer* LocalPlayer = Cast<ULocalPlayer>(Player);
+	if (LocalPlayer != NULL && LocalPlayer->ViewportClient != NULL && LocalPlayer->ViewportClient->Viewport != NULL)
+	{
+		// Create a view family for the game viewport
+		FSceneViewFamilyContext ViewFamily( FSceneViewFamily::ConstructionValues(
+			LocalPlayer->ViewportClient->Viewport,
+			GetWorld()->Scene,
+			LocalPlayer->ViewportClient->EngineShowFlags )
+			.SetRealtimeUpdate(true) );
+
+		// Calculate a view where the player is to update the streaming from the players start location
+		FVector ViewLocation;
+		FRotator ViewRotation;
+		FSceneView* SceneView = LocalPlayer->CalcSceneView( &ViewFamily, /*out*/ ViewLocation, /*out*/ ViewRotation, LocalPlayer->ViewportClient->Viewport );
+
+		if (SceneView) 
+		{
+			SceneView->WorldToPixel(WorldLocation, ScreenLocation);
+		}
+	} 
+}
+
+
 bool APlayerController::GetHitResultAtScreenPosition(const FVector2D ScreenPosition, const ECollisionChannel TraceChannel, bool bTraceComplex, FHitResult& HitResult) const
 {
 	// Early out if we clicked on a HUD hitbox
