[341253] Blueprints - Apply custom scale to non-native root components at construction time.

This matches what's done for native root components when Blueprint actor instances are constructed.

Existing Blueprint assets with non-native RootComponents will first be reset back to unit scale if necessary in order to avoid any unexpected scaling of existing actor instances as a result of this change.

diff --git a/Engine/Source/Runtime/Core/Public/UObject/ObjectVersion.h b/Engine/Source/Runtime/Core/Public/UObject/ObjectVersion.h
--- a/Engine/Source/Runtime/Core/Public/UObject/ObjectVersion.h
+++ b/Engine/Source/Runtime/Core/Public/UObject/ObjectVersion.h
@@ -582,6 +582,8 @@ enum EUnrealEngineObjectUE4Version
 	VER_UE4_ADD_LINEAR_COLOR_SAMPLER,
 	// Added StringAssetReferencesMap to support renames of FStringAssetReference properties.
 	VER_UE4_ADD_STRING_ASSET_REFERENCES_MAP,
+	// Apply scale from SCS RootComponent details in the Blueprint Editor to new actor instances at construction time
+	VER_UE4_BLUEPRINT_USE_SCS_ROOTCOMPONENT_SCALE,
 
 	// -----<new versions can be added before this line>-------------------------------------------------
 	// - this needs to be the last line (see note below)
diff --git a/Engine/Source/Runtime/Engine/Private/SCS_Node.cpp b/Engine/Source/Runtime/Engine/Private/SCS_Node.cpp
--- a/Engine/Source/Runtime/Engine/Private/SCS_Node.cpp
+++ b/Engine/Source/Runtime/Engine/Private/SCS_Node.cpp
@@ -55,7 +55,8 @@ void USCS_Node::ExecuteNodeOnActor(AActor* Actor, USceneComponent* ParentCompone
 			// If NULL is passed in, we are the root, so set transform and assign as RootComponent on Actor
 			if (ParentComponent == NULL || (ParentComponent && ParentComponent->IsPendingKill()))
 			{
-				NewSceneComp->SetWorldTransform(*RootTransform);
+				// Note: We use the scale vector from the component template rather than the passed in scale in this case (to match what happens with a native root)
+				NewSceneComp->SetWorldTransform(FTransform(RootTransform->GetRotation(), RootTransform->GetLocation(), NewSceneComp->RelativeScale3D));
 				Actor->SetRootComponent(NewSceneComp);
 			}
 			// Otherwise, attach to parent component passed in
diff --git a/Engine/Source/Runtime/Engine/Private/SimpleConstructionScript.cpp b/Engine/Source/Runtime/Engine/Private/SimpleConstructionScript.cpp
--- a/Engine/Source/Runtime/Engine/Private/SimpleConstructionScript.cpp
+++ b/Engine/Source/Runtime/Engine/Private/SimpleConstructionScript.cpp
@@ -38,10 +38,10 @@ void USimpleConstructionScript::Serialize(FArchive& Ar)
 
 	if(Ar.IsLoading())
 	{
+		int32 NodeIndex;
+
 		if(Ar.UE4Ver() < VER_UE4_REMOVE_NATIVE_COMPONENTS_FROM_BLUEPRINT_SCS)
 		{
-			int32 NodeIndex;
-
 			// If we previously had a root node, we need to move it into the new RootNodes array. This is done in Serialize() in order to support SCS preloading (which relies on a valid RootNodes array).
 			if(RootNode_DEPRECATED != NULL)
 			{
@@ -158,6 +158,27 @@ void USimpleConstructionScript::PostLoad()
 
 	// Ensure that we have a valid scene root
 	ValidateSceneRootNodes();
+
+	// Reset non-native "root" scene component scale values, prior to the change in which
+	// we began applying custom scale values to root components at construction time. This
+	// way older, existing Blueprint actor instances won't start unexpectedly getting scaled.
+	if(GetLinkerUE4Version() < VER_UE4_BLUEPRINT_USE_SCS_ROOTCOMPONENT_SCALE)
+	{
+		for(NodeIndex = 0; NodeIndex < RootNodes.Num(); ++NodeIndex)
+		{
+			USCS_Node* Node = RootNodes[NodeIndex];
+			if(Node->ParentComponentOrVariableName == NAME_None)
+			{
+				USceneComponent* SceneComponentTemplate = Cast<USceneComponent>(Node->ComponentTemplate);
+				if(SceneComponentTemplate != nullptr
+					&& SceneComponentTemplate->RelativeScale3D != FVector(1.0f, 1.0f, 1.0f))
+				{
+					UE_LOG(LogBlueprint, Warning, TEXT("Found custom SceneComponent scale for %s (%s) saved prior to being usable; reset to default."), *SceneComponentTemplate->GetName(), *SceneComponentTemplate->RelativeScale3D.ToString());
+					SceneComponentTemplate->RelativeScale3D = FVector(1.0f, 1.0f, 1.0f);
+				}
+			}
+		}
+	}
 }
 
 void USimpleConstructionScript::FixupRootNodeParentReferences()
