Fix for toggling fullscreen not working when -windowed or -fullscreen is specified on the command line.

diff --git a/Engine/Source/Runtime/Engine/Classes/GameFramework/GameUserSettings.h b/Engine/Source/Runtime/Engine/Classes/GameFramework/GameUserSettings.h
--- a/Engine/Source/Runtime/Engine/Classes/GameFramework/GameUserSettings.h
+++ b/Engine/Source/Runtime/Engine/Classes/GameFramework/GameUserSettings.h
@@ -11,8 +11,8 @@ class ENGINE_API UGameUserSettings : public UObject
 
 public:
 
-	/** Applies all current user settings to the game and saves to permanent storage (e.g. file). */
-	virtual void ApplySettings();
+	/** Applies all current user settings to the game and saves to permanent storage (e.g. file), optionally checking for command line overrides */
+	virtual void ApplySettings( bool bCheckCommandLineOverrides=false );
 
 	/** Returns the user setting for game screen resolution, in pixels. */
 	FIntPoint GetScreenResolution() const;
@@ -93,7 +93,7 @@ public:
 	static void LoadConfigIni( bool bForceReload=false );
 
 	/** Apply any runtime overrides to the settings */
-	static void RequestResolutionChange(int32 InResolutionX, int32 InResolutionY, EWindowMode::Type InWindowMode);
+	static void RequestResolutionChange(int32 InResolutionX, int32 InResolutionY, EWindowMode::Type InWindowMode, bool bConditionallyOverride);
 
 	/** Whether to use VSync or not. (public to allow UI to connect to it) */
 	UPROPERTY(config)
diff --git a/Engine/Source/Runtime/Engine/Private/GameEngine.cpp b/Engine/Source/Runtime/Engine/Private/GameEngine.cpp
--- a/Engine/Source/Runtime/Engine/Private/GameEngine.cpp
+++ b/Engine/Source/Runtime/Engine/Private/GameEngine.cpp
@@ -422,7 +422,7 @@ void UGameEngine::Init(IEngineLoop* InEngineLoop)
 
 	// Load and apply user game settings
 	GetGameUserSettings()->LoadSettings();
-	GetGameUserSettings()->ApplySettings();
+	GetGameUserSettings()->ApplySettings(true);
 
 	// Create game instance.  For GameEngine, this should be the only GameInstance that ever gets created.
 	{
diff --git a/Engine/Source/Runtime/Engine/Private/GameUserSettings.cpp b/Engine/Source/Runtime/Engine/Private/GameUserSettings.cpp
--- a/Engine/Source/Runtime/Engine/Private/GameUserSettings.cpp
+++ b/Engine/Source/Runtime/Engine/Private/GameUserSettings.cpp
@@ -182,7 +182,7 @@ void UGameUserSettings::ValidateSettings()
 	UpdateVersion();
 }
 
-void UGameUserSettings::ApplySettings()
+void UGameUserSettings::ApplySettings(bool bCheckCommandLineOverrides)
 {
 	ValidateSettings();
 
@@ -195,7 +195,7 @@ void UGameUserSettings::ApplySettings()
 	}
 
 	// Request a resolution change
-	RequestResolutionChange(ResolutionSizeX, ResolutionSizeY, NewWindowMode);
+	RequestResolutionChange(ResolutionSizeX, ResolutionSizeY, NewWindowMode, bCheckCommandLineOverrides);
 
 	// Update vsync cvar
 	{
@@ -237,12 +237,15 @@ void UGameUserSettings::LoadSettings( bool bForceReload/*=false*/ )
 	}
 }
 
-void UGameUserSettings::RequestResolutionChange(int32 InResolutionX, int32 InResolutionY, EWindowMode::Type InWindowMode)
+void UGameUserSettings::RequestResolutionChange(int32 InResolutionX, int32 InResolutionY, EWindowMode::Type InWindowMode, bool bConditionallyOverride)
 {
-	UGameEngine::ConditionallyOverrideSettings(InResolutionX, InResolutionY, InWindowMode);
+	if (bConditionallyOverride)
 	{
-		IConsoleVariable* CVar = IConsoleManager::Get().FindConsoleVariable(TEXT("r.FullScreenMode"));
-		CVar->Set(InWindowMode);
+		UGameEngine::ConditionallyOverrideSettings(InResolutionX, InResolutionY, InWindowMode);
+		{
+			IConsoleVariable* CVar = IConsoleManager::Get().FindConsoleVariable(TEXT("r.FullScreenMode"));
+			CVar->Set(InWindowMode);
+		}
 	}
 	FSystemResolution::RequestResolutionChange(InResolutionX, InResolutionY, InWindowMode);
 }
@@ -304,7 +307,7 @@ void UGameUserSettings::PreloadResolutionSettings()
 		CVar->Set((int32)WindowMode);
 	}
 
-	RequestResolutionChange(ResolutionX, ResolutionY, WindowMode);
+	RequestResolutionChange(ResolutionX, ResolutionY, WindowMode, true);
 
 	IConsoleManager::Get().CallAllConsoleVariableSinks();
 }
