Rolled back fix for -windowed disabling fullscreen/windowed toggling.  Needs a bit more work :)

diff --git a/Engine/Source/Runtime/Engine/Classes/GameFramework/GameUserSettings.h b/Engine/Source/Runtime/Engine/Classes/GameFramework/GameUserSettings.h
--- a/Engine/Source/Runtime/Engine/Classes/GameFramework/GameUserSettings.h
+++ b/Engine/Source/Runtime/Engine/Classes/GameFramework/GameUserSettings.h
@@ -11,8 +11,8 @@ class ENGINE_API UGameUserSettings : public UObject
 
 public:
 
-	/** Applies all current user settings to the game and saves to permanent storage (e.g. file), optionally checking for command line overrides */
-	virtual void ApplySettings( bool bCheckCommandLineOverrides=false );
+	/** Applies all current user settings to the game and saves to permanent storage (e.g. file). */
+	virtual void ApplySettings();
 
 	/** Returns the user setting for game screen resolution, in pixels. */
 	FIntPoint GetScreenResolution() const;
@@ -93,7 +93,7 @@ public:
 	static void LoadConfigIni( bool bForceReload=false );
 
 	/** Apply any runtime overrides to the settings */
-	static void RequestResolutionChange(int32 InResolutionX, int32 InResolutionY, EWindowMode::Type InWindowMode, bool bConditionallyOverride);
+	static void RequestResolutionChange(int32 InResolutionX, int32 InResolutionY, EWindowMode::Type InWindowMode);
 
 	/** Whether to use VSync or not. (public to allow UI to connect to it) */
 	UPROPERTY(config)
diff --git a/Engine/Source/Runtime/Engine/Private/GameEngine.cpp b/Engine/Source/Runtime/Engine/Private/GameEngine.cpp
--- a/Engine/Source/Runtime/Engine/Private/GameEngine.cpp
+++ b/Engine/Source/Runtime/Engine/Private/GameEngine.cpp
@@ -422,7 +422,7 @@ void UGameEngine::Init(IEngineLoop* InEngineLoop)
 
 	// Load and apply user game settings
 	GetGameUserSettings()->LoadSettings();
-	GetGameUserSettings()->ApplySettings(true);
+	GetGameUserSettings()->ApplySettings();
 
 	// Create game instance.  For GameEngine, this should be the only GameInstance that ever gets created.
 	{
diff --git a/Engine/Source/Runtime/Engine/Private/GameUserSettings.cpp b/Engine/Source/Runtime/Engine/Private/GameUserSettings.cpp
--- a/Engine/Source/Runtime/Engine/Private/GameUserSettings.cpp
+++ b/Engine/Source/Runtime/Engine/Private/GameUserSettings.cpp
@@ -182,7 +182,7 @@ void UGameUserSettings::ValidateSettings()
 	UpdateVersion();
 }
 
-void UGameUserSettings::ApplySettings(bool bCheckCommandLineOverrides)
+void UGameUserSettings::ApplySettings()
 {
 	ValidateSettings();
 
@@ -195,7 +195,7 @@ void UGameUserSettings::ApplySettings(bool bCheckCommandLineOverrides)
 	}
 
 	// Request a resolution change
-	RequestResolutionChange(ResolutionSizeX, ResolutionSizeY, NewWindowMode, bCheckCommandLineOverrides);
+	RequestResolutionChange(ResolutionSizeX, ResolutionSizeY, NewWindowMode);
 
 	// Update vsync cvar
 	{
@@ -237,15 +237,12 @@ void UGameUserSettings::LoadSettings( bool bForceReload/*=false*/ )
 	}
 }
 
-void UGameUserSettings::RequestResolutionChange(int32 InResolutionX, int32 InResolutionY, EWindowMode::Type InWindowMode, bool bConditionallyOverride)
+void UGameUserSettings::RequestResolutionChange(int32 InResolutionX, int32 InResolutionY, EWindowMode::Type InWindowMode)
 {
-	if (bConditionallyOverride)
+	UGameEngine::ConditionallyOverrideSettings(InResolutionX, InResolutionY, InWindowMode);
 	{
-		UGameEngine::ConditionallyOverrideSettings(InResolutionX, InResolutionY, InWindowMode);
-		{
-			IConsoleVariable* CVar = IConsoleManager::Get().FindConsoleVariable(TEXT("r.FullScreenMode"));
-			CVar->Set(InWindowMode);
-		}
+		IConsoleVariable* CVar = IConsoleManager::Get().FindConsoleVariable(TEXT("r.FullScreenMode"));
+		CVar->Set(InWindowMode);
 	}
 	FSystemResolution::RequestResolutionChange(InResolutionX, InResolutionY, InWindowMode);
 }
@@ -307,7 +304,7 @@ void UGameUserSettings::PreloadResolutionSettings()
 		CVar->Set((int32)WindowMode);
 	}
 
-	RequestResolutionChange(ResolutionX, ResolutionY, WindowMode, true);
+	RequestResolutionChange(ResolutionX, ResolutionY, WindowMode);
 
 	IConsoleManager::Get().CallAllConsoleVariableSinks();
 }
