- Fixed the additive animation help UDN file
- Fixed display text to be more user friendly

diff --git a/Engine/Source/Runtime/Engine/Classes/Animation/AnimSequence.h b/Engine/Source/Runtime/Engine/Classes/Animation/AnimSequence.h
--- a/Engine/Source/Runtime/Engine/Classes/Animation/AnimSequence.h
+++ b/Engine/Source/Runtime/Engine/Classes/Animation/AnimSequence.h
@@ -66,13 +66,13 @@ UENUM()
 enum EAdditiveBasePoseType
 {
 	// will be deprecated
-	ABPT_None,
+	ABPT_None UMETA(DisplayName="None"),
 	// use ref pose of Skeleton as base
-	ABPT_RefPose,
+	ABPT_RefPose UMETA(DisplayName="Reference Pose"),
 	// use whole animation as a base pose. Need BasePoseSeq.
-	ABPT_AnimScaled,
+	ABPT_AnimScaled UMETA(DisplayName="Selected animation scaled"),
 	// use animation as a base pose. Need BasePoseSeq and RefFrameIndex (will clamp if outside).
-	ABPT_AnimFrame,
+	ABPT_AnimFrame UMETA(DisplayName="Selected animation frame"),
 	ABPT_MAX,
 };
 
@@ -437,11 +437,11 @@ class UAnimSequence : public UAnimSequenceBase
 	TEnumAsByte<enum EAdditiveAnimationType> AdditiveAnimType;
 
 	/* Additive refrerence pose type. Refer above enum type */
-	UPROPERTY(EditAnywhere, Category=AdditiveSettings)
+	UPROPERTY(EditAnywhere, Category=AdditiveSettings, meta=(DisplayName = "Base Pose Type"))
 	TEnumAsByte<enum EAdditiveBasePoseType> RefPoseType;
 
 	/* Additive reference animation if it's relevant - i.e. AnimScaled or AnimFrame **/
-	UPROPERTY(EditAnywhere, Category=AdditiveSettings)
+	UPROPERTY(EditAnywhere, Category=AdditiveSettings, meta=(DisplayName = "Base Pose Animation"))
 	class UAnimSequence* RefPoseSeq;
 
 	/* Additve reference frame if RefPoseType == AnimFrame **/
