Header cleanup: VectorFieldVisualization.h, AISystemBase.h

diff --git a/Engine/Source/Runtime/Engine/Classes/AI/AISystemBase.h b/Engine/Source/Runtime/Engine/Classes/AI/AISystemBase.h
--- a/Engine/Source/Runtime/Engine/Classes/AI/AISystemBase.h
+++ b/Engine/Source/Runtime/Engine/Classes/AI/AISystemBase.h
@@ -2,7 +2,6 @@
 
 #pragma once
 #include "Math/IntPoint.h"
-#include "ModuleManager.h"
 #include "AISystemBase.generated.h"
 
 class UWorld;
diff --git a/Engine/Source/Runtime/Engine/Private/VectorFieldVisualization.h b/Engine/Source/Runtime/Engine/Private/VectorFieldVisualization.h
--- a/Engine/Source/Runtime/Engine/Private/VectorFieldVisualization.h
+++ b/Engine/Source/Runtime/Engine/Private/VectorFieldVisualization.h
@@ -5,10 +5,8 @@
 ==============================================================================*/
 
 #pragma once
-
 #include "VertexFactory.h"
 #include "UniformBuffer.h"
-#include "ShaderParameters.h"
 
 /*------------------------------------------------------------------------------
 	Forward declarations.
