Added USplineComponent::GetDistanceAlongSplineAtPoint.
#branch UE4
#proj Runtime.Engine

diff --git a/Engine/Source/Runtime/Engine/Classes/Components/SplineComponent.h b/Engine/Source/Runtime/Engine/Classes/Components/SplineComponent.h
--- a/Engine/Source/Runtime/Engine/Classes/Components/SplineComponent.h
+++ b/Engine/Source/Runtime/Engine/Classes/Components/SplineComponent.h
@@ -80,6 +80,10 @@ class ENGINE_API USplineComponent : public USceneComponent
 	UFUNCTION(BlueprintCallable, Category=Spline)
 	void GetLocalLocationAndTangentAtSplinePoint(int32 PointIndex, FVector& LocalLocation, FVector& LocalTangent) const;
 
+	/** Get the distance along the spline at the spline point */
+	UFUNCTION(BlueprintCallable, Category=Spline)
+	float GetDistanceAlongSplineAtSplinePoint(int32 PointIndex) const;
+
 	/** Returns total length along this spline */
 	UFUNCTION(BlueprintCallable, Category=Spline) 
 	float GetSplineLength() const;
diff --git a/Engine/Source/Runtime/Engine/Private/SplineComponent.cpp b/Engine/Source/Runtime/Engine/Private/SplineComponent.cpp
--- a/Engine/Source/Runtime/Engine/Private/SplineComponent.cpp
+++ b/Engine/Source/Runtime/Engine/Private/SplineComponent.cpp
@@ -266,6 +266,16 @@ void USplineComponent::GetLocalLocationAndTangentAtSplinePoint(int32 PointIndex,
 	}
 }
 
+float USplineComponent::GetDistanceAlongSplineAtSplinePoint(int32 PointIndex) const
+{
+	if ((PointIndex >= 0) && (PointIndex < SplineInfo.Points.Num()))
+	{
+		return SplineReparamTable.Points[PointIndex * ReparamStepsPerSegment].InVal;
+	}
+
+	return 0.0f;
+}
+
 float USplineComponent::GetSplineLength() const
 {
 	const int32 NumPoints = SplineReparamTable.Points.Num();
@@ -320,9 +330,8 @@ FVector USplineComponent::GetWorldDirectionAtDistanceAlongSpline(float Distance)
 
 FRotator USplineComponent::GetWorldRotationAtDistanceAlongSpline(float Distance) const
 {
-	const float Param = SplineReparamTable.Eval(Distance, 0.f);
-	const FVector Forward = SplineInfo.EvalDerivative(Param, FVector(1.0f, 0.0f, 0.0f));
-	return Forward.Rotation() + ComponentToWorld.Rotator();
+	const FVector Dir = GetWorldDirectionAtDistanceAlongSpline(Distance);
+	return Dir.Rotation();
 }
 
 
@@ -364,19 +373,8 @@ FVector USplineComponent::GetWorldDirectionAtTime(float Time, bool bUseConstantV
 
 FRotator USplineComponent::GetWorldRotationAtTime(float Time, bool bUseConstantVelocity) const
 {
-	if (Duration == 0.0f)
-	{
-		return FRotator::ZeroRotator;
-	}
-
-	if (bUseConstantVelocity)
-	{
-		return GetWorldRotationAtDistanceAlongSpline(Time / Duration * GetSplineLength());
-	}
-
-	const float TimeMultiplier = (SplineInfo.Points.Num() - 1.0f) / Duration;
-	const FVector Forward = SplineInfo.EvalDerivative(Time * TimeMultiplier, FVector(1.0f, 0.0f, 0.0f));
-	return Forward.Rotation() + ComponentToWorld.Rotator();
+	FVector WorldDir = GetWorldDirectionAtTime(Time, bUseConstantVelocity);
+	return WorldDir.Rotation();
 }
 
 
