Fix crash due to corrupted model data.

diff --git a/Engine/Source/Runtime/Engine/Private/ModelRender.cpp b/Engine/Source/Runtime/Engine/Private/ModelRender.cpp
--- a/Engine/Source/Runtime/Engine/Private/ModelRender.cpp
+++ b/Engine/Source/Runtime/Engine/Private/ModelRender.cpp
@@ -96,39 +96,43 @@ void UModelComponent::BuildRenderData()
 		for(int32 NodeIndex = 0;NodeIndex < Element.Nodes.Num();NodeIndex++)
 		{
 			const uint16& NodeIdx = Element.Nodes[NodeIndex];
-			checkf( TheModel->Nodes.IsValidIndex( NodeIdx ), TEXT( "Invalid Node Index, Idx:%d, Num:%d" ), NodeIdx, TheModel->Nodes.Num() );
-			const FBspNode& Node = TheModel->Nodes[NodeIdx];
-			checkf( TheModel->Surfs.IsValidIndex( Node.iSurf ), TEXT( "Invalid Surf Index, Idx:%d, Num:%d" ), Node.iSurf, TheModel->Surfs.Num() );
-			const FBspSurf& Surf = TheModel->Surfs[Node.iSurf];
+			if( ensureMsgf( TheModel->Nodes.IsValidIndex( NodeIdx ), TEXT( "Invalid Node Index, Idx:%d, Num:%d" ), NodeIdx, TheModel->Nodes.Num() ) )
+			{
+				const FBspNode& Node = TheModel->Nodes[NodeIdx];
+				if( ensureMsgf( TheModel->Surfs.IsValidIndex(Node.iSurf), TEXT("Invalid Surf Index, Idx:%d, Num:%d"), Node.iSurf, TheModel->Surfs.Num() ) )
+				{
+					const FBspSurf& Surf = TheModel->Surfs[Node.iSurf];
 
 #if WITH_EDITOR
-			// If we're not in a game world, check the surface visibility
-			if (!bIsGameWorld && !ShouldDrawSurface(Surf))
-			{
-				continue;
-			}
+					// If we're not in a game world, check the surface visibility
+					if(!bIsGameWorld && !ShouldDrawSurface(Surf))
+					{
+						continue;
+					}
 #endif
 
-			// Don't put portal polygons in the static index buffer.
-			if (Surf.PolyFlags & PF_Portal)
-				continue;
+					// Don't put portal polygons in the static index buffer.
+					if(Surf.PolyFlags & PF_Portal)
+						continue;
 
-			for(uint32 BackFace = 0;BackFace < (uint32)((Surf.PolyFlags & PF_TwoSided) ? 2 : 1);BackFace++)
-			{
-				for(int32 VertexIndex = 0;VertexIndex < Node.NumVertices;VertexIndex++)
-				{
-					Element.BoundingBox += TheModel->Points[TheModel->Verts[Node.iVertPool + VertexIndex].pVertex];
-				}
+					for(uint32 BackFace = 0; BackFace < (uint32)((Surf.PolyFlags & PF_TwoSided) ? 2 : 1); BackFace++)
+					{
+						for(int32 VertexIndex = 0; VertexIndex < Node.NumVertices; VertexIndex++)
+						{
+							Element.BoundingBox += TheModel->Points[TheModel->Verts[Node.iVertPool + VertexIndex].pVertex];
+						}
 
-				for(int32 VertexIndex = 2;VertexIndex < Node.NumVertices;VertexIndex++)
-				{
-					IndexBuffer->Indices.Add(Node.iVertexIndex + Node.NumVertices * BackFace);
-					IndexBuffer->Indices.Add(Node.iVertexIndex + Node.NumVertices * BackFace + VertexIndex);
-					IndexBuffer->Indices.Add(Node.iVertexIndex + Node.NumVertices * BackFace + VertexIndex - 1);
-					Element.NumTriangles++;
+						for(int32 VertexIndex = 2; VertexIndex < Node.NumVertices; VertexIndex++)
+						{
+							IndexBuffer->Indices.Add(Node.iVertexIndex + Node.NumVertices * BackFace);
+							IndexBuffer->Indices.Add(Node.iVertexIndex + Node.NumVertices * BackFace + VertexIndex);
+							IndexBuffer->Indices.Add(Node.iVertexIndex + Node.NumVertices * BackFace + VertexIndex - 1);
+							Element.NumTriangles++;
+						}
+						Element.MinVertexIndex = FMath::Min(Node.iVertexIndex + Node.NumVertices * BackFace, Element.MinVertexIndex);
+						Element.MaxVertexIndex = FMath::Max(Node.iVertexIndex + Node.NumVertices * BackFace + Node.NumVertices - 1, Element.MaxVertexIndex);
+					}
 				}
-				Element.MinVertexIndex = FMath::Min(Node.iVertexIndex + Node.NumVertices * BackFace,Element.MinVertexIndex);
-				Element.MaxVertexIndex = FMath::Max(Node.iVertexIndex + Node.NumVertices * BackFace + Node.NumVertices - 1,Element.MaxVertexIndex);
 			}
 		}
 
