Fix unnecessary BP dirtying when opening a BP asset with a native component attachment.

diff --git a/Engine/Source/Runtime/Engine/Private/SCS_Node.cpp b/Engine/Source/Runtime/Engine/Private/SCS_Node.cpp
--- a/Engine/Source/Runtime/Engine/Private/SCS_Node.cpp
+++ b/Engine/Source/Runtime/Engine/Private/SCS_Node.cpp
@@ -305,7 +305,7 @@ void USCS_Node::SetParent(USceneComponent* InParentComponent)
 	const FName NewParentComponentOwnerClassName = NAME_None;
 
 	// Only modify if it differs from current
-	if(bIsParentComponentNative
+	if(!bIsParentComponentNative
 		|| ParentComponentOrVariableName != NewParentComponentOrVariableName
 		|| ParentComponentOwnerClassName != NewParentComponentOwnerClassName)
 	{
