Revert mouse focus modification

diff --git a/Engine/Source/Runtime/Engine/Private/GameViewportClient.cpp b/Engine/Source/Runtime/Engine/Private/GameViewportClient.cpp
--- a/Engine/Source/Runtime/Engine/Private/GameViewportClient.cpp
+++ b/Engine/Source/Runtime/Engine/Private/GameViewportClient.cpp
@@ -217,7 +217,7 @@ bool UGameViewportClient::InputKey(FViewport* InViewport, int32 ControllerId, FK
 	if (!bResult)
 	{
 		ULocalPlayer* const TargetPlayer = GEngine->GetLocalPlayerFromControllerId(this, ControllerId);
-		if (TargetPlayer && TargetPlayer->PlayerController && InViewport->HasMouseCapture())
+		if (TargetPlayer && TargetPlayer->PlayerController)
 		{
 			bResult = TargetPlayer->PlayerController->InputKey(Key, EventType, AmountDepressed, bGamepad);
 		}
diff --git a/Engine/Source/Runtime/Engine/Private/Slate/SceneViewport.cpp b/Engine/Source/Runtime/Engine/Private/Slate/SceneViewport.cpp
--- a/Engine/Source/Runtime/Engine/Private/Slate/SceneViewport.cpp
+++ b/Engine/Source/Runtime/Engine/Private/Slate/SceneViewport.cpp
@@ -309,12 +309,8 @@ FReply FSceneViewport::OnMouseButtonDown( const FGeometry& InGeometry, const FPo
 	// Start a new reply state
 	// Prevent throttling when interacting with the viewport so we can move around in it
 	CurrentReplyState = FReply::Handled().PreventThrottling();
-	
-	if (!IsPlayInEditorViewport() || HasMouseCapture())
-	{
-		KeyStateMap.Add(InMouseEvent.GetEffectingButton(), true);
-	}
 
+	KeyStateMap.Add(InMouseEvent.GetEffectingButton(), true);
 	UpdateModifierKeys( InMouseEvent );
 	UpdateCachedMousePos( InGeometry, InMouseEvent );
 	UpdateCachedGeometry(InGeometry);
