Fix defaultplayername not being read off the current game mode

diff --git a/Engine/Source/Runtime/Engine/Private/Controller.cpp b/Engine/Source/Runtime/Engine/Private/Controller.cpp
--- a/Engine/Source/Runtime/Engine/Private/Controller.cpp
+++ b/Engine/Source/Runtime/Engine/Private/Controller.cpp
@@ -422,7 +422,7 @@ void AController::InitPlayerState()
 			{
 				// don't call SetPlayerName() as that will broadcast entry messages but the GameMode hasn't had a chance
 				// to potentially apply a player/bot name yet
-				PlayerState->PlayerName = GetDefault<AGameMode>()->DefaultPlayerName;
+				PlayerState->PlayerName = GameMode->DefaultPlayerName;
 			}
 		}
 	}
