[342381] Fix broken BP actor scaling in QA_Blueprints map.

Updated previous fix to apply non-native root component scale from template only when new Actors are first spawned into a map.

Also revised SCS node scale fixup code to more explicitly ensure that only the non-native root component node will be updated on BP postload.

diff --git a/Engine/Source/Runtime/Engine/Private/Actor.cpp b/Engine/Source/Runtime/Engine/Private/Actor.cpp
--- a/Engine/Source/Runtime/Engine/Private/Actor.cpp
+++ b/Engine/Source/Runtime/Engine/Private/Actor.cpp
@@ -2254,7 +2254,7 @@ void AActor::PostSpawnInitialize(FVector const& SpawnLocation, FRotator const& S
 	if (!bDeferConstruction)
 	{
 		// Preserve original root component scale
-		const FVector SpawnScale = GetRootComponent() ? GetRootComponent()->RelativeScale3D : FVector(1.0f, 1.0f, 1.0f);
+		const FVector SpawnScale = GetRootComponent() ? GetRootComponent()->RelativeScale3D : FVector::ZeroVector;
 		ExecuteConstruction( FTransform(SpawnRotation, SpawnLocation, SpawnScale), NULL );
 		PostActorConstruction();
 	}
diff --git a/Engine/Source/Runtime/Engine/Private/ActorConstruction.cpp b/Engine/Source/Runtime/Engine/Private/ActorConstruction.cpp
--- a/Engine/Source/Runtime/Engine/Private/ActorConstruction.cpp
+++ b/Engine/Source/Runtime/Engine/Private/ActorConstruction.cpp
@@ -382,7 +382,13 @@ void AActor::ExecuteConstruction(const FTransform& Transform, const FComponentIn
 #if WITH_EDITOR
 				BillboardComponent->Sprite = (UTexture2D*)(StaticLoadObject(UTexture2D::StaticClass(), NULL, TEXT("/Engine/EditorResources/BadBlueprintSprite.BadBlueprintSprite"), NULL, LOAD_None, NULL));
 #endif
-				BillboardComponent->SetRelativeTransform(Transform);
+				FTransform RelativeTransform = Transform;
+				if(RelativeTransform.GetScale3D().IsZero())
+				{
+					RelativeTransform.SetScale3D(FVector(1.0f, 1.0f, 1.0f));
+				}
+
+				BillboardComponent->SetRelativeTransform(RelativeTransform);
 
 				SetRootComponent(BillboardComponent);
 				FinishAndRegisterComponent(BillboardComponent);
diff --git a/Engine/Source/Runtime/Engine/Private/SCS_Node.cpp b/Engine/Source/Runtime/Engine/Private/SCS_Node.cpp
--- a/Engine/Source/Runtime/Engine/Private/SCS_Node.cpp
+++ b/Engine/Source/Runtime/Engine/Private/SCS_Node.cpp
@@ -55,8 +55,13 @@ void USCS_Node::ExecuteNodeOnActor(AActor* Actor, USceneComponent* ParentCompone
 			// If NULL is passed in, we are the root, so set transform and assign as RootComponent on Actor
 			if (ParentComponent == NULL || (ParentComponent && ParentComponent->IsPendingKill()))
 			{
-				// Note: We use the scale vector from the component template rather than the passed in scale in this case (to match what happens with a native root)
-				NewSceneComp->SetWorldTransform(FTransform(RootTransform->GetRotation(), RootTransform->GetLocation(), NewSceneComp->RelativeScale3D));
+				FTransform SpawnTransform = *RootTransform;
+				if(SpawnTransform.GetScale3D().IsZero())
+				{
+					// Note: We use the scale vector from the component template when spawning (to match what happens with a native root)
+					SpawnTransform.SetScale3D(NewSceneComp->RelativeScale3D);
+				}
+				NewSceneComp->SetWorldTransform(SpawnTransform);
 				Actor->SetRootComponent(NewSceneComp);
 			}
 			// Otherwise, attach to parent component passed in
diff --git a/Engine/Source/Runtime/Engine/Private/SimpleConstructionScript.cpp b/Engine/Source/Runtime/Engine/Private/SimpleConstructionScript.cpp
--- a/Engine/Source/Runtime/Engine/Private/SimpleConstructionScript.cpp
+++ b/Engine/Source/Runtime/Engine/Private/SimpleConstructionScript.cpp
@@ -164,17 +164,38 @@ void USimpleConstructionScript::PostLoad()
 	// way older, existing Blueprint actor instances won't start unexpectedly getting scaled.
 	if(GetLinkerUE4Version() < VER_UE4_BLUEPRINT_USE_SCS_ROOTCOMPONENT_SCALE)
 	{
-		for(NodeIndex = 0; NodeIndex < RootNodes.Num(); ++NodeIndex)
+		// Get the BlueprintGeneratedClass that owns the SCS
+		UClass* BPGeneratedClass = GetOwnerClass();
+		if(BPGeneratedClass != nullptr)
 		{
-			USCS_Node* Node = RootNodes[NodeIndex];
-			if(Node->ParentComponentOrVariableName == NAME_None)
+			// Get the Blueprint class default object
+			AActor* CDO = Cast<AActor>(BPGeneratedClass->GetDefaultObject(false));
+			if(CDO != NULL)
 			{
-				USceneComponent* SceneComponentTemplate = Cast<USceneComponent>(Node->ComponentTemplate);
-				if(SceneComponentTemplate != nullptr
-					&& SceneComponentTemplate->RelativeScale3D != FVector(1.0f, 1.0f, 1.0f))
+				// Check for a native root component
+				USceneComponent* NativeRootComponent = CDO->GetRootComponent();
+				if(NativeRootComponent == nullptr)
 				{
-					UE_LOG(LogBlueprint, Warning, TEXT("Found custom SceneComponent scale for %s (%s) saved prior to being usable; reset to default."), *SceneComponentTemplate->GetName(), *SceneComponentTemplate->RelativeScale3D.ToString());
-					SceneComponentTemplate->RelativeScale3D = FVector(1.0f, 1.0f, 1.0f);
+					// If no native root component exists, find the first non-native, non-parented SCS node with a
+					// scene component template. This will be designated as the root component at construction time.
+					for(NodeIndex = 0; NodeIndex < RootNodes.Num(); ++NodeIndex)
+					{
+						USCS_Node* Node = RootNodes[NodeIndex];
+						if(Node->ParentComponentOrVariableName == NAME_None)
+						{
+							// Note that we have to check for nullptr here, because it may be an ActorComponent type
+							USceneComponent* SceneComponentTemplate = Cast<USceneComponent>(Node->ComponentTemplate);
+							if(SceneComponentTemplate != nullptr
+								&& SceneComponentTemplate->RelativeScale3D != FVector(1.0f, 1.0f, 1.0f))
+							{
+								UE_LOG(LogBlueprint, Warning, TEXT("%s: Found non-native root component custom scale for %s (%s) saved prior to being usable; reverting to default scale."), *BPGeneratedClass->GetName(), *Node->GetVariableName().ToString(), *SceneComponentTemplate->RelativeScale3D.ToString());
+								SceneComponentTemplate->RelativeScale3D = FVector(1.0f, 1.0f, 1.0f);
+							}
+
+							// Done - no need to fix up any other nodes.
+							break;
+						}
+					}
 				}
 			}
 		}
@@ -318,10 +339,16 @@ void USimpleConstructionScript::ExecuteScriptOnActor(AActor* Actor, const FTrans
 	}
 	else if(Actor->GetRootComponent() == NULL) // Must have a root component at the end of SCS, so if we don't have one already (from base class), create a SceneComponent now
 	{
+		FTransform SpawnTransform = RootTransform;
+		if(SpawnTransform.GetScale3D().IsZero())
+		{
+			SpawnTransform.SetScale3D(FVector(1.0f, 1.0f, 1.0f));
+		}
+
 		USceneComponent* SceneComp = NewObject<USceneComponent>(Actor);
 		SceneComp->SetFlags(RF_Transactional);
 		SceneComp->bCreatedByConstructionScript = true;
-		SceneComp->SetWorldTransform(RootTransform);
+		SceneComp->SetWorldTransform(SpawnTransform);
 		Actor->SetRootComponent(SceneComp);
 		SceneComp->RegisterComponent();
 	}
