Fixed null ptr crash in undoing camera group modification

#ttp 341926 - EDITOR: MATINEE: CRASH: In a cameraAnim asset, editing the property of a group and then undoing results in a crash

#branch UE4

diff --git a/Engine/Source/Runtime/Engine/Private/Interpolation.cpp b/Engine/Source/Runtime/Engine/Private/Interpolation.cpp
--- a/Engine/Source/Runtime/Engine/Private/Interpolation.cpp
+++ b/Engine/Source/Runtime/Engine/Private/Interpolation.cpp
@@ -9655,7 +9655,7 @@ UInterpGroupInstDirector::UInterpGroupInstDirector(const class FPostConstructIni
 void UInterpGroupCamera::PostEditChangeProperty(FPropertyChangedEvent& PropertyChangedEvent)
 {
 #if WITH_EDITORONLY_DATA
-	if (PropertyChangedEvent.Property->GetName() == TEXT("AnimSeqName"))
+	if (PropertyChangedEvent.Property && PropertyChangedEvent.Property->GetName() == TEXT("AnimSeqName"))
 	{
 		UCameraAnim * CameraAnim = CastChecked<UCameraAnim>(GetOuter());
 		UInterpGroup* Group = CameraAnim->PreviewInterpGroup;
