Header cleanup: GPUSort.h, PlayerController.h, NavLinkCustomComponent.h, Actor.h.

diff --git a/Engine/Source/Editor/UnrealEd/Classes/ThumbnailRendering/WorldThumbnailInfo.h b/Engine/Source/Editor/UnrealEd/Classes/ThumbnailRendering/WorldThumbnailInfo.h
--- a/Engine/Source/Editor/UnrealEd/Classes/ThumbnailRendering/WorldThumbnailInfo.h
+++ b/Engine/Source/Editor/UnrealEd/Classes/ThumbnailRendering/WorldThumbnailInfo.h
@@ -5,8 +5,8 @@
  */
 
 #pragma once
-
-#include "Editor/UnrealEd/Classes/ThumbnailRendering/SceneThumbnailInfo.h"
+#include "Camera/CameraTypes.h"
+#include "ThumbnailRendering/SceneThumbnailInfo.h"
 #include "WorldThumbnailInfo.generated.h"
 
 UENUM()
diff --git a/Engine/Source/Runtime/Engine/Classes/AI/Navigation/NavLinkCustomComponent.h b/Engine/Source/Runtime/Engine/Classes/AI/Navigation/NavLinkCustomComponent.h
--- a/Engine/Source/Runtime/Engine/Classes/AI/Navigation/NavLinkCustomComponent.h
+++ b/Engine/Source/Runtime/Engine/Classes/AI/Navigation/NavLinkCustomComponent.h
@@ -1,7 +1,6 @@
 // Copyright 1998-2014 Epic Games, Inc. All Rights Reserved.
 
 #pragma once
-#include "NavigationTypes.h"
 #include "AI/Navigation/NavRelevantComponent.h"
 #include "AI/Navigation/NavLinkCustomInterface.h"
 #include "NavLinkCustomComponent.generated.h"
diff --git a/Engine/Source/Runtime/Engine/Classes/GameFramework/Actor.h b/Engine/Source/Runtime/Engine/Classes/GameFramework/Actor.h
--- a/Engine/Source/Runtime/Engine/Classes/GameFramework/Actor.h
+++ b/Engine/Source/Runtime/Engine/Classes/GameFramework/Actor.h
@@ -1,20 +1,16 @@
 // Copyright 1998-2014 Epic Games, Inc. All Rights Reserved.
 
 #pragma once
-
+#include "ComponentInstanceDataCache.h"
+#include "Components/SceneComponent.h"
 #include "EngineDefines.h"
-
 #include "Engine/EngineBaseTypes.h"
-#include "Components/ActorComponent.h"
 #include "Engine/EngineTypes.h"
+#include "InputCoreTypes.h"
+#include "RenderCommandFence.h"
 
-#include "Runtime/InputCore/Classes/InputCoreTypes.h"
-#include "Runtime/RenderCore/Public/RenderCommandFence.h"
-
-#include "Camera/CameraTypes.h"
-#include "Components/SceneComponent.h"
-#include "ComponentInstanceDataCache.h"
-
+struct FHitResult;
+class AActor;
 class FTimerManager; 
 
 UENUM(BlueprintType)
@@ -36,11 +32,11 @@ ENGINE_API DECLARE_LOG_CATEGORY_EXTERN(LogActor, Log, Warning);
  
 
 // Delegate signatures
-DECLARE_DYNAMIC_MULTICAST_DELEGATE_FourParams( FTakeAnyDamageSignature, float, Damage, const class UDamageType*, DamageType, class AController*, InstigatedBy, class AActor*, DamageCauser );
-DECLARE_DYNAMIC_MULTICAST_DELEGATE_EightParams( FTakePointDamageSignature, float, Damage, class AController*, InstigatedBy, FVector, HitLocation, class UPrimitiveComponent*, FHitComponent, FName, BoneName, FVector, ShotFromDirection, const class UDamageType*, DamageType, class AActor*, DamageCauser );
-DECLARE_DYNAMIC_MULTICAST_DELEGATE_OneParam( FActorBeginOverlapSignature, class AActor*, OtherActor );
-DECLARE_DYNAMIC_MULTICAST_DELEGATE_OneParam( FActorEndOverlapSignature, class AActor*, OtherActor );
-DECLARE_DYNAMIC_MULTICAST_DELEGATE_FourParams( FActorHitSignature, class AActor*, SelfActor, class AActor*, OtherActor, FVector, NormalImpulse, const FHitResult&, Hit );
+DECLARE_DYNAMIC_MULTICAST_DELEGATE_FourParams( FTakeAnyDamageSignature, float, Damage, const class UDamageType*, DamageType, class AController*, InstigatedBy, AActor*, DamageCauser );
+DECLARE_DYNAMIC_MULTICAST_DELEGATE_EightParams( FTakePointDamageSignature, float, Damage, class AController*, InstigatedBy, FVector, HitLocation, class UPrimitiveComponent*, FHitComponent, FName, BoneName, FVector, ShotFromDirection, const class UDamageType*, DamageType, AActor*, DamageCauser );
+DECLARE_DYNAMIC_MULTICAST_DELEGATE_OneParam( FActorBeginOverlapSignature, AActor*, OtherActor );
+DECLARE_DYNAMIC_MULTICAST_DELEGATE_OneParam( FActorEndOverlapSignature, AActor*, OtherActor );
+DECLARE_DYNAMIC_MULTICAST_DELEGATE_FourParams( FActorHitSignature, AActor*, SelfActor, AActor*, OtherActor, FVector, NormalImpulse, const FHitResult&, Hit );
 
 DECLARE_DYNAMIC_MULTICAST_DELEGATE( FActorBeginCursorOverSignature );
 DECLARE_DYNAMIC_MULTICAST_DELEGATE( FActorEndCursorOverSignature );
@@ -54,7 +50,7 @@ DECLARE_DYNAMIC_MULTICAST_DELEGATE_OneParam( FActorEndTouchOverSignature, ETouch
 DECLARE_DYNAMIC_MULTICAST_DELEGATE(FActorDestroyedSignature);
 DECLARE_DYNAMIC_MULTICAST_DELEGATE_OneParam(FActorEndPlaySignature, EEndPlayReason::Type, EndPlayReason);
 
-DECLARE_DELEGATE_FourParams(FMakeNoiseDelegate, class AActor*, float, class APawn*, const FVector&);
+DECLARE_DELEGATE_FourParams(FMakeNoiseDelegate, AActor*, float, class APawn*, const FVector&);
 
 DECLARE_CYCLE_STAT_EXTERN(TEXT("GetComponentsTime"),STAT_GetComponentsTime,STATGROUP_Engine,ENGINE_API);
 
@@ -92,7 +88,7 @@ class ENGINE_API AActor : public UObject
 private:
 	/** Used for replication (bNetUseOwnerRelevancy & bOnlyRelevantToOwner) and visibility (PrimitiveComponent bOwnerNoSee and bOnlyOwnerSee) */
 	UPROPERTY(replicated)
-	class AActor* Owner;
+	AActor* Owner;
 
 	/** Enables any collision on this actor */
 	UPROPERTY()
@@ -290,7 +286,7 @@ public:
 
 	/** Array of Actors whose Owner is this actor */
 	UPROPERTY(transient)
-	TArray<class AActor*> Children;    
+	TArray<AActor*> Children;    
 	
 	// Animation update rate control.
 public:
@@ -399,7 +395,7 @@ public:
 
 	/** The group this actor is a part of. */
 	UPROPERTY(transient)
-	class AActor* GroupActor;
+	AActor* GroupActor;
 
 	/** The scale to apply to any billboard components in editor builds (happens in any WITH_EDITOR build, including non-cooked games) */
 	UPROPERTY(EditAnywhere, BlueprintReadWrite, Category=Rendering, meta=(DisplayName="Editor Billboard Scale"))
@@ -411,7 +407,7 @@ public:
 
 	/** The Actor that owns the ChildActorComponent that owns this Actor */
 	UPROPERTY()
-	TWeakObjectPtr<class AActor> ParentComponentActor;	
+	TWeakObjectPtr<AActor> ParentComponentActor;	
 
 #if ENABLE_VISUAL_LOG
 	/** indicates where given actor will be writing its visual log data. Defaults to 'this' */
@@ -855,18 +851,18 @@ public:
 
 	/** Event when this actor takes ANY damage */
 	UFUNCTION(BlueprintImplementableEvent, BlueprintAuthorityOnly, meta=(FriendlyName = "AnyDamage"), Category="Damage")
-	virtual void ReceiveAnyDamage(float Damage, const class UDamageType* DamageType, class AController* InstigatedBy, class AActor* DamageCauser);
+	virtual void ReceiveAnyDamage(float Damage, const class UDamageType* DamageType, class AController* InstigatedBy, AActor* DamageCauser);
 	
 	/** 
 	 * Event when this actor takes RADIAL damage 
 	 * @todo Pass it the full array of hits instead of just one?
 	 */
 	UFUNCTION(BlueprintImplementableEvent, BlueprintAuthorityOnly, meta=(FriendlyName = "RadialDamage"), Category="Damage")
-	virtual void ReceiveRadialDamage(float DamageReceived, const class UDamageType* DamageType, FVector Origin, const struct FHitResult& HitInfo, class AController* InstigatedBy, class AActor* DamageCauser);
+	virtual void ReceiveRadialDamage(float DamageReceived, const class UDamageType* DamageType, FVector Origin, const struct FHitResult& HitInfo, class AController* InstigatedBy, AActor* DamageCauser);
 
 	/** Event when this actor takes POINT damage */
 	UFUNCTION(BlueprintImplementableEvent, BlueprintAuthorityOnly, meta=(FriendlyName = "PointDamage"), Category="Damage")
-	virtual void ReceivePointDamage(float Damage, const class UDamageType* DamageType, FVector HitLocation, FVector HitNormal, class UPrimitiveComponent* HitComponent, FName BoneName, FVector ShotFromDirection, class AController* InstigatedBy, class AActor* DamageCauser);
+	virtual void ReceivePointDamage(float Damage, const class UDamageType* DamageType, FVector HitLocation, FVector HitNormal, class UPrimitiveComponent* HitComponent, FName BoneName, FVector ShotFromDirection, class AController* InstigatedBy, AActor* DamageCauser);
 
 	/** Event called every frame */
 	UFUNCTION(BlueprintImplementableEvent, meta=(FriendlyName = "Tick"))
@@ -874,11 +870,11 @@ public:
 
 	/** Event when this actor overlaps another actor. */
 	UFUNCTION(BlueprintImplementableEvent, meta=(FriendlyName = "ActorBeginOverlap"), Category="Collision")
-	virtual void ReceiveActorBeginOverlap(class AActor* OtherActor);
+	virtual void ReceiveActorBeginOverlap(AActor* OtherActor);
 
 	/** Event when an actor no longer overlaps another actor, and they have separated. */
 	UFUNCTION(BlueprintImplementableEvent, meta=(FriendlyName = "ActorEndOverlap"), Category="Collision")
-	virtual void ReceiveActorEndOverlap(class AActor* OtherActor);
+	virtual void ReceiveActorEndOverlap(AActor* OtherActor);
 
 	/** Event when this actor has the mouse moved over it with the clickable interface. */
 	UFUNCTION(BlueprintImplementableEvent, meta=(FriendlyName = "ActorBeginCursorOver"), Category="Mouse Input")
@@ -930,7 +926,7 @@ public:
 
 	/** Event when this actor bumps into a blocking object, or blocks another actor that bumps into it. */
 	UFUNCTION(BlueprintImplementableEvent, meta=(FriendlyName = "Hit"), Category="Collision")
-	virtual void ReceiveHit(class UPrimitiveComponent* MyComp, class AActor* Other, class UPrimitiveComponent* OtherComp, bool bSelfMoved, FVector HitLocation, FVector HitNormal, FVector NormalImpulse, const FHitResult& Hit);
+	virtual void ReceiveHit(class UPrimitiveComponent* MyComp, AActor* Other, class UPrimitiveComponent* OtherComp, bool bSelfMoved, FVector HitLocation, FVector HitNormal, FVector NormalImpulse, const FHitResult& Hit);
 
 	/** Set the lifespan of this actor. When it expires the object will be destroyed. If requested lifespan is 0, the timer is cleared and the actor will not be destroyed. */
 	UFUNCTION(BlueprintCallable, Category="Utilities", meta=(Keywords = "delete destroy"))
@@ -1759,7 +1755,7 @@ protected:
 public:
 
 	/** Walk up the attachment chain from RootComponent until we encounter a different actor, and return it. If we are not attached to a component in a different actor, returns NULL */
-	virtual class AActor* GetAttachParentActor() const;
+	virtual AActor* GetAttachParentActor() const;
 
 	/** Walk up the attachment chain from RootComponent until we encounter a different actor, and return the socket name in the component. If we are not attached to a component in a different actor, returns NAME_None */
 	virtual FName GetAttachParentSocketName() const;
@@ -1854,11 +1850,11 @@ public:
 	 * @param DamageCauser		The Actor that directly caused the damage (e.g. the projectile that exploded, the rock that landed on you)
 	 * @return					The amount of damage actually applied.
 	 */
-	virtual float TakeDamage(float DamageAmount, struct FDamageEvent const& DamageEvent, class AController* EventInstigator, class AActor* DamageCauser);
+	virtual float TakeDamage(float DamageAmount, struct FDamageEvent const& DamageEvent, class AController* EventInstigator, AActor* DamageCauser);
 
 protected:
-	virtual float InternalTakeRadialDamage(float Damage, struct FRadialDamageEvent const& RadialDamageEvent, class AController* EventInstigator, class AActor* DamageCauser);
-	virtual float InternalTakePointDamage(float Damage, struct FPointDamageEvent const& RadialDamageEvent, class AController* EventInstigator, class AActor* DamageCauser);
+	virtual float InternalTakeRadialDamage(float Damage, struct FRadialDamageEvent const& RadialDamageEvent, class AController* EventInstigator, AActor* DamageCauser);
+	virtual float InternalTakePointDamage(float Damage, struct FPointDamageEvent const& RadialDamageEvent, class AController* EventInstigator, AActor* DamageCauser);
 public:
 
 	/* Called when this actor becomes the given PlayerController's ViewTarget. Triggers the Blueprint event K2_OnBecomeViewTarget. */
@@ -1926,7 +1922,7 @@ public:
 	 * @param RequestedBy - the Actor requesting the target location
 	 * @return the optimal location to fire weapons at this actor
 	 */
-	virtual FVector GetTargetLocation(class AActor* RequestedBy = NULL) const;
+	virtual FVector GetTargetLocation(AActor* RequestedBy = NULL) const;
 
 	/** 
 	  * Hook to allow actors to render HUD overlays for themselves.  Called from AHUD::DrawActorOverlays(). 
diff --git a/Engine/Source/Runtime/Engine/Classes/GameFramework/HUD.h b/Engine/Source/Runtime/Engine/Classes/GameFramework/HUD.h
--- a/Engine/Source/Runtime/Engine/Classes/GameFramework/HUD.h
+++ b/Engine/Source/Runtime/Engine/Classes/GameFramework/HUD.h
@@ -5,6 +5,7 @@
 #include "HUD.generated.h"
 
 class FCanvasTextItem;
+class UTexture;
 
 /** List of actors and debug text to draw, @see AddDebugText(), RemoveDebugText(), and DrawDebugTextList() */
 USTRUCT()
diff --git a/Engine/Source/Runtime/Engine/Classes/GameFramework/PlayerController.h b/Engine/Source/Runtime/Engine/Classes/GameFramework/PlayerController.h
--- a/Engine/Source/Runtime/Engine/Classes/GameFramework/PlayerController.h
+++ b/Engine/Source/Runtime/Engine/Classes/GameFramework/PlayerController.h
@@ -4,7 +4,7 @@
 
 #include "GameFramework/PlayerMuteList.h"
 #include "Camera/PlayerCameraManager.h"
-#include "GameFramework/GameMode.h"
+#include "Camera/CameraTypes.h"
 #include "Components/InputComponent.h"
 #include "GameFramework/ForceFeedbackEffect.h"
 #include "GameFramework/OnlineReplStructs.h"
@@ -14,6 +14,9 @@
 
 class FPrimitiveComponentId;
 
+/** Default delegate that provides an implementation for those that don't have special needs other than a toggle */
+DECLARE_DELEGATE_RetVal(bool, FCanUnpause);
+
 /** delegate used to override default viewport audio listener position calculated from camera */
 DECLARE_DELEGATE_ThreeParams(FGetAudioListenerPos, FVector& /*Location*/, FVector& /*ProjFront*/, FVector& /*ProjRight*/);
 
@@ -599,7 +602,7 @@ public:
 	 * @param CustomPlaySpace - Matrix used when Space = CAPS_UserDefined
 	 */
 	UFUNCTION(unreliable, client, BlueprintCallable, Category = "Game|Feedback")
-	void ClientPlayCameraAnim(class UCameraAnim* AnimToPlay, float Scale=0, float Rate=0, float BlendInTime=0, float BlendOutTime=0, bool bLoop=false, bool bRandomStartTime=false, enum ECameraAnimPlaySpace::Type Space=ECameraAnimPlaySpace::CameraLocal, FRotator CustomPlaySpace=FRotator::ZeroRotator);
+	void ClientPlayCameraAnim(class UCameraAnim* AnimToPlay, float Scale=0, float Rate=0, float BlendInTime=0, float BlendOutTime=0, bool bLoop=false, bool bRandomStartTime=false, ECameraAnimPlaySpace::Type Space=ECameraAnimPlaySpace::CameraLocal, FRotator CustomPlaySpace=FRotator::ZeroRotator);
 
 	/** 
 	 * Play Camera Shake 
diff --git a/Engine/Source/Runtime/Engine/Classes/GameFramework/WorldSettings.h b/Engine/Source/Runtime/Engine/Classes/GameFramework/WorldSettings.h
--- a/Engine/Source/Runtime/Engine/Classes/GameFramework/WorldSettings.h
+++ b/Engine/Source/Runtime/Engine/Classes/GameFramework/WorldSettings.h
@@ -1,6 +1,7 @@
 // Copyright 1998-2014 Epic Games, Inc. All Rights Reserved.
 
 #pragma once
+#include "GameFramework/Info.h"
 #include "Sound/ReverbVolume.h"
 #include "WorldSettings.generated.h"
 
diff --git a/Engine/Source/Runtime/Engine/Private/GPUSort.h b/Engine/Source/Runtime/Engine/Private/GPUSort.h
--- a/Engine/Source/Runtime/Engine/Private/GPUSort.h
+++ b/Engine/Source/Runtime/Engine/Private/GPUSort.h
@@ -6,8 +6,6 @@
 
 #pragma once
 
-#include "RHI.h"
-
 /**
  * The input buffers required for sorting on the GPU.
  */
