Fix for Param Crossfading and exposing buggy engine audio component

Parameter Crossfading was always ignored for 2d/preview sounds because they don't have a location and were still controlled by the distance node it inherited from. Added virtual function to check whether crossfading is allowed that both classes can use.
Exposing engine audio component of the buggy pawn so that it can have parameters updated on it via blueprints.

diff --git a/Engine/Source/Runtime/Engine/Classes/Sound/SoundNodeDistanceCrossFade.h b/Engine/Source/Runtime/Engine/Classes/Sound/SoundNodeDistanceCrossFade.h
--- a/Engine/Source/Runtime/Engine/Classes/Sound/SoundNodeDistanceCrossFade.h
+++ b/Engine/Source/Runtime/Engine/Classes/Sound/SoundNodeDistanceCrossFade.h
@@ -102,4 +102,9 @@ public:
 	// End USoundNode interface. 
 
 	virtual float GetCurrentDistance(FAudioDevice* AudioDevice, FActiveSound& ActiveSound, const FSoundParseParameters& ParseParams) const;
+
+	/**
+	 * Determines whether Crossfading is currently allowed for the active sound
+	 */
+	virtual bool AllowCrossfading(FActiveSound& ActiveSound) const;
 };
diff --git a/Engine/Source/Runtime/Engine/Classes/Sound/SoundNodeParamCrossFade.h b/Engine/Source/Runtime/Engine/Classes/Sound/SoundNodeParamCrossFade.h
--- a/Engine/Source/Runtime/Engine/Classes/Sound/SoundNodeParamCrossFade.h
+++ b/Engine/Source/Runtime/Engine/Classes/Sound/SoundNodeParamCrossFade.h
@@ -20,4 +20,5 @@ public:
 
 	virtual FString GetUniqueString() const override;
 	virtual float GetCurrentDistance(FAudioDevice* AudioDevice, FActiveSound& ActiveSound, const FSoundParseParameters& ParseParams) const override;
+	virtual bool AllowCrossfading(FActiveSound& ActiveSound) const override;
 };
diff --git a/Engine/Source/Runtime/Engine/Private/SoundNodeDistanceCrossFade.cpp b/Engine/Source/Runtime/Engine/Private/SoundNodeDistanceCrossFade.cpp
--- a/Engine/Source/Runtime/Engine/Private/SoundNodeDistanceCrossFade.cpp
+++ b/Engine/Source/Runtime/Engine/Private/SoundNodeDistanceCrossFade.cpp
@@ -58,8 +58,8 @@ void USoundNodeDistanceCrossFade::ParseNodes( FAudioDevice* AudioDevice, const U
 			float VolumeToSet = 1.0f;
 			//UE_LOG(LogAudio, Log,  TEXT("  USoundNodeDistanceCrossFade.  Distance: %f ChildNodeIndex: %d CurrLoc: %s  ListenerLoc: %s"), Distance, ChildNodeIndex, *AudioComponent->CurrentLocation.ToString(), *AudioComponent->Listener->Location.ToString() );
 
-			// Ignore distance calculations for preview components as they are undefined
-			if( !ActiveSound.bLocationDefined )
+			// Check whether crossfading is allowed
+			if( !AllowCrossfading(ActiveSound) )
 			{
 				VolumeToSet = CrossFadeInput[ ChildNodeIndex ].Volume;
 			}
@@ -152,3 +152,9 @@ FString USoundNodeDistanceCrossFade::GetUniqueString() const
 {
 	return TEXT( "DistanceCrossFadeComplex/" );
 }
+
+bool USoundNodeDistanceCrossFade::AllowCrossfading(FActiveSound& ActiveSound) const
+{
+	// Ignore distance calculations for preview components as they are undefined
+	return ActiveSound.bLocationDefined;
+}
diff --git a/Engine/Source/Runtime/Engine/Private/SoundNodeParamCrossFade.cpp b/Engine/Source/Runtime/Engine/Private/SoundNodeParamCrossFade.cpp
--- a/Engine/Source/Runtime/Engine/Private/SoundNodeParamCrossFade.cpp
+++ b/Engine/Source/Runtime/Engine/Private/SoundNodeParamCrossFade.cpp
@@ -24,3 +24,9 @@ FString USoundNodeParamCrossFade::GetUniqueString() const
 {
 	return TEXT( "ParamCrossFadeComplex/" );
 }
+
+bool USoundNodeParamCrossFade::AllowCrossfading(FActiveSound& ActiveSound) const
+{
+	// Always allow parameter to control crossfading, even on 2D/preview sounds
+	return true;
+}
