bp: we seem to be missing the ability to make particle system parameter structs within blueprints TTP #341814

diff --git a/Engine/Source/Runtime/Engine/Classes/Particles/ParticleSystemComponent.h b/Engine/Source/Runtime/Engine/Classes/Particles/ParticleSystemComponent.h
--- a/Engine/Source/Runtime/Engine/Classes/Particles/ParticleSystemComponent.h
+++ b/Engine/Source/Runtime/Engine/Classes/Particles/ParticleSystemComponent.h
@@ -19,7 +19,7 @@ struct FDynamicEmitterReplayDataBase;
 struct FParticleAnimTrailEmitterInstance;
 
 /** Enum for specifying type of a name instance parameter. */
-UENUM()
+UENUM(BlueprintType)
 enum EParticleSysParamType
 {
 	PSPT_None UMETA(DisplayName="None"),
@@ -74,7 +74,7 @@ DECLARE_DYNAMIC_DELEGATE_OneParam( FOnSystemFinished, class UParticleSystemCompo
 
 
 /** Struct used for a particular named instance parameter for this ParticleSystemComponent. */
-USTRUCT()
+USTRUCT(BlueprintType)
 struct FParticleSysParam
 {
 	GENERATED_USTRUCT_BODY()
