Structs, that are not POD, are copied (by default) using native copy constructor.

BP:  Passing a struct by val leads to unexpected results ttp341538

diff --git a/Engine/Source/Runtime/CoreUObject/Public/UObject/Class.h b/Engine/Source/Runtime/CoreUObject/Public/UObject/Class.h
--- a/Engine/Source/Runtime/CoreUObject/Public/UObject/Class.h
+++ b/Engine/Source/Runtime/CoreUObject/Public/UObject/Class.h
@@ -471,6 +471,10 @@ struct TStructOpsTypeTraitsBase
 template<class CPPSTRUCT>
 struct TStructOpsTypeTraits : public TStructOpsTypeTraitsBase
 {
+	enum
+	{
+		WithCopy = !TIsPODType<CPPSTRUCT>::Value
+	};
 };
 
 
diff --git a/Engine/Source/Runtime/Engine/Classes/Components/StaticMeshComponent.h b/Engine/Source/Runtime/Engine/Classes/Components/StaticMeshComponent.h
--- a/Engine/Source/Runtime/Engine/Classes/Components/StaticMeshComponent.h
+++ b/Engine/Source/Runtime/Engine/Classes/Components/StaticMeshComponent.h
@@ -98,6 +98,15 @@ private:
 	FStaticMeshComponentLODInfo &operator=( const FStaticMeshComponentLODInfo &rhs ) { check(0); return *this; }
 };
 
+template<>
+struct TStructOpsTypeTraits<FStaticMeshComponentLODInfo> : public TStructOpsTypeTraitsBase
+{
+	enum
+	{
+		WithCopy = false
+	};
+};
+
 /** A StaticMeshComponent is a mesh that does not animate. */
 UCLASS(ClassGroup=(Rendering, Common), hidecategories=(Object,Activation,"Components|Activation"), ShowCategories=(Mobility), editinlinenew, meta=(BlueprintSpawnableComponent))
 class ENGINE_API UStaticMeshComponent : public UMeshComponent
