Header cleanup:  BlueprintGeneratedClass.h, PlayerInput.h, Blueprint.h.

diff --git a/Engine/Source/Runtime/Engine/Classes/Engine/Blueprint.h b/Engine/Source/Runtime/Engine/Classes/Engine/Blueprint.h
--- a/Engine/Source/Runtime/Engine/Classes/Engine/Blueprint.h
+++ b/Engine/Source/Runtime/Engine/Classes/Engine/Blueprint.h
@@ -1,12 +1,9 @@
 // Copyright 1998-2014 Epic Games, Inc. All Rights Reserved.
 
 #pragma once
-
-#include "EdGraph/EdGraph.h"
 #include "EdGraph/EdGraphPin.h"
 #include "EdGraph/EdGraphNode.h"
 #include "BlueprintCore.h"
-
 #include "Blueprint.generated.h"
 
 /** States a blueprint can be in */
diff --git a/Engine/Source/Runtime/Engine/Classes/Engine/BlueprintGeneratedClass.h b/Engine/Source/Runtime/Engine/Classes/Engine/BlueprintGeneratedClass.h
--- a/Engine/Source/Runtime/Engine/Classes/Engine/BlueprintGeneratedClass.h
+++ b/Engine/Source/Runtime/Engine/Classes/Engine/BlueprintGeneratedClass.h
@@ -1,9 +1,6 @@
 // Copyright 1998-2014 Epic Games, Inc. All Rights Reserved.
 
 #pragma once
-
-#include "Components/InputComponent.h"
-#include "EdGraph/EdGraphNode.h"
 #include "BlueprintGeneratedClass.generated.h"
 
 class UEdGraphPin;
diff --git a/Engine/Source/Runtime/Engine/Classes/GameFramework/PlayerInput.h b/Engine/Source/Runtime/Engine/Classes/GameFramework/PlayerInput.h
--- a/Engine/Source/Runtime/Engine/Classes/GameFramework/PlayerInput.h
+++ b/Engine/Source/Runtime/Engine/Classes/GameFramework/PlayerInput.h
@@ -7,12 +7,9 @@
 //=============================================================================
 
 #pragma once
-
-#include "Engine/EngineTypes.h"
 #include "Components/InputComponent.h"
+#include "Engine/EngineTypes.h"
 #include "GestureRecognizer.h"
-#include "KeyState.h"
-
 #include "PlayerInput.generated.h"
 
 ENGINE_API DECLARE_LOG_CATEGORY_EXTERN(LogInput, Log, All);
@@ -21,6 +18,7 @@ ENGINE_API DECLARE_LOG_CATEGORY_EXTERN(LogInput, Log, All);
 // Forward declarations.
 //
 struct FDelegateDispatchDetails;
+struct FKeyState;
 
 USTRUCT()
 struct FKeyBind
