UE4 - Remove bogus check

diff --git a/Engine/Source/Runtime/Engine/Private/InstancedStaticMesh.cpp b/Engine/Source/Runtime/Engine/Private/InstancedStaticMesh.cpp
--- a/Engine/Source/Runtime/Engine/Private/InstancedStaticMesh.cpp
+++ b/Engine/Source/Runtime/Engine/Private/InstancedStaticMesh.cpp
@@ -1050,7 +1050,7 @@ inline void FInstancedStaticMeshSceneProxy::SetupInstancedMeshBatch(int32 LODInd
 	auto* OutBatchElement0 = OutMeshBatches[0].Elements.GetTypedData();
 	OutBatchElement0->UserData = (void*)&UserData_AllInstances;
 	OutBatchElement0->UserIndex = 0;
-	if (bInstanced && NumInstances > 1)
+	if (bInstanced)
 	{
 		OutBatchElement0->NumInstances = NumInstances;
 	}
