Fix for component registration error which shows up with tappy chicken; due to finding a parents for a component we end up processing the parent twice thus double adding it to the 'added' count and causing the function to fail because we 'add', say, 119 when we were expecting 118

diff --git a/Engine/Source/Runtime/Engine/Private/Actor.cpp b/Engine/Source/Runtime/Engine/Private/Actor.cpp
--- a/Engine/Source/Runtime/Engine/Private/Actor.cpp
+++ b/Engine/Source/Runtime/Engine/Private/Actor.cpp
@@ -3095,7 +3095,7 @@ bool AActor::IncrementalRegisterComponents(int32 NumComponentsToRegister)
 	int32 NumRegisteredComponentsThisRun = 0;
 	TArray<UActorComponent*> Components;
 	GetComponents(Components);
-	
+
 	for (int32 CompIdx = 0; CompIdx < Components.Num() && NumRegisteredComponentsThisRun < NumComponentsToRegister; CompIdx++)
 	{
 		UActorComponent* Component = Components[CompIdx];
@@ -3108,6 +3108,7 @@ bool AActor::IncrementalRegisterComponents(int32 NumComponentsToRegister)
 				// Register parent first, then return to this component on a next iteration
 				Component = ParentComponent;
 				CompIdx--;
+				NumTotalRegisteredComponents--; // because we will try to register the parent again later...
 			}
 				
 			//Before we register our component, save it to our transaction buffer so if "undone" it will return to an unregistered state.
