Fix for TTP 341939 - Remove warning about es2 and precomputed visibility. We have no workable way of determining projects that should be using precomputed visibility, so this warning is pointless in those cases (i.e. 2d games).

diff --git a/Engine/Source/Runtime/Engine/Private/UnrealEngine.cpp b/Engine/Source/Runtime/Engine/Private/UnrealEngine.cpp
--- a/Engine/Source/Runtime/Engine/Private/UnrealEngine.cpp
+++ b/Engine/Source/Runtime/Engine/Private/UnrealEngine.cpp
@@ -6388,15 +6388,7 @@ void DrawStatsHUD( UWorld* World, FViewport* Viewport, FCanvas* Canvas, UCanvas*
 				Canvas->DrawItem( SmallTextItem, FVector2D( MessageX, MessageY ) );
 				MessageY += 20;
 			}
-
-			if (World->FeatureLevel == ERHIFeatureLevel::ES2 && !World->GetWorldSettings()->bPrecomputeVisibility)
-			{
-				SmallTextItem.SetColor(FLinearColor::Red);
-				SmallTextItem.Text = LOCTEXT("ES2ButNoPrecomputedVisibility", "ES2 RENDERING IS ENABLED BUT PRECOMPUTED VISIBILITY IS DISABLED IN THE WORLD SETTINGS");
-				Canvas->DrawItem(SmallTextItem, FVector2D(MessageX, MessageY));
-				MessageY += 20;
-			}
-
+			
 			// check navmesh
 #if WITH_EDITOR
 			const bool bIsNavigationAutoUpdateEnabled = UNavigationSystem::GetIsNavigationAutoUpdateEnabled();
