Header cleanup: MovementComponent.h, PImplRecastNavMesh.h, NavRelevantComponent.h, EdGraph.h, NavigationGraphNode.h.

diff --git a/Engine/Source/Runtime/AIModule/Private/Navigation/CrowdManager.cpp b/Engine/Source/Runtime/AIModule/Private/Navigation/CrowdManager.cpp
--- a/Engine/Source/Runtime/AIModule/Private/Navigation/CrowdManager.cpp
+++ b/Engine/Source/Runtime/AIModule/Private/Navigation/CrowdManager.cpp
@@ -6,6 +6,7 @@
 #include "AI/Navigation/PImplRecastNavMesh.h"
 #include "AI/Navigation/RecastHelpers.h"
 #include "DetourCrowd/DetourCrowd.h"
+#include "Detour/DetourCommon.h"
 #endif
 
 #include "Navigation/CrowdManager.h"
diff --git a/Engine/Source/Runtime/Engine/Classes/AI/Navigation/NavigationGraphNode.h b/Engine/Source/Runtime/Engine/Classes/AI/Navigation/NavigationGraphNode.h
--- a/Engine/Source/Runtime/Engine/Classes/AI/Navigation/NavigationGraphNode.h
+++ b/Engine/Source/Runtime/Engine/Classes/AI/Navigation/NavigationGraphNode.h
@@ -2,7 +2,7 @@
 
 #pragma once
 
-#include "NavigationSystem.h"
+#include "GameFramework/Actor.h"
 
 #include "NavigationGraphNode.generated.h"
 
diff --git a/Engine/Source/Runtime/Engine/Classes/Components/BoxComponent.h b/Engine/Source/Runtime/Engine/Classes/Components/BoxComponent.h
--- a/Engine/Source/Runtime/Engine/Classes/Components/BoxComponent.h
+++ b/Engine/Source/Runtime/Engine/Classes/Components/BoxComponent.h
@@ -2,6 +2,7 @@
 
 
 #pragma once
+#include "Components/ShapeComponent.h"
 #include "BoxComponent.generated.h"
 
 /** 
diff --git a/Engine/Source/Runtime/Engine/Classes/EdGraph/EdGraph.h b/Engine/Source/Runtime/Engine/Classes/EdGraph/EdGraph.h
--- a/Engine/Source/Runtime/Engine/Classes/EdGraph/EdGraph.h
+++ b/Engine/Source/Runtime/Engine/Classes/EdGraph/EdGraph.h
@@ -1,13 +1,11 @@
 // Copyright 1998-2014 Epic Games, Inc. All Rights Reserved.
 
 #pragma once
-
 #include "BlueprintUtilities.h"
-#include "EdGraphSchema.h"
-#include "EdGraphNode.h"
-
 #include "EdGraph.generated.h"
 
+struct FEdGraphEditAction;
+
 USTRUCT()
 struct ENGINE_API FGraphReference
 {
diff --git a/Engine/Source/Runtime/Engine/Classes/GameFramework/MovementComponent.h b/Engine/Source/Runtime/Engine/Classes/GameFramework/MovementComponent.h
--- a/Engine/Source/Runtime/Engine/Classes/GameFramework/MovementComponent.h
+++ b/Engine/Source/Runtime/Engine/Classes/GameFramework/MovementComponent.h
@@ -6,7 +6,6 @@
 
 #pragma once
 #include "Components/ActorComponent.h"
-#include "Components/PrimitiveComponent.h"
 #include "MovementComponent.generated.h"
 
 struct FCollisionQueryParams;
diff --git a/Engine/Source/Runtime/Engine/Private/AI/Navigation/PImplRecastNavMesh.h b/Engine/Source/Runtime/Engine/Private/AI/Navigation/PImplRecastNavMesh.h
--- a/Engine/Source/Runtime/Engine/Private/AI/Navigation/PImplRecastNavMesh.h
+++ b/Engine/Source/Runtime/Engine/Private/AI/Navigation/PImplRecastNavMesh.h
@@ -10,10 +10,8 @@
 
 #if WITH_RECAST
 
-#include "DetourCommon.h"
-#include "DetourNavMesh.h"
-#include "DetourNavMeshQuery.h"
-#include "AI/Navigation/RecastNavMesh.h"
+#include "AI/Navigation/NavFilters/NavigationQueryFilter.h"
+#include "Detour/DetourNavMeshQuery.h"
 
 #define RECAST_VERY_SMALL_AGENT_RADIUS 0.0f
 
