Added an include for Actor.h to NavigationData.h so that it works with "include what you use" (without Engine.h)

diff --git a/Engine/Source/Runtime/Engine/Classes/AI/Navigation/NavigationData.h b/Engine/Source/Runtime/Engine/Classes/AI/Navigation/NavigationData.h
--- a/Engine/Source/Runtime/Engine/Classes/AI/Navigation/NavigationData.h
+++ b/Engine/Source/Runtime/Engine/Classes/AI/Navigation/NavigationData.h
@@ -1,6 +1,7 @@
 // Copyright 1998-2014 Epic Games, Inc. All Rights Reserved.
 
 #pragma once
+#include "GameFramework/Actor.h"
 #include "AI/Navigation/NavigationTypes.h"
 #include "AI/NavDataGenerator.h"
 #include "AI/Navigation/NavFilters/NavigationQueryFilter.h"
