Compilation Fix (no editor)

diff --git a/Engine/Source/Runtime/Engine/Private/Blueprint.cpp b/Engine/Source/Runtime/Engine/Private/Blueprint.cpp
--- a/Engine/Source/Runtime/Engine/Private/Blueprint.cpp
+++ b/Engine/Source/Runtime/Engine/Private/Blueprint.cpp
@@ -164,13 +164,14 @@ void UBlueprintCore::GenerateDeterministicGuid()
 
 UBlueprint::UBlueprint(const class FPostConstructInitializeProperties& PCIP)
 	: Super(PCIP)
+#if WITH_EDITORONLY_DATA
 	, bCachedDependenciesUpTpDate(false)
-{
-
+#endif
 #if WITH_EDITOR
-	bRunConstructionScriptOnDrag = true;
-	bGenerateConstClass = false;
+	, bRunConstructionScriptOnDrag(true)
+	, bGenerateConstClass(false)
 #endif
+{
 }
 
 void UBlueprint::Serialize(FArchive& Ar)
