* nicer name

diff --git a/Engine/Source/Runtime/Engine/Classes/Engine/Scene.h b/Engine/Source/Runtime/Engine/Classes/Engine/Scene.h
--- a/Engine/Source/Runtime/Engine/Classes/Engine/Scene.h
+++ b/Engine/Source/Runtime/Engine/Classes/Engine/Scene.h
@@ -558,7 +558,7 @@ struct FPostProcessSettings
 	 * Logarithmic adjustment for the exposure. Only used if a tonemapper is specified.
 	 * 0: no adjustment, -1:2x darker, -2:4x darker, 1:2x brighter, 2:4x brighter, ...
 	 */
-	UPROPERTY(interp, Category=AutoExposure, meta=(UIMin = "-8.0", UIMax = "8.0", editcondition = "bOverride_AutoExposureBias"))
+	UPROPERTY(interp, Category = AutoExposure, meta = (UIMin = "-8.0", UIMax = "8.0", editcondition = "bOverride_AutoExposureBias", DisplayName = "Exposure Bias"))
 	float AutoExposureBias;
 
 	/** temporary exposed until we found good values, -8: 1/256, -10: 1/1024 */
