Engine: Fix typo in build

diff --git a/Engine/Source/Runtime/Engine/Private/NiagaraComponent.cpp b/Engine/Source/Runtime/Engine/Private/NiagaraComponent.cpp
--- a/Engine/Source/Runtime/Engine/Private/NiagaraComponent.cpp
+++ b/Engine/Source/Runtime/Engine/Private/NiagaraComponent.cpp
@@ -282,7 +282,7 @@ namespace ENiagaraVectorAttr
 
 /**
  * A niagara particle simulation.
- */n
+ */
 class FNiagaraSimulation
 {
 public:
