Fix a false negative that was happening when comparing FUniqueNetId's. This fixes guests not working on XB1.

diff --git a/Engine/Source/Runtime/Engine/Private/GameInstance.cpp b/Engine/Source/Runtime/Engine/Private/GameInstance.cpp
--- a/Engine/Source/Runtime/Engine/Private/GameInstance.cpp
+++ b/Engine/Source/Runtime/Engine/Private/GameInstance.cpp
@@ -517,9 +517,26 @@ ULocalPlayer* UGameInstance::FindLocalPlayerFromControllerId(int32 ControllerId)
 
 ULocalPlayer* UGameInstance::FindLocalPlayerFromUniqueNetId(TSharedPtr<FUniqueNetId> UniqueNetId) const
 {
+	if (!UniqueNetId.IsValid())
+	{
+		return NULL;
+	}
+
 	for (ULocalPlayer* Player : LocalPlayers)
 	{
-		if (Player && Player->GetUniqueNetId() == UniqueNetId)
+		if (Player == NULL)
+		{
+			continue;
+		}
+
+		TSharedPtr<FUniqueNetId> OtherUniqueNetId = Player->GetUniqueNetId();
+
+		if (!OtherUniqueNetId.IsValid())
+		{
+			continue;
+		}
+
+		if (*OtherUniqueNetId == *UniqueNetId)
 		{
 			// Match
 			return Player;
