Fixing build error in Win32 (unaligned FVector4); fixing style violation.

diff --git a/Engine/Source/Runtime/Engine/Private/NiagaraComponent.cpp b/Engine/Source/Runtime/Engine/Private/NiagaraComponent.cpp
--- a/Engine/Source/Runtime/Engine/Private/NiagaraComponent.cpp
+++ b/Engine/Source/Runtime/Engine/Private/NiagaraComponent.cpp
@@ -374,12 +374,7 @@ public:
 
 	FBox GetBounds() const { return CachedBounds; }
 
-	void AddConstant(FVector4 Value)
-	{
-		ConstantTable.Add(Value);
-	}
-
-	void SetConstant(int Idx, FVector4 Value)
+	void SetConstant(int Idx, const FVector4 &Value)
 	{
 		check(Idx < ConstantTable.Num());
 		ConstantTable[Idx] = Value;
@@ -421,7 +416,9 @@ private:
 		for (int i = NiagaraConstants::NumBuiltinConstants + 1; i < Script->ConstantTable.Num(); i++)
 		{
 			if (ConstantTable.Num() <= i)
+			{
 				ConstantTable.Add(FVector4());
+			}
 			ConstantTable[i] = Script->ConstantTable[i];
 		}
 	}
