Framework - ProjectWorldLocationToScreen Now returns true/false if the projection was successful or failed.

diff --git a/Engine/Source/Runtime/Engine/Classes/GameFramework/PlayerController.h b/Engine/Source/Runtime/Engine/Classes/GameFramework/PlayerController.h
--- a/Engine/Source/Runtime/Engine/Classes/GameFramework/PlayerController.h
+++ b/Engine/Source/Runtime/Engine/Classes/GameFramework/PlayerController.h
@@ -187,7 +187,7 @@ class ENGINE_API APlayerController : public AController
 	UPROPERTY(DuplicateTransient)
 	uint8 NetPlayerIndex;
 
-    /** List of muted players in various categories */
+	/** List of muted players in various categories */
 	FPlayerMuteList MuteList;
 
 	/** this is set on the OLD PlayerController when performing a swap over a network connection
@@ -198,7 +198,7 @@ class ENGINE_API APlayerController : public AController
 	UPROPERTY(DuplicateTransient)
 	class UNetConnection* PendingSwapConnection;
 
-    /** The net connection this controller is communicating on, NULL for local players on server */
+	/** The net connection this controller is communicating on, NULL for local players on server */
 	UPROPERTY(DuplicateTransient)
 	class UNetConnection* NetConnection;
 
@@ -384,9 +384,12 @@ public:
 	UFUNCTION(BlueprintCallable, Category = "Game|Player", meta = (FriendlyName = "ConvertScreenLocationToWorldSpace"))
 	void DeprojectScreenPositionToWorld(float ScreenX, float ScreenY, FVector & WorldLocation, FVector & WorldDirection) const;
 
-	/** Convert a World Space 3D position into a 2D Screen Space position. **/
-	UFUNCTION(BlueprintCallable, Category = "Game|Player", meta = (FriendlyName = "ConvertWorldLocationToScreenLocation"))
-	void ProjectWorldLocationToScreen(FVector WorldLocation, FVector2D& ScreenLocation) const;
+	/**
+	 * Convert a World Space 3D position into a 2D Screen Space position.
+	 * @return true if the world coordinate was successfully projected to the screen.
+	 */
+	UFUNCTION(BlueprintCallable, Category = "Game|Player", meta = ( FriendlyName = "ConvertWorldLocationToScreenLocation" ))
+	bool ProjectWorldLocationToScreen(FVector WorldLocation, FVector2D& ScreenLocation) const;
 	
 	/**
 	  * Updates the rotation of player, based on ControlRotation after RotationInput has been applied.
@@ -439,22 +442,22 @@ public:
 	 */
 	virtual void SeamlessTravelFrom(class APlayerController* OldPC);
 
-   	/** 
+	/** 
 	 * Tell the client to enable or disable voice chat (not muting)
-     * @param bEnable enable or disable voice chat
+	 * @param bEnable enable or disable voice chat
 	 */
 	UFUNCTION(client, reliable)
 	virtual void ClientEnableNetworkVoice(bool bEnable);
 
-   	/** Enable voice chat transmission	 */
+	/** Enable voice chat transmission	 */
 	void StartTalking();
 
-   	/** Disable voice chat transmission	 */
+	/** Disable voice chat transmission	 */
 	void StopTalking();
 
-   	/** 
+	/** 
 	 * Toggle voice chat on and off
-     * @param bSpeaking enable or disable voice chat
+	 * @param bSpeaking enable or disable voice chat
 	 */
 	UFUNCTION(exec)
 	virtual void ToggleSpeaking(bool bSpeaking);
diff --git a/Engine/Source/Runtime/Engine/Private/PlayerController.cpp b/Engine/Source/Runtime/Engine/Private/PlayerController.cpp
--- a/Engine/Source/Runtime/Engine/Private/PlayerController.cpp
+++ b/Engine/Source/Runtime/Engine/Private/PlayerController.cpp
@@ -80,7 +80,7 @@ UPlayer* APlayerController::GetNetOwningPlayer()
 
 bool APlayerController::HasNetOwner() const
 {
-    // Player controllers are their own net owners
+	// Player controllers are their own net owners
 	return true;
 }
 
@@ -1448,11 +1448,11 @@ void APlayerController::ClientSetCameraFade_Implementation(bool bEnableFading, F
 			PlayerCameraManager->FadeTimeRemaining = FadeTime;
 			PlayerCameraManager->bFadeAudio = bFadeAudio;
 		}
- 		else
- 		{
- 			// Make sure FadeAmount finishes at the correct value
- 			PlayerCameraManager->FadeAmount = PlayerCameraManager->FadeAlpha.Y;
- 		}
+		else
+		{
+			// Make sure FadeAmount finishes at the correct value
+			PlayerCameraManager->FadeAmount = PlayerCameraManager->FadeAlpha.Y;
+		}
 	}
 }
 
@@ -1841,7 +1841,7 @@ void APlayerController::DeprojectScreenPositionToWorld(float ScreenX, float Scre
 }
 
 
-void APlayerController::ProjectWorldLocationToScreen(FVector WorldLocation, FVector2D& ScreenLocation) const
+bool APlayerController::ProjectWorldLocationToScreen(FVector WorldLocation, FVector2D& ScreenLocation) const
 {
 	ULocalPlayer* LocalPlayer = Cast<ULocalPlayer>(Player);
 	if (LocalPlayer != NULL && LocalPlayer->ViewportClient != NULL && LocalPlayer->ViewportClient->Viewport != NULL)
@@ -1860,9 +1860,11 @@ void APlayerController::ProjectWorldLocationToScreen(FVector WorldLocation, FVec
 
 		if (SceneView) 
 		{
-			SceneView->WorldToPixel(WorldLocation, ScreenLocation);
+			return SceneView->WorldToPixel(WorldLocation, ScreenLocation);
 		}
-	} 
+	}
+
+	return false;
 }
 
 
