Clean up DrawDebug* blueprint nodes
* All draw debug now allow Duration to be specified (Cylinder, Cone, Frustum were missing it)
* Default Duration is now 0, not -1 (which has same semantics, but less odd looking)
* All functions default color to White (instead of inconsistent Black and White)
* Don't pass parameters as reference so they can be edited on the pin and not require a connection
* DrawDebugCone has default parameters that will show something (along the same lines as Sphere and Cylinder)
* DrawDebugCone takes angle parameters in degrees, not radians (old radians version of function deprecated with useful message)
* DrawDebugPlane has better parameter names

diff --git a/Engine/Config/BaseEngine.ini b/Engine/Config/BaseEngine.ini
--- a/Engine/Config/BaseEngine.ini
+++ b/Engine/Config/BaseEngine.ini
@@ -1011,6 +1011,9 @@ bUseStreamingPause=false
 
 +K2ParamRedirects=(NodeName="/Script/BlueprintGraph.K2Node_SpawnActor", OldParamName="bNoCollisionFail", NewParamName="SpawnEvenIfColliding")
 
++K2ParamRedirects=(NodeName="/Script/BlueprintGraph.K2Node_CallFunction", OldParamName="KismetSystemLibrary.DrawDebugPlane.P", NewParamName="PlaneCoordinates")
++K2ParamRedirects=(NodeName="/Script/BlueprintGraph.K2Node_CallFunction", OldParamName="KismetSystemLibrary.DrawDebugPlane.Loc", NewParamName="Location")
+
 ; Temporary animation redirects from macro-expansion style to runtime node style
 +K2ParamRedirects=(NodeName="/Script/AnimGraph.AnimGraphNode_Root", OldParamName="Pose", NewParamName="Result")
 +K2ParamRedirects=(NodeName="/Script/AnimGraph.AnimGraphNode_StateResult", OldParamName="Pose", NewParamName="Result")
diff --git a/Engine/Source/Runtime/Engine/Classes/Kismet/KismetSystemLibrary.h b/Engine/Source/Runtime/Engine/Classes/Kismet/KismetSystemLibrary.h
--- a/Engine/Source/Runtime/Engine/Classes/Kismet/KismetSystemLibrary.h
+++ b/Engine/Source/Runtime/Engine/Classes/Kismet/KismetSystemLibrary.h
@@ -991,43 +991,50 @@ class UKismetSystemLibrary : public UBlueprintFunctionLibrary
 
 	/** Draw a debug line */
 	UFUNCTION(BlueprintCallable, Category="Rendering|Debug", meta=(HidePin="WorldContextObject", DefaultToSelf="WorldContextObject"))
-	static ENGINE_API void DrawDebugLine(UObject* WorldContextObject, const FVector& LineStart, const FVector& LineEnd, FLinearColor LineColor, float Duration=-1.f, float Thickness = 0.f);
+	static ENGINE_API void DrawDebugLine(UObject* WorldContextObject, const FVector LineStart, const FVector LineEnd, FLinearColor LineColor, float Duration=0.f, float Thickness = 0.f);
 
 	/** Draw a debug point */
 	UFUNCTION(BlueprintCallable, Category="Rendering|Debug", meta=(HidePin="WorldContextObject", DefaultToSelf="WorldContextObject"))
-	static ENGINE_API void DrawDebugPoint(UObject* WorldContextObject, const FVector& Position, float Size, FLinearColor PointColor, float Duration=0.f);
+	static ENGINE_API void DrawDebugPoint(UObject* WorldContextObject, const FVector Position, float Size, FLinearColor PointColor, float Duration=0.f);
 
 	/** Draw directional arrow, pointing from LineStart to LineEnd. */
 	UFUNCTION(BlueprintCallable, Category="Rendering|Debug", meta=(HidePin="WorldContextObject", DefaultToSelf="WorldContextObject"))
-	static ENGINE_API void DrawDebugArrow(UObject* WorldContextObject, const FVector& LineStart, const FVector& LineEnd, float ArrowSize, FLinearColor LineColor, float Duration=-1.f);
+	static ENGINE_API void DrawDebugArrow(UObject* WorldContextObject, const FVector LineStart, const FVector LineEnd, float ArrowSize, FLinearColor LineColor, float Duration=0.f);
 
 	/** Draw a debug box */
 	UFUNCTION(BlueprintCallable, Category="Rendering|Debug", meta=(HidePin="WorldContextObject", DefaultToSelf="WorldContextObject"))
-	static ENGINE_API void DrawDebugBox(UObject* WorldContextObject, const FVector& Center, FVector Extent, FLinearColor LineColor, const FRotator Rotation=FRotator::ZeroRotator, float Duration=-1.f);
+	static ENGINE_API void DrawDebugBox(UObject* WorldContextObject, const FVector Center, FVector Extent, FLinearColor LineColor, const FRotator Rotation=FRotator::ZeroRotator, float Duration=0.f);
 
 	/** Draw a debug coordinate system. */
 	UFUNCTION(BlueprintCallable, Category="Rendering|Debug", meta=(HidePin="WorldContextObject", DefaultToSelf="WorldContextObject"))
-	static ENGINE_API void DrawDebugCoordinateSystem(UObject* WorldContextObject, const FVector& AxisLoc, const FRotator& AxisRot, float Scale);
+	static ENGINE_API void DrawDebugCoordinateSystem(UObject* WorldContextObject, const FVector AxisLoc, const FRotator AxisRot, float Scale=1.f, float Duration=0.f);
 
 	/** Draw a debug sphere */
-	UFUNCTION(BlueprintCallable, Category="Rendering|Debug", meta=(Segments="12",Radius="100",HidePin="WorldContextObject", DefaultToSelf="WorldContextObject"))
-	static ENGINE_API void DrawDebugSphere(UObject* WorldContextObject, const FVector& Center, float Radius, int32 Segments, FLinearColor LineColor, float Duration=-1.f);
+	UFUNCTION(BlueprintCallable, Category="Rendering|Debug", meta=(HidePin="WorldContextObject", DefaultToSelf="WorldContextObject"))
+	static ENGINE_API void DrawDebugSphere(UObject* WorldContextObject, const FVector Center, float Radius=100.f, int32 Segments=12, FLinearColor LineColor = FLinearColor::White, float Duration=0.f);
 
 	/** Draw a debug cylinder */
-	UFUNCTION(BlueprintCallable, Category="Rendering|Debug", meta=(Segments="12",Radius="100",HidePin="WorldContextObject", DefaultToSelf="WorldContextObject"))
-	static ENGINE_API void DrawDebugCylinder(UObject* WorldContextObject, const FVector& Start, const FVector& End, float Radius, int32 Segments, FLinearColor LineColor);
+	UFUNCTION(BlueprintCallable, Category="Rendering|Debug", meta=(HidePin="WorldContextObject", DefaultToSelf="WorldContextObject"))
+	static ENGINE_API void DrawDebugCylinder(UObject* WorldContextObject, const FVector Start, const FVector End, float Radius=100.f, int32 Segments=12, FLinearColor LineColor = FLinearColor::White, float Duration=0.f);
 	
 	/** Draw a debug cone */
-	UFUNCTION(BlueprintCallable, Category="Rendering|Debug", meta=(HidePin="WorldContextObject", DefaultToSelf="WorldContextObject"))
-	static ENGINE_API void DrawDebugCone(UObject* WorldContextObject, const FVector& Origin, const FVector& Direction, float Length, float AngleWidth, float AngleHeight, int32 NumSides, FLinearColor LineColor);
+	UFUNCTION(BlueprintCallable, Category="Rendering|Debug", meta=(HidePin="WorldContextObject", DefaultToSelf="WorldContextObject", DeprecatedFunction, DeprecationMessage="DrawDebugCone has been changed to use degrees for angles instead of radians. Place a new DrawDebugCone node and pass your angles as degrees."))
+	static ENGINE_API void DrawDebugCone(UObject* WorldContextObject, const FVector Origin, const FVector Direction, float Length, float AngleWidth, float AngleHeight, int32 NumSides, FLinearColor LineColor);
+
+	/** 
+	 * Draw a debug cone 
+	 * Angles are specified in degrees
+	 */
+	UFUNCTION(BlueprintCallable, Category="Rendering|Debug", meta=(HidePin="WorldContextObject", DefaultToSelf="WorldContextObject", FriendlyName="DrawDebugCone"))
+	static ENGINE_API void DrawDebugConeInDegrees(UObject* WorldContextObject, const FVector Origin, const FVector Direction, float Length=100.f, float AngleWidth=45.f, float AngleHeight=45.f, int32 NumSides = 12, FLinearColor LineColor = FLinearColor::White, float Duration=0.f);
 
 	/** Draw a debug capsule */
 	UFUNCTION(BlueprintCallable, Category="Rendering|Debug", meta=(HidePin="WorldContextObject", DefaultToSelf="WorldContextObject"))
-	static ENGINE_API void DrawDebugCapsule(UObject* WorldContextObject, const FVector& Center, float HalfHeight, float Radius, const FRotator& Rotation, FLinearColor LineColor, float Duration=-1.f);
+	static ENGINE_API void DrawDebugCapsule(UObject* WorldContextObject, const FVector Center, float HalfHeight, float Radius, const FRotator Rotation, FLinearColor LineColor = FLinearColor::White, float Duration=0.f);
 
 	/** Draw a debug string at a 3d world location. */
 	UFUNCTION(BlueprintCallable, Category="Rendering|Debug", meta=(HidePin="WorldContextObject", DefaultToSelf="WorldContextObject"))
-	static ENGINE_API void DrawDebugString(UObject* WorldContextObject, const FVector& TextLocation, const FString& Text, class AActor* TestBaseActor = NULL, FLinearColor TextColor = FLinearColor::White, float Duration=-1.f);
+	static ENGINE_API void DrawDebugString(UObject* WorldContextObject, const FVector TextLocation, const FString& Text, class AActor* TestBaseActor = NULL, FLinearColor TextColor = FLinearColor::White, float Duration=0.f);
 	/** 
 	 * Removes all debug strings. 
 	 *
@@ -1038,7 +1045,7 @@ class UKismetSystemLibrary : public UBlueprintFunctionLibrary
 
 	/** Draws a debug plane. */
 	UFUNCTION(BlueprintCallable, Category="Rendering|Debug", meta=(HidePin="WorldContextObject", DefaultToSelf="WorldContextObject"))
-	static ENGINE_API void DrawDebugPlane(UObject* WorldContextObject, const FPlane& P, const FVector& Loc, float Size, FLinearColor PlaneColor, float Duration=-1.f);
+	static ENGINE_API void DrawDebugPlane(UObject* WorldContextObject, const FPlane PlaneCoordinates, const FVector Location, float Size, FLinearColor PlaneColor = FLinearColor::White, float Duration=0.f);
 
 	/** 
 	 * Flush all persistent debug lines and shapes.
@@ -1050,19 +1057,19 @@ class UKismetSystemLibrary : public UBlueprintFunctionLibrary
 
 	/** Draws a debug frustum. */
 	UFUNCTION(BlueprintCallable, Category="Rendering|Debug", meta=(HidePin="WorldContextObject", DefaultToSelf="WorldContextObject"))
-	static ENGINE_API void DrawDebugFrustum(UObject* WorldContextObject, const FTransform& FrustumTransform, FLinearColor FrustumColor);
+	static ENGINE_API void DrawDebugFrustum(UObject* WorldContextObject, const FTransform& FrustumTransform, FLinearColor FrustumColor = FLinearColor::White, float Duration=0.f);
 
 	/** Draw a debug camera shape. */
 	UFUNCTION(BlueprintCallable, Category="Rendering|Debug")
-	static ENGINE_API void DrawDebugCamera(const ACameraActor* CameraActor, FLinearColor CameraColor = FLinearColor::White, float Duration=-1.f);
+	static ENGINE_API void DrawDebugCamera(const ACameraActor* CameraActor, FLinearColor CameraColor = FLinearColor::White, float Duration=0.f);
 
 	/* Draws a 2D Histogram of size 'DrawSize' based FDebugFloatHistory struct, using DrawTransform for the position in the world. */
 	UFUNCTION(BlueprintCallable, Category = "Rendering|Debug")
-	static ENGINE_API void DrawDebugFloatHistoryTransform(UObject * WorldContextObject, const FDebugFloatHistory & FloatHistory, const FTransform & DrawTransform, FVector2D DrawSize, FLinearColor DrawColor, float Duration = -1.f);
+	static ENGINE_API void DrawDebugFloatHistoryTransform(UObject * WorldContextObject, const FDebugFloatHistory & FloatHistory, const FTransform & DrawTransform, FVector2D DrawSize, FLinearColor DrawColor = FLinearColor::White, float Duration = 0.f);
 
 	/* Draws a 2D Histogram of size 'DrawSize' based FDebugFloatHistory struct, using DrawLocation for the location in the world, rotation will face camera of first player. */
 	UFUNCTION(BlueprintCallable, Category = "Rendering|Debug")
-	static ENGINE_API void DrawDebugFloatHistoryLocation(UObject * WorldContextObject, const FDebugFloatHistory & FloatHistory, FVector DrawLocation, FVector2D DrawSize, FLinearColor DrawColor, float Duration = -1.f);
+	static ENGINE_API void DrawDebugFloatHistoryLocation(UObject * WorldContextObject, const FDebugFloatHistory & FloatHistory, FVector DrawLocation, FVector2D DrawSize, FLinearColor DrawColor = FLinearColor::White, float Duration = 0.f);
 
 	UFUNCTION(BlueprintPure, Category = "Rendering|Debug")
 	static ENGINE_API FDebugFloatHistory AddFloatHistorySample(float Value, const FDebugFloatHistory & FloatHistory);
diff --git a/Engine/Source/Runtime/Engine/Private/KismetSystemLibrary.cpp b/Engine/Source/Runtime/Engine/Private/KismetSystemLibrary.cpp
--- a/Engine/Source/Runtime/Engine/Private/KismetSystemLibrary.cpp
+++ b/Engine/Source/Runtime/Engine/Private/KismetSystemLibrary.cpp
@@ -2289,25 +2289,25 @@ bool UKismetSystemLibrary::CapsuleTraceMultiByObject_DEPRECATED(UObject* WorldCo
 }
 
 /** Draw a debug line */
-void UKismetSystemLibrary::DrawDebugLine(UObject* WorldContextObject, FVector const& LineStart, FVector const& LineEnd, FLinearColor Color, float LifeTime, float Thickness)
+void UKismetSystemLibrary::DrawDebugLine(UObject* WorldContextObject, FVector const LineStart, FVector const LineEnd, FLinearColor Color, float LifeTime, float Thickness)
 {
 	::DrawDebugLine(GEngine->GetWorldFromContextObject(WorldContextObject), LineStart, LineEnd, Color, false, LifeTime, SDPG_World, Thickness);
 }
 
 /** Draw a debug point */
-void UKismetSystemLibrary::DrawDebugPoint(UObject* WorldContextObject, FVector const& Position, float Size, FLinearColor PointColor, float LifeTime)
+void UKismetSystemLibrary::DrawDebugPoint(UObject* WorldContextObject, FVector const Position, float Size, FLinearColor PointColor, float LifeTime)
 {
 	::DrawDebugPoint(GEngine->GetWorldFromContextObject(WorldContextObject), Position, Size, PointColor, false, LifeTime, SDPG_World);
 }
 
 /** Draw directional arrow, pointing from LineStart to LineEnd. */
-void UKismetSystemLibrary::DrawDebugArrow(UObject* WorldContextObject, FVector const& LineStart, FVector const& LineEnd, float ArrowSize, FLinearColor Color, float LifeTime)
+void UKismetSystemLibrary::DrawDebugArrow(UObject* WorldContextObject, FVector const LineStart, FVector const LineEnd, float ArrowSize, FLinearColor Color, float LifeTime)
 {
 	::DrawDebugDirectionalArrow(GEngine->GetWorldFromContextObject(WorldContextObject), LineStart, LineEnd, ArrowSize, Color, false, LifeTime, SDPG_World);
 }
 
 /** Draw a debug box */
-void UKismetSystemLibrary::DrawDebugBox(UObject* WorldContextObject, FVector const& Center, FVector Extent, FLinearColor Color, const FRotator Rotation, float LifeTime)
+void UKismetSystemLibrary::DrawDebugBox(UObject* WorldContextObject, FVector const Center, FVector Extent, FLinearColor Color, const FRotator Rotation, float LifeTime)
 {
 	UWorld* World = GEngine->GetWorldFromContextObject(WorldContextObject);
 	if (Rotation == FRotator::ZeroRotator)
@@ -2321,37 +2321,42 @@ void UKismetSystemLibrary::DrawDebugBox(UObject* WorldContextObject, FVector con
 }
 
 /** Draw a debug coordinate system. */
-void UKismetSystemLibrary::DrawDebugCoordinateSystem(UObject* WorldContextObject, FVector const& AxisLoc, FRotator const& AxisRot, float Scale)
+void UKismetSystemLibrary::DrawDebugCoordinateSystem(UObject* WorldContextObject, FVector const AxisLoc, FRotator const AxisRot, float Scale, float LifeTime)
 {
-	::DrawDebugCoordinateSystem(GEngine->GetWorldFromContextObject(WorldContextObject), AxisLoc, AxisRot, Scale, false, SDPG_World);
+	::DrawDebugCoordinateSystem(GEngine->GetWorldFromContextObject(WorldContextObject), AxisLoc, AxisRot, Scale, false, LifeTime, SDPG_World);
 }
 
 /** Draw a debug sphere */
-void UKismetSystemLibrary::DrawDebugSphere(UObject* WorldContextObject, FVector const& Center, float Radius, int32 Segments, FLinearColor Color, float LifeTime)
+void UKismetSystemLibrary::DrawDebugSphere(UObject* WorldContextObject, FVector const Center, float Radius, int32 Segments, FLinearColor Color, float LifeTime)
 {
 	::DrawDebugSphere(GEngine->GetWorldFromContextObject(WorldContextObject), Center, Radius, Segments, Color.ToFColor(true), false, LifeTime, SDPG_World);
 }
 
 /** Draw a debug cylinder */
-void UKismetSystemLibrary::DrawDebugCylinder(UObject* WorldContextObject, FVector const& Start, FVector const& End, float Radius, int32 Segments, FLinearColor Color)
+void UKismetSystemLibrary::DrawDebugCylinder(UObject* WorldContextObject, FVector const Start, FVector const End, float Radius, int32 Segments, FLinearColor Color, float LifeTime)
 {
-	::DrawDebugCylinder(GEngine->GetWorldFromContextObject(WorldContextObject), Start, End, Radius, Segments, Color.ToFColor(true), false, SDPG_World);
+	::DrawDebugCylinder(GEngine->GetWorldFromContextObject(WorldContextObject), Start, End, Radius, Segments, Color.ToFColor(true), false, LifeTime, SDPG_World);
 }
 
 /** Draw a debug cone */
-void UKismetSystemLibrary::DrawDebugCone(UObject* WorldContextObject, FVector const& Origin, FVector const& Direction, float Length, float AngleWidth, float AngleHeight, int32 NumSides, FLinearColor Color)
+void UKismetSystemLibrary::DrawDebugCone(UObject* WorldContextObject, FVector const Origin, FVector const Direction, float Length, float AngleWidth, float AngleHeight, int32 NumSides, FLinearColor Color)
 {
-	::DrawDebugCone(GEngine->GetWorldFromContextObject(WorldContextObject), Origin, Direction, Length, AngleWidth, AngleHeight, NumSides, Color.ToFColor(true), false, SDPG_World);
+	::DrawDebugCone(GEngine->GetWorldFromContextObject(WorldContextObject), Origin, Direction, Length, AngleWidth, AngleHeight, NumSides, Color.ToFColor(true), false, -1.f, SDPG_World);
+}
+
+void UKismetSystemLibrary::DrawDebugConeInDegrees(UObject* WorldContextObject, FVector const Origin, FVector const Direction, float Length, float AngleWidth, float AngleHeight, int32 NumSides, FLinearColor Color, float LifeTime)
+{
+	::DrawDebugCone(GEngine->GetWorldFromContextObject(WorldContextObject), Origin, Direction, Length, FMath::DegreesToRadians(AngleWidth), FMath::DegreesToRadians(AngleHeight), NumSides, Color.ToFColor(true), false, LifeTime, SDPG_World);
 }
 
 /** Draw a debug capsule */
-void UKismetSystemLibrary::DrawDebugCapsule(UObject* WorldContextObject, FVector const& Center, float HalfHeight, float Radius, const FRotator& Rotation, FLinearColor Color, float LifeTime)
+void UKismetSystemLibrary::DrawDebugCapsule(UObject* WorldContextObject, FVector const Center, float HalfHeight, float Radius, const FRotator Rotation, FLinearColor Color, float LifeTime)
 {
 	::DrawDebugCapsule(GEngine->GetWorldFromContextObject(WorldContextObject), Center, HalfHeight, Radius, Rotation.Quaternion(), Color.ToFColor(true), false, LifeTime, SDPG_World);
 }
 
 /** Draw a debug string at a 3d world location. */
-void UKismetSystemLibrary::DrawDebugString(UObject* WorldContextObject, FVector const& TextLocation, const FString& Text, class AActor* TestBaseActor, FLinearColor TextColor, float Duration)
+void UKismetSystemLibrary::DrawDebugString(UObject* WorldContextObject, FVector const TextLocation, const FString& Text, class AActor* TestBaseActor, FLinearColor TextColor, float Duration)
 {
 	::DrawDebugString(GEngine->GetWorldFromContextObject( WorldContextObject ), TextLocation, Text, TestBaseActor, TextColor.ToFColor(true), Duration);
 }
@@ -2364,7 +2369,7 @@ void UKismetSystemLibrary::FlushDebugStrings( UObject* WorldContextObject )
 }
 
 /** Draws a debug plane. */
-void UKismetSystemLibrary::DrawDebugPlane(UObject* WorldContextObject, FPlane const& P, FVector const& Loc, float Size, FLinearColor Color, float LifeTime)
+void UKismetSystemLibrary::DrawDebugPlane(UObject* WorldContextObject, FPlane const P, FVector const Loc, float Size, FLinearColor Color, float LifeTime)
 {
 	::DrawDebugSolidPlane(GEngine->GetWorldFromContextObject( WorldContextObject ), P, Loc, Size, Color.ToFColor(true), false, LifeTime, SDPG_World);
 }
@@ -2377,12 +2382,12 @@ void UKismetSystemLibrary::FlushPersistentDebugLines(UObject* WorldContextObject
 }
 
 /** Draws a debug frustum. */
-void UKismetSystemLibrary::DrawDebugFrustum(UObject* WorldContextObject, const FTransform& FrustumTransform, FLinearColor FrustumColor)
+void UKismetSystemLibrary::DrawDebugFrustum(UObject* WorldContextObject, const FTransform& FrustumTransform, FLinearColor FrustumColor, float Duration)
 {
 	if( FrustumTransform.IsRotationNormalized() )
 	{
 		FMatrix FrustumToWorld =  FrustumTransform.ToMatrixWithScale();
-		::DrawDebugFrustum(GEngine->GetWorldFromContextObject(WorldContextObject), FrustumToWorld, FrustumColor, false, SDPG_World);
+		::DrawDebugFrustum(GEngine->GetWorldFromContextObject(WorldContextObject), FrustumToWorld, FrustumColor, false, Duration, SDPG_World);
 	}
 }
 
