Fix for toggling fullscreen not working when -windowed or -fullscreen is specified on the command line.

diff --git a/Engine/Source/Editor/UnrealEd/Private/Editor.cpp b/Engine/Source/Editor/UnrealEd/Private/Editor.cpp
--- a/Engine/Source/Editor/UnrealEd/Private/Editor.cpp
+++ b/Engine/Source/Editor/UnrealEd/Private/Editor.cpp
@@ -676,7 +676,7 @@ void UEditorEngine::Init(IEngineLoop* InEngineLoop)
 	if (GameUserSettings)
 	{
 		GameUserSettings->LoadSettings();
-		GameUserSettings->ApplySettings();
+		GameUserSettings->ApplySettings(true);
 	}
 
 	UEditorStyleSettings* Settings = GetMutableDefault<UEditorStyleSettings>();
diff --git a/Engine/Source/Runtime/Engine/Classes/GameFramework/GameUserSettings.h b/Engine/Source/Runtime/Engine/Classes/GameFramework/GameUserSettings.h
--- a/Engine/Source/Runtime/Engine/Classes/GameFramework/GameUserSettings.h
+++ b/Engine/Source/Runtime/Engine/Classes/GameFramework/GameUserSettings.h
@@ -11,7 +11,11 @@ class ENGINE_API UGameUserSettings : public UObject
 
 public:
 
+	/** Applies all current user settings to the game and saves to permanent storage (e.g. file), optionally checking for command line overrides. */
+	virtual void ApplySettings(bool bCheckForCommandLineOverrides);
+
 	/** Applies all current user settings to the game and saves to permanent storage (e.g. file). */
+	DEPRECATED(4.4, "ApplySettings() is deprecated, use ApplySettings(bool) instead.")
 	virtual void ApplySettings();
 
 	/** Returns the user setting for game screen resolution, in pixels. */
@@ -93,7 +97,7 @@ public:
 	static void LoadConfigIni( bool bForceReload=false );
 
 	/** Apply any runtime overrides to the settings */
-	static void RequestResolutionChange(int32 InResolutionX, int32 InResolutionY, EWindowMode::Type InWindowMode);
+	static void RequestResolutionChange(int32 InResolutionX, int32 InResolutionY, EWindowMode::Type InWindowMode, bool bConditionallyOverride);
 
 	/** Whether to use VSync or not. (public to allow UI to connect to it) */
 	UPROPERTY(config)
diff --git a/Engine/Source/Runtime/Engine/Private/GameEngine.cpp b/Engine/Source/Runtime/Engine/Private/GameEngine.cpp
--- a/Engine/Source/Runtime/Engine/Private/GameEngine.cpp
+++ b/Engine/Source/Runtime/Engine/Private/GameEngine.cpp
@@ -422,7 +422,7 @@ void UGameEngine::Init(IEngineLoop* InEngineLoop)
 
 	// Load and apply user game settings
 	GetGameUserSettings()->LoadSettings();
-	GetGameUserSettings()->ApplySettings();
+	GetGameUserSettings()->ApplySettings(true);
 
 	// Create game instance.  For GameEngine, this should be the only GameInstance that ever gets created.
 	{
@@ -695,7 +695,7 @@ bool UGameEngine::HandleCancelCommand( const TCHAR* Cmd, FOutputDevice& Ar, UWor
 #if !UE_BUILD_SHIPPING
 bool UGameEngine::HandleApplyUserSettingsCommand( const TCHAR* Cmd, FOutputDevice& Ar )
 {
-	GetGameUserSettings()->ApplySettings();
+	GetGameUserSettings()->ApplySettings(false);
 	return true;
 }
 #endif // !UE_BUILD_SHIPPING
diff --git a/Engine/Source/Runtime/Engine/Private/GameUserSettings.cpp b/Engine/Source/Runtime/Engine/Private/GameUserSettings.cpp
--- a/Engine/Source/Runtime/Engine/Private/GameUserSettings.cpp
+++ b/Engine/Source/Runtime/Engine/Private/GameUserSettings.cpp
@@ -182,7 +182,7 @@ void UGameUserSettings::ValidateSettings()
 	UpdateVersion();
 }
 
-void UGameUserSettings::ApplySettings()
+void UGameUserSettings::ApplySettings(bool bCheckForCommandLineOverrides)
 {
 	ValidateSettings();
 
@@ -195,7 +195,7 @@ void UGameUserSettings::ApplySettings()
 	}
 
 	// Request a resolution change
-	RequestResolutionChange(ResolutionSizeX, ResolutionSizeY, NewWindowMode);
+	RequestResolutionChange(ResolutionSizeX, ResolutionSizeY, NewWindowMode, bCheckForCommandLineOverrides);
 
 	// Update vsync cvar
 	{
@@ -216,6 +216,13 @@ void UGameUserSettings::ApplySettings()
 	SaveSettings();
 }
 
+void UGameUserSettings::ApplySettings()
+{
+	// deprecated.  calling the new version with true to match 4.3 behavior
+	// should generally use false going forward except when you want command line options to override settings
+	ApplySettings(true);
+}
+
 void UGameUserSettings::LoadSettings( bool bForceReload/*=false*/ )
 {
 	if ( bForceReload )
@@ -237,12 +244,16 @@ void UGameUserSettings::LoadSettings( bool bForceReload/*=false*/ )
 	}
 }
 
-void UGameUserSettings::RequestResolutionChange(int32 InResolutionX, int32 InResolutionY, EWindowMode::Type InWindowMode)
+void UGameUserSettings::RequestResolutionChange(int32 InResolutionX, int32 InResolutionY, EWindowMode::Type InWindowMode, bool bConditionallyOverride)
 {
-	UGameEngine::ConditionallyOverrideSettings(InResolutionX, InResolutionY, InWindowMode);
+	if (bConditionallyOverride)
 	{
-		IConsoleVariable* CVar = IConsoleManager::Get().FindConsoleVariable(TEXT("r.FullScreenMode"));
-		CVar->Set(InWindowMode);
+		// check to see if the command line params should override the requested change
+		UGameEngine::ConditionallyOverrideSettings(InResolutionX, InResolutionY, InWindowMode);
+		{
+			IConsoleVariable* CVar = IConsoleManager::Get().FindConsoleVariable(TEXT("r.FullScreenMode"));
+			CVar->Set(InWindowMode);
+		}
 	}
 	FSystemResolution::RequestResolutionChange(InResolutionX, InResolutionY, InWindowMode);
 }
@@ -304,7 +315,7 @@ void UGameUserSettings::PreloadResolutionSettings()
 		CVar->Set((int32)WindowMode);
 	}
 
-	RequestResolutionChange(ResolutionX, ResolutionY, WindowMode);
+	RequestResolutionChange(ResolutionX, ResolutionY, WindowMode, true);
 
 	IConsoleManager::Get().CallAllConsoleVariableSinks();
 }
diff --git a/Engine/Source/Runtime/Engine/Private/GameViewportClient.cpp b/Engine/Source/Runtime/Engine/Private/GameViewportClient.cpp
--- a/Engine/Source/Runtime/Engine/Private/GameViewportClient.cpp
+++ b/Engine/Source/Runtime/Engine/Private/GameViewportClient.cpp
@@ -2361,7 +2361,7 @@ bool UGameViewportClient::SetDisplayConfiguration(const FIntPoint* Dimensions, E
 			UserSettings->SetScreenResolution(*Dimensions);
 		}
 
-		UserSettings->ApplySettings();
+		UserSettings->ApplySettings(false);
 	}
 	else
 	{
