Landscape : Fixed Visibility layer to be no weight blend, to avoid other layer override its value

diff --git a/Engine/Source/Runtime/Engine/Private/Landscape/Landscape.cpp b/Engine/Source/Runtime/Engine/Private/Landscape/Landscape.cpp
--- a/Engine/Source/Runtime/Engine/Private/Landscape/Landscape.cpp
+++ b/Engine/Source/Runtime/Engine/Private/Landscape/Landscape.cpp
@@ -514,6 +514,10 @@ ALandscapeProxy::ALandscapeProxy(const class FPostConstructInitializeProperties&
 
 		VisibilityLayer = ConstructorStatics.DataLayer.Get();
 		check(VisibilityLayer);
+#if WITH_EDITORONLY_DATA
+		// This layer should be no weight blending
+		VisibilityLayer->bNoWeightBlend = true;
+#endif
 		VisibilityLayer->AddToRoot();
 	}
 }
