add feature level to UMaterialInterface::GetUsedTextures.

diff --git a/Engine/Source/Editor/MeshPaint/Private/MeshPaintEdMode.cpp b/Engine/Source/Editor/MeshPaint/Private/MeshPaintEdMode.cpp
--- a/Engine/Source/Editor/MeshPaint/Private/MeshPaintEdMode.cpp
+++ b/Engine/Source/Editor/MeshPaint/Private/MeshPaintEdMode.cpp
@@ -2503,7 +2503,7 @@ void FEdModeMeshPaint::ClearStaticMeshTextureOverrides(UStaticMeshComponent* InS
 		if( Material != NULL )
 		{
 			TArray<UTexture*> UsedTextures;
-			Material->GetUsedTextures( UsedTextures, EMaterialQualityLevel::Num, false );
+			Material->GetUsedTextures( UsedTextures, EMaterialQualityLevel::Num, false, GRHIFeatureLevel, false );
 
 			for( int32 UsedIndex = 0; UsedIndex < UsedTextures.Num(); UsedIndex++ )
 			{
@@ -2553,7 +2553,7 @@ void FEdModeMeshPaint::SetAllTextureOverrides(UStaticMeshComponent* InStaticMesh
 		if( Material != NULL )
 		{
 			TArray<UTexture*> UsedTextures;
-			Material->GetUsedTextures( UsedTextures, EMaterialQualityLevel::Num, false );
+			Material->GetUsedTextures( UsedTextures, EMaterialQualityLevel::Num, false, GRHIFeatureLevel, false );
 
 			for( int32 UsedIndex = 0; UsedIndex < UsedTextures.Num(); UsedIndex++ )
 			{
@@ -2832,7 +2832,7 @@ void FEdModeMeshPaint::SaveSettingsForActor( AActor* InActor )
 				}
 
 				TArray<UTexture*> UsedTextures;
-				UsedMaterials[ MatIndex ]->GetUsedTextures( UsedTextures, EMaterialQualityLevel::Num, false );
+				UsedMaterials[ MatIndex ]->GetUsedTextures( UsedTextures, EMaterialQualityLevel::Num, false, GRHIFeatureLevel, false );
 
 				for( int32 TexIndex = 0; TexIndex < UsedTextures.Num(); TexIndex++ )
 				{
@@ -2882,7 +2882,7 @@ void FEdModeMeshPaint::UpdateSettingsForStaticMeshComponent( UStaticMeshComponen
 			}
 
 			TArray<UTexture*> UsedTextures;
-			UsedMaterials[ MatIndex ]->GetUsedTextures( UsedTextures, EMaterialQualityLevel::Num, false );
+			UsedMaterials[MatIndex]->GetUsedTextures(UsedTextures, EMaterialQualityLevel::Num, false, GRHIFeatureLevel, false);
 
 			for( int32 TexIndex = 0; TexIndex < UsedTextures.Num(); TexIndex++ )
 			{
diff --git a/Engine/Source/Editor/PropertyEditor/Private/PropertyCustomizationHelpers.cpp b/Engine/Source/Editor/PropertyEditor/Private/PropertyCustomizationHelpers.cpp
--- a/Engine/Source/Editor/PropertyEditor/Private/PropertyCustomizationHelpers.cpp
+++ b/Engine/Source/Editor/PropertyEditor/Private/PropertyCustomizationHelpers.cpp
@@ -689,7 +689,7 @@ private:
 			UMaterialInterface* Material = MaterialItem.Material.Get();
 
 			TArray< UTexture* > Textures;
-			Material->GetUsedTextures( Textures, EMaterialQualityLevel::Num, false );
+			Material->GetUsedTextures(Textures, EMaterialQualityLevel::Num, false, GRHIFeatureLevel, false);
 
 			// Add a menu item for each texture.  Clicking on the texture will display it in the content browser
 			for( int32 TextureIndex = 0; TextureIndex < Textures.Num(); ++TextureIndex )
diff --git a/Engine/Source/Editor/PropertyEditor/PropertyEditor.Build.cs b/Engine/Source/Editor/PropertyEditor/PropertyEditor.Build.cs
--- a/Engine/Source/Editor/PropertyEditor/PropertyEditor.Build.cs
+++ b/Engine/Source/Editor/PropertyEditor/PropertyEditor.Build.cs
@@ -45,6 +45,7 @@ public class PropertyEditor : ModuleRules
 				"SlateCore",
 				"EditorWidgets",
 				"Documentation",
+                "RHI",
 			}
 		);
 
diff --git a/Engine/Source/Editor/StatsViewer/Private/StatsPages/TextureStatsPage.cpp b/Engine/Source/Editor/StatsViewer/Private/StatsPages/TextureStatsPage.cpp
--- a/Engine/Source/Editor/StatsViewer/Private/StatsPages/TextureStatsPage.cpp
+++ b/Engine/Source/Editor/StatsViewer/Private/StatsPages/TextureStatsPage.cpp
@@ -299,7 +299,7 @@ struct TextureStatsGenerator : public FFindReferencedAssets
 						{
 							TArray<UTexture*> UsedTextures;
 
-							CurrentMaterial->GetUsedTextures(UsedTextures, EMaterialQualityLevel::Num, false);
+							CurrentMaterial->GetUsedTextures(UsedTextures, EMaterialQualityLevel::Num, false, GRHIFeatureLevel, false);
 							for(int32 TextureIndex = 0; TextureIndex < UsedTextures.Num(); TextureIndex++)
 							{
 								UTexture* CurrentUsedTexture = UsedTextures[TextureIndex];
diff --git a/Engine/Source/Editor/UnrealEd/Private/EditorExporters.cpp b/Engine/Source/Editor/UnrealEd/Private/EditorExporters.cpp
--- a/Engine/Source/Editor/UnrealEd/Private/EditorExporters.cpp
+++ b/Engine/Source/Editor/UnrealEd/Private/EditorExporters.cpp
@@ -1143,7 +1143,7 @@ public:
 
 		TArray<UTexture*> MaterialTextures;
 
-		InMaterialInterface->GetUsedTextures(MaterialTextures, EMaterialQualityLevel::Num, false);
+		InMaterialInterface->GetUsedTextures(MaterialTextures, EMaterialQualityLevel::Num, false, GRHIFeatureLevel, false);
 
 		// find the largest texture in the list (applying it's LOD bias)
 		FIntPoint MaxSize = MinimumSize;
diff --git a/Engine/Source/Runtime/Engine/Classes/Materials/Material.h b/Engine/Source/Runtime/Engine/Classes/Materials/Material.h
--- a/Engine/Source/Runtime/Engine/Classes/Materials/Material.h
+++ b/Engine/Source/Runtime/Engine/Classes/Materials/Material.h
@@ -751,7 +751,7 @@ public:
 	ENGINE_API virtual bool GetRefractionSettings(float& OutBiasValue) const override;
 	ENGINE_API virtual FMaterialRenderProxy* GetRenderProxy(bool Selected, bool bHovered=false) const override;
 	ENGINE_API virtual UPhysicalMaterial* GetPhysicalMaterial() const override;
-	ENGINE_API virtual void GetUsedTextures(TArray<UTexture*>& OutTextures, EMaterialQualityLevel::Type QualityLevel, bool bAllQualityLevels) const override;
+	ENGINE_API virtual void GetUsedTextures(TArray<UTexture*>& OutTextures, EMaterialQualityLevel::Type QualityLevel, bool bAllQualityLevels, ERHIFeatureLevel::Type FeatureLevel, bool bAllFeatureLevels) const override;
 	ENGINE_API virtual void OverrideTexture( const UTexture* InTextureToOverride, UTexture* OverrideTexture ) override;
 	ENGINE_API virtual bool CheckMaterialUsage(const EMaterialUsage Usage, const bool bSkipPrim = false) override;
 	ENGINE_API virtual bool CheckMaterialUsage_Concurrent(const EMaterialUsage Usage, const bool bSkipPrim = false) const override;
diff --git a/Engine/Source/Runtime/Engine/Classes/Materials/MaterialInstance.h b/Engine/Source/Runtime/Engine/Classes/Materials/MaterialInstance.h
--- a/Engine/Source/Runtime/Engine/Classes/Materials/MaterialInstance.h
+++ b/Engine/Source/Runtime/Engine/Classes/Materials/MaterialInstance.h
@@ -218,7 +218,7 @@ public:
 	virtual bool GetScalarParameterValue(FName ParameterName,float& OutValue) const override;
 	virtual bool GetTextureParameterValue(FName ParameterName,class UTexture*& OutValue) const override;
 	virtual bool GetVectorParameterValue(FName ParameterName,FLinearColor& OutValue) const override;
-	virtual void GetUsedTextures(TArray<UTexture*>& OutTextures, EMaterialQualityLevel::Type QualityLevel, bool bAllQualityLevels) const override;
+	virtual void GetUsedTextures(TArray<UTexture*>& OutTextures, EMaterialQualityLevel::Type QualityLevel, bool bAllQualityLevels, ERHIFeatureLevel::Type FeatureLevel, bool bAllFeatureLevels) const override;
 	virtual void OverrideTexture( const UTexture* InTextureToOverride, UTexture* OverrideTexture ) override;
 	virtual bool CheckMaterialUsage(const EMaterialUsage Usage, const bool bSkipPrim = false) override;
 	virtual bool CheckMaterialUsage_Concurrent(const EMaterialUsage Usage, const bool bSkipPrim = false) const override;
diff --git a/Engine/Source/Runtime/Engine/Classes/Materials/MaterialInterface.h b/Engine/Source/Runtime/Engine/Classes/Materials/MaterialInterface.h
--- a/Engine/Source/Runtime/Engine/Classes/Materials/MaterialInterface.h
+++ b/Engine/Source/Runtime/Engine/Classes/Materials/MaterialInterface.h
@@ -252,7 +252,7 @@ public:
 	virtual UPhysicalMaterial* GetPhysicalMaterial() const PURE_VIRTUAL(UMaterialInterface::GetPhysicalMaterial,return NULL;);
 
 	/** Return the textures used to render this material. */
-	virtual void GetUsedTextures(TArray<UTexture*>& OutTextures, EMaterialQualityLevel::Type QualityLevel, bool bAllQualityLevels) const 
+	virtual void GetUsedTextures(TArray<UTexture*>& OutTextures, EMaterialQualityLevel::Type QualityLevel, bool bAllQualityLevels, ERHIFeatureLevel::Type FeatureLevel, bool bAllFeatureLevels) const
 		PURE_VIRTUAL(UMaterialInterface::GetUsedTextures,);
 
 	/**
diff --git a/Engine/Source/Runtime/Engine/Private/Landscape/LandscapeRender.cpp b/Engine/Source/Runtime/Engine/Private/Landscape/LandscapeRender.cpp
--- a/Engine/Source/Runtime/Engine/Private/Landscape/LandscapeRender.cpp
+++ b/Engine/Source/Runtime/Engine/Private/Landscape/LandscapeRender.cpp
@@ -2350,7 +2350,7 @@ void ULandscapeComponent::GetStreamingTextureInfo(TArray<FStreamingTexturePrimit
 	if (MaterialInstance)
 	{
 		TArray<UTexture*> Textures;
-		MaterialInstance->GetUsedTextures(Textures, EMaterialQualityLevel::Num, false);
+		MaterialInstance->GetUsedTextures(Textures, EMaterialQualityLevel::Num, false, GRHIFeatureLevel, false);
 		// Add each texture to the output with the appropriate parameters.
 		// TODO: Take into account which UVIndex is being used.
 		for (int32 TextureIndex = 0;TextureIndex < Textures.Num();TextureIndex++)
diff --git a/Engine/Source/Runtime/Engine/Private/Materials/Material.cpp b/Engine/Source/Runtime/Engine/Private/Materials/Material.cpp
--- a/Engine/Source/Runtime/Engine/Private/Materials/Material.cpp
+++ b/Engine/Source/Runtime/Engine/Private/Materials/Material.cpp
@@ -549,7 +549,7 @@ FMaterialResource* UMaterial::AllocateResource()
 	return new FMaterialResource();
 }
 
-void UMaterial::GetUsedTextures(TArray<UTexture*>& OutTextures, EMaterialQualityLevel::Type QualityLevel, bool bAllQualityLevels) const
+void UMaterial::GetUsedTextures(TArray<UTexture*>& OutTextures, EMaterialQualityLevel::Type QualityLevel, bool bAllQualityLevels, ERHIFeatureLevel::Type FeatureLevel, bool bAllFeatureLevels) const
 {
 	OutTextures.Empty();
 
@@ -562,27 +562,33 @@ void UMaterial::GetUsedTextures(TArray<UTexture*>& OutTextures, EMaterialQuality
 	{
 		for (int32 QualityLevelIndex = 0; QualityLevelIndex < EMaterialQualityLevel::Num; QualityLevelIndex++)
 		{
-			const FMaterialResource* CurrentResource = MaterialResources[QualityLevelIndex][GRHIFeatureLevel];
+			if (QualityLevelIndex != QualityLevel && !bAllQualityLevels)
+				continue;
 
-			if (QualityLevelIndex == QualityLevel || bAllQualityLevels)
+			for (int32 FeatureLevelIndex = 0; FeatureLevelIndex < ERHIFeatureLevel::Num; FeatureLevelIndex++)
 			{
+				const FMaterialResource* CurrentResource = MaterialResources[QualityLevelIndex][FeatureLevelIndex];
+				if (CurrentResource == nullptr || (FeatureLevelIndex != FeatureLevel && !bAllFeatureLevels))
+					continue;
+
 				const TArray<TRefCountPtr<FMaterialUniformExpressionTexture> >* ExpressionsByType[2] =
 				{
 					&CurrentResource->GetUniform2DTextureExpressions(),
 					&CurrentResource->GetUniformCubeTextureExpressions()
 				};
-				for(int32 TypeIndex = 0;TypeIndex < ARRAY_COUNT(ExpressionsByType);TypeIndex++)
+				for (int32 TypeIndex = 0; TypeIndex < ARRAY_COUNT(ExpressionsByType); TypeIndex++)
 				{
 					// Iterate over each of the material's texture expressions.
 					for (FMaterialUniformExpressionTexture* Expression : *ExpressionsByType[TypeIndex])
 					{
 						const bool bAllowOverride = false;
 						UTexture* Texture = NULL;
-						Expression->GetGameThreadTextureValue(this,*CurrentResource,Texture,bAllowOverride);
+						Expression->GetGameThreadTextureValue(this, *CurrentResource, Texture, bAllowOverride);
 
 						if (Texture)
 						{
 							OutTextures.Add(Texture);
+							//OutTextures.AddUnique(Texture); //AJB - maybe this?
 						}
 					}
 				}
@@ -2934,7 +2940,7 @@ bool UMaterial::UpdateLightmassTextureTracking()
 #if WITH_EDITORONLY_DATA
 	TArray<UTexture*> UsedTextures;
 	
-	GetUsedTextures(UsedTextures, EMaterialQualityLevel::Num, true);
+	GetUsedTextures(UsedTextures, EMaterialQualityLevel::Num, true, GRHIFeatureLevel, true);
 	if (UsedTextures.Num() != ReferencedTextureGuids.Num())
 	{
 		bTexturesHaveChanged = true;
diff --git a/Engine/Source/Runtime/Engine/Private/Materials/MaterialInstance.cpp b/Engine/Source/Runtime/Engine/Private/Materials/MaterialInstance.cpp
--- a/Engine/Source/Runtime/Engine/Private/Materials/MaterialInstance.cpp
+++ b/Engine/Source/Runtime/Engine/Private/Materials/MaterialInstance.cpp
@@ -673,7 +673,7 @@ void UMaterialInstance::GetTextureExpressionValues(const FMaterialResource* Mate
 	}
 }
 
-void UMaterialInstance::GetUsedTextures(TArray<UTexture*>& OutTextures, EMaterialQualityLevel::Type QualityLevel, bool bAllQualityLevels) const
+void UMaterialInstance::GetUsedTextures(TArray<UTexture*>& OutTextures, EMaterialQualityLevel::Type QualityLevel, bool bAllQualityLevels, ERHIFeatureLevel::Type FeatureLevel, bool bAllFeatureLevels) const
 {
 	OutTextures.Empty();
 
@@ -698,12 +698,17 @@ void UMaterialInstance::GetUsedTextures(TArray<UTexture*>& OutTextures, EMateria
 		{
 			for (int32 QualityLevelIndex = 0; QualityLevelIndex < EMaterialQualityLevel::Num; QualityLevelIndex++)
 			{
-				const FMaterialResource* CurrentResource = MaterialInstanceToUse->StaticPermutationMaterialResources[QualityLevelIndex][GRHIFeatureLevel];
-
-				//@todo - GetUsedTextures is incorrect during cooking since we don't cache shaders for the current platform during cooking
-				if (QualityLevelIndex == QualityLevel || bAllQualityLevels)
+				for (int32 FeatureLevelIndex = 0; FeatureLevelIndex < ERHIFeatureLevel::Num; FeatureLevelIndex++)
 				{
-					GetTextureExpressionValues(CurrentResource, OutTextures);
+					const FMaterialResource* CurrentResource = MaterialInstanceToUse->StaticPermutationMaterialResources[QualityLevelIndex][FeatureLevelIndex];
+					if (CurrentResource == nullptr || (FeatureLevelIndex != FeatureLevel && !bAllFeatureLevels))
+						continue;
+
+					//@todo - GetUsedTextures is incorrect during cooking since we don't cache shaders for the current platform during cooking
+					if (QualityLevelIndex == QualityLevel || bAllQualityLevels)
+					{
+						GetTextureExpressionValues(CurrentResource, OutTextures);
+					}
 				}
 			}
 		}
@@ -714,13 +719,13 @@ void UMaterialInstance::GetUsedTextures(TArray<UTexture*>& OutTextures, EMateria
 
 			if (Material)
 			{
-				const FMaterialResource* MaterialResource = Material->GetMaterialResource(GRHIFeatureLevel, QualityLevel);
+				const FMaterialResource* MaterialResource = Material->GetMaterialResource(FeatureLevel, QualityLevel);
 				GetTextureExpressionValues(MaterialResource, OutTextures);
 			}
 			else
 			{
 				// If the material instance has no material, use the default material.
-				UMaterial::GetDefaultMaterial(MD_Surface)->GetUsedTextures(OutTextures, QualityLevel, bAllQualityLevels);
+				UMaterial::GetDefaultMaterial(MD_Surface)->GetUsedTextures(OutTextures, QualityLevel, bAllQualityLevels, FeatureLevel, bAllFeatureLevels);
 			}
 		}
 	}
@@ -2025,7 +2030,7 @@ bool UMaterialInstance::UpdateLightmassTextureTracking()
 #if WITH_EDITORONLY_DATA
 	TArray<UTexture*> UsedTextures;
 	
-	GetUsedTextures(UsedTextures, EMaterialQualityLevel::Num, true);
+	GetUsedTextures(UsedTextures, EMaterialQualityLevel::Num, true, GRHIFeatureLevel, true);
 	if (UsedTextures.Num() != ReferencedTextureGuids.Num())
 	{
 		bTexturesHaveChanged = true;
diff --git a/Engine/Source/Runtime/Engine/Private/Materials/MaterialInterface.cpp b/Engine/Source/Runtime/Engine/Private/Materials/MaterialInterface.cpp
--- a/Engine/Source/Runtime/Engine/Private/Materials/MaterialInterface.cpp
+++ b/Engine/Source/Runtime/Engine/Private/Materials/MaterialInterface.cpp
@@ -89,7 +89,7 @@ void UMaterialInterface::SetForceMipLevelsToBeResident( bool OverrideForceMiplev
 {
 	TArray<UTexture*> Textures;
 	
-	GetUsedTextures(Textures, EMaterialQualityLevel::Num, false);
+	GetUsedTextures(Textures, EMaterialQualityLevel::Num, false, GRHIFeatureLevel, false);
 	for ( int32 TextureIndex=0; TextureIndex < Textures.Num(); ++TextureIndex )
 	{
 		UTexture2D* Texture = Cast<UTexture2D>(Textures[TextureIndex]);
@@ -217,7 +217,7 @@ bool DoesMaterialUseTexture(const UMaterialInterface* Material,const UTexture* C
 	}
 
 	TArray<UTexture*> Textures;
-	Material->GetUsedTextures(Textures, EMaterialQualityLevel::Num, true);
+	Material->GetUsedTextures(Textures, EMaterialQualityLevel::Num, true, GRHIFeatureLevel, true);
 	for (int32 i = 0; i < Textures.Num(); i++)
 	{
 		if (Textures[i] == CheckTexture)
diff --git a/Engine/Source/Runtime/Engine/Private/ModelComponent.cpp b/Engine/Source/Runtime/Engine/Private/ModelComponent.cpp
--- a/Engine/Source/Runtime/Engine/Private/ModelComponent.cpp
+++ b/Engine/Source/Runtime/Engine/Private/ModelComponent.cpp
@@ -368,7 +368,7 @@ void UModelComponent::GetStreamingTextureInfo(TArray<FStreamingTexturePrimitiveI
 				// Enumerate the textures used by the surface's material.
 				TArray<UTexture*> Textures;
 				
-				Material->GetUsedTextures(Textures, EMaterialQualityLevel::Num, false);
+				Material->GetUsedTextures(Textures, EMaterialQualityLevel::Num, false, GRHIFeatureLevel, false);
 
 				// Add each texture to the output with the appropriate parameters.
 				for(int32 TextureIndex = 0;TextureIndex < Textures.Num();TextureIndex++)
diff --git a/Engine/Source/Runtime/Engine/Private/PrimitiveComponent.cpp b/Engine/Source/Runtime/Engine/Private/PrimitiveComponent.cpp
--- a/Engine/Source/Runtime/Engine/Private/PrimitiveComponent.cpp
+++ b/Engine/Source/Runtime/Engine/Private/PrimitiveComponent.cpp
@@ -167,7 +167,7 @@ void UPrimitiveComponent::GetUsedTextures(TArray<UTexture*>& OutTextures, EMater
 		if( UsedMaterials[ MatIndex ] )
 		{
 			UsedTextures.Reset();
-			UsedMaterials[ MatIndex ]->GetUsedTextures( UsedTextures, QualityLevel, false );
+			UsedMaterials[ MatIndex ]->GetUsedTextures( UsedTextures, QualityLevel, false, GRHIFeatureLevel, false );
 
 			for( int32 TextureIndex=0; TextureIndex<UsedTextures.Num(); TextureIndex++ )
 			{
diff --git a/Engine/Source/Runtime/Engine/Private/SkinnedMeshComponent.cpp b/Engine/Source/Runtime/Engine/Private/SkinnedMeshComponent.cpp
--- a/Engine/Source/Runtime/Engine/Private/SkinnedMeshComponent.cpp
+++ b/Engine/Source/Runtime/Engine/Private/SkinnedMeshComponent.cpp
@@ -373,7 +373,7 @@ void USkinnedMeshComponent::GetStreamingTextureInfo(TArray<FStreamingTexturePrim
 			{
 				TArray<UTexture*> Textures;
 				
-				MaterialInterface->GetUsedTextures(Textures, EMaterialQualityLevel::Num, false);
+				MaterialInterface->GetUsedTextures(Textures, EMaterialQualityLevel::Num, false, GRHIFeatureLevel, false);
 				for(int32 TextureIndex = 0; TextureIndex < Textures.Num(); TextureIndex++ )
 				{
 					FStreamingTexturePrimitiveInfo& StreamingTexture = *new(OutStreamingTextures) FStreamingTexturePrimitiveInfo;
diff --git a/Engine/Source/Runtime/Engine/Private/StaticMeshComponent.cpp b/Engine/Source/Runtime/Engine/Private/StaticMeshComponent.cpp
--- a/Engine/Source/Runtime/Engine/Private/StaticMeshComponent.cpp
+++ b/Engine/Source/Runtime/Engine/Private/StaticMeshComponent.cpp
@@ -518,7 +518,7 @@ void UStaticMeshComponent::GetStreamingTextureInfo(TArray<FStreamingTexturePrimi
 				// Enumerate the textures used by the material.
 				TArray<UTexture*> Textures;
 
-				Material->GetUsedTextures(Textures, EMaterialQualityLevel::Num, false);
+				Material->GetUsedTextures(Textures, EMaterialQualityLevel::Num, false, GRHIFeatureLevel, false);
 
 				// Add each texture to the output with the appropriate parameters.
 				// TODO: Take into account which UVIndex is being used.
diff --git a/Engine/Source/Runtime/Engine/Private/UnrealEngine.cpp b/Engine/Source/Runtime/Engine/Private/UnrealEngine.cpp
--- a/Engine/Source/Runtime/Engine/Private/UnrealEngine.cpp
+++ b/Engine/Source/Runtime/Engine/Private/UnrealEngine.cpp
@@ -3225,7 +3225,7 @@ bool UEngine::HandleRemoteTextureStatsCommand( const TCHAR* Cmd, FOutputDevice&
 			if( UsedMaterials[ MatIndex ] )
 			{
 				UsedTextures.Reset();
-				UsedMaterials[ MatIndex ]->GetUsedTextures( UsedTextures, EMaterialQualityLevel::Num, false );
+				UsedMaterials[ MatIndex ]->GetUsedTextures( UsedTextures, EMaterialQualityLevel::Num, false, GRHIFeatureLevel, false );
 
 				// Increase usage count for all referenced textures
 				for( int32 TextureIndex=0; TextureIndex<UsedTextures.Num(); TextureIndex++ )
