Fix Win32

diff --git a/Engine/Source/Runtime/Engine/Classes/Kismet/KismetSystemLibrary.h b/Engine/Source/Runtime/Engine/Classes/Kismet/KismetSystemLibrary.h
--- a/Engine/Source/Runtime/Engine/Classes/Kismet/KismetSystemLibrary.h
+++ b/Engine/Source/Runtime/Engine/Classes/Kismet/KismetSystemLibrary.h
@@ -1045,7 +1045,7 @@ class UKismetSystemLibrary : public UBlueprintFunctionLibrary
 
 	/** Draws a debug plane. */
 	UFUNCTION(BlueprintCallable, Category="Rendering|Debug", meta=(HidePin="WorldContextObject", DefaultToSelf="WorldContextObject"))
-	static ENGINE_API void DrawDebugPlane(UObject* WorldContextObject, const FPlane PlaneCoordinates, const FVector Location, float Size, FLinearColor PlaneColor = FLinearColor::White, float Duration=0.f);
+	static ENGINE_API void DrawDebugPlane(UObject* WorldContextObject, const FPlane& PlaneCoordinates, const FVector Location, float Size, FLinearColor PlaneColor = FLinearColor::White, float Duration=0.f);
 
 	/** 
 	 * Flush all persistent debug lines and shapes.
diff --git a/Engine/Source/Runtime/Engine/Private/KismetSystemLibrary.cpp b/Engine/Source/Runtime/Engine/Private/KismetSystemLibrary.cpp
--- a/Engine/Source/Runtime/Engine/Private/KismetSystemLibrary.cpp
+++ b/Engine/Source/Runtime/Engine/Private/KismetSystemLibrary.cpp
@@ -2369,7 +2369,7 @@ void UKismetSystemLibrary::FlushDebugStrings( UObject* WorldContextObject )
 }
 
 /** Draws a debug plane. */
-void UKismetSystemLibrary::DrawDebugPlane(UObject* WorldContextObject, FPlane const P, FVector const Loc, float Size, FLinearColor Color, float LifeTime)
+void UKismetSystemLibrary::DrawDebugPlane(UObject* WorldContextObject, FPlane const& P, FVector const Loc, float Size, FLinearColor Color, float LifeTime)
 {
 	::DrawDebugSolidPlane(GEngine->GetWorldFromContextObject( WorldContextObject ), P, Loc, Size, Color.ToFColor(true), false, LifeTime, SDPG_World);
 }
