Add EngineMinimal.h, modify all templates to use it in their PCHs
Remove PCH from template pch header name (2dsidescroller and twinstick)

diff --git a/Engine/Source/Runtime/Engine/Classes/Components/AudioComponent.h b/Engine/Source/Runtime/Engine/Classes/Components/AudioComponent.h
--- a/Engine/Source/Runtime/Engine/Classes/Components/AudioComponent.h
+++ b/Engine/Source/Runtime/Engine/Classes/Components/AudioComponent.h
@@ -3,7 +3,6 @@
 
 #pragma once
 #include "Audio.h"
-#include "Sound/ReverbVolume.h"
 #include "Sound/SoundAttenuation.h"
 #include "Engine/EngineTypes.h"
 
diff --git a/Engine/Source/Runtime/Engine/Public/EngineMinimal.h b/Engine/Source/Runtime/Engine/Public/EngineMinimal.h
new file mode 100644
--- /dev/null
+++ b/Engine/Source/Runtime/Engine/Public/EngineMinimal.h
@@ -0,0 +1,52 @@
+// Copyright 1998-2014 Epic Games, Inc. All Rights Reserved.
+
+/*=============================================================================
+	EngineMinimal.h: Commonly used include for developing projects with UE4
+=============================================================================*/
+
+#pragma once
+
+// UObject core
+#include "CoreUObject.h"
+
+// Actor based classes
+#include "GameFramework/Actor.h"
+#include "GameFramework/Pawn.h"
+#include "GameFramework/Character.h"
+#include "GameFramework/Controller.h"
+#include "GameFramework/PlayerController.h"
+#include "GameFramework/GameMode.h"
+
+// ActorComponent based classes
+#include "Components/ActorComponent.h"
+#include "Components/SceneComponent.h"
+#include "Components/PrimitiveComponent.h"
+#include "Components/MeshComponent.h"
+#include "Components/StaticMeshComponent.h"
+#include "Components/SkeletalMeshComponent.h"
+#include "Components/CapsuleComponent.h"
+#include "Components/SphereComponent.h"
+#include "Components/BoxComponent.h"
+#include "Components/InputComponent.h"
+#include "Components/AudioComponent.h"
+#include "Camera/CameraComponent.h"
+#include "GameFramework/SpringArmComponent.h"
+#include "GameFramework/CharacterMovementComponent.h"
+#include "Particles/ParticleSystemComponent.h"
+#include "Components/LightComponent.h"
+#include "Components/PointLightComponent.h"
+#include "Components/SpotLightComponent.h"
+
+// Other
+#include "Engine/World.h"
+#include "Engine/StaticMesh.h"
+#include "Engine/SkeletalMesh.h"
+#include "Engine/Texture2D.h"
+#include "Engine/CollisionProfile.h"
+#include "Kismet/GameplayStatics.h"
+#include "Materials/Material.h"
+#include "Materials/MaterialInstance.h"
+#include "Materials/MaterialInstanceDynamic.h"
+#include "Sound/SoundBase.h"
+#include "Sound/SoundWave.h"
+
