Protect against ChildActor not having a RootComponent

diff --git a/Engine/Source/Runtime/Engine/Private/ChildActorComponent.cpp b/Engine/Source/Runtime/Engine/Private/ChildActorComponent.cpp
--- a/Engine/Source/Runtime/Engine/Private/ChildActorComponent.cpp
+++ b/Engine/Source/Runtime/Engine/Private/ChildActorComponent.cpp
@@ -44,18 +44,21 @@ public:
 		if (Component->ChildActor)
 		{
 			USceneComponent* ChildRootComponent = Component->ChildActor->GetRootComponent();
-			for (USceneComponent* AttachedComponent : ChildRootComponent->AttachChildren)
+			if (ChildRootComponent)
 			{
-				if (AttachedComponent)
+				for (USceneComponent* AttachedComponent : ChildRootComponent->AttachChildren)
 				{
-					AActor* AttachedActor = AttachedComponent->GetOwner();
-					if (AttachedActor != Component->ChildActor)
+					if (AttachedComponent)
 					{
-						FAttachedActorInfo Info;
-						Info.Actor = AttachedActor;
-						Info.SocketName = AttachedComponent->AttachSocketName;
-						Info.RelativeTransform = AttachedComponent->GetRelativeTransform();
-						AttachedActors.Add(Info);
+						AActor* AttachedActor = AttachedComponent->GetOwner();
+						if (AttachedActor != Component->ChildActor)
+						{
+							FAttachedActorInfo Info;
+							Info.Actor = AttachedActor;
+							Info.SocketName = AttachedComponent->AttachSocketName;
+							Info.RelativeTransform = AttachedComponent->GetRelativeTransform();
+							AttachedActors.Add(Info);
+						}
 					}
 				}
 			}
