- Added IOS hardware survey for analytics [ttp 342243]

diff --git a/Engine/Source/Runtime/Core/Private/IOS/IOSPlatformSurvey.cpp b/Engine/Source/Runtime/Core/Private/IOS/IOSPlatformSurvey.cpp
new file mode 100644
--- /dev/null
+++ b/Engine/Source/Runtime/Core/Private/IOS/IOSPlatformSurvey.cpp
@@ -0,0 +1,38 @@
+// Copyright 1998-2014 Epic Games, Inc. All Rights Reserved.
+
+/*=============================================================================
+	IOSPlatformSurvey.cpp: HardwareSurvey implementation
+=============================================================================*/
+#include "Core.h"
+#include "IOSAppDelegate.h"
+
+
+bool FIOSPlatformSurvey::GetSurveyResults(FHardwareSurveyResults& OutResults, bool bWait)
+{
+	FCString::Strcpy(OutResults.Platform, FPlatformMisc::GetDefaultDeviceProfileName());
+	FCString::Strcpy(OutResults.OSVersion, *FString([[UIDevice currentDevice] systemVersion]));
+#if PLATFORM_64BITS
+	OutResults.OSBits = 64;
+#else
+	OutResults.OSBits = 32;
+#endif
+	FCString::Strcpy(OutResults.OSLanguage, *FString([[NSLocale preferredLanguages] objectAtIndex:0]));
+	FCString::Strcpy(OutResults.MultimediaAPI, FPlatformMisc::HasPlatformFeature(TEXT("Metal")) ? TEXT("Metal") : TEXT("ES2"));
+	OutResults.CPUCount = FPlatformMisc::NumberOfCores();
+
+	// display 0 is max size
+	CGRect MainFrame = [[UIScreen mainScreen] bounds];
+	if ([IOSAppDelegate GetDelegate].OSVersion < 8.0f && ![IOSAppDelegate GetDelegate].bDeviceInPortraitMode)
+	{
+		Swap(MainFrame.size.width, MainFrame.size.height);
+	}
+	float Scale = [[UIScreen mainScreen] scale];
+	OutResults.Displays[0].CurrentModeWidth = MainFrame.size.width * Scale;
+	OutResults.Displays[0].CurrentModeHeight = MainFrame.size.height * Scale;
+
+	// display 1 is current size
+	FPlatformRect ScreenSize = FIOSWindow::GetScreenRect();
+	OutResults.Displays[1].CurrentModeWidth = ScreenSize.Right - ScreenSize.Left;
+	OutResults.Displays[1].CurrentModeHeight = ScreenSize.Bottom - ScreenSize.Top;
+	return true;
+}
diff --git a/Engine/Source/Runtime/Core/Public/IOS/IOSPlatformSurvey.h b/Engine/Source/Runtime/Core/Public/IOS/IOSPlatformSurvey.h
--- a/Engine/Source/Runtime/Core/Public/IOS/IOSPlatformSurvey.h
+++ b/Engine/Source/Runtime/Core/Public/IOS/IOSPlatformSurvey.h
@@ -12,7 +12,7 @@
 **/
 struct FIOSPlatformSurvey : public FGenericPlatformSurvey
 {
-	// default implementation for now
+	static bool GetSurveyResults(FHardwareSurveyResults& OutResults, bool bWait = false);
 };
 
 typedef FIOSPlatformSurvey FPlatformSurvey;
\ No newline at end of file
diff --git a/Engine/Source/Runtime/Engine/Private/UnrealEngine.cpp b/Engine/Source/Runtime/Engine/Private/UnrealEngine.cpp
--- a/Engine/Source/Runtime/Engine/Private/UnrealEngine.cpp
+++ b/Engine/Source/Runtime/Engine/Private/UnrealEngine.cpp
@@ -5321,13 +5321,31 @@ void UEngine::TickHardwareSurvey()
 
 bool UEngine::IsHardwareSurveyRequired()
 {
-#if PLATFORM_DESKTOP
 	// Analytics must have been initialized FIRST.
 	if (!FEngineAnalytics::IsAvailable())
 	{
 		return false;
 	}
 
+#if PLATFORM_IOS
+	// look up to see if how long ago the last survey was, if ever
+	NSUserDefaults* UserSettings = [NSUserDefaults standardUserDefaults];
+	NSDate* SurveyDateTime = [UserSettings objectForKey:@"HardwareSurveyDateTime"];
+
+	// if we never did one, do it now
+	if (SurveyDateTime == nil)
+	{
+		return true;
+	}
+
+	// how long has it been (negate since the past is negative)
+	NSTimeInterval SecondsBeforeNow = -[SurveyDateTime timeIntervalSinceNow];
+	
+	// survey again after a 30 days (TimeDiff is in seconds)
+	return SecondsBeforeNow > (30.0 * 24.0 * 60.0 * 60.0);
+
+#elif PLATFORM_DESKTOP
+
 	bool bSurveyDone = false;
 	GConfig->GetBool(TEXT("Engine.HardwareSurvey"), TEXT("bHardwareSurveyDone"), bSurveyDone, GEditorGameAgnosticIni);
 
@@ -5467,7 +5485,30 @@ FString UEngine::HardwareSurveyGetResolutionClass(uint32 LargestDisplayHeight)
 
 void UEngine::OnHardwareSurveyComplete(const FHardwareSurveyResults& SurveyResults)
 {
-#if PLATFORM_DESKTOP
+#if PLATFORM_IOS
+	if (FEngineAnalytics::IsAvailable())
+	{
+		// mark now as last survey time
+		NSUserDefaults* UserSettings = [NSUserDefaults standardUserDefaults];
+		[UserSettings setObject:[NSDate date] forKey:@"HardwareSurveyDateTime"];
+
+		TArray<FAnalyticsEventAttribute> HardwareStatsAttribs;
+		// copy from what IOSPlatformSurvey has filled out
+		HardwareStatsAttribs.Add(FAnalyticsEventAttribute(TEXT("Model"), SurveyResults.Platform));
+		HardwareStatsAttribs.Add(FAnalyticsEventAttribute(TEXT("OS.Version"), SurveyResults.OSVersion));
+		HardwareStatsAttribs.Add(FAnalyticsEventAttribute(TEXT("OS.Bits"), FString::Printf(TEXT("%d-bit"), SurveyResults.OSBits)));
+		HardwareStatsAttribs.Add(FAnalyticsEventAttribute(TEXT("OS.Language"), SurveyResults.OSLanguage));
+		HardwareStatsAttribs.Add(FAnalyticsEventAttribute(TEXT("RenderingAPI"), SurveyResults.MultimediaAPI));
+		HardwareStatsAttribs.Add(FAnalyticsEventAttribute(TEXT("CPU.Count"), FString::Printf(TEXT("%d"), SurveyResults.CPUCount)));
+		FString DisplayResolution = FString::Printf(TEXT("%dx%d"), SurveyResults.Displays[0].CurrentModeWidth, SurveyResults.Displays[0].CurrentModeHeight);
+		FString ViewResolution = FString::Printf(TEXT("%dx%d"), SurveyResults.Displays[0].CurrentModeWidth, SurveyResults.Displays[0].CurrentModeHeight);
+		HardwareStatsAttribs.Add(FAnalyticsEventAttribute(TEXT("DisplayResolution"), DisplayResolution));
+		HardwareStatsAttribs.Add(FAnalyticsEventAttribute(TEXT("ViewResolution"), ViewResolution));
+	
+		FEngineAnalytics::GetProvider().RecordEvent(TEXT("IOSHardwareStats"), HardwareStatsAttribs);
+	}
+
+#elif PLATFORM_DESKTOP
 	if (GConfig)
 	{
 		GConfig->SetBool(TEXT("Engine.HardwareSurvey"), TEXT("bHardwareSurveyDone"), true, GEditorGameAgnosticIni);
