Attachment replication was not adhering to absolute scale

diff --git a/Engine/Source/Runtime/Engine/Private/Actor.cpp b/Engine/Source/Runtime/Engine/Private/Actor.cpp
--- a/Engine/Source/Runtime/Engine/Private/Actor.cpp
+++ b/Engine/Source/Runtime/Engine/Private/Actor.cpp
@@ -1207,14 +1207,18 @@ void AActor::OnRep_AttachmentReplication()
 			if (ParentComponent)
 			{
 				// Calculate scale before attachment as ComponentToWorld will be modified after AttachTo()
-				FTransform ParentToWorld = ParentComponent->GetSocketTransform(AttachmentReplication.AttachSocket);
-				FTransform RelativeTM = RootComponent->ComponentToWorld.GetRelativeTransform(ParentToWorld);
-				FVector RelativeScale3D = RelativeTM.GetScale3D();
+				FVector NewRelativeScale3D = RootComponent->RelativeScale3D;
+				if (!RootComponent->bAbsoluteScale)
+				{
+					FTransform ParentToWorld = ParentComponent->GetSocketTransform(AttachmentReplication.AttachSocket);
+					FTransform RelativeTM = RootComponent->ComponentToWorld.GetRelativeTransform(ParentToWorld);
+					NewRelativeScale3D = RelativeTM.GetScale3D();
+				}
 
 				RootComponent->AttachTo(ParentComponent, AttachmentReplication.AttachSocket);
 				RootComponent->RelativeLocation = AttachmentReplication.LocationOffset;
 				RootComponent->RelativeRotation = AttachmentReplication.RotationOffset;
-				RootComponent->RelativeScale3D = RelativeScale3D;
+				RootComponent->RelativeScale3D = NewRelativeScale3D;
 
 				RootComponent->UpdateComponentToWorld();
 			}
