Ensure that Actor's in levels that have been unloaded do not get a tick if they were already scheduled.

diff --git a/Engine/Source/Runtime/Engine/Private/Actor.cpp b/Engine/Source/Runtime/Engine/Private/Actor.cpp
--- a/Engine/Source/Runtime/Engine/Private/Actor.cpp
+++ b/Engine/Source/Runtime/Engine/Private/Actor.cpp
@@ -740,7 +740,8 @@ void AActor::TickActor( float DeltaSeconds, ELevelTick TickType, FActorTickFunct
 	const bool bShouldTick = ((TickType!=LEVELTICK_ViewportsOnly) || ShouldTickIfViewportsOnly());
 	if(bShouldTick)
 	{
-		if (!IsPendingKill())
+		// If an Actor has been Destroyed or its level has been unloaded don't execute any queued ticks
+		if (!IsPendingKill() && GetWorld())
 		{
 			Tick(DeltaSeconds);	// perform any tick functions unique to an actor subclass
 		}
@@ -2748,7 +2749,8 @@ void AActor::GetActorBounds(bool bOnlyCollidingComponents, FVector& Origin, FVec
 
 AWorldSettings * AActor::GetWorldSettings() const
 {
-	return GetWorld()->GetWorldSettings();
+	UWorld* World = GetWorld();
+	return (World ? World->GetWorldSettings() : nullptr);
 }
 
 void AActor::PlaySoundOnActor(USoundCue* InSoundCue, float VolumeMultiplier/*=1.f*/, float PitchMultiplier/*=1.f*/)
