Don't change the console state when one of the modifier keys is being held - especially important on OS X where Command+~ is the native window cycling key-equivalent.

diff --git a/Engine/Source/Runtime/Engine/Private/UserInterface/Console.cpp b/Engine/Source/Runtime/Engine/Private/UserInterface/Console.cpp
--- a/Engine/Source/Runtime/Engine/Private/UserInterface/Console.cpp
+++ b/Engine/Source/Runtime/Engine/Private/UserInterface/Console.cpp
@@ -561,7 +561,10 @@ bool UConsole::InputKey_InputLine( int32 ControllerId, FKey Key, EInputEvent Eve
 	}
 
 	// cycle between console states
-	if ( Key==GetDefault<UInputSettings>()->ConsoleKey && Event == IE_Pressed )
+	bool bModifierDown = bCtrl;
+	FModifierKeysState KeyState = FSlateApplication::Get().GetModifierKeys();
+	bModifierDown |= KeyState.IsAltDown() || KeyState.IsCommandDown() || KeyState.IsShiftDown() || KeyState.IsControlDown();
+	if ( Key==GetDefault<UInputSettings>()->ConsoleKey && Event == IE_Pressed && !bModifierDown )
 	{
 		if (ConsoleState == NAME_Typing)
 		{
