Changes to animation evaluation to solve crash with MT on Orion

diff --git a/Engine/Source/Runtime/Engine/Classes/Components/SkeletalMeshComponent.h b/Engine/Source/Runtime/Engine/Classes/Components/SkeletalMeshComponent.h
--- a/Engine/Source/Runtime/Engine/Classes/Components/SkeletalMeshComponent.h
+++ b/Engine/Source/Runtime/Engine/Classes/Components/SkeletalMeshComponent.h
@@ -7,6 +7,44 @@
 #include "PhysicsEngine/PhysicsConstraintComponent.h"
 #include "SkeletalMeshComponent.generated.h"
 
+class UAnimInstance;
+
+struct FAnimationEvaluationContext
+{
+	// The anim instance we are evaluating
+	UAnimInstance* AnimInstance;
+
+	// The SkeletalMesh we are evaluating for
+	USkeletalMesh* SkeletalMesh;
+
+	// Double buffer evaluation data
+	TArray<FTransform> SpaceBases;
+	TArray<FTransform> LocalAtoms;
+	TArray<FActiveVertexAnim> VertexAnims;
+	FVector RootBoneTranslation;
+
+	// Are we performing interpolation this tick
+	bool bDoInterpolation;
+
+	// Are we evaluating this tick
+	bool bDoEvaluation;
+
+	// Are we storing data in cache bones this tick
+	bool bDuplicateToCacheBones;
+
+	FAnimationEvaluationContext()
+	{
+		Clear();
+	}
+
+	void Clear()
+	{
+		AnimInstance = NULL;
+		SkeletalMesh = NULL;
+	}
+
+};
+
 //#if WITH_APEX
 namespace physx
 { 
@@ -672,16 +710,6 @@ public:
 	/** freezing clothing actor now */
 	void FreezeClothSection(bool bFreeze);
 
-	/** Evaluate Anim System **/
-	void EvaluateAnimation( TArray<FTransform>& OutLocalAtoms, TArray<struct FActiveVertexAnim>& OutVertexAnims, FVector& OutRootBoneTranslation ) const;
-
-	/**
-	 * Take the SourceAtoms array (translation vector, rotation quaternion and scale vector) and update the array of component-space bone transformation matrices (DestSpaceBases).
-	 * It will work down hierarchy multiplying the component-space transform of the parent by the relative transform of the child.
-	 * This code also applies any per-bone rotators etc. as part of the composition process
-	 */
-	void FillSpaceBases( const TArray<FTransform>& SourceAtoms, TArray<FTransform>& DestSpaceBases ) const;
-
 	/** 
 	 * Recalculates the RequiredBones array in this SkeletalMeshComponent based on current SkeletalMesh, LOD and PhysicsAsset.
 	 * Is called when bRequiredBonesUpToDate = false
@@ -794,13 +822,15 @@ public:
 	/**
 	* Runs the animation evaluation for the current pose into the supplied variables
 	*
+	* @param	InSkeletalMesh			The skeletal mesh we are animating
+	* @param	InAnimInstance			The anim instance we are evaluating
 	* @param	OutSpaceBases			Component space bone transforms
 	* @param	OutLocalAtoms			Local space bone transforms
 	* @param	OutVertexAnims			Active vertex animations
 	* @param	OutRootBoneTranslation	Calculated root bone translation
 	*/
-	void PerformAnimationEvaluation( TArray<FTransform>& OutSpaceBases, TArray<FTransform>& OutLocalAtoms, TArray<FActiveVertexAnim>& OutVertexAnims, FVector& OutRootBoneTranslation ) const;
-	void PostAnimEvaluation( bool bDoInterpolation, bool bDuplicateToCacheBones );
+	void PerformAnimationEvaluation(const USkeletalMesh* InSkeletalMesh, UAnimInstance* InAnimInstance, TArray<FTransform>& OutSpaceBases, TArray<FTransform>& OutLocalAtoms, TArray<FActiveVertexAnim>& OutVertexAnims, FVector& OutRootBoneTranslation) const;
+	void PostAnimEvaluation( FAnimationEvaluationContext& EvaluationContext );
 
 	/**
 	 * Blend of Physics Bones with PhysicsWeight and Animated Bones with (1-PhysicsWeight)
@@ -1046,30 +1076,42 @@ protected:
 	bool NeedToSpawnAnimScriptInstance(bool bForceInit) const;
 	
 private:
+	/** Evaluate Anim System **/
+	void EvaluateAnimation(const USkeletalMesh* InSkeletalMesh, UAnimInstance* InAnimInstance, TArray<FTransform>& OutLocalAtoms, TArray<struct FActiveVertexAnim>& OutVertexAnims, FVector& OutRootBoneTranslation) const;
+
+	/**
+	* Take the SourceAtoms array (translation vector, rotation quaternion and scale vector) and update the array of component-space bone transformation matrices (DestSpaceBases).
+	* It will work down hierarchy multiplying the component-space transform of the parent by the relative transform of the child.
+	* This code also applies any per-bone rotators etc. as part of the composition process
+	*/
+	void FillSpaceBases(const USkeletalMesh* InSkeletalMesh, const TArray<FTransform>& SourceAtoms, TArray<FTransform>& DestSpaceBases) const;
+
 	void RenderAxisGizmo(const FTransform& Transform, class UCanvas* Canvas) const;
 
 	bool ShouldBlendPhysicsBones();	
 	void ClearAnimScriptInstance();
 
-	//Handle parallel evaluation of animation
-	TArray<FTransform> PTSpaceBases;
-	TArray<FTransform> PTLocalAtoms;
-	TArray<FActiveVertexAnim> PTVertexAnims;
-	FVector PTRootBoneTranslation;
-	bool bPTDoInterpolation;
-	bool bPTDuplicateToCacheBones;
+	//Data for parallel evaluation of animation
+	FAnimationEvaluationContext AnimEvaluationContext;
 
 public:
 	// Parallel evaluation wrappers
-	void ParallelAnimationEvaluation() { PerformAnimationEvaluation(PTSpaceBases, PTLocalAtoms, PTVertexAnims, PTRootBoneTranslation); }
+	void ParallelAnimationEvaluation() { PerformAnimationEvaluation(AnimEvaluationContext.SkeletalMesh, AnimEvaluationContext.AnimInstance, AnimEvaluationContext.SpaceBases, AnimEvaluationContext.LocalAtoms, AnimEvaluationContext.VertexAnims, AnimEvaluationContext.RootBoneTranslation); }
 	void CompleteParallelAnimationEvaluation()
 	{
-		Exchange(PTSpaceBases, bPTDoInterpolation ? CachedSpaceBases : SpaceBases);
-		Exchange(PTLocalAtoms, bPTDoInterpolation ? CachedLocalAtoms : LocalAtoms);
-		Exchange(PTVertexAnims, ActiveVertexAnims);
-		Exchange(PTRootBoneTranslation, RootBoneTranslation);
-
-		PostAnimEvaluation(bPTDoInterpolation, bPTDuplicateToCacheBones);
+		if (AnimEvaluationContext.AnimInstance == AnimScriptInstance && AnimEvaluationContext.SkeletalMesh == SkeletalMesh)
+		{
+			Exchange(AnimEvaluationContext.SpaceBases, AnimEvaluationContext.bDoInterpolation ? CachedSpaceBases : SpaceBases);
+			Exchange(AnimEvaluationContext.LocalAtoms, AnimEvaluationContext.bDoInterpolation ? CachedLocalAtoms : LocalAtoms);
+			Exchange(AnimEvaluationContext.VertexAnims, ActiveVertexAnims);
+			Exchange(AnimEvaluationContext.RootBoneTranslation, RootBoneTranslation);
+
+			PostAnimEvaluation(AnimEvaluationContext);
+		}
+		else
+		{
+			AnimEvaluationContext.Clear();
+		}
 	}
 
 	friend class FSkeletalMeshComponentDetails;
diff --git a/Engine/Source/Runtime/Engine/Classes/Components/SkinnedMeshComponent.h b/Engine/Source/Runtime/Engine/Classes/Components/SkinnedMeshComponent.h
--- a/Engine/Source/Runtime/Engine/Classes/Components/SkinnedMeshComponent.h
+++ b/Engine/Source/Runtime/Engine/Classes/Components/SkinnedMeshComponent.h
@@ -527,11 +527,6 @@ public:
 
 	/** Vertex anim */
 
-	/** 
-	 * Combine CurveKeys (that reference morph targets by name) and ActiveAnims (that reference vertex anims by reference) into the ActiveVertexAnims array.
-	 */
-	TArray<struct FActiveVertexAnim> UpdateActiveVertexAnims(const TMap<FName, float>& MorphCurveAnims, const TArray<FActiveVertexAnim>& ActiveAnims) const;
-
 	/**
 	 * Checks/updates material usage on proxy based on current morph target usage
 	 */
@@ -539,6 +534,11 @@ public:
 
 
 protected:
+	/**
+	* Combine CurveKeys (that reference morph targets by name) and ActiveAnims (that reference vertex anims by reference) into the ActiveVertexAnims array.
+	*/
+	static TArray<struct FActiveVertexAnim> UpdateActiveVertexAnims(const USkeletalMesh* InSkeletalMesh, const TMap<FName, float>& InMorphCurveAnims, const TArray<FActiveVertexAnim>& InActiveAnims);
+
 	/** 
 	 * Should update transform in Tick
 	 * 
diff --git a/Engine/Source/Runtime/Engine/Private/SkeletalMeshComponent.cpp b/Engine/Source/Runtime/Engine/Private/SkeletalMeshComponent.cpp
--- a/Engine/Source/Runtime/Engine/Private/SkeletalMeshComponent.cpp
+++ b/Engine/Source/Runtime/Engine/Private/SkeletalMeshComponent.cpp
@@ -577,18 +577,18 @@ static void IntersectBoneIndexArrays(TArray<FBoneIndexType>& Output, const TArra
 }
 
 
-void USkeletalMeshComponent::FillSpaceBases(const TArray<FTransform>& SourceAtoms, TArray<FTransform>& DestSpaceBases) const
+void USkeletalMeshComponent::FillSpaceBases(const USkeletalMesh* InSkeletalMesh, const TArray<FTransform>& SourceAtoms, TArray<FTransform>& DestSpaceBases) const
 {
 	SCOPE_CYCLE_COUNTER(STAT_SkelComposeTime);
 
-	if( !SkeletalMesh )
+	if( !InSkeletalMesh )
 	{
 		return;
 	}
 
 	// right now all this does is populate DestSpaceBases
-	check(SkeletalMesh->RefSkeleton.GetNum() == SourceAtoms.Num());
-	check( SkeletalMesh->RefSkeleton.GetNum() == DestSpaceBases.Num() );
+	check( InSkeletalMesh->RefSkeleton.GetNum() == SourceAtoms.Num());
+	check( InSkeletalMesh->RefSkeleton.GetNum() == DestSpaceBases.Num());
 
 	const int32 NumBones = SourceAtoms.Num();
 
@@ -622,7 +622,7 @@ void USkeletalMeshComponent::FillSpaceBases(const TArray<FTransform>& SourceAtom
 		BoneProcessed[BoneIndex] = 1;
 #endif
 		// For all bones below the root, final component-space transform is relative transform * component-space transform of parent.
-		const int32 ParentIndex = SkeletalMesh->RefSkeleton.GetParentIndex(BoneIndex);
+		const int32 ParentIndex = InSkeletalMesh->RefSkeleton.GetParentIndex(BoneIndex);
 		FPlatformMisc::Prefetch(SpaceBasesData + ParentIndex);
 
 #if (UE_BUILD_DEBUG || UE_BUILD_DEVELOPMENT)
@@ -802,29 +802,29 @@ void USkeletalMeshComponent::RecalcRequiredBones(int32 LODIndex)
 	CachedSpaceBases.Empty();
 }
 
-void USkeletalMeshComponent::EvaluateAnimation(TArray<FTransform>& OutLocalAtoms, TArray<FActiveVertexAnim>& OutVertexAnims, FVector& OutRootBoneTranslation) const
+void USkeletalMeshComponent::EvaluateAnimation(const USkeletalMesh* InSkeletalMesh, UAnimInstance* InAnimInstance, TArray<FTransform>& OutLocalAtoms, TArray<FActiveVertexAnim>& OutVertexAnims, FVector& OutRootBoneTranslation) const
 {
 	SCOPE_CYCLE_COUNTER(STAT_AnimBlendTime);
 
-	if( !SkeletalMesh )
+	if( !InSkeletalMesh )
 	{
 		return;
 	}
 
 	// We can only evaluate animation if RequiredBones is properly setup for the right mesh!
-	if( SkeletalMesh->Skeleton && AnimScriptInstance 
+	if( InSkeletalMesh->Skeleton && InAnimInstance
 		&& ensure(bRequiredBonesUpToDate)
-		&& AnimScriptInstance->RequiredBones.IsValid() 
-		&& (AnimScriptInstance->RequiredBones.GetAsset() == SkeletalMesh) )
+		&& InAnimInstance->RequiredBones.IsValid()
+		&& (InAnimInstance->RequiredBones.GetAsset() == InSkeletalMesh) )
 	{
 		if( !bForceRefpose )
 		{
 			// Create an evaluation context
-			FPoseContext EvaluationContext(AnimScriptInstance);
+			FPoseContext EvaluationContext(InAnimInstance);
 			EvaluationContext.ResetToRefPose();
 			
 			// Run the anim blueprint
-			AnimScriptInstance->EvaluateAnimation(EvaluationContext);
+			InAnimInstance->EvaluateAnimation(EvaluationContext);
 
 			// can we avoid that copy?
 			if( EvaluationContext.Pose.Bones.Num() > 0 )
@@ -836,30 +836,30 @@ void USkeletalMeshComponent::EvaluateAnimation(TArray<FTransform>& OutLocalAtoms
 			}
 			else
 			{
-				FAnimationRuntime::FillWithRefPose(OutLocalAtoms, AnimScriptInstance->RequiredBones);
+				FAnimationRuntime::FillWithRefPose(OutLocalAtoms, InAnimInstance->RequiredBones);
 			}
 		}
 		else
 		{
-			FAnimationRuntime::FillWithRefPose(OutLocalAtoms, AnimScriptInstance->RequiredBones);
+			FAnimationRuntime::FillWithRefPose(OutLocalAtoms, InAnimInstance->RequiredBones);
 		}
 
-		OutVertexAnims = UpdateActiveVertexAnims(AnimScriptInstance->MorphTargetCurves, AnimScriptInstance->VertexAnims);
+		OutVertexAnims = UpdateActiveVertexAnims(InSkeletalMesh, InAnimInstance->MorphTargetCurves, InAnimInstance->VertexAnims);
 	}
 	else
 	{
-		OutLocalAtoms = SkeletalMesh->RefSkeleton.GetRefBonePose();
+		OutLocalAtoms = InSkeletalMesh->RefSkeleton.GetRefBonePose();
 
 		// if it's only morph, there is no reason to blend
 		if ( MorphTargetCurves.Num() > 0 )
 		{
 			TArray<struct FActiveVertexAnim> EmptyVertexAnims;
-			OutVertexAnims = UpdateActiveVertexAnims(MorphTargetCurves, EmptyVertexAnims);
+			OutVertexAnims = UpdateActiveVertexAnims(InSkeletalMesh, MorphTargetCurves, EmptyVertexAnims);
 		}
 	}
 
 	// Remember the root bone's translation so we can move the bounds.
-	OutRootBoneTranslation = OutLocalAtoms[0].GetTranslation() - SkeletalMesh->RefSkeleton.GetRefBonePose()[0].GetTranslation();
+	OutRootBoneTranslation = OutLocalAtoms[0].GetTranslation() - InSkeletalMesh->RefSkeleton.GetRefBonePose()[0].GetTranslation();
 }
 
 void USkeletalMeshComponent::UpdateSlaveComponent()
@@ -872,30 +872,27 @@ void USkeletalMeshComponent::UpdateSlaveComponent()
 
 		if ( MasterSMC->AnimScriptInstance )
 		{
-			ActiveVertexAnims = UpdateActiveVertexAnims(MasterSMC->AnimScriptInstance->MorphTargetCurves, MasterSMC->AnimScriptInstance->VertexAnims);
+			ActiveVertexAnims = UpdateActiveVertexAnims(SkeletalMesh, MasterSMC->AnimScriptInstance->MorphTargetCurves, MasterSMC->AnimScriptInstance->VertexAnims);
 		}
 	}
 
 	Super::UpdateSlaveComponent();
 }
 
-void USkeletalMeshComponent::PerformAnimationEvaluation(TArray<FTransform>& OutSpaceBases, TArray<FTransform>& OutLocalAtoms, TArray<FActiveVertexAnim>& OutVertexAnims, FVector& OutRootBoneTranslation) const
+void USkeletalMeshComponent::PerformAnimationEvaluation(const USkeletalMesh* InSkeletalMesh, UAnimInstance* InAnimInstance, TArray<FTransform>& OutSpaceBases, TArray<FTransform>& OutLocalAtoms, TArray<FActiveVertexAnim>& OutVertexAnims, FVector& OutRootBoneTranslation) const
 {
 	SCOPE_CYCLE_COUNTER(STAT_PerformAnimEvaluation);
 	// Can't do anything without a SkeletalMesh
 	// Do nothing more if no bones in skeleton.
-	if (!SkeletalMesh || OutSpaceBases.Num() == 0)
+	if (!InSkeletalMesh || OutSpaceBases.Num() == 0)
 	{
 		return;
 	}
-	AActor * Owner = GetOwner();
-	const FAnimUpdateRateParameters  & UpdateRateParams = Owner ? Owner->AnimUpdateRateParams : FAnimUpdateRateParameters();
-	UE_LOG(LogAnimation, Verbose, TEXT("RefreshBoneTransforms(%s)"), *GetNameSafe(Owner));
 
 	// evaluate pure animations, and fill up LocalAtoms
-	EvaluateAnimation(OutLocalAtoms, OutVertexAnims, OutRootBoneTranslation);
+	EvaluateAnimation(InSkeletalMesh, InAnimInstance, OutLocalAtoms, OutVertexAnims, OutRootBoneTranslation);
 	// Fill SpaceBases from LocalAtoms
-	FillSpaceBases(OutLocalAtoms, OutSpaceBases);
+	FillSpaceBases(InSkeletalMesh, OutLocalAtoms, OutSpaceBases);
 }
 
 void USkeletalMeshComponent::RefreshBoneTransforms(FActorComponentTickFunction* TickFunction)
@@ -916,6 +913,9 @@ void USkeletalMeshComponent::RefreshBoneTransforms(FActorComponentTickFunction*
 		RecalcRequiredBones(PredictedLODLevel);
 	}
 
+	AnimEvaluationContext.SkeletalMesh = SkeletalMesh;
+	AnimEvaluationContext.AnimInstance = AnimScriptInstance;
+
 	//Handle update rate optimization setup
 	const FAnimUpdateRateParameters  & UpdateRateParams = Owner ? Owner->AnimUpdateRateParams : FAnimUpdateRateParameters();
 	
@@ -926,10 +926,10 @@ void USkeletalMeshComponent::RefreshBoneTransforms(FActorComponentTickFunction*
 									  || (LocalAtoms.Num() != CachedLocalAtoms.Num())
 									  || (SpaceBases.Num() != CachedSpaceBases.Num()) );
 
-	const bool bDoEvaluation = !bDoUpdateRateOptimization || bInvalidCachedBones || !UpdateRateParams.ShouldSkipEvaluation();
+	AnimEvaluationContext.bDoEvaluation = !bDoUpdateRateOptimization || bInvalidCachedBones || !UpdateRateParams.ShouldSkipEvaluation();
 	
-	bPTDoInterpolation = bDoUpdateRateOptimization && !bInvalidCachedBones && UpdateRateParams.ShouldInterpolateSkippedFrames();
-	bPTDuplicateToCacheBones = bInvalidCachedBones || (bDoUpdateRateOptimization && bDoEvaluation && !bPTDoInterpolation);
+	AnimEvaluationContext.bDoInterpolation = bDoUpdateRateOptimization && !bInvalidCachedBones && UpdateRateParams.ShouldInterpolateSkippedFrames();
+	AnimEvaluationContext.bDuplicateToCacheBones = bInvalidCachedBones || (bDoUpdateRateOptimization && AnimEvaluationContext.bDoEvaluation && !AnimEvaluationContext.bDoInterpolation);
 
 	if (!bDoUpdateRateOptimization)
 	{
@@ -939,14 +939,14 @@ void USkeletalMeshComponent::RefreshBoneTransforms(FActorComponentTickFunction*
 	}
 
 	const bool bDoPAE = !!CVarUseParallelAnimationEvaluation.GetValueOnGameThread();
-	if (bDoEvaluation && TickFunction && bDoPAE)
+	if (AnimEvaluationContext.bDoEvaluation && TickFunction && bDoPAE)
 	{
-		if (SkeletalMesh->RefSkeleton.GetNum() != PTLocalAtoms.Num())
+		if (SkeletalMesh->RefSkeleton.GetNum() != AnimEvaluationContext.LocalAtoms.Num())
 		{
 			// Initialize Parallel Task arrays
-			PTLocalAtoms = LocalAtoms;
-			PTSpaceBases = SpaceBases;
-			PTVertexAnims = ActiveVertexAnims;
+			AnimEvaluationContext.LocalAtoms = LocalAtoms;
+			AnimEvaluationContext.SpaceBases = SpaceBases;
+			AnimEvaluationContext.VertexAnims = ActiveVertexAnims;
 		}
 
 		// start parallel work
@@ -961,38 +961,40 @@ void USkeletalMeshComponent::RefreshBoneTransforms(FActorComponentTickFunction*
 	}
 	else
 	{
-		if (bDoEvaluation)
+		if (AnimEvaluationContext.bDoEvaluation)
 		{
-			if (bPTDoInterpolation)
+			if (AnimEvaluationContext.bDoInterpolation)
 			{
-				PerformAnimationEvaluation(CachedSpaceBases, CachedLocalAtoms, ActiveVertexAnims, RootBoneTranslation);
+				PerformAnimationEvaluation(SkeletalMesh, AnimScriptInstance, CachedSpaceBases, CachedLocalAtoms, ActiveVertexAnims, RootBoneTranslation);
 			}
 			else
 			{
-				PerformAnimationEvaluation(SpaceBases, LocalAtoms, ActiveVertexAnims, RootBoneTranslation);
+				PerformAnimationEvaluation(SkeletalMesh, AnimScriptInstance, SpaceBases, LocalAtoms, ActiveVertexAnims, RootBoneTranslation);
 			}
 		}
-		else if (!bPTDoInterpolation)
+		else if (!AnimEvaluationContext.bDoInterpolation)
 		{
 			LocalAtoms = CachedLocalAtoms;
 			SpaceBases = CachedSpaceBases;
 		}
 
-		PostAnimEvaluation(bPTDoInterpolation, bPTDuplicateToCacheBones);
+		PostAnimEvaluation(AnimEvaluationContext);
 	}
 
 }
 
-void USkeletalMeshComponent::PostAnimEvaluation(bool bDoInterpolation, bool bDuplicateToCacheBones)
+void USkeletalMeshComponent::PostAnimEvaluation(FAnimationEvaluationContext& EvaluationContext)
 {
+	AnimEvaluationContext.Clear();
+
 	SCOPE_CYCLE_COUNTER(STAT_PostAnimEvaluation);
-	if( bDuplicateToCacheBones )
+	if (EvaluationContext.bDuplicateToCacheBones)
 	{
 		CachedSpaceBases = SpaceBases;
 		CachedLocalAtoms = LocalAtoms;
 	}
 
-	if( bDoInterpolation )
+	if (EvaluationContext.bDoInterpolation)
 	{
 		SCOPE_CYCLE_COUNTER(STAT_InterpolateSkippedFrames);
 		AActor * Owner = GetOwner();
diff --git a/Engine/Source/Runtime/Engine/Private/SkinnedMeshComponent.cpp b/Engine/Source/Runtime/Engine/Private/SkinnedMeshComponent.cpp
--- a/Engine/Source/Runtime/Engine/Private/SkinnedMeshComponent.cpp
+++ b/Engine/Source/Runtime/Engine/Private/SkinnedMeshComponent.cpp
@@ -1668,7 +1668,7 @@ int32 FindVertexAnim(const TArray<FActiveVertexAnim>& ActiveAnims, UVertexAnimBa
 	return INDEX_NONE;
 }
 
-TArray<FActiveVertexAnim> USkinnedMeshComponent::UpdateActiveVertexAnims(const TMap<FName, float>& MorphCurveAnims, const TArray<FActiveVertexAnim>& ActiveAnims) const
+TArray<FActiveVertexAnim> USkinnedMeshComponent::UpdateActiveVertexAnims(const USkeletalMesh* InSkeletalMesh, const TMap<FName, float>& MorphCurveAnims, const TArray<FActiveVertexAnim>& ActiveAnims)
 {
 	TArray<struct FActiveVertexAnim> OutVertexAnims;
 
@@ -1679,7 +1679,7 @@ TArray<FActiveVertexAnim> USkinnedMeshComponent::UpdateActiveVertexAnims(const T
 		// Check it has valid weight, and works on this SkeletalMesh
 		if(	ActiveAnim.Weight > MinVertexAnimBlendWeight &&
 			ActiveAnim.VertAnim != NULL &&
-			ActiveAnim.VertAnim->BaseSkelMesh == SkeletalMesh )
+			ActiveAnim.VertAnim->BaseSkelMesh == InSkeletalMesh)
 		{
 			OutVertexAnims.Add(ActiveAnim);
 		}
@@ -1696,7 +1696,7 @@ TArray<FActiveVertexAnim> USkinnedMeshComponent::UpdateActiveVertexAnims(const T
 		if(Weight > MinVertexAnimBlendWeight)
 		{
 			// Find morph reference
-			UMorphTarget* Target = FindMorphTarget(CurveName);
+			UMorphTarget* Target = InSkeletalMesh ? InSkeletalMesh->FindMorphTarget(CurveName) : NULL;
 			if(Target != NULL)				
 			{
 				// See if this morph target already has an entry
