Unregister Actor Tick functions immediately upon EndPlay being received.  Otherwise other tick functions that have this Actor as a prerequisite can cause the Actor to be ticked incorrectly.

diff --git a/Engine/Source/Runtime/Engine/Private/Actor.cpp b/Engine/Source/Runtime/Engine/Private/Actor.cpp
--- a/Engine/Source/Runtime/Engine/Private/Actor.cpp
+++ b/Engine/Source/Runtime/Engine/Private/Actor.cpp
@@ -673,11 +673,11 @@ void AActor::RegisterAllActorTickFunctions(bool bRegister, bool bDoComponents)
 
 			for (int32 Index = 0; Index < Components.Num(); Index++)
 			{
-					Components[Index]->RegisterAllComponentTickFunctions(bRegister);
-				}
+				Components[Index]->RegisterAllComponentTickFunctions(bRegister);
 			}
 		}
 	}
+}
 
 void AActor::SetActorTickEnabled(bool bEnabled)
 {
@@ -1211,8 +1211,8 @@ void AActor::OnRep_AttachmentReplication()
 				FVector NewRelativeScale3D = RootComponent->RelativeScale3D;
 				if (!RootComponent->bAbsoluteScale)
 				{
-					FTransform ParentToWorld = ParentComponent->GetSocketTransform(AttachmentReplication.AttachSocket);
-					FTransform RelativeTM = RootComponent->ComponentToWorld.GetRelativeTransform(ParentToWorld);
+				FTransform ParentToWorld = ParentComponent->GetSocketTransform(AttachmentReplication.AttachSocket);
+				FTransform RelativeTM = RootComponent->ComponentToWorld.GetRelativeTransform(ParentToWorld);
 					NewRelativeScale3D = RelativeTM.GetScale3D();
 				}
 
@@ -1566,6 +1566,8 @@ void AActor::EndPlay(const EEndPlayReason::Type EndPlayReason)
 	ReceiveEndPlay(EndPlayReason);
 	OnEndPlay.Broadcast(EndPlayReason);
 
+	RegisterAllActorTickFunctions(false, true);
+
 	// Behaviors specific to an actor being unloaded due to a streaming level removal
 	if (EndPlayReason == EEndPlayReason::RemovedFromWorld)
 	{
@@ -3101,7 +3103,7 @@ bool AActor::IncrementalRegisterComponents(int32 NumComponentsToRegister)
 	int32 NumRegisteredComponentsThisRun = 0;
 	TArray<UActorComponent*> Components;
 	GetComponents(Components);
-
+	
 	for (int32 CompIdx = 0; CompIdx < Components.Num() && NumRegisteredComponentsThisRun < NumComponentsToRegister; CompIdx++)
 	{
 		UActorComponent* Component = Components[CompIdx];
