Attachment replication now takes into account RelativeScale3D.

diff --git a/Engine/Source/Runtime/Engine/Classes/Engine/EngineTypes.h b/Engine/Source/Runtime/Engine/Classes/Engine/EngineTypes.h
--- a/Engine/Source/Runtime/Engine/Classes/Engine/EngineTypes.h
+++ b/Engine/Source/Runtime/Engine/Classes/Engine/EngineTypes.h
@@ -2209,6 +2209,9 @@ struct FRepAttachment
 	UPROPERTY()
 	FVector_NetQuantize100 LocationOffset;
 
+	UPROPERTY()
+	FVector_NetQuantize100 RelativeScale3D;
+
 	UPROPERTY()
 	FRotator RotationOffset;
 
@@ -2221,6 +2224,7 @@ struct FRepAttachment
 	FRepAttachment()
 		: AttachParent(NULL)
 		, LocationOffset(ForceInit)
+		, RelativeScale3D(ForceInit)
 		, RotationOffset(ForceInit)
 		, AttachSocket(NAME_None)
 		, AttachComponent(NULL)
diff --git a/Engine/Source/Runtime/Engine/Private/Actor.cpp b/Engine/Source/Runtime/Engine/Private/Actor.cpp
--- a/Engine/Source/Runtime/Engine/Private/Actor.cpp
+++ b/Engine/Source/Runtime/Engine/Private/Actor.cpp
@@ -1187,6 +1187,7 @@ void AActor::AttachRootComponentTo(USceneComponent* InParent, FName InSocketName
 		AttachmentReplication.AttachParent = InParent->GetAttachmentRootActor();
 		AttachmentReplication.LocationOffset = RootComponent->RelativeLocation;
 		AttachmentReplication.RotationOffset = RootComponent->RelativeRotation;
+		AttachmentReplication.RelativeScale3D = RootComponent->RelativeScale3D;
 		AttachmentReplication.AttachSocket = InSocketName;
 		AttachmentReplication.AttachComponent = InParent;
 	}
@@ -1207,19 +1208,10 @@ void AActor::OnRep_AttachmentReplication()
 
 			if (ParentComponent)
 			{
-				// Calculate scale before attachment as ComponentToWorld will be modified after AttachTo()
-				FVector NewRelativeScale3D = RootComponent->RelativeScale3D;
-				if (!RootComponent->bAbsoluteScale)
-				{
-				FTransform ParentToWorld = ParentComponent->GetSocketTransform(AttachmentReplication.AttachSocket);
-				FTransform RelativeTM = RootComponent->ComponentToWorld.GetRelativeTransform(ParentToWorld);
-					NewRelativeScale3D = RelativeTM.GetScale3D();
-				}
-
 				RootComponent->AttachTo(ParentComponent, AttachmentReplication.AttachSocket);
 				RootComponent->RelativeLocation = AttachmentReplication.LocationOffset;
 				RootComponent->RelativeRotation = AttachmentReplication.RotationOffset;
-				RootComponent->RelativeScale3D = NewRelativeScale3D;
+				RootComponent->RelativeScale3D = AttachmentReplication.RelativeScale3D;
 
 				RootComponent->UpdateComponentToWorld();
 			}
@@ -1243,6 +1235,7 @@ void AActor::AttachRootComponentToActor(AActor* InParentActor, FName InSocketNam
 			AttachmentReplication.AttachParent = InParentActor;
 			AttachmentReplication.LocationOffset = RootComponent->RelativeLocation;
 			AttachmentReplication.RotationOffset = RootComponent->RelativeRotation;
+			AttachmentReplication.RelativeScale3D = RootComponent->RelativeScale3D;
 			AttachmentReplication.AttachSocket = InSocketName;
 			AttachmentReplication.AttachComponent = NULL;
 		}
diff --git a/Engine/Source/Runtime/Engine/Private/ActorReplication.cpp b/Engine/Source/Runtime/Engine/Private/ActorReplication.cpp
--- a/Engine/Source/Runtime/Engine/Private/ActorReplication.cpp
+++ b/Engine/Source/Runtime/Engine/Private/ActorReplication.cpp
@@ -180,6 +180,7 @@ void AActor::GatherCurrentMovement()
 			{
 				AttachmentReplication.LocationOffset = RootComponent->RelativeLocation;
 				AttachmentReplication.RotationOffset = RootComponent->RelativeRotation;
+				AttachmentReplication.RelativeScale3D = RootComponent->RelativeScale3D;
 			}
 		}
 		else
