Fix a crash when trying to rebuild missing streaming mips of a texture that is in the process of being garbage collected.

diff --git a/Engine/Source/Runtime/Engine/Private/Texture2D.cpp b/Engine/Source/Runtime/Engine/Private/Texture2D.cpp
--- a/Engine/Source/Runtime/Engine/Private/Texture2D.cpp
+++ b/Engine/Source/Runtime/Engine/Private/Texture2D.cpp
@@ -680,7 +680,14 @@ bool UTexture2D::UpdateStreamingStatus( bool bWaitForMipFading /*= false*/ )
 			{
 				// A streaming request to the derived data cache failed. That means the texture itself needs to be updated.
 				UE_LOG(LogTexture, Warning, TEXT("Failed to stream mip data from the derived data cache for %s. Streaming mips will be recached."), *GetPathName() );
-				ForceRebuildPlatformData();
+
+				// We can't load the source art from a bulk data object if the texture itself is pending kill because the linker will have been detached.
+				// In this case we don't rebuild the data and instead let the streaming request be cancelled. This will let the garbage collector finish
+				// destroying the object.
+				if (!HasAnyFlags(RF_PendingKill | RF_Unreachable))
+				{
+					ForceRebuildPlatformData();
+				}
 			}
 #endif // #if WITH_EDITORONLY_DATA
 		}
