Start of no-gbuffer rendering option. Still have a lot of work to do before this is functional.

diff --git a/Engine/Shaders/BasePassPixelShader.usf b/Engine/Shaders/BasePassPixelShader.usf
--- a/Engine/Shaders/BasePassPixelShader.usf
+++ b/Engine/Shaders/BasePassPixelShader.usf
@@ -496,7 +496,7 @@ void Main(
 #endif
 	OPTIONAL_IsFrontFace,
 	out float4 OutColor		: SV_Target0
-#if FEATURE_LEVEL >= FEATURE_LEVEL_SM4 && (MATERIALBLENDING_SOLID || MATERIALBLENDING_MASKED)
+#if FEATURE_LEVEL >= FEATURE_LEVEL_SM4 && (MATERIALBLENDING_SOLID || MATERIALBLENDING_MASKED) && !NO_GBUFFER
 	,out float4 OutGBufferA : SV_Target1 
 	,out float4 OutGBufferB : SV_Target2
 	,out float4 OutGBufferC : SV_Target3
@@ -704,7 +704,7 @@ void Main(
 
 	float4 PrecomputedShadowFactors = GetPrecomputedShadowMasks(Interpolants);
 
-	#if FEATURE_LEVEL >= FEATURE_LEVEL_SM4 && (MATERIALBLENDING_SOLID || MATERIALBLENDING_MASKED)
+	#if FEATURE_LEVEL >= FEATURE_LEVEL_SM4 && (MATERIALBLENDING_SOLID || MATERIALBLENDING_MASKED) && !NO_GBUFFER
 
 		FGBufferData GBuffer = (FGBufferData)0;
 
diff --git a/Engine/Source/Runtime/Core/Private/HAL/ConsoleManager.cpp b/Engine/Source/Runtime/Core/Private/HAL/ConsoleManager.cpp
--- a/Engine/Source/Runtime/Core/Private/HAL/ConsoleManager.cpp
+++ b/Engine/Source/Runtime/Core/Private/HAL/ConsoleManager.cpp
@@ -1858,4 +1858,11 @@ static TAutoConsoleVariable<int32> CVarFeatureLevelPreview(
 	TEXT("r.FeatureLevelPreview"),
 	0,
 	TEXT("If 1 the quick settings menu will contain an option to enable feature level preview modes"),
+	ECVF_RenderThreadSafe);
+
+
+static TAutoConsoleVariable<int32> CVarGBuffer(
+	TEXT("r.GBuffer"),
+	1,
+	TEXT("0=Do not use GBuffer (fast minimal rendering path), 1=Use GBuffer [default]."),
 	ECVF_RenderThreadSafe);
\ No newline at end of file
diff --git a/Engine/Source/Runtime/Engine/Classes/Engine/RendererSettings.h b/Engine/Source/Runtime/Engine/Classes/Engine/RendererSettings.h
--- a/Engine/Source/Runtime/Engine/Classes/Engine/RendererSettings.h
+++ b/Engine/Source/Runtime/Engine/Classes/Engine/RendererSettings.h
@@ -121,6 +121,12 @@ class ENGINE_API URendererSettings
 		ToolTip="Whether to allow any static lighting to be generated and used, like lightmaps and shadowmaps. Games that only use dynamic lighting should set this to 0 to save some static lighting overhead. Changing this setting requires restarting the editor."))
 	uint32 bAllowStaticLighting:1;
 
+	UPROPERTY(config, EditAnywhere, Category=Lighting, meta=(
+		ConsoleVariable="r.GBuffer",
+		ToolTip="1=Use GBuffer, 0=Don't use GBuffer (minimal limited renderer)."))
+		uint32 bGBuffer:1;
+
+
 	UPROPERTY(config, EditAnywhere, Category=Lighting, meta=(
 		ConsoleVariable="r.AllowMeshDistanceFieldRepresentations",
 		ToolTip="Whether to build distance fields of static meshes, needed for distance field AO, which is used to implement Movable SkyLight shadows.  Enabling will increase mesh build times and memory usage.  Changing this setting requires restarting the editor."))
diff --git a/Engine/Source/Runtime/Engine/Private/ShaderCompiler/ShaderCompiler.cpp b/Engine/Source/Runtime/Engine/Private/ShaderCompiler/ShaderCompiler.cpp
--- a/Engine/Source/Runtime/Engine/Private/ShaderCompiler/ShaderCompiler.cpp
+++ b/Engine/Source/Runtime/Engine/Private/ShaderCompiler/ShaderCompiler.cpp
@@ -2065,6 +2065,11 @@ void GlobalBeginCompileShader(
 		Input.Environment.SetDefine(TEXT("ALLOW_STATIC_LIGHTING"), CVar ? (CVar->GetValueOnGameThread() != 0) : 1);
 	}
 
+	{
+		static const auto CVar = IConsoleManager::Get().FindTConsoleVariableDataInt(TEXT("r.GBuffer"));
+		Input.Environment.SetDefine(TEXT("NO_GBUFFER"), CVar ? (CVar->GetValueOnGameThread() == 0) : 0);
+	}
+
 	{
 		static IConsoleVariable* CVar = IConsoleManager::Get().FindConsoleVariable(TEXT("r.DBuffer"));
 		Input.Environment.SetDefine(TEXT("USE_DBUFFER"), CVar ? CVar->GetInt() : 0);
diff --git a/Engine/Source/Runtime/Engine/Private/ShowFlags.cpp b/Engine/Source/Runtime/Engine/Private/ShowFlags.cpp
--- a/Engine/Source/Runtime/Engine/Private/ShowFlags.cpp
+++ b/Engine/Source/Runtime/Engine/Private/ShowFlags.cpp
@@ -338,6 +338,14 @@ void EngineShowFlagOverride(EShowFlagInitMode ShowFlagInitMode, EViewModeIndex V
 		}
 	}
 
+	{
+		static const auto ICVar = IConsoleManager::Get().FindTConsoleVariableDataInt(TEXT("r.GBuffer"));
+		if(ICVar->GetValueOnGameThread() == 0)
+		{
+			EngineShowFlags.DynamicShadows = 0;
+		}
+	}
+	
 	{
 		static const auto ICVar = IConsoleManager::Get().FindTConsoleVariableDataInt(TEXT("r.RefractionQuality"));
 		if(ICVar->GetValueOnGameThread() <= 0)
diff --git a/Engine/Source/Runtime/Engine/Public/ShowFlags.h b/Engine/Source/Runtime/Engine/Public/ShowFlags.h
--- a/Engine/Source/Runtime/Engine/Public/ShowFlags.h
+++ b/Engine/Source/Runtime/Engine/Public/ShowFlags.h
@@ -210,6 +210,7 @@ struct FEngineShowFlags
 		SetSnap(false);
 		SetVisualizeDistanceFieldAO(false);
 		SetVisualizeSSR(false);
+		SetForceGBuffer(false);
 	}
 
 	static bool CanBeToggledInEditor(const TCHAR* Name)
@@ -250,6 +251,8 @@ struct FEngineShowFlags
 		SetHMDDistortion(false);
 		SetStereoRendering(false);
 		SetDistanceFieldAO(false);
+
+		SetForceGBuffer(true);
 	}
 
 	// ---------------------------------------------------------
diff --git a/Engine/Source/Runtime/Engine/Public/ShowFlagsValues.inl b/Engine/Source/Runtime/Engine/Public/ShowFlagsValues.inl
--- a/Engine/Source/Runtime/Engine/Public/ShowFlagsValues.inl
+++ b/Engine/Source/Runtime/Engine/Public/ShowFlagsValues.inl
@@ -268,6 +268,9 @@ SHOWFLAG_ALWAYS_ACCESSIBLE(VisualizeDistanceFieldAO, SFG_Visualize, LOCTEXT("Vis
 SHOWFLAG_ALWAYS_ACCESSIBLE(DistanceFieldAO, SFG_LightingFeatures, LOCTEXT("DistanceFieldAOSF", "Distance Field AO"))
 /** Visualize screen space reflections, for developer (by default off): */
 SHOWFLAG_FIXED_IN_SHIPPING(VisualizeSSR, 0, SFG_Visualize, LOCTEXT("VisualizeSSR", "Visualize ScreenSpaceReflections"))
+/** Force the use of the GBuffer. */
+SHOWFLAG_ALWAYS_ACCESSIBLE(ForceGBuffer, SFG_Hidden, LOCTEXT("ForceGBuffer", "Force usage of GBuffer"))
+
 
 #undef SHOWFLAG_ALWAYS_ACCESSIBLE
 #undef SHOWFLAG_FIXED_IN_SHIPPING
diff --git a/Engine/Source/Runtime/Renderer/Private/DeferredShadingRenderer.cpp b/Engine/Source/Runtime/Renderer/Private/DeferredShadingRenderer.cpp
--- a/Engine/Source/Runtime/Renderer/Private/DeferredShadingRenderer.cpp
+++ b/Engine/Source/Runtime/Renderer/Private/DeferredShadingRenderer.cpp
@@ -643,6 +643,13 @@ void FDeferredShadingSceneRenderer::Render(FRHICommandListImmediate& RHICmdList)
 	bool bRequiresRHIClear = true;
 	bool bRequiresFarZQuadClear = false;
 
+	static const auto GBufferCVar = IConsoleManager::Get().FindTConsoleVariableDataInt(TEXT("r.GBuffer"));
+	bool bGBuffer = GBufferCVar ? (GBufferCVar->GetValueOnRenderThread() != 0) : true;
+	if(ViewFamily.EngineShowFlags.ForceGBuffer)
+	{
+		bGBuffer = true;
+	}
+
 	if (ClearMethodCVar)
 	{
 		switch (ClearMethodCVar->GetValueOnRenderThread())
@@ -812,6 +819,7 @@ void FDeferredShadingSceneRenderer::Render(FRHICommandListImmediate& RHICmdList)
 	if (ViewFamily.EngineShowFlags.Lighting
 		&& FeatureLevel >= ERHIFeatureLevel::SM4
 		&& ViewFamily.EngineShowFlags.DeferredLighting
+		&& bGBuffer
 		)
 	{
 		GRenderTargetPool.AddPhaseEvent(TEXT("Lighting"));
diff --git a/Engine/Source/Runtime/ShaderCore/Private/Shader.cpp b/Engine/Source/Runtime/ShaderCore/Private/Shader.cpp
--- a/Engine/Source/Runtime/ShaderCore/Private/Shader.cpp
+++ b/Engine/Source/Runtime/ShaderCore/Private/Shader.cpp
@@ -1119,6 +1119,14 @@ void ShaderMapAppendKeyString(EShaderPlatform Platform, FString& KeyString)
 		KeyString += bValue ? TEXT("_SL") : TEXT("_NoSL");
 	}
 
+	{
+		static const auto CVar = IConsoleManager::Get().FindTConsoleVariableDataInt(TEXT("r.GBuffer"));
+		if(CVar ? CVar->GetValueOnGameThread() == 0 : false)
+		{
+			KeyString += TEXT("_NoGB");
+		}
+	}
+
 	{
 		static IConsoleVariable* CVar = IConsoleManager::Get().FindConsoleVariable(TEXT("r.DBuffer"));
 		KeyString += (CVar && CVar->GetInt() != 0) ? TEXT("_DBuf") : TEXT("_NoDBuf");
