[GitHub] 292 : Essential Fix for KismetMath::RandomPointInBoundingBox

#ttp 341669

diff --git a/Engine/Source/Runtime/Engine/Classes/Kismet/KismetMathLibrary.h b/Engine/Source/Runtime/Engine/Classes/Kismet/KismetMathLibrary.h
--- a/Engine/Source/Runtime/Engine/Classes/Kismet/KismetMathLibrary.h
+++ b/Engine/Source/Runtime/Engine/Classes/Kismet/KismetMathLibrary.h
@@ -496,7 +496,8 @@ class UKismetMathLibrary : public UBlueprintFunctionLibrary
 
 	/** Returns a random point within the specified bounding box */
 	UFUNCTION(BlueprintPure, Category = "Math|Random")
-	static FVector RandomPointInBoundingBox(FBox BoundingBox);
+	static FVector RandomPointInBoundingBox(const FVector& Origin, const FVector& BoxExtent);
+
 	/* 
 	 * Returns a random vector with length of 1, within the specified cone, with uniform random distribution. 
 	 * @param ConeDir	The base "center" direction of the cone.
diff --git a/Engine/Source/Runtime/Engine/Private/KismetMathLibrary.cpp b/Engine/Source/Runtime/Engine/Private/KismetMathLibrary.cpp
--- a/Engine/Source/Runtime/Engine/Private/KismetMathLibrary.cpp
+++ b/Engine/Source/Runtime/Engine/Private/KismetMathLibrary.cpp
@@ -689,9 +689,11 @@ FVector UKismetMathLibrary::RandomUnitVector()
 	return FMath::VRand();
 }
 
-FVector UKismetMathLibrary::RandomPointInBoundingBox(FBox BoundingBox)
+FVector UKismetMathLibrary::RandomPointInBoundingBox(const FVector& Origin, const FVector& BoxExtent)
 {
-	return FMath::RandPointInBox(BoundingBox);
+	const FVector BoxMin = Origin - BoxExtent;
+	const FVector BoxMax = Origin + BoxExtent;
+	return FMath::RandPointInBox(FBox(BoxMin, BoxMax));
 }
 
 FVector UKismetMathLibrary::RandomUnitVectorInCone(FVector ConeDir, float ConeHalfAngle)
