* added new shading model Subsurface_Profile, can specify SubsurfaceProfile asset on the material, create it in content browser
  no backscattering yet, might replace Preintegrated and Subsurface shading models, can be optimized, postprocess pass only runs if an object on the screen is using it, uses SeparableSSS by [REDACTED_PERSON] and Diego Gutierrez

diff --git a/Engine/Shaders/BasePassPixelShader.usf b/Engine/Shaders/BasePassPixelShader.usf
--- a/Engine/Shaders/BasePassPixelShader.usf
+++ b/Engine/Shaders/BasePassPixelShader.usf
@@ -206,8 +206,11 @@ float3 GetTranslucencyLighting(FMaterialPixelParameters MaterialParameters, floa
 				// Setup a pow lobe to approximate anisotropic in-scattering near to the light direction
 				float InScattering = pow(saturate(dot(View.DirectionalLightDirection, MaterialParameters.CameraVector)), InScatterPower);
 
+				float4 SSData = GetMaterialSubsurfaceData(MaterialParameters);
+				float3 SubsurfaceColor = SSData.rgb;
+
 				BackscatteredLighting = 
-					GetMaterialSubsurfaceColor(MaterialParameters) 
+					SubsurfaceColor
 					* InScattering 
 					* View.DirectionalLightColor.rgb
 					// Energy normalization, tighter lobes should be brighter
@@ -573,10 +576,20 @@ void Main(
 	float MaterialAO = GetMaterialAmbientOcclusion(MaterialParameters);
 	float Roughness = GetMaterialRoughness(MaterialParameters);
 
+	// 0..1, SubsurfaceProfileId = int(x * 255)
+	float SubsurfaceProfile = 0;
+
 	// If we don't use this shading model the color should be black (don't generate shader code for unused data, don't do indirectlighting cache lighting with this color).
 	float3 SubsurfaceColor = 0;
+#if MATERIAL_SHADINGMODEL_SUBSURFACE || MATERIAL_SHADINGMODEL_PREINTEGRATED_SKIN || MATERIAL_SHADINGMODEL_SUBSURFACE_PROFILE
+	{
+		float4 SubsurfaceData = GetMaterialSubsurfaceData(MaterialParameters);
+
 #if MATERIAL_SHADINGMODEL_SUBSURFACE || MATERIAL_SHADINGMODEL_PREINTEGRATED_SKIN
-	SubsurfaceColor = GetMaterialSubsurfaceColor(MaterialParameters);
+		SubsurfaceColor = SubsurfaceData.rgb;
+#endif
+		SubsurfaceProfile = SubsurfaceData.a;
+	}
 #endif
 
 #if MATERIAL_FULLY_ROUGH
@@ -731,10 +744,13 @@ void Main(
 			GBuffer.ShadingModelID = SHADINGMODELID_DEFAULT_LIT;
 		#elif MATERIAL_SHADINGMODEL_SUBSURFACE
 			GBuffer.ShadingModelID = SHADINGMODELID_SUBSURFACE;
-			GBuffer.CustomData = EncodeSubsurfaceColor( SubsurfaceColor );
+			GBuffer.CustomData = EncodeSubsurfaceColor(SubsurfaceColor);
 		#elif MATERIAL_SHADINGMODEL_PREINTEGRATED_SKIN
 			GBuffer.ShadingModelID = SHADINGMODELID_PREINTEGRATED_SKIN;
-			GBuffer.CustomData = EncodeSubsurfaceColor( SubsurfaceColor );
+			GBuffer.CustomData = EncodeSubsurfaceColor(SubsurfaceColor);
+		#elif MATERIAL_SHADINGMODEL_SUBSURFACE_PROFILE
+			GBuffer.ShadingModelID = SHADINGMODELID_SUBSURFACE_PROFILE;
+			GBuffer.CustomData = EncodeSubsurfaceProfile(SubsurfaceProfile);
 		#elif MATERIAL_SHADINGMODEL_CLEAR_COAT
 			GBuffer.ShadingModelID = SHADINGMODELID_CLEAR_COAT;
 
diff --git a/Engine/Shaders/DeferredLightingCommon.usf b/Engine/Shaders/DeferredLightingCommon.usf
--- a/Engine/Shaders/DeferredLightingCommon.usf
+++ b/Engine/Shaders/DeferredLightingCommon.usf
@@ -231,7 +231,7 @@ float3 ClearCoatShading( FGBufferData GBuffer, float3 LobeRoughness, float3 Lobe
 
 float3 SubsurfaceShadingSubsurface( FGBufferData GBuffer, float3 L, float3 V, half3 N )
 {
-	float3 SubsurfaceColor = DecodeSubsurfaceColor( GBuffer.CustomData );
+	float3 SubsurfaceColor = ExtractSubsurfaceColor(GBuffer);
 
 	float3 H = normalize(V + L);
 
@@ -253,7 +253,7 @@ SamplerState	PreIntegratedBRDFSampler;
 
 float3 SubsurfaceShadingPreintegratedSkin( FGBufferData GBuffer, float3 L, float3 V, half3 N )
 {
-	float3 SubsurfaceColor = DecodeSubsurfaceColor( GBuffer.CustomData );
+	float3 SubsurfaceColor = ExtractSubsurfaceColor(GBuffer);
 
 	float OpacityFactor = GBuffer.Opacity;
 	float3 PreintegratedBRDF = Texture2DSampleLevel(PreIntegratedBRDF, PreIntegratedBRDFSampler, float2(saturate(dot(N, L) * .5 + .5), 1 - OpacityFactor), 0).rgb;
@@ -269,6 +269,7 @@ float3 SurfaceShading( FGBufferData GBuffer, float3 LobeRoughness, float3 LobeEn
 		case SHADINGMODELID_DEFAULT_LIT:
 		case SHADINGMODELID_SUBSURFACE:
 		case SHADINGMODELID_PREINTEGRATED_SKIN:
+		case SHADINGMODELID_SUBSURFACE_PROFILE:
 			return StandardShading( GBuffer, LobeRoughness, LobeEnergy, L, V, N, DiffSpecMask);
 		case SHADINGMODELID_CLEAR_COAT:
 			// this path does not support DiffSpecMask yet
diff --git a/Engine/Shaders/DeferredShadingCommon.usf b/Engine/Shaders/DeferredShadingCommon.usf
--- a/Engine/Shaders/DeferredShadingCommon.usf
+++ b/Engine/Shaders/DeferredShadingCommon.usf
@@ -180,9 +180,10 @@ float3 EncodeSubsurfaceColor(float3 SubsurfaceColor)
 	return sqrt(saturate(SubsurfaceColor));
 }
 
-float3 DecodeSubsurfaceColor(float3 SubsurfaceColor)
+// @param SubsurfaceProfile 0..1, SubsurfaceProfileId = int(x * 255)
+float3 EncodeSubsurfaceProfile(float SubsurfaceProfile)
 {
-	return Square(SubsurfaceColor);
+	return float3(SubsurfaceProfile, 0, 0);
 }
 
 float EncodeIndirectIrradiance(float IndirectIrradiance)
@@ -260,7 +261,8 @@ uint DecodeShadingModelId(float InPackedChannel)
 #define SHADINGMODELID_SUBSURFACE			2
 #define SHADINGMODELID_PREINTEGRATED_SKIN	3
 #define SHADINGMODELID_CLEAR_COAT			4
-#define SHADINGMODELID_NUM					5
+#define SHADINGMODELID_SUBSURFACE_PROFILE	5
+#define SHADINGMODELID_NUM					6
 
 // all values that are output by the forward rendering pass
 struct FGBufferData
@@ -356,7 +358,7 @@ void EncodeGBuffer(
 	}
 	else
 	{
-		bool bSubsurface = (GBuffer.ShadingModelID == SHADINGMODELID_SUBSURFACE || GBuffer.ShadingModelID == SHADINGMODELID_PREINTEGRATED_SKIN);
+		bool bSubsurface = (GBuffer.ShadingModelID == SHADINGMODELID_SUBSURFACE || GBuffer.ShadingModelID == SHADINGMODELID_PREINTEGRATED_SKIN || GBuffer.ShadingModelID == SHADINGMODELID_SUBSURFACE_PROFILE);
 
 #if DIFFUSE_SPEC_INPUTS
 		OutGBufferB.rgb = EncodeSpecularColor(GBuffer.SpecularColor);
@@ -406,7 +408,7 @@ FGBufferData DecodeGBufferData(
 	GBuffer.HasDistanceFieldRepresentation = PackedAlpha & 1;
 	GBuffer.ShadingModelID = DecodeShadingModelId(InGBufferB.a);
 	
-	bool bSubsurface = (GBuffer.ShadingModelID == SHADINGMODELID_SUBSURFACE || GBuffer.ShadingModelID == SHADINGMODELID_PREINTEGRATED_SKIN);
+	bool bSubsurface = (GBuffer.ShadingModelID == SHADINGMODELID_SUBSURFACE || GBuffer.ShadingModelID == SHADINGMODELID_PREINTEGRATED_SKIN || GBuffer.ShadingModelID == SHADINGMODELID_SUBSURFACE_PROFILE);
 
 #if DIFFUSE_SPEC_INPUTS
 	GBuffer.BaseColor = 0;
@@ -491,6 +493,17 @@ void EncodeDBufferData(FGBufferData GBufferData, float3 MultiOpacity,
 	DBufferC = float4(GBufferData.Roughness, 0, 0, MultiOpacity.z);
 }
 
+float3 ExtractSubsurfaceColor(FGBufferData BufferData)
+{
+	return Square(BufferData.CustomData);
+}
+
+uint ExtractSubsurfaceProfileInt(FGBufferData BufferData)
+{
+	// can be optimized
+	return uint(BufferData.CustomData.r * 255.0f + 0.5f);
+}
+
 /** Populates DBufferA, DBufferB, DBufferC as float4 and puts opacity in alpha for frame buffer blending */
 void ApplyDBufferData(FDBufferData DBufferData, inout float3 WorldNormal, inout float3 SubsurfaceColor, inout float Roughness, 
 #if DIFFUSE_SPEC_INPUTS
diff --git a/Engine/Shaders/MaterialTemplate.usf b/Engine/Shaders/MaterialTemplate.usf
--- a/Engine/Shaders/MaterialTemplate.usf
+++ b/Engine/Shaders/MaterialTemplate.usf
@@ -849,6 +849,7 @@ float3 GetShadingModelColor(FGBufferData GBuffer)
 	else if (GBuffer.ShadingModelID == SHADINGMODELID_DEFAULT_LIT) return float3(0.1f, 1.0f, 0.1f); // Green
 	else if (GBuffer.ShadingModelID == SHADINGMODELID_SUBSURFACE) return float3(1.0f, 0.1f, 0.1f); // Red
 	else if (GBuffer.ShadingModelID == SHADINGMODELID_PREINTEGRATED_SKIN) return float3(0.6f, 0.4f, 0.1f); // Brown
+	else if (GBuffer.ShadingModelID == SHADINGMODELID_SUBSURFACE_PROFILE) return float3(0.2f, 0.6f, 0.5f); // Cyan
 	else return float3(1.0f, 1.0f, 1.0f); // White
 #else
 	switch(GBuffer.ShadingModelID)
@@ -857,6 +858,7 @@ float3 GetShadingModelColor(FGBufferData GBuffer)
 		case SHADINGMODELID_DEFAULT_LIT: return float3(0.1f, 1.0f, 0.1f); // Green
 		case SHADINGMODELID_SUBSURFACE: return float3(1.0f, 0.1f, 0.1f); // Red
 		case SHADINGMODELID_PREINTEGRATED_SKIN: return float3(0.6f, 0.4f, 0.1f); // Brown
+		case SHADINGMODELID_SUBSURFACE_PROFILE: return float3(0.2f, 0.6f, 0.5f); // Cyan
 		default: return float3(1.0f, 1.0f, 1.0f); // White
 	}
 #endif
@@ -1147,7 +1149,8 @@ half GetMaterialTessellationMultiplier(FMaterialTessellationParameters Parameter
 %s;
 }
 
-half3 GetMaterialSubsurfaceColor(FMaterialPixelParameters Parameters)
+// .rgb:SubsurfaceColor, .a:SSProfileId in 0..1 range
+half4 GetMaterialSubsurfaceData(FMaterialPixelParameters Parameters)
 {
 %s;
 }
diff --git a/Engine/Shaders/PostProcessAmbient.usf b/Engine/Shaders/PostProcessAmbient.usf
--- a/Engine/Shaders/PostProcessAmbient.usf
+++ b/Engine/Shaders/PostProcessAmbient.usf
@@ -488,7 +488,7 @@ void MainPS(in float4 UVAndScreenPos : TEXCOORD0, out float4 OutColor : SV_Targe
 		// some view dependent and some non view dependent (hard coded)
 		float DependentSplit = 0.5f;
 
-		float3 SubsurfaceColor = DecodeSubsurfaceColor( GBuffer.CustomData );
+		float3 SubsurfaceColor = ExtractSubsurfaceColor(GBuffer);
 
 		// view independent (shared lookup for diffuse for better performance
 		NonSpecularContribution += DiffuseLookup * SubsurfaceColor * (DependentSplit);
diff --git a/Engine/Shaders/PostProcessGBufferHints.usf b/Engine/Shaders/PostProcessGBufferHints.usf
--- a/Engine/Shaders/PostProcessGBufferHints.usf
+++ b/Engine/Shaders/PostProcessGBufferHints.usf
@@ -19,7 +19,7 @@ uint ComputeAdvice(FGBufferData BufferData, float3 OriginalSceneColor)
 	{
 		FLATTEN if(BufferData.ShadingModelID == SHADINGMODELID_SUBSURFACE)
 		{
-			SubsurfaceLum = dot(OneThird.xxx, DecodeSubsurfaceColor(BufferData.CustomData));
+			SubsurfaceLum = dot(OneThird.xxx, ExtractSubsurfaceColor(BufferData));
 		}
 
 		float DiffuseLum = dot(OneThird.xxx, BufferData.DiffuseColor);
diff --git a/Engine/Shaders/PostProcessSubsurface.usf b/Engine/Shaders/PostProcessSubsurface.usf
--- a/Engine/Shaders/PostProcessSubsurface.usf
+++ b/Engine/Shaders/PostProcessSubsurface.usf
@@ -12,13 +12,27 @@
 Texture2D SSProfilesTexture;
 SamplerState SSProfilesTextureSampler;
 
+// x:Radius, y:DistanceToProjectionWindow, zw:xy scale correction for BufferSize!=ViewportSize
+float4 SSSParams;
+
 // ------------------------------------------
 
 // setup for "SeparableSSS.usf"
-
+ 
 #define SSSS_HLSL_4 1
+
+// can be optimized
+float GetSubsurfaceStrength(float2 UV)
+{
+	FScreenSpaceData ScreenSpaceData = GetScreenSpaceData(UV);
+
+	float Mask = ScreenSpaceData.GBuffer.ShadingModelID >= SHADINGMODELID_SUBSURFACE_PROFILE;
+
+	return Mask * ScreenSpaceData.GBuffer.Opacity;
+}
+
 //#define SSSS_STREGTH_SOURCE 1	// todo: masking fom LightingModel or SSSColor
-#define SSSS_STREGTH_SOURCE (GetScreenSpaceData(texcoord).GBuffer.ShadingModelID >= SHADINGMODELID_SUBSURFACE)
+#define SSSS_STREGTH_SOURCE GetSubsurfaceStrength(texcoord)
 
 // 0:speculars leak into SSS / 1:requires alpha channel for scenecolor
 // #define SSSS_SPECULAR_CORRECTION 1
@@ -29,6 +43,7 @@ SamplerState SSProfilesTextureSampler;
 
 // from [REDACTED_PROFILE_URL]/separable-sss/tree/master/Demo
 // [REDACTED_PERSON] [REDACTED_PROFILE_URL]
+// [REDACTED_PROFILE_URL]
 #include "SeparableSSS.usf"
 
 // ------------------------------------------
@@ -56,9 +71,6 @@ SamplerState SSProfilesTextureSampler;
 // 0 / 1
 #define VISUALIZE_KERNEL 0
 
-// x:Radius
-float4 SSSParams;
-
 
 void SetupPS(in float4 UVAndScreenPos : TEXCOORD0, out float4 OutColor : SV_Target0)
 {
@@ -68,7 +80,7 @@ void SetupPS(in float4 UVAndScreenPos : TEXCOORD0, out float4 OutColor : SV_Targ
 
 	FScreenSpaceData ScreenSpaceData = GetScreenSpaceData(UV);
 
-	FLATTEN if(ScreenSpaceData.GBuffer.ShadingModelID >= SHADINGMODELID_SUBSURFACE)
+	FLATTEN if(ScreenSpaceData.GBuffer.ShadingModelID >= SHADINGMODELID_SUBSURFACE_PROFILE)
 	{
 		float4 SceneColor4 = Texture2DSample(PostprocessInput0, PostprocessInput0Sampler, UV);
 
@@ -83,67 +95,74 @@ void SetupPS(in float4 UVAndScreenPos : TEXCOORD0, out float4 OutColor : SV_Targ
 	}
 }
 
-// Gaussian in window -2/c0 .. 2/c0, 2d integral: 1
-// @param XSquared x * x
-float ComputeWindowedGaussian(float x, float c0)
-{
-	x *= c0;
-
-	float XSquared = x * x;
-	return max(exp2(-XSquared) - 0.065f, 0) / 3.4324f * (c0 * c0);
-}
-
-// Subsurface function
-// @param x > 0
-float3 ColorOverDistance(float3 SubsurfaceColor, float x)
-{
-	// larger means tighter gaussian
-	x *= 2.0f;	// -1..1 to -2..2
-
-	float3 Color = SubsurfaceColor;
-
-	// to avoid stuff gets brighter
-//	Color /= dot(Color, 1);
-
-	float3 Hack = 1;//float3(0.8f, 0.5f, 0.35f);
-
-	return Hack * float3(
-		ComputeWindowedGaussian(x, Color.r),
-		ComputeWindowedGaussian(x, Color.g),
-		ComputeWindowedGaussian(x, Color.b) );
-}
-
-float3 SampleDiffuseLighting(float2 LocalUV)
-{
-#if VISUALIZE_KERNEL
-	float Aspect = ViewportSize.x / ViewportSize.y;
-
-	float2 Dxy = (LocalUV - float2(0.75f, 0.25f)) * float2(Aspect, 1);
-
-	if(dot(Dxy, Dxy) < 0.000006f) return 2;
-	if(dot(Dxy, Dxy) < 0.05f) return 0;
-#endif
-
-	return Texture2DSample(PostprocessInput0, PostprocessInput0Sampler, LocalUV).rgb;
-}
-
 // input0 is created by the SetupPS shader
 void MainPS(float4 UVAndScreenPos : TEXCOORD0, out float4 OutColor : SV_Target0)
 {
 	float2 UV = UVAndScreenPos.xy;
 
 	FScreenSpaceData ScreenSpaceData = GetScreenSpaceData(UV);
+	//
+
+	// test
+//	if (UV.y > 0.5f)
+	if (0)
+	{
+		uint SubsurfaceProfileInt = ExtractSubsurfaceProfileInt(ScreenSpaceData.GBuffer);
+		OutColor = float4(0.5f, 0.5f, 0.5f, 0);
+		if (SubsurfaceProfileInt == 0)
+		{
+			OutColor = float4(0.8f, 0.7f, 0.6f, 0);
+		}
+		if (SubsurfaceProfileInt == 1)
+		{
+			OutColor = float4(1, 0, 0, 0);
+		}
+		if (SubsurfaceProfileInt == 2)
+		{
+			OutColor = float4(0, 1, 0, 0);
+		}
+		if (SubsurfaceProfileInt == 3)
+		{ 
+			OutColor = float4(0, 0, 1, 0);
+		}
+		if (SubsurfaceProfileInt == 4)
+		{
+			OutColor = float4(1, 0, 1, 0);
+		}
+		if (SubsurfaceProfileInt == 5)
+		{
+			OutColor = float4(0, 1, 1, 0);
+		}
+		if (SubsurfaceProfileInt == 6)
+		{
+			OutColor = float4(1, 1, 0, 0);
+		}
+		if (SubsurfaceProfileInt == 100)
+		{
+			OutColor = float4(0, 0.2f, 0, 0);
+		}
+		if (SubsurfaceProfileInt == 255)
+		{
+			OutColor = float4(1, 1, 1, 0);
+		}
+		return;
+	}
 
 	// call into "SeparableSSS.usf"
 
-	// in world units
-	float sssWidth = SSSParams.x * 0.01f;
+	// in world units, *0.01 to scale from units(cm) to m
+	float sssWidth = SSSParams.x;
+
+	// the viewport is only a fraction of the buffersize, here we compensate for that
+	// *1000 as we store the sample distances / 1000 to allows for scaling
+	float Scaler = View.ViewSizeAndSceneTexelSize.x * View.ViewSizeAndSceneTexelSize.z * 1000.0f;
+
 #if METHOD == 0
 	// first pass
-	float2 Direction = float2(1, 0);
+	float2 Direction = float2(1, 0) * Scaler;
 #else
 	// second pass
-	float2 Direction = float2(0, 1);
+	float2 Direction = float2(0, View.ViewSizeAndSceneTexelSize.w / View.ViewSizeAndSceneTexelSize.z) * Scaler;
 #endif
 	OutColor = SSSSBlurPS(UV, sssWidth, Direction, false);
 
@@ -154,16 +173,18 @@ void MainPS(float4 UVAndScreenPos : TEXCOORD0, out float4 OutColor : SV_Target0)
 #if METHOD > 0
 	// on second pass we recombine with the scene color
 
-	float4 SceneColor4 = Texture2DSample(PostprocessInput1, PostprocessInput1Sampler, UV);
+	float4 SceneColor4 = Texture2DSample(PostprocessInput0, PostprocessInput0Sampler, UV);
+
 	float4 SSSColor = OutColor;
 
-	SSSColor.a = (ScreenSpaceData.GBuffer.ShadingModelID >= SHADINGMODELID_SUBSURFACE);
+	SSSColor.a = (ScreenSpaceData.GBuffer.ShadingModelID >= SHADINGMODELID_SUBSURFACE_PROFILE);
 
 #if METHOD > 1
 	// we took the specular highlights out, now we add them back in
 	SSSColor.rgb += LightAccumulator_ReconstructNonDiffuseLighting(SceneColor4);
 #endif
 
-	OutColor = float4(lerp(SceneColor4.rgb, SSSColor.rgb, SSSColor.a), 1);
+	// alpha channel can be anything - we consumed the value and likely output to a format that doesn't need the value
+	OutColor = float4(lerp(SceneColor4.rgb, SSSColor.rgb, SSSColor.a), 0);
 #endif
 }
diff --git a/Engine/Shaders/ReflectionEnvironmentComputeShaders.usf b/Engine/Shaders/ReflectionEnvironmentComputeShaders.usf
--- a/Engine/Shaders/ReflectionEnvironmentComputeShaders.usf
+++ b/Engine/Shaders/ReflectionEnvironmentComputeShaders.usf
@@ -658,11 +658,12 @@ void ReflectionEnvironmentTiledDeferredMain(
 		case SHADINGMODELID_DEFAULT_LIT:
 		case SHADINGMODELID_SUBSURFACE:
 		case SHADINGMODELID_PREINTEGRATED_SKIN:
+		case SHADINGMODELID_SUBSURFACE_PROFILE:
 		{
 			BRANCH
 			if( GBuffer.ShadingModelID == SHADINGMODELID_SUBSURFACE )
 			{
-				GBuffer.DiffuseColor += DecodeSubsurfaceColor(GBuffer.CustomData);
+				GBuffer.DiffuseColor += ExtractSubsurfaceColor(GBuffer);
 			}
 
 		#if USE_LIGHTMAPS
diff --git a/Engine/Shaders/SeparableSSS.usf b/Engine/Shaders/SeparableSSS.usf
--- a/Engine/Shaders/SeparableSSS.usf
+++ b/Engine/Shaders/SeparableSSS.usf
@@ -95,9 +95,9 @@
 /**
  * SSSS_FOV must be set to the value used to render the scene.
  */
-#ifndef SSSS_FOVY
-#define SSSS_FOVY 20.0
-#endif
+//#ifndef SSSS_FOVY
+//#define SSSS_FOVY 20.0
+//#endif
 
 /**
  * Light diffusion should occur on the surface of the object, not in a screen 
@@ -147,7 +147,7 @@ float4 kernel[SSSS_N_KERNELWEIGHTCOUNT];
 
 #if SSSS_QUALITY == 2
 #define SSSS_N_KERNELWEIGHTCOUNT 13
-const static float4 kernel[] = {
+/*const static float4 kernel[] = {
     float4(0.530605, 0.613514, 0.739601, 0),
     float4(0.0211412, 0.0459286, 0.0378196, 0.0208333),
     float4(0.0402784, 0.0657244, 0.04631, 0.0833333),
@@ -161,6 +161,21 @@ const static float4 kernel[] = {
     float4(0.00500364, 0.00020094, 5.28848e-005, 2.08333),
     float4(0.00333804, 7.85443e-005, 1.2945e-005, 2.52083),
     float4(0.000973794, 1.11862e-005, 9.43437e-007, 3),
+};*/
+const static float4 kernel[] = {
+	float4(0.540384, 0.622941, 0.739601, 0),
+	float4(0.0207007, 0.0448084, 0.0378196, 0.0208333),
+	float4(0.0394393, 0.0641214, 0.04631, 0.0833333),
+	float4(0.0483305, 0.0358757, 0.0219485, 0.1875),
+	float4(0.0401626, 0.0195023, 0.0118481, 0.333333),
+	float4(0.0258149, 0.0116795, 0.00684598, 0.520833),
+	float4(0.0175505, 0.00694333, 0.00347194, 0.75),
+	float4(0.0125819, 0.00347638, 0.00132016, 1.02083),
+	float4(0.00924226, 0.00135726, 0.000416598, 1.33333),
+	float4(0.00686372, 0.000481619, 0.000151938, 1.6875),
+	float4(0.0048994, 0.000196039, 5.28848e-005, 2.08333),
+	float4(0.0032685, 7.66286e-005, 1.2945e-005, 2.52083),
+	float4(0.000953506, 1.09134e-005, 9.43437e-007, 3),
 };
 #elif SSSS_QUALITY == 1
 #define SSSS_N_KERNELWEIGHTCOUNT 9
@@ -195,18 +210,18 @@ float4 Quantize8Bit(float4 In)
 	return floor(In * 255.999f) / 255.0f;
 }
 
-float4 GetKernel(uint Index)
+float4 GetKernel(uint SampleIndex, uint SubsurfaceProfileInt)
 {
-	const float4 TableMax = float4(0.784728f, 0.784728f, 0.784728f, 3.0f);
-	
-	// profiled on NV670 fullscreen, one pass (total: 2 pass), samples: 13+1+13, large head on screen, r.SSSS 200
+	const float4 TableMax = float4(1, 1, 1, 3.0f);
+
+	// profiled on NV670 fullscreen, one pass (total: 2 pass), samples: 13+1+13, large object filling the screen
 
 	// texture lookup
-//	float4 Value = SSProfilesTexture.Load(int3(Index, 0, 0)) * TableMax;		// 0.88ms
+	float4 Value = SSProfilesTexture.Load(int3(SampleIndex, SubsurfaceProfileInt, 0)) * TableMax;		// 0.88ms for 8bit
 	// simulated  8 bit lookup
-//	float4 Value = Quantize8Bit(kernel[Index] / TableMax) * TableMax;			// 1.33ms
+//	float4 Value = Quantize8Bit(kernel[SampleIndex] / TableMax) * TableMax;			// 1.33ms same look as texture
 	// float reference
-	float4 Value = kernel[Index];												// 0.85ms
+//	float4 Value = kernel[SampleIndex];												// 0.85ms, quite a bit different in look, need to use more than 8 bit?
 
 	return Value;
 }
@@ -237,9 +252,9 @@ float4 GetKernel(uint Index)
 
 // EPIC changed start
 #define SSSSSampleShadowmap(coord)	float4(0,0,0,0)	// todo: not used yet
-#define SSSSSampleMasking(texcoord)	(GetScreenSpaceData(texcoord).GBuffer.ShadingModelID >= SHADINGMODELID_SUBSURFACE)	// todo: masking doesn't work with filtering ideally put in alpa to get it filtered
-#define SSSSSampleSceneColor(coord) (SSSSSampleMasking(coord) * Texture2DSample(PostprocessInput0, PostprocessInput0Sampler, coord))
-#define SSSSSampleSceneColorPoint(coord) ( Texture2DSample(PostprocessInput0, PostprocessInput0Sampler, coord))		// todo: could be made point
+#define SSSSSampleMasking(texcoord)	(GetScreenSpaceData(texcoord).GBuffer.ShadingModelID == SHADINGMODELID_SUBSURFACE_PROFILE)	// todo: masking doesn't work with filtering ideally put in alpha to get it filtered
+#define SSSSSampleSceneColor(coord) (SSSSSampleMasking(coord) * Texture2DSample(PostprocessInput1, PostprocessInput1Sampler, coord))
+#define SSSSSampleSceneColorPoint(coord) ( Texture2DSample(PostprocessInput1, PostprocessInput1Sampler, coord))		// todo: could be made point
 #define SSSSSampleSceneDepth(coord) CalcSceneDepth(coord)		// todo: could be made linear
 #define SSSSSampleSceneDepthPoint(coord) CalcSceneDepth(coord)
 // EPIC changed end
@@ -442,26 +457,32 @@ float4 SSSSBlurPS(
 
     // Calculate the sssWidth scale (1.0 for a unit plane sitting on the
     // projection window):
-    float distanceToProjectionWindow = 1.0 / tan(0.5 * radians(SSSS_FOVY));
-    float scale = distanceToProjectionWindow / depthM;
+//    float distanceToProjectionWindow = 1.0 / tan(0.5 * radians(SSSS_FOVY));
+	float distanceToProjectionWindow = SSSParams.y;
+//	float distanceToProjectionWindow = 1.0 / tan(0.5 * radians(90.0f));
+
+	float scale = distanceToProjectionWindow / depthM;
 
     // Calculate the final step to fetch the surrounding pixels:
-    float2 finalStep = sssWidth * scale * dir;
-    finalStep *= SSSS_STREGTH_SOURCE; // Modulate it using the alpha channel.
+    float2 finalStep = sssWidth * scale * dir * 0.5f;										// hacked in 0.5 - -1..1 to 0..1 but why this isn't in demo code?
+	finalStep *= SSSS_STREGTH_SOURCE; // Modulate it using the alpha channel.
     finalStep *= 1.0 / 3.0; // Divide by 3 as the kernels range from -3 to 3.
 
+	FGBufferData GBufferData = GetGBufferData(texcoord);
+
+	// 0..255, which SubSurface profile to pick
+	uint SubsurfaceProfileInt = ExtractSubsurfaceProfileInt(GBufferData);
+
     // Accumulate the center sample:
     float4 colorBlurred = colorM;
-	colorBlurred.rgb *= GetKernel(0).rgb;
-
-	// hack
-   // colorBlurred.rgb = 0;
+	colorBlurred.rgb *= GetKernel(0, SubsurfaceProfileInt).rgb;
 
 	// Accumulate the other samples:
     SSSS_UNROLL
 	for (int i = 1; i < SSSS_N_KERNELWEIGHTCOUNT; i++) 
 	{
-		half4 Kernel = GetKernel(i);
+		// Kernel.a = 0..3
+		half4 Kernel = GetKernel(i, SubsurfaceProfileInt);
 
 		float3 LocalAccum = 0;
 
diff --git a/Engine/Source/Developer/AssetTools/Private/AssetTools.cpp b/Engine/Source/Developer/AssetTools/Private/AssetTools.cpp
--- a/Engine/Source/Developer/AssetTools/Private/AssetTools.cpp
+++ b/Engine/Source/Developer/AssetTools/Private/AssetTools.cpp
@@ -47,6 +47,7 @@ FAssetTools::FAssetTools()
 	RegisterAssetTypeActions( MakeShareable(new FAssetTypeActions_Struct) );
 	RegisterAssetTypeActions( MakeShareable(new FAssetTypeActions_Font) );
 	RegisterAssetTypeActions( MakeShareable(new FAssetTypeActions_ForceFeedbackEffect) );
+	RegisterAssetTypeActions( MakeShareable(new FAssetTypeActions_SubsurfaceProfile));
 	RegisterAssetTypeActions( MakeShareable(new FAssetTypeActions_InstancedFoliageSettings) );
 	RegisterAssetTypeActions( MakeShareable(new FAssetTypeActions_InterpData) );
 	RegisterAssetTypeActions( MakeShareable(new FAssetTypeActions_LandscapeLayer) );
diff --git a/Engine/Source/Developer/AssetTools/Private/AssetToolsPrivatePCH.h b/Engine/Source/Developer/AssetTools/Private/AssetToolsPrivatePCH.h
--- a/Engine/Source/Developer/AssetTools/Private/AssetToolsPrivatePCH.h
+++ b/Engine/Source/Developer/AssetTools/Private/AssetToolsPrivatePCH.h
@@ -47,6 +47,7 @@ DECLARE_LOG_CATEGORY_EXTERN(LogAssetTools, Log, All);
 #include "AssetTypeActions/AssetTypeActions_Struct.h"
 #include "AssetTypeActions/AssetTypeActions_Font.h"
 #include "AssetTypeActions/AssetTypeActions_ForceFeedbackEffect.h"
+#include "AssetTypeActions/AssetTypeActions_SubsurfaceProfile.h"
 #include "AssetTypeActions/AssetTypeActions_InstancedFoliageSettings.h"
 #include "AssetTypeActions/AssetTypeActions_InterpData.h"
 #include "AssetTypeActions/AssetTypeActions_LandscapeLayer.h"
diff --git a/Engine/Source/Developer/AssetTools/Private/AssetTypeActions/AssetTypeActions_SubsurfaceProfile.cpp b/Engine/Source/Developer/AssetTools/Private/AssetTypeActions/AssetTypeActions_SubsurfaceProfile.cpp
new file mode 100644
--- /dev/null
+++ b/Engine/Source/Developer/AssetTools/Private/AssetTypeActions/AssetTypeActions_SubsurfaceProfile.cpp
@@ -0,0 +1,8 @@
+// Copyright 1998-2014 Epic Games, Inc. All Rights Reserved.
+
+#include "AssetToolsPrivatePCH.h"
+
+UClass* FAssetTypeActions_SubsurfaceProfile::GetSupportedClass() const
+{
+	return USubsurfaceProfile::StaticClass();
+}
\ No newline at end of file
diff --git a/Engine/Source/Developer/AssetTools/Private/AssetTypeActions/AssetTypeActions_SubsurfaceProfile.h b/Engine/Source/Developer/AssetTools/Private/AssetTypeActions/AssetTypeActions_SubsurfaceProfile.h
new file mode 100644
--- /dev/null
+++ b/Engine/Source/Developer/AssetTools/Private/AssetTypeActions/AssetTypeActions_SubsurfaceProfile.h
@@ -0,0 +1,14 @@
+// Copyright 1998-2014 Epic Games, Inc. All Rights Reserved.
+
+#pragma once
+
+class FAssetTypeActions_SubsurfaceProfile : public FAssetTypeActions_Base
+{
+public:
+	// IAssetTypeActions Implementation
+	virtual FText GetName() const override { return NSLOCTEXT("AssetTypeActions", "AssetTypeActions_SubsurfaceProfile", "Subsurface Profile"); }
+	virtual FColor GetTypeColor() const override { return FColor(192, 121, 64); }
+	virtual UClass* GetSupportedClass() const override;
+
+	virtual uint32 GetCategories() override { return EAssetTypeCategories::MaterialsAndTextures; }
+};
\ No newline at end of file
diff --git a/Engine/Source/Editor/UnrealEd/Classes/Factories/SubsurfaceProfileFactory.h b/Engine/Source/Editor/UnrealEd/Classes/Factories/SubsurfaceProfileFactory.h
new file mode 100644
--- /dev/null
+++ b/Engine/Source/Editor/UnrealEd/Classes/Factories/SubsurfaceProfileFactory.h
@@ -0,0 +1,18 @@
+// Copyright 1998-2014 Epic Games, Inc. All Rights Reserved.
+
+#pragma once
+#include "SubsurfaceProfileFactory.generated.h"
+
+UCLASS(hidecategories=Object)
+class USubsurfaceProfileFactory : public UFactory
+{
+	GENERATED_UCLASS_BODY()
+
+
+	// Begin UFactory Interface
+	virtual UObject* FactoryCreateNew(UClass* Class,UObject* InParent,FName Name,EObjectFlags Flags,UObject* Context,FFeedbackContext* Warn) override;
+	// Begin UFactory Interface	
+};
+
+
+
diff --git a/Engine/Source/Editor/UnrealEd/Private/Factories/EditorFactories.cpp b/Engine/Source/Editor/UnrealEd/Private/Factories/EditorFactories.cpp
--- a/Engine/Source/Editor/UnrealEd/Private/Factories/EditorFactories.cpp
+++ b/Engine/Source/Editor/UnrealEd/Private/Factories/EditorFactories.cpp
@@ -7125,6 +7125,24 @@ UObject* UForceFeedbackEffectFactory::FactoryCreateNew( UClass* InClass, UObject
 	return ConstructObject<UForceFeedbackEffect>( UForceFeedbackEffect::StaticClass(), InParent, InName, Flags );
 }
 
+/*-----------------------------------------------------------------------------
+USubsurfaceProfileFactory implementation.
+-----------------------------------------------------------------------------*/
+USubsurfaceProfileFactory::USubsurfaceProfileFactory(const class FPostConstructInitializeProperties& PCIP)
+: Super(PCIP)
+{
+
+	SupportedClass = USubsurfaceProfile::StaticClass();
+	bCreateNew = true;
+	bEditorImport = false;
+	bEditAfterNew = true;
+}
+
+UObject* USubsurfaceProfileFactory::FactoryCreateNew(UClass* InClass, UObject* InParent, FName InName, EObjectFlags Flags, UObject* Context, FFeedbackContext* Warn)
+{
+	return ConstructObject<USubsurfaceProfile>(USubsurfaceProfile::StaticClass(), InParent, InName, Flags);
+}
+
 /*-----------------------------------------------------------------------------
 	UTouchInterfaceFactory implementation.
 -----------------------------------------------------------------------------*/
diff --git a/Engine/Source/Editor/UnrealEd/Private/Lightmass/LightmassRender.cpp b/Engine/Source/Editor/UnrealEd/Private/Lightmass/LightmassRender.cpp
--- a/Engine/Source/Editor/UnrealEd/Private/Lightmass/LightmassRender.cpp
+++ b/Engine/Source/Editor/UnrealEd/Private/Lightmass/LightmassRender.cpp
@@ -841,7 +841,8 @@ bool FLightmassMaterialRenderer::GenerateMaterialData(
  	if ((ShadingModel != MSM_DefaultLit) &&
 		(ShadingModel != MSM_Unlit) &&
 		(ShadingModel != MSM_Subsurface) &&
-		(ShadingModel != MSM_PreintegratedSkin))
+		(ShadingModel != MSM_PreintegratedSkin) &&
+		(ShadingModel != MSM_SubsurfaceProfile))
 	{
 		UE_LOG(LogLightmassRender, Warning, TEXT("LIGHTMASS: Material has an unsupported shading model: %d on %s"), 
 			(int32)ShadingModel,
diff --git a/Engine/Source/Runtime/Core/Public/Math/UnrealMathUtility.h b/Engine/Source/Runtime/Core/Public/Math/UnrealMathUtility.h
--- a/Engine/Source/Runtime/Core/Public/Math/UnrealMathUtility.h
+++ b/Engine/Source/Runtime/Core/Public/Math/UnrealMathUtility.h
@@ -965,6 +965,29 @@ struct FMath : public FPlatformMath
 	 * @param Dst in and out
 	 */
 	static CORE_API void ApplyScaleToFloat(float& Dst, const FVector& DeltaScale, float Magnitude = 1.0f);
+
+	// @param x assumed to be in this range: 0..1
+	// @return 0..255
+	static uint8 Quantize8UnsignedByte(float x)
+	{
+		// 0..1 -> 0..255
+		int32 Ret = (int32)(x * 255.999f);
+
+		check(Ret >= 0);
+		check(Ret <= 255);
+
+		return Ret;
+	}
+	
+	// @param x assumed to be in this range: -1..1
+	// @return 0..255
+	static uint8 Quantize8SignedByte(float x)
+	{
+		// -1..1 -> 0..1
+		float y = x * 0.5f + 0.5f;
+
+		return Quantize8UnsignedByte(y);
+	}
 };
 
 // Platform specific vector intrinsics include.
diff --git a/Engine/Source/Runtime/Engine/Classes/Engine/EngineTypes.h b/Engine/Source/Runtime/Engine/Classes/Engine/EngineTypes.h
--- a/Engine/Source/Runtime/Engine/Classes/Engine/EngineTypes.h
+++ b/Engine/Source/Runtime/Engine/Classes/Engine/EngineTypes.h
@@ -118,6 +118,7 @@ enum EMaterialShadingModel
 	MSM_Subsurface			UMETA(DisplayName="Subsurface"),
 	MSM_PreintegratedSkin	UMETA(DisplayName="Preintegrated Skin"),
 	MSM_ClearCoat			UMETA(DisplayName="Clear Coat"),
+	MSM_SubsurfaceProfile	UMETA(DisplayName="Subsurface Profile"),
 	MSM_MAX,
 };
 
diff --git a/Engine/Source/Runtime/Engine/Classes/Engine/SubsurfaceProfile.h b/Engine/Source/Runtime/Engine/Classes/Engine/SubsurfaceProfile.h
new file mode 100644
--- /dev/null
+++ b/Engine/Source/Runtime/Engine/Classes/Engine/SubsurfaceProfile.h
@@ -0,0 +1,149 @@
+// Copyright 1998-2014 Epic Games, Inc. All Rights Reserved.
+#pragma once
+
+#include "../Curves/CurveFloat.h"
+#include "SubsurfaceProfile.generated.h"
+
+// struct with all the settings we want in USubsurfaceProfile, separate to make it easer to pass this data around in the engine.
+USTRUCT(BlueprintType)
+struct FSubsurfaceProfileStruct
+{
+	GENERATED_USTRUCT_BODY()
+	
+	/** in world/unreal units (cm) */
+	UPROPERTY(Category = "SubsurfaceProfileStruct", EditAnywhere, BlueprintReadOnly, meta = (ClampMin = "0.1", UIMax = "50.0", ClampMax = "1000.0"))
+	float ScatterRadius;
+
+	/**
+	* Specifies the how much of the diffuse light gets into the material,
+	* can be seen as a per-channel mix factor between the original image,
+	* and the SSS-filtered image (called "strength" in SeparableSSS, default there: 0.48, 0.41, 0.28)
+	*/
+	UPROPERTY(Category = "SubsurfaceProfileStruct", EditAnywhere, BlueprintReadOnly, meta = (HideAlphaChannel))
+	FLinearColor SubsurfaceColor;
+
+	/**
+	* defines the per-channel falloff of the gradients
+	* produced by the subsurface scattering events, can be used to fine tune the color of the gradients
+	* (called "falloff" in SeparableSSS, default there: 1, 0.37, 0.3)
+	*/
+	UPROPERTY(Category = "SubsurfaceProfileStruct", EditAnywhere, BlueprintReadOnly, meta = (HideAlphaChannel))
+	FLinearColor FalloffColor;
+
+	// constructor
+	FSubsurfaceProfileStruct()
+	{
+		// defaults from SeparableSSS.h and the demo
+		ScatterRadius = 1.2f;
+		SubsurfaceColor = FLinearColor(0.48f, 0.41f, 0.28f);
+		FalloffColor = FLinearColor(1.0f, 0.37f, 0.3f);
+	}
+
+};
+
+/**
+ * Subsurface Scattering profile asset, can be specified at the material. Only for "Subsurface Profile" materials, is use during Screenspace Subsurface Scattering
+ * Don't change at runtime. All properties in here are per material - texture like variations need to come from properties that are in the GBuffer.
+ */
+UCLASS(autoexpandcategories = SubsurfaceProfile, MinimalAPI)
+class USubsurfaceProfile : public UObject
+{
+	GENERATED_UCLASS_BODY()
+
+	UPROPERTY(Category = USubsurfaceProfile, EditAnywhere, meta = (ShowOnlyInnerProperties))
+	struct FSubsurfaceProfileStruct Settings;
+
+	// Begin UObject interface
+	virtual void BeginDestroy();
+	virtual void PostEditChangeProperty(struct FPropertyChangedEvent& PropertyChangedEvent);
+	// End UObject interface
+};
+
+
+
+// render thread
+class FSubsurfaceProfileTexture : public FRenderResource
+{
+public:
+	// constructor runs on main thread (global variable)
+	FSubsurfaceProfileTexture();
+
+	// destructor
+	~FSubsurfaceProfileTexture();
+
+	void SetRendererModule(class IRendererModule* InRendererModule) { RendererModule = InRendererModule; }
+
+	// convenience, can be optimized 
+	int32 AddOrUpdateProfile(const FSubsurfaceProfileStruct Settings, const USubsurfaceProfile& Profile)
+	{
+		int32 AllocationId = FindAllocationId(Profile); 
+		
+		if (AllocationId != -1)
+		{
+			UpdateProfile(AllocationId, Settings);
+		}
+		else
+		{
+			AllocationId = AddProfile(Settings, Profile);
+		}
+
+		return AllocationId;
+	}
+
+	// O(n) n is a small number
+	// @param Profile must not be 0
+	// @return AllocationId -1: no allocation, should be deallocated with DeallocateSubsurfaceProfile()
+	int32 AddProfile(const FSubsurfaceProfileStruct Settings, const USubsurfaceProfile& Profile);
+
+	// O(n) to find the element, n is the SSProfile count and usually quite small
+	void RemoveProfile(USubsurfaceProfile& InProfile);
+
+	// @param InRendererModule first call needs to have it != 0
+	void UpdateProfile(const FSubsurfaceProfileStruct Settings, const USubsurfaceProfile& Profile) { int32 AllocationId = FindAllocationId(Profile); UpdateProfile(AllocationId, Settings); }
+
+	// @param InRendererModule first call needs to have it != 0
+	void UpdateProfile(int32 AllocationId, const FSubsurfaceProfileStruct Settings);
+
+	const struct FSceneRenderTargetItem* GetTexture(FRHICommandListImmediate& RHICmdList);
+
+
+	// FRenderResource interface.
+	/**
+	* Release textures when device is lost/destroyed.
+	*/
+	virtual void ReleaseDynamicRHI() override;
+
+	// for debugging, can be removed
+	void Dump();
+
+	// @return -1 if not found
+	int32 FindAllocationId(const USubsurfaceProfile& InProfile) const;
+
+private:
+
+	struct FSubsurfaceProfileEntry
+	{
+		FSubsurfaceProfileEntry(FSubsurfaceProfileStruct InSubsurfaceProfileStruct, const USubsurfaceProfile* InGameThreadObject)
+			: Settings(InSubsurfaceProfileStruct)
+			, GameThreadObject(InGameThreadObject)
+		{
+		}
+
+		FSubsurfaceProfileStruct Settings;
+
+		// not referenced, void and not USubsurfaceProfile* to prevent accidential access, 0 if the entry can be reused or it's [0] which is used as default
+		const void* GameThreadObject;
+	};
+
+	// 0 if there if no user of the feature
+	IRendererModule* RendererModule;
+	//
+	TArray<FSubsurfaceProfileEntry> SubsurfaceProfileEntries;
+
+	// Could be optimized but should not happen too often (during level load or editor operations).
+	void CreateTexture(FRHICommandListImmediate& RHICmdList);
+
+};
+
+// lives on the render thread
+extern TGlobalResource<FSubsurfaceProfileTexture> GSubsufaceProfileTextureObject;
diff --git a/Engine/Source/Runtime/Engine/Classes/Materials/Material.h b/Engine/Source/Runtime/Engine/Classes/Materials/Material.h
--- a/Engine/Source/Runtime/Engine/Classes/Materials/Material.h
+++ b/Engine/Source/Runtime/Engine/Classes/Materials/Material.h
@@ -846,6 +846,9 @@ private:
 	/** Sets up transient properties in MaterialResources. */
 	void UpdateResourceAllocations();
 
+	/** to share code for PostLoad() and PostEditChangeProperty() */
+	void PropagateDataToMaterialProxy();
+
 public:
 
 	/** @return the name of the given usage flag. */
diff --git a/Engine/Source/Runtime/Engine/Classes/Materials/MaterialInstance.h b/Engine/Source/Runtime/Engine/Classes/Materials/MaterialInstance.h
--- a/Engine/Source/Runtime/Engine/Classes/Materials/MaterialInstance.h
+++ b/Engine/Source/Runtime/Engine/Classes/Materials/MaterialInstance.h
@@ -375,6 +375,9 @@ protected:
 	/** Copies over parameters given a material interface */
 	ENGINE_API void CopyMaterialInstanceParameters(UMaterialInterface* MaterialInterface);
 
+	// to share code between PostLoad() and PostEditChangeProperty()
+	void UpdateMaterialInstanceData();
+
 	/** Allow resource to access private members. */
 	friend class FMaterialInstanceResource;
 	/** Editor-only access to private members. */
diff --git a/Engine/Source/Runtime/Engine/Classes/Materials/MaterialInterface.h b/Engine/Source/Runtime/Engine/Classes/Materials/MaterialInterface.h
--- a/Engine/Source/Runtime/Engine/Classes/Materials/MaterialInterface.h
+++ b/Engine/Source/Runtime/Engine/Classes/Materials/MaterialInterface.h
@@ -53,6 +53,9 @@ struct ENGINE_API FMaterialRelevance
 	UPROPERTY()
 		uint32 bDisableDepthTest : 1;
 
+	UPROPERTY()
+		uint32 bSubsurfaceProfile : 1;
+
 	/** Default constructor. */
 	FMaterialRelevance()
 		: bOpaque(false)
@@ -61,6 +64,7 @@ struct ENGINE_API FMaterialRelevance
 		, bSeparateTranslucency(false)
 		, bNormalTranslucency(false)
 		, bDisableDepthTest(false)
+		, bSubsurfaceProfile(false)
 	{}
 
 	/** Bitwise OR operator.  Sets any relevance bits which are present in either FMaterialRelevance. */
@@ -72,6 +76,7 @@ struct ENGINE_API FMaterialRelevance
 		bSeparateTranslucency |= B.bSeparateTranslucency;
 		bNormalTranslucency |= B.bNormalTranslucency;
 		bDisableDepthTest |= B.bDisableDepthTest;
+		bSubsurfaceProfile |= B.bSubsurfaceProfile;
 		return *this;
 	}
 
@@ -162,6 +167,12 @@ class UMaterialInterface : public UObject, public IBlendableInterface
 {
 	GENERATED_UCLASS_BODY()
 
+	/** SubsurfaceProfile, for Screen Space Subsurface Scattering */
+	UPROPERTY(EditAnywhere, BlueprintReadOnly, Category = Material, meta = (DisplayName = "Subsurface Profile"))
+	class USubsurfaceProfile* SubsurfaceProfile;
+
+	/* -------------------------- */
+
 	/** A fence to track when the primitive is no longer used as a parent */
 	FRenderCommandFence ParentRefFence;
 
@@ -569,6 +580,8 @@ protected:
 	/** Returns a bitfield indicating which feature levels should be compiled for rendering. */
 	ENGINE_API uint32 GetFeatureLevelsToCompileForRendering() const;
 
+	void UpdateMaterialRenderProxy(FMaterialRenderProxy& Proxy);
+
 private:
 	/**
 	 * Post loads all default materials.
diff --git a/Engine/Source/Runtime/Engine/Private/EnginePrivate.h b/Engine/Source/Runtime/Engine/Private/EnginePrivate.h
--- a/Engine/Source/Runtime/Engine/Private/EnginePrivate.h
+++ b/Engine/Source/Runtime/Engine/Private/EnginePrivate.h
@@ -38,6 +38,7 @@
 #include "Curves/CurveBase.h"
 #include "Curves/CurveFloat.h"
 #include "GameFramework/ForceFeedbackEffect.h"
+#include "Engine/SubSurfaceProfile.h"
 #include "GameFramework/PlayerController.h"
 #include "Engine/DebugCameraController.h"
 #include "Engine/DecalActor.h"
diff --git a/Engine/Source/Runtime/Engine/Private/Materials/HLSLMaterialTranslator.h b/Engine/Source/Runtime/Engine/Private/Materials/HLSLMaterialTranslator.h
--- a/Engine/Source/Runtime/Engine/Private/Materials/HLSLMaterialTranslator.h
+++ b/Engine/Source/Runtime/Engine/Private/Materials/HLSLMaterialTranslator.h
@@ -336,14 +336,28 @@ public:
 				Chunk[MP_WorldDisplacement]			= Constant3(0.0f, 0.0f, 0.0f);
 			}
 			Chunk[MP_TessellationMultiplier]		= Material->CompilePropertyAndSetMaterialProperty(MP_TessellationMultiplier,this);
-			Chunk[MP_SubsurfaceColor]				= Material->CompilePropertyAndSetMaterialProperty(MP_SubsurfaceColor       ,this);
+
+			EMaterialShadingModel MaterialShadingModel = Material->GetShadingModel();
+
+			if (MaterialShadingModel == MSM_Subsurface || MaterialShadingModel == MSM_PreintegratedSkin || MaterialShadingModel == MSM_SubsurfaceProfile)
+			{
+				int32 SubsurfaceColor = Material->CompilePropertyAndSetMaterialProperty(MP_SubsurfaceColor, this);
+
+				static FName SubsurfaceProfileName(TEXT("SubsurfaceProfile"));
+
+				// 1.0f is is a not used profile - later this gets replaced with the actual profile
+				int32 SubsurfaceProfile = ForceCast(ScalarParameter(SubsurfaceProfileName, 1.0f), MCT_Float1);
+
+				Chunk[MP_SubsurfaceColor] = AppendVector(SubsurfaceColor, SubsurfaceProfile);		
+			}
+
 			Chunk[MP_ClearCoat]						= Material->CompilePropertyAndSetMaterialProperty(MP_ClearCoat			   ,this);
 			Chunk[MP_ClearCoatRoughness]			= Material->CompilePropertyAndSetMaterialProperty(MP_ClearCoatRoughness    ,this);
 			Chunk[MP_AmbientOcclusion]				= Material->CompilePropertyAndSetMaterialProperty(MP_AmbientOcclusion      ,this);
 
 			{
 				int32 UserRefraction = ForceCast(Material->CompilePropertyAndSetMaterialProperty(MP_Refraction, this), MCT_Float2);
-				int32 RefractionDepthBias = ForceCast(ScalarParameter(FName(TEXT("RefractionDepthBias")), GetRefractionDepthBiasValue()), MCT_Float1);
+				int32 RefractionDepthBias = ForceCast(ScalarParameter(FName(TEXT("RefractionDepthBias")), Material->GetRefractionDepthBiasValue()), MCT_Float1);
 
 				Chunk[MP_Refraction] = AppendVector(ForceCast(UserRefraction, MCT_Float1), RefractionDepthBias);
 			}
@@ -362,10 +376,6 @@ public:
 				{
 					Chunk[CustomUVIndex] = Material->CompilePropertyAndSetMaterialProperty((EMaterialProperty)CustomUVIndex, this);
 				}
-				else
-				{
-					Chunk[CustomUVIndex] = INDEX_NONE;
-				}
 			}
 
 			check(FunctionStack.Num() == 1);
@@ -395,7 +405,7 @@ public:
 			}
 
 
-			if (Material->GetBlendMode() == BLEND_Modulate && Material->GetShadingModel() != MSM_Unlit && !Material->IsUsedWithDeferredDecal())
+			if (Material->GetBlendMode() == BLEND_Modulate && MaterialShadingModel != MSM_Unlit && !Material->IsUsedWithDeferredDecal())
 			{
 				Errorf(TEXT("Dynamically lit translucency is not supported for BLEND_Modulate materials."));
 			}
@@ -431,12 +441,12 @@ public:
 				Errorf(TEXT("Light function materials must be opaque."));
 			}
 
-			if (Material->IsLightFunction() && Material->GetShadingModel() != MSM_Unlit)
+			if (Material->IsLightFunction() && MaterialShadingModel != MSM_Unlit)
 			{
 				Errorf(TEXT("Light function materials must use unlit."));
 			}
 
-			if (Material->GetMaterialDomain() == MD_PostProcess && Material->GetShadingModel() != MSM_Unlit)
+			if (Material->GetMaterialDomain() == MD_PostProcess && MaterialShadingModel != MSM_Unlit)
 			{
 				Errorf(TEXT("Post process materials must use unlit."));
 			}
@@ -511,6 +521,13 @@ public:
 
 			MaterialCompilationOutput.UniformExpressionSet.SetParameterCollections(ParameterCollections);
 
+/* test
+			if(MaterialShadingModel == MSM_SubsurfaceProfile)
+			{
+				MaterialCompilationOutput.UniformExpressionSet.SetCompileInSubsurfaceProfile();
+			}
+*/
+
 			// Create the material uniform buffer struct.
 			MaterialCompilationOutput.UniformExpressionSet.CreateBufferStruct();
 		}
@@ -1292,11 +1309,6 @@ protected:
 		return FeatureLevel;
 	}
 
-	virtual float GetRefractionDepthBiasValue() override
-	{
-		return Material->GetRefractionDepthBiasValue();
-	}
-
 	/** 
 	 * Casts the passed in code to DestType, or generates a compile error if the cast is not valid. 
 	 * This will truncate a type (float4 -> float3) but not add components (float2 -> float3), however a float1 can be cast to any float type by replication. 
diff --git a/Engine/Source/Runtime/Engine/Private/Materials/Material.cpp b/Engine/Source/Runtime/Engine/Private/Materials/Material.cpp
--- a/Engine/Source/Runtime/Engine/Private/Materials/Material.cpp
+++ b/Engine/Source/Runtime/Engine/Private/Materials/Material.cpp
@@ -23,6 +23,7 @@
 #include "Materials/MaterialExpressionStaticSwitchParameter.h"
 #include "Materials/MaterialExpressionTextureSampleParameter.h"
 #include "Materials/MaterialExpressionVectorParameter.h"
+#include "Engine/SubsurfaceProfile.h"
 #include "EditorSupportDelegates.h"
 #include "MaterialShaderType.h"
 #include "MaterialInstanceSupport.h"
@@ -205,6 +206,24 @@ public:
 		const FMaterialResource* MaterialResource = Material->GetMaterialResource(Context.Material.GetFeatureLevel());
 		if(MaterialResource && MaterialResource->GetRenderingThreadShaderMap())
 		{
+			static FName SubsurfaceProfile(TEXT("SubsurfaceProfile"));
+			if (ParameterName == SubsurfaceProfile)
+			{
+				const USubsurfaceProfile* SubsurfaceProfileRT = GetSubsurfaceProfile();
+
+				if(SubsurfaceProfileRT)
+				{
+					// can be optimized (cached)
+					*OutValue = GSubsufaceProfileTextureObject.FindAllocationId(*SubsurfaceProfileRT) / 255.0f;
+				}
+				else
+				{
+					// no profile specified means we use the default one stored at [0] which is human skin
+					*OutValue = 0.0f;
+				}
+				return true;
+			}
+
 			return false;
 		}
 		else
@@ -258,6 +277,8 @@ private:
 	}
 
 	UMaterial* Material;
+
+	// maintained by the render thread
 	float DistanceFieldPenumbraScale;
 };
 
@@ -1989,6 +2010,8 @@ void UMaterial::PostLoad()
 	}
 #endif // #if WITH_EDITOR
 
+	PropagateDataToMaterialProxy();
+
 	STAT(double MaterialLoadTime = 0);
 	{
 		SCOPE_SECONDS_COUNTER(MaterialLoadTime);
@@ -2031,13 +2054,6 @@ void UMaterial::PostLoad()
 		UpdateLightmassTextureTracking();
 	}
 
-	for (int32 i = 0; i < ARRAY_COUNT(DefaultMaterialInstances); i++)
-	{
-		if (DefaultMaterialInstances[i])
-		{
-			DefaultMaterialInstances[i]->GameThread_UpdateDistanceFieldPenumbraScale(GetDistanceFieldPenumbraScale());
-		}
-	}
 
 #if WITH_EDITOR
 	if (GMaterialsThatNeedExpressionsFlipped.Get(this))
@@ -2062,6 +2078,19 @@ void UMaterial::PostLoad()
 #endif // #if WITH_EDITOR
 }
 
+void UMaterial::PropagateDataToMaterialProxy()
+{
+	for (int32 i = 0; i < ARRAY_COUNT(DefaultMaterialInstances); i++)
+	{
+		if (DefaultMaterialInstances[i])
+		{
+			DefaultMaterialInstances[i]->GameThread_UpdateDistanceFieldPenumbraScale(GetDistanceFieldPenumbraScale());
+
+			UpdateMaterialRenderProxy(*DefaultMaterialInstances[i]);
+		}
+	}
+}
+
 void UMaterial::BeginCacheForCookedPlatformData( const ITargetPlatform *TargetPlatform )
 {
 	TArray<FName> DesiredShaderFormats;
@@ -2191,6 +2220,11 @@ bool UMaterial::CanEditChange(const UProperty* InProperty) const
 		{
 			return IsTranslucentBlendMode(BlendMode) && ShadingModel != MSM_Unlit;
 		}
+
+		if (PropertyName == GET_MEMBER_NAME_STRING_CHECKED(UMaterial, SubsurfaceProfile))
+		{
+			return MaterialDomain == MD_Surface && ShadingModel == MSM_SubsurfaceProfile && (BlendMode == BLEND_Opaque || BlendMode == BLEND_Masked);
+		}
 	}
 
 	return true;
@@ -2269,14 +2303,8 @@ void UMaterial::PostEditChangeProperty(FPropertyChangedEvent& PropertyChangedEve
 			FGlobalComponentReregisterContext RecreateComponents;
 		}
 	}
-	
-	for (int32 i = 0; i < ARRAY_COUNT(DefaultMaterialInstances); i++)
-	{
-		if (DefaultMaterialInstances[i])
-		{
-			DefaultMaterialInstances[i]->GameThread_UpdateDistanceFieldPenumbraScale(GetDistanceFieldPenumbraScale());
-		}
-	}
+
+	PropagateDataToMaterialProxy();
 
 	// many property changes can require rebuild of graph so always mark as changed
 	// not interested in PostEditChange calls though as the graph may have instigated it
@@ -3524,7 +3552,7 @@ bool UMaterial::IsPropertyActive(EMaterialProperty InProperty)const
 		break;
 	case MP_Opacity:
 		Active = IsTranslucentBlendMode((EBlendMode)BlendMode) && BlendMode != BLEND_Modulate;
-		if(ShadingModel == MSM_Subsurface || ShadingModel == MSM_PreintegratedSkin)
+		if (ShadingModel == MSM_Subsurface || ShadingModel == MSM_PreintegratedSkin || ShadingModel == MSM_SubsurfaceProfile)
 		{
 			Active = true;
 		}
diff --git a/Engine/Source/Runtime/Engine/Private/Materials/MaterialInstance.cpp b/Engine/Source/Runtime/Engine/Private/Materials/MaterialInstance.cpp
--- a/Engine/Source/Runtime/Engine/Private/Materials/MaterialInstance.cpp
+++ b/Engine/Source/Runtime/Engine/Private/Materials/MaterialInstance.cpp
@@ -280,6 +280,19 @@ void FMaterialInstanceResource::GameThread_UpdateDistanceFieldPenumbraScale(floa
 	});
 }
 
+void UMaterialInstance::UpdateMaterialInstanceData()
+{
+	for (int32 i = 0; i < ARRAY_COUNT(Resources); i++)
+	{
+		if (Resources[i])
+		{
+			Resources[i]->GameThread_UpdateDistanceFieldPenumbraScale(GetDistanceFieldPenumbraScale());
+
+			UpdateMaterialRenderProxy(*Resources[i]);
+		}
+	}
+}
+
 void FMaterialInstanceResource::GameThread_SetParent(UMaterialInterface* InParent)
 {
 	check(IsInGameThread());
@@ -1690,13 +1703,7 @@ void UMaterialInstance::PostLoad()
 		UpdateLightmassTextureTracking();
 	}
 
-	for (int32 i = 0; i < ARRAY_COUNT(Resources); i++)
-	{
-		if (Resources[i])
-		{
-			Resources[i]->GameThread_UpdateDistanceFieldPenumbraScale(GetDistanceFieldPenumbraScale());
-		}
-	}
+	UpdateMaterialInstanceData();
 
 	// Fixup for legacy instances which didn't recreate the lighting guid properly on duplication
 	if (GetLinker() && GetLinker()->UE4Ver() < VER_UE4_BUMPED_MATERIAL_EXPORT_GUIDS)
@@ -2003,13 +2010,7 @@ void UMaterialInstance::PostEditChangeProperty(FPropertyChangedEvent& PropertyCh
 		UpdateLightmassTextureTracking();
 	}
 
-	for (int32 i = 0; i < ARRAY_COUNT(Resources); i++)
-	{
-		if (Resources[i])
-		{
-			Resources[i]->GameThread_UpdateDistanceFieldPenumbraScale(GetDistanceFieldPenumbraScale());
-		}
-	}
+	UpdateMaterialInstanceData();
 
 	InitResources();
 
diff --git a/Engine/Source/Runtime/Engine/Private/Materials/MaterialInterface.cpp b/Engine/Source/Runtime/Engine/Private/Materials/MaterialInterface.cpp
--- a/Engine/Source/Runtime/Engine/Private/Materials/MaterialInterface.cpp
+++ b/Engine/Source/Runtime/Engine/Private/Materials/MaterialInterface.cpp
@@ -5,6 +5,7 @@
 =============================================================================*/
 
 #include "EnginePrivate.h"
+#include "Engine/SubsurfaceProfile.h"
 
 //////////////////////////////////////////////////////////////////////////
 
@@ -16,6 +17,7 @@ void FMaterialRelevance::SetPrimitiveViewRelevance(FPrimitiveViewRelevance& OutV
 	OutViewRelevance.bDistortionRelevance = bDistortion;
 	OutViewRelevance.bSeparateTranslucencyRelevance = bSeparateTranslucency;
 	OutViewRelevance.bNormalTranslucencyRelevance = bNormalTranslucency;
+	OutViewRelevance.bSubsurfaceProfileRelevance = bSubsurfaceProfile;
 }
 
 //////////////////////////////////////////////////////////////////////////
@@ -51,6 +53,7 @@ FMaterialRelevance UMaterialInterface::GetRelevance_Internal(const UMaterial* Ma
 		MaterialRelevance.bSeparateTranslucency = bIsTranslucent && Material->bEnableSeparateTranslucency;
 		MaterialRelevance.bNormalTranslucency = bIsTranslucent && !Material->bEnableSeparateTranslucency;
 		MaterialRelevance.bDisableDepthTest = bIsTranslucent && Material->bDisableDepthTest;
+		MaterialRelevance.bSubsurfaceProfile = (Material->MaterialDomain == MD_Surface) && !bIsTranslucent && (Material->GetShadingModel() == MSM_SubsurfaceProfile);
 		return MaterialRelevance;
 	}
 	else
@@ -325,3 +328,60 @@ uint32 UMaterialInterface::GetFeatureLevelsToCompileForRendering() const
 {
 	return FeatureLevelsToForceCompile | (1 << GRHIFeatureLevel);
 }
+
+
+void UMaterialInterface::UpdateMaterialRenderProxy(FMaterialRenderProxy& Proxy)
+{
+	// no 0 pointer
+	check(&Proxy);
+
+	FSubsurfaceProfileStruct Settings;
+
+	if (SubsurfaceProfile)
+	{
+		Settings = SubsurfaceProfile->Settings;
+	}
+
+	// this can be improved, it doesn't support Renderer hot reload
+	{
+		static bool bFirst = true;
+
+		if(bFirst)
+		{
+			bFirst = false;
+
+			static const FName RendererModuleName("Renderer");
+			IRendererModule& RendererModule = FModuleManager::GetModuleChecked<IRendererModule>(RendererModuleName);
+
+			ENQUEUE_UNIQUE_RENDER_COMMAND_ONEPARAMETER(
+				UpdateMaterialRenderProxySubsurfaceSetup,
+				IRendererModule&, RendererModule, RendererModule,
+			{
+				GSubsufaceProfileTextureObject.SetRendererModule(&RendererModule);
+			});
+		}
+	}
+
+	EMaterialShadingModel MaterialShadingModel = GetShadingModel_Internal();
+
+	// for better performance we only update SubsurfaceProfileRT if the feature is used
+	if (MaterialShadingModel == MSM_SubsurfaceProfile)
+	{
+		ENQUEUE_UNIQUE_RENDER_COMMAND_THREEPARAMETER(
+			UpdateMaterialRenderProxySubsurface,
+			const FSubsurfaceProfileStruct, Settings, Settings,
+			const USubsurfaceProfile*, SubsurfaceProfile, SubsurfaceProfile,
+			FMaterialRenderProxy&, Proxy, Proxy,
+		{
+			uint32 AllocationId = 0;
+
+			if(SubsurfaceProfile)
+			{
+				AllocationId = GSubsufaceProfileTextureObject.AddOrUpdateProfile(Settings, *SubsurfaceProfile);
+
+				check(AllocationId >= 0 && AllocationId <= 255);
+			}
+			Proxy.SetSubsurfaceProfile(SubsurfaceProfile);
+		});
+	}
+}
diff --git a/Engine/Source/Runtime/Engine/Private/Materials/MaterialShader.cpp b/Engine/Source/Runtime/Engine/Private/Materials/MaterialShader.cpp
--- a/Engine/Source/Runtime/Engine/Private/Materials/MaterialShader.cpp
+++ b/Engine/Source/Runtime/Engine/Private/Materials/MaterialShader.cpp
@@ -40,6 +40,7 @@ FString GetShadingModelString(EMaterialShadingModel ShadingModel)
 		case MSM_Subsurface:		ShadingModelName = TEXT("MSM_Subsurface"); break;
 		case MSM_PreintegratedSkin:	ShadingModelName = TEXT("MSM_PreintegratedSkin"); break;
 		case MSM_ClearCoat:			ShadingModelName = TEXT("MSM_ClearCoat"); break;
+		case MSM_SubsurfaceProfile:	ShadingModelName = TEXT("MSM_SubsurfaceProfile"); break;
 		default: ShadingModelName = TEXT("Unknown"); break;
 	}
 	return ShadingModelName;
@@ -82,6 +83,7 @@ void UpdateMaterialShaderCompilingStats(const FMaterial* Material)
 		case MSM_Subsurface:
 		case MSM_PreintegratedSkin:
 		case MSM_ClearCoat:
+		case MSM_SubsurfaceProfile:
 			INC_DWORD_STAT_BY(STAT_ShaderCompiling_NumLitMaterialShaders,1);
 			break;
 		default: break;
diff --git a/Engine/Source/Runtime/Engine/Private/Materials/MaterialShared.cpp b/Engine/Source/Runtime/Engine/Private/Materials/MaterialShared.cpp
--- a/Engine/Source/Runtime/Engine/Private/Materials/MaterialShared.cpp
+++ b/Engine/Source/Runtime/Engine/Private/Materials/MaterialShared.cpp
@@ -1126,6 +1126,7 @@ void FMaterial::SetupMaterialEnvironment(
 		case MSM_DefaultLit:		OutEnvironment.SetDefine(TEXT("MATERIAL_SHADINGMODEL_DEFAULT_LIT"),			TEXT("1")); break;
 		case MSM_Subsurface:		OutEnvironment.SetDefine(TEXT("MATERIAL_SHADINGMODEL_SUBSURFACE"),			TEXT("1")); break;
 		case MSM_PreintegratedSkin: OutEnvironment.SetDefine(TEXT("MATERIAL_SHADINGMODEL_PREINTEGRATED_SKIN"),	TEXT("1")); break;
+		case MSM_SubsurfaceProfile: OutEnvironment.SetDefine(TEXT("MATERIAL_SHADINGMODEL_SUBSURFACE_PROFILE"),	TEXT("1")); break;
 		case MSM_ClearCoat:			OutEnvironment.SetDefine(TEXT("MATERIAL_SHADINGMODEL_CLEAR_COAT"),			TEXT("1")); break;
 		default: 
 			UE_LOG(LogMaterial, Warning, TEXT("Unknown material shading model: %u  Setting to MSM_DefaultLit"),(int32)GetShadingModel());
@@ -1545,12 +1546,14 @@ void FMaterialRenderProxy::InvalidateUniformExpressionCache()
 FMaterialRenderProxy::FMaterialRenderProxy()
 	: bSelected(false)
 	, bHovered(false)
+	, SubsurfaceProfileRT(0)
 {
 }
 
 FMaterialRenderProxy::FMaterialRenderProxy(bool bInSelected, bool bInHovered)
 	: bSelected(bInSelected)
 	, bHovered(bInHovered)
+	, SubsurfaceProfileRT(0)
 {
 }
 
diff --git a/Engine/Source/Runtime/Engine/Private/Rendering/SeparableSSS.cpp b/Engine/Source/Runtime/Engine/Private/Rendering/SeparableSSS.cpp
new file mode 100644
--- /dev/null
+++ b/Engine/Source/Runtime/Engine/Private/Rendering/SeparableSSS.cpp
@@ -0,0 +1,181 @@
+// Copyright 1998-2014 Epic Games, Inc. All Rights Reserved.
+
+/*=============================================================================
+SeparableSSS.cpp: Computing the kernel for the Separable Screen Space Subsurface Scattering, based on SeparableSSS, see copyright below
+=============================================================================*/
+
+/**
+ * Copyright (C) 2012 [REDACTED_PERSON] ([REDACTED_EMAIL])
+ * Copyright (C) 2012 Diego Gutierrez ([REDACTED_EMAIL])
+ * All rights reserved.
+ *
+ * Redistribution and use in source and binary forms, with or without
+ * modification, are permitted provided that the following conditions are met:
+ * 
+ *    1. Redistributions of source code must retain the above copyright notice,
+ *       this list of conditions and the following disclaimer.
+ *
+ *    2. Redistributions in binary form must reproduce the following disclaimer
+ *       in the documentation and/or other materials provided with the 
+ *       distribution:
+ *
+ *       "Uses Separable SSS. Copyright (C) 2012 by [REDACTED_PERSON] and Diego
+ *        Gutierrez."
+ *
+ * THIS SOFTWARE IS PROVIDED BY THE COPYRIGHT HOLDERS AND CONTRIBUTORS ``AS 
+ * IS'' AND ANY EXPRESS OR IMPLIED WARRANTIES, INCLUDING, BUT NOT LIMITED TO, 
+ * THE IMPLIED WARRANTIES OF MERCHANTABILITY AND FITNESS FOR A PARTICULAR 
+ * PURPOSE ARE DISCLAIMED. IN NO EVENT SHALL COPYRIGHT HOLDERS OR CONTRIBUTORS 
+ * BE LIABLE FOR ANY DIRECT, INDIRECT, INCIDENTAL, SPECIAL, EXEMPLARY, OR 
+ * CONSEQUENTIAL DAMAGES (INCLUDING, BUT NOT LIMITED TO, PROCUREMENT OF 
+ * SUBSTITUTE GOODS OR SERVICES; LOSS OF USE, DATA, OR PROFITS; OR BUSINESS 
+ * INTERRUPTION) HOWEVER CAUSED AND ON ANY THEORY OF LIABILITY, WHETHER IN 
+ * CONTRACT, STRICT LIABILITY, OR TORT (INCLUDING NEGLIGENCE OR OTHERWISE) 
+ * ARISING IN ANY WAY OUT OF THE USE OF THIS SOFTWARE, EVEN IF ADVISED OF THE 
+ * POSSIBILITY OF SUCH DAMAGE.
+ *
+ * The views and conclusions contained in the software and documentation are 
+ * those of the authors and should not be interpreted as representing official
+ * policies, either expressed or implied, of the copyright holders.
+ */
+
+
+#include "EnginePrivate.h"
+#include "SeparableSSS.h"
+
+// helper function for ComputeMirroredSSSKernel
+inline FVector SeparableSSS_Gaussian(float variance, float r, FLinearColor FalloffColor)
+{
+	FVector Ret;
+
+	/**
+	* We use a falloff to modulate the shape of the profile. Big falloffs
+	* spreads the shape making it wider, while small falloffs make it
+	* narrower.
+	*/
+	for (int i = 0; i < 3; i++)
+	{
+		float rr = r / (0.001f + FalloffColor.Component(i));
+		Ret[i] = exp((-(rr * rr)) / (2.0f * variance)) / (2.0f * 3.14f * variance);
+	}
+
+	return Ret;
+}
+
+// helper function for ComputeMirroredSSSKernel
+// r is in mm
+inline FVector SeparableSSS_Profile(float r, FLinearColor FalloffColor)
+{
+	/**
+	* We used the red channel of the original skin profile defined in
+	* [d'Eon07] for all three channels. We noticed it can be used for green
+	* and blue channels (scaled using the falloff parameter) without
+	* introducing noticeable differences and allowing for total control over
+	* the profile. For example, it allows to create blue SSS gradients, which
+	* could be useful in case of rendering blue creatures.
+	*/
+	// first parameter is variance in mm^2
+	return  // 0.233f * SeparableSSS_Gaussian(0.0064f, r, FalloffColor) + /* We consider this one to be directly bounced light, accounted by the strength parameter (see @STRENGTH) */
+		0.100f * SeparableSSS_Gaussian(0.0484f, r, FalloffColor) +
+		0.118f * SeparableSSS_Gaussian(0.187f, r, FalloffColor) +
+		0.113f * SeparableSSS_Gaussian(0.567f, r, FalloffColor) +
+		0.358f * SeparableSSS_Gaussian(1.99f, r, FalloffColor) +
+		0.078f * SeparableSSS_Gaussian(7.41f, r, FalloffColor);
+}
+
+inline void ComputeMirroredSSSKernel(TArray<FLinearColor>& Out, FLinearColor SubsurfaceColor, FLinearColor FalloffColor)
+{
+	// needs to be preallocated with the expected size e.g. 5 for a 9 sample kernel with one center and 4 left, 4 right and one center sample
+	check(Out.Num());
+
+	uint32 nNonMirroredSamples = Out.Num();
+	int32 nTotalSamples = nNonMirroredSamples * 2 - 1;
+
+	// we could generate Out directly but the original code form SeparableSSS wasn't done like that so we convert it later
+	// .a is in mm
+	TArray<FLinearColor> kernel;
+	{
+
+		kernel.AddZeroed(nTotalSamples);
+
+		const float Range = nTotalSamples > 20 ? 3.0f : 2.0f;
+		// tweak constant
+		const float Exponent = 2.0f;
+
+		// Calculate the offsets:
+		float step = 2.0f * Range / (nTotalSamples - 1);
+		for (int i = 0; i < nTotalSamples; i++)
+		{
+			float o = -Range + float(i) * step;
+			float sign = o < 0.0f ? -1.0f : 1.0f;
+			kernel[i].A = Range * sign * abs(pow(o, Exponent)) / pow(Range, Exponent);
+		}
+
+		// Calculate the weights:
+		for (int32 i = 0; i < nTotalSamples; i++)
+		{
+			float w0 = i > 0 ? abs(kernel[i].A - kernel[i - 1].A) : 0.0f;
+			float w1 = i < nTotalSamples - 1 ? abs(kernel[i].A - kernel[i + 1].A) : 0.0f;
+			float area = (w0 + w1) / 2.0f;
+			FVector t = area * SeparableSSS_Profile(kernel[i].A, FalloffColor);
+			kernel[i].R = t.X;
+			kernel[i].G = t.Y;
+			kernel[i].B = t.Z;
+		}
+
+		// We want the offset 0.0 to come first:
+		FLinearColor t = kernel[nTotalSamples / 2];
+
+		for (int i = nTotalSamples / 2; i > 0; i--)
+		{
+			kernel[i] = kernel[i - 1];
+		}
+		kernel[0] = t;
+
+		// Normalize the weights in RGB
+		{
+			FVector sum = FVector(0, 0, 0);
+
+			for (int i = 0; i < nTotalSamples; i++)
+			{
+				sum.X += kernel[i].R;
+				sum.Y += kernel[i].G;
+				sum.Z += kernel[i].B;
+			}
+
+			for (int i = 0; i < nTotalSamples; i++)
+			{
+				kernel[i].R /= sum.X;
+				kernel[i].G /= sum.Y;
+				kernel[i].B /= sum.Z;
+			}
+		}
+
+		// Tweak them using the desired strength. The first one is:
+		//     lerp(1.0, kernel[0].rgb, strength)
+		kernel[0].R = (1.0f - SubsurfaceColor.R) * 1.0f + SubsurfaceColor.R * kernel[0].R;
+		kernel[0].G = (1.0f - SubsurfaceColor.G) * 1.0f + SubsurfaceColor.G * kernel[0].G;
+		kernel[0].B = (1.0f - SubsurfaceColor.B) * 1.0f + SubsurfaceColor.B * kernel[0].B;
+
+		for (int i = 1; i < nTotalSamples; i++)
+		{
+			kernel[i].R *= SubsurfaceColor.R;
+			kernel[i].G *= SubsurfaceColor.G;
+			kernel[i].B *= SubsurfaceColor.B;
+		}
+	}
+
+	// generate output (remove negative samples)
+	{
+		check(kernel[0].A == 0.0f);
+
+		// center sample
+		Out[0] = kernel[0];
+
+		// all positive samples
+		for (uint32 i = 0; i < nNonMirroredSamples - 1; i++)
+		{
+			Out[i + 1] = kernel[nNonMirroredSamples + i];
+		}
+	}
+}
\ No newline at end of file
diff --git a/Engine/Source/Runtime/Engine/Private/Rendering/SeparableSSS.h b/Engine/Source/Runtime/Engine/Private/Rendering/SeparableSSS.h
new file mode 100644
--- /dev/null
+++ b/Engine/Source/Runtime/Engine/Private/Rendering/SeparableSSS.h
@@ -0,0 +1,14 @@
+// Copyright 1998-2014 Epic Games, Inc. All Rights Reserved.
+
+/*=============================================================================
+SeparableSSS.h: Computing the kernel for the Separable Screen Space Subsurface Scattering
+=============================================================================*/
+
+#pragma once
+
+// @param Out, needs to be preallocated with the expected size (RGB is the sample weight, A is the offset), [0] is the center samples, following elements need to be mirrored with A, -A
+// @parma SubsurfaceColor see SubsurfaceProfile.h
+// @parma FalloffColor see SubsurfaceProfile.h
+void ComputeMirroredSSSKernel(TArray<FLinearColor>& Out, FLinearColor SubsurfaceColor, FLinearColor FalloffColor);
+
+
diff --git a/Engine/Source/Runtime/Engine/Private/Rendering/SubsurfaceProfile.cpp b/Engine/Source/Runtime/Engine/Private/Rendering/SubsurfaceProfile.cpp
new file mode 100644
--- /dev/null
+++ b/Engine/Source/Runtime/Engine/Private/Rendering/SubsurfaceProfile.cpp
@@ -0,0 +1,263 @@
+// Copyright 1998-2014 Epic Games, Inc. All Rights Reserved.
+
+#include "EnginePrivate.h"
+#include "SeparableSSS.h"
+
+DEFINE_LOG_CATEGORY_STATIC(LogSubsurfaceProfile, Log, All);
+
+// lives on the render thread
+TGlobalResource<FSubsurfaceProfileTexture> GSubsufaceProfileTextureObject;
+
+// Texture with one or more SubSurfaceProfiles or 0 if there is no user
+static TRefCountPtr<IPooledRenderTarget> GSSProfiles;
+
+
+FSubsurfaceProfileTexture::FSubsurfaceProfileTexture()
+	: RendererModule(0)
+{
+	FSubsurfaceProfileStruct DefaultSkin;
+
+	// add element 0, it is used as default profile
+	SubsurfaceProfileEntries.Add(FSubsurfaceProfileEntry(DefaultSkin, 0));
+}
+
+FSubsurfaceProfileTexture::~FSubsurfaceProfileTexture()
+{
+	// we assume all ~USubsurfaceProfile() have bene called already
+	for (int32 i = 0; i < SubsurfaceProfileEntries.Num(); ++i)
+	{
+		check(SubsurfaceProfileEntries[i].GameThreadObject == 0);
+	}
+}
+
+int32 FSubsurfaceProfileTexture::AddProfile(const FSubsurfaceProfileStruct Settings, const USubsurfaceProfile& InProfile)
+{
+	check(&InProfile);
+	check(FindAllocationId(InProfile) == -1);
+
+	int32 RetAllocationId = -1;
+	{
+		for (int32 i = 1; i < SubsurfaceProfileEntries.Num(); ++i)
+		{
+			if (SubsurfaceProfileEntries[i].GameThreadObject == 0)
+			{
+				RetAllocationId = i; break;
+			}
+		}
+
+		if(RetAllocationId == -1)
+		{
+			RetAllocationId = SubsurfaceProfileEntries.Num();
+			SubsurfaceProfileEntries.Add(FSubsurfaceProfileEntry(Settings, &InProfile));
+		}
+	}
+
+	UpdateProfile(RetAllocationId, Settings);
+
+	return RetAllocationId;
+}
+
+
+void FSubsurfaceProfileTexture::RemoveProfile(USubsurfaceProfile& InProfile)
+{
+	int32 AllocationId = FindAllocationId(InProfile);
+
+	if(AllocationId == -1)
+	{
+		// -1: no allocation, no work needed
+		return;
+	}
+
+	// >0 as 0 is used as default profile which should never be removed
+	check(AllocationId > 0);
+
+	check(SubsurfaceProfileEntries[AllocationId].GameThreadObject == &InProfile);
+
+	// make it available for reuse
+	SubsurfaceProfileEntries[AllocationId].GameThreadObject = 0;
+}
+
+void FSubsurfaceProfileTexture::UpdateProfile(int32 AllocationId, const FSubsurfaceProfileStruct Settings)
+{
+	check(IsInRenderingThread());
+
+	// call SetRendererModule() is missing
+	check(RendererModule);
+
+	if (AllocationId == -1)
+	{
+		// if we modify a profile that is not assigned/used yet, no work is needed
+		return;
+	}
+
+	check(AllocationId < SubsurfaceProfileEntries.Num());
+
+	SubsurfaceProfileEntries[AllocationId].Settings = Settings;
+
+	GSSProfiles.SafeRelease();
+}
+
+const FSceneRenderTargetItem* FSubsurfaceProfileTexture::GetTexture(FRHICommandListImmediate& RHICmdList)
+{
+	// call SetRendererModule() is missing
+	check(RendererModule);
+
+	if (!GSSProfiles)
+	{
+		CreateTexture(RHICmdList);
+	}
+
+	return &GSSProfiles->GetRenderTargetItem();
+}
+
+void FSubsurfaceProfileTexture::ReleaseDynamicRHI()
+{
+	GSSProfiles.SafeRelease();
+}
+
+void FSubsurfaceProfileTexture::CreateTexture(FRHICommandListImmediate& RHICmdList)
+{
+	// call SetRendererModule() is missing
+	check(RendererModule);
+
+	uint32 Height = SubsurfaceProfileEntries.Num();
+
+	check(Height);
+
+	// true:16bit (currently required to have very small and very large kernel sizes), false: 8bit
+	const bool b16Bit = true;
+
+	// at minimum 64 lines (less reallocations)
+	FPooledRenderTargetDesc Desc(FPooledRenderTargetDesc::Create2DDesc(FIntPoint(16, FMath::Max(Height, (uint32)64)), PF_B8G8R8A8, TexCreate_None, TexCreate_None, false));
+
+	if (b16Bit)
+	{
+		Desc.Format = PF_A16B16G16R16;
+	}
+
+	RendererModule->RenderTargetPoolFindFreeElement(Desc, GSSProfiles, TEXT("SSProfiles"));
+
+	// Write the contents of the texture.
+	uint32 DestStride;
+	uint8* DestBuffer = (uint8*)RHICmdList.LockTexture2D((FTexture2DRHIRef&)GSSProfiles->GetRenderTargetItem().ShaderResourceTexture, 0, RLM_WriteOnly, DestStride, false);
+
+	TArray<FLinearColor> kernel;
+	kernel.AddZeroed(13);	// highest quality
+	// todo: for optimizations
+//	kernel.AddZeroed(9);
+//	kernel.AddZeroed(6);	// lowest quality
+
+	for (uint32 y = 0; y < Height; ++y)
+	{
+		FSubsurfaceProfileStruct& Data = SubsurfaceProfileEntries[y].Settings;
+
+		ComputeMirroredSSSKernel(kernel, Data.SubsurfaceColor, Data.FalloffColor);
+
+		const float TableMaxRGB = 1.0f;
+		const float TableMaxA = 3.0f;
+
+		for (int32 Pos = 0; Pos < kernel.Num(); ++Pos)
+		{
+			FVector4 C = kernel[Pos] * FLinearColor(1.0f / TableMaxRGB, 1.0f / TableMaxRGB, 1.0f / TableMaxRGB, 1.0f / TableMaxA);
+
+			// requires 16bit (could be made with 8 bit e.g. using sample0.w as 8bit scale applied to all samples (more multiplications in the shader))
+			C.W *= Data.ScatterRadius / 1000.0f;
+
+			if (b16Bit)
+			{
+				uint16* Dest = (uint16*)(DestBuffer + DestStride * y);
+
+				Dest[Pos * 4 + 0] = (uint16)(C.X * (256 * 256 - 0.0001f));
+				Dest[Pos * 4 + 1] = (uint16)(C.Y * (256 * 256 - 0.0001f));
+				Dest[Pos * 4 + 2] = (uint16)(C.Z * (256 * 256 - 0.0001f));
+				Dest[Pos * 4 + 3] = (uint16)(C.W * (256 * 256 - 0.0001f));
+			}
+			else
+			{
+				FColor* Dest = (FColor*)(DestBuffer + DestStride * y);
+
+				Dest[Pos] = FColor(FMath::Quantize8UnsignedByte(C.X), FMath::Quantize8UnsignedByte(C.Y), FMath::Quantize8UnsignedByte(C.Z), FMath::Quantize8UnsignedByte(C.W));
+			}
+		}
+	}
+
+	RHICmdList.UnlockTexture2D((FTexture2DRHIRef&)GSSProfiles->GetRenderTargetItem().ShaderResourceTexture, 0, false);
+}
+
+int32 FSubsurfaceProfileTexture::FindAllocationId(const USubsurfaceProfile& InProfile) const
+{
+	// we start at 1 because [0] is the default profile and always [0].GameThreadObject = 0 so we don't need to iterate that one
+	for (int32 i = 1; i < SubsurfaceProfileEntries.Num(); ++i)
+	{
+		if (SubsurfaceProfileEntries[i].GameThreadObject == &InProfile)
+		{
+			return i;
+		}
+	}
+
+	return -1;
+}
+
+// for debugging
+void FSubsurfaceProfileTexture::Dump()
+{
+#if !(UE_BUILD_SHIPPING || UE_BUILD_TEST)
+	UE_LOG(LogSubsurfaceProfile, Log, TEXT("USubsurfaceProfileManager::Dump"));
+	for (int32 i = 0; i < SubsurfaceProfileEntries.Num(); ++i)
+	{
+		// + 1 as the Id is one higher than the array index, 0 is used for the default profile (not assigned)
+		UE_LOG(LogSubsurfaceProfile, Log, TEXT("  %d. AllocationId=%d, Pointer=%p"), i, i + 1, SubsurfaceProfileEntries[i].GameThreadObject);
+
+		{
+			UE_LOG(LogSubsurfaceProfile, Log, TEXT("     ScatterRadius = %f"),
+				SubsurfaceProfileEntries[i].Settings.ScatterRadius);
+			UE_LOG(LogSubsurfaceProfile, Log, TEXT("     SubsurfaceColor=%f %f %f"),
+				SubsurfaceProfileEntries[i].Settings.SubsurfaceColor.R, SubsurfaceProfileEntries[i].Settings.SubsurfaceColor.G, SubsurfaceProfileEntries[i].Settings.SubsurfaceColor.B);
+			UE_LOG(LogSubsurfaceProfile, Log, TEXT("     FalloffColor=%f %f %f"),
+				SubsurfaceProfileEntries[i].Settings.FalloffColor.R, SubsurfaceProfileEntries[i].Settings.FalloffColor.G, SubsurfaceProfileEntries[i].Settings.FalloffColor.B);
+		}
+	}
+
+	UE_LOG(LogSubsurfaceProfile, Log, TEXT(""));
+#endif
+}
+
+
+
+ENGINE_API const FSceneRenderTargetItem* GetSubsufaceProfileTexture_RT(FRHICommandListImmediate& RHICmdList)
+{
+	check(IsInRenderingThread());
+
+	return GSubsufaceProfileTextureObject.GetTexture(RHICmdList);
+}
+
+// ------------------------------------------------------
+
+USubsurfaceProfile::USubsurfaceProfile(const class FPostConstructInitializeProperties& PCIP)
+	: Super(PCIP)
+{
+}
+
+void USubsurfaceProfile::BeginDestroy()
+{
+	ENQUEUE_UNIQUE_RENDER_COMMAND_ONEPARAMETER(
+		RemoveSubsurfaceProfile,
+		USubsurfaceProfile&, Ref, *this,
+		{
+			GSubsufaceProfileTextureObject.RemoveProfile(Ref);
+		});
+
+	Super::BeginDestroy();
+}
+
+void USubsurfaceProfile::PostEditChangeProperty(struct FPropertyChangedEvent& PropertyChangedEvent)
+{
+	ENQUEUE_UNIQUE_RENDER_COMMAND_TWOPARAMETER(
+		UpdateSubsurfaceProfile,
+		const FSubsurfaceProfileStruct, Settings, this->Settings,
+		USubsurfaceProfile*, Profile, this,
+	{
+		// any changes to the setting require an update of the texture
+		GSubsufaceProfileTextureObject.UpdateProfile(Settings, *Profile);
+	});
+}
\ No newline at end of file
diff --git a/Engine/Source/Runtime/Engine/Public/Engine.h b/Engine/Source/Runtime/Engine/Public/Engine.h
--- a/Engine/Source/Runtime/Engine/Public/Engine.h
+++ b/Engine/Source/Runtime/Engine/Public/Engine.h
@@ -58,6 +58,7 @@
 #include "Curves/CurveBase.h"
 #include "Curves/CurveFloat.h"
 #include "GameFramework/ForceFeedbackEffect.h"
+#include "Engine/SubsurfaceProfile.h"
 #include "GameFramework/PlayerController.h"
 #include "Engine/DebugCameraController.h"
 #include "Engine/DecalActor.h"
diff --git a/Engine/Source/Runtime/Engine/Public/MaterialCompiler.h b/Engine/Source/Runtime/Engine/Public/MaterialCompiler.h
--- a/Engine/Source/Runtime/Engine/Public/MaterialCompiler.h
+++ b/Engine/Source/Runtime/Engine/Public/MaterialCompiler.h
@@ -35,8 +35,6 @@ public:
 
 	virtual ERHIFeatureLevel::Type GetFeatureLevel() = 0;
 
-	virtual float GetRefractionDepthBiasValue() = 0;
-
 	/** 
 	 * Casts the passed in code to DestType, or generates a compile error if the cast is not valid. 
 	 * This will truncate a type (float4 -> float3) but not add components (float2 -> float3), however a float1 can be cast to any float type by replication. 
@@ -204,7 +202,6 @@ public:
 	virtual EMaterialValueType GetType(int32 Code) override { return Compiler->GetType(Code); }
 	virtual EMaterialQualityLevel::Type GetQualityLevel() override { return Compiler->GetQualityLevel(); }
 	virtual ERHIFeatureLevel::Type GetFeatureLevel() override { return Compiler->GetFeatureLevel(); }
-	virtual float GetRefractionDepthBiasValue() override { return Compiler->GetRefractionDepthBiasValue(); }
 	virtual int32 ValidCast(int32 Code,EMaterialValueType DestType) override { return Compiler->ValidCast(Code, DestType); }
 	virtual int32 ForceCast(int32 Code,EMaterialValueType DestType,bool bExactMatch=false,bool bReplicateValue=false) override
 	{ return Compiler->ForceCast(Code,DestType,bExactMatch,bReplicateValue); }
diff --git a/Engine/Source/Runtime/Engine/Public/MaterialShared.h b/Engine/Source/Runtime/Engine/Public/MaterialShared.h
--- a/Engine/Source/Runtime/Engine/Public/MaterialShared.h
+++ b/Engine/Source/Runtime/Engine/Public/MaterialShared.h
@@ -1220,6 +1220,8 @@ struct FUniformExpressionCache
 	}
 };
 
+class USubsurfaceProfile;
+
 /**
  * A material render proxy used by the renderer.
  */
@@ -1283,12 +1285,17 @@ public:
 		return MaterialRenderProxyMap;
 	}
 
+	void SetSubsurfaceProfile(const USubsurfaceProfile* Ptr) { SubsurfaceProfileRT = Ptr; }
+	const USubsurfaceProfile* GetSubsurfaceProfile() const { return SubsurfaceProfileRT; }
+
 private:
 
 	/** true if the material is selected. */
 	bool bSelected : 1;
 	/** true if the material is hovered. */
 	bool bHovered : 1;
+	/** 0 if not set */
+	const USubsurfaceProfile* SubsurfaceProfileRT;
 
 	/** 
 	 * Tracks all material render proxies in all scenes, can only be accessed on the rendering thread.
diff --git a/Engine/Source/Runtime/Engine/Public/PrimitiveViewRelevance.h b/Engine/Source/Runtime/Engine/Public/PrimitiveViewRelevance.h
--- a/Engine/Source/Runtime/Engine/Public/PrimitiveViewRelevance.h
+++ b/Engine/Source/Runtime/Engine/Public/PrimitiveViewRelevance.h
@@ -39,6 +39,8 @@ struct FPrimitiveViewRelevance
 	uint32 bSeparateTranslucencyRelevance : 1;
 	/** The primitive has one or more elements that have normal translucency. */
 	uint32 bNormalTranslucencyRelevance : 1;
+	/** The primitive has one or more elements that have the SubsurfaceProfile shading model. */
+	uint32 bSubsurfaceProfileRelevance : 1;
 
 	/** 
 	 * Whether this primitive view relevance has been initialized this frame.  
@@ -69,6 +71,7 @@ struct FPrimitiveViewRelevance
 		bDistortionRelevance(false),
 		bSeparateTranslucencyRelevance(false),
 		bNormalTranslucencyRelevance(false),
+		bSubsurfaceProfileRelevance(false),
 		bInitializedThisFrame(false)
 	{}
 
@@ -91,6 +94,7 @@ struct FPrimitiveViewRelevance
 		bSeparateTranslucencyRelevance |= B.bSeparateTranslucencyRelevance != 0;
 		bNormalTranslucencyRelevance |= B.bNormalTranslucencyRelevance != 0;
 		bInitializedThisFrame |= B.bInitializedThisFrame;
+		bSubsurfaceProfileRelevance |= B.bSubsurfaceProfileRelevance != 0;
 		return *this;
 	}
 
diff --git a/Engine/Source/Runtime/Engine/Public/ShowFlagsValues.inl b/Engine/Source/Runtime/Engine/Public/ShowFlagsValues.inl
--- a/Engine/Source/Runtime/Engine/Public/ShowFlagsValues.inl
+++ b/Engine/Source/Runtime/Engine/Public/ShowFlagsValues.inl
@@ -108,6 +108,8 @@ SHOWFLAG_ALWAYS_ACCESSIBLE(Specular, SFG_LightingComponents, LOCTEXT("SpecularSF
 SHOWFLAG_ALWAYS_ACCESSIBLE(SelectionOutline, SFG_Hidden, LOCTEXT("SelectionOutlineSF", "Selection Outline"))
 /** If screen space reflections are enabled */
 SHOWFLAG_ALWAYS_ACCESSIBLE(ScreenSpaceReflections, SFG_LightingFeatures, LOCTEXT("ScreenSpaceReflectionsSF", "Screen Space Reflections"))
+/** If Screen Space Subsurface Scattering enabled */
+SHOWFLAG_ALWAYS_ACCESSIBLE(SubsurfaceScattering, SFG_LightingFeatures, LOCTEXT("SubsurfaceScatteringSF", "Screen Space Subsurface Scattering"))
 /** If the indirect lighting cache is enabled */
 SHOWFLAG_ALWAYS_ACCESSIBLE(IndirectLightingCache, SFG_LightingFeatures, LOCTEXT("IndirectLightingCacheSF", "Indirect Lighting Cache"))
 /** calls debug drawing for AIs */
diff --git a/Engine/Source/Runtime/Renderer/Private/CompositionLighting/CompositionLighting.cpp b/Engine/Source/Runtime/Renderer/Private/CompositionLighting/CompositionLighting.cpp
--- a/Engine/Source/Runtime/Renderer/Private/CompositionLighting/CompositionLighting.cpp
+++ b/Engine/Source/Runtime/Renderer/Private/CompositionLighting/CompositionLighting.cpp
@@ -19,12 +19,25 @@
 #include "ScreenSpaceReflections.h"
 #include "PostProcessWeightedSampleSum.h"
 #include "PostProcessTemporalAA.h"
+#include "PostProcessSubsurface.h"
 
 /** The global center for all deferred lighting activities. */
 FCompositionLighting GCompositionLighting;
 
 // -------------------------------------------------------
 
+static TAutoConsoleVariable<float> CVarSSSSS(
+	TEXT("r.SSSSS"),
+	1.0f,
+	TEXT("Experimental screen space subsurface scattering pass")
+	TEXT("(use shadingmodel SubsurfaceProfile, get near to the object as the default)\n")
+	TEXT("is human skin which only scatters about 1.2cm)\n")
+	TEXT(" 0: off (if there is no object on the screen using this pass it should automatically disable the post process pass)\n")
+	TEXT("<1: scale scatter radius down (for testing)\n")
+	TEXT(" 1: use given radius form the Subsurface scattering asset (default)\n")
+	TEXT(">1: scale scatter radius up (for testing)"),
+	ECVF_RenderThreadSafe);
+
 static bool IsAmbientCubemapPassRequired(FPostprocessContext& Context)
 {
 	FScene* Scene = (FScene*)Context.View.Family->Scene;
@@ -326,6 +339,33 @@ void FCompositionLighting::ProcessLighting(FRHICommandListImmediate& RHICmdList,
 			FRenderingCompositePass* SSAO = Context.Graph.RegisterPass(new FRCPassPostProcessInput(GSceneRenderTargets.ScreenSpaceAO));
 			AddPostProcessingLpvIndirect( Context, SSAO );
 		}
+
+		// Screen Space Subsurface Scattering
+		{
+			float Radius = CVarSSSSS.GetValueOnRenderThread();
+
+			bool bSimpleDynamicLighting = IsSimpleDynamicLightingEnabled();
+
+			if (View.bScreenSpaceSubsurfacePassNeeded && Radius > 0 && !bSimpleDynamicLighting && View.Family->EngineShowFlags.SubsurfaceScattering)
+			{
+				FRenderingCompositePass* PassSetup = Context.Graph.RegisterPass(new(FMemStack::Get()) FRCPassPostProcessSubsurfaceSetup());
+				PassSetup->SetInput(ePId_Input0, Context.FinalOutput);
+
+				FRenderingCompositePass* Pass0 = Context.Graph.RegisterPass(new(FMemStack::Get()) FRCPassPostProcessSubsurface(0, Radius));
+				Pass0->SetInput(ePId_Input1, PassSetup);
+
+				FRenderingCompositePass* Pass1 = Context.Graph.RegisterPass(new(FMemStack::Get()) FRCPassPostProcessSubsurface(1, Radius));
+				Pass1->SetInput(ePId_Input0, Context.FinalOutput);
+				Pass1->SetInput(ePId_Input1, Pass0);
+				Context.FinalOutput = FRenderingCompositeOutputRef(Pass1);
+
+				// waste of performance?
+				FRenderingCompositePass* NullPass = Context.Graph.RegisterPass(new(FMemStack::Get()) FRCPassPostProcessPassThrough(0));
+				NullPass->SetInput(ePId_Input0, Context.FinalOutput);
+				Context.FinalOutput = FRenderingCompositeOutputRef(NullPass);
+			}
+		}
+
 		// The graph setup should be finished before this line ----------------------------------------
 
 		SCOPED_DRAW_EVENT(CompositionLighting, DEC_SCENE_ITEMS);
diff --git a/Engine/Source/Runtime/Renderer/Private/CompositionLighting/PostProcessAmbient.h b/Engine/Source/Runtime/Renderer/Private/CompositionLighting/PostProcessAmbient.h
--- a/Engine/Source/Runtime/Renderer/Private/CompositionLighting/PostProcessAmbient.h
+++ b/Engine/Source/Runtime/Renderer/Private/CompositionLighting/PostProcessAmbient.h
@@ -19,5 +19,6 @@ public:
 	virtual const TCHAR* GetDebugName() { return TEXT("FRCPassPostProcessAmbient"); }
 	virtual void Process(FRenderingCompositePassContext& Context);
 	virtual void Release() override { delete this; }
+	virtual bool FrameBufferBlendingWithInput0() const { return true; }
 	virtual FPooledRenderTargetDesc ComputeOutputDesc(EPassOutputId InPassOutputId) const;
 };
diff --git a/Engine/Source/Runtime/Renderer/Private/PostProcess/PostProcessSubsurface.cpp b/Engine/Source/Runtime/Renderer/Private/PostProcess/PostProcessSubsurface.cpp
--- a/Engine/Source/Runtime/Renderer/Private/PostProcess/PostProcessSubsurface.cpp
+++ b/Engine/Source/Runtime/Renderer/Private/PostProcess/PostProcessSubsurface.cpp
@@ -193,8 +193,8 @@ FPooledRenderTargetDesc FRCPassPostProcessSubsurfaceSetup::ComputeOutputDesc(EPa
 
 	Ret.Reset();
 	Ret.DebugName = TEXT("SubsurfaceSetup");
-
-	// alpha is unsed, todo: consider 32bit format
+	// we don't need alpha any more
+	Ret.Format = PF_FloatRGB;
 
 	return Ret;
 }
@@ -224,6 +224,8 @@ public:
 	FPostProcessPassParameters PostprocessParameter;
 	FDeferredPixelShaderParameters DeferredParameters;
 	FShaderParameter SSSParams;
+	FShaderResourceParameter SSProfilesTexture;
+	FShaderResourceParameter SSProfilesTextureSampler;
 
 	/** Initialization constructor. */
 	FPostProcessSubsurfacePS(const ShaderMetaType::CompiledShaderInitializerType& Initializer)
@@ -232,13 +234,15 @@ public:
 		PostprocessParameter.Bind(Initializer.ParameterMap);
 		DeferredParameters.Bind(Initializer.ParameterMap);
 		SSSParams.Bind(Initializer.ParameterMap, TEXT("SSSParams"));
+		SSProfilesTexture.Bind(Initializer.ParameterMap, TEXT("SSProfilesTexture"));
+		SSProfilesTextureSampler.Bind(Initializer.ParameterMap, TEXT("SSProfilesTextureSampler"));
 	}
 
 	// FShader interface.
 	virtual bool Serialize(FArchive& Ar)
 	{
 		bool bShaderHasOutdatedParameters = FGlobalShader::Serialize(Ar);
-		Ar << PostprocessParameter << DeferredParameters << SSSParams;
+		Ar << PostprocessParameter << DeferredParameters << SSSParams << SSProfilesTexture << SSProfilesTextureSampler;
 		return bShaderHasOutdatedParameters;
 	}
 
@@ -251,9 +255,29 @@ public:
 		PostprocessParameter.SetPS(ShaderRHI, Context, TStaticSamplerState<SF_Bilinear,AM_Border,AM_Border,AM_Border>::GetRHI());
 
 		{
-			FVector4 ColorScale(InRadius, 0, 0, 0);
+			// from Separabale.usf: float distanceToProjectionWindow = 1.0 / tan(0.5 * radians(SSSS_FOVY))
+			// can be extracted out of projection matrix
+
+			float ScaleCorrectionX = Context.View.ViewRect.Width() / (float)GSceneRenderTargets.GetBufferSizeXY().X;
+			float ScaleCorrectionY = Context.View.ViewRect.Height() / (float)GSceneRenderTargets.GetBufferSizeXY().Y;
+
+			FVector4 ColorScale(InRadius, Context.View.ViewMatrices.ProjMatrix.M[0][0], ScaleCorrectionX, ScaleCorrectionY);
 			SetShaderValue(Context.RHICmdList, ShaderRHI, SSSParams, ColorScale);
 		}
+
+		{
+			ENGINE_API const FSceneRenderTargetItem* GetSubsufaceProfileTexture_RT(FRHICommandListImmediate& RHICmdList);
+
+			const FSceneRenderTargetItem* Item = GetSubsufaceProfileTexture_RT(Context.RHICmdList);
+
+			if (!Item)
+			{
+				// should never happen
+				Item = &GSystemTextures.BlackDummy->GetRenderTargetItem();
+			}
+
+			SetTextureParameter(Context.RHICmdList, ShaderRHI, SSProfilesTexture, SSProfilesTextureSampler, TStaticSamplerState<SF_Point, AM_Wrap, AM_Wrap, AM_Wrap>::GetRHI(), Item->ShaderResourceTexture);
+		}
 	}
 
 	static const TCHAR* GetSourceFilename()
@@ -289,7 +313,6 @@ void SetSubsurfaceShader(const FRenderingCompositePassContext& Context, float In
 	TShaderMapRef<FPostProcessSubsurfacePS<Method> > PixelShader(GetGlobalShaderMap());
 
 	static FGlobalBoundShaderState BoundShaderState;
-	
 
 	SetGlobalBoundShaderState(Context.RHICmdList, BoundShaderState, GFilterVertexDeclaration.VertexDeclarationRHI, *VertexShader, *PixelShader);
 
@@ -300,7 +323,7 @@ void SetSubsurfaceShader(const FRenderingCompositePassContext& Context, float In
 
 void FRCPassPostProcessSubsurface::Process(FRenderingCompositePassContext& Context)
 {
-	const FPooledRenderTargetDesc* InputDesc = GetInputDesc(ePId_Input0);
+	const FPooledRenderTargetDesc* InputDesc = GetInputDesc(ePId_Input1);
 
 	if(!InputDesc)
 	{
@@ -330,37 +353,21 @@ void FRCPassPostProcessSubsurface::Process(FRenderingCompositePassContext& Conte
 
 	Context.SetViewportAndCallRHI(0, 0, 0.0f, DestSize.X, DestSize.Y, 1.0f );
 
-	// set the state
 	Context.RHICmdList.SetBlendState(TStaticBlendState<>::GetRHI());
 	Context.RHICmdList.SetRasterizerState(TStaticRasterizerState<>::GetRHI());
 	Context.RHICmdList.SetDepthStencilState(TStaticDepthStencilState<false, CF_Always>::GetRHI());
 
 	TShaderMapRef<FPostProcessVS> VertexShader(GetGlobalShaderMap());
 
-	if(Pass == 0)
+	SCOPED_DRAW_EVENT(SubsurfacePass, DEC_SCENE_ITEMS);
+
+	if (Pass == 0)
 	{
-		SCOPED_DRAW_EVENT(SubsurfacePass0, DEC_SCENE_ITEMS);
-	
 		SetSubsurfaceShader<0>(Context, Radius);
-
-		// Draw a quad mapping scene color to the view's render target
-		DrawRectangle(
-			Context.RHICmdList,
-			DestRect.Min.X, DestRect.Min.Y,
-			DestRect.Width(), DestRect.Height(),
-			SrcRect.Min.X, SrcRect.Min.Y,
-			SrcRect.Width(), SrcRect.Height(),
-			DestSize,
-			SrcSize,
-			*VertexShader,
-			EDRF_UseTriangleOptimization);
-
-		Context.RHICmdList.CopyToResolveTarget(DestRenderTarget.TargetableTexture, DestRenderTarget.ShaderResourceTexture, false, FResolveParams());
 	}
-	
-	if(Pass == 1)
+	else
 	{
-		SCOPED_DRAW_EVENT(SubsurfacePass1, DEC_SCENE_ITEMS);
+		check(Pass == 1);
 
 		if(DoSpecularCorrection())
 		{
@@ -371,33 +378,31 @@ void FRCPassPostProcessSubsurface::Process(FRenderingCompositePassContext& Conte
 		{
 			SetSubsurfaceShader<1>(Context, Radius);
 		}
-
-		// Draw a quad mapping scene color to the view's render target
-		DrawRectangle(
-			Context.RHICmdList,
-			DestRect.Min.X, DestRect.Min.Y,
-			DestRect.Width(), DestRect.Height(),
-			SrcRect.Min.X, SrcRect.Min.Y,
-			SrcRect.Width(), SrcRect.Height(),
-			DestSize,
-			SrcSize,
-			*VertexShader,
-			EDRF_UseTriangleOptimization);
-
-		Context.RHICmdList.CopyToResolveTarget(DestRenderTarget.TargetableTexture, DestRenderTarget.ShaderResourceTexture, false, FResolveParams());
-
-		// matching AdjustGBufferRefCount(1) call is PostProcessing.cpp
-		GSceneRenderTargets.AdjustGBufferRefCount(-1);
 	}
+
+	DrawRectangle(
+		Context.RHICmdList,
+		DestRect.Min.X, DestRect.Min.Y,
+		DestRect.Width(), DestRect.Height(),
+		SrcRect.Min.X, SrcRect.Min.Y,
+		SrcRect.Width(), SrcRect.Height(),
+		DestSize,
+		SrcSize,
+		*VertexShader,
+		EDRF_UseTriangleOptimization);
+
+	Context.RHICmdList.CopyToResolveTarget(DestRenderTarget.TargetableTexture, DestRenderTarget.ShaderResourceTexture, false, FResolveParams());
 }
 
 
 FPooledRenderTargetDesc FRCPassPostProcessSubsurface::ComputeOutputDesc(EPassOutputId InPassOutputId) const
 {
-	FPooledRenderTargetDesc Ret = PassInputs[0].GetOutput()->RenderTargetDesc;
+	FPooledRenderTargetDesc Ret;
+
+	Ret = PassInputs[1].GetOutput()->RenderTargetDesc;
 
 	Ret.Reset();
-	Ret.DebugName = (Pass == 0) ? TEXT("SubsurfaceTemp") : TEXT("SubsurfaceSceneColor");
+	Ret.DebugName = (Pass == 0) ? TEXT("SubsurfaceTemp") : TEXT("SceneColor");
 
 	return Ret;
 }
\ No newline at end of file
diff --git a/Engine/Source/Runtime/Renderer/Private/PostProcess/PostProcessSubsurface.h b/Engine/Source/Runtime/Renderer/Private/PostProcess/PostProcessSubsurface.h
--- a/Engine/Source/Runtime/Renderer/Private/PostProcess/PostProcessSubsurface.h
+++ b/Engine/Source/Runtime/Renderer/Private/PostProcess/PostProcessSubsurface.h
@@ -24,8 +24,8 @@ public:
 
 
 // derives from TRenderingCompositePassBase<InputCount, OutputCount>
-// ePId_Input0: SceneColor (horizontal blur) or the pass before (vertical blur)
-// ePId_Input1: optional SubsurfaceTemp input
+// ePId_Input0: optional SubsurfaceTemp input
+// ePId_Input1: SceneColor (horizontal blur) or the pass before (vertical blur)
 // modifies SceneColor, uses some GBuffer attributes
 class FRCPassPostProcessSubsurface : public TRenderingCompositePassBase<2, 1>
 {
diff --git a/Engine/Source/Runtime/Renderer/Private/PostProcess/PostProcessing.cpp b/Engine/Source/Runtime/Renderer/Private/PostProcess/PostProcessing.cpp
--- a/Engine/Source/Runtime/Renderer/Private/PostProcess/PostProcessing.cpp
+++ b/Engine/Source/Runtime/Renderer/Private/PostProcess/PostProcessing.cpp
@@ -40,7 +40,6 @@
 #include "PostProcessAmbientOcclusion.h"
 #include "ScreenSpaceReflections.h"
 #include "PostProcessTestImage.h"
-#include "PostProcessSubsurface.h"
 #include "PostProcessUIBlur.h"
 #include "HighResScreenshot.h"
 
@@ -79,14 +78,6 @@ static TAutoConsoleVariable<float> CVarMotionBlurSoftEdgeSize(
 	TEXT(" 0:off (not free and does never completely disable), >0, 1.0 (default)"),
 	ECVF_RenderThreadSafe);
 
-static TAutoConsoleVariable<float> CVarSSSSS(
-	TEXT("r.SSSSS"),
-	0.0f,
-	TEXT("Very experimental screen space subsurface scattering on non unlit/lit materials\n")
-	TEXT("0: off\n")
-	TEXT("x: SSS radius in world space e.g. 10"),
-	ECVF_RenderThreadSafe);
-
 IMPLEMENT_SHADER_TYPE(,FPostProcessVS,TEXT("PostProcessBloom"),TEXT("MainPostprocessCommonVS"),SF_Vertex);
 
 // -------------------------------------------------------
@@ -754,8 +745,6 @@ void FPostProcessing::Process(FRHICommandListImmediate& RHICmdList, const FViewI
 
 	check(IsInRenderingThread());
 
-	bool bSimpleDynamicLighting = IsSimpleDynamicLightingEnabled();
-
 	const auto FeatureLevel = View.GetFeatureLevel();
 
 	GRenderTargetPool.AddPhaseEvent(TEXT("PostProcessing"));
@@ -806,28 +795,6 @@ void FPostProcessing::Process(FRHICommandListImmediate& RHICmdList, const FViewI
 
 		if (View.Family->EngineShowFlags.PostProcessing && FeatureLevel >= ERHIFeatureLevel::SM4)
 		{
-			// Screen Space Subsurface Scattering
-			{
-				float Radius = CVarSSSSS.GetValueOnRenderThread();
-
-				if(Radius > 0 && !bSimpleDynamicLighting)
-				{
-					// matching AdjustGBufferRefCount(-1) call is in FRCPassPostProcessSubsurface::Process()
-					GSceneRenderTargets.AdjustGBufferRefCount(1);
-
-					FRenderingCompositePass* PassSetup = Context.Graph.RegisterPass(new(FMemStack::Get()) FRCPassPostProcessSubsurfaceSetup());
-					PassSetup->SetInput(ePId_Input0, Context.FinalOutput);
-
-					FRenderingCompositePass* Pass0 = Context.Graph.RegisterPass(new(FMemStack::Get()) FRCPassPostProcessSubsurface(0, Radius));
-					Pass0->SetInput(ePId_Input0, PassSetup);
-
-					FRenderingCompositePass* Pass1 = Context.Graph.RegisterPass(new(FMemStack::Get()) FRCPassPostProcessSubsurface(1, Radius));
-					Pass1->SetInput(ePId_Input0, Pass0);
-					Pass1->SetInput(ePId_Input1, Context.FinalOutput);
-					Context.FinalOutput = FRenderingCompositeOutputRef(Pass1);
-				}
-			}
-
 			FRenderingCompositeOutputRef VelocityInput;
 			if(VelocityRT)
 			{
diff --git a/Engine/Source/Runtime/Renderer/Private/PostProcess/SceneRenderTargets.cpp b/Engine/Source/Runtime/Renderer/Private/PostProcess/SceneRenderTargets.cpp
--- a/Engine/Source/Runtime/Renderer/Private/PostProcess/SceneRenderTargets.cpp
+++ b/Engine/Source/Runtime/Renderer/Private/PostProcess/SceneRenderTargets.cpp
@@ -966,20 +966,6 @@ void FSceneRenderTargets::SetBufferSize(int32 InBufferSizeX, int32 InBufferSizeY
 	BufferSize.Y = InBufferSizeY;
 }
 
-uint8 Quantize8SignedByte(float x)
-{
-	// -1..1 -> 0..1
-	float y = x * 0.5f + 0.5f;
-
-	// 0..1 -> 0..255
-	int32 Ret = (int32)(y * 255.999f);
-
-	check(Ret >= 0);
-	check(Ret <= 255);
-
-	return Ret;
-}
-
 void FSceneRenderTargets::AllocateForwardShadingPathRenderTargets()
 {
 	// on ES2 we don't do on demand allocation of SceneColor yet (in non ES2 it's released in the Tonemapper Process())
diff --git a/Engine/Source/Runtime/Renderer/Private/SceneRendering.cpp b/Engine/Source/Runtime/Renderer/Private/SceneRendering.cpp
--- a/Engine/Source/Runtime/Renderer/Private/SceneRendering.cpp
+++ b/Engine/Source/Runtime/Renderer/Private/SceneRendering.cpp
@@ -77,12 +77,6 @@ static TAutoConsoleVariable<float> CVarTessellationAdaptivePixelsPerTriangle(
  */
 FViewInfo::FViewInfo(const FSceneViewInitOptions& InitOptions)
 	:	FSceneView(InitOptions)
-	,	bPrevTransformsReset( false )
-	,	bIgnoreExistingQueries( false )
-	,	bDisableQuerySubmissions( false )
-	,	bDisableDistanceBasedFadeTransitions( false )
-	,	NumVisibleStaticMeshElements(0)
-	,	PrecomputedVisibilityData(NULL)
 	,	IndividualOcclusionQueries((FSceneViewState*)InitOptions.SceneViewStateInterface, 1)
 	,	GroupedOcclusionQueries((FSceneViewState*)InitOptions.SceneViewStateInterface, FOcclusionQueryBatcher::OccludedPrimitiveQueryBatchSize)
 {
@@ -95,13 +89,6 @@ FViewInfo::FViewInfo(const FSceneViewInitOptions& InitOptions)
  */
 FViewInfo::FViewInfo(const FSceneView* InView)
 	:	FSceneView(*InView)
-	,	bHasTranslucentViewMeshElements( 0 )
-	,	bPrevTransformsReset( false )
-	,	bIgnoreExistingQueries( false )
-	,	bDisableQuerySubmissions( false )
-	,	bDisableDistanceBasedFadeTransitions( false )
-	,	NumVisibleStaticMeshElements(0)
-	,	PrecomputedVisibilityData(NULL)
 	,	IndividualOcclusionQueries((FSceneViewState*)InView->State,1)
 	,	GroupedOcclusionQueries((FSceneViewState*)InView->State,FOcclusionQueryBatcher::OccludedPrimitiveQueryBatchSize)
 {
@@ -110,6 +97,15 @@ FViewInfo::FViewInfo(const FSceneView* InView)
 
 void FViewInfo::Init()
 {
+	bHasTranslucentViewMeshElements = 0;
+	bPrevTransformsReset = false;
+	bIgnoreExistingQueries = false;
+	bDisableQuerySubmissions = false;
+	bDisableDistanceBasedFadeTransitions = false;
+	bScreenSpaceSubsurfacePassNeeded = false;
+	NumVisibleStaticMeshElements = false;
+	PrecomputedVisibilityData = 0;
+
 	bIsViewInfo = true;
 	PrevViewProjMatrix.SetIdentity();
 	PrevViewRotationProjMatrix.SetIdentity();
diff --git a/Engine/Source/Runtime/Renderer/Private/SceneRendering.h b/Engine/Source/Runtime/Renderer/Private/SceneRendering.h
--- a/Engine/Source/Runtime/Renderer/Private/SceneRendering.h
+++ b/Engine/Source/Runtime/Renderer/Private/SceneRendering.h
@@ -385,6 +385,8 @@ public:
 	uint32 bDisableQuerySubmissions : 1;
 	/** Whether we should disable distance-based fade transitions for this frame (usually after a large camera movement.) */
 	uint32 bDisableDistanceBasedFadeTransitions : 1;
+	/** Whether we render any objects with a material using SubsurfaceProfile as shading model. */
+	uint32 bScreenSpaceSubsurfacePassNeeded : 1;
 
 	FViewMatrices PrevViewMatrices;
 
diff --git a/Engine/Source/Runtime/Renderer/Private/SceneVisibility.cpp b/Engine/Source/Runtime/Renderer/Private/SceneVisibility.cpp
--- a/Engine/Source/Runtime/Renderer/Private/SceneVisibility.cpp
+++ b/Engine/Source/Runtime/Renderer/Private/SceneVisibility.cpp
@@ -713,6 +713,11 @@ static void ComputeRelevanceForView(
 			}
 		}
 		
+		if (ViewRelevance.bSubsurfaceProfileRelevance)
+		{
+			View.bScreenSpaceSubsurfacePassNeeded = true;
+		}
+		
 		if (ViewRelevance.bRenderCustomDepth)
 		{
 			// Add to set of dynamic distortion primitives
diff --git a/Engine/Source/Runtime/Renderer/Private/SystemTextures.cpp b/Engine/Source/Runtime/Renderer/Private/SystemTextures.cpp
--- a/Engine/Source/Runtime/Renderer/Private/SystemTextures.cpp
+++ b/Engine/Source/Runtime/Renderer/Private/SystemTextures.cpp
@@ -14,8 +14,6 @@ SystemTextures
 /** The global render targets used for scene rendering. */
 TGlobalResource<FSystemTextures> GSystemTextures;
 
-uint8 Quantize8SignedByte(float x);
-
 void FSystemTextures::InitializeTextures(FRHICommandListImmediate& RHICmdList, ERHIFeatureLevel::Type InFeatureLevel)
 {
 	if (bTexturesInitialized && FeatureLevelInitializedTo >= InFeatureLevel)
@@ -229,7 +227,7 @@ void FSystemTextures::InitializeTextures(FRHICommandListImmediate& RHICmdList, E
 			float s = FMath::Sin(ww) * lenm;
 			float c = FMath::Cos(ww) * lenm;
 
-			Bases[Pos] = FColor(Quantize8SignedByte(c), Quantize8SignedByte(s), Quantize8SignedByte(-s), Quantize8SignedByte(c));
+			Bases[Pos] = FColor(FMath::Quantize8SignedByte(c), FMath::Quantize8SignedByte(s), FMath::Quantize8SignedByte(-s), FMath::Quantize8SignedByte(c));
 		}
 
 		FPooledRenderTargetDesc Desc(FPooledRenderTargetDesc::Create2DDesc(FIntPoint(64, 64), PF_B8G8R8A8, TexCreate_HideInVisualizeTexture, TexCreate_None, false));
