SimulatedRootMotionPositionFixup: handle zero play rate, and better communicate "Server ahead of Client" pre-requisite.

diff --git a/Engine/Source/Runtime/Engine/Private/Character.cpp b/Engine/Source/Runtime/Engine/Private/Character.cpp
--- a/Engine/Source/Runtime/Engine/Private/Character.cpp
+++ b/Engine/Source/Runtime/Engine/Private/Character.cpp
@@ -952,15 +952,20 @@ void ACharacter::SimulatedRootMotionPositionFixup(float DeltaSeconds)
 					if( CharacterMovement )
 					{
 						// Guess time it takes for this delta track position, so we can get falling physics accurate.
-						float DeltaTime = DeltaPosition / ClientMontageInstance->PlayRate;
-						check( DeltaTime > 0.f );
-						CharacterMovement->SimulateRootMotion(DeltaTime, LocalRootMotionTransform);
-
-						// After movement correction, smooth out error in position if any.
-						INetworkPredictionInterface* PredictionInterface = InterfaceCast<INetworkPredictionInterface>(GetMovementComponent());
-						if (PredictionInterface)
+						if (!FMath::IsNearlyZero(ClientMontageInstance->PlayRate))
 						{
-							PredictionInterface->SmoothCorrection(OldLocation);
+							float DeltaTime = DeltaPosition / ClientMontageInstance->PlayRate;
+
+							// Even with negative playrate deltatime should be positive.
+							check(DeltaTime > 0.f);
+							CharacterMovement->SimulateRootMotion(DeltaTime, LocalRootMotionTransform);
+
+							// After movement correction, smooth out error in position if any.
+							INetworkPredictionInterface* PredictionInterface = InterfaceCast<INetworkPredictionInterface>(GetMovementComponent());
+							if (PredictionInterface)
+							{
+								PredictionInterface->SmoothCorrection(OldLocation);
+							}
 						}
 					}
 				}
@@ -1013,13 +1018,13 @@ bool ACharacter::CanUseRootMotionRepMove(const FSimulatedRootMotionReplicatedMov
 				const bool bSameSections = (AnimMontage->GetSectionIndexFromPosition(ServerPosition) == CurrentSectionIndex);
 				// if we are looping and just wrapped over, skip. That's also not easy to handle and not frequent.
 				const bool bHasLooped = (NextSectionIndex == CurrentSectionIndex) && (FMath::Abs(DeltaPosition) > (AnimMontage->GetSectionLength(CurrentSectionIndex) / 2.f));
-				// Can only simulate forward in time, so we need to find a move from the server that is behind the client in time.
-				const bool bClientAheadOfServer = ((DeltaPosition * ClientMontageInstance.PlayRate) >= 0.f);
+				// Can only simulate forward in time, so we need to make sure server move is not ahead of the client.
+				const bool bServerAheadOfClient = ((DeltaPosition * ClientMontageInstance.PlayRate) < 0.f);
 
-				UE_LOG(LogRootMotion, Log,  TEXT("\t\tACharacter::CanUseRootMotionRepMove ServerPosition: %.3f, ClientPosition: %.3f, DeltaPosition: %.3f, bSameSections: %d, bHasLooped: %d, bClientAheadOfServer: %d"), 
-					ServerPosition, ClientPosition, DeltaPosition, bSameSections, bHasLooped, bClientAheadOfServer);
+				UE_LOG(LogRootMotion, Log,  TEXT("\t\tACharacter::CanUseRootMotionRepMove ServerPosition: %.3f, ClientPosition: %.3f, DeltaPosition: %.3f, bSameSections: %d, bHasLooped: %d, bServerAheadOfClient: %d"), 
+					ServerPosition, ClientPosition, DeltaPosition, bSameSections, bHasLooped, bServerAheadOfClient);
 
-				return bSameSections && !bHasLooped && bClientAheadOfServer;
+				return bSameSections && !bHasLooped && !bServerAheadOfClient;
 			}
 		}
 	}
