Allow overriding FUniqueNetId on ULocalPlayer. This is so we can decouple controller id, and OSS identity.

This is necessary for platforms that need to allow overriding the controller while still maintaining the original user for achievements, etc.

diff --git a/Engine/Source/Runtime/Engine/Classes/Engine/LocalPlayer.h b/Engine/Source/Runtime/Engine/Classes/Engine/LocalPlayer.h
--- a/Engine/Source/Runtime/Engine/Classes/Engine/LocalPlayer.h
+++ b/Engine/Source/Runtime/Engine/Classes/Engine/LocalPlayer.h
@@ -55,6 +55,9 @@ class ENGINE_API ULocalPlayer : public UPlayer
 	UPROPERTY()
 	int32 ControllerId;
 
+	/** The FUniqueNetId which this player is associated with. */
+	TSharedPtr<class FUniqueNetId> CachedUniqueNetId;
+
 	/** The master viewport containing this player's view. */
 	UPROPERTY()
 	class UGameViewportClient* ViewportClient;
@@ -238,7 +241,27 @@ public:
 	 *
 	 * @return unique Id associated with this player
 	 */
-	TSharedPtr<class FUniqueNetId> GetUniqueNetId() const;
+	TSharedPtr<class FUniqueNetId> GetUniqueNetIdFromCachedControllerId() const;
+
+	/** 
+	 * Retrieves this player's unique net ID that was previously cached
+	 *
+	 * @return unique Id associated with this player
+	 */
+	TSharedPtr<class FUniqueNetId> GetCachedUniqueNetId() const;
+
+	/** Sets the players current cached unique net id */
+	void SetCachedUniqueNetId( TSharedPtr<class FUniqueNetId> NewUniqueNetId );
+
+	/** 
+	 * Retrieves the preferred unique net id. This is for backwards compatibility for games that don't use the cached unique net id logic
+	 *
+	 * @return unique Id associated with this player
+	 */
+	TSharedPtr<FUniqueNetId> GetPreferredUniqueNetId() const;
+
+	/** Returns true if the cached unique net id, is the one assigned to the controller id from the OSS */
+	bool IsCachedUniqueNetIdPairedWithControllerId() const;
 
 	/**
 	 * This function will give you two points in Pixel Space that surround the World Space box.
diff --git a/Engine/Source/Runtime/Engine/Private/GameInstance.cpp b/Engine/Source/Runtime/Engine/Private/GameInstance.cpp
--- a/Engine/Source/Runtime/Engine/Private/GameInstance.cpp
+++ b/Engine/Source/Runtime/Engine/Private/GameInstance.cpp
@@ -529,7 +529,7 @@ ULocalPlayer* UGameInstance::FindLocalPlayerFromUniqueNetId(TSharedPtr<FUniqueNe
 			continue;
 		}
 
-		TSharedPtr<FUniqueNetId> OtherUniqueNetId = Player->GetUniqueNetId();
+		TSharedPtr<FUniqueNetId> OtherUniqueNetId = Player->GetPreferredUniqueNetId();
 
 		if (!OtherUniqueNetId.IsValid())
 		{
diff --git a/Engine/Source/Runtime/Engine/Private/LocalPlayer.cpp b/Engine/Source/Runtime/Engine/Private/LocalPlayer.cpp
--- a/Engine/Source/Runtime/Engine/Private/LocalPlayer.cpp
+++ b/Engine/Source/Runtime/Engine/Private/LocalPlayer.cpp
@@ -137,8 +137,8 @@ void ULocalPlayer::PostInitProperties()
 
 void ULocalPlayer::PlayerAdded(UGameViewportClient* InViewportClient, int32 InControllerID)
 {
-	ViewportClient = InViewportClient;
-	ControllerId = InControllerID;
+	ViewportClient		= InViewportClient;
+	ControllerId		= InControllerID;
 }
 
 void ULocalPlayer::InitOnlineSession()
@@ -201,7 +201,7 @@ bool ULocalPlayer::SpawnPlayActor(const FString& URL,FString& OutError, UWorld*
 		}
 
 		// Get player unique id
-		TSharedPtr<FUniqueNetId> UniqueId = GetUniqueNetId();
+		TSharedPtr<FUniqueNetId> UniqueId = GetPreferredUniqueNetId();
 
 		PlayerController = InWorld->SpawnPlayActor(this, ROLE_SimulatedProxy, PlayerURL, UniqueId, OutError, GEngine->GetGamePlayers(InWorld).Find(this));
 	}
@@ -273,7 +273,7 @@ void ULocalPlayer::SendSplitJoin()
 			}
 
 			// Send the player unique Id at login
-			FUniqueNetIdRepl UniqueIdRepl(GetUniqueNetId());
+			FUniqueNetIdRepl UniqueIdRepl(GetPreferredUniqueNetId());
 
 			FString URLString = URL.ToString();
 			FNetControlMessage<NMT_JoinSplit>::Send(NetDriver->ServerConnection, URLString, UniqueIdRepl);
@@ -1556,7 +1556,7 @@ FString ULocalPlayer::GetNickname() const
 	return TEXT("");
 }
 
-TSharedPtr<FUniqueNetId> ULocalPlayer::GetUniqueNetId() const
+TSharedPtr<FUniqueNetId> ULocalPlayer::GetUniqueNetIdFromCachedControllerId() const
 {
 	UWorld* World = GetWorld();
 	if (World != NULL)
@@ -1575,6 +1575,51 @@ TSharedPtr<FUniqueNetId> ULocalPlayer::GetUniqueNetId() const
 	return NULL;
 }
 
+TSharedPtr<FUniqueNetId> ULocalPlayer::GetCachedUniqueNetId() const
+{
+	return CachedUniqueNetId;
+}
+
+void ULocalPlayer::SetCachedUniqueNetId( TSharedPtr<class FUniqueNetId> NewUniqueNetId )
+{
+	CachedUniqueNetId = NewUniqueNetId;
+}
+
+TSharedPtr<FUniqueNetId> ULocalPlayer::GetPreferredUniqueNetId() const
+{
+	// Prefer the cached unique net id (only if it's valid)
+	// This is for backwards compatibility for games that don't yet cache the unique id properly
+	if (GetCachedUniqueNetId().IsValid() && GetCachedUniqueNetId()->IsValid())
+	{
+		return GetCachedUniqueNetId();
+	}
+
+	// If the cached unique net id is not valid, then get the one paired with the controller
+	return GetUniqueNetIdFromCachedControllerId();
+}
+
+bool ULocalPlayer::IsCachedUniqueNetIdPairedWithControllerId() const
+{
+	// Get the UniqueNetId that is paired with the controller
+	TSharedPtr<FUniqueNetId> UniqueIdFromController = GetUniqueNetIdFromCachedControllerId();
+
+	if (CachedUniqueNetId.IsValid() != UniqueIdFromController.IsValid())
+	{
+		// Definitely can't match if one is valid and not the other
+		return false;
+	}
+
+	if (!CachedUniqueNetId.IsValid())
+	{
+		// Both are invalid, technically they match
+		check(!UniqueIdFromController.IsValid());
+		return true;
+	}
+
+	// Both are valid, ask them if they match
+	return *CachedUniqueNetId == *UniqueIdFromController;
+}
+
 UWorld* ULocalPlayer::GetWorld() const
 {
 	return ViewportClient ? ViewportClient->GetWorld() : nullptr;
diff --git a/Engine/Source/Runtime/Engine/Private/PendingNetGame.cpp b/Engine/Source/Runtime/Engine/Private/PendingNetGame.cpp
--- a/Engine/Source/Runtime/Engine/Private/PendingNetGame.cpp
+++ b/Engine/Source/Runtime/Engine/Private/PendingNetGame.cpp
@@ -188,7 +188,7 @@ void UPendingNetGame::NotifyControlMessage(UNetConnection* Connection, uint8 Mes
 			}
 			
 			// Send the player unique Id at login
-			FUniqueNetIdRepl UniqueIdRepl(LocalPlayer->GetUniqueNetId());
+			FUniqueNetIdRepl UniqueIdRepl(LocalPlayer->GetPreferredUniqueNetId());
 
 			Connection->ClientResponse = TEXT("0");
 			FString URLString(PartialURL.ToString());
