* fixed bad jump when using black in Subsurface Scattering Profile

diff --git a/Engine/Source/Runtime/Engine/Private/Rendering/SubsurfaceProfile.cpp b/Engine/Source/Runtime/Engine/Private/Rendering/SubsurfaceProfile.cpp
--- a/Engine/Source/Runtime/Engine/Private/Rendering/SubsurfaceProfile.cpp
+++ b/Engine/Source/Runtime/Engine/Private/Rendering/SubsurfaceProfile.cpp
@@ -151,6 +151,13 @@ void FSubsurfaceProfileTexture::CreateTexture(FRHICommandListImmediate& RHICmdLi
 	{
 		FSubsurfaceProfileStruct& Data = SubsurfaceProfileEntries[y].Settings;
 
+		if(Data.SubsurfaceColor.IsAlmostBlack())
+		{
+			// to avoid div by 0 and a jump to a different value
+			// this basically means we don't want subsurface scattering
+			Data.SubsurfaceColor = FLinearColor(0.0001f, 0.0001f, 0.0001f);
+		}
+
 		ComputeMirroredSSSKernel(kernel, Data.SubsurfaceColor, Data.FalloffColor);
 
 		const float TableMaxRGB = 1.0f;
