Make sure Shader Debug Info Root path is always passed correctly to the shader compilers.

diff --git a/Engine/Source/Runtime/Engine/Private/ShaderCompiler/ShaderCompiler.cpp b/Engine/Source/Runtime/Engine/Private/ShaderCompiler/ShaderCompiler.cpp
--- a/Engine/Source/Runtime/Engine/Private/ShaderCompiler/ShaderCompiler.cpp
+++ b/Engine/Source/Runtime/Engine/Private/ShaderCompiler/ShaderCompiler.cpp
@@ -1988,11 +1988,10 @@ void GlobalBeginCompileShader(
 	Input.ShaderFormat = LegacyShaderPlatformToShaderFormat(EShaderPlatform(Target.Platform));
 	Input.SourceFilename = SourceFilename;
 	Input.EntryPointName = FunctionName;
-	Input.DumpDebugInfoRootPath = GShaderCompilingManager->GetAbsoluteShaderDebugInfoDirectory();
+	Input.DumpDebugInfoRootPath = GShaderCompilingManager->GetAbsoluteShaderDebugInfoDirectory() / Input.ShaderFormat.ToString();
 
 	if (GDumpShaderDebugInfo != 0)
 	{
-		Input.DumpDebugInfoRootPath = GShaderCompilingManager->GetAbsoluteShaderDebugInfoDirectory() / Input.ShaderFormat.ToString();
 		Input.DumpDebugInfoPath = Input.DumpDebugInfoRootPath / DebugGroupName;
 
 		if (VFType)
