Header cleanup: PlayerCameraManager.h, NavigationSystem.h, RecastNavMesh.h, MatineeActorCameraAnim.h, TriggerBase.h.

diff --git a/Engine/Source/Runtime/Engine/Classes/AI/Navigation/NavigationData.h b/Engine/Source/Runtime/Engine/Classes/AI/Navigation/NavigationData.h
--- a/Engine/Source/Runtime/Engine/Classes/AI/Navigation/NavigationData.h
+++ b/Engine/Source/Runtime/Engine/Classes/AI/Navigation/NavigationData.h
@@ -1,10 +1,10 @@
 // Copyright 1998-2014 Epic Games, Inc. All Rights Reserved.
 
 #pragma once
-#include "GameFramework/Actor.h"
-#include "AI/Navigation/NavigationTypes.h"
 #include "AI/NavDataGenerator.h"
+#include "AI/Navigation/NavigationTypes.h"
 #include "AI/Navigation/NavFilters/NavigationQueryFilter.h"
+#include "GameFramework/Actor.h"
 #include "NavigationData.generated.h"
 
 class FNavDataGenerator; 
@@ -38,6 +38,310 @@ struct FNavPathRecalculationRequest
 	bool operator==(const FNavPathRecalculationRequest& Other) const { return Path == Other.Path;  }
 };
 
+struct FPathFindingResult
+{
+	FNavPathSharedPtr Path;
+	ENavigationQueryResult::Type Result;
+
+	FPathFindingResult(ENavigationQueryResult::Type InResult = ENavigationQueryResult::Invalid) : Result(InResult)
+	{ }
+
+	FORCEINLINE bool IsSuccessful() const
+	{
+		return Result == ENavigationQueryResult::Success;
+	}
+	FORCEINLINE bool IsPartial() const;
+};
+
+struct ENGINE_API FNavigationPath : public TSharedFromThis<FNavigationPath, ESPMode::ThreadSafe>
+{
+	//DECLARE_DELEGATE_OneParam(FPathObserverDelegate, FNavigationPath*);
+	DECLARE_MULTICAST_DELEGATE_TwoParams(FPathObserverDelegate, FNavigationPath*, ENavPathEvent::Type);
+
+	FNavigationPath();
+	FNavigationPath(const TArray<FVector>& Points, AActor* Base = NULL);
+	virtual ~FNavigationPath()
+	{ }
+
+	FORCEINLINE bool IsValid() const
+	{
+		return bIsReady && PathPoints.Num() > 1 && bUpToDate;
+	}
+	FORCEINLINE bool IsUpToDate() const
+	{
+		return bUpToDate;
+	}
+	FORCEINLINE bool IsReady() const
+	{
+		return bIsReady;
+	}
+	FORCEINLINE bool IsPartial() const
+	{
+		return bIsPartial;
+	}
+	FORCEINLINE bool DidSearchReachedLimit() const
+	{
+		return bReachedSearchLimit;
+	}
+	FORCEINLINE bool IsDirect() const
+	{
+		return NavigationDataUsed.Get() == NULL;
+	}
+	FORCEINLINE FVector GetDestinationLocation() const
+	{
+		return IsValid() ? PathPoints.Last().Location : INVALID_NAVEXTENT;
+	}
+	FORCEINLINE FPathObserverDelegate& GetObserver()
+	{
+		return ObserverDelegate;
+	}
+	FORCEINLINE void AddObserver(FPathObserverDelegate::FDelegate& NewObserver)
+	{
+		ObserverDelegate.Add(NewObserver);
+	}
+	FORCEINLINE void RemoveObserver(FPathObserverDelegate::FDelegate& ObserverToRemove)
+	{
+		ObserverDelegate.Remove(ObserverToRemove);
+	}
+
+	FORCEINLINE void MarkReady()
+	{
+		bIsReady = true;
+	}
+
+	FORCEINLINE void SetNavigationDataUsed(const ANavigationData* const NewData);
+
+	FORCEINLINE ANavigationData* GetNavigationDataUsed() const
+	{
+		return NavigationDataUsed.Get();
+	}
+	FORCEINLINE void SetQuerier(const UObject* InQuerier)
+	{
+		return Querier = InQuerier;
+	}
+	FORCEINLINE const UObject* GetQuerier() const
+	{
+		return Querier.Get();
+	}
+	//FORCEINLINE void SetObserver(const FPathObserverDelegate& Observer) { ObserverDelegate = Observer; }
+	FORCEINLINE void SetIsPartial(const bool bPartial)
+	{
+		bIsPartial = bPartial;
+	}
+	FORCEINLINE void SetSearchReachedLimit(const bool bLimited)
+	{
+		bReachedSearchLimit = bLimited;
+	}
+
+	FORCEINLINE void SetFilter(TSharedPtr<const FNavigationQueryFilter> InFilter)
+	{
+		Filter = InFilter;
+	}
+	FORCEINLINE TSharedPtr<const FNavigationQueryFilter> GetFilter() const
+	{
+		return Filter;
+	}
+	FORCEINLINE AActor* GetBaseActor() const
+	{
+		return Base.Get();
+	}
+
+	FVector GetStartLocation() const
+	{
+		return PathPoints.Num() > 0 ? PathPoints[0].Location : FNavigationSystem::InvalidLocation;
+	}
+	FVector GetEndLocation() const
+	{
+		return PathPoints.Num() > 0 ? PathPoints.Last().Location : FNavigationSystem::InvalidLocation;
+	}
+
+	FORCEINLINE void DoneUpdating(ENavPathUpdateType::Type UpdateType)
+	{
+		bUpToDate = true; ObserverDelegate.Broadcast(this, UpdateType == ENavPathUpdateType::GoalMoved ? ENavPathEvent::UpdatedDueToGoalMoved : ENavPathEvent::UpdatedDueToNavigationChanged);
+	}
+	void Invalidate();
+	void RePathFailed();
+
+	virtual void DebugDraw(const ANavigationData* NavData, FColor PathColor, class UCanvas* Canvas, bool bPersistent, const uint32 NextPathPointIndex = 0) const;
+#if ENABLE_VISUAL_LOG
+	virtual void DescribeSelfToVisLog(struct FVisLogEntry* Snapshot) const;
+	virtual FString GetDescription() const;
+#endif // ENABLE_VISUAL_LOG
+
+	/** check if path contains specific custom nav link */
+	virtual bool ContainsCustomLink(uint32 UniqueLinkId) const;
+
+	/** check if path contains any custom nav link */
+	virtual bool ContainsAnyCustomLink() const;
+
+	/** check if path contains given node */
+	virtual bool ContainsNode(NavNodeRef NodeRef) const;
+
+	/** get cost of path, starting from given point */
+	virtual float GetCostFromIndex(int32 PathPointIndex) const
+	{
+		return 0.0f;
+	}
+
+	/** get cost of path, starting from given node */
+	virtual float GetCostFromNode(NavNodeRef PathNode) const
+	{
+		return 0.0f;
+	}
+
+	FORCEINLINE float GetCost() const
+	{
+		return GetCostFromIndex(0);
+	}
+
+	/** calculates total length of segments from NextPathPoint to the end of path, plus distance from CurrentPosition to NextPathPoint */
+	float GetLengthFromPosition(FVector SegmentStart, uint32 NextPathPointIndex) const;
+
+	FORCEINLINE float GetLength() const
+	{
+		return PathPoints.Num() ? GetLengthFromPosition(PathPoints[0].Location, 1) : 0.0f;
+	}
+
+	static bool GetPathPoint(const FNavigationPath* Path, uint32 PathVertIdx, FNavPathPoint& PathPoint)
+	{
+		if (Path && Path->GetPathPoints().IsValidIndex((int32)PathVertIdx))
+		{
+			PathPoint = Path->PathPoints[PathVertIdx];
+			return true;
+		}
+
+		return false;
+	}
+
+	FORCEINLINE const TArray<FNavPathPoint>& GetPathPoints() const
+	{
+		return PathPoints;
+	}
+	FORCEINLINE TArray<FNavPathPoint>& GetPathPoints()
+	{
+		return PathPoints;
+	}
+
+	virtual bool DoesIntersectBox(const FBox& Box, uint32 StartingIndex = 0, int32* IntersectingSegmentIndex = NULL) const;
+
+	/** type safe casts */
+	template<typename PathClass>
+	FORCEINLINE const PathClass* CastPath() const
+	{
+		return PathType == PathClass::Type ? static_cast<PathClass*>(this) : NULL;
+	}
+
+	template<typename PathClass>
+	FORCEINLINE PathClass* CastPath()
+	{
+		return PathType == PathClass::Type ? static_cast<PathClass*>(this) : NULL;
+	}
+
+	/** enables path observing specified AActor's location and update itself if actor changes location */
+	void SetGoalActorObservation(const AActor& ActorToObserve, float TetherDistance);
+	/** turns goal actor location's observation */
+	void DisableGoalActorObservation();
+	/** set's up the path to use SourceActor's location in case of recalculation */
+	void SetSourceActor(const AActor& InSourceActor);
+
+	/** if enabled path will request recalculation if it gets invalidated due to a change to underlying navigation */
+	void EnableRecalculationOnInvalidation(bool bShouldAutoUpdate)
+	{
+		bDoAutoUpdateOnInvalidation = bShouldAutoUpdate;
+	}
+	bool WillRecalculateOnInvalidation() const
+	{
+		return bDoAutoUpdateOnInvalidation;
+	}
+
+	EPathObservationResult::Type TickPathObservation();
+
+	/** If GoalActor is set it retrieved its navigation location, if not retrieved last path point location */
+	FVector GetGoalLocation() const;
+
+	/** retrieved location to start path finding from (in case of path recalculation) */
+	FVector GetPathFindingStartLocation() const;
+
+	const AActor* GetGoalActor() const
+	{
+		return GoalActor.Get();
+	}
+	const INavAgentInterface* GetGoalActorAsNavAgent() const
+	{
+		return GoalActor.IsValid() ? GoalActorAsNavAgent : NULL;
+	}
+
+	// @todo this is navigation-type specific and should not be implemented here.
+	/** additional node refs used during path following shortcuts */
+	TArray<NavNodeRef> ShortcutNodeRefs;
+
+protected:
+	/**
+	* IMPORTANT: path is assumed to be valid if it contains _MORE_ than _ONE_ point
+	*	point 0 is path's starting point - if it's the only point on the path then there's no path per se
+	*/
+	TArray<FNavPathPoint> PathPoints;
+
+	/** base actor, if exist path points locations will be relative to it */
+	TWeakObjectPtr<AActor> Base;
+
+private:
+	/** if set path will observe GoalActor's location and update itself if goal moves more then
+	*	@note only actual navigation paths can use this feature, meaning the ones associated with
+	*	a NavigationData instance (meaning NavigationDataUsed != NULL) */
+	TWeakObjectPtr<const AActor> GoalActor;
+
+	/** cached result of GoalActor casting to INavAgentInterface */
+	const INavAgentInterface* GoalActorAsNavAgent;
+
+	/** if set will be queried for location in case of path's recalculation */
+	TWeakObjectPtr<const AActor> SourceActor;
+
+	/** cached result of PathSource casting to INavAgentInterface */
+	const INavAgentInterface* SourceActorAsNavAgent;
+
+	/** path's querier is an object that was used (depending on navigation used) for supplying additional path-finding details
+	*	and will be reused when re-calculating the path*/
+	TWeakObjectPtr<const UObject> Querier;
+
+protected:
+	/** filter used to build this path */
+	TSharedPtr<const struct FNavigationQueryFilter> Filter;
+
+	/** type of path */
+	static const FNavPathType Type;
+
+	FNavPathType PathType;
+
+	/** A delegate that will be called when path becomes invalid */
+	FPathObserverDelegate ObserverDelegate;
+
+	/** "true" until navigation data used to generate this path has been changed/invalidated */
+	uint32 bUpToDate : 1;
+
+	/** when false it means path instance has been created, but not filled with data yet */
+	uint32 bIsReady : 1;
+
+	/** "true" when path is only partially generated, when goal is unreachable and path represent best guess */
+	uint32 bIsPartial : 1;
+
+	/** set to true when path finding algorithm reached a technical limit (like limit of A* nodes).
+	*	This generally means path cannot be trusted to lead to requested destination
+	*	although it might lead closer to destination. */
+	uint32 bReachedSearchLimit : 1;
+
+	/** if true path will request re-pathing if it gets invalidated due to underlying navigation changed */
+	uint32 bDoAutoUpdateOnInvalidation : 1;
+
+	/** navigation data used to generate this path */
+	TWeakObjectPtr<ANavigationData> NavigationDataUsed;
+
+private:
+	/* if GoalActor is set this is the distance we'll try to keep GoalActor from end of path. If GoalActor
+	* moves more then this from the end of the path we'll recalculate the path */
+	float GoalActorLocationTetherDistanceSq;
+};
+
 /** 
  *	Represents abstract Navigation Data (sub-classed as NavMesh, NavGraph, etc)
  *	Used as a common interface for all navigation types handled by NavigationSystem
@@ -433,3 +737,36 @@ private:
 
 	static uint16 GetNextUniqueID();
 };
+
+struct FAsyncPathFindingQuery : public FPathFindingQuery
+{
+	const uint32 QueryID;
+	const FNavPathQueryDelegate OnDoneDelegate;
+	const TEnumAsByte<EPathFindingMode::Type> Mode;
+	FPathFindingResult Result;
+
+	FAsyncPathFindingQuery()
+		: QueryID(INVALID_NAVQUERYID)
+	{ }
+
+	FAsyncPathFindingQuery(const UObject* InOwner, const ANavigationData* InNavData, const FVector& Start, const FVector& End, const FNavPathQueryDelegate& Delegate, TSharedPtr<const FNavigationQueryFilter> SourceQueryFilter);
+	FAsyncPathFindingQuery(const FPathFindingQuery& Query, const FNavPathQueryDelegate& Delegate, const EPathFindingMode::Type QueryMode);
+
+protected:
+	FORCEINLINE static uint32 GetUniqueID()
+	{
+		return ++LastPathFindingUniqueID;
+	}
+
+	static uint32 LastPathFindingUniqueID;
+};
+
+FORCEINLINE bool FPathFindingResult::IsPartial() const
+{
+	return Result == ENavigationQueryResult::Fail && Path.IsValid() && Path->IsPartial();
+}
+
+FORCEINLINE void FNavigationPath::SetNavigationDataUsed(const ANavigationData* const NewData)
+{
+	NavigationDataUsed = NewData;
+}
\ No newline at end of file
diff --git a/Engine/Source/Runtime/Engine/Classes/AI/Navigation/NavigationSystem.h b/Engine/Source/Runtime/Engine/Classes/AI/Navigation/NavigationSystem.h
--- a/Engine/Source/Runtime/Engine/Classes/AI/Navigation/NavigationSystem.h
+++ b/Engine/Source/Runtime/Engine/Classes/AI/Navigation/NavigationSystem.h
@@ -2,11 +2,9 @@
 
 #pragma once
 
-#include "AI/Navigation/NavFilters/NavigationQueryFilter.h"
-#include "AI/Navigation/NavigationTypes.h"
+#include "AI/Navigation/NavigationData.h"
 #include "Kismet/BlueprintFunctionLibrary.h"
 #include "GenericOctreePublic.h"
-#include "AI/Navigation/NavigationData.h"
 #include "NavigationSystem.generated.h"
 
 #define NAVSYS_DEBUG (0 && UE_BUILD_DEBUG)
@@ -23,6 +21,8 @@ class ANavigationData;
 class UNavigationQueryFilter;
 class UWorld;
 class UCrowdManager;
+struct FNavDataConfig;
+struct FPathFindingResult;
 
 ENGINE_API DECLARE_LOG_CATEGORY_EXTERN(LogNavigation, Warning, All);
 
diff --git a/Engine/Source/Runtime/Engine/Classes/AI/Navigation/NavigationTypes.h b/Engine/Source/Runtime/Engine/Classes/AI/Navigation/NavigationTypes.h
--- a/Engine/Source/Runtime/Engine/Classes/AI/Navigation/NavigationTypes.h
+++ b/Engine/Source/Runtime/Engine/Classes/AI/Navigation/NavigationTypes.h
@@ -258,197 +258,6 @@ typedef TSharedRef<struct FNavigationPath, ESPMode::ThreadSafe> FNavPathSharedRe
 typedef TSharedPtr<struct FNavigationPath, ESPMode::ThreadSafe> FNavPathSharedPtr;
 typedef TWeakPtr<struct FNavigationPath, ESPMode::ThreadSafe> FNavPathWeakPtr;
 
-struct ENGINE_API FNavigationPath : public TSharedFromThis<FNavigationPath, ESPMode::ThreadSafe>
-{
-	//DECLARE_DELEGATE_OneParam(FPathObserverDelegate, FNavigationPath*);
-	DECLARE_MULTICAST_DELEGATE_TwoParams(FPathObserverDelegate, FNavigationPath*, ENavPathEvent::Type);
-
-	FNavigationPath();
-	FNavigationPath(const TArray<FVector>& Points, AActor* Base = NULL);
-	virtual ~FNavigationPath() {}
-
-	FORCEINLINE bool IsValid() const { return bIsReady && PathPoints.Num() > 1 && bUpToDate; }
-	FORCEINLINE bool IsUpToDate() const { return bUpToDate; }
-	FORCEINLINE bool IsReady() const { return bIsReady; }
-	FORCEINLINE bool IsPartial() const { return bIsPartial; }
-	FORCEINLINE bool DidSearchReachedLimit() const { return bReachedSearchLimit; }	
-	FORCEINLINE bool IsDirect() const { return NavigationDataUsed.Get() == NULL; }
-	FORCEINLINE FVector GetDestinationLocation() const { return IsValid() ? PathPoints.Last().Location : INVALID_NAVEXTENT; }
-	FORCEINLINE FPathObserverDelegate& GetObserver() { return ObserverDelegate; }
-	FORCEINLINE void AddObserver(FPathObserverDelegate::FDelegate& NewObserver) { ObserverDelegate.Add(NewObserver); }
-	FORCEINLINE void RemoveObserver(FPathObserverDelegate::FDelegate& ObserverToRemove) { ObserverDelegate.Remove(ObserverToRemove); }
-
-	FORCEINLINE void MarkReady() { bIsReady = true; }
-
-	FORCEINLINE void SetNavigationDataUsed(const ANavigationData* const NewData) { NavigationDataUsed = NewData; }
-	FORCEINLINE ANavigationData* GetNavigationDataUsed() const { return NavigationDataUsed.Get(); }
-	FORCEINLINE void SetQuerier(const UObject* InQuerier) { return Querier = InQuerier; }
-	FORCEINLINE const UObject* GetQuerier() const { return Querier.Get(); }
-	//FORCEINLINE void SetObserver(const FPathObserverDelegate& Observer) { ObserverDelegate = Observer; }
-	FORCEINLINE void SetIsPartial(const bool bPartial) { bIsPartial = bPartial; }
-	FORCEINLINE void SetSearchReachedLimit(const bool bLimited) { bReachedSearchLimit = bLimited; }
-
-	FORCEINLINE void SetFilter(TSharedPtr<const FNavigationQueryFilter> InFilter) { Filter = InFilter; }
-	FORCEINLINE TSharedPtr<const FNavigationQueryFilter> GetFilter() const { return Filter; }
-	FORCEINLINE AActor* GetBaseActor() const { return Base.Get(); }
-
-	FVector GetStartLocation() const { return PathPoints.Num() > 0 ? PathPoints[0].Location : FNavigationSystem::InvalidLocation; }
-	FVector GetEndLocation() const { return PathPoints.Num() > 0 ? PathPoints.Last().Location : FNavigationSystem::InvalidLocation; }
-	
-	FORCEINLINE void DoneUpdating(ENavPathUpdateType::Type UpdateType) { bUpToDate = true; ObserverDelegate.Broadcast(this, UpdateType == ENavPathUpdateType::GoalMoved ? ENavPathEvent::UpdatedDueToGoalMoved : ENavPathEvent::UpdatedDueToNavigationChanged); }
-	void Invalidate();
-	void RePathFailed();
-
-	virtual void DebugDraw(const ANavigationData* NavData, FColor PathColor, class UCanvas* Canvas, bool bPersistent, const uint32 NextPathPointIndex = 0) const;
-#if ENABLE_VISUAL_LOG
-	virtual void DescribeSelfToVisLog(struct FVisLogEntry* Snapshot) const;
-	virtual FString GetDescription() const;
-#endif // ENABLE_VISUAL_LOG
-
-	/** check if path contains specific custom nav link */
-	virtual bool ContainsCustomLink(uint32 UniqueLinkId) const;
-
-	/** check if path contains any custom nav link */
-	virtual bool ContainsAnyCustomLink() const;
-
-	/** check if path contains given node */
-	virtual bool ContainsNode(NavNodeRef NodeRef) const;
-
-	/** get cost of path, starting from given point */
-	virtual float GetCostFromIndex(int32 PathPointIndex) const { return 0.0f; }
-
-	/** get cost of path, starting from given node */
-	virtual float GetCostFromNode(NavNodeRef PathNode) const { return 0.0f; }
-
-	FORCEINLINE float GetCost() const { return GetCostFromIndex(0); }
-
-	/** calculates total length of segments from NextPathPoint to the end of path, plus distance from CurrentPosition to NextPathPoint */
-	float GetLengthFromPosition(FVector SegmentStart, uint32 NextPathPointIndex) const;
-
-	FORCEINLINE float GetLength() const { return PathPoints.Num() ? GetLengthFromPosition(PathPoints[0].Location, 1) : 0.0f; }
-
-	static bool GetPathPoint(const FNavigationPath* Path, uint32 PathVertIdx, FNavPathPoint& PathPoint)
-	{
-		if (Path && Path->GetPathPoints().IsValidIndex((int32)PathVertIdx))
-		{
-			PathPoint = Path->PathPoints[PathVertIdx];
-			return true;
-		}
-
-		return false;
-	}
-
-	FORCEINLINE const TArray<FNavPathPoint>& GetPathPoints() const { return PathPoints; }
-	FORCEINLINE TArray<FNavPathPoint>& GetPathPoints() { return PathPoints; }
-
-	virtual bool DoesIntersectBox(const FBox& Box, uint32 StartingIndex = 0, int32* IntersectingSegmentIndex = NULL) const;
-
-	/** type safe casts */
-	template<typename PathClass>
-	FORCEINLINE const PathClass* CastPath() const
-	{
-		return PathType == PathClass::Type ? static_cast<PathClass*>(this) : NULL;
-	}
-
-	template<typename PathClass>
-	FORCEINLINE PathClass* CastPath()
-	{
-		return PathType == PathClass::Type ? static_cast<PathClass*>(this) : NULL;
-	}
-
-	/** enables path observing specified AActor's location and update itself if actor changes location */
-	void SetGoalActorObservation(const AActor& ActorToObserve, float TetherDistance);
-	/** turns goal actor location's observation */
-	void DisableGoalActorObservation();
-	/** set's up the path to use SourceActor's location in case of recalculation */
-	void SetSourceActor(const AActor& InSourceActor);
-
-	/** if enabled path will request recalculation if it gets invalidated due to a change to underlying navigation */
-	void EnableRecalculationOnInvalidation(bool bShouldAutoUpdate) { bDoAutoUpdateOnInvalidation = bShouldAutoUpdate; }
-	bool WillRecalculateOnInvalidation() const { return bDoAutoUpdateOnInvalidation; }
-	
-	EPathObservationResult::Type TickPathObservation();
-
-	/** If GoalActor is set it retrieved its navigation location, if not retrieved last path point location */
-	FVector GetGoalLocation() const;
-
-	/** retrieved location to start path finding from (in case of path recalculation) */
-	FVector GetPathFindingStartLocation() const;
-
-	const AActor* GetGoalActor() const { return GoalActor.Get();  }	
-	const INavAgentInterface* GetGoalActorAsNavAgent() const { return GoalActor.IsValid() ? GoalActorAsNavAgent : NULL; }
-
-	// @todo this is navigation-type specific and should not be implemented here.
-	/** additional node refs used during path following shortcuts */
-	TArray<NavNodeRef> ShortcutNodeRefs;
-
-protected:
-	/** 
-	 * IMPORTANT: path is assumed to be valid if it contains _MORE_ than _ONE_ point 
-	 *	point 0 is path's starting point - if it's the only point on the path then there's no path per se
-	 */
-	TArray<FNavPathPoint> PathPoints;
-	
-	/** base actor, if exist path points locations will be relative to it */
-	TWeakObjectPtr<AActor> Base;
-
-private:
-	/** if set path will observe GoalActor's location and update itself if goal moves more then
-	 *	@note only actual navigation paths can use this feature, meaning the ones associated with
-	 *	a NavigationData instance (meaning NavigationDataUsed != NULL) */
-	TWeakObjectPtr<const AActor> GoalActor;
-
-	/** cached result of GoalActor casting to INavAgentInterface */
-	const INavAgentInterface* GoalActorAsNavAgent;
-
-	/** if set will be queried for location in case of path's recalculation */
-	TWeakObjectPtr<const AActor> SourceActor;
-
-	/** cached result of PathSource casting to INavAgentInterface */
-	const INavAgentInterface* SourceActorAsNavAgent;
-
-	/** path's querier is an object that was used (depending on navigation used) for supplying additional path-finding details
-	 *	and will be reused when re-calculating the path*/
-	TWeakObjectPtr<const UObject> Querier;
-
-protected:
-	/** filter used to build this path */
-	TSharedPtr<const struct FNavigationQueryFilter> Filter;
-
-	/** type of path */
-	static const FNavPathType Type;
-
-	FNavPathType PathType;
-
-	/** A delegate that will be called when path becomes invalid */
-	FPathObserverDelegate ObserverDelegate;
-
-	/** "true" until navigation data used to generate this path has been changed/invalidated */
-	uint32 bUpToDate : 1;
-
-	/** when false it means path instance has been created, but not filled with data yet */
-	uint32 bIsReady : 1;
-
-	/** "true" when path is only partially generated, when goal is unreachable and path represent best guess */
-	uint32 bIsPartial : 1;
-
-	/** set to true when path finding algorithm reached a technical limit (like limit of A* nodes).
-	 *	This generally means path cannot be trusted to lead to requested destination
-	 *	although it might lead closer to destination. */
-	uint32 bReachedSearchLimit : 1;
-
-	/** if true path will request re-pathing if it gets invalidated due to underlying navigation changed */
-	uint32 bDoAutoUpdateOnInvalidation : 1;
-
-	/** navigation data used to generate this path */
-	TWeakObjectPtr<ANavigationData> NavigationDataUsed;
-
-private:
-	/* if GoalActor is set this is the distance we'll try to keep GoalActor from end of path. If GoalActor
-	* moves more then this from the end of the path we'll recalculate the path */
-	float GoalActorLocationTetherDistanceSq;
-};
-
 USTRUCT()
 struct FMovementProperties
 {
@@ -557,18 +366,6 @@ namespace ENavigationQueryResult
 	};
 }
 
-struct FPathFindingResult
-{
-	FNavPathSharedPtr Path;
-	ENavigationQueryResult::Type Result;
-
-	FPathFindingResult(ENavigationQueryResult::Type InResult = ENavigationQueryResult::Invalid) : Result(InResult)
-	{}
-
-	FORCEINLINE bool IsSuccessful() const { return Result == ENavigationQueryResult::Success; }
-	FORCEINLINE bool IsPartial() const { return Result == ENavigationQueryResult::Fail && Path.IsValid() && Path->IsPartial(); }
-};
-
 struct ENGINE_API FPathFindingQuery
 {
 	TWeakObjectPtr<const ANavigationData> NavData;
@@ -611,30 +408,6 @@ namespace EPathFindingMode
 */
 DECLARE_DELEGATE_ThreeParams(FNavPathQueryDelegate, uint32, ENavigationQueryResult::Type, FNavPathSharedPtr);
 
-struct FAsyncPathFindingQuery : public FPathFindingQuery
-{
-	const uint32 QueryID;
-	const FNavPathQueryDelegate OnDoneDelegate;
-	const TEnumAsByte<EPathFindingMode::Type> Mode;
-	FPathFindingResult Result;
-
-	FAsyncPathFindingQuery()
-		: QueryID(INVALID_NAVQUERYID)
-	{
-	}
-
-	FAsyncPathFindingQuery(const UObject* InOwner, const ANavigationData* InNavData, const FVector& Start, const FVector& End, const FNavPathQueryDelegate& Delegate, TSharedPtr<const FNavigationQueryFilter> SourceQueryFilter);
-	FAsyncPathFindingQuery(const FPathFindingQuery& Query, const FNavPathQueryDelegate& Delegate, const EPathFindingMode::Type QueryMode);
-
-protected:
-	FORCEINLINE static uint32 GetUniqueID()
-	{
-		return ++LastPathFindingUniqueID;
-	}
-
-	static uint32 LastPathFindingUniqueID;
-};
-
 //////////////////////////////////////////////////////////////////////////
 // Custom path following data
 
diff --git a/Engine/Source/Runtime/Engine/Classes/AI/Navigation/RecastNavMesh.h b/Engine/Source/Runtime/Engine/Classes/AI/Navigation/RecastNavMesh.h
--- a/Engine/Source/Runtime/Engine/Classes/AI/Navigation/RecastNavMesh.h
+++ b/Engine/Source/Runtime/Engine/Classes/AI/Navigation/RecastNavMesh.h
@@ -1,11 +1,9 @@
 // Copyright 1998-2014 Epic Games, Inc. All Rights Reserved.
 
 #pragma once
-
+#include "AI/Navigation/NavigationData.h"
 #include "AI/Navigation/NavigationTypes.h"
-#include "NavAreas/NavArea.h"
 #include "Tickable.h"
-#include "AI/Navigation/NavigationData.h"
 #include "RecastNavMesh.generated.h"
 
 /** Initial checkin. */
diff --git a/Engine/Source/Runtime/Engine/Classes/Camera/PlayerCameraManager.h b/Engine/Source/Runtime/Engine/Classes/Camera/PlayerCameraManager.h
--- a/Engine/Source/Runtime/Engine/Classes/Camera/PlayerCameraManager.h
+++ b/Engine/Source/Runtime/Engine/Classes/Camera/PlayerCameraManager.h
@@ -3,9 +3,6 @@
 #pragma once
 #include "GameFramework/Actor.h"
 #include "CameraTypes.h"
-#include "Engine/Scene.h"
-#include "Engine/EngineBaseTypes.h"
-
 #include "PlayerCameraManager.generated.h"
 
 class APawn;
diff --git a/Engine/Source/Runtime/Engine/Classes/Engine/TriggerBase.h b/Engine/Source/Runtime/Engine/Classes/Engine/TriggerBase.h
--- a/Engine/Source/Runtime/Engine/Classes/Engine/TriggerBase.h
+++ b/Engine/Source/Runtime/Engine/Classes/Engine/TriggerBase.h
@@ -5,11 +5,12 @@
  */
 
 #pragma once
-
-#include "EngineBaseTypes.h"
-
+#include "GameFramework/Actor.h"
 #include "TriggerBase.generated.h"
 
+class UShapeComponent;
+class UBillboardComponent;
+
 UCLASS(ClassGroup=Common, abstract, ConversionRoot, MinimalAPI)
 class ATriggerBase : public AActor
 {
@@ -17,7 +18,7 @@ class ATriggerBase : public AActor
 
 	/** Base collision component */
 	UPROPERTY(Category=TriggerBase, VisibleAnywhere, BlueprintReadOnly)
-	TSubobjectPtr<class UShapeComponent> CollisionComponent;
+	TSubobjectPtr<UShapeComponent> CollisionComponent;
 
 	// Reference to the billboard component
 	UPROPERTY()
diff --git a/Engine/Source/Runtime/Engine/Classes/Matinee/MatineeActorCameraAnim.h b/Engine/Source/Runtime/Engine/Classes/Matinee/MatineeActorCameraAnim.h
--- a/Engine/Source/Runtime/Engine/Classes/Matinee/MatineeActorCameraAnim.h
+++ b/Engine/Source/Runtime/Engine/Classes/Matinee/MatineeActorCameraAnim.h
@@ -1,12 +1,11 @@
 // Copyright 1998-2014 Epic Games, Inc. All Rights Reserved.
 
 #pragma once
-
 #include "Matinee/MatineeActor.h"
-#include "Engine/EngineBaseTypes.h"
-
 #include "MatineeActorCameraAnim.generated.h"
 
+class UCameraAnim;
+
 /**
  * Actor used to control temporary matinees for camera anims that only exist in the editor
  */
@@ -17,7 +16,7 @@ class AMatineeActorCameraAnim : public AMatineeActor
 
 	/** The camera anim we are editing */
 	UPROPERTY(Transient)
-	class UCameraAnim* CameraAnim;
+	UCameraAnim* CameraAnim;
 
 	// Begin UObject interface
 	virtual bool NeedsLoadForClient() const override
diff --git a/Engine/Source/Runtime/Engine/Public/AI/NavDataGenerator.h b/Engine/Source/Runtime/Engine/Public/AI/NavDataGenerator.h
--- a/Engine/Source/Runtime/Engine/Public/AI/NavDataGenerator.h
+++ b/Engine/Source/Runtime/Engine/Public/AI/NavDataGenerator.h
@@ -2,6 +2,8 @@
 
 #pragma once 
 
+struct FNavigationDirtyArea;
+
 class FNavDataGenerator : public TSharedFromThis<FNavDataGenerator, ESPMode::ThreadSafe>
 {
 public:
