Persistent level always added/build for a texture streaming in the Editor.
Landscape correctly returns lightmaps and shadowmaps as streaming textures.
Foliage instances threated as static by streaming system.
Heightmaps, weighmaps, lightmaps and shadowmaps treated as static when building level streaming data.

diff --git a/Engine/Source/Runtime/Engine/Private/EnginePrivate.h b/Engine/Source/Runtime/Engine/Private/EnginePrivate.h
--- a/Engine/Source/Runtime/Engine/Private/EnginePrivate.h
+++ b/Engine/Source/Runtime/Engine/Private/EnginePrivate.h
@@ -316,6 +316,7 @@
 #include "CanvasTypes.h"							// Canvas.
 #include "EngineUtils.h"
 #include "InstancedFoliage.h"				// Instanced foliage.
+#include "Foliage/InstancedFoliageActor.h"	// Instanced foliage.
 #include "UnrealExporter.h"					// Exporter definition.
 #include "TimerManager.h"					// Game play timers
 #include "EngineService.h"
diff --git a/Engine/Source/Runtime/Engine/Private/Landscape/LandscapeRender.cpp b/Engine/Source/Runtime/Engine/Private/Landscape/LandscapeRender.cpp
--- a/Engine/Source/Runtime/Engine/Private/Landscape/LandscapeRender.cpp
+++ b/Engine/Source/Runtime/Engine/Private/Landscape/LandscapeRender.cpp
@@ -2401,6 +2401,39 @@ void ULandscapeComponent::GetStreamingTextureInfo(TArray<FStreamingTexturePrimit
 				}
 			}
 		}
+
+		// Lightmap
+		FLightMap2D* Lightmap = LightMap ? LightMap->GetLightMap2D() : nullptr;
+		uint32 LightmapIndex = AllowHighQualityLightmaps() ? 0 : 1;
+		if (Lightmap && Lightmap->IsValid(LightmapIndex))
+		{
+			const FVector2D& Scale = Lightmap->GetCoordinateScale();
+			if (Scale.X > SMALL_NUMBER && Scale.Y > SMALL_NUMBER)
+			{
+				float LightmapFactorX = TexelFactor / Scale.X;
+				float LightmapFactorY = TexelFactor / Scale.Y;
+				FStreamingTexturePrimitiveInfo& StreamingTexture = *new(OutStreamingTextures) FStreamingTexturePrimitiveInfo;
+				StreamingTexture.Bounds		 = BoundingSphere;
+				StreamingTexture.TexelFactor = FMath::Max(LightmapFactorX, LightmapFactorY);
+				StreamingTexture.Texture	 = Lightmap->GetTexture(LightmapIndex);
+			}
+		}
+		
+		// Shadowmap
+		FShadowMap2D* Shadowmap = ShadowMap ? ShadowMap->GetShadowMap2D() : nullptr;
+		if (Shadowmap && Shadowmap->IsValid())
+		{
+			const FVector2D& Scale = Shadowmap->GetCoordinateScale();
+			if (Scale.X > SMALL_NUMBER && Scale.Y > SMALL_NUMBER)
+			{
+				float ShadowmapFactorX		 = TexelFactor / Scale.X;
+				float ShadowmapFactorY		 = TexelFactor / Scale.Y;
+				FStreamingTexturePrimitiveInfo& StreamingTexture = *new(OutStreamingTextures) FStreamingTexturePrimitiveInfo;
+				StreamingTexture.Bounds		 = BoundingSphere;
+				StreamingTexture.TexelFactor = FMath::Max(ShadowmapFactorX, ShadowmapFactorY);
+				StreamingTexture.Texture	 = Shadowmap->GetTexture();
+			}
+		}
 	}
 
 	// Weightmap
diff --git a/Engine/Source/Runtime/Engine/Private/Level.cpp b/Engine/Source/Runtime/Engine/Private/Level.cpp
--- a/Engine/Source/Runtime/Engine/Private/Level.cpp
+++ b/Engine/Source/Runtime/Engine/Private/Level.cpp
@@ -1328,7 +1328,8 @@ void ULevel::BuildStreamingData(UTexture2D* UpdateSpecificTextureOnly/*=NULL*/)
 				const AActor* const Owner				= Primitive->GetOwner();
 				const bool bIsStaticMeshComponent		= Primitive->IsA(UStaticMeshComponent::StaticClass());
 				const bool bIsSkeletalMeshComponent		= Primitive->IsA(USkeletalMeshComponent::StaticClass());
-				const bool bIsStatic					= Owner == NULL || Primitive->Mobility == EComponentMobility::Static;
+				const bool bIsFoliage					= Owner && Owner->IsA(AInstancedFoliageActor::StaticClass()) && Primitive->IsA(UInstancedStaticMeshComponent::StaticClass());
+				const bool bIsStatic					= Owner == NULL || Primitive->Mobility == EComponentMobility::Static || bIsFoliage;
 
 				TArray<FStreamingTexturePrimitiveInfo> PrimitiveStreamingTextures;
 
@@ -1380,7 +1381,11 @@ void ULevel::BuildStreamingData(UTexture2D* UpdateSpecificTextureOnly/*=NULL*/)
 						const bool bIsWorldTexture			= 
 							Texture2D->LODGroup == TEXTUREGROUP_World ||
 							Texture2D->LODGroup == TEXTUREGROUP_WorldNormalMap ||
-							Texture2D->LODGroup == TEXTUREGROUP_WorldSpecular;
+							Texture2D->LODGroup == TEXTUREGROUP_WorldSpecular ||
+							Texture2D->LODGroup == TEXTUREGROUP_Terrain_Heightmap ||
+							Texture2D->LODGroup == TEXTUREGROUP_Terrain_Weightmap ||
+							Texture2D->LODGroup == TEXTUREGROUP_Shadowmap ||
+							Texture2D->LODGroup == TEXTUREGROUP_Lightmap;
 
 						// Check if we should consider this a static mesh texture instance.
 						bool bIsStaticMeshTextureInstance = bIsWorldTexture && !bIsSkeletalMeshComponent;
diff --git a/Engine/Source/Runtime/Engine/Private/World.cpp b/Engine/Source/Runtime/Engine/Private/World.cpp
--- a/Engine/Source/Runtime/Engine/Private/World.cpp
+++ b/Engine/Source/Runtime/Engine/Private/World.cpp
@@ -1265,6 +1265,7 @@ void UWorld::UpdateWorldComponents(bool bRerunConstructionScripts, bool bCurrent
 			if(!StreamingLevel || Level->bIsVisible)
 			{
 				Level->UpdateLevelComponents(bRerunConstructionScripts);
+				IStreamingManager::Get().AddLevel(Level);
 			}
 		}
 	}
