only include where needed policy

diff --git a/Engine/Source/Runtime/Engine/Private/EnginePrivate.h b/Engine/Source/Runtime/Engine/Private/EnginePrivate.h
--- a/Engine/Source/Runtime/Engine/Private/EnginePrivate.h
+++ b/Engine/Source/Runtime/Engine/Private/EnginePrivate.h
@@ -316,7 +316,6 @@
 #include "CanvasTypes.h"							// Canvas.
 #include "EngineUtils.h"
 #include "InstancedFoliage.h"				// Instanced foliage.
-#include "Foliage/InstancedFoliageActor.h"	// Instanced foliage.
 #include "UnrealExporter.h"					// Exporter definition.
 #include "TimerManager.h"					// Game play timers
 #include "EngineService.h"
diff --git a/Engine/Source/Runtime/Engine/Private/Level.cpp b/Engine/Source/Runtime/Engine/Private/Level.cpp
--- a/Engine/Source/Runtime/Engine/Private/Level.cpp
+++ b/Engine/Source/Runtime/Engine/Private/Level.cpp
@@ -24,6 +24,7 @@ Level.cpp: Level-related functions
 #include "LevelUtils.h"
 #include "TargetPlatform.h"
 #include "ContentStreaming.h"
+#include "Foliage/InstancedFoliageActor.h"
 
 DEFINE_LOG_CATEGORY(LogLevel);
 
