Moved FAnimUpdateRateParameters from Actor into SkinnedMeshComponent

diff --git a/Engine/Source/Runtime/Engine/Classes/Components/SkeletalMeshComponent.h b/Engine/Source/Runtime/Engine/Classes/Components/SkeletalMeshComponent.h
--- a/Engine/Source/Runtime/Engine/Classes/Components/SkeletalMeshComponent.h
+++ b/Engine/Source/Runtime/Engine/Classes/Components/SkeletalMeshComponent.h
@@ -812,6 +812,9 @@ public:
 	virtual void SetPhysicsAsset(class UPhysicsAsset* NewPhysicsAsset,bool bForceReInit = false) override;
 	virtual void SetSkeletalMesh(class USkeletalMesh* NewMesh) override;
 	virtual FVector GetSkinnedVertexPosition(int32 VertexIndex) const override;
+
+	virtual bool IsPlayingRootMotion() override;
+
 	// End USkinnedMeshComponent interface
 	/** 
 	 *	Iterate over each joint in the physics for this mesh, setting its AngularPositionTarget based on the animation information.
diff --git a/Engine/Source/Runtime/Engine/Classes/Components/SkinnedMeshComponent.h b/Engine/Source/Runtime/Engine/Classes/Components/SkinnedMeshComponent.h
--- a/Engine/Source/Runtime/Engine/Classes/Components/SkinnedMeshComponent.h
+++ b/Engine/Source/Runtime/Engine/Classes/Components/SkinnedMeshComponent.h
@@ -841,4 +841,28 @@ private:
 	 */
 	template <bool bExtraBoneInfluencesT, bool bCachedMatrices>
 	FVector GetTypedSkinnedVertexPosition(const FSkelMeshChunk& Chunk, const FSkeletalMeshVertexBuffer& VertexBufferGPUSkin, int32 VertIndex, bool bSoftVertex, const TArray<FMatrix> & RefToLocals = TArray<FMatrix>()) const;
+
+	// Animation update rate control.
+public:
+	/** Animation Update Rate optimization parameters. */
+	UPROPERTY(Transient)
+	struct FAnimUpdateRateParameters AnimUpdateRateParams;
+
+	/** Unique Tag assigned to spread updates of SkinnedMeshes over time. */
+	UPROPERTY(Transient)
+	uint32 AnimUpdateRateShiftTag;
+
+	/** Frame counter to call AnimUpdateRateTick() just once per frame. */
+	UPROPERTY(Transient)
+	uint32 AnimUpdateRateFrameCount;
+
+	/** Aimation Update Rate Tick. */
+	void AnimUpdateRateTick();
+
+	/** Updates AnimUpdateRateParams, used by SkinnedMeshComponents.
+	* @param bRecentlyRendered : true if at least one SkinnedMeshComponent on this Actor has been rendered in the last second.
+	* @param MaxDistanceFactor : Largest SkinnedMeshComponent of this Actor drawn on screen. */
+	void AnimUpdateRateSetParams(const bool & bRecentlyRendered, const float & MaxDistanceFactor, const bool & bPlayingRootMotion);
+
+	virtual bool IsPlayingRootMotion(){ return false; }
 };
diff --git a/Engine/Source/Runtime/Engine/Classes/Engine/EngineTypes.h b/Engine/Source/Runtime/Engine/Classes/Engine/EngineTypes.h
--- a/Engine/Source/Runtime/Engine/Classes/Engine/EngineTypes.h
+++ b/Engine/Source/Runtime/Engine/Classes/Engine/EngineTypes.h
@@ -1596,7 +1596,7 @@ public:
 	 * @param : NewEvaluationRate. How often animation will be evaluated. 1 = every frame, 2 = every 2 frames, etc.
 	 * @param : bNewInterpSkippedFrames. When skipping a frame, should it be interpolated or frozen?
 	 */
-	void Set(const class AActor & Owner, const int32 & NewUpdateRate, const int32 & NewEvaluationRate, const bool & bNewInterpSkippedFrames);
+	void Set(const class USkinnedMeshComponent & Owner, const int32 & NewUpdateRate, const int32 & NewEvaluationRate, const bool & bNewInterpSkippedFrames);
 
 	/* Getter for UpdateRate */
 	int32 GetUpdateRate() const
diff --git a/Engine/Source/Runtime/Engine/Classes/GameFramework/Actor.h b/Engine/Source/Runtime/Engine/Classes/GameFramework/Actor.h
--- a/Engine/Source/Runtime/Engine/Classes/GameFramework/Actor.h
+++ b/Engine/Source/Runtime/Engine/Classes/GameFramework/Actor.h
@@ -288,28 +288,6 @@ public:
 	UPROPERTY(transient)
 	TArray<AActor*> Children;    
 	
-	// Animation update rate control.
-public:
-	/** Unique Tag assigned to spread updates of SkinnedMeshes over time. */
-	UPROPERTY(Transient)
-	uint32 AnimUpdateRateShiftTag;
-
-	/** Frame counter to call AnimUpdateRateTick() just once per frame. */
-	UPROPERTY(Transient)
-	uint32 AnimUpdateRateFrameCount;
-
-	/** Animation Update Rate optimization parameters. */
-	UPROPERTY(Transient)
-	struct FAnimUpdateRateParameters AnimUpdateRateParams;
-
-	/** Aimation Update Rate Tick. */
-	void AnimUpdateRateTick();
-
-	/** Updates AnimUpdateRateParams, used by SkinnedMeshComponents.
-	 * @param bRecentlyRendered : true if at least one SkinnedMeshComponent on this Actor has been rendered in the last second.
-	 * @param MaxDistanceFactor : Largest SkinnedMeshComponent of this Actor drawn on screen. */
-	void AnimUpdateRateSetParams(const bool & bRecentlyRendered, const float & MaxDistanceFactor, const bool & bPlayingRootMotion);
-
 protected:
 	// Collision primitive.
 	UPROPERTY()
diff --git a/Engine/Source/Runtime/Engine/Private/Actor.cpp b/Engine/Source/Runtime/Engine/Private/Actor.cpp
--- a/Engine/Source/Runtime/Engine/Private/Actor.cpp
+++ b/Engine/Source/Runtime/Engine/Private/Actor.cpp
@@ -73,114 +73,6 @@ AActor::AActor(const class FPostConstructInitializeProperties& PCIP)
 	bCanBeDamaged = true;
 	bFindCameraComponentWhenViewTarget = true;
 	bAllowReceiveTickEventOnDedicatedServer = true;
-	AnimUpdateRateShiftTag = 0;
-	AnimUpdateRateFrameCount = 0;
-}
-
-// Global counter to spread SkinnedMeshComponent tick updates.
-static uint32 SkinnedMeshUpdateRateGroupCount = 0;
-
-void AActor::AnimUpdateRateTick()
-{
-	// If hasn't been initialized yet, pick a unique ID, to spread population over frames.
-	if( AnimUpdateRateShiftTag == 0 )
-	{
-		AnimUpdateRateShiftTag = ++SkinnedMeshUpdateRateGroupCount;
-	}
-
-	// Go through components and figure out if they've been recently rendered, and the biggest MaxDistanceFactor
-	bool bRecentlyRendered = false;
-	bool bPlayingRootMotion = false;
-	float MaxDistanceFactor = 0.f;
-
-	// Gather Actor's components
-	TArray<USceneComponent *> ComponentStack;
-	USceneComponent const * CurrentComponent = GetRootComponent();
-
-	while( CurrentComponent )
-	{
-		// push children on the stack so they get tested later
-		ComponentStack.Append(CurrentComponent->AttachChildren);
-
-		if( CurrentComponent->IsRegistered() )
-		{
-			USkeletalMeshComponent const * SkelMeshComp = Cast<USkeletalMeshComponent>(CurrentComponent);
-			if( SkelMeshComp )
-			{
-				bRecentlyRendered = bRecentlyRendered || (SkelMeshComp->LastRenderTime > (GetWorld()->TimeSeconds - 1.f));
-				MaxDistanceFactor = FMath::Max(MaxDistanceFactor, SkelMeshComp->MaxDistanceFactor);
-				bPlayingRootMotion = bPlayingRootMotion || (SkelMeshComp->AnimScriptInstance ? (SkelMeshComp->AnimScriptInstance->GetRootMotionMontageInstance() != NULL) : false);
-			}
-			else
-			{
-				USkinnedMeshComponent const * SkinMeshComp = Cast<USkinnedMeshComponent>(CurrentComponent);
-				if( SkinMeshComp )
-				{
-					bRecentlyRendered = bRecentlyRendered || (SkinMeshComp->LastRenderTime > (GetWorld()->TimeSeconds - 1.f));
-					MaxDistanceFactor = FMath::Max(MaxDistanceFactor, SkinMeshComp->MaxDistanceFactor);
-				}
-			}
-		}
-
-		CurrentComponent = (ComponentStack.Num() > 0) ? ComponentStack.Pop() : NULL;
-	}
-
-	// Figure out which update rate should be used.
-	AnimUpdateRateSetParams(bRecentlyRendered, MaxDistanceFactor, bPlayingRootMotion);
-}
-
-void AActor::AnimUpdateRateSetParams(const bool & bRecentlyRendered, const float & MaxDistanceFactor, const bool & bPlayingRootMotion)
-{
-	// default rules for setting update rates
-
-	// Human controlled characters should be ticked always fully to minimize latency w/ game play events triggered by animation.
-	AController * Controller = GetInstigatorController();
-	const bool bHumanControlled = Controller && (Cast<APlayerController>(Controller) != NULL);
-
-	// Not rendered, including dedicated servers. we can skip the Evaluation part.
-	if( !bRecentlyRendered )
-	{
-		AnimUpdateRateParams.Set(*this, (bHumanControlled ?  1 : 4), 4, false);
-	}
-	// Visible controlled characters or playing root motion. Need evaluation and ticking done every frame.
-	else  if( bHumanControlled || bPlayingRootMotion )
-	{
-		AnimUpdateRateParams.Set(*this, 1, 1, false);
-	}
-	else
-	{
-		// For visible meshes, figure out how often bones should be evaluated VS interpolated to previous evaluation.
-		// This is based on screen size and makes sense for a 3D FPS game, different games or Actors can implement different rules.
-		int32 DesiredEvaluationRate;
-		if( MaxDistanceFactor > 0.4f )
-		{
-			DesiredEvaluationRate = 1;
-		}
-		else if( MaxDistanceFactor > .2f )
-		{
-			DesiredEvaluationRate = 2;
-		}
-		else 
-		{
-			DesiredEvaluationRate = 3;
-		}
-
-		AnimUpdateRateParams.Set(*this, DesiredEvaluationRate, DesiredEvaluationRate, true);
-	}
-}
-
-void FAnimUpdateRateParameters::Set(const class AActor & Owner, const int32 & NewUpdateRate, const int32 & NewEvaluationRate, const bool & bNewInterpSkippedFrames)
-{
-	UpdateRate = FMath::Max(NewUpdateRate, 1);
-	// Make sure EvaluationRate is a multiple of UpdateRate.
-	EvaluationRate = FMath::Max((NewEvaluationRate / UpdateRate) * UpdateRate, 1);
-	bInterpolateSkippedFrames = bNewInterpSkippedFrames;	
-
-	// Make sure we don't overflow. we don't need very large numbers.
-	const uint32 Counter = (GFrameCounter % MAX_uint8) + (Owner.AnimUpdateRateShiftTag % MAX_uint8);
-
-	bSkipUpdate = ((Counter % UpdateRate) > 0);
-	bSkipEvaluation = ((Counter % EvaluationRate) > 0);
 }
 
 void FActorTickFunction::ExecuteTick(float DeltaTime, enum ELevelTick TickType, ENamedThreads::Type CurrentThread, const FGraphEventRef& MyCompletionGraphEvent)
diff --git a/Engine/Source/Runtime/Engine/Private/SkeletalMesh.cpp b/Engine/Source/Runtime/Engine/Private/SkeletalMesh.cpp
--- a/Engine/Source/Runtime/Engine/Private/SkeletalMesh.cpp
+++ b/Engine/Source/Runtime/Engine/Private/SkeletalMesh.cpp
@@ -4465,6 +4465,9 @@ USkinnedMeshComponent::USkinnedMeshComponent(const class FPostConstructInitializ
 	SectionIndexPreview = -1;
 #endif // WITH_EDITORONLY_DATA
 	bPerBoneMotionBlur = true;
+
+	AnimUpdateRateShiftTag = 0;
+	AnimUpdateRateFrameCount = 0;
 }
 
 
diff --git a/Engine/Source/Runtime/Engine/Private/SkeletalMeshComponent.cpp b/Engine/Source/Runtime/Engine/Private/SkeletalMeshComponent.cpp
--- a/Engine/Source/Runtime/Engine/Private/SkeletalMeshComponent.cpp
+++ b/Engine/Source/Runtime/Engine/Private/SkeletalMeshComponent.cpp
@@ -917,18 +917,16 @@ void USkeletalMeshComponent::RefreshBoneTransforms(FActorComponentTickFunction*
 	AnimEvaluationContext.AnimInstance = AnimScriptInstance;
 
 	//Handle update rate optimization setup
-	const FAnimUpdateRateParameters  & UpdateRateParams = Owner ? Owner->AnimUpdateRateParams : FAnimUpdateRateParameters();
-	
-	const bool bDoUpdateRateOptimization = bEnableUpdateRateOptimizations && (UpdateRateParams.GetEvaluationRate() > 1);
+	const bool bDoUpdateRateOptimization = bEnableUpdateRateOptimizations && (AnimUpdateRateParams.GetEvaluationRate() > 1);
 	//Dont mark cache as invalid if we aren't performing optimization anyway
 	const bool bInvalidCachedBones = bDoUpdateRateOptimization &&
 									( (LocalAtoms.Num() != SkeletalMesh->RefSkeleton.GetNum())
 									  || (LocalAtoms.Num() != CachedLocalAtoms.Num())
 									  || (SpaceBases.Num() != CachedSpaceBases.Num()) );
 
-	AnimEvaluationContext.bDoEvaluation = !bDoUpdateRateOptimization || bInvalidCachedBones || !UpdateRateParams.ShouldSkipEvaluation();
+	AnimEvaluationContext.bDoEvaluation = !bDoUpdateRateOptimization || bInvalidCachedBones || !AnimUpdateRateParams.ShouldSkipEvaluation();
 	
-	AnimEvaluationContext.bDoInterpolation = bDoUpdateRateOptimization && !bInvalidCachedBones && UpdateRateParams.ShouldInterpolateSkippedFrames();
+	AnimEvaluationContext.bDoInterpolation = bDoUpdateRateOptimization && !bInvalidCachedBones && AnimUpdateRateParams.ShouldInterpolateSkippedFrames();
 	AnimEvaluationContext.bDuplicateToCacheBones = bInvalidCachedBones || (bDoUpdateRateOptimization && AnimEvaluationContext.bDoEvaluation && !AnimEvaluationContext.bDoInterpolation);
 
 	if (!bDoUpdateRateOptimization)
@@ -997,10 +995,8 @@ void USkeletalMeshComponent::PostAnimEvaluation(FAnimationEvaluationContext& Eva
 	if (EvaluationContext.bDoInterpolation)
 	{
 		SCOPE_CYCLE_COUNTER(STAT_InterpolateSkippedFrames);
-		AActor * Owner = GetOwner();
-		const FAnimUpdateRateParameters  & UpdateRateParams = Owner ? Owner->AnimUpdateRateParams : FAnimUpdateRateParameters();
 
-		const float Alpha = 0.25f + (1.f / float(FMath::Max(UpdateRateParams.GetEvaluationRate(), 2) * 2));
+		const float Alpha = 0.25f + (1.f / float(FMath::Max(AnimUpdateRateParams.GetEvaluationRate(), 2) * 2));
 		FAnimationRuntime::LerpBoneTransforms(LocalAtoms, CachedLocalAtoms, Alpha, RequiredBones);
 		FAnimationRuntime::LerpBoneTransforms(SpaceBases, CachedSpaceBases, Alpha, RequiredBones);
 	}
@@ -1756,4 +1752,9 @@ void USkeletalMeshComponent::ValidateAnimation()
 	}
 }
 
+bool USkeletalMeshComponent::IsPlayingRootMotion()
+{
+	return (AnimScriptInstance ? (AnimScriptInstance->GetRootMotionMontageInstance() != NULL) : false);
+}
+
 #endif
diff --git a/Engine/Source/Runtime/Engine/Private/SkinnedMeshComponent.cpp b/Engine/Source/Runtime/Engine/Private/SkinnedMeshComponent.cpp
--- a/Engine/Source/Runtime/Engine/Private/SkinnedMeshComponent.cpp
+++ b/Engine/Source/Runtime/Engine/Private/SkinnedMeshComponent.cpp
@@ -237,32 +237,28 @@ void USkinnedMeshComponent::TickUpdateRate()
 	SCOPE_CYCLE_COUNTER(STAT_TickUpdateRate);
 	if( bEnableUpdateRateOptimizations )
 	{
-		AActor * Owner = GetOwner();
-		if( Owner )
-		{
-			// Convert current frame counter from 64 to 32 bits.
-			const uint32 CurrentFrame32 = uint32(GFrameCounter % MAX_uint32);
+		// Convert current frame counter from 64 to 32 bits.
+		const uint32 CurrentFrame32 = uint32(GFrameCounter % MAX_uint32);
 
-			// Tick Owner once per frame. All attached SkinnedMeshComponents will share the same settings.
-			if( Owner->AnimUpdateRateFrameCount != CurrentFrame32 )
-			{
-				Owner->AnimUpdateRateFrameCount = CurrentFrame32;
-				Owner->AnimUpdateRateTick();
-			}
+		// Tick Owner once per frame. All attached SkinnedMeshComponents will share the same settings.
+		if( AnimUpdateRateFrameCount != CurrentFrame32 )
+		{
+			AnimUpdateRateFrameCount = CurrentFrame32;
+			AnimUpdateRateTick();
+		}
 
-			// debug -- @todo: hook this up to a console command.
-			if( bDisplayDebugUpdateRateOptimizations )
+		// debug -- @todo: hook this up to a console command.
+		if( bDisplayDebugUpdateRateOptimizations )
+		{
+			FColor DrawColor;
+			switch( AnimUpdateRateParams.GetUpdateRate() )
 			{
-				FColor DrawColor;
-				switch( Owner->AnimUpdateRateParams.GetUpdateRate() )
-				{
-				case 1 : DrawColor = FColor::Red; break;
-				case 2 : DrawColor = FColor::Green; break;
-				case 3 : DrawColor = FColor::Blue; break;
-				default: DrawColor = FColor::Black; break;
-				}
-				DrawDebugBox(GetWorld(), Bounds.Origin, Bounds.BoxExtent, FQuat::Identity, DrawColor, false);
+			case 1 : DrawColor = FColor::Red; break;
+			case 2 : DrawColor = FColor::Green; break;
+			case 3 : DrawColor = FColor::Blue; break;
+			default: DrawColor = FColor::Black; break;
 			}
+			DrawDebugBox(GetWorld(), Bounds.Origin, Bounds.BoxExtent, FQuat::Identity, DrawColor, false);
 		}
 	}
 }
@@ -278,14 +274,11 @@ void USkinnedMeshComponent::TickComponent(float DeltaTime, enum ELevelTick TickT
 	bRecentlyRendered = (LastRenderTime > GetWorld()->TimeSeconds - 1.0f);
 	TickUpdateRate();
 
-	AActor * Owner = GetOwner();
-	const FAnimUpdateRateParameters  & UpdateRateParams = Owner ? Owner->AnimUpdateRateParams : FAnimUpdateRateParameters();
-
 	bPoseTicked = false;
 	// Tick Pose first
 	if( ShouldTickPose() )
 	{
-		if( bEnableUpdateRateOptimizations && UpdateRateParams.ShouldSkipUpdate() )
+		if (bEnableUpdateRateOptimizations && AnimUpdateRateParams.ShouldSkipUpdate())
 		{
 			SkippedTickDeltaTime += DeltaTime;
 			if( !bRecentlyRendered )
@@ -1719,3 +1712,102 @@ TArray<FActiveVertexAnim> USkinnedMeshComponent::UpdateActiveVertexAnims(const U
 	return OutVertexAnims;
 }
 
+
+// Global counter to spread SkinnedMeshComponent tick updates.
+static uint32 SkinnedMeshUpdateRateGroupCount = 0;
+
+void USkinnedMeshComponent::AnimUpdateRateTick()
+{
+	// If hasn't been initialized yet, pick a unique ID, to spread population over frames.
+	if (AnimUpdateRateShiftTag == 0)
+	{
+		AnimUpdateRateShiftTag = ++SkinnedMeshUpdateRateGroupCount;
+	}
+
+	// Go through components and figure out if they've been recently rendered, and the biggest MaxDistanceFactor
+	bool bRecentlyRendered = false;
+	bool bPlayingRootMotion = false;
+	float MaxDistanceFactor = 0.f;
+
+	// Gather Actor's components
+	TArray<USceneComponent *> ComponentStack;
+	USceneComponent* CurrentComponent = this;
+
+	while (CurrentComponent)
+	{
+		// push children on the stack so they get tested later
+		ComponentStack.Append(CurrentComponent->AttachChildren);
+
+		if (CurrentComponent->IsRegistered())
+		{
+			USkinnedMeshComponent* SkinMeshComp = Cast<USkinnedMeshComponent>(CurrentComponent);
+			if (SkinMeshComp)
+			{
+				bRecentlyRendered = bRecentlyRendered || (SkinMeshComp->LastRenderTime > (GetWorld()->TimeSeconds - 1.f));
+				MaxDistanceFactor = FMath::Max(MaxDistanceFactor, SkinMeshComp->MaxDistanceFactor);
+				bPlayingRootMotion = bPlayingRootMotion || SkinMeshComp->IsPlayingRootMotion();
+			}
+		}
+
+		CurrentComponent = (ComponentStack.Num() > 0) ? ComponentStack.Pop() : NULL;
+	}
+
+	// Figure out which update rate should be used.
+	AnimUpdateRateSetParams(bRecentlyRendered, MaxDistanceFactor, bPlayingRootMotion);
+}
+
+void USkinnedMeshComponent::AnimUpdateRateSetParams(const bool & bRecentlyRendered, const float & MaxDistanceFactor, const bool & bPlayingRootMotion)
+{
+	// default rules for setting update rates
+
+	// Human controlled characters should be ticked always fully to minimize latency w/ game play events triggered by animation.
+	AActor* Owner = GetOwner();
+	AController * Controller = Owner ? Owner->GetInstigatorController() : NULL;
+	const bool bHumanControlled = Controller && (Cast<APlayerController>(Controller) != NULL);
+
+	// Not rendered, including dedicated servers. we can skip the Evaluation part.
+	if (!bRecentlyRendered)
+	{
+		AnimUpdateRateParams.Set(*this, (bHumanControlled ? 1 : 4), 4, false);
+	}
+	// Visible controlled characters or playing root motion. Need evaluation and ticking done every frame.
+	else  if (bHumanControlled || bPlayingRootMotion)
+	{
+		AnimUpdateRateParams.Set(*this, 1, 1, false);
+	}
+	else
+	{
+		// For visible meshes, figure out how often bones should be evaluated VS interpolated to previous evaluation.
+		// This is based on screen size and makes sense for a 3D FPS game, different games or Actors can implement different rules.
+		int32 DesiredEvaluationRate;
+		if (MaxDistanceFactor > 0.4f)
+		{
+			DesiredEvaluationRate = 1;
+		}
+		else if (MaxDistanceFactor > .2f)
+		{
+			DesiredEvaluationRate = 2;
+		}
+		else
+		{
+			DesiredEvaluationRate = 3;
+		}
+
+		AnimUpdateRateParams.Set(*this, DesiredEvaluationRate, DesiredEvaluationRate, true);
+	}
+}
+
+void FAnimUpdateRateParameters::Set(const class USkinnedMeshComponent& Owner, const int32 & NewUpdateRate, const int32 & NewEvaluationRate, const bool & bNewInterpSkippedFrames)
+{
+	UpdateRate = FMath::Max(NewUpdateRate, 1);
+	// Make sure EvaluationRate is a multiple of UpdateRate.
+	EvaluationRate = FMath::Max((NewEvaluationRate / UpdateRate) * UpdateRate, 1);
+	bInterpolateSkippedFrames = bNewInterpSkippedFrames;
+
+	// Make sure we don't overflow. we don't need very large numbers.
+	const uint32 Counter = (GFrameCounter % MAX_uint8) + (Owner.AnimUpdateRateShiftTag % MAX_uint8);
+
+	bSkipUpdate = ((Counter % UpdateRate) > 0);
+	bSkipEvaluation = ((Counter % EvaluationRate) > 0);
+}
+
