Move EEndPlayReason to EngineTypes.h to clean up Actor.h a little
Remove long-deprecated Layer_DEPRECATED and unused Group_DEPRECATED
Remove unused FPrimitiveTriangleVertex and FPrimitiveTriangleDefinitionInterface (from PrimitiveComponent.h)
Change Engine.h to include EngineMinimal.h (and remove those headers from Engine.h)
Add DecalActor.cpp and move relevant functions there
Do a little more culling on EnginePrivate.h

diff --git a/Engine/Source/Runtime/Engine/Classes/Components/PrimitiveComponent.h b/Engine/Source/Runtime/Engine/Classes/Components/PrimitiveComponent.h
--- a/Engine/Source/Runtime/Engine/Classes/Components/PrimitiveComponent.h
+++ b/Engine/Source/Runtime/Engine/Classes/Components/PrimitiveComponent.h
@@ -18,24 +18,6 @@
 class FPrimitiveSceneProxy;
 class AController; 
 
-/** Information about a vertex of a primitive's triangle. */
-struct FPrimitiveTriangleVertex
-{
-	FVector WorldPosition;
-	FVector WorldTangentX;
-	FVector WorldTangentY;
-	FVector WorldTangentZ;
-};
-
-/** An interface to some consumer of the primitive's triangles. */
-class FPrimitiveTriangleDefinitionInterface
-{
-public:
-
-	/** Defines a triangle by its vertices. */
-	virtual void DefineTriangle(const TStaticArray<FPrimitiveTriangleVertex,3>& InVertices,const UMaterialInterface* Material) = 0;
-};
-
 /** Information about a streaming texture that a primitive uses for rendering. */
 struct FStreamingTexturePrimitiveInfo
 {
@@ -44,8 +26,6 @@ struct FStreamingTexturePrimitiveInfo
 	float TexelFactor;
 };
 
-
-
 UENUM()
 enum ECanBeCharacterBase
 {
@@ -74,8 +54,6 @@ namespace EHasCustomNavigableGeometry
 	};
 };
 
-
-
 /** Information about the sprite category */
 USTRUCT()
 struct FSpriteCategoryInfo
diff --git a/Engine/Source/Runtime/Engine/Classes/Components/SceneComponent.h b/Engine/Source/Runtime/Engine/Classes/Components/SceneComponent.h
--- a/Engine/Source/Runtime/Engine/Classes/Components/SceneComponent.h
+++ b/Engine/Source/Runtime/Engine/Classes/Components/SceneComponent.h
@@ -2,7 +2,7 @@
 
 #pragma once
 
-#include "ActorComponent.h"
+#include "Components/ActorComponent.h"
 #include "SceneComponent.generated.h"
 
 /** Overlap info consisting of the primitive and the body that is overlapping */
diff --git a/Engine/Source/Runtime/Engine/Classes/Engine/BrushBuilder.h b/Engine/Source/Runtime/Engine/Classes/Engine/BrushBuilder.h
--- a/Engine/Source/Runtime/Engine/Classes/Engine/BrushBuilder.h
+++ b/Engine/Source/Runtime/Engine/Classes/Engine/BrushBuilder.h
@@ -69,9 +69,6 @@ protected:
 	UPROPERTY()
 	TArray<struct FBuilderPoly> Polys;
 
-	UPROPERTY()
-	FName Group_DEPRECATED;
-
 	UPROPERTY()
 	FName Layer;
 
diff --git a/Engine/Source/Runtime/Engine/Classes/Engine/EngineTypes.h b/Engine/Source/Runtime/Engine/Classes/Engine/EngineTypes.h
--- a/Engine/Source/Runtime/Engine/Classes/Engine/EngineTypes.h
+++ b/Engine/Source/Runtime/Engine/Classes/Engine/EngineTypes.h
@@ -2070,6 +2070,19 @@ namespace EAutoReceiveInput
 	};
 }
 
+UENUM(BlueprintType)
+namespace EEndPlayReason
+{
+	enum Type
+	{
+		ActorDestroyed,		// When the Actor is explicitly destroyed
+		LevelTransition,	// When the world is being unloaded for a level transition
+		EndPlayInEditor,	// When the world is being unloaded because PIE is ending
+		RemovedFromWorld,	// When the level it is a member of is streamed out
+	};
+
+}
+
 /** Replicated movement data of our RootComponent.
   * Struct used for efficient replication as velocity and location are generally replicated together (this saves a repindex) 
   * and velocity.Z is commonly zero (most position replications are for walking pawns). 
diff --git a/Engine/Source/Runtime/Engine/Classes/Engine/GeneratedMeshAreaLight.h b/Engine/Source/Runtime/Engine/Classes/Engine/GeneratedMeshAreaLight.h
--- a/Engine/Source/Runtime/Engine/Classes/Engine/GeneratedMeshAreaLight.h
+++ b/Engine/Source/Runtime/Engine/Classes/Engine/GeneratedMeshAreaLight.h
@@ -5,7 +5,7 @@
 /**
  * GeneratedMeshAreaLight - A light type that is created after a lighting build with Lightmass and handles mesh area light influence on dynamic objects.
  */
-
+#include "Engine/SpotLight.h"
 #include "GeneratedMeshAreaLight.generated.h"
 
 UCLASS(notplaceable, MinimalAPI)
diff --git a/Engine/Source/Runtime/Engine/Classes/Engine/InteractiveFoliageActor.h b/Engine/Source/Runtime/Engine/Classes/Engine/InteractiveFoliageActor.h
--- a/Engine/Source/Runtime/Engine/Classes/Engine/InteractiveFoliageActor.h
+++ b/Engine/Source/Runtime/Engine/Classes/Engine/InteractiveFoliageActor.h
@@ -4,6 +4,7 @@
 //=============================================================================
 
 #pragma once
+#include "Engine/StaticMeshActor.h"
 #include "InteractiveFoliageActor.generated.h"
 
 UCLASS(MinimalAPI)
diff --git a/Engine/Source/Runtime/Engine/Classes/Engine/SpotLight.h b/Engine/Source/Runtime/Engine/Classes/Engine/SpotLight.h
--- a/Engine/Source/Runtime/Engine/Classes/Engine/SpotLight.h
+++ b/Engine/Source/Runtime/Engine/Classes/Engine/SpotLight.h
@@ -2,7 +2,7 @@
 
 #pragma once
 
-
+#include "Engine/Light.h"
 #include "SpotLight.generated.h"
 
 UCLASS(ClassGroup=(Lights, SpotLights), MinimalAPI, meta=(ChildCanTick))
@@ -16,7 +16,7 @@ class ASpotLight : public ALight
 #if WITH_EDITORONLY_DATA
 	// Reference to editor arrow component visualization 
 	UPROPERTY()
-	TSubobjectPtr<UArrowComponent> ArrowComponent;
+	TSubobjectPtr<class UArrowComponent> ArrowComponent;
 #endif
 
 	// BEGIN DEPRECATED (use component functions now in level script)
diff --git a/Engine/Source/Runtime/Engine/Classes/GameFramework/Actor.h b/Engine/Source/Runtime/Engine/Classes/GameFramework/Actor.h
--- a/Engine/Source/Runtime/Engine/Classes/GameFramework/Actor.h
+++ b/Engine/Source/Runtime/Engine/Classes/GameFramework/Actor.h
@@ -13,19 +13,6 @@ struct FHitResult;
 class AActor;
 class FTimerManager; 
 
-UENUM(BlueprintType)
-namespace EEndPlayReason
-{
-	enum Type
-	{
-		ActorDestroyed,		// When the Actor is explicitly destroyed
-		LevelTransition,	// When the world is being unloaded for a level transition
-		EndPlayInEditor,	// When the world is being unloaded because PIE is ending
-		RemovedFromWorld,	// When the level it is a member of is streamed out
-	};
-
-}
-
 #include "Actor.generated.h"
 
 ENGINE_API DECLARE_LOG_CATEGORY_EXTERN(LogActor, Log, Warning);
@@ -56,15 +43,15 @@ DECLARE_CYCLE_STAT_EXTERN(TEXT("GetComponentsTime"),STAT_GetComponentsTime,STATG
 
 /**
  * Base class for an object that can be placed or spawned in a level. Actors may contain a collection of Components, and support network replication.
-*
-* The functions of interest to initialization order for an Actor is roughly as follows:
-* PostLoad/PostActorCreated - Do any setup of the actor required for construction. PostLoad for serialized actors, PostActorCreated for spawned.  
-* AActor::OnConstruction - The construction of the actor, this is where Blueprint actors have their components created and blueprint variables are initialized
-* AActor::PreInitializeComponents - Called before InitializeComponent is called on the actor's components
-* UActorComponent::InitializeComponent - Each component in the actor's components array gets an initialize call (if bWantsInitializeComponent is true for that component)
-* AActor::PostInitializeComponents - Called after the actor's components have been initialized
-* AActor::BeginPlay - Called when the level is started
-*/
+ *
+ * The functions of interest to initialization order for an Actor is roughly as follners:
+ * PostLoad/PostActorCreated - Do any setup of the actor required for construction. PostLoad for serialized actors, PostActorCreated for spawned.  
+ * AActor::OnConstruction - The construction of the actor, this is where Blueprint actors have their components created and blueprint variables are initialized
+ * AActor::PreInitializeComponents - Called before InitializeComponent is called on the actor's components
+ * UActorComponent::InitializeComponent - Each component in the actor's components array gets an initialize call (if bWantsInitializeComponent is true for that component)
+ * AActor::PostInitializeComponents - Called after the actor's components have been initialized
+ * AActor::BeginPlay - Called when the level is started
+ */
 UCLASS(abstract, BlueprintType, Blueprintable, config=Engine)
 class ENGINE_API AActor : public UObject
 {
@@ -364,13 +351,6 @@ public:
 	UPROPERTY()
 	uint32 bLockLocation:1;
 
-	// Actor's layer name.
-	UPROPERTY()
-	FName Layer_DEPRECATED;
-
-	UPROPERTY()
-	FName Group_DEPRECATED;
-
 	/** The group this actor is a part of. */
 	UPROPERTY(transient)
 	AActor* GroupActor;
diff --git a/Engine/Source/Runtime/Engine/Private/Actor.cpp b/Engine/Source/Runtime/Engine/Private/Actor.cpp
--- a/Engine/Source/Runtime/Engine/Private/Actor.cpp
+++ b/Engine/Source/Runtime/Engine/Private/Actor.cpp
@@ -441,21 +441,6 @@ void AActor::PostLoad()
 		{
 			UE_LOG(LogActor, Log,  TEXT("Loaded Actor (%s) with IsPendingKill() == true"), *GetName() );
 		}
-
-		// Move the old layers information into the new layers system
-		TArray< FString > OldLayerNames; 
-		Layer_DEPRECATED.ToString().ParseIntoArray( &OldLayerNames, TEXT(","), 0 );
-
-		for( int NameIndex = 0; NameIndex < OldLayerNames.Num(); ++NameIndex )
-		{
-			const FString& Name = OldLayerNames[ NameIndex ];
-			if( Name != TEXT( "None" ) )
-			{
-				Layers.AddUnique( FName( *Name ) );
-			}
-		}
-
-		Layer_DEPRECATED = FName( NAME_None ); 
 #endif // WITH_EDITORONLY_DATA
 	}
 }
diff --git a/Engine/Source/Runtime/Engine/Private/ActorConstruction.cpp b/Engine/Source/Runtime/Engine/Private/ActorConstruction.cpp
--- a/Engine/Source/Runtime/Engine/Private/ActorConstruction.cpp
+++ b/Engine/Source/Runtime/Engine/Private/ActorConstruction.cpp
@@ -5,6 +5,7 @@
 #include "LatentActions.h"
 #include "DeferRegisterComponents.h"
 #include "ComponentInstanceDataCache.h"
+#include "Engine/LevelScriptActor.h"
 
 #if WITH_EDITOR
 #include "Editor/UnrealEd/Public/Kismet2/BlueprintEditorUtils.h"
diff --git a/Engine/Source/Runtime/Engine/Private/ActorEditor.cpp b/Engine/Source/Runtime/Engine/Private/ActorEditor.cpp
--- a/Engine/Source/Runtime/Engine/Private/ActorEditor.cpp
+++ b/Engine/Source/Runtime/Engine/Private/ActorEditor.cpp
@@ -10,6 +10,7 @@
 #include "Foliage/InstancedFoliageActor.h"
 #include "Landscape/LandscapeHeightfieldCollisionComponent.h"
 #include "Landscape/LandscapeComponent.h"
+#include "Engine/LevelBounds.h"
 
 #if WITH_EDITOR
 
diff --git a/Engine/Source/Runtime/Engine/Private/Atmosphere/Atmosphere.cpp b/Engine/Source/Runtime/Engine/Private/Atmosphere/Atmosphere.cpp
--- a/Engine/Source/Runtime/Engine/Private/Atmosphere/Atmosphere.cpp
+++ b/Engine/Source/Runtime/Engine/Private/Atmosphere/Atmosphere.cpp
@@ -4,6 +4,7 @@
 #include "../../../Renderer/Private/ScenePrivate.h"
 #include "Atmosphere.h"
 #include "ComponentInstanceDataCache.h"
+#include "Atmosphere/AtmosphericFog.h"
 
 #if WITH_EDITOR
 #include "ObjectEditorUtils.h"
diff --git a/Engine/Source/Runtime/Engine/Private/Camera/CameraActor.cpp b/Engine/Source/Runtime/Engine/Private/Camera/CameraActor.cpp
--- a/Engine/Source/Runtime/Engine/Private/Camera/CameraActor.cpp
+++ b/Engine/Source/Runtime/Engine/Private/Camera/CameraActor.cpp
@@ -1,8 +1,7 @@
 // Copyright 1998-2014 Epic Games, Inc. All Rights Reserved.
 
 #include "EnginePrivate.h"
-#include "ParticleDefinitions.h"
-
+#include "Camera/CameraActor.h"
 #include "Net/UnrealNetwork.h"
 #include "MessageLog.h"
 #include "UObjectToken.h"
diff --git a/Engine/Source/Runtime/Engine/Private/DecalActor.cpp b/Engine/Source/Runtime/Engine/Private/DecalActor.cpp
new file mode 100644
--- /dev/null
+++ b/Engine/Source/Runtime/Engine/Private/DecalActor.cpp
@@ -0,0 +1,135 @@
+// Copyright 1998-2014 Epic Games, Inc. All Rights Reserved.
+
+#include "EnginePrivate.h"
+#include "Engine/DecalActor.h"
+
+
+#if WITH_EDITOR
+namespace DecalEditorConstants
+{
+	/** Scale factor to apply to get nice scaling behaviour in-editor when using percentage-based scaling */
+	static const float PercentageScalingMultiplier = 5.0f;
+
+	/** Scale factor to apply to get nice scaling behaviour in-editor when using additive-based scaling */
+	static const float AdditiveScalingMultiplier = 50.0f;
+}
+#endif
+
+ADecalActor::ADecalActor(const class FPostConstructInitializeProperties& PCIP)
+: Super(PCIP)
+{
+	Decal = PCIP.CreateDefaultSubobject<UDecalComponent>(this, TEXT("NewDecalComponent"));
+	Decal->RelativeScale3D = FVector(128.0f, 256.0f, 256.0f);
+
+	Decal->RelativeRotation = FRotator(-90, 0, 0);
+
+	RootComponent = Decal;
+
+#if WITH_EDITORONLY_DATA
+	BoxComponent = PCIP.CreateEditorOnlyDefaultSubobject<UBoxComponent>(this, TEXT("DrawBox0"));
+	if (BoxComponent != nullptr)
+	{
+		BoxComponent->BodyInstance.bEnableCollision_DEPRECATED = false;
+		BoxComponent->SetCollisionProfileName(UCollisionProfile::NoCollision_ProfileName);
+
+		BoxComponent->ShapeColor = FColor(80, 80, 200, 255);
+		BoxComponent->bDrawOnlyIfSelected = true;
+		BoxComponent->InitBoxExtent(FVector(1.0f, 1.0f, 1.0f));
+
+		BoxComponent->AttachParent = Decal;
+	}
+
+	ArrowComponent = PCIP.CreateEditorOnlyDefaultSubobject<UArrowComponent>(this, TEXT("ArrowComponent0"));
+	SpriteComponent = PCIP.CreateEditorOnlyDefaultSubobject<UBillboardComponent>(this, TEXT("Sprite"));
+
+	if (!IsRunningCommandlet())
+	{
+		// Structure to hold one-time initialization
+		struct FConstructorStatics
+		{
+			ConstructorHelpers::FObjectFinderOptional<UTexture2D> DecalTexture;
+			FName ID_Decals;
+			FText NAME_Decals;
+			FConstructorStatics()
+				: DecalTexture(TEXT("/Engine/EditorResources/S_DecalActorIcon"))
+				, ID_Decals(TEXT("Decals"))
+				, NAME_Decals(NSLOCTEXT("SpriteCategory", "Decals", "Decals"))
+			{
+			}
+		};
+		static FConstructorStatics ConstructorStatics;
+
+		if (ArrowComponent)
+		{
+			ArrowComponent->bTreatAsASprite = true;
+			ArrowComponent->ArrowSize = 1.0f;
+			ArrowComponent->ArrowColor = FColor(80, 80, 200, 255);
+			ArrowComponent->SpriteInfo.Category = ConstructorStatics.ID_Decals;
+			ArrowComponent->SpriteInfo.DisplayName = ConstructorStatics.NAME_Decals;
+			ArrowComponent->AttachParent = Decal;
+			ArrowComponent->bAbsoluteScale = true;
+			ArrowComponent->bIsScreenSizeScaled = true;
+		}
+
+		if (SpriteComponent)
+		{
+			SpriteComponent->Sprite = ConstructorStatics.DecalTexture.Get();
+			SpriteComponent->RelativeScale3D = FVector(0.5f, 0.5f, 0.5f);
+			SpriteComponent->AttachParent = Decal;
+			SpriteComponent->bIsScreenSizeScaled = true;
+			SpriteComponent->bAbsoluteScale = true;
+			SpriteComponent->bReceivesDecals = false;
+		}
+	}
+#endif // WITH_EDITORONLY_DATA
+
+	bCanBeDamaged = false;
+}
+
+#if WITH_EDITOR
+void ADecalActor::PostEditMove(bool bFinished)
+{
+	Super::PostEditMove(bFinished);
+
+	if (Decal.IsValid())
+	{
+		Decal->RecreateRenderState_Concurrent();
+	}
+}
+
+void ADecalActor::PostEditChangeProperty(FPropertyChangedEvent& PropertyChangedEvent)
+{
+	// AActor::PostEditChange will ForceUpdateComponents()
+	Super::PostEditChangeProperty(PropertyChangedEvent);
+
+	if (Decal.IsValid())
+	{
+		Decal->RecreateRenderState_Concurrent();
+	}
+}
+
+void ADecalActor::EditorApplyScale(const FVector& DeltaScale, const FVector* PivotLocation, bool bAltDown, bool bShiftDown, bool bCtrlDown)
+{
+	const FVector ModifiedScale = DeltaScale * (AActor::bUsePercentageBasedScaling ? DecalEditorConstants::PercentageScalingMultiplier : DecalEditorConstants::AdditiveScalingMultiplier);
+
+	Super::EditorApplyScale(ModifiedScale, PivotLocation, bAltDown, bShiftDown, bCtrlDown);
+}
+#endif // WITH_EDITOR
+
+void ADecalActor::SetDecalMaterial(class UMaterialInterface* NewDecalMaterial)
+{
+	if (Decal.IsValid())
+	{
+		Decal->SetDecalMaterial(NewDecalMaterial);
+	}
+}
+
+class UMaterialInterface* ADecalActor::GetDecalMaterial() const
+{
+	return Decal.IsValid() ? Decal->GetDecalMaterial() : NULL;
+}
+
+class UMaterialInstanceDynamic* ADecalActor::CreateDynamicMaterialInstance()
+{
+	return Decal.IsValid() ? Decal->CreateDynamicMaterialInstance() : NULL;
+}
diff --git a/Engine/Source/Runtime/Engine/Private/DecalComponent.cpp b/Engine/Source/Runtime/Engine/Private/DecalComponent.cpp
--- a/Engine/Source/Runtime/Engine/Private/DecalComponent.cpp
+++ b/Engine/Source/Runtime/Engine/Private/DecalComponent.cpp
@@ -7,119 +7,6 @@
 #include "EnginePrivate.h"
 #include "LevelUtils.h"
 
-#if WITH_EDITOR
-namespace DecalEditorConstants
-{
-	/** Scale factor to apply to get nice scaling behaviour in-editor when using percentage-based scaling */
-	static const float PercentageScalingMultiplier = 5.0f;
-
-	/** Scale factor to apply to get nice scaling behaviour in-editor when using additive-based scaling */
-	static const float AdditiveScalingMultiplier = 50.0f;
-}
-#endif
-
-ADecalActor::ADecalActor(const class FPostConstructInitializeProperties& PCIP)
-	: Super(PCIP)
-{
-	Decal = PCIP.CreateDefaultSubobject<UDecalComponent>(this, TEXT("NewDecalComponent"));
-	Decal->RelativeScale3D = FVector(128.0f, 256.0f, 256.0f);
-
-	Decal->RelativeRotation = FRotator(-90, 0, 0);
-
-	RootComponent = Decal;
-
-#if WITH_EDITORONLY_DATA
-	BoxComponent = PCIP.CreateEditorOnlyDefaultSubobject<UBoxComponent>(this, TEXT("DrawBox0"));
-	if (BoxComponent != nullptr)
-	{
-		BoxComponent->BodyInstance.bEnableCollision_DEPRECATED = false;
-		BoxComponent->SetCollisionProfileName(UCollisionProfile::NoCollision_ProfileName);
-
-		BoxComponent->ShapeColor = FColor(80, 80, 200, 255);
-		BoxComponent->bDrawOnlyIfSelected = true;
-		BoxComponent->InitBoxExtent(FVector(1.0f, 1.0f, 1.0f));
-
-		BoxComponent->AttachParent = Decal;
-	}
-
-	ArrowComponent = PCIP.CreateEditorOnlyDefaultSubobject<UArrowComponent>(this, TEXT("ArrowComponent0"));
-	SpriteComponent = PCIP.CreateEditorOnlyDefaultSubobject<UBillboardComponent>(this, TEXT("Sprite"));
-
-	if (!IsRunningCommandlet())
-	{
-		// Structure to hold one-time initialization
-		struct FConstructorStatics
-		{
-			ConstructorHelpers::FObjectFinderOptional<UTexture2D> DecalTexture;
-			FName ID_Decals;
-			FText NAME_Decals;
-			FConstructorStatics()
-				: DecalTexture(TEXT("/Engine/EditorResources/S_DecalActorIcon"))
-				, ID_Decals(TEXT("Decals"))
-				, NAME_Decals(NSLOCTEXT("SpriteCategory", "Decals", "Decals"))
-			{
-			}
-		};
-		static FConstructorStatics ConstructorStatics;
-
-		if (ArrowComponent)
-		{
-			ArrowComponent->bTreatAsASprite = true;
-			ArrowComponent->ArrowSize = 1.0f;
-			ArrowComponent->ArrowColor = FColor(80, 80, 200, 255);
-			ArrowComponent->SpriteInfo.Category = ConstructorStatics.ID_Decals;
-			ArrowComponent->SpriteInfo.DisplayName = ConstructorStatics.NAME_Decals;
-			ArrowComponent->AttachParent = Decal;
-			ArrowComponent->bAbsoluteScale = true;
-			ArrowComponent->bIsScreenSizeScaled = true;
-		}
-
-		if (SpriteComponent)
-		{
-			SpriteComponent->Sprite = ConstructorStatics.DecalTexture.Get();
-			SpriteComponent->RelativeScale3D = FVector(0.5f, 0.5f, 0.5f);
-			SpriteComponent->AttachParent = Decal;
-			SpriteComponent->bIsScreenSizeScaled = true;
-			SpriteComponent->bAbsoluteScale = true;
-			SpriteComponent->bReceivesDecals = false;
-		}
-	}
-#endif // WITH_EDITORONLY_DATA
-
-	bCanBeDamaged = false;
-}
-
-#if WITH_EDITOR
-void ADecalActor::PostEditMove(bool bFinished)
-{
-	Super::PostEditMove(bFinished);
-
-	if (Decal.IsValid())
-	{
-		Decal->RecreateRenderState_Concurrent();
-	}
-}
-
-void ADecalActor::PostEditChangeProperty(FPropertyChangedEvent& PropertyChangedEvent)
-{
-	// AActor::PostEditChange will ForceUpdateComponents()
-	Super::PostEditChangeProperty(PropertyChangedEvent);
-
-	if (Decal.IsValid())
-	{
-		Decal->RecreateRenderState_Concurrent();
-	}
-}
-
-void ADecalActor::EditorApplyScale(const FVector& DeltaScale, const FVector* PivotLocation, bool bAltDown, bool bShiftDown, bool bCtrlDown)
-{
-	const FVector ModifiedScale = DeltaScale * ( AActor::bUsePercentageBasedScaling ? DecalEditorConstants::PercentageScalingMultiplier : DecalEditorConstants::AdditiveScalingMultiplier );
-
-	Super::EditorApplyScale(ModifiedScale, PivotLocation, bAltDown, bShiftDown, bCtrlDown);
-}
-
-#endif // WITH_EDITOR
-
 FDeferredDecalProxy::FDeferredDecalProxy(const UDecalComponent* InComponent)
 {
 	UMaterialInterface* EffectiveMaterial = UMaterial::GetDefaultMaterial(MD_DeferredDecal);
@@ -230,20 +117,3 @@ void UDecalComponent::DestroyRenderState_Concurrent()
 	World->Scene->RemoveDecal(this);
 }
 
-void ADecalActor::SetDecalMaterial(class UMaterialInterface* NewDecalMaterial)
-{
-	if(Decal.IsValid())
-	{
-		Decal->SetDecalMaterial(NewDecalMaterial);
-	}
-}
-
-class UMaterialInterface* ADecalActor::GetDecalMaterial() const
-{
-	return Decal.IsValid() ? Decal->GetDecalMaterial() : NULL;
-}
-
-class UMaterialInstanceDynamic* ADecalActor::CreateDynamicMaterialInstance()
-{
-	return Decal.IsValid() ? Decal->CreateDynamicMaterialInstance() : NULL;
-}
diff --git a/Engine/Source/Runtime/Engine/Private/DestructibleActor.cpp b/Engine/Source/Runtime/Engine/Private/DestructibleActor.cpp
--- a/Engine/Source/Runtime/Engine/Private/DestructibleActor.cpp
+++ b/Engine/Source/Runtime/Engine/Private/DestructibleActor.cpp
@@ -7,7 +7,7 @@
 
 #include "EnginePrivate.h"
 #include "PhysicsEngine/PhysXSupport.h"
-
+#include "PhysicsEngine/DestructibleActor.h"
 
 ADestructibleActor::ADestructibleActor(const class FPostConstructInitializeProperties& PCIP)
 	: Super(PCIP)
diff --git a/Engine/Source/Runtime/Engine/Private/DirectionalLightComponent.cpp b/Engine/Source/Runtime/Engine/Private/DirectionalLightComponent.cpp
--- a/Engine/Source/Runtime/Engine/Private/DirectionalLightComponent.cpp
+++ b/Engine/Source/Runtime/Engine/Private/DirectionalLightComponent.cpp
@@ -5,6 +5,7 @@
 =============================================================================*/
 
 #include "EnginePrivate.h"
+#include "Engine/DirectionalLight.h"
 
 static float GMaxCSMRadiusToAllowPerObjectShadows = 8000;
 static FAutoConsoleVariableRef CVarMaxCSMRadiusToAllowPerObjectShadows(
diff --git a/Engine/Source/Runtime/Engine/Private/EnginePrivate.h b/Engine/Source/Runtime/Engine/Private/EnginePrivate.h
--- a/Engine/Source/Runtime/Engine/Private/EnginePrivate.h
+++ b/Engine/Source/Runtime/Engine/Private/EnginePrivate.h
@@ -22,7 +22,6 @@
 #include "Components/ActorComponent.h"
 #include "GameFramework/Actor.h"
 #include "Engine/Brush.h"
-#include "Camera/CameraActor.h"
 #include "AI/Navigation/NavAgentInterface.h"
 #include "GameFramework/Controller.h"
 #include "AI/Navigation/NavLinkDefinition.h"
@@ -30,38 +29,23 @@
 #include "Kismet/BlueprintFunctionLibrary.h"
 #include "AI/Navigation/NavigationSystem.h"
 #include "AI/Navigation/NavCollision.h"
-#include "GameFramework/PlayerMuteList.h"
 #include "Camera/PlayerCameraManager.h"
 #include "GameFramework/Info.h"
 #include "GameFramework/GameMode.h"
 #include "Components/InputComponent.h"
-#include "Curves/CurveBase.h"
-#include "Curves/CurveFloat.h"
 #include "GameFramework/ForceFeedbackEffect.h"
 #include "Engine/SubSurfaceProfile.h"
 #include "GameFramework/PlayerController.h"
 #include "Engine/DebugCameraController.h"
-#include "Engine/DecalActor.h"
-#include "PhysicsEngine/DestructibleActor.h"
 #include "Engine/Canvas.h"
-#include "Atmosphere/AtmosphericFog.h"
-#include "Engine/ExponentialHeightFog.h"
 #include "GameFramework/GameSession.h"
 #include "GameFramework/GameState.h"
 #include "GameFramework/PlayerState.h"
-#include "Engine/SkyLight.h"
-#include "Engine/WindDirectionalSource.h"
 #include "GameFramework/WorldSettings.h"
 #include "PhysicsEngine/BodyInstance.h"
 #include "Components/SceneComponent.h"
 #include "Components/LightComponentBase.h"
 #include "Components/LightComponent.h"
-#include "Engine/LevelBounds.h"
-#include "Engine/LevelScriptActor.h"
-#include "Engine/Light.h"
-#include "Engine/DirectionalLight.h"
-#include "Engine/PointLight.h"
-#include "Engine/SpotLight.h"
 #include "Engine/GeneratedMeshAreaLight.h"
 #include "AI/Navigation/NavigationData.h"
 #include "AI/Navigation/NavigationGraph.h"
@@ -95,7 +79,6 @@
 #include "PhysicsEngine/ConstraintInstance.h"
 #include "PhysicsEngine/PhysicsConstraintComponent.h"
 #include "Components/SkeletalMeshComponent.h"
-#include "Engine/StaticMeshActor.h"
 #include "Engine/TargetPoint.h"
 #include "Engine/TextRenderActor.h"
 #include "Engine/TriggerBase.h"
diff --git a/Engine/Source/Runtime/Engine/Private/HeightFogComponent.cpp b/Engine/Source/Runtime/Engine/Private/HeightFogComponent.cpp
--- a/Engine/Source/Runtime/Engine/Private/HeightFogComponent.cpp
+++ b/Engine/Source/Runtime/Engine/Private/HeightFogComponent.cpp
@@ -6,7 +6,7 @@
 
 #include "EnginePrivate.h"
 #include "Net/UnrealNetwork.h"
-
+#include "Engine/ExponentialHeightFog.h"
 
 UExponentialHeightFogComponent::UExponentialHeightFogComponent(const class FPostConstructInitializeProperties& PCIP)
 	: Super(PCIP)
diff --git a/Engine/Source/Runtime/Engine/Private/KismetSystemLibrary.cpp b/Engine/Source/Runtime/Engine/Private/KismetSystemLibrary.cpp
--- a/Engine/Source/Runtime/Engine/Private/KismetSystemLibrary.cpp
+++ b/Engine/Source/Runtime/Engine/Private/KismetSystemLibrary.cpp
@@ -7,6 +7,7 @@
 #include "InterpolateComponentToAction.h"
 #include "Advertising.h"
 #include "Online.h"
+#include "Camera/CameraActor.h"
 
 //////////////////////////////////////////////////////////////////////////
 // UKismetSystemLibrary
diff --git a/Engine/Source/Runtime/Engine/Private/Level.cpp b/Engine/Source/Runtime/Engine/Private/Level.cpp
--- a/Engine/Source/Runtime/Engine/Private/Level.cpp
+++ b/Engine/Source/Runtime/Engine/Private/Level.cpp
@@ -6,6 +6,7 @@ Level.cpp: Level-related functions
 
 #include "EnginePrivate.h"
 #include "Engine/LevelScriptBlueprint.h"
+#include "Engine/LevelScriptActor.h"
 #include "Engine/WorldComposition.h"
 #include "Sound/SoundNodeWave.h"
 #include "Net/UnrealNetwork.h"
@@ -16,6 +17,7 @@ Level.cpp: Level-related functions
 #include "TickTaskManagerInterface.h"
 #include "BlueprintUtilities.h"
 #include "DynamicMeshBuilder.h"
+#include "Engine/LevelBounds.h"
 #if WITH_EDITOR
 #include "Editor/UnrealEd/Public/Kismet2/KismetEditorUtilities.h"
 #include "Editor/UnrealEd/Public/Kismet2/BlueprintEditorUtils.h"
diff --git a/Engine/Source/Runtime/Engine/Private/LevelScriptActor.cpp b/Engine/Source/Runtime/Engine/Private/LevelScriptActor.cpp
--- a/Engine/Source/Runtime/Engine/Private/LevelScriptActor.cpp
+++ b/Engine/Source/Runtime/Engine/Private/LevelScriptActor.cpp
@@ -3,6 +3,7 @@
 
 #include "EnginePrivate.h"
 #include "Engine/LevelScriptBlueprint.h"
+#include "Engine/LevelScriptActor.h"
 #include "BlueprintUtilities.h"
 
 
diff --git a/Engine/Source/Runtime/Engine/Private/LevelStreaming.cpp b/Engine/Source/Runtime/Engine/Private/LevelStreaming.cpp
--- a/Engine/Source/Runtime/Engine/Private/LevelStreaming.cpp
+++ b/Engine/Source/Runtime/Engine/Private/LevelStreaming.cpp
@@ -2,6 +2,7 @@
 
 #include "EnginePrivate.h"
 #include "Engine/LevelStreamingVolume.h"
+#include "Engine/LevelBounds.h"
 #include "LevelUtils.h"
 #if WITH_EDITOR
 	#include "Slate.h"
diff --git a/Engine/Source/Runtime/Engine/Private/Light.cpp b/Engine/Source/Runtime/Engine/Private/Light.cpp
--- a/Engine/Source/Runtime/Engine/Private/Light.cpp
+++ b/Engine/Source/Runtime/Engine/Private/Light.cpp
@@ -2,6 +2,9 @@
 
 
 #include "EnginePrivate.h"
+#include "Engine/PointLight.h"
+#include "Engine/DirectionalLight.h"
+#include "Components/ArrowComponent.h"
 #include "LightingBuildOptions.h"
 #include "Net/UnrealNetwork.h"
 
diff --git a/Engine/Source/Runtime/Engine/Private/PlayerCameraManager.cpp b/Engine/Source/Runtime/Engine/Private/PlayerCameraManager.cpp
--- a/Engine/Source/Runtime/Engine/Private/PlayerCameraManager.cpp
+++ b/Engine/Source/Runtime/Engine/Private/PlayerCameraManager.cpp
@@ -6,6 +6,8 @@
 #include "Particles/EmitterCameraLensEffectBase.h"
 #include "IHeadMountedDisplay.h"
 #include "Particles/EmitterCameraLensEffectBase.h"
+#include "Camera/CameraActor.h"
+
 DEFINE_LOG_CATEGORY_STATIC(LogPlayerCameraManager, Log, All);
 
 
diff --git a/Engine/Source/Runtime/Engine/Private/PlayerController.cpp b/Engine/Source/Runtime/Engine/Private/PlayerController.cpp
--- a/Engine/Source/Runtime/Engine/Private/PlayerController.cpp
+++ b/Engine/Source/Runtime/Engine/Private/PlayerController.cpp
@@ -4,6 +4,7 @@
 #include "Engine/Console.h"
 #include "Net/UnrealNetwork.h"
 #include "Engine/WorldComposition.h"
+#include "Engine/LevelScriptActor.h"
 #include "GameFramework/GameNetworkManager.h"
 #include "GameFramework/Character.h"
 #include "Interfaces/NetworkPredictionInterface.h"
diff --git a/Engine/Source/Runtime/Engine/Private/SkyLightComponent.cpp b/Engine/Source/Runtime/Engine/Private/SkyLightComponent.cpp
--- a/Engine/Source/Runtime/Engine/Private/SkyLightComponent.cpp
+++ b/Engine/Source/Runtime/Engine/Private/SkyLightComponent.cpp
@@ -8,6 +8,7 @@
 #if WITH_EDITOR
 #include "ObjectEditorUtils.h"
 #endif
+#include "Engine/SkyLight.h"
 #include "MessageLog.h"
 #include "UObjectToken.h"
 #include "Net/UnrealNetwork.h"
diff --git a/Engine/Source/Runtime/Engine/Private/StaticMeshActor.cpp b/Engine/Source/Runtime/Engine/Private/StaticMeshActor.cpp
--- a/Engine/Source/Runtime/Engine/Private/StaticMeshActor.cpp
+++ b/Engine/Source/Runtime/Engine/Private/StaticMeshActor.cpp
@@ -9,6 +9,7 @@
 #include "StaticMeshResources.h"
 #include "UObjectToken.h"
 #include "MapErrors.h"
+#include "Engine/StaticMeshActor.h"
 
 #define LOCTEXT_NAMESPACE "StaticMeshActor"
 
diff --git a/Engine/Source/Runtime/Engine/Private/UnrealEngine.cpp b/Engine/Source/Runtime/Engine/Private/UnrealEngine.cpp
--- a/Engine/Source/Runtime/Engine/Private/UnrealEngine.cpp
+++ b/Engine/Source/Runtime/Engine/Private/UnrealEngine.cpp
@@ -47,6 +47,7 @@
 #include "GameFramework/HUD.h"
 #include "GameFramework/Character.h"
 #include "Engine/LevelStreamingVolume.h"
+#include "Engine/LevelScriptActor.h"
 #include "Vehicles/TireType.h"
 
 #include "Particles/Spawn/ParticleModuleSpawn.h"
diff --git a/Engine/Source/Runtime/Engine/Private/WindDirectionalSource.cpp b/Engine/Source/Runtime/Engine/Private/WindDirectionalSource.cpp
--- a/Engine/Source/Runtime/Engine/Private/WindDirectionalSource.cpp
+++ b/Engine/Source/Runtime/Engine/Private/WindDirectionalSource.cpp
@@ -1,6 +1,7 @@
 // Copyright 1998-2014 Epic Games, Inc. All Rights Reserved.
 
 #include "EnginePrivate.h"
+#include "Engine/WindDirectionalSource.h"
 
 AWindDirectionalSource::AWindDirectionalSource(const class FPostConstructInitializeProperties& PCIP)
 	: Super(PCIP)
diff --git a/Engine/Source/Runtime/Engine/Private/World.cpp b/Engine/Source/Runtime/Engine/Private/World.cpp
--- a/Engine/Source/Runtime/Engine/Private/World.cpp
+++ b/Engine/Source/Runtime/Engine/Private/World.cpp
@@ -6,6 +6,7 @@
 
 #include "EnginePrivate.h"
 #include "Engine/LevelScriptBlueprint.h"
+#include "Engine/LevelScriptActor.h"
 #include "Engine/CullDistanceVolume.h"
 #include "GameFramework/DefaultPhysicsVolume.h"
 #include "Engine/Console.h"
@@ -32,6 +33,7 @@
 #include "PhysicsPublic.h"
 #include "AI/AISystemBase.h"
 #include "SceneInterface.h"
+#include "Camera/CameraActor.h"
 
 #if WITH_EDITOR
 	#include "Editor/UnrealEd/Public/Kismet2/KismetEditorUtilities.h"
diff --git a/Engine/Source/Runtime/Engine/Public/Engine.h b/Engine/Source/Runtime/Engine/Public/Engine.h
--- a/Engine/Source/Runtime/Engine/Public/Engine.h
+++ b/Engine/Source/Runtime/Engine/Public/Engine.h
@@ -19,11 +19,11 @@
 	Engine public includes.
 -----------------------------------------------------------------------------*/
 
+#include "EngineMinimal.h" // 
+
 #include "Engine/EngineBaseTypes.h"
 #include "Camera/CameraTypes.h"
 #include "Engine/EngineTypes.h"
-#include "Components/ActorComponent.h"
-#include "GameFramework/Actor.h"
 #include "Sound/AmbientSound.h"
 #include "Engine/Brush.h"
 #include "GameFramework/Volume.h"
@@ -44,7 +44,6 @@
 #include "Engine/TriggerVolume.h"
 #include "Camera/CameraActor.h"
 #include "AI/Navigation/NavAgentInterface.h"
-#include "GameFramework/Controller.h"
 #include "AI/Navigation/NavLinkDefinition.h"
 #include "AI/Navigation/NavigationTypes.h"
 #include "Kismet/BlueprintFunctionLibrary.h"
@@ -53,13 +52,10 @@
 #include "GameFramework/PlayerMuteList.h"
 #include "Camera/PlayerCameraManager.h"
 #include "GameFramework/Info.h"
-#include "GameFramework/GameMode.h"
-#include "Components/InputComponent.h"
 #include "Curves/CurveBase.h"
 #include "Curves/CurveFloat.h"
 #include "GameFramework/ForceFeedbackEffect.h"
 #include "Engine/SubsurfaceProfile.h"
-#include "GameFramework/PlayerController.h"
 #include "Engine/DebugCameraController.h"
 #include "Engine/DecalActor.h"
 #include "PhysicsEngine/DestructibleActor.h"
@@ -74,9 +70,7 @@
 #include "Engine/WindDirectionalSource.h"
 #include "GameFramework/WorldSettings.h"
 #include "PhysicsEngine/BodyInstance.h"
-#include "Components/SceneComponent.h"
 #include "Components/LightComponentBase.h"
-#include "Components/LightComponent.h"
 #include "Engine/LevelBounds.h"
 #include "Engine/LevelScriptActor.h"
 #include "Engine/Light.h"
@@ -97,14 +91,10 @@
 #include "AI/Navigation/NavLinkHostInterface.h"
 #include "AI/Navigation/NavLinkProxy.h"
 #include "Engine/Note.h"
-#include "Components/PrimitiveComponent.h"
 #include "GameFramework/MovementComponent.h"
 #include "GameFramework/NavMovementComponent.h"
 #include "GameFramework/PawnMovementComponent.h"
-#include "GameFramework/Pawn.h"
 #include "Interfaces/NetworkPredictionInterface.h"
-#include "GameFramework/CharacterMovementComponent.h"
-#include "GameFramework/Character.h"
 #include "GameFramework/FloatingPawnMovement.h"
 #include "GameFramework/DefaultPawn.h"
 #include "GameFramework/SpectatorPawn.h"
@@ -120,11 +110,9 @@
 #include "Engine/SceneCapture.h"
 #include "Engine/SceneCapture2D.h"
 #include "Engine/SceneCaptureCube.h"
-#include "Components/MeshComponent.h"
 #include "Components/SkinnedMeshComponent.h"
 #include "PhysicsEngine/ConstraintInstance.h"
 #include "PhysicsEngine/PhysicsConstraintComponent.h"
-#include "Components/SkeletalMeshComponent.h"
 #include "Engine/StaticMeshActor.h"
 #include "Engine/TargetPoint.h"
 #include "Engine/TextRenderActor.h"
@@ -145,14 +133,10 @@
 #include "PhysicsEngine/PhysicsHandleComponent.h"
 #include "Atmosphere/AtmosphericFogComponent.h"
 #include "Sound/SoundAttenuation.h"
-#include "Components/AudioComponent.h"
-#include "Camera/CameraComponent.h"
 #include "Components/ChildActorComponent.h"
 #include "Components/DecalComponent.h"
 #include "Components/ExponentialHeightFogComponent.h"
 #include "Components/DirectionalLightComponent.h"
-#include "Components/PointLightComponent.h"
-#include "Components/SpotLightComponent.h"
 #include "Components/SkyLightComponent.h"
 #include "AI/Navigation/NavigationGraphNodeComponent.h"
 #include "PhysicsEngine/PhysicsThrusterComponent.h"
@@ -175,9 +159,6 @@
 #include "AI/Navigation/NavMeshRenderingComponent.h"
 #include "AI/Navigation/NavTestRenderingComponent.h"
 #include "Components/ShapeComponent.h"
-#include "Components/BoxComponent.h"
-#include "Components/CapsuleComponent.h"
-#include "Components/SphereComponent.h"
 #include "Components/DrawSphereComponent.h"
 #include "Components/TextRenderComponent.h"
 #include "Components/VectorFieldComponent.h"
@@ -189,7 +170,6 @@
 #include "Components/SceneCaptureComponent.h"
 #include "Components/SceneCaptureComponent2D.h"
 #include "Components/SceneCaptureComponentCube.h"
-#include "GameFramework/SpringArmComponent.h"
 #include "Components/WindDirectionalSourceComponent.h"
 #include "Components/TimelineComponent.h"
 #include "Animation/AnimSequenceBase.h"
@@ -219,7 +199,6 @@
 #include "Engine/Blueprint.h"
 #include "Animation/AnimBlueprint.h"
 #include "Sound/DialogueTypes.h"
-#include "Kismet/GameplayStatics.h"
 #include "Kismet/HeadMountedDisplayFunctionLibrary.h"
 #include "Kismet/KismetArrayLibrary.h"
 #include "Kismet/KismetInputLibrary.h"
@@ -238,7 +217,6 @@
 #include "Camera/CameraShake.h"
 #include "Camera/CameraModifier_CameraShake.h"
 #include "GameFramework/CheatManager.h"
-#include "Engine/CollisionProfile.h"
 #include "Curves/CurveLinearColor.h"
 #include "Curves/CurveVector.h"
 #include "Curves/CurveEdPresetCurve.h"
@@ -295,10 +273,7 @@
 #include "GameFramework/LocalMessage.h"
 #include "GameFramework/EngineMessage.h"
 #include "Materials/MaterialInterface.h"
-#include "Materials/Material.h"
-#include "Materials/MaterialInstance.h"
 #include "Materials/MaterialInstanceConstant.h"
-#include "Materials/MaterialInstanceDynamic.h"
 #include "Materials/MaterialParameterCollection.h"
 #include "Materials/MaterialParameterCollectionInstance.h"
 #include "Engine/NetDriver.h"
@@ -333,26 +308,21 @@
 #include "Engine/Selection.h"
 #include "Engine/SimpleConstructionScript.h"
 #include "Animation/PreviewAssetAttachComponent.h"
-#include "Engine/SkeletalMesh.h"
 #include "Engine/DestructibleMesh.h"
 #include "Engine/SkeletalMeshSocket.h"
 #include "Animation/Skeleton.h"
-#include "Sound/SoundBase.h"
 #include "Sound/DialogueSoundWaveProxy.h"
 #include "Sound/SoundClass.h"
 #include "Sound/SoundCue.h"
 #include "Sound/SoundGroups.h"
-#include "Sound/SoundWave.h"
 #include "Sound/SoundWaveStreaming.h"
 #include "Sound/SoundMix.h"
-#include "Engine/StaticMesh.h"
 #include "Engine/StaticMeshSocket.h"
 #include "Camera/CameraStackTypes.h"
 #include "Engine/StreamableManager.h"
 #include "Engine/TextureDefines.h"
 #include "Tests/TextPropertyTestObject.h"
 #include "Engine/Texture.h"
-#include "Engine/Texture2D.h"
 #include "Engine/LightMapTexture2D.h"
 #include "Engine/ShadowMapTexture2D.h"
 #include "Engine/TextureLightProfile.h"
@@ -372,7 +342,6 @@
 #include "VectorField/VectorFieldAnimated.h"
 #include "VectorField/VectorFieldStatic.h"
 #include "Animation/VertexAnim/MorphTarget.h"
-#include "Engine/World.h"
 
 #include "SystemSettings.h"					// Scalability options.
 #include "SceneManagement.h"				// Scene management.
