Add missing forward declare to BodyInstance.h

diff --git a/Engine/Source/Runtime/Engine/Classes/PhysicsEngine/BodyInstance.h b/Engine/Source/Runtime/Engine/Classes/PhysicsEngine/BodyInstance.h
--- a/Engine/Source/Runtime/Engine/Classes/PhysicsEngine/BodyInstance.h
+++ b/Engine/Source/Runtime/Engine/Classes/PhysicsEngine/BodyInstance.h
@@ -10,6 +10,7 @@
 #define UE_WITH_PHYSICS (WITH_PHYSX || WITH_BOX2D)
 
 struct FCollisionShape;
+struct FConstraintInstance;
 class UPhysicsConstraintComponent;
 
 #if WITH_PHYSX
