Fixed regression in PIE where streaming levels was not fully streamed before game start

diff --git a/Engine/Source/Runtime/Engine/Private/GameInstance.cpp b/Engine/Source/Runtime/Engine/Private/GameInstance.cpp
--- a/Engine/Source/Runtime/Engine/Private/GameInstance.cpp
+++ b/Engine/Source/Runtime/Engine/Private/GameInstance.cpp
@@ -186,6 +186,10 @@ bool UGameInstance::StartPIEGameInstance(ULocalPlayer* LocalPlayer, bool bInSimu
 		{
 			PlayWorld->NavigateTo(PlayWorld->GlobalOriginOffset);
 		}
+		else
+		{
+			PlayWorld->FlushLevelStreaming();
+		}
 
 		UNavigationSystem::InitializeForWorld(PlayWorld, LocalPlayers.Num() > 0 ? FNavigationSystem::PIEMode : FNavigationSystem::SimulationMode);
 		PlayWorld->CreateAISystem();
