fixed box collision export for navmesh

diff --git a/Engine/Source/Runtime/Engine/Private/AI/Navigation/RecastNavMeshGenerator.cpp b/Engine/Source/Runtime/Engine/Private/AI/Navigation/RecastNavMeshGenerator.cpp
--- a/Engine/Source/Runtime/Engine/Private/AI/Navigation/RecastNavMeshGenerator.cpp
+++ b/Engine/Source/Runtime/Engine/Private/AI/Navigation/RecastNavMeshGenerator.cpp
@@ -753,14 +753,14 @@ void ExportRigidBodyBoxElements(UBodySetup& BodySetup, TNavStatArray<float>& Ver
 
 		// add box vertices
 		FVector UnrealVerts[] = {
-			ElemTM.TransformPosition(FVector(-Extent.X,  Extent.Y, -Extent.Z)),
-			ElemTM.TransformPosition(FVector( Extent.X,  Extent.Y, -Extent.Z)),
-			ElemTM.TransformPosition(FVector( Extent.X, -Extent.Y, -Extent.Z)),
+			ElemTM.TransformPosition(FVector(-Extent.X, -Extent.Y,  Extent.Z)),
+			ElemTM.TransformPosition(FVector( Extent.X, -Extent.Y,  Extent.Z)),
 			ElemTM.TransformPosition(FVector(-Extent.X, -Extent.Y, -Extent.Z)),
+			ElemTM.TransformPosition(FVector( Extent.X, -Extent.Y, -Extent.Z)),
 			ElemTM.TransformPosition(FVector(-Extent.X,  Extent.Y,  Extent.Z)),
 			ElemTM.TransformPosition(FVector( Extent.X,  Extent.Y,  Extent.Z)),
-			ElemTM.TransformPosition(FVector( Extent.X, -Extent.Y,  Extent.Z)),
-			ElemTM.TransformPosition(FVector(-Extent.X, -Extent.Y,  Extent.Z))
+			ElemTM.TransformPosition(FVector(-Extent.X,  Extent.Y, -Extent.Z)),
+			ElemTM.TransformPosition(FVector( Extent.X,  Extent.Y, -Extent.Z))
 		};
 
 		for (int32 iv = 0; iv < ARRAY_COUNT(UnrealVerts); iv++)
@@ -772,29 +772,18 @@ void ExportRigidBodyBoxElements(UBodySetup& BodySetup, TNavStatArray<float>& Ver
 			VertexBuffer.Add(UnrealVerts[iv].Z);
 		}
 		
-		// bottom
-		IndexBuffer.Add(VertBase + 0); IndexBuffer.Add(VertBase + 1); IndexBuffer.Add(VertBase + 2);
-		IndexBuffer.Add(VertBase + 0); IndexBuffer.Add(VertBase + 2); IndexBuffer.Add(VertBase + 3);
-
-		// top
-		IndexBuffer.Add(VertBase + 4); IndexBuffer.Add(VertBase + 6); IndexBuffer.Add(VertBase + 5);
-		IndexBuffer.Add(VertBase + 4); IndexBuffer.Add(VertBase + 7); IndexBuffer.Add(VertBase + 6);
-
-		// front
-		IndexBuffer.Add(VertBase + 7); IndexBuffer.Add(VertBase + 2); IndexBuffer.Add(VertBase + 6);
-		IndexBuffer.Add(VertBase + 7); IndexBuffer.Add(VertBase + 3); IndexBuffer.Add(VertBase + 2);
-
-		// back
-		IndexBuffer.Add(VertBase + 4); IndexBuffer.Add(VertBase + 5); IndexBuffer.Add(VertBase + 1);
-		IndexBuffer.Add(VertBase + 4); IndexBuffer.Add(VertBase + 1); IndexBuffer.Add(VertBase + 0);
-
-		// left
-		IndexBuffer.Add(VertBase + 7); IndexBuffer.Add(VertBase + 4); IndexBuffer.Add(VertBase + 0);
-		IndexBuffer.Add(VertBase + 7); IndexBuffer.Add(VertBase + 0); IndexBuffer.Add(VertBase + 3);
-
-		// right
-		IndexBuffer.Add(VertBase + 6); IndexBuffer.Add(VertBase + 1); IndexBuffer.Add(VertBase + 5);
-		IndexBuffer.Add(VertBase + 6); IndexBuffer.Add(VertBase + 2); IndexBuffer.Add(VertBase + 1);
+		IndexBuffer.Add(VertBase + 3); IndexBuffer.Add(VertBase + 2); IndexBuffer.Add(VertBase + 0);
+		IndexBuffer.Add(VertBase + 3); IndexBuffer.Add(VertBase + 0); IndexBuffer.Add(VertBase + 1);
+		IndexBuffer.Add(VertBase + 7); IndexBuffer.Add(VertBase + 3); IndexBuffer.Add(VertBase + 1);
+		IndexBuffer.Add(VertBase + 7); IndexBuffer.Add(VertBase + 1); IndexBuffer.Add(VertBase + 5);
+		IndexBuffer.Add(VertBase + 6); IndexBuffer.Add(VertBase + 7); IndexBuffer.Add(VertBase + 5);
+		IndexBuffer.Add(VertBase + 6); IndexBuffer.Add(VertBase + 5); IndexBuffer.Add(VertBase + 4);
+		IndexBuffer.Add(VertBase + 2); IndexBuffer.Add(VertBase + 6); IndexBuffer.Add(VertBase + 4);
+		IndexBuffer.Add(VertBase + 2); IndexBuffer.Add(VertBase + 4); IndexBuffer.Add(VertBase + 0);
+		IndexBuffer.Add(VertBase + 1); IndexBuffer.Add(VertBase + 0); IndexBuffer.Add(VertBase + 4);
+		IndexBuffer.Add(VertBase + 1); IndexBuffer.Add(VertBase + 4); IndexBuffer.Add(VertBase + 5);
+		IndexBuffer.Add(VertBase + 7); IndexBuffer.Add(VertBase + 6); IndexBuffer.Add(VertBase + 2);
+		IndexBuffer.Add(VertBase + 7); IndexBuffer.Add(VertBase + 2); IndexBuffer.Add(VertBase + 3);
 	}
 }
 
