Non-unity fixes

diff --git a/Engine/Source/Runtime/Engine/Classes/GameFramework/PlayerController.h b/Engine/Source/Runtime/Engine/Classes/GameFramework/PlayerController.h
--- a/Engine/Source/Runtime/Engine/Classes/GameFramework/PlayerController.h
+++ b/Engine/Source/Runtime/Engine/Classes/GameFramework/PlayerController.h
@@ -5,6 +5,7 @@
 #include "GameFramework/PlayerMuteList.h"
 #include "Camera/PlayerCameraManager.h"
 #include "Camera/CameraTypes.h"
+#include "Camera/CameraActor.h"
 #include "Components/InputComponent.h"
 #include "GameFramework/ForceFeedbackEffect.h"
 #include "GameFramework/OnlineReplStructs.h"
diff --git a/Engine/Source/Runtime/Engine/Private/Camera/CameraAnim.cpp b/Engine/Source/Runtime/Engine/Private/Camera/CameraAnim.cpp
--- a/Engine/Source/Runtime/Engine/Private/Camera/CameraAnim.cpp
+++ b/Engine/Source/Runtime/Engine/Private/Camera/CameraAnim.cpp
@@ -2,6 +2,7 @@
 
 #include "EnginePrivate.h"
 #include "ParticleDefinitions.h"
+#include "Camera/CameraActor.h"
 #include "Matinee/MatineeActor.h"
 #include "Matinee/InterpData.h"
 #include "Matinee/InterpGroup.h"
diff --git a/Engine/Source/Runtime/Engine/Private/Camera/CameraAnimInst.cpp b/Engine/Source/Runtime/Engine/Private/Camera/CameraAnimInst.cpp
--- a/Engine/Source/Runtime/Engine/Private/Camera/CameraAnimInst.cpp
+++ b/Engine/Source/Runtime/Engine/Private/Camera/CameraAnimInst.cpp
@@ -2,6 +2,7 @@
 
 
 #include "EnginePrivate.h"
+#include "Camera/CameraActor.h"
 #include "Matinee/InterpTrackMove.h"
 #include "Matinee/InterpGroup.h"
 #include "Matinee/InterpGroupInst.h"
diff --git a/Engine/Source/Runtime/Engine/Private/CharacterMovementComponent.cpp b/Engine/Source/Runtime/Engine/Private/CharacterMovementComponent.cpp
--- a/Engine/Source/Runtime/Engine/Private/CharacterMovementComponent.cpp
+++ b/Engine/Source/Runtime/Engine/Private/CharacterMovementComponent.cpp
@@ -12,6 +12,7 @@
 #include "GameFramework/CharacterMovementComponent.h"
 #include "Components/PrimitiveComponent.h"
 #include "Animation/AnimMontage.h"
+#include "PhysicsEngine/DestructibleActor.h"
 
 // @todo this is here only due to circular dependency to AIModule. To be removed
 #include "Navigation/PathFollowingComponent.h"
diff --git a/Engine/Source/Runtime/Engine/Private/DestructibleComponent.cpp b/Engine/Source/Runtime/Engine/Private/DestructibleComponent.cpp
--- a/Engine/Source/Runtime/Engine/Private/DestructibleComponent.cpp
+++ b/Engine/Source/Runtime/Engine/Private/DestructibleComponent.cpp
@@ -7,6 +7,7 @@
 
 #include "EnginePrivate.h"
 #include "PhysicsPublic.h"
+#include "PhysicsEngine/DestructibleActor.h"
 #include "PhysicsEngine/PhysXSupport.h"
 #include "Collision/PhysXCollision.h"
 #include "ParticleDefinitions.h"
diff --git a/Engine/Source/Runtime/Engine/Private/GameMode.cpp b/Engine/Source/Runtime/Engine/Private/GameMode.cpp
--- a/Engine/Source/Runtime/Engine/Private/GameMode.cpp
+++ b/Engine/Source/Runtime/Engine/Private/GameMode.cpp
@@ -5,6 +5,7 @@
 =============================================================================*/
 
 #include "EnginePrivate.h"
+#include "Engine/LevelScriptActor.h"
 #include "GameFramework/GameNetworkManager.h"
 #include "Matinee/MatineeActor.h"
 #include "OnlineSubsystemUtils.h"
diff --git a/Engine/Source/Runtime/Engine/Private/Interpolation.cpp b/Engine/Source/Runtime/Engine/Private/Interpolation.cpp
--- a/Engine/Source/Runtime/Engine/Private/Interpolation.cpp
+++ b/Engine/Source/Runtime/Engine/Private/Interpolation.cpp
@@ -5,6 +5,7 @@
 =============================================================================*/
 
 #include "EnginePrivate.h"
+#include "Engine/LevelScriptActor.h"
 #include "Matinee/MatineeActorCameraAnim.h"
 #include "Matinee/InterpData.h"
 #include "Matinee/InterpTrackInstProperty.h"
diff --git a/Engine/Source/Runtime/Engine/Private/LevelBounds.cpp b/Engine/Source/Runtime/Engine/Private/LevelBounds.cpp
--- a/Engine/Source/Runtime/Engine/Private/LevelBounds.cpp
+++ b/Engine/Source/Runtime/Engine/Private/LevelBounds.cpp
@@ -1,6 +1,7 @@
 // Copyright 1998-2014 Epic Games, Inc. All Rights Reserved.
 
 #include "EnginePrivate.h"
+#include "Engine/LevelBounds.h"
 
 // Default size of the box (scale)
 static const FVector DefaultLevelSize = FVector(1000.f);
diff --git a/Engine/Source/Runtime/Engine/Private/LevelScriptBlueprint.cpp b/Engine/Source/Runtime/Engine/Private/LevelScriptBlueprint.cpp
--- a/Engine/Source/Runtime/Engine/Private/LevelScriptBlueprint.cpp
+++ b/Engine/Source/Runtime/Engine/Private/LevelScriptBlueprint.cpp
@@ -2,6 +2,7 @@
 
 #include "EnginePrivate.h"
 #include "Engine/LevelScriptBlueprint.h"
+#include "Engine/LevelScriptActor.h"
 #include "BlueprintUtilities.h"
 #include "LatentActions.h"
 
diff --git a/Engine/Source/Runtime/Engine/Private/UserInterface/Console.cpp b/Engine/Source/Runtime/Engine/Private/UserInterface/Console.cpp
--- a/Engine/Source/Runtime/Engine/Private/UserInterface/Console.cpp
+++ b/Engine/Source/Runtime/Engine/Private/UserInterface/Console.cpp
@@ -6,6 +6,7 @@
 
 #include "EnginePrivate.h"
 #include "Engine/Console.h"
+#include "Engine/LevelScriptActor.h"
 #include "Slate.h"
 #include "DefaultValueHelper.h"
 
