Small fix for non-unity build

diff --git a/Engine/Source/Runtime/Engine/Private/Rendering/SeparableSSS.cpp b/Engine/Source/Runtime/Engine/Private/Rendering/SeparableSSS.cpp
--- a/Engine/Source/Runtime/Engine/Private/Rendering/SeparableSSS.cpp
+++ b/Engine/Source/Runtime/Engine/Private/Rendering/SeparableSSS.cpp
@@ -83,7 +83,7 @@ inline FVector SeparableSSS_Profile(float r, FLinearColor FalloffColor)
 		0.078f * SeparableSSS_Gaussian(7.41f, r, FalloffColor);
 }
 
-inline void ComputeMirroredSSSKernel(TArray<FLinearColor>& Out, FLinearColor SubsurfaceColor, FLinearColor FalloffColor)
+void ComputeMirroredSSSKernel(TArray<FLinearColor>& Out, FLinearColor SubsurfaceColor, FLinearColor FalloffColor)
 {
 	// needs to be preallocated with the expected size e.g. 5 for a 9 sample kernel with one center and 4 left, 4 right and one center sample
 	check(Out.Num());
