Use GetUnitAxis instead of GetScaledAxis inside FConstraintInstance::SetRefFrame
https://udn.unrealengine.com/questions/215623/scaled-reference-frame-in-constraint-instance.html
Fix  FConstraintInstance::IsTerminated()
https://udn.unrealengine.com/questions/215873/bug-in-fconstraintinstanceisterminated.html
Remove out-of-date comment about constraint instance pooling

diff --git a/Engine/Source/Runtime/Engine/Classes/PhysicsEngine/ConstraintInstance.h b/Engine/Source/Runtime/Engine/Classes/PhysicsEngine/ConstraintInstance.h
--- a/Engine/Source/Runtime/Engine/Classes/PhysicsEngine/ConstraintInstance.h
+++ b/Engine/Source/Runtime/Engine/Classes/PhysicsEngine/ConstraintInstance.h
@@ -353,7 +353,7 @@ struct ENGINE_API FConstraintInstance
 
 	/** Limit to the force the angular drive can apply. */
 	UPROPERTY(EditAnywhere, BlueprintReadOnly, Category=AngularMotor)
-	float AngularDriveForceLimit;    //	Due to the way the ConstraintInstance pooling works, this MUST BE LAST PROPERTY OF THE CLASS.
+	float AngularDriveForceLimit;
 
 	float AverageMass;
 
diff --git a/Engine/Source/Runtime/Engine/Private/PhysicsEngine/ConstraintInstance.cpp b/Engine/Source/Runtime/Engine/Private/PhysicsEngine/ConstraintInstance.cpp
--- a/Engine/Source/Runtime/Engine/Private/PhysicsEngine/ConstraintInstance.cpp
+++ b/Engine/Source/Runtime/Engine/Private/PhysicsEngine/ConstraintInstance.cpp
@@ -481,7 +481,7 @@ void FConstraintInstance::TermConstraint()
 bool FConstraintInstance::IsTerminated() const
 {
 #if WITH_PHYSX
-	return (ConstraintData != NULL);
+	return (ConstraintData == NULL);
 #else 
 	return true;
 #endif //WITH_PHYSX
@@ -559,14 +559,14 @@ void FConstraintInstance::SetRefFrame(EConstraintFrame::Type Frame, const FTrans
 	if(Frame == EConstraintFrame::Frame1)
 	{
 		Pos1 = RefFrame.GetTranslation();
-		PriAxis1 = RefFrame.GetScaledAxis( EAxis::X );
-		SecAxis1 = RefFrame.GetScaledAxis( EAxis::Y );
+		PriAxis1 = RefFrame.GetUnitAxis( EAxis::X );
+		SecAxis1 = RefFrame.GetUnitAxis( EAxis::Y );
 	}
 	else
 	{
 		Pos2 = RefFrame.GetTranslation();
-		PriAxis2 = RefFrame.GetScaledAxis( EAxis::X );
-		SecAxis2 = RefFrame.GetScaledAxis( EAxis::Y );
+		PriAxis2 = RefFrame.GetUnitAxis( EAxis::X );
+		SecAxis2 = RefFrame.GetUnitAxis( EAxis::Y );
 	}
 
 #if WITH_PHYSX
