Blendspaces support -1.0 play rate as well as +1.0

diff --git a/Engine/Source/Runtime/Engine/Private/Animation/BlendSpaceBase.cpp b/Engine/Source/Runtime/Engine/Private/Animation/BlendSpaceBase.cpp
--- a/Engine/Source/Runtime/Engine/Private/Animation/BlendSpaceBase.cpp
+++ b/Engine/Source/Runtime/Engine/Private/Animation/BlendSpaceBase.cpp
@@ -274,19 +274,23 @@ void UBlendSpaceBase::TickAssetPlayerInstance(const FAnimTickRecord& Instance, c
 						const FBlendSample& Sample = SampleData[SampleDataIndex];
 						if( Sample.Animation )
 						{
-							const float PrevSampleDataTime = ClampedNormalizedPreviousTime * Sample.Animation->SequenceLength;
+							const float SampleNormalizedPreviousTime = Sample.Animation->RateScale >= 0.f ? ClampedNormalizedPreviousTime : 1.f - ClampedNormalizedPreviousTime;
+							const float SampleNormalizedCurrentTime = Sample.Animation->RateScale >= 0.f ? ClampedNormalizedCurrentTime : 1.f - ClampedNormalizedCurrentTime;
+
+							const float PrevSampleDataTime = SampleNormalizedPreviousTime * Sample.Animation->SequenceLength;
 							float& CurrentSampleDataTime = SampleEntry.Time;
-							CurrentSampleDataTime = ClampedNormalizedCurrentTime * Sample.Animation->SequenceLength;
+							CurrentSampleDataTime = SampleNormalizedCurrentTime * Sample.Animation->SequenceLength;
 
 							if( bGenerateNotifies && (!bTriggerNotifyHighestWeightedAnim || (I == HighestWeightIndex)))
 							{
 								// Figure out delta time 
 								float DeltaTimePosition = CurrentSampleDataTime - PrevSampleDataTime;
+								const float SampleMoveDelta = MoveDelta * Sample.Animation->RateScale;
 
 								// if we went against play rate, then loop around.
-								if( (MoveDelta * DeltaTimePosition) < 0.f )
+								if ((SampleMoveDelta * DeltaTimePosition) < 0.f)
 								{
-									DeltaTimePosition += FMath::Sign<float>(MoveDelta) * Sample.Animation->SequenceLength;
+									DeltaTimePosition += FMath::Sign<float>(SampleMoveDelta) * Sample.Animation->SequenceLength;
 								}
 
 								// Harvest and record notifies
@@ -844,7 +848,7 @@ float UBlendSpaceBase::GetAnimationLengthFromSampleData(const TArray<FBlendSampl
 			if (Sample.Animation)
 			{
 				// apply rate scale to get actual playback time
-				BlendAnimLength += (Sample.Animation->SequenceLength/(Sample.Animation->RateScale != 0.0f? Sample.Animation->RateScale : 1.0f))*SampleDataList[I].GetWeight();
+				BlendAnimLength += (Sample.Animation->SequenceLength / (Sample.Animation->RateScale != 0.0f ? FMath::Abs(Sample.Animation->RateScale) : 1.0f))*SampleDataList[I].GetWeight();
 				UE_LOG(LogAnimation, Verbose, TEXT("[%d] - Sample Animation(%s) : Weight(%0.5f) "), I+1, *Sample.Animation->GetName(), SampleDataList[I].GetWeight());
 			}
 		}
