Fix case-sensitive issue.

diff --git a/Engine/Source/Runtime/Engine/Private/EnginePrivate.h b/Engine/Source/Runtime/Engine/Private/EnginePrivate.h
--- a/Engine/Source/Runtime/Engine/Private/EnginePrivate.h
+++ b/Engine/Source/Runtime/Engine/Private/EnginePrivate.h
@@ -34,7 +34,7 @@
 #include "GameFramework/GameMode.h"
 #include "Components/InputComponent.h"
 #include "GameFramework/ForceFeedbackEffect.h"
-#include "Engine/SubSurfaceProfile.h"
+#include "Engine/SubsurfaceProfile.h"
 #include "GameFramework/PlayerController.h"
 #include "Engine/DebugCameraController.h"
 #include "Engine/Canvas.h"
