Pull request for Substance plugin #311

diff --git a/Engine/Source/Runtime/Engine/Private/Texture2DDynamic.cpp b/Engine/Source/Runtime/Engine/Private/Texture2DDynamic.cpp
--- a/Engine/Source/Runtime/Engine/Private/Texture2DDynamic.cpp
+++ b/Engine/Source/Runtime/Engine/Private/Texture2DDynamic.cpp
@@ -57,7 +57,7 @@ void FTexture2DDynamicResource::InitRHI()
 		Flags |= TexCreate_NoTiling;
 	}
 	FRHIResourceCreateInfo CreateInfo;
-	Texture2DRHI = RHICreateTexture2D( Owner->SizeX, Owner->SizeY, Owner->Format, Owner->NumMips, 1, Flags, CreateInfo );
+	Texture2DRHI = RHICreateTexture2D(GetSizeX(), GetSizeY(), Owner->Format, Owner->NumMips, 1, Flags, CreateInfo);
 	TextureRHI = Texture2DRHI;
 	RHIUpdateTextureReference(Owner->TextureReference.TextureReferenceRHI,TextureRHI);
 }
diff --git a/Engine/Source/Runtime/Engine/Public/TextureResource.h b/Engine/Source/Runtime/Engine/Public/TextureResource.h
--- a/Engine/Source/Runtime/Engine/Public/TextureResource.h
+++ b/Engine/Source/Runtime/Engine/Public/TextureResource.h
@@ -395,7 +395,7 @@ class FTexture2DDynamicResource : public FTextureResource
 {
 public:
 	/** Initialization constructor. */
-	FTexture2DDynamicResource(class UTexture2DDynamic* InOwner);
+	ENGINE_API FTexture2DDynamicResource(class UTexture2DDynamic* InOwner);
 
 	/** Returns the width of the texture in pixels. */
 	virtual uint32 GetSizeX() const;
@@ -404,10 +404,10 @@ public:
 	virtual uint32 GetSizeY() const;
 
 	/** Called when the resource is initialized. This is only called by the rendering thread. */
-	virtual void InitRHI() override;
+	ENGINE_API virtual void InitRHI() override;
 
 	/** Called when the resource is released. This is only called by the rendering thread. */
-	virtual void ReleaseRHI() override;
+	ENGINE_API virtual void ReleaseRHI() override;
 
 	/** Returns the Texture2DRHI, which can be used for locking/unlocking the mips. */
 	ENGINE_API FTexture2DRHIRef GetTexture2DRHI();
