Support duration on anim notify begin. Requested often, decided to add.

diff --git a/Engine/Source/Runtime/Engine/Classes/Animation/AnimNotifies/AnimNotifyState.h b/Engine/Source/Runtime/Engine/Classes/Animation/AnimNotifies/AnimNotifyState.h
--- a/Engine/Source/Runtime/Engine/Classes/Animation/AnimNotifies/AnimNotifyState.h
+++ b/Engine/Source/Runtime/Engine/Classes/Animation/AnimNotifies/AnimNotifyState.h
@@ -21,7 +21,7 @@ class ENGINE_API UAnimNotifyState : public UObject
 	FString GetNotifyName() const;
 
 	UFUNCTION(BlueprintImplementableEvent)
-	virtual bool Received_NotifyBegin(USkeletalMeshComponent * MeshComp, UAnimSequenceBase * Animation) const;
+	virtual bool Received_NotifyBegin(USkeletalMeshComponent * MeshComp, UAnimSequenceBase * Animation, float TotalDuration) const;
 	
 	UFUNCTION(BlueprintImplementableEvent)
 	virtual bool Received_NotifyTick(USkeletalMeshComponent * MeshComp, UAnimSequenceBase * Animation, float FrameDeltaTime) const;
@@ -36,7 +36,7 @@ class ENGINE_API UAnimNotifyState : public UObject
 
 #endif // WITH_EDITORONLY_DATA
 
-	virtual void NotifyBegin(USkeletalMeshComponent * MeshComp, UAnimSequenceBase * Animation);
+	virtual void NotifyBegin(USkeletalMeshComponent * MeshComp, UAnimSequenceBase * Animation, float TotalDuration);
 	virtual void NotifyTick(USkeletalMeshComponent * MeshComp, UAnimSequenceBase * Animation, float FrameDeltaTime);
 	virtual void NotifyEnd(USkeletalMeshComponent * MeshComp, UAnimSequenceBase * Animation);
 
diff --git a/Engine/Source/Runtime/Engine/Classes/Animation/AnimNotifies/AnimNotifyState_TimedParticleEffect.h b/Engine/Source/Runtime/Engine/Classes/Animation/AnimNotifies/AnimNotifyState_TimedParticleEffect.h
--- a/Engine/Source/Runtime/Engine/Classes/Animation/AnimNotifies/AnimNotifyState_TimedParticleEffect.h
+++ b/Engine/Source/Runtime/Engine/Classes/Animation/AnimNotifies/AnimNotifyState_TimedParticleEffect.h
@@ -48,7 +48,7 @@ class UAnimNotifyState_TimedParticleEffect : public UAnimNotifyState
 	virtual void PreEditChange(UProperty* PropertyAboutToChange) override;
 #endif
 
-	virtual void NotifyBegin(class USkeletalMeshComponent * MeshComp, class UAnimSequenceBase * Animation) override;
+	virtual void NotifyBegin(class USkeletalMeshComponent * MeshComp, class UAnimSequenceBase * Animation, float TotalDuration) override;
 	virtual void NotifyTick(class USkeletalMeshComponent * MeshComp, class UAnimSequenceBase * Animation, float FrameDeltaTime) override;
 	virtual void NotifyEnd(class USkeletalMeshComponent * MeshComp, class UAnimSequenceBase * Animation) override;
 
diff --git a/Engine/Source/Runtime/Engine/Classes/Animation/AnimNotifies/AnimNotifyState_Trail.h b/Engine/Source/Runtime/Engine/Classes/Animation/AnimNotifies/AnimNotifyState_Trail.h
--- a/Engine/Source/Runtime/Engine/Classes/Animation/AnimNotifies/AnimNotifyState_Trail.h
+++ b/Engine/Source/Runtime/Engine/Classes/Animation/AnimNotifies/AnimNotifyState_Trail.h
@@ -54,7 +54,7 @@ class UAnimNotifyState_Trail : public UAnimNotifyState
 	uint32 bRenderTessellation : 1;
 #endif // WITH_EDITORONLY_DATA
 
-	virtual void NotifyBegin(class USkeletalMeshComponent * MeshComp, class UAnimSequenceBase * Animation) override;
+	virtual void NotifyBegin(class USkeletalMeshComponent * MeshComp, class UAnimSequenceBase * Animation, float TotalDuration) override;
 	virtual void NotifyTick(class USkeletalMeshComponent * MeshComp, class UAnimSequenceBase * Animation, float FrameDeltaTime) override;
 	virtual void NotifyEnd(class USkeletalMeshComponent * MeshComp, class UAnimSequenceBase * Animation) override;
 
diff --git a/Engine/Source/Runtime/Engine/Private/Animation/AnimInstance.cpp b/Engine/Source/Runtime/Engine/Private/Animation/AnimInstance.cpp
--- a/Engine/Source/Runtime/Engine/Private/Animation/AnimInstance.cpp
+++ b/Engine/Source/Runtime/Engine/Private/Animation/AnimInstance.cpp
@@ -1136,7 +1136,7 @@ void UAnimInstance::TriggerAnimNotifies(float DeltaSeconds)
 	for (int32 Index = 0; Index < NotifyStateBeginEvent.Num(); Index++)
 	{
 		const FAnimNotifyEvent * AnimNotifyEvent = NotifyStateBeginEvent[Index];
-		AnimNotifyEvent->NotifyStateClass->NotifyBegin(SkelMeshComp, Cast<UAnimSequenceBase>(AnimNotifyEvent->NotifyStateClass->GetOuter()));
+		AnimNotifyEvent->NotifyStateClass->NotifyBegin(SkelMeshComp, Cast<UAnimSequenceBase>(AnimNotifyEvent->NotifyStateClass->GetOuter()), AnimNotifyEvent->Duration);
 	}
 
 	// Switch our arrays.
diff --git a/Engine/Source/Runtime/Engine/Private/Animation/AnimNotifyState.cpp b/Engine/Source/Runtime/Engine/Private/Animation/AnimNotifyState.cpp
--- a/Engine/Source/Runtime/Engine/Private/Animation/AnimNotifyState.cpp
+++ b/Engine/Source/Runtime/Engine/Private/Animation/AnimNotifyState.cpp
@@ -16,9 +16,9 @@ UAnimNotifyState::UAnimNotifyState(const class FPostConstructInitializePropertie
 }
 
 
-void UAnimNotifyState::NotifyBegin(class USkeletalMeshComponent * MeshComp, class UAnimSequenceBase * Animation)
+void UAnimNotifyState::NotifyBegin(class USkeletalMeshComponent * MeshComp, class UAnimSequenceBase * Animation, float TotalDuration)
 {
-	Received_NotifyBegin(MeshComp, Animation);
+	Received_NotifyBegin(MeshComp, Animation, TotalDuration);
 }
 
 void UAnimNotifyState::NotifyTick(class USkeletalMeshComponent * MeshComp, class UAnimSequenceBase * Animation, float FrameDeltaTime)
diff --git a/Engine/Source/Runtime/Engine/Private/Animation/AnimNotifyState_TimedParticleEffect.cpp b/Engine/Source/Runtime/Engine/Private/Animation/AnimNotifyState_TimedParticleEffect.cpp
--- a/Engine/Source/Runtime/Engine/Private/Animation/AnimNotifyState_TimedParticleEffect.cpp
+++ b/Engine/Source/Runtime/Engine/Private/Animation/AnimNotifyState_TimedParticleEffect.cpp
@@ -12,7 +12,7 @@ UAnimNotifyState_TimedParticleEffect::UAnimNotifyState_TimedParticleEffect(const
 	RotationOffset = FRotator(0.0f, 0.0f, 0.0f);
 }
 
-void UAnimNotifyState_TimedParticleEffect::NotifyBegin(USkeletalMeshComponent * MeshComp, class UAnimSequenceBase * Animation)
+void UAnimNotifyState_TimedParticleEffect::NotifyBegin(USkeletalMeshComponent * MeshComp, class UAnimSequenceBase * Animation, float TotalDuration)
 {
 	// Only spawn if we've got valid params
 	if(ValidateParameters(MeshComp))
@@ -20,7 +20,7 @@ void UAnimNotifyState_TimedParticleEffect::NotifyBegin(USkeletalMeshComponent *
 		UParticleSystemComponent* NewComponent = UGameplayStatics::SpawnEmitterAttached(PSTemplate, MeshComp, SocketName, LocationOffset, RotationOffset);
 	}
 
-	Received_NotifyBegin(MeshComp, Animation);
+	Received_NotifyBegin(MeshComp, Animation, TotalDuration);
 }
 
 void UAnimNotifyState_TimedParticleEffect::NotifyTick(USkeletalMeshComponent * MeshComp, class UAnimSequenceBase * Animation, float FrameDeltaTime)
diff --git a/Engine/Source/Runtime/Engine/Private/Animation/AnimNotifyState_Trail.cpp b/Engine/Source/Runtime/Engine/Private/Animation/AnimNotifyState_Trail.cpp
--- a/Engine/Source/Runtime/Engine/Private/Animation/AnimNotifyState_Trail.cpp
+++ b/Engine/Source/Runtime/Engine/Private/Animation/AnimNotifyState_Trail.cpp
@@ -38,7 +38,7 @@ void UAnimNotifyState_Trail::AnimNotifyEventChanged(class USkeletalMeshComponent
 {
 }
 
-void UAnimNotifyState_Trail::NotifyBegin(class USkeletalMeshComponent * MeshComp, class UAnimSequenceBase * Animation)
+void UAnimNotifyState_Trail::NotifyBegin(class USkeletalMeshComponent * MeshComp, class UAnimSequenceBase * Animation, float TotalDuration)
 {
 	bool bError = ValidateInput(MeshComp);
 
@@ -120,7 +120,7 @@ void UAnimNotifyState_Trail::NotifyBegin(class USkeletalMeshComponent * MeshComp
 		}
 	}
 
-	Received_NotifyBegin(MeshComp, Animation);
+	Received_NotifyBegin(MeshComp, Animation, TotalDuration);
 }
 
 void UAnimNotifyState_Trail::NotifyTick(class USkeletalMeshComponent * MeshComp, class UAnimSequenceBase * Animation, float FrameDeltaTime)
