Added SetVirtualJoystickVisibility() to PlayerController.

diff --git a/Engine/Source/Runtime/Engine/Classes/GameFramework/PlayerController.h b/Engine/Source/Runtime/Engine/Classes/GameFramework/PlayerController.h
--- a/Engine/Source/Runtime/Engine/Classes/GameFramework/PlayerController.h
+++ b/Engine/Source/Runtime/Engine/Classes/GameFramework/PlayerController.h
@@ -963,6 +963,10 @@ public:
 	UFUNCTION(BlueprintCallable, Category="Game|Player")
 	virtual void ActivateTouchInterface(class UTouchInterface* NewTouchInterface);
 
+	/** Set the virtual joystick visibility. */
+	UFUNCTION(BlueprintCallable, Category = "Game|Player")
+	virtual void SetVirtualJoystickVisibility(bool bVisible);
+
 public:
 	/**
 	 * Change Camera mode
diff --git a/Engine/Source/Runtime/Engine/Private/PlayerController.cpp b/Engine/Source/Runtime/Engine/Private/PlayerController.cpp
--- a/Engine/Source/Runtime/Engine/Private/PlayerController.cpp
+++ b/Engine/Source/Runtime/Engine/Private/PlayerController.cpp
@@ -4195,6 +4195,14 @@ void APlayerController::ActivateTouchInterface(UTouchInterface* NewTouchInterfac
 	CurrentTouchInterface = NewTouchInterface;
 }
 
+void APlayerController::SetVirtualJoystickVisibility(bool bVisible)
+{
+	if (VirtualJoystick.IsValid())
+	{
+		VirtualJoystick->SetVisibility(bVisible, false);
+	}
+}
+
 void APlayerController::UpdateCameraManager(float DeltaSeconds)
 {
 	if (PlayerCameraManager != NULL)
