Shifting texture streaming data during world origin rebasing

diff --git a/Engine/Source/Runtime/Engine/Private/Level.cpp b/Engine/Source/Runtime/Engine/Private/Level.cpp
--- a/Engine/Source/Runtime/Engine/Private/Level.cpp
+++ b/Engine/Source/Runtime/Engine/Private/Level.cpp
@@ -1824,16 +1824,22 @@ bool ULevel::IsCurrentLevel() const
 
 void ULevel::ApplyWorldOffset(const FVector& InWorldOffset, bool bWorldShift)
 {
-	// Update texture streaming data to account for the move
-	for (TMap< UTexture2D*, TArray<FStreamableTextureInstance> >::TIterator It(TextureToInstancesMap); It; ++It)
+	if (bTextureStreamingBuilt)
 	{
-		TArray<FStreamableTextureInstance>& TextureInfo = It.Value();
-		for (int32 i = 0; i < TextureInfo.Num(); i++)
+		// Update texture streaming data to account for the move
+		for (TMap< UTexture2D*, TArray<FStreamableTextureInstance> >::TIterator It(TextureToInstancesMap); It; ++It)
 		{
-			TextureInfo[i].BoundingSphere.Center+= InWorldOffset;
+			TArray<FStreamableTextureInstance>& TextureInfo = It.Value();
+			for (int32 i = 0; i < TextureInfo.Num(); i++)
+			{
+				TextureInfo[i].BoundingSphere.Center+= InWorldOffset;
+			}
 		}
+		
+		// Re-add level data to a manager
+		IStreamingManager::Get().AddPreparedLevel( this );
 	}
-	
+
 	// Move precomputed light samples
 	if (PrecomputedLightVolume)
 	{
