Header cleanup: NetConnection.h, NavigationGraph.h, Character.h, NavigationData.h, NavLinkProxy.h.

diff --git a/Engine/Source/Runtime/Engine/Classes/AI/Navigation/NavLinkHostInterface.h b/Engine/Source/Runtime/Engine/Classes/AI/Navigation/NavLinkHostInterface.h
--- a/Engine/Source/Runtime/Engine/Classes/AI/Navigation/NavLinkHostInterface.h
+++ b/Engine/Source/Runtime/Engine/Classes/AI/Navigation/NavLinkHostInterface.h
@@ -1,6 +1,7 @@
 // Copyright 1998-2014 Epic Games, Inc. All Rights Reserved.
 
 #pragma once
+#include "Interface.h"
 #include "NavLinkHostInterface.generated.h"
 
 struct FNavigationLink;
diff --git a/Engine/Source/Runtime/Engine/Classes/AI/Navigation/NavLinkProxy.h b/Engine/Source/Runtime/Engine/Classes/AI/Navigation/NavLinkProxy.h
--- a/Engine/Source/Runtime/Engine/Classes/AI/Navigation/NavLinkProxy.h
+++ b/Engine/Source/Runtime/Engine/Classes/AI/Navigation/NavLinkProxy.h
@@ -1,8 +1,9 @@
 // Copyright 1998-2014 Epic Games, Inc. All Rights Reserved.
 
 #pragma once
-#include "NavigationTypes.h"
+#include "AI/Navigation/NavLinkHostInterface.h"
 #include "AI/Navigation/NavRelevantActorInterface.h"
+#include "GameFramework/Actor.h"
 #include "NavLinkProxy.generated.h"
 
 DECLARE_DYNAMIC_MULTICAST_DELEGATE_TwoParams( FSmartLinkReachedSignature, class AActor*, MovingActor, const FVector&, DestinationPoint );
diff --git a/Engine/Source/Runtime/Engine/Classes/AI/Navigation/NavRelevantActorInterface.h b/Engine/Source/Runtime/Engine/Classes/AI/Navigation/NavRelevantActorInterface.h
--- a/Engine/Source/Runtime/Engine/Classes/AI/Navigation/NavRelevantActorInterface.h
+++ b/Engine/Source/Runtime/Engine/Classes/AI/Navigation/NavRelevantActorInterface.h
@@ -1,6 +1,7 @@
 // Copyright 1998-2014 Epic Games, Inc. All Rights Reserved.
 
 #pragma once
+#include "Interface.h"
 #include "NavRelevantActorInterface.generated.h"
 
 UINTERFACE(/**/MinimalAPI, meta=(CannotImplementInterfaceInBlueprint))
diff --git a/Engine/Source/Runtime/Engine/Classes/AI/Navigation/NavigationData.h b/Engine/Source/Runtime/Engine/Classes/AI/Navigation/NavigationData.h
--- a/Engine/Source/Runtime/Engine/Classes/AI/Navigation/NavigationData.h
+++ b/Engine/Source/Runtime/Engine/Classes/AI/Navigation/NavigationData.h
@@ -1,9 +1,8 @@
 // Copyright 1998-2014 Epic Games, Inc. All Rights Reserved.
 
 #pragma once
-#include "AI/NavDataGenerator.h"
-#include "AI/Navigation/NavigationTypes.h"
 #include "AI/Navigation/NavFilters/NavigationQueryFilter.h"
+#include "AI/Navigation/NavigationTypes.h"
 #include "GameFramework/Actor.h"
 #include "NavigationData.generated.h"
 
diff --git a/Engine/Source/Runtime/Engine/Private/Camera/CameraAnim.cpp b/Engine/Source/Runtime/Engine/Private/Camera/CameraAnim.cpp
--- a/Engine/Source/Runtime/Engine/Private/Camera/CameraAnim.cpp
+++ b/Engine/Source/Runtime/Engine/Private/Camera/CameraAnim.cpp
@@ -1,6 +1,8 @@
 // Copyright 1998-2014 Epic Games, Inc. All Rights Reserved.
 
 #include "EnginePrivate.h"
+#include "Camera/CameraActor.h"
+#include "Camera/CameraAnim.h"
 #include "ParticleDefinitions.h"
 #include "Camera/CameraActor.h"
 #include "Matinee/MatineeActor.h"
diff --git a/Engine/Source/Runtime/Engine/Private/Camera/CameraAnimInst.cpp b/Engine/Source/Runtime/Engine/Private/Camera/CameraAnimInst.cpp
--- a/Engine/Source/Runtime/Engine/Private/Camera/CameraAnimInst.cpp
+++ b/Engine/Source/Runtime/Engine/Private/Camera/CameraAnimInst.cpp
@@ -3,6 +3,7 @@
 
 #include "EnginePrivate.h"
 #include "Camera/CameraActor.h"
+#include "Camera/CameraAnim.h"
 #include "Matinee/InterpTrackMove.h"
 #include "Matinee/InterpGroup.h"
 #include "Matinee/InterpGroupInst.h"
diff --git a/Engine/Source/Runtime/Engine/Private/DestructibleComponent.cpp b/Engine/Source/Runtime/Engine/Private/DestructibleComponent.cpp
--- a/Engine/Source/Runtime/Engine/Private/DestructibleComponent.cpp
+++ b/Engine/Source/Runtime/Engine/Private/DestructibleComponent.cpp
@@ -13,6 +13,7 @@
 #include "ParticleDefinitions.h"
 #include "ObjectEditorUtils.h"
 #include "Particles/ParticleSystemComponent.h"
+#include "PhysicsEngine/DestructibleActor.h"
 
 UDestructibleComponent::UDestructibleComponent(const class FPostConstructInitializeProperties& PCIP)
 	: Super(PCIP)
diff --git a/Engine/Source/Runtime/Engine/Private/LevelScriptBlueprint.cpp b/Engine/Source/Runtime/Engine/Private/LevelScriptBlueprint.cpp
--- a/Engine/Source/Runtime/Engine/Private/LevelScriptBlueprint.cpp
+++ b/Engine/Source/Runtime/Engine/Private/LevelScriptBlueprint.cpp
@@ -1,6 +1,7 @@
 // Copyright 1998-2014 Epic Games, Inc. All Rights Reserved.
 
 #include "EnginePrivate.h"
+#include "Engine/LevelScriptActor.h"
 #include "Engine/LevelScriptBlueprint.h"
 #include "Engine/LevelScriptActor.h"
 #include "BlueprintUtilities.h"
diff --git a/Engine/Source/Runtime/Engine/Private/PlayerController.cpp b/Engine/Source/Runtime/Engine/Private/PlayerController.cpp
--- a/Engine/Source/Runtime/Engine/Private/PlayerController.cpp
+++ b/Engine/Source/Runtime/Engine/Private/PlayerController.cpp
@@ -24,6 +24,7 @@
 #include "GameFramework/HUD.h"
 #include "ContentStreaming.h"
 #include "GameFramework/PawnMovementComponent.h"
+#include "Camera/CameraActor.h"
 
 DEFINE_LOG_CATEGORY(LogPlayerController);
 
