Modified FAnimUpdateRateParameters and Moved AnimUpdateRateShiftTag into Actor again for efficiency

diff --git a/Engine/Source/Runtime/Engine/Classes/Components/SkinnedMeshComponent.h b/Engine/Source/Runtime/Engine/Classes/Components/SkinnedMeshComponent.h
--- a/Engine/Source/Runtime/Engine/Classes/Components/SkinnedMeshComponent.h
+++ b/Engine/Source/Runtime/Engine/Classes/Components/SkinnedMeshComponent.h
@@ -588,7 +588,7 @@ protected:
 
 	// Update Rate
 public:
-	/** if TRUE, Owner will determine how often animation will be updated and evaluated. See AActor::AnimUpdateRateTick() 
+	/** if TRUE, Owner will determine how often animation will be updated and evaluated. See AnimUpdateRateTick() 
 	 * This allows to skip frames for performance. (For example based on visibility and size on screen). */
 	UPROPERTY(EditAnywhere, AdvancedDisplay, BlueprintReadWrite, Category=Optimization)
 	bool bEnableUpdateRateOptimizations;
@@ -848,14 +848,6 @@ public:
 	UPROPERTY(Transient)
 	struct FAnimUpdateRateParameters AnimUpdateRateParams;
 
-	/** Unique Tag assigned to spread updates of SkinnedMeshes over time. */
-	UPROPERTY(Transient)
-	uint32 AnimUpdateRateShiftTag;
-
-	/** Frame counter to call AnimUpdateRateTick() just once per frame. */
-	UPROPERTY(Transient)
-	uint32 AnimUpdateRateFrameCount;
-
 	/** Aimation Update Rate Tick. */
 	void AnimUpdateRateTick();
 
diff --git a/Engine/Source/Runtime/Engine/Classes/Engine/EngineTypes.h b/Engine/Source/Runtime/Engine/Classes/Engine/EngineTypes.h
--- a/Engine/Source/Runtime/Engine/Classes/Engine/EngineTypes.h
+++ b/Engine/Source/Runtime/Engine/Classes/Engine/EngineTypes.h
@@ -1596,7 +1596,7 @@ public:
 	 * @param : NewEvaluationRate. How often animation will be evaluated. 1 = every frame, 2 = every 2 frames, etc.
 	 * @param : bNewInterpSkippedFrames. When skipping a frame, should it be interpolated or frozen?
 	 */
-	void Set(const class USkinnedMeshComponent & Owner, const int32 & NewUpdateRate, const int32 & NewEvaluationRate, const bool & bNewInterpSkippedFrames);
+	void Set(class AActor& Owner, const int32 & NewUpdateRate, const int32 & NewEvaluationRate, const bool & bNewInterpSkippedFrames);
 
 	/* Getter for UpdateRate */
 	int32 GetUpdateRate() const
diff --git a/Engine/Source/Runtime/Engine/Classes/GameFramework/Actor.h b/Engine/Source/Runtime/Engine/Classes/GameFramework/Actor.h
--- a/Engine/Source/Runtime/Engine/Classes/GameFramework/Actor.h
+++ b/Engine/Source/Runtime/Engine/Classes/GameFramework/Actor.h
@@ -275,6 +275,19 @@ public:
 	UPROPERTY(transient)
 	TArray<AActor*> Children;    
 	
+	// Animation update rate control.
+protected:
+	/** Unique Tag assigned to spread updates of SkinnedMeshes over time. */
+	UPROPERTY(Transient)
+	uint32 AnimUpdateRateShiftTag;
+public:
+	/** Frame counter to call AnimUpdateRateTick() just once per frame. */
+	UPROPERTY(Transient)
+	uint32 AnimUpdateRateFrameCount;
+
+	/** get a unique ID to share with all SkinnedMeshComponents in this actor. */
+	uint32 GetAnimUpdateRateShiftTag();
+
 protected:
 	// Collision primitive.
 	UPROPERTY()
diff --git a/Engine/Source/Runtime/Engine/Private/Actor.cpp b/Engine/Source/Runtime/Engine/Private/Actor.cpp
--- a/Engine/Source/Runtime/Engine/Private/Actor.cpp
+++ b/Engine/Source/Runtime/Engine/Private/Actor.cpp
@@ -73,6 +73,22 @@ AActor::AActor(const class FPostConstructInitializeProperties& PCIP)
 	bCanBeDamaged = true;
 	bFindCameraComponentWhenViewTarget = true;
 	bAllowReceiveTickEventOnDedicatedServer = true;
+	AnimUpdateRateShiftTag = 0;
+	AnimUpdateRateFrameCount = 0;
+}
+
+// Global counter to spread SkinnedMeshComponent tick updates.
+static uint32 SkinnedMeshUpdateRateGroupCount = 0;
+
+uint32 AActor::GetAnimUpdateRateShiftTag()
+{
+	// If hasn't been initialized yet, pick a unique ID, to spread population over frames.
+	if (AnimUpdateRateShiftTag == 0)
+	{
+		AnimUpdateRateShiftTag = ++SkinnedMeshUpdateRateGroupCount;
+	}
+
+	return AnimUpdateRateShiftTag;
 }
 
 void FActorTickFunction::ExecuteTick(float DeltaTime, enum ELevelTick TickType, ENamedThreads::Type CurrentThread, const FGraphEventRef& MyCompletionGraphEvent)
diff --git a/Engine/Source/Runtime/Engine/Private/SkeletalMesh.cpp b/Engine/Source/Runtime/Engine/Private/SkeletalMesh.cpp
--- a/Engine/Source/Runtime/Engine/Private/SkeletalMesh.cpp
+++ b/Engine/Source/Runtime/Engine/Private/SkeletalMesh.cpp
@@ -4465,9 +4465,6 @@ USkinnedMeshComponent::USkinnedMeshComponent(const class FPostConstructInitializ
 	SectionIndexPreview = -1;
 #endif // WITH_EDITORONLY_DATA
 	bPerBoneMotionBlur = true;
-
-	AnimUpdateRateShiftTag = 0;
-	AnimUpdateRateFrameCount = 0;
 }
 
 
diff --git a/Engine/Source/Runtime/Engine/Private/SkinnedMeshComponent.cpp b/Engine/Source/Runtime/Engine/Private/SkinnedMeshComponent.cpp
--- a/Engine/Source/Runtime/Engine/Private/SkinnedMeshComponent.cpp
+++ b/Engine/Source/Runtime/Engine/Private/SkinnedMeshComponent.cpp
@@ -237,28 +237,32 @@ void USkinnedMeshComponent::TickUpdateRate()
 	SCOPE_CYCLE_COUNTER(STAT_TickUpdateRate);
 	if( bEnableUpdateRateOptimizations )
 	{
-		// Convert current frame counter from 64 to 32 bits.
-		const uint32 CurrentFrame32 = uint32(GFrameCounter % MAX_uint32);
-
-		// Tick Owner once per frame. All attached SkinnedMeshComponents will share the same settings.
-		if( AnimUpdateRateFrameCount != CurrentFrame32 )
+		AActor* Owner = GetOwner();
+		if (Owner)
 		{
-			AnimUpdateRateFrameCount = CurrentFrame32;
-			AnimUpdateRateTick();
-		}
+			// Convert current frame counter from 64 to 32 bits.
+			const uint32 CurrentFrame32 = uint32(GFrameCounter % MAX_uint32);
 
-		// debug -- @todo: hook this up to a console command.
-		if( bDisplayDebugUpdateRateOptimizations )
-		{
-			FColor DrawColor;
-			switch( AnimUpdateRateParams.GetUpdateRate() )
+			// Tick Owner once per frame. All attached SkinnedMeshComponents will share the same settings.
+			if (Owner->AnimUpdateRateFrameCount != CurrentFrame32)
 			{
-			case 1 : DrawColor = FColor::Red; break;
-			case 2 : DrawColor = FColor::Green; break;
-			case 3 : DrawColor = FColor::Blue; break;
-			default: DrawColor = FColor::Black; break;
+				Owner->AnimUpdateRateFrameCount = CurrentFrame32;
+				AnimUpdateRateTick();
+			}
+
+			// debug -- @todo: hook this up to a console command.
+			if (bDisplayDebugUpdateRateOptimizations)
+			{
+				FColor DrawColor;
+				switch (AnimUpdateRateParams.GetUpdateRate())
+				{
+				case 1: DrawColor = FColor::Red; break;
+				case 2: DrawColor = FColor::Green; break;
+				case 3: DrawColor = FColor::Blue; break;
+				default: DrawColor = FColor::Black; break;
+				}
+				DrawDebugBox(GetWorld(), Bounds.Origin, Bounds.BoxExtent, FQuat::Identity, DrawColor, false);
 			}
-			DrawDebugBox(GetWorld(), Bounds.Origin, Bounds.BoxExtent, FQuat::Identity, DrawColor, false);
 		}
 	}
 }
@@ -1712,16 +1716,13 @@ TArray<FActiveVertexAnim> USkinnedMeshComponent::UpdateActiveVertexAnims(const U
 	return OutVertexAnims;
 }
 
-
-// Global counter to spread SkinnedMeshComponent tick updates.
-static uint32 SkinnedMeshUpdateRateGroupCount = 0;
-
 void USkinnedMeshComponent::AnimUpdateRateTick()
 {
-	// If hasn't been initialized yet, pick a unique ID, to spread population over frames.
-	if (AnimUpdateRateShiftTag == 0)
+	AActor* Owner = GetOwner();
+
+	if (!Owner)
 	{
-		AnimUpdateRateShiftTag = ++SkinnedMeshUpdateRateGroupCount;
+		return;
 	}
 
 	// Go through components and figure out if they've been recently rendered, and the biggest MaxDistanceFactor
@@ -1731,7 +1732,7 @@ void USkinnedMeshComponent::AnimUpdateRateTick()
 
 	// Gather Actor's components
 	TArray<USceneComponent *> ComponentStack;
-	USceneComponent* CurrentComponent = this;
+	USceneComponent* CurrentComponent = Owner->GetRootComponent();
 
 	while (CurrentComponent)
 	{
@@ -1762,18 +1763,23 @@ void USkinnedMeshComponent::AnimUpdateRateSetParams(const bool & bRecentlyRender
 
 	// Human controlled characters should be ticked always fully to minimize latency w/ game play events triggered by animation.
 	AActor* Owner = GetOwner();
-	AController * Controller = Owner ? Owner->GetInstigatorController() : NULL;
+	if (!Owner)
+	{
+		return;
+	}
+
+	AController * Controller = Owner->GetInstigatorController();
 	const bool bHumanControlled = Controller && (Cast<APlayerController>(Controller) != NULL);
 
 	// Not rendered, including dedicated servers. we can skip the Evaluation part.
 	if (!bRecentlyRendered)
 	{
-		AnimUpdateRateParams.Set(*this, (bHumanControlled ? 1 : 4), 4, false);
+		AnimUpdateRateParams.Set(*Owner, (bHumanControlled ? 1 : 4), 4, false);
 	}
 	// Visible controlled characters or playing root motion. Need evaluation and ticking done every frame.
 	else  if (bHumanControlled || bPlayingRootMotion)
 	{
-		AnimUpdateRateParams.Set(*this, 1, 1, false);
+		AnimUpdateRateParams.Set(*Owner, 1, 1, false);
 	}
 	else
 	{
@@ -1793,11 +1799,11 @@ void USkinnedMeshComponent::AnimUpdateRateSetParams(const bool & bRecentlyRender
 			DesiredEvaluationRate = 3;
 		}
 
-		AnimUpdateRateParams.Set(*this, DesiredEvaluationRate, DesiredEvaluationRate, true);
+		AnimUpdateRateParams.Set(*Owner, DesiredEvaluationRate, DesiredEvaluationRate, true);
 	}
 }
 
-void FAnimUpdateRateParameters::Set(const class USkinnedMeshComponent& Owner, const int32 & NewUpdateRate, const int32 & NewEvaluationRate, const bool & bNewInterpSkippedFrames)
+void FAnimUpdateRateParameters::Set(class AActor& Owner, const int32 & NewUpdateRate, const int32 & NewEvaluationRate, const bool & bNewInterpSkippedFrames)
 {
 	UpdateRate = FMath::Max(NewUpdateRate, 1);
 	// Make sure EvaluationRate is a multiple of UpdateRate.
@@ -1805,7 +1811,7 @@ void FAnimUpdateRateParameters::Set(const class USkinnedMeshComponent& Owner, co
 	bInterpolateSkippedFrames = bNewInterpSkippedFrames;
 
 	// Make sure we don't overflow. we don't need very large numbers.
-	const uint32 Counter = (GFrameCounter % MAX_uint8) + (Owner.AnimUpdateRateShiftTag % MAX_uint8);
+	const uint32 Counter = (GFrameCounter % MAX_uint8) + (Owner.GetAnimUpdateRateShiftTag() % MAX_uint8);
 
 	bSkipUpdate = ((Counter % UpdateRate) > 0);
 	bSkipEvaluation = ((Counter % EvaluationRate) > 0);
