Header cleanup: DataTableFunctionLibrary.h, BodySetup2D.h, CameraModifier_CameraShake.h, SoundCue.h, ForceFeedbackEffect.h

diff --git a/Engine/Source/Runtime/Engine/Classes/Camera/CameraModifier_CameraShake.h b/Engine/Source/Runtime/Engine/Classes/Camera/CameraModifier_CameraShake.h
--- a/Engine/Source/Runtime/Engine/Classes/Camera/CameraModifier_CameraShake.h
+++ b/Engine/Source/Runtime/Engine/Classes/Camera/CameraModifier_CameraShake.h
@@ -5,10 +5,7 @@
  */
 
 #pragma once
-
-#include "CameraTypes.h"
-#include "CameraShake.h"
-
+#include "Camera/CameraModifier.h"
 #include "CameraModifier_CameraShake.generated.h"
 
 USTRUCT()
@@ -122,6 +119,3 @@ public:
 	virtual bool ModifyCamera(APlayerCameraManager* Camera, float DeltaTime, struct FMinimalViewInfo& InOutPOV) override;
 	// End UCameraModifer Interface
 };
-
-
-
diff --git a/Engine/Source/Runtime/Engine/Classes/GameFramework/ForceFeedbackEffect.h b/Engine/Source/Runtime/Engine/Classes/GameFramework/ForceFeedbackEffect.h
--- a/Engine/Source/Runtime/Engine/Classes/GameFramework/ForceFeedbackEffect.h
+++ b/Engine/Source/Runtime/Engine/Classes/GameFramework/ForceFeedbackEffect.h
@@ -1,8 +1,6 @@
 // Copyright 1998-2014 Epic Games, Inc. All Rights Reserved.
 #pragma once
-
-#include "../Curves/CurveFloat.h"
-#include "IForceFeedbackSystem.h"
+#include "Curves/CurveFloat.h"
 #include "ForceFeedbackEffect.generated.h"
 
 USTRUCT()
@@ -66,7 +64,6 @@ struct FActiveForceFeedbackEffect
 	bool Update(const float DeltaTime, struct FForceFeedbackValues& Values);
 };
 
-
 /**
  * A predefined force-feedback effect to be played on a controller
  */
diff --git a/Engine/Source/Runtime/Engine/Classes/GameFramework/PlayerController.h b/Engine/Source/Runtime/Engine/Classes/GameFramework/PlayerController.h
--- a/Engine/Source/Runtime/Engine/Classes/GameFramework/PlayerController.h
+++ b/Engine/Source/Runtime/Engine/Classes/GameFramework/PlayerController.h
@@ -11,6 +11,7 @@
 #include "GameFramework/OnlineReplStructs.h"
 #include "GameFramework/Controller.h"
 #include "Engine/LatentActionManager.h"
+#include "GenericPlatform/IForceFeedbackSystem.h"
 #include "PlayerController.generated.h"
 
 class FPrimitiveComponentId;
@@ -51,25 +52,7 @@ struct FDynamicForceFeedbackDetails
 		, Intensity(0.f)
 	{}
 
-	void Update(FForceFeedbackValues& Values) const
-	{
-		if (bAffectsLeftLarge)
-		{
-			Values.LeftLarge = FMath::Clamp(Intensity, Values.LeftLarge, 1.f);
-		}
-		if (bAffectsLeftSmall)
-		{
-			Values.LeftSmall = FMath::Clamp(Intensity, Values.LeftSmall, 1.f);
-		}
-		if (bAffectsRightLarge)
-		{
-			Values.RightLarge = FMath::Clamp(Intensity, Values.RightLarge, 1.f);
-		}
-		if (bAffectsRightSmall)
-		{
-			Values.RightSmall = FMath::Clamp(Intensity, Values.RightSmall, 1.f);
-		}
-	}
+	void Update(FForceFeedbackValues& Values) const;
 };
 
 //=============================================================================
diff --git a/Engine/Source/Runtime/Engine/Classes/Kismet/DataTableFunctionLibrary.h b/Engine/Source/Runtime/Engine/Classes/Kismet/DataTableFunctionLibrary.h
--- a/Engine/Source/Runtime/Engine/Classes/Kismet/DataTableFunctionLibrary.h
+++ b/Engine/Source/Runtime/Engine/Classes/Kismet/DataTableFunctionLibrary.h
@@ -1,9 +1,8 @@
 // Copyright 1998-2014 Epic Games, Inc. All Rights Reserved.
 
 #pragma once
-#include "../Engine/DataTable.h"
-#include "../Engine/CurveTable.h"
-#include "KismetSystemLibrary.h"
+#include "Engine/DataTable.h"
+#include "Kismet/BlueprintFunctionLibrary.h"
 #include "DataTableFunctionLibrary.generated.h"
 
 /** Enum used to indicate success or failure of EvaluateCurveTableRow */
diff --git a/Engine/Source/Runtime/Engine/Classes/PhysicsEngine/BodySetup2D.h b/Engine/Source/Runtime/Engine/Classes/PhysicsEngine/BodySetup2D.h
--- a/Engine/Source/Runtime/Engine/Classes/PhysicsEngine/BodySetup2D.h
+++ b/Engine/Source/Runtime/Engine/Classes/PhysicsEngine/BodySetup2D.h
@@ -1,10 +1,8 @@
 // Copyright 1998-2014 Epic Games, Inc. All Rights Reserved.
 
 #pragma once
-
-#include "AggregateGeometry2D.h"
-#include "Engine/EngineTypes.h"
-
+#include "PhysicsEngine/AggregateGeometry2D.h"
+#include "PhysicsEngine/BodySetup.h"
 #include "BodySetup2D.generated.h"
 
 UCLASS()
diff --git a/Engine/Source/Runtime/Engine/Classes/Sound/SoundCue.h b/Engine/Source/Runtime/Engine/Classes/Sound/SoundCue.h
--- a/Engine/Source/Runtime/Engine/Classes/Sound/SoundCue.h
+++ b/Engine/Source/Runtime/Engine/Classes/Sound/SoundCue.h
@@ -1,10 +1,7 @@
 // Copyright 1998-2014 Epic Games, Inc. All Rights Reserved.
 
 #pragma once
-
-#include "SoundClass.h"
 #include "Sound/SoundBase.h"
-
 #include "SoundCue.generated.h"
 
 struct FActiveSound;
diff --git a/Engine/Source/Runtime/Engine/Private/PlayerController.cpp b/Engine/Source/Runtime/Engine/Private/PlayerController.cpp
--- a/Engine/Source/Runtime/Engine/Private/PlayerController.cpp
+++ b/Engine/Source/Runtime/Engine/Private/PlayerController.cpp
@@ -25,6 +25,7 @@
 #include "ContentStreaming.h"
 #include "GameFramework/PawnMovementComponent.h"
 #include "Camera/CameraActor.h"
+#include "GenericPlatform/IForceFeedbackSystem.h"
 
 DEFINE_LOG_CATEGORY(LogPlayerController);
 
@@ -4258,3 +4259,23 @@ void APlayerController::ClientRepObjRef_Implementation(UObject *Object)
 {
 	UE_LOG(LogPlayerController, Warning, TEXT("APlayerController::ClientRepObjRef repped: %s"), Object ? *Object->GetName() : TEXT("NULL") );
 }
+
+void FDynamicForceFeedbackDetails::Update(FForceFeedbackValues& Values) const
+{
+	if (bAffectsLeftLarge)
+	{
+		Values.LeftLarge = FMath::Clamp(Intensity, Values.LeftLarge, 1.f);
+	}
+	if (bAffectsLeftSmall)
+	{
+		Values.LeftSmall = FMath::Clamp(Intensity, Values.LeftSmall, 1.f);
+	}
+	if (bAffectsRightLarge)
+	{
+		Values.RightLarge = FMath::Clamp(Intensity, Values.RightLarge, 1.f);
+	}
+	if (bAffectsRightSmall)
+	{
+		Values.RightSmall = FMath::Clamp(Intensity, Values.RightSmall, 1.f);
+	}
+}
