Documentation actor revisions. Now has a billboard material.

diff --git a/Engine/Source/Runtime/Engine/Classes/Engine/DocumentationActor.h b/Engine/Source/Runtime/Engine/Classes/Engine/DocumentationActor.h
--- a/Engine/Source/Runtime/Engine/Classes/Engine/DocumentationActor.h
+++ b/Engine/Source/Runtime/Engine/Classes/Engine/DocumentationActor.h
@@ -6,7 +6,7 @@
 #pragma once
 #include "DocumentationActor.generated.h"
 
-UCLASS(hidecategories = (Actor, Transform, Tags, Materials, StaticMesh, Rendering, Components, Blueprint, StaticMeshActor, Object, Collision, Display, Rendering, Physics, Input, Lighting, Layers))
+UCLASS(hidecategories = (Sprite, MaterialSprite, Actor, Transform, Tags, Materials, Rendering, Components, Blueprint, bject, Collision, Display, Rendering, Physics, Input, Lighting, Layers))
 class ENGINE_API ADocumentationActor : public AActor
 {
 	GENERATED_UCLASS_BODY()
@@ -22,8 +22,9 @@ class ENGINE_API ADocumentationActor : public AActor
 	UPROPERTY(Category = HelpDocumentation, EditAnywhere, BlueprintReadWrite, AdvancedDisplay)
 	FString DocumentLink; 
 	
-	UPROPERTY(Category = StaticMeshActor, VisibleAnywhere, BlueprintReadOnly)
-	TSubobjectPtr<class UStaticMeshComponent> StaticMeshComponent;
+ 	UPROPERTY(Category = Sprite, VisibleAnywhere, BlueprintReadOnly)
+	TSubobjectPtr<class UMaterialBillboardComponent> Billboard;
+
 #endif
 
 };
diff --git a/Engine/Source/Runtime/Engine/Private/DocumentationActor.cpp b/Engine/Source/Runtime/Engine/Private/DocumentationActor.cpp
--- a/Engine/Source/Runtime/Engine/Private/DocumentationActor.cpp
+++ b/Engine/Source/Runtime/Engine/Private/DocumentationActor.cpp
@@ -2,8 +2,8 @@
 
 #include "EnginePrivate.h"
 #include "Engine/DocumentationActor.h"
-
 #if WITH_EDITORONLY_DATA
+#include "Components/MaterialBillboardComponent.h"
 #include "IDocumentation.h"
 #endif
 
@@ -16,14 +16,11 @@ ADocumentationActor::ADocumentationActor(const class FPostConstructInitializePro
 	RootComponent = SceneComponent;	
 
 #if WITH_EDITORONLY_DATA
-	static ConstructorHelpers::FObjectFinder<UStaticMesh> DocuemtMesh(TEXT("/Engine/EditorMeshes/EditorHelp"));
-	
-	// Create a mesh to represent our actor
-	StaticMeshComponent = PCIP.CreateDefaultSubobject<UStaticMeshComponent>(this, TEXT("MeshComponent"));
-	StaticMeshComponent->SetCollisionProfileName(UCollisionProfile::NoCollision_ProfileName);
-	StaticMeshComponent->bGenerateOverlapEvents = false;
-	StaticMeshComponent->AttachParent = RootComponent;		
-	StaticMeshComponent->SetStaticMesh(DocuemtMesh.Object);
+ 	static ConstructorHelpers::FObjectFinder<UMaterial> MaterialAsset(TEXT("/Engine/EditorMaterials/HelpActorMaterial.HelpActorMaterial"));
+ 	// Create a Material billboard to represent our actor
+	Billboard = PCIP.CreateDefaultSubobject<UMaterialBillboardComponent>(this, TEXT("BillboardComponent"));
+	Billboard->AddElement(MaterialAsset.Object, nullptr, false, 32.0f, 32.0f, nullptr);
+	Billboard->AttachParent = RootComponent;
 #endif //WITH_EDITORONLY_DATA
 }
 
@@ -47,5 +44,5 @@ bool ADocumentationActor::HasValidDocumentLink() const
 #endif // WITH_EDITORONLY_DATA
 
 	return bDocumentValid;
-	
 }
+
