Removed unnecessary gamethread task from FRecastRenderingSceneProxy #ue4

diff --git a/Engine/Source/Runtime/Engine/Classes/AI/Navigation/NavMeshRenderingComponent.h b/Engine/Source/Runtime/Engine/Classes/AI/Navigation/NavMeshRenderingComponent.h
--- a/Engine/Source/Runtime/Engine/Classes/AI/Navigation/NavMeshRenderingComponent.h
+++ b/Engine/Source/Runtime/Engine/Classes/AI/Navigation/NavMeshRenderingComponent.h
@@ -22,7 +22,7 @@ class ENGINE_API UNavMeshRenderingComponent : public UPrimitiveComponent
 	// End USceneComponent Interface
 
 	void GatherData(struct FNavMeshSceneProxyData*) const;
-	void GatherDataForProxy();
+
 protected:
 	TSharedPtr<struct FNavMeshSceneProxyData, ESPMode::ThreadSafe>	ProxyData;
 };
diff --git a/Engine/Source/Runtime/Engine/Private/NavMeshRenderingComponent.cpp b/Engine/Source/Runtime/Engine/Private/NavMeshRenderingComponent.cpp
--- a/Engine/Source/Runtime/Engine/Private/NavMeshRenderingComponent.cpp
+++ b/Engine/Source/Runtime/Engine/Private/NavMeshRenderingComponent.cpp
@@ -31,7 +31,9 @@ FPrimitiveSceneProxy* UNavMeshRenderingComponent::CreateSceneProxy()
 		{
 			ProxyData = MakeShareable(new FNavMeshSceneProxyData());
 		}
-		ProxyData->bNeedsNewData = true;
+		GatherData(ProxyData.Get());
+		ProxyData->bNeedsNewData = false;
+
 		SceneProxy = new FRecastRenderingSceneProxy(this, ProxyData);
 	}
 	return SceneProxy;
@@ -40,19 +42,6 @@ FPrimitiveSceneProxy* UNavMeshRenderingComponent::CreateSceneProxy()
 #endif
 }
 
-void UNavMeshRenderingComponent::GatherDataForProxy()
-{
-	if (ProxyData.IsValid())
-	{
-#if WITH_RECAST
-		FScopeLock ScopeLock(&ProxyData.Get()->CriticalSection);
-		GatherData( ProxyData.Get() );
-
-		ProxyData->bNeedsNewData = false;
-#endif
-	}
-}
-
 void UNavMeshRenderingComponent::CreateRenderState_Concurrent()
 {
 	Super::CreateRenderState_Concurrent();
diff --git a/Engine/Source/Runtime/Engine/Private/NavMeshRenderingHelpers.h b/Engine/Source/Runtime/Engine/Private/NavMeshRenderingHelpers.h
--- a/Engine/Source/Runtime/Engine/Private/NavMeshRenderingHelpers.h
+++ b/Engine/Source/Runtime/Engine/Private/NavMeshRenderingHelpers.h
@@ -198,145 +198,126 @@ public:
 	{
 		QUICK_SCOPE_CYCLE_COUNTER( STAT_RecastRenderingSceneProxy_DrawDynamicElements );
 
-		if (!bRequestedData && ProxyData.IsValid() && ProxyData->bNeedsNewData)
+		if (ProxyData.IsValid() == false || !ProxyData->bEnableDrawing) //check if we have any data to render
 		{
-			if ((RenderingComponent.IsValid() && !RenderingComponent->IsPendingKill()) || TaskDeletegate.IsBound())
-			{
-				FGraphEventRef CompleteHandle = FSimpleDelegateGraphTask::CreateAndDispatchWhenReady(
-					TaskDeletegate.IsBound() ? TaskDeletegate : FSimpleDelegateGraphTask::FDelegate::CreateUObject(RenderingComponent.Get(), &UNavMeshRenderingComponent::GatherDataForProxy)
-					, TEXT("Requesting navmesh data")
-					, NULL
-					, ENamedThreads::GameThread
-					);
-				bRequestedData = true;
-			}
 			return;
 		}
 
-		if (ProxyData.IsValid() == false || !ProxyData->bEnableDrawing)
-		{
-			return;
-		}
-
-		// scope to lock access to rendering data (it can be changed on game thread)
-		{
-			FScopeLock ScopeLock(&ProxyData->CriticalSection);
-			ProxyData->bSkipDistanceCheck = GIsEditor && (GEngine->GetDebugLocalPlayer() == NULL);
+		ProxyData->bSkipDistanceCheck = GIsEditor && (GEngine->GetDebugLocalPlayer() == NULL);
 
-			FVector const PosX(1.f,0,0);
-			FVector const PosY(0,1.f,0);
-			FVector const PosZ(0,0,1.f);
+		FVector const PosX(1.f,0,0);
+		FVector const PosY(0,1.f,0);
+		FVector const PosZ(0,0,1.f);
 
-			const TArray<FVector>& MeshVerts = ProxyData->NavMeshGeometry.MeshVerts;
+		const TArray<FVector>& MeshVerts = ProxyData->NavMeshGeometry.MeshVerts;
 
-			// Draw Mesh
-			for (int32 Index = 0; Index < ProxyData->MeshBuilders.Num(); ++Index)
-			{
-				const FColoredMaterialRenderProxy *MeshColorInstance = new(FMemStack::Get()) FColoredMaterialRenderProxy(GEngine->DebugMeshMaterial->GetRenderProxy(false), ProxyData->MeshBuilders[Index].ClusterColor);						
-				FDynamicMeshBuilder	MeshBuilder;
-				MeshBuilder.AddVertices( ProxyData->MeshBuilders[Index].Vertices );
-				MeshBuilder.AddTriangles( ProxyData->MeshBuilders[Index].Indices );
-				MeshBuilder.Draw(PDI, FMatrix::Identity, MeshColorInstance, GetDepthPriorityGroup(View));
-			}
+		// Draw Mesh
+		for (int32 Index = 0; Index < ProxyData->MeshBuilders.Num(); ++Index)
+		{
+			const FColoredMaterialRenderProxy *MeshColorInstance = new(FMemStack::Get()) FColoredMaterialRenderProxy(GEngine->DebugMeshMaterial->GetRenderProxy(false), ProxyData->MeshBuilders[Index].ClusterColor);						
+			FDynamicMeshBuilder	MeshBuilder;
+			MeshBuilder.AddVertices( ProxyData->MeshBuilders[Index].Vertices );
+			MeshBuilder.AddTriangles( ProxyData->MeshBuilders[Index].Indices );
+			MeshBuilder.Draw(PDI, FMatrix::Identity, MeshColorInstance, GetDepthPriorityGroup(View));
+		}
 
-			FHitProxyId HitProxyId;
-			FBatchedElements BatchedElements;
-			int32 Num = ProxyData->NavMeshEdgeLines.Num();
-			for (int32 Index = 0; Index < Num; ++Index)
+		FHitProxyId HitProxyId;
+		FBatchedElements BatchedElements;
+		int32 Num = ProxyData->NavMeshEdgeLines.Num();
+		for (int32 Index = 0; Index < Num; ++Index)
+		{
+			const FDebugLine &Line = ProxyData->NavMeshEdgeLines[Index];
+			if( LineInView(Line.Start,Line.End,View,false) )
 			{
-				const FDebugLine &Line = ProxyData->NavMeshEdgeLines[Index];
-				if( LineInView(Line.Start,Line.End,View,false) )
+				if (LineInCorrectDistance(Line.Start,Line.End,View))
+				{
+					BatchedElements.AddLine(Line.Start, Line.End, Line.Color,HitProxyId,NavMeshEdges_LineThickness, 0, true);
+				}
+				else if (GIsEditor)
 				{
-					if (LineInCorrectDistance(Line.Start,Line.End,View))
-					{
-						BatchedElements.AddLine(Line.Start, Line.End, Line.Color,HitProxyId,NavMeshEdges_LineThickness, 0, true);
-					}
-					else if (GIsEditor)
-					{
-						BatchedElements.AddLine(Line.Start, Line.End, Line.Color,HitProxyId,DefaultEdges_LineThickness, 0, true);
-					}
+					BatchedElements.AddLine(Line.Start, Line.End, Line.Color,HitProxyId,DefaultEdges_LineThickness, 0, true);
 				}
 			}
+		}
 
-			Num = ProxyData->ClusterLinkLines.Num();
-			for (int32 Index = 0; Index < Num; ++Index)
+		Num = ProxyData->ClusterLinkLines.Num();
+		for (int32 Index = 0; Index < Num; ++Index)
+		{
+			const FDebugLine &Line = ProxyData->ClusterLinkLines[Index];
+			if( LineInView(Line.Start,Line.End,View,false) )
 			{
-				const FDebugLine &Line = ProxyData->ClusterLinkLines[Index];
-				if( LineInView(Line.Start,Line.End,View,false) )
+				if (LineInCorrectDistance(Line.Start,Line.End,View))
+				{
+					BatchedElements.AddLine(Line.Start, Line.End, Line.Color,HitProxyId,ClusterLinkLines_LineThickness, 0, true);
+				}
+				else if (GIsEditor)
 				{
-					if (LineInCorrectDistance(Line.Start,Line.End,View))
-					{
-						BatchedElements.AddLine(Line.Start, Line.End, Line.Color,HitProxyId,ClusterLinkLines_LineThickness, 0, true);
-					}
-					else if (GIsEditor)
-					{
-						BatchedElements.AddLine(Line.Start, Line.End, Line.Color, HitProxyId, DefaultEdges_LineThickness, 0, true);
-					}
+					BatchedElements.AddLine(Line.Start, Line.End, Line.Color, HitProxyId, DefaultEdges_LineThickness, 0, true);
 				}
 			}
+		}
 
-			Num = ProxyData->TileEdgeLines.Num();
-			for (int32 Index = 0; Index < Num; ++Index)
+		Num = ProxyData->TileEdgeLines.Num();
+		for (int32 Index = 0; Index < Num; ++Index)
+		{
+			const FDebugLine &Line = ProxyData->TileEdgeLines[Index];
+			if( LineInView(Line.Start,Line.End,View,false) )
 			{
-				const FDebugLine &Line = ProxyData->TileEdgeLines[Index];
-				if( LineInView(Line.Start,Line.End,View,false) )
+				if (LineInCorrectDistance(Line.Start,Line.End,View))
+				{
+					BatchedElements.AddLine(Line.Start, Line.End, Line.Color,HitProxyId,PolyEdges_LineThickness, 0, true);
+				}
+				else if (GIsEditor)
 				{
-					if (LineInCorrectDistance(Line.Start,Line.End,View))
-					{
-						BatchedElements.AddLine(Line.Start, Line.End, Line.Color,HitProxyId,PolyEdges_LineThickness, 0, true);
-					}
-					else if (GIsEditor)
-					{
-						BatchedElements.AddLine(Line.Start, Line.End, Line.Color, HitProxyId, DefaultEdges_LineThickness, 0, true);
-					}
+					BatchedElements.AddLine(Line.Start, Line.End, Line.Color, HitProxyId, DefaultEdges_LineThickness, 0, true);
 				}
 			}
+		}
 
-			Num = ProxyData->NavLinkLines.Num();
-			for (int32 Index = 0; Index < Num; ++Index)
+		Num = ProxyData->NavLinkLines.Num();
+		for (int32 Index = 0; Index < Num; ++Index)
+		{
+			const FDebugLine &Line = ProxyData->NavLinkLines[Index];
+			if( LineInView(Line.Start,Line.End,View,false) )
 			{
-				const FDebugLine &Line = ProxyData->NavLinkLines[Index];
-				if( LineInView(Line.Start,Line.End,View,false) )
+				if (LineInCorrectDistance(Line.Start,Line.End,View))
 				{
-					if (LineInCorrectDistance(Line.Start,Line.End,View))
-					{
-						BatchedElements.AddLine(Line.Start, Line.End, Line.Color,HitProxyId,LinkLines_LineThickness, 0, true);
-					}
-					else if (GIsEditor)
-					{
-						BatchedElements.AddLine(Line.Start, Line.End, Line.Color, HitProxyId, DefaultEdges_LineThickness, 0, true);
-					}
+					BatchedElements.AddLine(Line.Start, Line.End, Line.Color,HitProxyId,LinkLines_LineThickness, 0, true);
+				}
+				else if (GIsEditor)
+				{
+					BatchedElements.AddLine(Line.Start, Line.End, Line.Color, HitProxyId, DefaultEdges_LineThickness, 0, true);
 				}
 			}
+		}
 
-			const bool bNeedToSwitchVerticalAxis = IsES2Platform(GRHIShaderPlatform) && !IsPCPlatform(GRHIShaderPlatform);
-			const FTexture2DRHIRef DepthTexture;
+		const bool bNeedToSwitchVerticalAxis = IsES2Platform(GRHIShaderPlatform) && !IsPCPlatform(GRHIShaderPlatform);
+		const FTexture2DRHIRef DepthTexture;
 			
-			FRHICommandListImmediate& RHICmdList = FRHICommandListExecutor::GetImmediateCommandList();
-			BatchedElements.Draw(
-				RHICmdList,
-				bNeedToSwitchVerticalAxis,
-				View->ViewProjectionMatrix,
-				View->ViewRect.Width(),
-				View->ViewRect.Height(),
-				View->Family->EngineShowFlags.HitProxies,
-				1.0f,
-				View,
-				DepthTexture
-				);
-
-			ProxyData->BatchedElements.Draw(
-				RHICmdList,
-				bNeedToSwitchVerticalAxis,
-				View->ViewProjectionMatrix,
-				View->ViewRect.Width(),
-				View->ViewRect.Height(),
-				View->Family->EngineShowFlags.HitProxies,
-				1.0f,
-				View,
-				DepthTexture
-				);
-		}
+		FRHICommandListImmediate& RHICmdList = FRHICommandListExecutor::GetImmediateCommandList();
+		BatchedElements.Draw(
+			RHICmdList,
+			bNeedToSwitchVerticalAxis,
+			View->ViewProjectionMatrix,
+			View->ViewRect.Width(),
+			View->ViewRect.Height(),
+			View->Family->EngineShowFlags.HitProxies,
+			1.0f,
+			View,
+			DepthTexture
+			);
+
+		ProxyData->BatchedElements.Draw(
+			RHICmdList,
+			bNeedToSwitchVerticalAxis,
+			View->ViewProjectionMatrix,
+			View->ViewRect.Width(),
+			View->ViewRect.Height(),
+			View->Family->EngineShowFlags.HitProxies,
+			1.0f,
+			View,
+			DepthTexture
+			);
 	}
 
 	FORCEINLINE bool PointInView(const FVector& Location, const FSceneView* View)
