Header cleanup: InstancedFoliageActor.h, ActorChannel.h, NavigationTypes.h, SoundWave.h, VoiceChannel.h.

diff --git a/Engine/Source/Editor/UnrealEd/Private/EditorExporters.cpp b/Engine/Source/Editor/UnrealEd/Private/EditorExporters.cpp
--- a/Engine/Source/Editor/UnrealEd/Private/EditorExporters.cpp
+++ b/Engine/Source/Editor/UnrealEd/Private/EditorExporters.cpp
@@ -22,6 +22,7 @@
 #include "Foliage/FoliageType.h"
 #include "UnrealExporter.h"
 #include "EngineModule.h"
+#include "InstancedFoliage.h"
 
 DEFINE_LOG_CATEGORY_STATIC(LogEditorExporters, Log, All);
 
diff --git a/Engine/Source/Editor/UnrealEd/Private/EditorObject.cpp b/Engine/Source/Editor/UnrealEd/Private/EditorObject.cpp
--- a/Engine/Source/Editor/UnrealEd/Private/EditorObject.cpp
+++ b/Engine/Source/Editor/UnrealEd/Private/EditorObject.cpp
@@ -10,6 +10,7 @@
 
 #include "Foliage/InstancedFoliageActor.h"
 #include "Foliage/FoliageType.h"
+#include "InstancedFoliage.h"
 
 DEFINE_LOG_CATEGORY_STATIC(LogEditorObject, Log, All);
 
diff --git a/Engine/Source/Runtime/AIModule/Classes/Navigation/CrowdFollowingComponent.h b/Engine/Source/Runtime/AIModule/Classes/Navigation/CrowdFollowingComponent.h
--- a/Engine/Source/Runtime/AIModule/Classes/Navigation/CrowdFollowingComponent.h
+++ b/Engine/Source/Runtime/AIModule/Classes/Navigation/CrowdFollowingComponent.h
@@ -1,6 +1,7 @@
 // Copyright 1998-2014 Epic Games, Inc. All Rights Reserved.
 
 #pragma once
+#include "AI/Navigation/NavigationAvoidanceTypes.h"
 #include "Navigation/CrowdAgentInterface.h"
 #include "Navigation/PathFollowingComponent.h"
 #include "CrowdFollowingComponent.generated.h"
diff --git a/Engine/Source/Runtime/Engine/Classes/AI/Navigation/NavLinkCustomComponent.h b/Engine/Source/Runtime/Engine/Classes/AI/Navigation/NavLinkCustomComponent.h
--- a/Engine/Source/Runtime/Engine/Classes/AI/Navigation/NavLinkCustomComponent.h
+++ b/Engine/Source/Runtime/Engine/Classes/AI/Navigation/NavLinkCustomComponent.h
@@ -1,8 +1,9 @@
 // Copyright 1998-2014 Epic Games, Inc. All Rights Reserved.
 
 #pragma once
-#include "AI/Navigation/NavRelevantComponent.h"
 #include "AI/Navigation/NavLinkCustomInterface.h"
+#include "AI/Navigation/NavLinkDefinition.h"
+#include "AI/Navigation/NavRelevantComponent.h"
 #include "NavLinkCustomComponent.generated.h"
 
 /**
diff --git a/Engine/Source/Runtime/Engine/Classes/AI/Navigation/NavLinkCustomInterface.h b/Engine/Source/Runtime/Engine/Classes/AI/Navigation/NavLinkCustomInterface.h
--- a/Engine/Source/Runtime/Engine/Classes/AI/Navigation/NavLinkCustomInterface.h
+++ b/Engine/Source/Runtime/Engine/Classes/AI/Navigation/NavLinkCustomInterface.h
@@ -1,8 +1,12 @@
 // Copyright 1998-2014 Epic Games, Inc. All Rights Reserved.
 
 #pragma once
+#include "AI/Navigation/NavLinkDefinition.h"
+#include "Interface.h"
 #include "NavLinkCustomInterface.generated.h"
 
+struct FNavigationLink;
+
 /** 
  *  Interface for custom navigation links
  *
diff --git a/Engine/Source/Runtime/Engine/Classes/AI/Navigation/NavLinkProxy.h b/Engine/Source/Runtime/Engine/Classes/AI/Navigation/NavLinkProxy.h
--- a/Engine/Source/Runtime/Engine/Classes/AI/Navigation/NavLinkProxy.h
+++ b/Engine/Source/Runtime/Engine/Classes/AI/Navigation/NavLinkProxy.h
@@ -1,6 +1,7 @@
 // Copyright 1998-2014 Epic Games, Inc. All Rights Reserved.
 
 #pragma once
+#include "AI/Navigation/NavLinkDefinition.h"
 #include "AI/Navigation/NavLinkHostInterface.h"
 #include "AI/Navigation/NavRelevantActorInterface.h"
 #include "GameFramework/Actor.h"
diff --git a/Engine/Source/Runtime/Engine/Classes/AI/Navigation/NavigationTypes.h b/Engine/Source/Runtime/Engine/Classes/AI/Navigation/NavigationTypes.h
--- a/Engine/Source/Runtime/Engine/Classes/AI/Navigation/NavigationTypes.h
+++ b/Engine/Source/Runtime/Engine/Classes/AI/Navigation/NavigationTypes.h
@@ -1,10 +1,6 @@
 // Copyright 1998-2014 Epic Games, Inc. All Rights Reserved.
 
 #pragma once
-
-#include "NavLinkDefinition.h"
-#include "NavigationModifier.h"
-#include "NavigationAvoidanceTypes.h"
 #include "NavigationTypes.generated.h"
 
 #define INVALID_NAVNODEREF (0)
diff --git a/Engine/Source/Runtime/Engine/Classes/AI/Navigation/RecastNavMesh.h b/Engine/Source/Runtime/Engine/Classes/AI/Navigation/RecastNavMesh.h
--- a/Engine/Source/Runtime/Engine/Classes/AI/Navigation/RecastNavMesh.h
+++ b/Engine/Source/Runtime/Engine/Classes/AI/Navigation/RecastNavMesh.h
@@ -37,6 +37,7 @@ class FRecastQueryFilter;
 class INavLinkCustomInterface;
 class UNavArea;
 class UNavigationSystem;
+struct FAreaNavModifier;
 
 UENUM()
 namespace ERecastPartitioning
diff --git a/Engine/Source/Runtime/Engine/Classes/Engine/ActorChannel.h b/Engine/Source/Runtime/Engine/Classes/Engine/ActorChannel.h
--- a/Engine/Source/Runtime/Engine/Classes/Engine/ActorChannel.h
+++ b/Engine/Source/Runtime/Engine/Classes/Engine/ActorChannel.h
@@ -33,8 +33,6 @@
  */
 
 #pragma once
-#include "Net/DataBunch.h"
-#include "Net/DataReplication.h"
 #include "Engine/Channel.h"
 #include "ActorChannel.generated.h"
 
diff --git a/Engine/Source/Runtime/Engine/Classes/Engine/VoiceChannel.h b/Engine/Source/Runtime/Engine/Classes/Engine/VoiceChannel.h
--- a/Engine/Source/Runtime/Engine/Classes/Engine/VoiceChannel.h
+++ b/Engine/Source/Runtime/Engine/Classes/Engine/VoiceChannel.h
@@ -5,9 +5,8 @@
 //
 
 #pragma once
-#include "Net/DataBunch.h"
-#include "Net/VoiceDataCommon.h"
 #include "Engine/Channel.h"
+#include "Net/VoiceDataCommon.h"
 #include "VoiceChannel.generated.h"
 
 UCLASS(transient, customConstructor)
diff --git a/Engine/Source/Runtime/Engine/Classes/Foliage/InstancedFoliageActor.h b/Engine/Source/Runtime/Engine/Classes/Foliage/InstancedFoliageActor.h
--- a/Engine/Source/Runtime/Engine/Classes/Foliage/InstancedFoliageActor.h
+++ b/Engine/Source/Runtime/Engine/Classes/Foliage/InstancedFoliageActor.h
@@ -2,15 +2,15 @@
 
 
 #pragma once
-
-#include "UniqueObj.h"
-#include "InstancedFoliage.h"
+#include "GameFramework/Actor.h"
+#include "Templates/UniqueObj.h"
 #include "InstancedFoliageActor.generated.h"
 
 // Forward declarations
 class UFoliageType;
 class ULandscapeHeightfieldCollisionComponent;
 struct FFoliageInstancePlacementInfo;
+struct FFoliageMeshInfo;
 
 UCLASS(notplaceable, hidecategories = Object, MinimalAPI, NotBlueprintable)
 class AInstancedFoliageActor : public AActor
diff --git a/Engine/Source/Runtime/Engine/Classes/Sound/SoundWave.h b/Engine/Source/Runtime/Engine/Classes/Sound/SoundWave.h
--- a/Engine/Source/Runtime/Engine/Classes/Sound/SoundWave.h
+++ b/Engine/Source/Runtime/Engine/Classes/Sound/SoundWave.h
@@ -5,11 +5,8 @@
 /** 
  * Playable sound object for raw wave files
  */
-
 #include "Sound/SoundBase.h"
 #include "Sound/SoundGroups.h"
-#include "SoundGroups.h"
-
 #include "SoundWave.generated.h"
 
 struct FActiveSound;
diff --git a/Engine/Source/Runtime/Engine/Private/AI/Navigation/RecastNavMeshGenerator.h b/Engine/Source/Runtime/Engine/Private/AI/Navigation/RecastNavMeshGenerator.h
--- a/Engine/Source/Runtime/Engine/Private/AI/Navigation/RecastNavMeshGenerator.h
+++ b/Engine/Source/Runtime/Engine/Private/AI/Navigation/RecastNavMeshGenerator.h
@@ -8,6 +8,7 @@
 #include "DetourNavMesh.h"
 #include "DetourNavMeshBuilder.h"
 #include "AI/NavDataGenerator.h"
+#include "AI/NavigationModifier.h"
 
 #define MAX_VERTS_PER_POLY	6
 
diff --git a/Engine/Source/Runtime/Engine/Private/DataChannel.cpp b/Engine/Source/Runtime/Engine/Private/DataChannel.cpp
--- a/Engine/Source/Runtime/Engine/Private/DataChannel.cpp
+++ b/Engine/Source/Runtime/Engine/Private/DataChannel.cpp
@@ -5,10 +5,10 @@
 =============================================================================*/
 
 #include "EnginePrivate.h"
-#include "Net/UnrealNetwork.h"
-
-#include "Net/NetworkProfiler.h"
 #include "Net/DataChannel.h"
+#include "Net/DataReplication.h"
+#include "Net/NetworkProfiler.h"
+#include "Net/UnrealNetwork.h"
 #include "Engine/ActorChannel.h"
 #include "Engine/ControlChannel.h"
 
diff --git a/Engine/Source/Runtime/Engine/Private/NetworkDriver.cpp b/Engine/Source/Runtime/Engine/Private/NetworkDriver.cpp
--- a/Engine/Source/Runtime/Engine/Private/NetworkDriver.cpp
+++ b/Engine/Source/Runtime/Engine/Private/NetworkDriver.cpp
@@ -5,6 +5,7 @@
 =============================================================================*/
 
 #include "EnginePrivate.h"
+#include "Net/DataReplication.h"
 #include "Net/UnrealNetwork.h"
 #include "Net/NetworkProfiler.h"
 #include "Engine/ActorChannel.h"
diff --git a/Engine/Source/Runtime/Engine/Public/AI/NavigationOctree.h b/Engine/Source/Runtime/Engine/Public/AI/NavigationOctree.h
--- a/Engine/Source/Runtime/Engine/Public/AI/NavigationOctree.h
+++ b/Engine/Source/Runtime/Engine/Public/AI/NavigationOctree.h
@@ -1,7 +1,7 @@
 // Copyright 1998-2014 Epic Games, Inc. All Rights Reserved.
 
 #pragma once
-
+#include "AI/NavigationModifier.h"
 #include "GenericOctree.h"
 #include "AI/Navigation/NavigationProxy.h"
 
