added ClearActionBindings() method; documentation & code cleanup pass

diff --git a/Engine/Source/Runtime/Engine/Classes/Components/InputComponent.h b/Engine/Source/Runtime/Engine/Classes/Components/InputComponent.h
--- a/Engine/Source/Runtime/Engine/Classes/Components/InputComponent.h
+++ b/Engine/Source/Runtime/Engine/Classes/Components/InputComponent.h
@@ -6,6 +6,7 @@
 #include "Runtime/InputCore/Classes/InputCoreTypes.h"
 #include "InputComponent.generated.h"
 
+
 // Utility delegate class to allow binding to either a C++ function or a blueprint script delegate
 template<class DelegateType, class DynamicDelegateType>
 struct TInputUnifiedDelegate
@@ -70,6 +71,7 @@ protected:
 	DynamicDelegateType FuncDynDelegate;
 };
 
+
 // An Input Chord is a key and the modifier keys that are to be held with it
 USTRUCT()
 struct FInputChord
@@ -108,8 +110,7 @@ struct FInputChord
 		, bCtrl(false)
 		, bAlt(false)
 		, bCmd(false)
-	{
-	}
+	{ }
 
 	FInputChord(const FKey InKey, const bool bInShift, const bool bInCtrl, const bool bInAlt, const bool bInCmd)
 		: Key(InKey)
@@ -117,8 +118,7 @@ struct FInputChord
 		, bCtrl(bInCtrl)
 		, bAlt(bInAlt)
 		, bCmd(bInCmd)
-	{
-	}
+	{ }
 
 	bool operator==(const FInputChord& Other) const
 	{
@@ -130,10 +130,12 @@ struct FInputChord
 	}
 };
 
+
 /** Delegate signature for action handlers. */
 DECLARE_DELEGATE( FInputActionHandlerSignature );
 DECLARE_DYNAMIC_DELEGATE( FInputActionHandlerDynamicSignature );
 
+
 struct FInputActionUnifiedDelegate : public TInputUnifiedDelegate<FInputActionHandlerSignature, FInputActionHandlerDynamicSignature>
 {
 	inline void Execute() const
@@ -149,6 +151,7 @@ struct FInputActionUnifiedDelegate : public TInputUnifiedDelegate<FInputActionHa
 	}
 };
 
+
 struct FInputBinding
 {
 	/** Whether the binding should consume the input or allow it to pass to another component */
@@ -163,6 +166,7 @@ struct FInputBinding
 	{}
 };
 
+
 /** Binds a game-defined action. */
 struct FInputActionBinding : public FInputBinding
 {
@@ -182,16 +186,17 @@ struct FInputActionBinding : public FInputBinding
 		, ActionName(NAME_None)
 		, KeyEvent(EInputEvent::IE_Pressed)
 		, bPaired(false)
-	{}
+	{ }
 
 	FInputActionBinding(const FName InActionName, const enum EInputEvent InKeyEvent)
 		: FInputBinding()
 		, ActionName(InActionName)
 		, KeyEvent(InKeyEvent)
 		, bPaired(false)
-	{}
+	{ }
 };
 
+
 struct FInputKeyBinding : public FInputBinding
 {
 	/** Raw key to bind to */
@@ -205,15 +210,16 @@ struct FInputKeyBinding : public FInputBinding
 	FInputKeyBinding()
 		: FInputBinding()
 		, KeyEvent(EInputEvent::IE_Pressed)
-	{}
+	{ }
 
 	FInputKeyBinding(const FInputChord InChord, const enum EInputEvent InKeyEvent)
 		: FInputBinding()
 		, Chord(InChord)
 		, KeyEvent(InKeyEvent)
-	{}
+	{ }
 };
 
+
 /** 
  * Delegate signature for touch handlers. 
  * @FigerIndex: Which finger touched
@@ -222,6 +228,7 @@ struct FInputKeyBinding : public FInputBinding
 DECLARE_DELEGATE_TwoParams( FInputTouchHandlerSignature, ETouchIndex::Type, FVector );
 DECLARE_DYNAMIC_DELEGATE_TwoParams( FInputTouchHandlerDynamicSignature, ETouchIndex::Type, FingerIndex, FVector, Location );
 
+
 struct FInputTouchUnifiedDelegate : public TInputUnifiedDelegate<FInputTouchHandlerSignature, FInputTouchHandlerDynamicSignature>
 {
 	inline void Execute(const ETouchIndex::Type FingerIndex, const FVector Location) const
@@ -248,14 +255,15 @@ struct FInputTouchBinding : public FInputBinding
 	FInputTouchBinding()
 		: FInputBinding()
 		, KeyEvent(EInputEvent::IE_Pressed)
-	{}
+	{ }
 
 	FInputTouchBinding(const enum EInputEvent InKeyEvent)
 		: FInputBinding()
 		, KeyEvent(InKeyEvent)
-	{}
+	{ }
 };
 
+
 /** 
  * Delegate signature for axis handlers. 
  * @AxisValue: "Value" to pass to the axis.  This value will be the device-dependent, so a mouse will report absolute change since the last update, 
@@ -265,6 +273,7 @@ struct FInputTouchBinding : public FInputBinding
 DECLARE_DELEGATE_OneParam( FInputAxisHandlerSignature, float );
 DECLARE_DYNAMIC_DELEGATE_OneParam( FInputAxisHandlerDynamicSignature, float, AxisValue );
 
+
 struct FInputAxisUnifiedDelegate : public TInputUnifiedDelegate<FInputAxisHandlerSignature, FInputAxisHandlerDynamicSignature>
 {
 	inline void Execute(const float AxisValue) const
@@ -294,15 +303,16 @@ struct FInputAxisBinding : public FInputBinding
 		: FInputBinding()
 		, AxisName(NAME_None)
 		, AxisValue(0.f)
-	{}
+	{ }
 
 	FInputAxisBinding(const FName InAxisName)
 		: FInputBinding()
 		, AxisName(InAxisName)
 		, AxisValue(0.f)
-	{}
+	{ }
 };
 
+
 /** Binds a raw axis to a function */
 struct FInputAxisKeyBinding : public FInputBinding
 {
@@ -315,7 +325,7 @@ struct FInputAxisKeyBinding : public FInputBinding
 	FInputAxisKeyBinding()
 		: FInputBinding()
 		, AxisValue(0.f)
-	{}
+	{ }
 
 	FInputAxisKeyBinding(const FKey InAxisKey)
 		: FInputBinding()
@@ -326,6 +336,7 @@ struct FInputAxisKeyBinding : public FInputBinding
 	}
 };
 
+
 /**
 * Delegate signature for vector axis handlers.
 * @AxisValue: "Value" to pass to the axis.
@@ -333,6 +344,7 @@ struct FInputAxisKeyBinding : public FInputBinding
 DECLARE_DELEGATE_OneParam(FInputVectorAxisHandlerSignature, FVector);
 DECLARE_DYNAMIC_DELEGATE_OneParam(FInputVectorAxisHandlerDynamicSignature, FVector, AxisValue);
 
+
 struct FInputVectorAxisUnifiedDelegate : public TInputUnifiedDelegate<FInputVectorAxisHandlerSignature, FInputVectorAxisHandlerDynamicSignature>
 {
 	inline void Execute(const FVector AxisValue) const
@@ -348,6 +360,7 @@ struct FInputVectorAxisUnifiedDelegate : public TInputUnifiedDelegate<FInputVect
 	}
 };
 
+
 struct FInputVectorAxisBinding : public FInputBinding
 {
 	FKey AxisKey;
@@ -358,7 +371,7 @@ struct FInputVectorAxisBinding : public FInputBinding
 
 	FInputVectorAxisBinding()
 		: FInputBinding()
-	{}
+	{ }
 
 	FInputVectorAxisBinding(const FKey InAxisKey)
 		: FInputBinding()
@@ -368,6 +381,7 @@ struct FInputVectorAxisBinding : public FInputBinding
 	}
 };
 
+
 /** 
  * Delegate signature for gesture handlers. 
  * @Value: "Value" to pass to the axis.  Note that by convention this is assumed to be a framerate-independent "delta" value, i.e. absolute change for this frame
@@ -376,6 +390,7 @@ struct FInputVectorAxisBinding : public FInputBinding
 DECLARE_DELEGATE_OneParam( FInputGestureHandlerSignature, float );
 DECLARE_DYNAMIC_DELEGATE_OneParam( FInputGestureHandlerDynamicSignature, float, Value );
 
+
 struct FInputGestureUnifiedDelegate : public TInputUnifiedDelegate<FInputGestureHandlerSignature, FInputGestureHandlerDynamicSignature>
 {
 	inline void Execute(const float Value) const
@@ -391,6 +406,7 @@ struct FInputGestureUnifiedDelegate : public TInputUnifiedDelegate<FInputGesture
 	}
 };
 
+
 /** Binds a gesture to a function. */
 struct FInputGestureBinding : public FInputBinding
 {
@@ -404,15 +420,16 @@ struct FInputGestureBinding : public FInputBinding
 	FInputGestureBinding()
 		: FInputBinding()
 		, GestureValue(0.f)
-	{}
+	{ }
 
 	FInputGestureBinding(const FKey InGestureKey)
 		: FInputBinding()
 		, GestureKey(InGestureKey)
 		, GestureValue(0.f)
-	{}
+	{ }
 };
 
+
 // -----------------------------------------------------------------
 // Helper macros to make setting up axis bindings more readable in native code. 
 //
@@ -449,49 +466,135 @@ namespace EControllerAnalogStick
 	};
 }
 
-// Input Component is a component that can be added to an Actor that allows the Actor to bind various kinds of input events
-// to delegate functions.  Input components are processed from a stack managed by the PlayerController and each bind can 
-// consume the input event preventing other components on the input stack from processing the input.
+
+/**
+ * Implement an Actor component for input bindings.
+ *
+ * Input Component is a component that can be added to an Actor that allows the Actor to bind various kinds of input events
+ * to delegate functions.  Input components are processed from a stack managed by the PlayerController and each bind can 
+ * consume the input event preventing other components on the input stack from processing the input.
+ */
 UCLASS(transient, config=Input, hidecategories=(Activation, "Components|Activation"))
-class ENGINE_API UInputComponent : public UActorComponent
+class ENGINE_API UInputComponent
+	: public UActorComponent
 {
 	GENERATED_UCLASS_BODY()
 
-private:
-	TArray<struct FInputActionBinding> ActionBindings;
-
-public:
+	/** The collection of key bindings. */
 	TArray<struct FInputKeyBinding> KeyBindings;
+
+	/** The collection of touch bindings. */
 	TArray<struct FInputTouchBinding> TouchBindings;
+
+	/** The collection of axis bindings. */
 	TArray<struct FInputAxisBinding> AxisBindings;
+
+	/** The collection of axis key bindings. */
 	TArray<struct FInputAxisKeyBinding> AxisKeyBindings;
+
+	/** The collection of vector axis bindings. */
 	TArray<struct FInputVectorAxisBinding> VectorAxisBindings;
+
+	/** The collection of gesture bindings. */
 	TArray<struct FInputGestureBinding> GestureBindings;
 
-	// Whether any components lower on the input stack should be allowed to receive input
+	/** Whether any components lower on the input stack should be allowed to receive input. */
 	uint32 bBlockInput:1;
 
-	// Returns the current value of an Axis defined in the project settings bound in this input component
-	float GetAxisValue(const FName AxisName) const;
+public:
 
-	// Returns the 
-	float GetAxisKeyValue(const FKey AxisKey) const;
-	FVector GetVectorAxisValue(const FKey AxisKey) const;
+	/**
+	 * Gets the current value of the axis with the specified name.
+	 *
+	 * @param AxisName The name of the axis.
+	 * @return Axis value.
+	 * @see GetAxisKeyValue, GetVectorAxisValue
+	 */
+	float GetAxisValue( const FName AxisName ) const;
+
+	/**
+	 * Gets the current value of the axis with the specified key.
+	 *
+	 * @param AxisKey The key of the axis.
+	 * @return Axis value.
+	 * @see GetAxisKeyValue, GetVectorAxisValue
+	 */
+	float GetAxisKeyValue( const FKey AxisKey ) const;
+
+	/**
+	 * Gets the current vector value of the axis with the specified key.
+	 *
+	 * @param AxisKey The key of the axis.
+	 * @return Axis value.
+	 * @see GetAxisValue, GetAxisKeyValue
+	 */
+	FVector GetVectorAxisValue( const FKey AxisKey ) const;
+
+	/**
+	 * Checks whether this component has any input bindings.
+	 *
+	 * @return true if any bindings are set, false otherwise.
+	 */
+	bool HasBindings( ) const;
+
+	/**
+	 * Adds the specified action binding.
+	 *
+	 * @param Binding The binding to add.
+	 * @return The last binding in the list.
+	 * @see ClearActionBindings, GetActionBinding, GetNumActionBindings, RemoveActionBinding
+	 */
+	FInputActionBinding& AddActionBinding( const FInputActionBinding& Binding );
+
+	/**
+	 * Removes all action bindings.
+	 *
+	 * @see AddActionBinding, GetActionBinding, GetNumActionBindings, RemoveActionBinding
+	 */
+	void ClearActionBindings( );
+
+	/**
+	 * Gets the action binding with the specified index.
+	 *
+	 * @param BindingIndex The index of the binding to get.
+	 * @see AddActionBinding, ClearActionBindings, GetNumActionBindings, RemoveActionBinding
+	 */
+	FInputActionBinding& GetActionBinding( const int32 BindingIndex );
+
+	/**
+	 * Gets the number of action bindings.
+	 *
+	 * @return Number of bindings.
+	 * @see AddActionBinding, ClearActionBindings, GetActionBinding, RemoveActionBinding
+	 */
+	int32 GetNumActionBindings( ) const;
+
+	/**
+	 * Removes the action binding at the specified index.
+	 *
+	 * @param BindingIndex The index of the binding to remove.
+	 * @see AddActionBinding, ClearActionBindings, GetActionBinding, GetNumActionBindings
+	 */
+	void RemoveActionBinding( const int32 BindingIndex );
 
-	bool HasBindings() const;
+public:
 
-	// Binds a delegate function to an Action defined in the project settings
+	/**
+	 * Binds a delegate function to an Action defined in the project settings.
+	 */
 	template<class UserClass>
-	FInputActionBinding& BindAction(const FName ActionName, const EInputEvent KeyEvent, UserClass* Object, typename FInputActionHandlerSignature::TUObjectMethodDelegate< UserClass >::FMethodPtr Func)
+	FInputActionBinding& BindAction( const FName ActionName, const EInputEvent KeyEvent, UserClass* Object, typename FInputActionHandlerSignature::TUObjectMethodDelegate< UserClass >::FMethodPtr Func )
 	{
 		FInputActionBinding AB( ActionName, KeyEvent );
 		AB.ActionDelegate.BindDelegate(Object, Func);
 		return AddActionBinding(AB);
 	}
 
-	// Binds a delegate function an Axis defined in the project settings
+	/**
+	 * Binds a delegate function an Axis defined in the project settings.
+	 */
 	template<class UserClass>
-	FInputAxisBinding& BindAxis(const FName AxisName, UserClass* Object, typename FInputAxisHandlerSignature::TUObjectMethodDelegate< UserClass >::FMethodPtr Func)
+	FInputAxisBinding& BindAxis( const FName AxisName, UserClass* Object, typename FInputAxisHandlerSignature::TUObjectMethodDelegate< UserClass >::FMethodPtr Func )
 	{
 		FInputAxisBinding AB( AxisName );
 		AB.AxisDelegate.BindDelegate(Object, Func);
@@ -499,17 +602,22 @@ public:
 		return AxisBindings.Last();
 	}
 
-	// Indicates that the InputComponent is interested in knowing the Axis value (via GetAxisValue) but does not want a delegate function called each frame
-	FInputAxisBinding& BindAxis(const FName AxisName)
+	/**
+	 * Indicates that the InputComponent is interested in knowing the Axis value
+	 * (via GetAxisValue) but does not want a delegate function called each frame.
+	 */
+	FInputAxisBinding& BindAxis( const FName AxisName )
 	{
 		FInputAxisBinding AB( AxisName );
 		AxisBindings.Add(AB);
 		return AxisBindings.Last();
 	}
 
-	// Binds a delegate function for an axis key (e.g. Mouse X)
+	/**
+	 * Binds a delegate function for an axis key (e.g. Mouse X).
+	 */
 	template<class UserClass>
-	FInputAxisKeyBinding& BindAxisKey(const FKey AxisKey, UserClass* Object, typename FInputAxisHandlerSignature::TUObjectMethodDelegate< UserClass >::FMethodPtr Func)
+	FInputAxisKeyBinding& BindAxisKey( const FKey AxisKey, UserClass* Object, typename FInputAxisHandlerSignature::TUObjectMethodDelegate< UserClass >::FMethodPtr Func )
 	{
 		FInputAxisKeyBinding AB(AxisKey);
 		AB.AxisDelegate.BindDelegate(Object, Func);
@@ -517,8 +625,11 @@ public:
 		return AxisKeyBindings.Last();
 	}
 
-	// Indicates that the InputComponent is interested in knowing/consuming an axis key's value (via GetAxisKeyValue) but does not want a delegate function called each frame
-	FInputAxisKeyBinding& BindAxisKey(const FKey AxisKey)
+	/**
+	 * Indicates that the InputComponent is interested in knowing/consuming an axis key's
+	 * value (via GetAxisKeyValue) but does not want a delegate function called each frame.
+	 */
+	FInputAxisKeyBinding& BindAxisKey( const FKey AxisKey )
 	{
 		FInputAxisKeyBinding AB(AxisKey);
 		AxisKeyBindings.Add(AB);
@@ -527,7 +638,7 @@ public:
 
 	// Binds a delegate function to a vector axis key (e.g. Tilt)
 	template<class UserClass>
-	FInputVectorAxisBinding& BindVectorAxis(const FKey AxisKey, UserClass* Object, typename FInputVectorAxisHandlerSignature::TUObjectMethodDelegate< UserClass >::FMethodPtr Func)
+	FInputVectorAxisBinding& BindVectorAxis( const FKey AxisKey, UserClass* Object, typename FInputVectorAxisHandlerSignature::TUObjectMethodDelegate< UserClass >::FMethodPtr Func )
 	{
 		FInputVectorAxisBinding AB(AxisKey);
 		AB.AxisDelegate.BindDelegate(Object, Func);
@@ -535,17 +646,22 @@ public:
 		return VectorAxisBindings.Last();
 	}
 
-	// Indicates that the InputComponent is interested in knowing/consuming a vector axis key's value (via GetVectorAxisKeyValue) but does not want a delegate function called each frame
-	FInputVectorAxisBinding& BindVectorAxis(const FKey AxisKey)
+	/**
+	 * Indicates that the InputComponent is interested in knowing/consuming a vector axis key's
+	 * value (via GetVectorAxisKeyValue) but does not want a delegate function called each frame.
+	 */
+	FInputVectorAxisBinding& BindVectorAxis( const FKey AxisKey )
 	{
 		FInputVectorAxisBinding AB(AxisKey);
 		VectorAxisBindings.Add(AB);
 		return VectorAxisBindings.Last();
 	}
 
-	// Binds a chord event to a delegate function
+	/**
+	 * Binds a chord event to a delegate function.
+	 */
 	template<class UserClass>
-	FInputKeyBinding& BindKey(const FInputChord Chord, const EInputEvent KeyEvent, UserClass* Object, typename FInputActionHandlerSignature::TUObjectMethodDelegate< UserClass >::FMethodPtr Func)
+	FInputKeyBinding& BindKey( const FInputChord Chord, const EInputEvent KeyEvent, UserClass* Object, typename FInputActionHandlerSignature::TUObjectMethodDelegate< UserClass >::FMethodPtr Func )
 	{
 		FInputKeyBinding KB(Chord, KeyEvent);
 		KB.KeyDelegate.BindDelegate(Object, Func);
@@ -553,16 +669,20 @@ public:
 		return KeyBindings.Last();
 	}
 
-	// Binds a key event to a delegate function
+	/**
+	 * Binds a key event to a delegate function.
+	 */
 	template<class UserClass>
-	FInputKeyBinding& BindKey(const FKey Key, const EInputEvent KeyEvent, UserClass* Object, typename FInputActionHandlerSignature::TUObjectMethodDelegate< UserClass >::FMethodPtr Func)
+	FInputKeyBinding& BindKey( const FKey Key, const EInputEvent KeyEvent, UserClass* Object, typename FInputActionHandlerSignature::TUObjectMethodDelegate< UserClass >::FMethodPtr Func )
 	{
 		return BindKey(FInputChord(Key, false, false, false, false), KeyEvent, Object, Func);
 	}
 
-	// Binds this input component to touch events
+	/**
+	 * Binds this input component to touch events.
+	 */
 	template<class UserClass>
-	FInputTouchBinding& BindTouch(const EInputEvent KeyEvent, UserClass* Object, typename FInputTouchHandlerSignature::TUObjectMethodDelegate< UserClass >::FMethodPtr Func)
+	FInputTouchBinding& BindTouch( const EInputEvent KeyEvent, UserClass* Object, typename FInputTouchHandlerSignature::TUObjectMethodDelegate< UserClass >::FMethodPtr Func )
 	{
 		FInputTouchBinding TB(KeyEvent);
 		TB.TouchDelegate.BindDelegate(Object, Func);
@@ -570,9 +690,11 @@ public:
 		return TouchBindings.Last();
 	}
 
-	// Binds a gesture event to a delegate function
+	/**
+	 * Binds a gesture event to a delegate function.
+	 */
 	template<class UserClass>
-	FInputGestureBinding& BindGesture(const FKey GestureKey, UserClass* Object, typename FInputGestureHandlerSignature::TUObjectMethodDelegate< UserClass >::FMethodPtr Func)
+	FInputGestureBinding& BindGesture( const FKey GestureKey, UserClass* Object, typename FInputGestureHandlerSignature::TUObjectMethodDelegate< UserClass >::FMethodPtr Func )
 	{
 		FInputGestureBinding GB(GestureKey);
 		GB.GestureDelegate.BindDelegate(Object, Func);
@@ -580,13 +702,8 @@ public:
 		return GestureBindings.Last();
 	}
 
-	// Utility functions for managing action bindings array as it needs to maintain knowledge of the pairing state
-	FInputActionBinding& AddActionBinding(const FInputActionBinding& Binding);
-	void RemoveActionBinding(const int32 BindingIndex);
-	int32 GetNumActionBindings() const;
-	FInputActionBinding& GetActionBinding(const int32 BindingIndex);
-
 private:
+
 	/** Returns true if the given key/button is pressed on the input of the controller (if present) */
 	UFUNCTION(BlueprintCallable, meta=(DeprecatedFunction, DeprecationMessage="Use PlayerController.IsInputKeyDown instead."))
 	bool IsControllerKeyDown(FKey Key) const;
@@ -622,4 +739,9 @@ private:
 	/** Retrieves the X and Y displacement of the given analog stick.  For WhickStick, 0 = left, 1 = right. */
 	UFUNCTION(BlueprintCallable, meta=(DeprecatedFunction, DeprecationMessage="Use PlayerController.GetInputAnalogStickState instead."))
 	void GetControllerAnalogStickState(EControllerAnalogStick::Type WhichStick, float& StickX, float& StickY) const;
+
+private:
+
+	/** Holds the collection of action bindings. */
+	TArray<struct FInputActionBinding> ActionBindings;
 };
diff --git a/Engine/Source/Runtime/Engine/Private/InputComponent.cpp b/Engine/Source/Runtime/Engine/Private/InputComponent.cpp
--- a/Engine/Source/Runtime/Engine/Private/InputComponent.cpp
+++ b/Engine/Source/Runtime/Engine/Private/InputComponent.cpp
@@ -2,30 +2,34 @@
 
 #include "EnginePrivate.h"
 
-FInputChord::RelationshipType FInputChord::GetRelationship(const FInputChord& OtherChord) const
+
+/* FInputChord interface
+ *****************************************************************************/
+
+FInputChord::RelationshipType FInputChord::GetRelationship( const FInputChord& OtherChord ) const
 {
 	RelationshipType Relationship = None;
 
 	if (Key == OtherChord.Key)
 	{
-		if (    bAlt == OtherChord.bAlt
-			 && bCtrl == OtherChord.bCtrl
-			 && bShift == OtherChord.bShift
-			 && bCmd == OtherChord.bCmd)
+		if ((bAlt == OtherChord.bAlt) &&
+			(bCtrl == OtherChord.bCtrl) &&
+			(bShift == OtherChord.bShift) &&
+			(bCmd == OtherChord.bCmd))
 		{
 			Relationship = Same;
 		}
-		else if (    (bAlt || !OtherChord.bAlt)
-				  && (bCtrl || !OtherChord.bCtrl)
-				  && (bShift || !OtherChord.bShift)
-				  && (bCmd || !OtherChord.bCmd))
+		else if ((bAlt || !OtherChord.bAlt) &&
+				(bCtrl || !OtherChord.bCtrl) &&
+				(bShift || !OtherChord.bShift) &&
+				(bCmd || !OtherChord.bCmd))
 		{
 			Relationship = Masks;
 		}
-		else if (    (!bAlt || OtherChord.bAlt)
-				  && (!bCtrl || OtherChord.bCtrl)
-				  && (!bShift || OtherChord.bShift)
-				  && (!bCmd || OtherChord.bCmd))
+		else if ((!bAlt || OtherChord.bAlt) &&
+				(!bCtrl || OtherChord.bCtrl) &&
+				(!bShift || OtherChord.bShift) &&
+				(!bCmd || OtherChord.bCmd))
 		{
 			Relationship = Masked;
 		}
@@ -34,13 +38,18 @@ FInputChord::RelationshipType FInputChord::GetRelationship(const FInputChord& Ot
 	return Relationship;
 }
 
-UInputComponent::UInputComponent(const class FPostConstructInitializeProperties& PCIP)
+
+/* UInputComponent interface
+ *****************************************************************************/
+
+UInputComponent::UInputComponent( const class FPostConstructInitializeProperties& PCIP )
 	: Super(PCIP)
 {
 	bBlockInput = false;
 }
 
-float UInputComponent::GetAxisValue(const FName AxisName) const
+
+float UInputComponent::GetAxisValue( const FName AxisName ) const
 {
 	float AxisValue = 0.f;
 	bool bFound = false;
@@ -63,7 +72,8 @@ float UInputComponent::GetAxisValue(const FName AxisName) const
 	return AxisValue;
 }
 
-float UInputComponent::GetAxisKeyValue(const FKey AxisKey) const
+
+float UInputComponent::GetAxisKeyValue( const FKey AxisKey ) const
 {
 	float AxisValue = 0.f;
 	bool bFound = false;
@@ -86,7 +96,8 @@ float UInputComponent::GetAxisKeyValue(const FKey AxisKey) const
 	return AxisValue;
 }
 
-FVector UInputComponent::GetVectorAxisValue(const FKey AxisKey) const
+
+FVector UInputComponent::GetVectorAxisValue( const FKey AxisKey ) const
 {
 	FVector AxisValue;
 	bool bFound = false;
@@ -109,18 +120,20 @@ FVector UInputComponent::GetVectorAxisValue(const FKey AxisKey) const
 	return AxisValue;
 }
 
-bool UInputComponent::HasBindings() const
+
+bool UInputComponent::HasBindings( ) const
 {
-	return (   (ActionBindings.Num() > 0) 
-			|| (AxisBindings.Num() > 0) 
-			|| (AxisKeyBindings.Num() > 0) 
-			|| (KeyBindings.Num() > 0) 
-			|| (TouchBindings.Num() > 0) 
-			|| (GestureBindings.Num() > 0)
-			|| (VectorAxisBindings.Num() > 0));
+	return ((ActionBindings.Num() > 0) ||
+			(AxisBindings.Num() > 0) ||
+			(AxisKeyBindings.Num() > 0) ||
+			(KeyBindings.Num() > 0) ||
+			(TouchBindings.Num() > 0) ||
+			(GestureBindings.Num() > 0) ||
+			(VectorAxisBindings.Num() > 0));
 }
 
-FInputActionBinding& UInputComponent::AddActionBinding(const FInputActionBinding& Binding)
+
+FInputActionBinding& UInputComponent::AddActionBinding( const FInputActionBinding& Binding )
 {
 	ActionBindings.Add(Binding);
 	if (Binding.KeyEvent == IE_Pressed || Binding.KeyEvent == IE_Released)
@@ -151,7 +164,26 @@ FInputActionBinding& UInputComponent::AddActionBinding(const FInputActionBinding
 	return ActionBindings.Last();
 }
 
-void UInputComponent::RemoveActionBinding(const int32 BindingIndex)
+
+void UInputComponent::ClearActionBindings( )
+{
+	ActionBindings.Reset();
+}
+
+
+FInputActionBinding& UInputComponent::GetActionBinding( const int32 BindingIndex )
+{
+	return ActionBindings[BindingIndex];
+}
+
+
+int32 UInputComponent::GetNumActionBindings( ) const
+{
+	return ActionBindings.Num();
+}
+
+
+void UInputComponent::RemoveActionBinding( const int32 BindingIndex )
 {
 	if (BindingIndex >= 0 && BindingIndex < ActionBindings.Num())
 	{
@@ -184,6 +216,7 @@ void UInputComponent::RemoveActionBinding(const int32 BindingIndex)
 					}
 				}
 			}
+
 			for (int32 ClearIndex = 0; ClearIndex < IndicesToClear.Num(); ++ClearIndex)
 			{
 				ActionBindings[IndicesToClear[ClearIndex]].bPaired = false;
@@ -194,18 +227,10 @@ void UInputComponent::RemoveActionBinding(const int32 BindingIndex)
 	}
 }
 
-int32 UInputComponent::GetNumActionBindings() const
-{
-	return ActionBindings.Num();
-}
-
-FInputActionBinding& UInputComponent::GetActionBinding(const int32 BindingIndex)
-{
-	return ActionBindings[BindingIndex];
-}
 
+/* Deprecated functions (needed for Blueprints)
+ *****************************************************************************/
 
-// Deprecated UFUNCTIONS needed for blueprint purposes
 bool UInputComponent::IsControllerKeyDown(FKey Key) const { return false; }
 bool UInputComponent::WasControllerKeyJustPressed(FKey Key) const { return false; }
 bool UInputComponent::WasControllerKeyJustReleased(FKey Key) const { return false; }
