Fix compiling with distcc (| inside &).

diff --git a/Engine/Source/Runtime/Engine/Private/Tests/EngineAutomationTests.cpp b/Engine/Source/Runtime/Engine/Private/Tests/EngineAutomationTests.cpp
--- a/Engine/Source/Runtime/Engine/Private/Tests/EngineAutomationTests.cpp
+++ b/Engine/Source/Runtime/Engine/Private/Tests/EngineAutomationTests.cpp
@@ -307,7 +307,7 @@ bool FExecWorldStringLatentCommand::Update()
 /**
  * Automation test to load a map and capture FPS performance charts
  */
-IMPLEMENT_COMPLEX_AUTOMATION_TEST(FCinematicFPSPerfTest, "Engine.Cinematic FPS Perf Capture", EAutomationTestFlags::ATF_Game | EAutomationTestFlags::ATF_NonNullRHI);
+IMPLEMENT_COMPLEX_AUTOMATION_TEST(FCinematicFPSPerfTest, "Engine.Cinematic FPS Perf Capture", (EAutomationTestFlags::ATF_Game | EAutomationTestFlags::ATF_NonNullRHI));
 
 void FCinematicFPSPerfTest::GetTests(TArray<FString>& OutBeautifiedNames, TArray <FString>& OutTestCommands) const
 {
