Media: removed legacy movie texture support

diff --git a/Engine/Config/BaseEditor.ini b/Engine/Config/BaseEditor.ini
--- a/Engine/Config/BaseEditor.ini
+++ b/Engine/Config/BaseEditor.ini
@@ -19,7 +19,6 @@ MaxNumberOfLogLines=2000
 +RenderableThumbnailTypes=(ClassNeedingThumbnailName="/Script/Engine.Texture2D",RendererClassName="/Script/UnrealEd.TextureThumbnailRenderer")
 +RenderableThumbnailTypes=(ClassNeedingThumbnailName="/Script/Engine.TextureRenderTarget",RendererClassName="/Script/UnrealEd.TextureThumbnailRenderer")
 +RenderableThumbnailTypes=(ClassNeedingThumbnailName="/Script/Engine.TextureRenderTargetCube",RendererClassName="/Script/UnrealEd.TextureThumbnailRenderer")
-+RenderableThumbnailTypes=(ClassNeedingThumbnailName="/Script/Engine.TextureMovie",RendererClassName="/Script/UnrealEd.TextureThumbnailRenderer")
 +RenderableThumbnailTypes=(ClassNeedingThumbnailName="/Script/Engine.LightMapTexture2D",RendererClassName="/Script/UnrealEd.TextureThumbnailRenderer")
 +RenderableThumbnailTypes=(ClassNeedingThumbnailName="/Script/Engine.TextureCube",RendererClassName="/Script/UnrealEd.TextureCubeThumbnailRenderer")
 +RenderableThumbnailTypes=(ClassNeedingThumbnailName="/Script/Engine.Material",RendererClassName="/Script/UnrealEd.MaterialInstanceThumbnailRenderer")
@@ -283,4 +282,4 @@ bBlueprintIsNotBlueprintType=false
 
 [StructSerialization]
 SkipByteCodeSerialization=true
-BreakSerializationRecursion=true
\ No newline at end of file
+BreakSerializationRecursion=true
diff --git a/Engine/Source/Developer/AssetTools/Private/AssetTools.cpp b/Engine/Source/Developer/AssetTools/Private/AssetTools.cpp
--- a/Engine/Source/Developer/AssetTools/Private/AssetTools.cpp
+++ b/Engine/Source/Developer/AssetTools/Private/AssetTools.cpp
@@ -77,7 +77,6 @@ FAssetTools::FAssetTools()
 	RegisterAssetTypeActions( MakeShareable(new FAssetTypeActions_Texture) );
 	RegisterAssetTypeActions( MakeShareable(new FAssetTypeActions_Texture2D) );
 	RegisterAssetTypeActions( MakeShareable(new FAssetTypeActions_TextureCube) );
-	RegisterAssetTypeActions( MakeShareable(new FAssetTypeActions_TextureMovie) );
 	RegisterAssetTypeActions( MakeShareable(new FAssetTypeActions_TextureRenderTarget) );
 	RegisterAssetTypeActions( MakeShareable(new FAssetTypeActions_TextureRenderTarget2D) );
 	RegisterAssetTypeActions( MakeShareable(new FAssetTypeActions_TextureRenderTargetCube) );
diff --git a/Engine/Source/Developer/AssetTools/Private/AssetToolsPrivatePCH.h b/Engine/Source/Developer/AssetTools/Private/AssetToolsPrivatePCH.h
--- a/Engine/Source/Developer/AssetTools/Private/AssetToolsPrivatePCH.h
+++ b/Engine/Source/Developer/AssetTools/Private/AssetToolsPrivatePCH.h
@@ -74,7 +74,6 @@ DECLARE_LOG_CATEGORY_EXTERN(LogAssetTools, Log, All);
 #include "AssetTypeActions/AssetTypeActions_StaticMesh.h"
 #include "AssetTypeActions/AssetTypeActions_Texture2D.h"
 #include "AssetTypeActions/AssetTypeActions_TextureCube.h"
-#include "AssetTypeActions/AssetTypeActions_TextureMovie.h"
 #include "AssetTypeActions/AssetTypeActions_TextureRenderTarget2D.h"
 #include "AssetTypeActions/AssetTypeActions_TextureRenderTargetCube.h"
 #include "AssetTypeActions/AssetTypeActions_TextureLightProfile.h"
diff --git a/Engine/Source/Developer/AssetTools/Private/AssetTypeActions/AssetTypeActions_TextureMovie.h b/Engine/Source/Developer/AssetTools/Private/AssetTypeActions/AssetTypeActions_TextureMovie.h
deleted file mode 100644
diff --git a/Engine/Source/Runtime/Engine/Classes/Engine/CodecMovie.h b/Engine/Source/Runtime/Engine/Classes/Engine/CodecMovie.h
deleted file mode 100644
diff --git a/Engine/Source/Runtime/Engine/Classes/Engine/CodecMovieFallback.h b/Engine/Source/Runtime/Engine/Classes/Engine/CodecMovieFallback.h
deleted file mode 100644
diff --git a/Engine/Source/Runtime/Engine/Classes/Engine/GameEngine.h b/Engine/Source/Runtime/Engine/Classes/Engine/GameEngine.h
--- a/Engine/Source/Runtime/Engine/Classes/Engine/GameEngine.h
+++ b/Engine/Source/Runtime/Engine/Classes/Engine/GameEngine.h
@@ -104,7 +104,6 @@ public:
 	 */
 
 	bool HandleReattachComponentsCommand( const TCHAR* Cmd, FOutputDevice& Ar );
-	bool HandleMovieCommand( const TCHAR* Cmd, FOutputDevice& Ar );
 	bool HandleCommand( const TCHAR* Cmd, FOutputDevice& Ar );
 	bool HandleExitCommand( const TCHAR* Cmd, FOutputDevice& Ar );
 	bool HandleGetMaxTickRateCommand( const TCHAR* Cmd, FOutputDevice& Ar );
diff --git a/Engine/Source/Runtime/Engine/Classes/Engine/TextureMovie.h b/Engine/Source/Runtime/Engine/Classes/Engine/TextureMovie.h
deleted file mode 100644
diff --git a/Engine/Source/Runtime/Engine/Classes/Materials/MaterialExpressionTextureSampleParameterMovie.h b/Engine/Source/Runtime/Engine/Classes/Materials/MaterialExpressionTextureSampleParameterMovie.h
deleted file mode 100644
diff --git a/Engine/Source/Runtime/Engine/Private/EnginePrivate.h b/Engine/Source/Runtime/Engine/Private/EnginePrivate.h
--- a/Engine/Source/Runtime/Engine/Private/EnginePrivate.h
+++ b/Engine/Source/Runtime/Engine/Private/EnginePrivate.h
@@ -270,7 +270,6 @@
 #include "Engine/TextureLightProfile.h"
 #include "Engine/Texture2DDynamic.h"
 #include "Engine/TextureCube.h"
-#include "Engine/TextureMovie.h"
 #include "Engine/TextureRenderTarget.h"
 #include "Engine/TextureRenderTarget2D.h"
 #include "Engine/CanvasRenderTarget2D.h"
diff --git a/Engine/Source/Runtime/Engine/Private/GameEngine.cpp b/Engine/Source/Runtime/Engine/Private/GameEngine.cpp
--- a/Engine/Source/Runtime/Engine/Private/GameEngine.cpp
+++ b/Engine/Source/Runtime/Engine/Private/GameEngine.cpp
@@ -546,10 +546,6 @@ bool UGameEngine::Exec( UWorld* InWorld, const TCHAR* Cmd, FOutputDevice& Ar )
 		return HandleReattachComponentsCommand( Cmd, Ar );
 	}
 	// exec to 
-	else if( FParse::Command( &Cmd,TEXT("MOVIE")) )
-	{
-		return HandleMovieCommand( Cmd, Ar );
-	}
 	else if( FParse::Command( &Cmd,TEXT("EXIT")) || FParse::Command(&Cmd,TEXT("QUIT")))
 	{
 		if ( FPlatformProperties::SupportsQuit() )
@@ -632,35 +628,6 @@ bool UGameEngine::HandleReattachComponentsCommand( const TCHAR* Cmd, FOutputDevi
 	return true;
 }
 
-bool UGameEngine::HandleMovieCommand( const TCHAR* Cmd, FOutputDevice& Ar )
-{
-	FString MovieCmd = FParse::Token(Cmd,0);
-	if( MovieCmd.Contains(TEXT("PLAY")) )
-	{
-		for( TObjectIterator<UTextureMovie> It; It; ++It )
-		{
-			UTextureMovie* Movie = *It;
-			Movie->Play();
-		}
-	}
-	else if( MovieCmd.Contains(TEXT("PAUSE")) )
-	{
-		for( TObjectIterator<UTextureMovie> It; It; ++It )
-		{
-			UTextureMovie* Movie = *It;
-			Movie->Pause();
-		}
-	}
-	else if( MovieCmd.Contains(TEXT("STOP")) )
-	{
-		for( TObjectIterator<UTextureMovie> It; It; ++It )
-		{
-			UTextureMovie* Movie = *It;
-			Movie->Stop();
-		}
-	}
-	return true;
-}
 
 bool UGameEngine::HandleExitCommand( const TCHAR* Cmd, FOutputDevice& Ar )
 {
diff --git a/Engine/Source/Runtime/Engine/Private/Materials/MaterialExpressions.cpp b/Engine/Source/Runtime/Engine/Private/Materials/MaterialExpressions.cpp
--- a/Engine/Source/Runtime/Engine/Private/Materials/MaterialExpressions.cpp
+++ b/Engine/Source/Runtime/Engine/Private/Materials/MaterialExpressions.cpp
@@ -122,7 +122,6 @@
 #include "Materials/MaterialExpressionAntialiasedTextureMask.h"
 #include "Materials/MaterialExpressionTextureSampleParameterSubUV.h"
 #include "Materials/MaterialExpressionTextureSampleParameterCube.h"
-#include "Materials/MaterialExpressionTextureSampleParameterMovie.h"
 #include "Materials/MaterialExpressionTextureCoordinate.h"
 #include "Materials/MaterialExpressionTime.h"
 #include "Materials/MaterialExpressionTransform.h"
@@ -1821,48 +1820,6 @@ void UMaterialExpressionTextureSampleParameterCube::SetDefaultTexture()
 	Texture = LoadObject<UTextureCube>(NULL, TEXT("/Engine/EngineResources/DefaultTextureCube.DefaultTextureCube"), NULL, LOAD_None, NULL);
 }
 
-UMaterialExpressionTextureSampleParameterMovie::UMaterialExpressionTextureSampleParameterMovie(const class FPostConstructInitializeProperties& PCIP)
-	: Super(PCIP)
-{
-	// Structure to hold one-time initialization
-	struct FConstructorStatics
-	{
-		FString NAME_Texture;
-		FString NAME_Parameters;
-		FConstructorStatics()
-			: NAME_Texture(LOCTEXT( "Texture", "Texture" ).ToString())
-			, NAME_Parameters(LOCTEXT( "Parameters", "Parameters" ).ToString())
-		{
-		}
-	};
-	static FConstructorStatics ConstructorStatics;
-
-	MenuCategories.Empty();
-	MenuCategories.Add(ConstructorStatics.NAME_Texture);
-	MenuCategories.Add(ConstructorStatics.NAME_Parameters);
-}
-
-void UMaterialExpressionTextureSampleParameterMovie::GetCaption(TArray<FString>& OutCaptions) const
-{
-	OutCaptions.Add(TEXT("ParamMovie")); 
-	OutCaptions.Add(FString::Printf(TEXT("'%s'"), *ParameterName.ToString()));
-}
-
-bool UMaterialExpressionTextureSampleParameterMovie::TextureIsValid( UTexture* InTexture )
-{
-	bool Result=false;
-	if (InTexture)
-	{
-		Result = (InTexture->GetClass() == UTextureMovie::StaticClass());
-	}
-	return Result;
-}
-
-const TCHAR* UMaterialExpressionTextureSampleParameterMovie::GetRequirements()
-{
-	return TEXT("Requires TextureMovie");
-}
-
 /** 
  * Performs a SubUV operation, which is doing a texture lookup into a sub rectangle of a texture, and optionally blending with another rectangle.  
  * This supports both sprites and mesh emitters.
diff --git a/Engine/Source/Runtime/Engine/Private/MovieCodecs.cpp b/Engine/Source/Runtime/Engine/Private/MovieCodecs.cpp
deleted file mode 100644
diff --git a/Engine/Source/Runtime/Engine/Private/TextureMovie.cpp b/Engine/Source/Runtime/Engine/Private/TextureMovie.cpp
deleted file mode 100644
diff --git a/Engine/Source/Runtime/Engine/Public/Engine.h b/Engine/Source/Runtime/Engine/Public/Engine.h
--- a/Engine/Source/Runtime/Engine/Public/Engine.h
+++ b/Engine/Source/Runtime/Engine/Public/Engine.h
@@ -328,7 +328,6 @@
 #include "Engine/TextureLightProfile.h"
 #include "Engine/Texture2DDynamic.h"
 #include "Engine/TextureCube.h"
-#include "Engine/TextureMovie.h"
 #include "Engine/TextureRenderTarget.h"
 #include "Engine/TextureRenderTarget2D.h"
 #include "Engine/CanvasRenderTarget2D.h"
