Making it so we first set the UserDefinedStruct's defaults before we serialize in over them (so BP arrays of structs don't get cleared when CopyPropertiesFromUnrelatedObjects is called).
(TTP #342798)

diff --git a/Engine/Source/Runtime/Engine/Private/UserDefinedStruct.cpp b/Engine/Source/Runtime/Engine/Private/UserDefinedStruct.cpp
--- a/Engine/Source/Runtime/Engine/Private/UserDefinedStruct.cpp
+++ b/Engine/Source/Runtime/Engine/Private/UserDefinedStruct.cpp
@@ -75,10 +75,9 @@ void UUserDefinedStruct::SerializeTaggedProperties(FArchive& Ar, uint8* Data, US
 
 	/*	Object serialized from delta will have missing properties filled with zeroed data, 
 		we want structure's default data instead */
-	const bool bLoadDefaultFirst = UDDefaultsStruct 
+	const bool bLoadDefaultFirst = UDDefaultsStruct
 		&& !bDuplicate
-		&& Ar.IsLoading() 
-		&& Ar.IsPersistent();
+		&& Ar.IsLoading();
 
 	const bool bPrepareDefaultStruct = bUseNewDefaults || bLoadDefaultFirst;
 	FStructOnScope StructDefaultMem(bPrepareDefaultStruct ? UDDefaultsStruct : NULL);
