Media: forgot to check in (for legacy movie texture removal)

diff --git a/Engine/Source/Runtime/Engine/Public/TextureResource.h b/Engine/Source/Runtime/Engine/Public/TextureResource.h
--- a/Engine/Source/Runtime/Engine/Public/TextureResource.h
+++ b/Engine/Source/Runtime/Engine/Public/TextureResource.h
@@ -826,58 +826,5 @@ private:
 	ECubeFace CurrentTargetFace;
 };
 
-/**
- * FTextureResource type for movie textures.
- */
-class FTextureMovieResource : public FTextureResource, public FRenderTarget, public FDeferredUpdateResource
-{
-public:
-
-	/** 
-	 * Constructor
-	 * @param InOwner - movie texture object to create a resource for
-	 */
-	FTextureMovieResource(const class UTextureMovie* InOwner)
-		:	Owner(InOwner)
-	{
-	}
-
-	// FRenderResource interface.
-
-	/**
-	 * Initializes the dynamic RHI resource and/or RHI render target used by this resource.
-	 * Called when the resource is initialized, or when reseting all RHI resources.
-	 * Resources that need to initialize after a D3D device reset must implement this function.
-	 * This is only called by the rendering thread.
-	 */
-	virtual void InitDynamicRHI();
-
-	/**
-	 * Releases the dynamic RHI resource and/or RHI render target resources used by this resource.
-	 * Called when the resource is released, or when reseting all RHI resources.
-	 * Resources that need to release before a D3D device reset must implement this function.
-	 * This is only called by the rendering thread.
-	 */
-	virtual void ReleaseDynamicRHI();	
-
-	// FDeferredClearResource interface
-
-	// FRenderTarget interface.
-	virtual FIntPoint GetSizeXY() const;
-
-protected:
-
-	/**
-	* Decodes the next frame of the movie stream and renders the result to this movie texture target
-	*/
-	virtual void UpdateDeferredResource(FRHICommandListImmediate& RHICmdList) override;
-
-
-private:
-	/** The UTextureRenderTarget2D which this resource represents. */
-	const class UTextureMovie* Owner;
-	/** Texture resource used for rendering with and resolving to */
-	FTexture2DRHIRef Texture2DRHI;
-};
 
 #define TEXTURERESOURCE_H_INCLUDED 1 // needed by TargetPlatform, so TTargetPlatformBase template knows if it can use UTexture in GetDefaultTextureFormatName, or rather just declare it
