Updating blueprint component attach hierarchies for native components that have changed since the blueprint was last saved.
(TTP #325936)

diff --git a/Engine/Source/Runtime/Engine/Private/Blueprint.cpp b/Engine/Source/Runtime/Engine/Private/Blueprint.cpp
--- a/Engine/Source/Runtime/Engine/Private/Blueprint.cpp
+++ b/Engine/Source/Runtime/Engine/Private/Blueprint.cpp
@@ -19,6 +19,104 @@
 
 DEFINE_LOG_CATEGORY(LogBlueprint);
 
+//////////////////////////////////////////////////////////////////////////
+// Static Helpers
+
+/**
+ * Updates the blueprint's OwnedComponents, such that they reflect changes made 
+ * natively since the blueprint was last saved (a change in AttachParents, etc.)
+ * 
+ * @param  Blueprint	The blueprint whose components you wish to vet.
+ */
+static void ConformNativeComponents(UBlueprint* Blueprint)
+{
+	UClass* const BlueprintClass = Blueprint->GeneratedClass;
+	UClass* const NativeSuperClass = FBlueprintEditorUtils::FindFirstNativeClass(BlueprintClass);
+
+	if (AActor* BlueprintCDO = Cast<AActor>(BlueprintClass->ClassDefaultObject))
+	{
+		TArray<UActorComponent*> OldNativeComponents;
+		// collect the native components that this blueprint was serialized out 
+		// with (the native components it had last time it was saved)
+		BlueprintCDO->GetComponents(OldNativeComponents);
+
+		AActor* NativeCDO = CastChecked<AActor>(NativeSuperClass->ClassDefaultObject);
+		// collect the more up to date native components (directly from the 
+		// native super-class)
+		TArray<UActorComponent*> NewNativeComponents;
+		NativeCDO->GetComponents(NewNativeComponents);
+
+		// utility lambda for finding named components in a supplied list
+		auto FindNamedComponentLambda = [](FName const ComponentName, TArray<UActorComponent*> const& ComponentList)->UActorComponent*
+		{
+			UActorComponent* FoundComponent = nullptr;
+			for (UActorComponent* Component : ComponentList)
+			{
+				if (Component->GetFName() == ComponentName)
+				{
+					FoundComponent = Component;
+					break;
+				}
+			}
+			return FoundComponent;
+		};
+
+		// utility lambda for finding matching components in the NewNativeComponents list
+		auto FindNativeComponentLambda = [&NewNativeComponents, &FindNamedComponentLambda](UActorComponent* BlueprintComponent)->UActorComponent*
+		{
+			UActorComponent* MatchingComponent = nullptr;
+			if (BlueprintComponent != nullptr)
+			{
+				FName const ComponentName = BlueprintComponent->GetFName();
+				MatchingComponent = FindNamedComponentLambda(ComponentName, NewNativeComponents);
+			}
+			return MatchingComponent;
+		};
+
+		// loop through all components that this blueprint thinks come from its
+		// native super-class (last time it was saved)
+		for (UActorComponent* Component : OldNativeComponents)
+		{
+			// if we found this component also listed for the native class
+			if (UActorComponent* NativeComponent = FindNativeComponentLambda(Component))
+			{
+				USceneComponent* BlueprintSceneComponent = Cast<USceneComponent>(Component);
+				if (BlueprintSceneComponent == nullptr)
+				{
+					// if this isn't a scene-component, then we don't care
+					// (we're looking to fixup scene-component parents)
+					continue;
+				}
+				UActorComponent* OldNativeParent = FindNativeComponentLambda(BlueprintSceneComponent->AttachParent);
+
+				USceneComponent* NativeSceneComponent = CastChecked<USceneComponent>(NativeComponent);
+				// if this native component has since been reparented, we need
+				// to make sure that this blueprint reflects that change
+				if (OldNativeParent != NativeSceneComponent->AttachParent)
+				{
+					USceneComponent* NewParent = nullptr;
+					if (NativeSceneComponent->AttachParent != nullptr)
+					{
+						NewParent = CastChecked<USceneComponent>(FindNamedComponentLambda(NativeSceneComponent->AttachParent->GetFName(), OldNativeComponents));
+					}
+					BlueprintSceneComponent->AttachParent = NewParent;
+				}
+			}
+			else // the component has been removed from the native class
+			{
+				// @TODO: I think we already handle removed native components elsewhere, so maybe we should error here?
+// 				BlueprintCDO->RemoveOwnedComponent(Component);
+// 
+// 				USimpleConstructionScript* BlueprintSCS = Blueprint->SimpleConstructionScript;
+// 				USCS_Node* ComponentNode = BlueprintSCS->CreateNode(Component, Component->GetFName());
+// 
+// 				BlueprintSCS->AddNode(ComponentNode);
+			}
+		}
+	}
+}
+
+
 //////////////////////////////////////////////////////////////////////////
 // FBPVariableDescription
 
@@ -1066,6 +1164,8 @@ void UBlueprint::PostLoadSubobjects(FObjectInstancingGraph* OuterInstanceGraph)
 {
 	Super::PostLoadSubobjects(OuterInstanceGraph);
 	ChangeOwnerOfTemplates();
+
+	ConformNativeComponents(this);
 }
 
 bool UBlueprint::Modify(bool bAlwaysMarkDirty)
