Shuffle information about initialization order out of the class comment used for tooltips

diff --git a/Engine/Source/Runtime/Engine/Classes/GameFramework/Actor.h b/Engine/Source/Runtime/Engine/Classes/GameFramework/Actor.h
--- a/Engine/Source/Runtime/Engine/Classes/GameFramework/Actor.h
+++ b/Engine/Source/Runtime/Engine/Classes/GameFramework/Actor.h
@@ -43,18 +43,20 @@ DECLARE_CYCLE_STAT_EXTERN(TEXT("GetComponentsTime"),STAT_GetComponentsTime,STATG
 
 /**
  * Base class for an object that can be placed or spawned in a level. Actors may contain a collection of Components, and support network replication.
- *
- * The functions of interest to initialization order for an Actor is roughly as follners:
- * PostLoad/PostActorCreated - Do any setup of the actor required for construction. PostLoad for serialized actors, PostActorCreated for spawned.  
- * AActor::OnConstruction - The construction of the actor, this is where Blueprint actors have their components created and blueprint variables are initialized
- * AActor::PreInitializeComponents - Called before InitializeComponent is called on the actor's components
- * UActorComponent::InitializeComponent - Each component in the actor's components array gets an initialize call (if bWantsInitializeComponent is true for that component)
- * AActor::PostInitializeComponents - Called after the actor's components have been initialized
- * AActor::BeginPlay - Called when the level is started
- */
+*/
 UCLASS(abstract, BlueprintType, Blueprintable, config=Engine)
 class ENGINE_API AActor : public UObject
 {
+	/**
+	* The functions of interest to initialization order for an Actor is roughly as follows:
+	* PostLoad/PostActorCreated - Do any setup of the actor required for construction. PostLoad for serialized actors, PostActorCreated for spawned.  
+	* AActor::OnConstruction - The construction of the actor, this is where Blueprint actors have their components created and blueprint variables are initialized
+	* AActor::PreInitializeComponents - Called before InitializeComponent is called on the actor's components
+	* UActorComponent::InitializeComponent - Each component in the actor's components array gets an initialize call (if bWantsInitializeComponent is true for that component)
+	* AActor::PostInitializeComponents - Called after the actor's components have been initialized
+	* AActor::BeginPlay - Called when the level is started
+	*/
+
 	GENERATED_UCLASS_BODY()
 
 	/** Primary Actor tick function */
