ScreenPercentage was broken for non-stereo in anything above 100. Now fixed.

diff --git a/Engine/Source/Runtime/Engine/Private/SceneView.cpp b/Engine/Source/Runtime/Engine/Private/SceneView.cpp
--- a/Engine/Source/Runtime/Engine/Private/SceneView.cpp
+++ b/Engine/Source/Runtime/Engine/Private/SceneView.cpp
@@ -904,14 +904,7 @@ void FSceneView::EndFinalPostprocessSettings()
 
 		if(Value >= 0.0)
 		{
-			if (!GEngine->IsStereoscopic3D())
-			{
-				FinalPostProcessSettings.ScreenPercentage = FMath::Min(FinalPostProcessSettings.ScreenPercentage, Value);
-			}
-			else
-			{
-				FinalPostProcessSettings.ScreenPercentage = Value;
-			}
+			FinalPostProcessSettings.ScreenPercentage = Value;
 		}
 
 		// Not supported in ES2.
