Fixing broken builder with include needed for non-unity builds.

diff --git a/Engine/Source/Runtime/Engine/Private/Blueprint.cpp b/Engine/Source/Runtime/Engine/Private/Blueprint.cpp
--- a/Engine/Source/Runtime/Engine/Private/Blueprint.cpp
+++ b/Engine/Source/Runtime/Engine/Private/Blueprint.cpp
@@ -15,6 +15,7 @@
 #include "Crc.h"
 #include "MessageLog.h"
 #include "Editor/UnrealEd/Classes/Settings/EditorLoadingSavingSettings.h"
+#include "BlueprintEditorUtils.h"
 #endif
 
 DEFINE_LOG_CATEGORY(LogBlueprint);
@@ -30,6 +31,7 @@ DEFINE_LOG_CATEGORY(LogBlueprint);
  */
 static void ConformNativeComponents(UBlueprint* Blueprint)
 {
+#if WITH_EDITOR
 	UClass* const BlueprintClass = Blueprint->GeneratedClass;
 	UClass* const NativeSuperClass = FBlueprintEditorUtils::FindFirstNativeClass(BlueprintClass);
 
@@ -114,6 +116,7 @@ static void ConformNativeComponents(UBlueprint* Blueprint)
 			}
 		}
 	}
+#endif // #if WITH_EDITOR
 }
 
 
