Fix WorldType during PostLoad operations for seamless travel. (Fixes BSP missing after seamless travel)

diff --git a/Engine/Source/Runtime/Engine/Private/World.cpp b/Engine/Source/Runtime/Engine/Private/World.cpp
--- a/Engine/Source/Runtime/Engine/Private/World.cpp
+++ b/Engine/Source/Runtime/Engine/Private/World.cpp
@@ -4041,6 +4041,10 @@ void FSeamlessTravelHandler::SetHandlerLoadedData(UObject* InLevelPackage, UWorl
 /** callback sent to async loading code to inform us when the level package is complete */
 void FSeamlessTravelHandler::SeamlessTravelLoadCallback(const FString& PackageName, UPackage* LevelPackage)
 {
+	// make sure we remove the name, even if travel was cancelled.
+	const FName URLMapFName = FName(*PendingTravelURL.Map);
+	UWorld::WorldTypePreLoadMap.Remove(URLMapFName);
+
 	// defer until tick when it's safe to perform the transition
 	if (IsInTransition())
 	{
@@ -4207,6 +4211,10 @@ void FSeamlessTravelHandler::StartLoadingDestination()
 			}
 		}
 
+		// Set the world type in the static map, so that UWorld::PostLoad can set the world type
+		const FName URLMapFName = FName(*PendingTravelURL.Map);
+		UWorld::WorldTypePreLoadMap.FindOrAdd(URLMapFName) = CurrentWorld->WorldType;
+
 		LoadPackageAsync(PendingTravelURL.Map, 
 			FLoadPackageAsyncDelegate::CreateRaw(this, &FSeamlessTravelHandler::SeamlessTravelLoadCallback), 
 			PendingTravelGuid.IsValid() ? &PendingTravelGuid : NULL
