Engine: Exporting selected TextureResource and TextureLODSettings methods, so these types can be subclassed in other modules

diff --git a/Engine/Source/Runtime/Engine/Public/TextureLODSettings.h b/Engine/Source/Runtime/Engine/Public/TextureLODSettings.h
--- a/Engine/Source/Runtime/Engine/Public/TextureLODSettings.h
+++ b/Engine/Source/Runtime/Engine/Public/TextureLODSettings.h
@@ -63,7 +63,7 @@ struct FTextureLODSettings
 	 * @param	InLODGroup		The LOD Group ID 
 	 * @return	LODBias
 	 */
-	int32 GetTextureLODGroupLODBias( int32 InLODGroup ) const;
+	ENGINE_API int32 GetTextureLODGroupLODBias( int32 InLODGroup ) const;
 
 	/**
 	 * Returns the LODGroup setting for number of streaming mip-levels.
@@ -72,7 +72,7 @@ struct FTextureLODSettings
 	 * @param	InLODGroup		The LOD Group ID 
 	 * @return	Number of streaming mip-levels for textures in the specified LODGroup
 	 */
-	int32 GetNumStreamedMips( int32 InLODGroup ) const;
+	ENGINE_API int32 GetNumStreamedMips( int32 InLODGroup ) const;
 
 	/**
 	 * Returns the filter state that should be used for the passed in texture, taking
@@ -81,7 +81,7 @@ struct FTextureLODSettings
 	 * @param	Texture		Texture to retrieve filter state for, must not be 0
 	 * @return	Filter sampler state for passed in texture
 	 */
-	ESamplerFilter GetSamplerFilter( const UTexture* Texture ) const;
+	ENGINE_API ESamplerFilter GetSamplerFilter( const UTexture* Texture ) const;
 
 	/**
 	 * Returns the texture group names, sorted like enum.
diff --git a/Engine/Source/Runtime/Engine/Public/TextureResource.h b/Engine/Source/Runtime/Engine/Public/TextureResource.h
--- a/Engine/Source/Runtime/Engine/Public/TextureResource.h
+++ b/Engine/Source/Runtime/Engine/Public/TextureResource.h
@@ -546,19 +546,22 @@ private:
 class FDeferredUpdateResource
 {
 public:
+
 	/**
 	 * Constructor, initializing UpdateListLink.
 	 */
 	FDeferredUpdateResource()
 		:	UpdateListLink(NULL)
 		,	bOnlyUpdateOnce(false)
-	{}
+	{ }
+
+public:
 
 	/**
 	 * Iterate over the global list of resources that need to
 	 * be updated and call UpdateResource on each one.
 	 */
-	ENGINE_API static void UpdateResources(FRHICommandListImmediate& RHICmdList);
+	ENGINE_API static void UpdateResources( FRHICommandListImmediate& RHICmdList );
 
 	/** 
 	 * This is reset after all viewports have been rendered
@@ -568,27 +571,26 @@ public:
 		bNeedsUpdate = true;
 	}
 
-	// FDeferredUpdateResource interface
-
 protected:
 
 	/**
 	* Updates the resource
 	*/
-	virtual void UpdateDeferredResource(FRHICommandListImmediate& RHICmdList) = 0;
+	virtual void UpdateDeferredResource( FRHICommandListImmediate& RHICmdList ) = 0;
 
 	/**
 	 * Add this resource to deferred update list
 	 * @param OnlyUpdateOnce - flag this resource for a single update if true
 	 */
-	void AddToDeferredUpdateList( bool OnlyUpdateOnce );
+	ENGINE_API void AddToDeferredUpdateList( bool OnlyUpdateOnce );
 
 	/**
 	 * Remove this resource from deferred update list
 	 */
-	void RemoveFromDeferredUpdateList();
+	ENGINE_API void RemoveFromDeferredUpdateList();
 
 private:
+
 	/** 
 	 * Resources can be added to this list if they need a deferred update during scene rendering.
 	 * @return global list of resource that need to be updated. 
