Make sure the transition map for seamless travel, when used, has the correct WorldType

diff --git a/Engine/Source/Runtime/Engine/Private/World.cpp b/Engine/Source/Runtime/Engine/Private/World.cpp
--- a/Engine/Source/Runtime/Engine/Private/World.cpp
+++ b/Engine/Source/Runtime/Engine/Private/World.cpp
@@ -4138,6 +4138,9 @@ bool FSeamlessTravelHandler::StartTravel(UWorld* InCurrentWorld, const FURL& InU
 			}
 			else
 			{
+				// Set the world type in the static map, so that UWorld::PostLoad can set the world type
+				UWorld::WorldTypePreLoadMap.FindOrAdd(*TransitionMap) = CurrentWorld->WorldType;
+
 				// first, load the entry level package
 				LoadPackageAsync(TransitionMap, 
 					FLoadPackageAsyncDelegate::CreateRaw(this, &FSeamlessTravelHandler::SeamlessTravelLoadCallback)
