Streaming strategy for a texture calculated based only on primitive mobility settings

diff --git a/Engine/Source/Runtime/Engine/Private/Level.cpp b/Engine/Source/Runtime/Engine/Private/Level.cpp
--- a/Engine/Source/Runtime/Engine/Private/Level.cpp
+++ b/Engine/Source/Runtime/Engine/Private/Level.cpp
@@ -1329,10 +1329,10 @@ void ULevel::BuildStreamingData(UTexture2D* UpdateSpecificTextureOnly/*=NULL*/)
 			if ( !bIsClassDefaultObject && Primitive->IsRegistered() )
 			{
 				const AActor* const Owner				= Primitive->GetOwner();
-				const bool bIsStaticMeshComponent		= Primitive->IsA(UStaticMeshComponent::StaticClass());
-				const bool bIsSkeletalMeshComponent		= Primitive->IsA(USkeletalMeshComponent::StaticClass());
-				const bool bIsFoliage					= Owner && Owner->IsA(AInstancedFoliageActor::StaticClass()) && Primitive->IsA(UInstancedStaticMeshComponent::StaticClass());
-				const bool bIsStatic					= Owner == NULL || Primitive->Mobility == EComponentMobility::Static || bIsFoliage;
+				const bool bIsFoliage					= Owner && Owner->IsA(AInstancedFoliageActor::StaticClass()) && Primitive->IsA(UInstancedStaticMeshComponent::StaticClass()); 
+				const bool bIsStatic					= Owner == NULL 
+															|| Primitive->Mobility == EComponentMobility::Static 
+															|| bIsFoliage; // treat Foliage components as static, regardless of mobility settings
 
 				TArray<FStreamingTexturePrimitiveInfo> PrimitiveStreamingTextures;
 
@@ -1380,33 +1380,14 @@ void ULevel::BuildStreamingData(UTexture2D* UpdateSpecificTextureOnly/*=NULL*/)
 
 					if(bShouldHandleTexture)
 					{
-						// Check if this is a world texture.
-						const bool bIsWorldTexture			= 
-							Texture2D->LODGroup == TEXTUREGROUP_World ||
-							Texture2D->LODGroup == TEXTUREGROUP_WorldNormalMap ||
-							Texture2D->LODGroup == TEXTUREGROUP_WorldSpecular ||
-							Texture2D->LODGroup == TEXTUREGROUP_Terrain_Heightmap ||
-							Texture2D->LODGroup == TEXTUREGROUP_Terrain_Weightmap ||
-							Texture2D->LODGroup == TEXTUREGROUP_Shadowmap ||
-							Texture2D->LODGroup == TEXTUREGROUP_Lightmap;
-
-						// Check if we should consider this a static mesh texture instance.
-						bool bIsStaticMeshTextureInstance = bIsWorldTexture && !bIsSkeletalMeshComponent;
-
-						// Treat non-static textures dynamically instead.
-						if ( !bIsStatic && bUseDynamicStreaming )
-						{
-							bIsStaticMeshTextureInstance = false;
-						}
-
 						// Is the primitive set to force its textures to be resident?
 						if ( Primitive->bForceMipStreaming )
 						{
 							// Add them to the ForceStreamTextures set.
 							ForceStreamTextures.Add(Texture2D,true);
 						}
-						// Is this a static mesh texture instance?
-						else if ( bIsStaticMeshTextureInstance && bCanBeStreamedByDistance )
+						// Is this texture used by a static object?
+						else if ( bIsStatic && bCanBeStreamedByDistance )
 						{
 							// Texture instance information.
 							FStreamableTextureInstance TextureInstance;
