Header cleanup: WheeledVehicleMovementComponent.h, EdGraphSchema.h, PlayerState.h, GameplayStatics.h CharacterMovementComponent.h.

diff --git a/Engine/Source/Runtime/Engine/Classes/GameFramework/CharacterMovementComponent.h b/Engine/Source/Runtime/Engine/Classes/GameFramework/CharacterMovementComponent.h
--- a/Engine/Source/Runtime/Engine/Classes/GameFramework/CharacterMovementComponent.h
+++ b/Engine/Source/Runtime/Engine/Classes/GameFramework/CharacterMovementComponent.h
@@ -6,13 +6,13 @@
  */
 
 #pragma once
-
-#include "WorldCollision.h"
-#include "GameFramework/PawnMovementComponent.h"
-#include "Interfaces/NetworkPredictionInterface.h"
-#include "Engine/EngineTypes.h"
 #include "AI/Navigation/NavigationAvoidanceTypes.h"
 #include "Animation/AnimationAsset.h"
+#include "Engine/EngineBaseTypes.h"
+#include "Engine/EngineTypes.h"
+#include "GameFramework/PawnMovementComponent.h"
+#include "Interfaces/NetworkPredictionInterface.h"
+#include "WorldCollision.h"
 #include "CharacterMovementComponent.generated.h"
 
 struct FVector_NetQuantize100;
diff --git a/Engine/Source/Runtime/Engine/Classes/GameFramework/PlayerState.h b/Engine/Source/Runtime/Engine/Classes/GameFramework/PlayerState.h
--- a/Engine/Source/Runtime/Engine/Classes/GameFramework/PlayerState.h
+++ b/Engine/Source/Runtime/Engine/Classes/GameFramework/PlayerState.h
@@ -1,7 +1,7 @@
 // Copyright 1998-2014 Epic Games, Inc. All Rights Reserved.
 
 #pragma once
-#include "OnlineReplStructs.h"
+#include "GameFramework/Info.h"
 #include "PlayerState.generated.h"
 
 /**
diff --git a/Engine/Source/Runtime/Engine/Private/Animation/AnimNode_WheelHandler.cpp b/Engine/Source/Runtime/Engine/Private/Animation/AnimNode_WheelHandler.cpp
--- a/Engine/Source/Runtime/Engine/Private/Animation/AnimNode_WheelHandler.cpp
+++ b/Engine/Source/Runtime/Engine/Private/Animation/AnimNode_WheelHandler.cpp
@@ -5,6 +5,7 @@
 #include "Animation/BoneControllers/AnimNode_WheelHandler.h"
 #include "GameFramework/WheeledVehicle.h"
 #include "Vehicles/WheeledVehicleMovementComponent.h"
+#include "Vehicles/VehicleWheel.h"
 
 /////////////////////////////////////////////////////
 // FAnimNode_WheelHandler
diff --git a/Engine/Source/Runtime/Engine/Private/EdGraph/EdGraphSchema.cpp b/Engine/Source/Runtime/Engine/Private/EdGraph/EdGraphSchema.cpp
--- a/Engine/Source/Runtime/Engine/Private/EdGraph/EdGraphSchema.cpp
+++ b/Engine/Source/Runtime/Engine/Private/EdGraph/EdGraphSchema.cpp
@@ -2,6 +2,7 @@
 
 #include "EnginePrivate.h"
 #include "EdGraph/EdGraph.h"
+#include "EdGraph/EdGraphPin.h"
 #include "EdGraph/EdGraphSchema.h"
 #include "BlueprintUtilities.h"
 #if WITH_EDITOR
@@ -641,4 +642,9 @@ void UEdGraphSchema::GetContextMenuActions(const UEdGraph* CurrentGraph, const U
 #endif
 }
 
+FString UEdGraphSchema::IsCurrentPinDefaultValid(const UEdGraphPin* Pin) const
+{
+	return IsPinDefaultValid(Pin, Pin->DefaultValue, Pin->DefaultObject, Pin->DefaultTextValue);
+}
+
 #undef LOCTEXT_NAMESPACE
diff --git a/Engine/Source/Runtime/Engine/Private/Vehicles/WheeledVehicleMovementComponent.cpp b/Engine/Source/Runtime/Engine/Private/Vehicles/WheeledVehicleMovementComponent.cpp
--- a/Engine/Source/Runtime/Engine/Private/Vehicles/WheeledVehicleMovementComponent.cpp
+++ b/Engine/Source/Runtime/Engine/Private/Vehicles/WheeledVehicleMovementComponent.cpp
@@ -1376,3 +1376,11 @@ void UWheeledVehicleMovementComponent::ComputeConstants()
 }
 
 #undef LOCTEXT_NAMESPACE
+
+FWheelSetup::FWheelSetup()
+: WheelClass(UVehicleWheel::StaticClass())
+, BoneName(NAME_None)
+, AdditionalOffset(0.0f)
+{
+
+}
