Missing files from previous checkin.

diff --git a/Engine/Source/Runtime/Engine/Classes/EdGraph/EdGraphSchema.h b/Engine/Source/Runtime/Engine/Classes/EdGraph/EdGraphSchema.h
--- a/Engine/Source/Runtime/Engine/Classes/EdGraph/EdGraphSchema.h
+++ b/Engine/Source/Runtime/Engine/Classes/EdGraph/EdGraphSchema.h
@@ -1,13 +1,12 @@
 // Copyright 1998-2014 Epic Games, Inc. All Rights Reserved.
 
 #pragma once
-
-#include "EdGraphPin.h"
-
 #include "EdGraphSchema.generated.h"
 
 class FSlateRect;
 class UEdGraphNode;
+class UEdGraphPin;
+struct FEdGraphPinType;
 
 /** Distinguishes between different graph types. Graphs can have different properties; for example: functions have one entry point, ubergraphs can have multiples. */
 UENUM()
@@ -509,10 +508,7 @@ class ENGINE_API UEdGraphSchema : public UObject
 	 *	Determine whether the current pin default values are valid
 	 *	@see IsPinDefaultValid
 	 */
-	FString IsCurrentPinDefaultValid(const UEdGraphPin* Pin) const
-	{
-		return IsPinDefaultValid(Pin, Pin->DefaultValue, Pin->DefaultObject, Pin->DefaultTextValue);
-	}
+	FString IsCurrentPinDefaultValid(const UEdGraphPin* Pin) const;
 
 	/**
 	 * An easy way to check to see if the current graph system supports pin watching.
diff --git a/Engine/Source/Runtime/Engine/Classes/Vehicles/WheeledVehicleMovementComponent.h b/Engine/Source/Runtime/Engine/Classes/Vehicles/WheeledVehicleMovementComponent.h
--- a/Engine/Source/Runtime/Engine/Classes/Vehicles/WheeledVehicleMovementComponent.h
+++ b/Engine/Source/Runtime/Engine/Classes/Vehicles/WheeledVehicleMovementComponent.h
@@ -2,10 +2,10 @@
 
 #pragma once
 #include "GameFramework/PawnMovementComponent.h"
-#include "VehicleWheel.h"
-#include "Curves/CurveBase.h"
 #include "WheeledVehicleMovementComponent.generated.h"
 
+class UVehicleWheel;
+
 #if WITH_VEHICLE
 namespace physx
 {
@@ -98,12 +98,7 @@ struct FWheelSetup
 	UPROPERTY(EditAnywhere, Category=WheelSetup)
 	FVector AdditionalOffset;
 
-	FWheelSetup()
-		: WheelClass(UVehicleWheel::StaticClass())
-		, BoneName(NAME_None)
-		, AdditionalOffset(0.0f)
-	{
-	}
+	FWheelSetup();
 };
 
 USTRUCT()
