Fix framerate and video driver crashing problems when selecting SpeedTree meshes painted as foliage

diff --git a/Engine/Shaders/LocalVertexFactory.usf b/Engine/Shaders/LocalVertexFactory.usf
--- a/Engine/Shaders/LocalVertexFactory.usf
+++ b/Engine/Shaders/LocalVertexFactory.usf
@@ -122,7 +122,7 @@ struct FVertexFactoryIntermediates
 
 	half4 Color;
 #if USE_INSTANCING
-	float2 PerInstanceParams;
+	float3 PerInstanceParams;
 #endif	// USE_INSTANCING
 };
 
@@ -441,6 +441,10 @@ FVertexFactoryIntermediates GetVertexFactoryIntermediates(FVertexFactoryInput In
 	Intermediates.PerInstanceParams.x = CPUInstanceInverseTransform[2].w;
 	#endif // !USE_INSTANCING_EMULATED
 	Intermediates.PerInstanceParams.y = 1.0 - saturate((length(InstanceLocation) - InstancingFadeOutParams.x) * InstancingFadeOutParams.y);
+	// PerInstanceParams.z stores a hide/show flag for this instance
+	float SelectedValue = trunc(GetInstanceShadowMapBias(Input).z * (1.0 / 256.0));
+	// InstancingFadeOutParams.z,w are RenderSelected and RenderDeselected respectively.
+	Intermediates.PerInstanceParams.z = InstancingFadeOutParams.z * SelectedValue + InstancingFadeOutParams.w * (1-SelectedValue);
 #endif	// USE_INSTANCING
 	return Intermediates;
 }
@@ -461,10 +465,7 @@ half3x3 VertexFactoryGetTangentToLocal( FVertexFactoryInput Input, FVertexFactor
 float4 VertexFactoryGetWorldPosition(FVertexFactoryInput Input, FVertexFactoryIntermediates Intermediates)
 {
 #if USE_INSTANCING
-	float SelectedValue = trunc(GetInstanceShadowMapBias(Input).z * (1.0 / 256.0));
-	// InstancingFadeOutParams.z,w are RenderSelected and RenderDeselected respectively.
-	float SelectionMask = InstancingFadeOutParams.z * SelectedValue + InstancingFadeOutParams.w * (1-SelectedValue);
-	return CalcWorldPosition(Input.Position, GetInstanceTransform(Input)) * SelectionMask;
+	return CalcWorldPosition(Input.Position, GetInstanceTransform(Input)) * Intermediates.PerInstanceParams.z;
 #else
 	return CalcWorldPosition(Input.Position);
 #endif	// USE_INSTANCING
diff --git a/Engine/Shaders/LocalVertexFactoryCommon.usf b/Engine/Shaders/LocalVertexFactoryCommon.usf
--- a/Engine/Shaders/LocalVertexFactoryCommon.usf
+++ b/Engine/Shaders/LocalVertexFactoryCommon.usf
@@ -17,8 +17,8 @@ struct FVertexFactoryInterpolantsVSToPS
 #endif
 
 #if USE_INSTANCING
-	// x = per-instance random, y = per-instance fade out amount
-	float2  PerInstanceParams : COLOR1;
+	// x = per-instance random, y = per-instance fade out amount, z = hide/show flag
+	float3  PerInstanceParams : COLOR1;
 #endif
 
 #if NUM_MATERIAL_TEXCOORDS
@@ -143,7 +143,7 @@ struct FVertexFactoryInterpolantsVSToPS
 
 #if USE_INSTANCING
 	// x = per-instance random, y = per-instance fade out amount
-	float2  PerInstanceParams : COLOR1;
+	float3  PerInstanceParams : COLOR1;
 #endif
 };
 
diff --git a/Engine/Shaders/MaterialTemplate.usf b/Engine/Shaders/MaterialTemplate.usf
--- a/Engine/Shaders/MaterialTemplate.usf
+++ b/Engine/Shaders/MaterialTemplate.usf
@@ -132,7 +132,7 @@ struct FMaterialPixelParameters
 #endif
 
 #if USE_INSTANCING
-	half2 PerInstanceParams;
+	half3 PerInstanceParams;
 #endif
 
 	/** Per-particle properties. Only valid for particle vertex factories. */
@@ -229,7 +229,7 @@ struct FMaterialVertexParameters
 	float4x4 InstanceLocalToWorld;
 	float3x3 InstanceWorldToLocal;
 	float3 InstanceLocalPosition;
-	float2 PerInstanceParams;
+	float3 PerInstanceParams;
 #elif PARTICLE_MESH_FACTORY
 	/** Per-particle properties. */
 	float4x4 InstanceLocalToWorld;
@@ -679,6 +679,12 @@ float3 GetSpeedTreeVertexOffset(FMaterialVertexParameters Parameters, int Geomet
 		float3x3 LocalToWorld = (float3x3)Parameters.InstanceLocalToWorld;
 		float3x3 WorldToLocal = Parameters.InstanceWorldToLocal;
 		float3 LocalPosition = Parameters.InstanceLocalPosition;
+
+		// skip if this instance is hidden
+		if (Parameters.PerInstanceParams.z < 1.f)
+		{
+			return float4(0,0,0,0);
+		}
 	#else
 		float3x3 LocalToWorld = (float3x3)Primitive.LocalToWorld;
 		float3x3 WorldToLocal = (float3x3)Primitive.WorldToLocal;
diff --git a/Engine/Source/Runtime/Engine/Private/InstancedStaticMesh.cpp b/Engine/Source/Runtime/Engine/Private/InstancedStaticMesh.cpp
--- a/Engine/Source/Runtime/Engine/Private/InstancedStaticMesh.cpp
+++ b/Engine/Source/Runtime/Engine/Private/InstancedStaticMesh.cpp
@@ -35,6 +35,7 @@
 #include "ComponentInstanceDataCache.h"
 #include "InstancedFoliage.h"
 #include "VertexFactory.h"
+#include "LocalVertexFactory.h"
 
 #if WITH_PHYSX
 #include "PhysicsEngine/PhysXSupport.h"
@@ -535,11 +536,13 @@ void FInstancedStaticMeshVertexFactory::InitRHI()
 	InitDeclaration(Elements,Data);
 }
 
-class FInstancedStaticMeshVertexFactoryShaderParameters : public FVertexFactoryShaderParameters
+class FInstancedStaticMeshVertexFactoryShaderParameters : public FLocalVertexFactoryShaderParameters
 {
 	virtual void Bind(const FShaderParameterMap& ParameterMap) override
 	{
-		InstancingFadeOutParamsParameter.Bind(ParameterMap,TEXT("InstancingFadeOutParams"));
+		FLocalVertexFactoryShaderParameters::Bind(ParameterMap);
+
+		InstancingFadeOutParamsParameter.Bind(ParameterMap, TEXT("InstancingFadeOutParams"));
 		CPUInstanceShadowMapBias.Bind(ParameterMap, TEXT("CPUInstanceShadowMapBias"));
 		CPUInstanceTransform.Bind(ParameterMap, TEXT("CPUInstanceTransform"));
 		CPUInstanceInverseTransform.Bind(ParameterMap, TEXT("CPUInstanceInverseTransform"));
@@ -549,6 +552,7 @@ class FInstancedStaticMeshVertexFactoryShaderParameters : public FVertexFactoryS
 
 	void Serialize(FArchive& Ar)
 	{
+		FLocalVertexFactoryShaderParameters::Serialize(Ar);
 		Ar << InstancingFadeOutParamsParameter;
 		Ar << CPUInstanceShadowMapBias;
 		Ar << CPUInstanceTransform;
@@ -862,7 +866,7 @@ public:
 #if WITH_EDITOR
 		if( bHasSelectedInstances )
 		{
-			// if we have selected indices, we must be in Foliage edit mode, so we mark the selection
+			// if we have selected indices, mark scene proxy as selected.
 			SetSelection_GameThread(true);
 		}
 #endif
@@ -1749,6 +1753,8 @@ void UInstancedStaticMeshComponent::SelectInstance(bool bInSelected, int32 InIns
 
 void FInstancedStaticMeshVertexFactoryShaderParameters::SetMesh( FRHICommandList& RHICmdList, FShader* VertexShader,const class FVertexFactory* VertexFactory,const class FSceneView& View,const struct FMeshBatchElement& BatchElement,uint32 DataFlags ) const 
 {
+	FLocalVertexFactoryShaderParameters::SetMesh(RHICmdList, VertexShader, VertexFactory, View, BatchElement, DataFlags);
+
 	FRHIVertexShader* VS = VertexShader->GetVertexShader();
 	if( InstancingFadeOutParamsParameter.IsBound() )
 	{
@@ -1792,14 +1798,4 @@ void FInstancedStaticMeshVertexFactoryShaderParameters::SetMesh( FRHICommandList
 			SetShaderValueArray(RHICmdList, VS, CPUInstanceInverseTransform, InstanceStream->InstanceInverseTransform, 3);
 		}
 	}
-
-	// upload SpeedTree buffer, if available
-	if (View.Family != NULL && View.Family->Scene != NULL)
-	{
-		FUniformBufferRHIParamRef SpeedTreeUniformBuffer = View.Family->Scene->GetSpeedTreeUniformBuffer(VertexFactory);
-		if (SpeedTreeUniformBuffer != NULL)
-		{
-			SetUniformBufferParameter(RHICmdList, VS, VertexShader->GetUniformBufferParameter<FSpeedTreeUniformParameters>(), SpeedTreeUniformBuffer);
-		}
-	}
 }
diff --git a/Engine/Source/Runtime/Engine/Private/LocalVertexFactory.cpp b/Engine/Source/Runtime/Engine/Private/LocalVertexFactory.cpp
--- a/Engine/Source/Runtime/Engine/Private/LocalVertexFactory.cpp
+++ b/Engine/Source/Runtime/Engine/Private/LocalVertexFactory.cpp
@@ -9,39 +9,33 @@
 #include "ShaderParameterUtils.h"
 #include "LocalVertexFactory.h"
 
-class FSpeedTreeVertexFactoryShaderParameters : public FVertexFactoryShaderParameters
+void FLocalVertexFactoryShaderParameters::Bind(const FShaderParameterMap& ParameterMap)
 {
-public:
-
-	virtual void Bind(const FShaderParameterMap& ParameterMap) override
-	{
-		LODParameter.Bind(ParameterMap, TEXT("SpeedTreeLODInfo"));
-	}
+	LODParameter.Bind(ParameterMap, TEXT("SpeedTreeLODInfo"));
+}
 
-	virtual void Serialize(FArchive& Ar) override
-	{
-		Ar << LODParameter;
-	}
+void FLocalVertexFactoryShaderParameters::Serialize(FArchive& Ar)
+{
+	Ar << LODParameter;
+}
 
-	virtual void SetMesh(FRHICommandList& RHICmdList, FShader* Shader,const FVertexFactory* VertexFactory,const FSceneView& View,const FMeshBatchElement& BatchElement,uint32 DataFlags) const override
+void FLocalVertexFactoryShaderParameters::SetMesh(FRHICommandList& RHICmdList, FShader* Shader, const FVertexFactory* VertexFactory, const FSceneView& View, const FMeshBatchElement& BatchElement, uint32 DataFlags) const
+{
+	if (View.Family != NULL && View.Family->Scene != NULL)
 	{
-		if (View.Family != NULL && View.Family->Scene != NULL)
+		FUniformBufferRHIParamRef SpeedTreeUniformBuffer = View.Family->Scene->GetSpeedTreeUniformBuffer(VertexFactory);
+		if (SpeedTreeUniformBuffer != NULL)
 		{
-			FUniformBufferRHIParamRef SpeedTreeUniformBuffer = View.Family->Scene->GetSpeedTreeUniformBuffer(VertexFactory);
-			if (SpeedTreeUniformBuffer != NULL)
-			{
-				SetUniformBufferParameter(RHICmdList, Shader->GetVertexShader(), Shader->GetUniformBufferParameter<FSpeedTreeUniformParameters>(), SpeedTreeUniformBuffer);
+			SetUniformBufferParameter(RHICmdList, Shader->GetVertexShader(), Shader->GetUniformBufferParameter<FSpeedTreeUniformParameters>(), SpeedTreeUniformBuffer);
 
-				if (LODParameter.IsBound())
-				{
-					FVector LODData(BatchElement.MinScreenSize, BatchElement.MaxScreenSize, BatchElement.MaxScreenSize - BatchElement.MinScreenSize);
-					SetShaderValue(RHICmdList, Shader->GetVertexShader(), LODParameter, LODData);
-				}
+			if (LODParameter.IsBound())
+			{
+				FVector LODData(BatchElement.MinScreenSize, BatchElement.MaxScreenSize, BatchElement.MaxScreenSize - BatchElement.MinScreenSize);
+				SetShaderValue(RHICmdList, Shader->GetVertexShader(), LODParameter, LODData);
 			}
 		}
 	}
-	FShaderParameter LODParameter;
-};
+}
 
 /**
  * Should we cache the material's shadertype on this platform with this vertex factory? 
@@ -153,7 +147,7 @@ FVertexFactoryShaderParameters* FLocalVertexFactory::ConstructShaderParameters(E
 {
 	if (ShaderFrequency == SF_Vertex)
 	{
-		return new FSpeedTreeVertexFactoryShaderParameters();
+		return new FLocalVertexFactoryShaderParameters();
 	}
 
 	return NULL;
diff --git a/Engine/Source/Runtime/Engine/Public/LocalVertexFactory.h b/Engine/Source/Runtime/Engine/Public/LocalVertexFactory.h
--- a/Engine/Source/Runtime/Engine/Public/LocalVertexFactory.h
+++ b/Engine/Source/Runtime/Engine/Public/LocalVertexFactory.h
@@ -61,3 +61,16 @@ public:
 protected:
 	DataType Data;
 };
+
+/**
+ * Shader parameters for LocalVertexFactory.
+ */
+class FLocalVertexFactoryShaderParameters : public FVertexFactoryShaderParameters
+{
+public:
+	virtual void Bind(const FShaderParameterMap& ParameterMap) override;
+	virtual void Serialize(FArchive& Ar) override;
+	virtual void SetMesh(FRHICommandList& RHICmdList, FShader* Shader, const FVertexFactory* VertexFactory, const FSceneView& View, const FMeshBatchElement& BatchElement, uint32 DataFlags) const override;
+	// SpeedTree LOD parameter
+	FShaderParameter LODParameter;
+};
diff --git a/Engine/Source/Runtime/Renderer/Private/ScenePrivate.h b/Engine/Source/Runtime/Renderer/Private/ScenePrivate.h
--- a/Engine/Source/Runtime/Renderer/Private/ScenePrivate.h
+++ b/Engine/Source/Runtime/Renderer/Private/ScenePrivate.h
@@ -1221,7 +1221,7 @@ public:
 	/** The wind sources in the scene. */
 	TArray<class FWindSourceSceneProxy*> WindSources;
 
-	/** SpeedTree wind objects in the scene. FSpeedTreeVertexFactoryShaderParameters needs to lookup by FVertexFactory, but wind objects are per tree (i.e. per UStaticMesh)*/
+	/** SpeedTree wind objects in the scene. FLocalVertexFactoryShaderParameters needs to lookup by FVertexFactory, but wind objects are per tree (i.e. per UStaticMesh)*/
 	TMap<UStaticMesh*, FSpeedTreeWindComputation*> SpeedTreeWindComputationMap;
 	TMap<FVertexFactory*, UStaticMesh*> SpeedTreeVertexFactoryMap;
 
