Ensure destructible mesh tangent data is being correctly generated.

*  Pass through vertex binormal to APEX.
* Corrected ordering of binormal and tangent storage on output vertex.
* Ensure Submesh array is being correctly initialized (ctor was not being called).

#TTP 339005 - RENDERING: Destructible meshes have lower quality materials than static meshes.

diff --git a/Engine/Source/Editor/UnrealEd/Private/Factories/ApexDestructibleAssetImport.cpp b/Engine/Source/Editor/UnrealEd/Private/Factories/ApexDestructibleAssetImport.cpp
--- a/Engine/Source/Editor/UnrealEd/Private/Factories/ApexDestructibleAssetImport.cpp
+++ b/Engine/Source/Editor/UnrealEd/Private/Factories/ApexDestructibleAssetImport.cpp
@@ -527,7 +527,7 @@ static bool FillSkelMeshImporterFromApexDestructibleAsset(FSkeletalMeshImportDat
 		bHaveAllNormals = bHaveAllNormals && bHaveNormals;
 		bHaveAllTangents = bHaveAllTangents && bHaveTangents;
 
-		// Binromals
+		// Binormomals
 		const PxI32 BinormalBufferIndex = VBFormat.getBufferIndexFromID(VBFormat.getSemanticID(NxRenderVertexSemantic::BINORMAL));
 		TArray<FVector> Binormals;
 		Binormals.AddUninitialized(SubmeshVertexCount);
@@ -600,8 +600,8 @@ static bool FillSkelMeshImporterFromApexDestructibleAsset(FSkeletalMeshImportDat
 				for (uint32 V = 0; V < 3; ++V)
 				{
 					// Tangent basis
-					Triangle.TangentX[V] = Binormals[SubmeshVertexIndex[V]];
-					Triangle.TangentY[V] = Tangents[SubmeshVertexIndex[V]];
+					Triangle.TangentX[V] = Tangents[SubmeshVertexIndex[V]];
+					Triangle.TangentY[V] = Binormals[SubmeshVertexIndex[V]];
 					Triangle.TangentZ[V] = Normals[SubmeshVertexIndex[V]];
 #if INVERT_Y_AND_V
 					Triangle.TangentX[V].Y *= -1.0f;
diff --git a/Engine/Source/Runtime/Engine/Private/DestructibleMesh.cpp b/Engine/Source/Runtime/Engine/Private/DestructibleMesh.cpp
--- a/Engine/Source/Runtime/Engine/Private/DestructibleMesh.cpp
+++ b/Engine/Source/Runtime/Engine/Private/DestructibleMesh.cpp
@@ -447,6 +447,7 @@ bool CreateSubmeshFromSMSection(FStaticMeshLODResources& RenderMesh, int32 Subme
 	// Create submesh descriptor, just a material name and a vertex format
 	FCStringAnsi::Strncpy(SubmeshData.mMaterialName, TCHAR_TO_ANSI(*FString::Printf(TEXT("Material%d"),Section.MaterialIndex)), physx::NxExplicitSubmeshData::MaterialNameBufferSize);
 	SubmeshData.mVertexFormat.mHasStaticPositions = SubmeshData.mVertexFormat.mHasStaticNormals = SubmeshData.mVertexFormat.mHasStaticTangents = true;
+	SubmeshData.mVertexFormat.mHasStaticBinormals = true;
 	SubmeshData.mVertexFormat.mBonesPerVertex = 1;
 	SubmeshData.mVertexFormat.mUVCount =  FMath::Min((physx::PxU32)RenderMesh.VertexBuffer.GetNumTexCoords(), (physx::PxU32)NxVertexFormat::MAX_UV_COUNT);
 
@@ -466,7 +467,7 @@ bool CreateSubmeshFromSMSection(FStaticMeshLODResources& RenderMesh, int32 Subme
 			Vertex.position = U2PVector(RenderMesh.PositionVertexBuffer.VertexPosition(UnrealVertIndex));	Vertex.position.y *= -1;
 			Vertex.normal = U2PVector((FVector)RenderMesh.VertexBuffer.VertexTangentZ(UnrealVertIndex));	Vertex.normal.y *= -1;
 			Vertex.tangent = U2PVector((FVector)RenderMesh.VertexBuffer.VertexTangentX(UnrealVertIndex));	Vertex.tangent.y *= -1;
-			// Vertex.binormal = // Not recoding the bitangent/binormal
+			Vertex.binormal = U2PVector((FVector)RenderMesh.VertexBuffer.VertexTangentY(UnrealVertIndex));	Vertex.binormal.y *= -1;
 			for (int32 TexCoordSourceIndex = 0; TexCoordSourceIndex < (int32)SubmeshData.mVertexFormat.mUVCount; ++TexCoordSourceIndex)
 			{
 				const FVector2D& TexCoord = RenderMesh.VertexBuffer.GetVertexUV(UnrealVertIndex, TexCoordSourceIndex);
@@ -537,7 +538,7 @@ bool UDestructibleMesh::BuildFractureSettingsFromStaticMesh(UStaticMesh* StaticM
 	TArray<physx::NxExplicitRenderTriangle> Triangles;
 	Triangles.Reserve(OverallTriangleCount);
 	TArray<physx::NxExplicitSubmeshData> Submeshes;	// Elements <-> Submeshes
-	Submeshes.AddUninitialized(OverallSubmeshCount);
+	Submeshes.SetNum(OverallSubmeshCount);
 
 	// UE Materials
 	TArray<UMaterialInterface*> Materials;
