Header cleanup: PhysicalMaterial.h, SoundMix.h, TimelineTemplate.h,  EmitterCameraLensEffectBase.h, OnlineReplStructs.h, SpectatorPawnMovement.h, GameSession.h, GameUserSettings.h, StreamableManager.h

diff --git a/Engine/Source/Runtime/Engine/Classes/Distributions/DistributionVector.h b/Engine/Source/Runtime/Engine/Classes/Distributions/DistributionVector.h
--- a/Engine/Source/Runtime/Engine/Classes/Distributions/DistributionVector.h
+++ b/Engine/Source/Runtime/Engine/Classes/Distributions/DistributionVector.h
@@ -2,7 +2,6 @@
 
 
 #pragma once
-#include "Distributions.h"
 #include "Distributions/Distribution.h"
 #include "DistributionVector.generated.h"
 
diff --git a/Engine/Source/Runtime/Engine/Classes/Engine/StreamableManager.h b/Engine/Source/Runtime/Engine/Classes/Engine/StreamableManager.h
--- a/Engine/Source/Runtime/Engine/Classes/Engine/StreamableManager.h
+++ b/Engine/Source/Runtime/Engine/Classes/Engine/StreamableManager.h
@@ -1,9 +1,6 @@
 // Copyright 1998-2014 Epic Games, Inc. All Rights Reserved.
 
 #pragma once
-
-#include "StringAssetReference.h"
-
 #include "StreamableManager.generated.h"
 
 /** Defines FStreamableDelegate delegate interface */
diff --git a/Engine/Source/Runtime/Engine/Classes/Engine/TimelineTemplate.h b/Engine/Source/Runtime/Engine/Classes/Engine/TimelineTemplate.h
--- a/Engine/Source/Runtime/Engine/Classes/Engine/TimelineTemplate.h
+++ b/Engine/Source/Runtime/Engine/Classes/Engine/TimelineTemplate.h
@@ -1,8 +1,6 @@
 // Copyright 1998-2014 Epic Games, Inc. All Rights Reserved.
 
 #pragma once
-
-#include "Components/TimelineComponent.h"
 #include "TimelineTemplate.generated.h"
 
 USTRUCT()
diff --git a/Engine/Source/Runtime/Engine/Classes/GameFramework/GameSession.h b/Engine/Source/Runtime/Engine/Classes/GameFramework/GameSession.h
--- a/Engine/Source/Runtime/Engine/Classes/GameFramework/GameSession.h
+++ b/Engine/Source/Runtime/Engine/Classes/GameFramework/GameSession.h
@@ -5,7 +5,7 @@
  */
 
 #pragma once
-#include "Runtime/Online/OnlineSubsystem/Public/OnlineSubsystemTypes.h"
+#include "GameFramework/Info.h"
 #include "GameSession.generated.h"
 
 class UWorld;
diff --git a/Engine/Source/Runtime/Engine/Classes/GameFramework/SpectatorPawnMovement.h b/Engine/Source/Runtime/Engine/Classes/GameFramework/SpectatorPawnMovement.h
--- a/Engine/Source/Runtime/Engine/Classes/GameFramework/SpectatorPawnMovement.h
+++ b/Engine/Source/Runtime/Engine/Classes/GameFramework/SpectatorPawnMovement.h
@@ -5,10 +5,7 @@
  */
 
 #pragma once
-
-#include "FloatingPawnMovement.h"
 #include "GameFramework/FloatingPawnMovement.h"
-
 #include "SpectatorPawnMovement.generated.h"
 
 UCLASS()
diff --git a/Engine/Source/Runtime/Engine/Classes/Particles/EmitterCameraLensEffectBase.h b/Engine/Source/Runtime/Engine/Classes/Particles/EmitterCameraLensEffectBase.h
--- a/Engine/Source/Runtime/Engine/Classes/Particles/EmitterCameraLensEffectBase.h
+++ b/Engine/Source/Runtime/Engine/Classes/Particles/EmitterCameraLensEffectBase.h
@@ -7,10 +7,7 @@
  */
 
 #pragma once
-
 #include "Particles/Emitter.h"
-#include "Emitter.h"
-
 #include "EmitterCameraLensEffectBase.generated.h"
 
 UCLASS(abstract, Blueprintable, MinimalAPI)
diff --git a/Engine/Source/Runtime/Engine/Classes/PhysicalMaterials/PhysicalMaterial.h b/Engine/Source/Runtime/Engine/Classes/PhysicalMaterials/PhysicalMaterial.h
--- a/Engine/Source/Runtime/Engine/Classes/PhysicalMaterials/PhysicalMaterial.h
+++ b/Engine/Source/Runtime/Engine/Classes/PhysicalMaterials/PhysicalMaterial.h
@@ -1,8 +1,6 @@
 // Copyright 1998-2014 Epic Games, Inc. All Rights Reserved.
 
 #pragma once
-
-#include "PhysxUserData.h"
 #include "PhysicsEngine/PhysicsSettings.h"
 #include "PhysicalMaterial.generated.h"
 
diff --git a/Engine/Source/Runtime/Engine/Classes/Sound/SoundMix.h b/Engine/Source/Runtime/Engine/Classes/Sound/SoundMix.h
--- a/Engine/Source/Runtime/Engine/Classes/Sound/SoundMix.h
+++ b/Engine/Source/Runtime/Engine/Classes/Sound/SoundMix.h
@@ -1,8 +1,6 @@
 // Copyright 1998-2014 Epic Games, Inc. All Rights Reserved.
 
 #pragma once
-
-#include "SoundClass.h"
 #include "Audio.h"
 #include "SoundMix.generated.h"
 
