Add checked version of get world context by PIE instance ID for consistency.

diff --git a/Engine/Source/Runtime/Engine/Classes/Engine/Engine.h b/Engine/Source/Runtime/Engine/Classes/Engine/Engine.h
--- a/Engine/Source/Runtime/Engine/Classes/Engine/Engine.h
+++ b/Engine/Source/Runtime/Engine/Classes/Engine/Engine.h
@@ -2370,13 +2370,14 @@ public:
 	FWorldContext* GetWorldContextFromPendingNetGame(const UPendingNetGame *InPendingNetGame);	
 	FWorldContext* GetWorldContextFromPendingNetGameNetDriver(const UNetDriver *InPendingNetGame);	
 	FWorldContext* GetWorldContextFromHandle(const FName WorldContextHandle);
-	FWorldContext* GetWorldContextFromPIEInstance(int32 PIEInstance);
+	FWorldContext* GetWorldContextFromPIEInstance(const int32 PIEInstance);
 
 	FWorldContext& GetWorldContextFromWorldChecked(UWorld * InWorld);
 	FWorldContext& GetWorldContextFromGameViewportChecked(const UGameViewportClient *InViewport);
 	FWorldContext& GetWorldContextFromPendingNetGameChecked(const UPendingNetGame *InPendingNetGame);	
 	FWorldContext& GetWorldContextFromPendingNetGameNetDriverChecked(const UNetDriver *InPendingNetGame);	
 	FWorldContext& GetWorldContextFromHandleChecked(const FName WorldContextHandle);
+	FWorldContext& GetWorldContextFromPIEInstanceChecked(const int32 PIEInstance);
 
 	const TIndirectArray<FWorldContext>& GetWorldContexts() { return WorldList;	}
 
diff --git a/Engine/Source/Runtime/Engine/Private/UnrealEngine.cpp b/Engine/Source/Runtime/Engine/Private/UnrealEngine.cpp
--- a/Engine/Source/Runtime/Engine/Private/UnrealEngine.cpp
+++ b/Engine/Source/Runtime/Engine/Private/UnrealEngine.cpp
@@ -9009,7 +9009,7 @@ FWorldContext& UEngine::GetWorldContextFromPendingNetGameNetDriverChecked(const
 	return HandleInvalidWorldContext();
 }
 
-FWorldContext* UEngine::GetWorldContextFromPIEInstance(int32 PIEInstance)
+FWorldContext* UEngine::GetWorldContextFromPIEInstance(const int32 PIEInstance)
 {
 	for (FWorldContext& WorldContext : WorldList)
 	{
@@ -9021,6 +9021,15 @@ FWorldContext* UEngine::GetWorldContextFromPIEInstance(int32 PIEInstance)
 	return NULL;
 }
 
+FWorldContext& UEngine::GetWorldContextFromPIEInstanceChecked(const int32 PIEInstance)
+{
+	if (FWorldContext* WorldContext = GetWorldContextFromPIEInstance(PIEInstance))
+	{
+		return *WorldContext;
+	}
+	return HandleInvalidWorldContext();
+}
+
 UPendingNetGame* UEngine::PendingNetGameFromWorld( UWorld* InWorld )
 {
 	return GetWorldContextFromWorldChecked(InWorld).PendingNetGame;
