Overloaded BeginPlay() method now calling into super so that CameraActor Blueprints get their BeginPlay events triggered.
(TTP #343000)

diff --git a/Engine/Source/Runtime/Engine/Private/Camera/CameraActor.cpp b/Engine/Source/Runtime/Engine/Private/Camera/CameraActor.cpp
--- a/Engine/Source/Runtime/Engine/Private/Camera/CameraActor.cpp
+++ b/Engine/Source/Runtime/Engine/Private/Camera/CameraActor.cpp
@@ -109,6 +109,8 @@ void ACameraActor::BeginPlay()
 			PC->SetViewTarget(this);
 		}
 	}
+
+	Super::BeginPlay();
 }
 
 #undef LOCTEXT_NAMESPACE
