Move per-frame simulation parameters for GPU particles in to the global constant buffer.

diff --git a/Engine/Shaders/ParticleSimulationShader.usf b/Engine/Shaders/ParticleSimulationShader.usf
--- a/Engine/Shaders/ParticleSimulationShader.usf
+++ b/Engine/Shaders/ParticleSimulationShader.usf
@@ -71,6 +71,12 @@ void VertexMain(
 	#define DEPTH_BUFFER_COLLISION 0
 #endif
 
+/** Position (XYZ) and squared radius (W) of the point attractor. */
+float4 PointAttractor;
+/** Position offset (XYZ) to add to particles and strength of the attractor (W). */
+float4 PositionOffsetAndAttractorStrength;
+/** Amount by which to scale bounds for collision purposes. */
+float2 LocalToWorldScale;
 /** Amount of time by which to simulate particles. */
 float DeltaSeconds;
 
@@ -194,9 +200,9 @@ float3 ComputeDrag(float3 Velocity, float DragCoefficient)
  */
 float3 ComputeAttractionForce(float3 Position)
 {
-	float3 PointLoc = SimulationPerFrame.PointAttractor.xyz;
-	float RadiusSq = SimulationPerFrame.PointAttractor.w;
-	float Strength = SimulationPerFrame.PointAttractorStrength;
+	float3 PointLoc = PointAttractor.xyz;
+	float RadiusSq = PointAttractor.w;
+	float Strength = PositionOffsetAndAttractorStrength.w;
 
 	float3 DirectionToPoint = PointLoc - Position + float3(0, 0, 0.0001f);
 	float DistSq = max(dot(DirectionToPoint,DirectionToPoint), RadiusSq);
@@ -373,7 +379,7 @@ void PixelMain(
 	// Compute the size of the sprite. Note it is (0,0) if the sprite is dead.
 	float2 InitialSize = RenderAttributeSample.xy;
 	float2 SizeScale = MiscCurve.xy;
-	float2 Size = InitialSize * SizeScale * SimulationPerFrame.LocalToWorldScale;
+	float2 Size = InitialSize * SizeScale * LocalToWorldScale;
 
 	// Compute the radius with which to perform collision checks.
 	float CollisionRadius = min(Size.x,Size.y) * Simulation.CollisionRadiusScale + Simulation.CollisionRadiusBias;
@@ -403,7 +409,7 @@ void PixelMain(
 	const float3 OrbitVelocity = ComputeOrbitVelocity(RelativeTime, OrbitRandom);
 
 	// Store the new position, time, and velocity for the particle.
-	OutPosition.xyz = NewPosition + OrbitVelocity * DeltaSeconds + SimulationPerFrame.PositionOffset;
+	OutPosition.xyz = NewPosition + OrbitVelocity * DeltaSeconds + PositionOffsetAndAttractorStrength.xyz;
 	OutPosition.w = RelativeTime;
 	OutVelocity.xyz = NewVelocity;
 	OutVelocity.w = TimeScale;
diff --git a/Engine/Source/Runtime/Engine/Private/Particles/ParticleGpuSimulation.cpp b/Engine/Source/Runtime/Engine/Private/Particles/ParticleGpuSimulation.cpp
--- a/Engine/Source/Runtime/Engine/Private/Particles/ParticleGpuSimulation.cpp
+++ b/Engine/Source/Runtime/Engine/Private/Particles/ParticleGpuSimulation.cpp
@@ -700,42 +700,64 @@ IMPLEMENT_UNIFORM_BUFFER_STRUCT(FParticleSimulationParameters,TEXT("Simulation")
 typedef TUniformBufferRef<FParticleSimulationParameters> FParticleSimulationBufferRef;
 
 /**
- * Uniform buffer to hold per-frame parameters for particle simulation.
+ * Per-frame parameters for particle simulation.
  */
-BEGIN_UNIFORM_BUFFER_STRUCT(FParticlePerFrameSimulationParameters,)
-	DECLARE_UNIFORM_BUFFER_STRUCT_MEMBER(FVector4, PointAttractor)
-	DECLARE_UNIFORM_BUFFER_STRUCT_MEMBER(float, PointAttractorStrength)
-	DECLARE_UNIFORM_BUFFER_STRUCT_MEMBER(FVector2D, LocalToWorldScale)
-	DECLARE_UNIFORM_BUFFER_STRUCT_MEMBER(FVector, PositionOffset)
-END_UNIFORM_BUFFER_STRUCT(FParticlePerFrameSimulationParameters)
-
-IMPLEMENT_UNIFORM_BUFFER_STRUCT(FParticlePerFrameSimulationParameters,TEXT("SimulationPerFrame"));
+struct FParticlePerFrameSimulationParameters
+{
+	/** Position (XYZ) and squared radius (W) of the point attractor. */
+	FVector4 PointAttractor;
+	/** Position offset (XYZ) to add to particles and strength of the attractor (W). */
+	FVector4 PositionOffsetAndAttractorStrength;
+	/** Amount by which to scale bounds for collision purposes. */
+	FVector2D LocalToWorldScale;
+	/** Amount of time by which to simulate particles. */
+	float DeltaSeconds;
 
-typedef TUniformBufferRef<FParticlePerFrameSimulationParameters> FParticlePerFrameSimulationBufferRef;
+	FParticlePerFrameSimulationParameters()
+		: PointAttractor(FVector::ZeroVector,0.0f)
+		, PositionOffsetAndAttractorStrength(FVector::ZeroVector,0.0f)
+		, LocalToWorldScale(1.0f, 1.0f)
+		, DeltaSeconds(0.0f)
+	{
+	}
+};
 
 /**
- * Default uniform buffer for per frame simulation parameters
+ * Per-frame shader parameters for particle simulation.
  */
-class FDefaultParticlePerFrameSimulationParametersUniformBuffer : public TUniformBuffer<FParticlePerFrameSimulationParameters>
+struct FParticlePerFrameSimulationShaderParameters
 {
-public:
+	FShaderParameter PointAttractor;
+	FShaderParameter PositionOffsetAndAttractorStrength;
+	FShaderParameter LocalToWorldScale;
+	FShaderParameter DeltaSeconds;
 
-	/** Default constructor. */
-	FDefaultParticlePerFrameSimulationParametersUniformBuffer()
+	void Bind(const FShaderParameterMap& ParameterMap)
 	{
-		FParticlePerFrameSimulationParameters DefaultParams;
-
-		DefaultParams.PointAttractor = FVector4(FVector::ZeroVector,0);
-		DefaultParams.PointAttractorStrength = 0;
-		DefaultParams.LocalToWorldScale = FVector2D(1.0f, 1.0f);
-		DefaultParams.PositionOffset = FVector::ZeroVector;
+		PointAttractor.Bind(ParameterMap,TEXT("PointAttractor"));
+		PositionOffsetAndAttractorStrength.Bind(ParameterMap,TEXT("PositionOffsetAndAttractorStrength"));
+		LocalToWorldScale.Bind(ParameterMap,TEXT("LocalToWorldScale"));
+		DeltaSeconds.Bind(ParameterMap,TEXT("DeltaSeconds"));
+	}
 
-		SetContents(DefaultParams);
+	template <typename ShaderRHIParamRef>
+	void Set(FRHICommandList& RHICmdList, const ShaderRHIParamRef& ShaderRHI, const FParticlePerFrameSimulationParameters& Parameters) const
+	{
+		SetShaderValue(RHICmdList,ShaderRHI,PointAttractor,Parameters.PointAttractor);
+		SetShaderValue(RHICmdList,ShaderRHI,PositionOffsetAndAttractorStrength,Parameters.PositionOffsetAndAttractorStrength);
+		SetShaderValue(RHICmdList,ShaderRHI,LocalToWorldScale,Parameters.LocalToWorldScale);
+		SetShaderValue(RHICmdList,ShaderRHI,DeltaSeconds,Parameters.DeltaSeconds);
 	}
 };
 
-/** Global default uniform buffer for per frame simulation parameters */
-TGlobalResource<FDefaultParticlePerFrameSimulationParametersUniformBuffer> GDefaultParticlePerFrameSimulationParametersUniformBuffer;
+FArchive& operator<<(FArchive& Ar, FParticlePerFrameSimulationShaderParameters& PerFrameParameters)
+{
+	Ar << PerFrameParameters.PointAttractor;
+	Ar << PerFrameParameters.PositionOffsetAndAttractorStrength;
+	Ar << PerFrameParameters.LocalToWorldScale;
+	Ar << PerFrameParameters.DeltaSeconds;
+	return Ar;
+}
 
 /**
  * Uniform buffer to hold parameters for vector fields sampled during particle
@@ -869,8 +891,8 @@ public:
 		SceneDepthTextureParameterSampler.Bind(Initializer.ParameterMap,TEXT("SceneDepthTextureSampler"));
 		GBufferATextureParameter.Bind(Initializer.ParameterMap,TEXT("GBufferATexture"));
 		GBufferATextureParameterSampler.Bind(Initializer.ParameterMap,TEXT("GBufferATextureSampler"));
-		DeltaSeconds.Bind(Initializer.ParameterMap,TEXT("DeltaSeconds"));
 		CollisionDepthBounds.Bind(Initializer.ParameterMap,TEXT("CollisionDepthBounds"));
+		PerFrameParameters.Bind(Initializer.ParameterMap);
 	}
 
 	/** Serialization. */
@@ -895,8 +917,8 @@ public:
 		Ar << SceneDepthTextureParameterSampler;
 		Ar << GBufferATextureParameter;
 		Ar << GBufferATextureParameterSampler;
-		Ar << DeltaSeconds;
 		Ar << CollisionDepthBounds;
+		Ar << PerFrameParameters;
 		return bShaderHasOutdatedParameters;
 	}
 
@@ -971,12 +993,11 @@ public:
 	/**
 	 * Set per-instance parameters for this shader.
 	 */
-	void SetInstanceParameters(FRHICommandList& RHICmdList, float InDeltaSeconds, FUniformBufferRHIParamRef UniformBuffer, FUniformBufferRHIParamRef PerFrameUniformBuffer)
+	void SetInstanceParameters(FRHICommandList& RHICmdList, FUniformBufferRHIParamRef UniformBuffer, const FParticlePerFrameSimulationParameters& InPerFrameParameters)
 	{
 		FPixelShaderRHIParamRef PixelShaderRHI = GetPixelShader();
-		SetShaderValue(RHICmdList, PixelShaderRHI, DeltaSeconds, InDeltaSeconds);
 		SetUniformBufferParameter(RHICmdList, PixelShaderRHI, GetUniformBufferParameter<FParticleSimulationParameters>(), UniformBuffer);
-		SetUniformBufferParameter(RHICmdList, PixelShaderRHI, GetUniformBufferParameter<FParticlePerFrameSimulationParameters>(), PerFrameUniformBuffer);
+		PerFrameParameters.Set(RHICmdList, PixelShaderRHI, InPerFrameParameters);
 	}
 
 	/**
@@ -1024,8 +1045,8 @@ private:
 	/** The GBufferATexture parameter for depth buffer collision. */
 	FShaderResourceParameter GBufferATextureParameter;
 	FShaderResourceParameter GBufferATextureParameterSampler;
-	/** The number of seconds by which to simulate particles. */
-	FShaderParameter DeltaSeconds;
+	/** Per frame simulation parameters. */
+	FParticlePerFrameSimulationShaderParameters PerFrameParameters;
 	/** Collision depth bounds. */
 	FShaderParameter CollisionDepthBounds;
 };
@@ -1184,15 +1205,28 @@ struct FSimulationCommandGPU
 	/** Uniform buffer containing simulation parameters. */
 	FUniformBufferRHIParamRef UniformBuffer;
 	/** Uniform buffer containing per-frame simulation parameters. */
-	FUniformBufferRHIRef PerFrameUniformBuffer;
+	FParticlePerFrameSimulationParameters PerFrameParameters;
 	/** Parameters to sample the local vector field for this simulation. */
 	FVectorFieldUniformBufferRef VectorFieldsUniformBuffer;
 	/** Vector field volume textures for this simulation. */
 	FTexture3DRHIParamRef VectorFieldTexturesRHI[MAX_VECTOR_FIELDS];
-	/** The number of seconds by which to simulate. */
-	float DeltaSeconds;
 	/** The number of tiles to simulate. */
 	int32 TileCount;
+
+	/** Initialization constructor. */
+	FSimulationCommandGPU(FParticleShaderParamRef InTileOffsetsRef, FUniformBufferRHIParamRef InUniformBuffer, const FParticlePerFrameSimulationParameters& InPerFrameParameters, FVectorFieldUniformBufferRef& InVectorFieldsUniformBuffer, int32 InTileCount)
+		: TileOffsetsRef(InTileOffsetsRef)
+		, UniformBuffer(InUniformBuffer)
+		, PerFrameParameters(InPerFrameParameters)
+		, VectorFieldsUniformBuffer(InVectorFieldsUniformBuffer)
+		, TileCount(InTileCount)
+	{
+		FTexture3DRHIParamRef BlackVolumeTextureRHI = (FTexture3DRHIParamRef)(FTextureRHIParamRef)GBlackVolumeTexture->TextureRHI;
+		for (int32 i = 0; i < MAX_VECTOR_FIELDS; ++i)
+		{
+			VectorFieldTexturesRHI[i] = BlackVolumeTextureRHI;
+		}
+	}
 };
 
 /**
@@ -1240,7 +1274,7 @@ void ExecuteSimulationCommands(
 	{
 		const FSimulationCommandGPU& Command = SimulationCommands[CommandIndex];
 		VertexShader->SetParameters(RHICmdList, Command.TileOffsetsRef);
-		PixelShader->SetInstanceParameters(RHICmdList, Command.DeltaSeconds, Command.UniformBuffer, Command.PerFrameUniformBuffer);
+		PixelShader->SetInstanceParameters(RHICmdList, Command.UniformBuffer, Command.PerFrameParameters);
 		PixelShader->SetVectorFieldParameters(
 			RHICmdList, 
 			Command.VectorFieldsUniformBuffer,
@@ -2233,9 +2267,7 @@ public:
 	/** The per-emitter simulation resources. */
 	const FParticleEmitterSimulationResources* EmitterSimulationResources;
 	/** The per-frame simulation uniform buffer. */
-	FParticlePerFrameSimulationBufferRef PerFrameSimulationUniformBuffer;
-	/** Pending update time. */
-	float PendingDeltaSeconds;
+	FParticlePerFrameSimulationParameters PerFrameSimulationParameters;
 	/** Bounds for particles in the simulation. */
 	FBox Bounds;
 
@@ -2274,7 +2306,6 @@ public:
 	/** Default constructor. */
 	FParticleSimulationGPU()
 		: EmitterSimulationResources(NULL)
-		, PendingDeltaSeconds(0.0f)
 		, VectorFieldVisualizationVertexFactory(NULL)
 		, SimulationIndex(INDEX_NONE)
 		, SimulationPhase(EParticleSimulatePhase::Main)
@@ -2477,8 +2508,6 @@ struct FGPUSpriteDynamicEmitterData : FDynamicEmitterDataBase
 	float LocalVectorFieldIntensity;
 	/** Vector field tightness. */
 	float LocalVectorFieldTightness;
-	/** Pending simulation time. */
-	float PendingDeltaSeconds;
 	/** Per-frame simulation parameters. */
 	FParticlePerFrameSimulationParameters PerFrameSimulationParameters;
 	/** Per-emitter parameters that may change*/
@@ -2504,7 +2533,6 @@ struct FGPUSpriteDynamicEmitterData : FDynamicEmitterDataBase
 		, Resources(NULL)
 		, Simulation(NULL)
 		, Material(NULL)
-		, PendingDeltaSeconds(0.0f)
 		, SortMode(PSORTMODE_None)
 		, bLocalVectorFieldTileX(false)
 		, bLocalVectorFieldTileY(false)
@@ -2516,7 +2544,11 @@ struct FGPUSpriteDynamicEmitterData : FDynamicEmitterDataBase
 	 * Called to create render thread resources.
 	 */
 	virtual void CreateRenderThreadResources(const FParticleSystemSceneProxy* InOwnerProxy)
-	{		
+	{
+		check(Simulation);
+
+		// Update the per-frame simulation parameters with those provided from the game thread.
+		Simulation->PerFrameSimulationParameters = PerFrameSimulationParameters;
 	}
 
 	/**
@@ -2548,12 +2580,6 @@ struct FGPUSpriteDynamicEmitterData : FDynamicEmitterDataBase
 			{
 				// Create view agnostic render data.  Do here rather than in CreateRenderThreadResources because in some cases Render can be called before CreateRenderThreadResources
 				{
-					Simulation->PendingDeltaSeconds = PendingDeltaSeconds;
-
-					// Create the per-frame uniform buffer.
-					Simulation->PerFrameSimulationUniformBuffer =
-						FParticlePerFrameSimulationBufferRef::CreateUniformBufferImmediate(PerFrameSimulationParameters, UniformBuffer_SingleFrame);
-
 					// Create per-emitter uniform buffer for dynamic parameters
 					DynamicUniformBuffer = FGPUSpriteEmitterDynamicUniformBufferRef::CreateUniformBufferImmediate(EmitterDynamicParameters, UniformBuffer_SingleFrame);
 
@@ -2967,10 +2993,9 @@ public:
 		{
 			FVector PointAttractorPosition = ComponentToWorld.TransformPosition(EmitterInfo.PointAttractorPosition);
 			DynamicData->PerFrameSimulationParameters.PointAttractor = FVector4(PointAttractorPosition, EmitterInfo.PointAttractorRadiusSq);
-			DynamicData->PerFrameSimulationParameters.PointAttractorStrength = PointAttractorStrength;
-			DynamicData->PendingDeltaSeconds = PendingDeltaSeconds;
+			DynamicData->PerFrameSimulationParameters.PositionOffsetAndAttractorStrength = FVector4(PositionOffsetThisTick, PointAttractorStrength);
 			DynamicData->PerFrameSimulationParameters.LocalToWorldScale = DynamicData->EmitterDynamicParameters.LocalToWorldScale;
-			DynamicData->PerFrameSimulationParameters.PositionOffset = PositionOffsetThisTick;
+			DynamicData->PerFrameSimulationParameters.DeltaSeconds = PendingDeltaSeconds;
 			Exchange(DynamicData->TilesToClear, TilesToClear);
 			Exchange(DynamicData->NewParticles, NewParticles);
 		}
@@ -4066,19 +4091,19 @@ void FFXSystem::SimulateGPUParticles(
 		if (Simulation->SimulationPhase == Phase
 			&& Simulation->TileVertexBuffer.AlignedTileCount > 0)
 		{
-			FSimulationCommandGPU* SimulationCommand = new(SimulationCommands) FSimulationCommandGPU;
-			SimulationCommand->TileOffsetsRef = Simulation->TileVertexBuffer.GetShaderParam();
-			SimulationCommand->UniformBuffer = Simulation->EmitterSimulationResources->SimulationUniformBuffer;
-			SimulationCommand->PerFrameUniformBuffer = Simulation->PerFrameSimulationUniformBuffer ? Simulation->PerFrameSimulationUniformBuffer : GDefaultParticlePerFrameSimulationParametersUniformBuffer;
-			SimulationCommand->DeltaSeconds = Simulation->PendingDeltaSeconds;
-			SimulationCommand->TileCount = Simulation->TileVertexBuffer.AlignedTileCount;
+			FSimulationCommandGPU* SimulationCommand = new(SimulationCommands) FSimulationCommandGPU(
+				Simulation->TileVertexBuffer.GetShaderParam(),
+				Simulation->EmitterSimulationResources->SimulationUniformBuffer,
+				Simulation->PerFrameSimulationParameters,
+				EmptyVectorFieldUniformBuffer,
+				Simulation->TileVertexBuffer.AlignedTileCount
+				);
 
 			// Determine which vector fields affect this simulation and build the appropriate parameters.
 			{
 				SCOPE_CYCLE_COUNTER(STAT_GPUParticleVFCullTime);
 				FVectorFieldUniformParameters VectorFieldParameters;
 				const FBox SimulationBounds = Simulation->Bounds;
-				FTexture3DRHIParamRef BlackVolumeTextureRHI = (FTexture3DRHIParamRef)(FTextureRHIParamRef)GBlackVolumeTexture->TextureRHI;
 
 				// Add the local vector field.
 				VectorFieldParameters.Count = 0;
@@ -4087,7 +4112,7 @@ void FFXSystem::SimulateGPUParticles(
 					const float Intensity = Simulation->LocalVectorField.Intensity * Simulation->LocalVectorField.Resource->Intensity;
 					if (FMath::Abs(Intensity) > 0.0f)
 					{
-						Simulation->LocalVectorField.Resource->Update(RHICmdList, Simulation->PendingDeltaSeconds);
+						Simulation->LocalVectorField.Resource->Update(RHICmdList, Simulation->PerFrameSimulationParameters.DeltaSeconds);
 						SimulationCommand->VectorFieldTexturesRHI[0] = Simulation->LocalVectorField.Resource->VolumeTextureRHI;
 						SetParametersForVectorField(VectorFieldParameters, &Simulation->LocalVectorField, /*EmitterScale=*/ 1.0f, /*EmitterTightness=*/ -1, VectorFieldParameters.Count++);
 					}
@@ -4113,15 +4138,7 @@ void FFXSystem::SimulateGPUParticles(
 				}
 
 				// Fill out any remaining vector field entries.
-				if (VectorFieldParameters.Count == 0)
-				{
-					SimulationCommand->VectorFieldsUniformBuffer = EmptyVectorFieldUniformBuffer;
-					for (int32 Index = 0; Index < MAX_VECTOR_FIELDS; ++Index)
-					{
-						SimulationCommand->VectorFieldTexturesRHI[Index] = BlackVolumeTextureRHI;
-					}
-				}
-				else
+				if (VectorFieldParameters.Count > 0)
 				{
 					int32 PadCount = VectorFieldParameters.Count;
 					while (PadCount < MAX_VECTOR_FIELDS)
@@ -4131,7 +4148,6 @@ void FFXSystem::SimulateGPUParticles(
 						VectorFieldParameters.VolumeToWorld[Index] = FMatrix::Identity;
 						VectorFieldParameters.VolumeSize[Index] = FVector(1.0f);
 						VectorFieldParameters.IntensityAndTightness[Index] = FVector2D::ZeroVector;
-						SimulationCommand->VectorFieldTexturesRHI[Index] = BlackVolumeTextureRHI;
 					}
 					SimulationCommand->VectorFieldsUniformBuffer = FVectorFieldUniformBufferRef::CreateUniformBufferImmediate(VectorFieldParameters, UniformBuffer_SingleFrame);
 				}
@@ -4145,13 +4161,8 @@ void FFXSystem::SimulateGPUParticles(
 			NewParticles.Append(Simulation->NewParticles);
 			Simulation->NewParticles.Reset();
 
-			// Reset pending simulation time.
-			Simulation->PendingDeltaSeconds = 0.0f;
-
-			// PerFrameSimulationUniformBuffer was created with UniformBuffer_SingleFrame.  Release it so we use
-			// the default uniform buffer if we simulate again before creating a fresh buffer (i.e. when the Emitter goes offscreen)
-			// this way the simulation shader won't read garbage.
-			Simulation->PerFrameSimulationUniformBuffer.SafeRelease();			
+			// Reset pending simulation time. This prevents an emitter from simulating twice if we don't get an update from the game thread, e.g. the component didn't tick last frame.
+			Simulation->PerFrameSimulationParameters.DeltaSeconds = 0.0f;
 		}
 	}
 
