Fixed wrong condition for a tickable level

diff --git a/Engine/Source/Runtime/Engine/Public/EngineUtils.h b/Engine/Source/Runtime/Engine/Public/EngineUtils.h
--- a/Engine/Source/Runtime/Engine/Public/EngineUtils.h
+++ b/Engine/Source/Runtime/Engine/Public/EngineUtils.h
@@ -288,7 +288,7 @@ struct FTickableLevelFilter
 	 */
 	static bool CanIterateLevel(ULevel* Level)
 	{
-		return !Level->HasVisibilityRequestPending() || Level->bIsAssociatingLevel;
+		return Level->bIsVisible;
 	}
 };
 
