Header cleanup: RecastNavMeshGenerator.h, CaputreSource.h, CapturePin.h, FXSystemPrivate.h, CollisionDebugDrawing.h.

diff --git a/Engine/Plugins/Media/MediaAssetEditor/Source/MediaAssetEditor/Private/AssetTools/MediaTextureActions.cpp b/Engine/Plugins/Media/MediaAssetEditor/Source/MediaAssetEditor/Private/AssetTools/MediaTextureActions.cpp
--- a/Engine/Plugins/Media/MediaAssetEditor/Source/MediaAssetEditor/Private/AssetTools/MediaTextureActions.cpp
+++ b/Engine/Plugins/Media/MediaAssetEditor/Source/MediaAssetEditor/Private/AssetTools/MediaTextureActions.cpp
@@ -1,6 +1,7 @@
 // Copyright 1998-2014 Epic Games, Inc. All Rights Reserved.
 
 #include "MediaAssetEditorPrivatePCH.h"
+#include "Editor/ContentBrowser/Public/ContentBrowserModule.h"
 #include "TextureEditor.h"
 
 
diff --git a/Engine/Plugins/Media/MediaAssetEditor/Source/MediaAssetEditor/Private/MediaAssetEditorModule.cpp b/Engine/Plugins/Media/MediaAssetEditor/Source/MediaAssetEditor/Private/MediaAssetEditorModule.cpp
--- a/Engine/Plugins/Media/MediaAssetEditor/Source/MediaAssetEditor/Private/MediaAssetEditorModule.cpp
+++ b/Engine/Plugins/Media/MediaAssetEditor/Source/MediaAssetEditor/Private/MediaAssetEditorModule.cpp
@@ -7,9 +7,6 @@
 
 #define LOCTEXT_NAMESPACE "FMediaAssetEditorModule"
 
-static const FName MediaAssetEditorAppIdentifier("MediaAssetEditorApp");
-
-
 /**
  * Implements the MediaAssetEditor module.
  */
diff --git a/Engine/Plugins/Media/MediaAssetEditor/Source/MediaAssetEditor/Private/MediaAssetEditorToolkit.cpp b/Engine/Plugins/Media/MediaAssetEditor/Source/MediaAssetEditor/Private/MediaAssetEditorToolkit.cpp
--- a/Engine/Plugins/Media/MediaAssetEditor/Source/MediaAssetEditor/Private/MediaAssetEditorToolkit.cpp
+++ b/Engine/Plugins/Media/MediaAssetEditor/Source/MediaAssetEditor/Private/MediaAssetEditorToolkit.cpp
@@ -13,7 +13,7 @@ DEFINE_LOG_CATEGORY_STATIC(LogMediaAssetEditor, Log, All);
 
 static const FName DetailsTabId("Details");
 static const FName PlayerTabId("Player");
-
+static const FName MediaAssetEditorAppIdentifier("MediaAssetEditorApp");
 
 /* FMediaAssetEditorToolkit structors
  *****************************************************************************/
diff --git a/Engine/Source/Editor/DetailCustomizations/Private/MediaTextureCustomization.cpp b/Engine/Source/Editor/DetailCustomizations/Private/MediaTextureCustomization.cpp
--- a/Engine/Source/Editor/DetailCustomizations/Private/MediaTextureCustomization.cpp
+++ b/Engine/Source/Editor/DetailCustomizations/Private/MediaTextureCustomization.cpp
@@ -1,6 +1,7 @@
 // Copyright 1998-2014 Epic Games, Inc. All Rights Reserved.
 
 #include "DetailCustomizationsPrivatePCH.h"
+#include "MediaAsset.h"
 #include "MediaTextureCustomization.h"
 
 
diff --git a/Engine/Source/Runtime/Engine/Private/AI/Navigation/RecastNavMeshGenerator.cpp b/Engine/Source/Runtime/Engine/Private/AI/Navigation/RecastNavMeshGenerator.cpp
--- a/Engine/Source/Runtime/Engine/Private/AI/Navigation/RecastNavMeshGenerator.cpp
+++ b/Engine/Source/Runtime/Engine/Private/AI/Navigation/RecastNavMeshGenerator.cpp
@@ -15,6 +15,7 @@
 
 // recast includes
 #include "Recast.h"
+#include "DetourCommon.h"
 #include "DetourNavMeshBuilder.h"
 #include "DetourNavMeshQuery.h"
 #include "RecastAlloc.h"
diff --git a/Engine/Source/Runtime/Engine/Private/AI/Navigation/RecastNavMeshGenerator.h b/Engine/Source/Runtime/Engine/Private/AI/Navigation/RecastNavMeshGenerator.h
--- a/Engine/Source/Runtime/Engine/Private/AI/Navigation/RecastNavMeshGenerator.h
+++ b/Engine/Source/Runtime/Engine/Private/AI/Navigation/RecastNavMeshGenerator.h
@@ -4,9 +4,7 @@
 #if WITH_RECAST
 
 #include "Recast.h"
-#include "DetourCommon.h"
 #include "DetourNavMesh.h"
-#include "DetourNavMeshBuilder.h"
 #include "AI/NavDataGenerator.h"
 #include "AI/NavigationModifier.h"
 
diff --git a/Engine/Source/Runtime/Engine/Private/AVIWriter.cpp b/Engine/Source/Runtime/Engine/Private/AVIWriter.cpp
--- a/Engine/Source/Runtime/Engine/Private/AVIWriter.cpp
+++ b/Engine/Source/Runtime/Engine/Private/AVIWriter.cpp
@@ -5,6 +5,7 @@
 =============================================================================*/
 #include "EnginePrivate.h"
 #include "AVIWriter.h"
+#include "CapturePin.h"
 
 DEFINE_LOG_CATEGORY_STATIC(LogAVIWriter, Log, All);
 
diff --git a/Engine/Source/Runtime/Engine/Private/CaptureSource.cpp b/Engine/Source/Runtime/Engine/Private/CaptureSource.cpp
--- a/Engine/Source/Runtime/Engine/Private/CaptureSource.cpp
+++ b/Engine/Source/Runtime/Engine/Private/CaptureSource.cpp
@@ -5,9 +5,10 @@
 =============================================================================*/
 #include "EnginePrivate.h"
 
-
+#include "CapturePin.h"
 #include "CaptureSource.h"
 
+
 #if PLATFORM_WINDOWS && !UE_BUILD_MINIMAL
 
 FCaptureSource::FCaptureSource(IUnknown *pUnk, HRESULT *phr)
diff --git a/Engine/Source/Runtime/Engine/Private/CaptureSource.h b/Engine/Source/Runtime/Engine/Private/CaptureSource.h
--- a/Engine/Source/Runtime/Engine/Private/CaptureSource.h
+++ b/Engine/Source/Runtime/Engine/Private/CaptureSource.h
@@ -8,8 +8,14 @@
 #define _CAPTURESOURCE_HEADER_
 
 #if PLATFORM_WINDOWS && !UE_BUILD_MINIMAL
-
-#include "CapturePin.h"
+#pragma warning(disable : 4263) // 'function' : member function does not override any base class virtual member function
+#pragma warning(disable : 4264) // 'virtual_function' : no override available for virtual member function from base 
+#include "AllowWindowsPlatformTypes.h"
+#include <streams.h>
+#include "HideWindowsPlatformTypes.h"
+#pragma warning(default : 4263) // 'function' : member function does not override any base class virtual member function
+#pragma warning(default : 4264) // 'virtual_function' : no override available for virtual member function from base
+class FCapturePin;
 
 // {9A80E195-3BBA-4821-B18B-21BB496F80F8}
 DEFINE_GUID(CLSID_CaptureSource, 
diff --git a/Engine/Source/Runtime/Engine/Private/Collision/CollisionDebugDrawing.h b/Engine/Source/Runtime/Engine/Private/Collision/CollisionDebugDrawing.h
--- a/Engine/Source/Runtime/Engine/Private/Collision/CollisionDebugDrawing.h
+++ b/Engine/Source/Runtime/Engine/Private/Collision/CollisionDebugDrawing.h
@@ -4,8 +4,6 @@
 
 #pragma once 
 
-#include "CollisionDebugDrawingPublic.h"
-
 #if WITH_PHYSX
 /** Draw PhysX geom with overlaps */
 void DrawGeomOverlaps(const UWorld* InWorld, const PxGeometry& PGeom, const PxTransform& PGeomPose, TArray<struct FOverlapResult>& Overlaps, float Lifetime);
diff --git a/Engine/Source/Runtime/Engine/Private/Collision/PhysXCollision.cpp b/Engine/Source/Runtime/Engine/Private/Collision/PhysXCollision.cpp
--- a/Engine/Source/Runtime/Engine/Private/Collision/PhysXCollision.cpp
+++ b/Engine/Source/Runtime/Engine/Private/Collision/PhysXCollision.cpp
@@ -3,6 +3,7 @@
 #include "EnginePrivate.h"
 #include "PhysicsPublic.h"
 #include "Collision.h"
+#include "CollisionDebugDrawingPublic.h"
 
 #if WITH_BOX2D
 #include "../PhysicsEngine2D/Box2DIntegration.h"
diff --git a/Engine/Source/Runtime/Engine/Private/PrimitiveComponent.cpp b/Engine/Source/Runtime/Engine/Private/PrimitiveComponent.cpp
--- a/Engine/Source/Runtime/Engine/Private/PrimitiveComponent.cpp
+++ b/Engine/Source/Runtime/Engine/Private/PrimitiveComponent.cpp
@@ -19,6 +19,7 @@
 #include "MessageLog.h"
 #include "UObjectToken.h"
 #include "MapErrors.h"
+#include "CollisionDebugDrawingPublic.h"
 
 #define LOCTEXT_NAMESPACE "PrimitiveComponent"
 
diff --git a/Engine/Source/Runtime/Engine/Private/SkeletalMeshComponentPhysics.cpp b/Engine/Source/Runtime/Engine/Private/SkeletalMeshComponentPhysics.cpp
--- a/Engine/Source/Runtime/Engine/Private/SkeletalMeshComponentPhysics.cpp
+++ b/Engine/Source/Runtime/Engine/Private/SkeletalMeshComponentPhysics.cpp
@@ -6,6 +6,7 @@
 #include "SkeletalRenderPublic.h"
 
 #include "MessageLog.h"
+#include "CollisionDebugDrawingPublic.h"
 
 #if WITH_PHYSX
 	#include "PhysicsEngine/PhysXSupport.h"
diff --git a/Engine/Source/Runtime/MediaAssets/Private/Assets/MediaTextureResource.cpp b/Engine/Source/Runtime/MediaAssets/Private/Assets/MediaTextureResource.cpp
--- a/Engine/Source/Runtime/MediaAssets/Private/Assets/MediaTextureResource.cpp
+++ b/Engine/Source/Runtime/MediaAssets/Private/Assets/MediaTextureResource.cpp
@@ -1,7 +1,7 @@
 // Copyright 1998-2014 Epic Games, Inc. All Rights Reserved.
 
 #include "MediaAssetsPrivatePCH.h"
-
+#include "Assets/MediaTextureResource.h"
 
 /* FMediaTextureResource structors
  *****************************************************************************/
diff --git a/Engine/Source/Runtime/MediaAssets/Public/MediaSoundWave.h b/Engine/Source/Runtime/MediaAssets/Public/MediaSoundWave.h
--- a/Engine/Source/Runtime/MediaAssets/Public/MediaSoundWave.h
+++ b/Engine/Source/Runtime/MediaAssets/Public/MediaSoundWave.h
@@ -6,6 +6,7 @@
 #include "Sound/SoundWave.h"
 #include "MediaSoundWave.generated.h"
 
+class UMediaAsset;
 
 /**
  * Implements a playable sound object for audio streams from media assets.
diff --git a/Engine/Source/Runtime/MediaAssets/Public/MediaTexture.h b/Engine/Source/Runtime/MediaAssets/Public/MediaTexture.h
--- a/Engine/Source/Runtime/MediaAssets/Public/MediaTexture.h
+++ b/Engine/Source/Runtime/MediaAssets/Public/MediaTexture.h
@@ -5,6 +5,7 @@
 #include "MediaSampleBuffer.h"
 #include "MediaTexture.generated.h"
 
+class UMediaAsset;
 
 /**
  * Implements a texture for rendering video tracks from media assets.
