Stationary primitives use static texture streaming stategy

diff --git a/Engine/Source/Runtime/Engine/Private/Level.cpp b/Engine/Source/Runtime/Engine/Private/Level.cpp
--- a/Engine/Source/Runtime/Engine/Private/Level.cpp
+++ b/Engine/Source/Runtime/Engine/Private/Level.cpp
@@ -1332,6 +1332,7 @@ void ULevel::BuildStreamingData(UTexture2D* UpdateSpecificTextureOnly/*=NULL*/)
 				const bool bIsFoliage					= Owner && Owner->IsA(AInstancedFoliageActor::StaticClass()) && Primitive->IsA(UInstancedStaticMeshComponent::StaticClass()); 
 				const bool bIsStatic					= Owner == NULL 
 															|| Primitive->Mobility == EComponentMobility::Static 
+															|| Primitive->Mobility == EComponentMobility::Stationary
 															|| bIsFoliage; // treat Foliage components as static, regardless of mobility settings
 
 				TArray<FStreamingTexturePrimitiveInfo> PrimitiveStreamingTextures;
