Fixed TTP #342981 "CRITICAL: REGRESSION: CRASH: QAGame -game -featureleveles2 Assertion failed: (Index >= 0) & (Index < ArrayNum)"

diff --git a/Engine/Source/Developer/GameplayDebugger/Private/GameplayDebuggingComponent.cpp b/Engine/Source/Developer/GameplayDebugger/Private/GameplayDebuggingComponent.cpp
--- a/Engine/Source/Developer/GameplayDebugger/Private/GameplayDebuggingComponent.cpp
+++ b/Engine/Source/Developer/GameplayDebugger/Private/GameplayDebuggingComponent.cpp
@@ -1132,16 +1132,15 @@ FPrimitiveSceneProxy* UGameplayDebuggingComponent::CreateSceneProxy()
 	//const APlayerController* MyPC = MyPawn ? Cast<APlayerController>(MyPawn->Controller) : NULL;
 	if (ShouldReplicateData(EAIDebugDrawDataView::NavMesh) && World && World->GetNetMode() != NM_DedicatedServer)
 	{
-		TSharedPtr<struct FNavMeshSceneProxyData, ESPMode::ThreadSafe> NewNavmeshRenderData = MakeShareable(new FNavMeshSceneProxyData());
-		NewNavmeshRenderData->bNeedsNewData = false;
-		NewNavmeshRenderData->bEnableDrawing = false;
-		PrepareNavMeshData(NewNavmeshRenderData.Get());
-		if (NewNavmeshRenderData.IsValid())
-		{
-			NavMeshBounds = NewNavmeshRenderData->Bounds;
-			CompositeProxy = CompositeProxy ? CompositeProxy : (new FDebugRenderSceneCompositeProxy(this));
-			CompositeProxy->AddChild(new FRecastRenderingSceneProxy(this, NewNavmeshRenderData, FSimpleDelegateGraphTask::FDelegate(), true));
-		}
+		FNavMeshSceneProxyData NewNavmeshRenderData;
+		NewNavmeshRenderData.Reset();
+		NewNavmeshRenderData.bNeedsNewData = false;
+		NewNavmeshRenderData.bEnableDrawing = false;
+		PrepareNavMeshData(&NewNavmeshRenderData);
+
+		NavMeshBounds = NewNavmeshRenderData.Bounds;
+		CompositeProxy = CompositeProxy ? CompositeProxy : (new FDebugRenderSceneCompositeProxy(this));
+		CompositeProxy->AddChild(new FRecastRenderingSceneProxy(this, &NewNavmeshRenderData, true));
 	}
 #endif
 
diff --git a/Engine/Source/Runtime/Engine/Classes/AI/Navigation/NavMeshRenderingComponent.h b/Engine/Source/Runtime/Engine/Classes/AI/Navigation/NavMeshRenderingComponent.h
--- a/Engine/Source/Runtime/Engine/Classes/AI/Navigation/NavMeshRenderingComponent.h
+++ b/Engine/Source/Runtime/Engine/Classes/AI/Navigation/NavMeshRenderingComponent.h
@@ -22,7 +22,4 @@ class ENGINE_API UNavMeshRenderingComponent : public UPrimitiveComponent
 	// End USceneComponent Interface
 
 	void GatherData(struct FNavMeshSceneProxyData*) const;
-
-protected:
-	TSharedPtr<struct FNavMeshSceneProxyData, ESPMode::ThreadSafe>	ProxyData;
 };
diff --git a/Engine/Source/Runtime/Engine/Private/NavMeshRenderingComponent.cpp b/Engine/Source/Runtime/Engine/Private/NavMeshRenderingComponent.cpp
--- a/Engine/Source/Runtime/Engine/Private/NavMeshRenderingComponent.cpp
+++ b/Engine/Source/Runtime/Engine/Private/NavMeshRenderingComponent.cpp
@@ -27,14 +27,12 @@ FPrimitiveSceneProxy* UNavMeshRenderingComponent::CreateSceneProxy()
 	FPrimitiveSceneProxy* SceneProxy = NULL;
 	if (IsVisible())
 	{
-		if (!ProxyData.IsValid())
-		{
-			ProxyData = MakeShareable(new FNavMeshSceneProxyData());
-		}
-		GatherData(ProxyData.Get());
-		ProxyData->bNeedsNewData = false;
+		FNavMeshSceneProxyData ProxyData;
+		ProxyData.Reset();
+		GatherData(&ProxyData);
 
-		SceneProxy = new FRecastRenderingSceneProxy(this, ProxyData);
+		SceneProxy = new FRecastRenderingSceneProxy(this, &ProxyData);
+		ProxyData.Reset();
 	}
 	return SceneProxy;
 #else
@@ -86,6 +84,7 @@ void UNavMeshRenderingComponent::GatherData(struct FNavMeshSceneProxyData* Curre
 
 	CurrentData->Reset();
 	CurrentData->bEnableDrawing = NavMesh->bEnableDrawing;
+	CurrentData->bNeedsNewData = false;
 
 	if (CurrentData && NavMesh && NavMesh->bEnableDrawing)
 	{
diff --git a/Engine/Source/Runtime/Engine/Private/NavMeshRenderingHelpers.h b/Engine/Source/Runtime/Engine/Private/NavMeshRenderingHelpers.h
--- a/Engine/Source/Runtime/Engine/Private/NavMeshRenderingHelpers.h
+++ b/Engine/Source/Runtime/Engine/Private/NavMeshRenderingHelpers.h
@@ -102,7 +102,6 @@ struct FNavMeshSceneProxyData : public TSharedFromThis<FNavMeshSceneProxyData, E
 	FBox Bounds;
 
 	FVector NavMeshDrawOffset;
-	FCriticalSection CriticalSection;
 	uint32 bDataGathered : 1;
 	uint32 bSkipDistanceCheck : 1;
 	uint32 bDrawClusters : 1;
@@ -150,20 +149,22 @@ FORCEINLINE bool LineInCorrectDistance(const FVector& Start, const FVector& End,
 class FRecastRenderingSceneProxy : public FDebugRenderSceneProxy
 {
 public:
-	FRecastRenderingSceneProxy(const UPrimitiveComponent* InComponent, TSharedPtr<FNavMeshSceneProxyData, ESPMode::ThreadSafe> InProxyData, 
-		const FSimpleDelegateGraphTask::FDelegate& InTaskDeletegate = FSimpleDelegateGraphTask::FDelegate(), bool ForceToRender = false) 
+	FRecastRenderingSceneProxy(const UPrimitiveComponent* InComponent, FNavMeshSceneProxyData* InProxyData,  bool ForceToRender = false) 
 		: FDebugRenderSceneProxy(InComponent)
-		, ProxyData(InProxyData)
 		, bRequestedData(false)
 		, bForceRendering(ForceToRender)
 	{
-		TaskDeletegate = InTaskDeletegate;
+		ProxyData.Reset();
+		if (InProxyData)
+		{
+			ProxyData = *InProxyData;
+		}
 		RenderingComponent = Cast<const UNavMeshRenderingComponent>(InComponent);
 	}
 
 	virtual ~FRecastRenderingSceneProxy() 
 	{
-		ProxyData = NULL;
+		ProxyData.Reset();
 	}
 
 
@@ -198,35 +199,35 @@ public:
 	{
 		QUICK_SCOPE_CYCLE_COUNTER( STAT_RecastRenderingSceneProxy_DrawDynamicElements );
 
-		if (ProxyData.IsValid() == false || !ProxyData->bEnableDrawing) //check if we have any data to render
+		if (!ProxyData.bEnableDrawing) //check if we have any data to render
 		{
 			return;
 		}
 
-		ProxyData->bSkipDistanceCheck = GIsEditor && (GEngine->GetDebugLocalPlayer() == NULL);
+		ProxyData.bSkipDistanceCheck = GIsEditor && (GEngine->GetDebugLocalPlayer() == NULL);
 
 		FVector const PosX(1.f,0,0);
 		FVector const PosY(0,1.f,0);
 		FVector const PosZ(0,0,1.f);
 
-		const TArray<FVector>& MeshVerts = ProxyData->NavMeshGeometry.MeshVerts;
+		const TArray<FVector>& MeshVerts = ProxyData.NavMeshGeometry.MeshVerts;
 
 		// Draw Mesh
-		for (int32 Index = 0; Index < ProxyData->MeshBuilders.Num(); ++Index)
+		for (int32 Index = 0; Index < ProxyData.MeshBuilders.Num(); ++Index)
 		{
-			const FColoredMaterialRenderProxy *MeshColorInstance = new(FMemStack::Get()) FColoredMaterialRenderProxy(GEngine->DebugMeshMaterial->GetRenderProxy(false), ProxyData->MeshBuilders[Index].ClusterColor);						
+			const FColoredMaterialRenderProxy *MeshColorInstance = new(FMemStack::Get()) FColoredMaterialRenderProxy(GEngine->DebugMeshMaterial->GetRenderProxy(false), ProxyData.MeshBuilders[Index].ClusterColor);						
 			FDynamicMeshBuilder	MeshBuilder;
-			MeshBuilder.AddVertices( ProxyData->MeshBuilders[Index].Vertices );
-			MeshBuilder.AddTriangles( ProxyData->MeshBuilders[Index].Indices );
+			MeshBuilder.AddVertices( ProxyData.MeshBuilders[Index].Vertices );
+			MeshBuilder.AddTriangles( ProxyData.MeshBuilders[Index].Indices );
 			MeshBuilder.Draw(PDI, FMatrix::Identity, MeshColorInstance, GetDepthPriorityGroup(View));
 		}
 
 		FHitProxyId HitProxyId;
 		FBatchedElements BatchedElements;
-		int32 Num = ProxyData->NavMeshEdgeLines.Num();
+		int32 Num = ProxyData.NavMeshEdgeLines.Num();
 		for (int32 Index = 0; Index < Num; ++Index)
 		{
-			const FDebugLine &Line = ProxyData->NavMeshEdgeLines[Index];
+			const FDebugLine &Line = ProxyData.NavMeshEdgeLines[Index];
 			if( LineInView(Line.Start,Line.End,View,false) )
 			{
 				if (LineInCorrectDistance(Line.Start,Line.End,View))
@@ -240,10 +241,10 @@ public:
 			}
 		}
 
-		Num = ProxyData->ClusterLinkLines.Num();
+		Num = ProxyData.ClusterLinkLines.Num();
 		for (int32 Index = 0; Index < Num; ++Index)
 		{
-			const FDebugLine &Line = ProxyData->ClusterLinkLines[Index];
+			const FDebugLine &Line = ProxyData.ClusterLinkLines[Index];
 			if( LineInView(Line.Start,Line.End,View,false) )
 			{
 				if (LineInCorrectDistance(Line.Start,Line.End,View))
@@ -257,10 +258,10 @@ public:
 			}
 		}
 
-		Num = ProxyData->TileEdgeLines.Num();
+		Num = ProxyData.TileEdgeLines.Num();
 		for (int32 Index = 0; Index < Num; ++Index)
 		{
-			const FDebugLine &Line = ProxyData->TileEdgeLines[Index];
+			const FDebugLine &Line = ProxyData.TileEdgeLines[Index];
 			if( LineInView(Line.Start,Line.End,View,false) )
 			{
 				if (LineInCorrectDistance(Line.Start,Line.End,View))
@@ -274,10 +275,10 @@ public:
 			}
 		}
 
-		Num = ProxyData->NavLinkLines.Num();
+		Num = ProxyData.NavLinkLines.Num();
 		for (int32 Index = 0; Index < Num; ++Index)
 		{
-			const FDebugLine &Line = ProxyData->NavLinkLines[Index];
+			const FDebugLine &Line = ProxyData.NavLinkLines[Index];
 			if( LineInView(Line.Start,Line.End,View,false) )
 			{
 				if (LineInCorrectDistance(Line.Start,Line.End,View))
@@ -307,7 +308,7 @@ public:
 			DepthTexture
 			);
 
-		ProxyData->BatchedElements.Draw(
+		ProxyData.BatchedElements.Draw(
 			RHICmdList,
 			bNeedToSwitchVerticalAxis,
 			View->ViewProjectionMatrix,
@@ -327,7 +328,7 @@ public:
 
 	void DrawDebugLabels(UCanvas* Canvas, APlayerController*)
 	{
-		if (ProxyData->bNeedsNewData == true || ProxyData->bEnableDrawing == false)
+		if (ProxyData.bNeedsNewData == true || ProxyData.bEnableDrawing == false)
 		{
 			return;
 		}
@@ -336,8 +337,8 @@ public:
 		Canvas->SetDrawColor(FColor::White);
 		const FSceneView* View = Canvas->SceneView;
 		UFont* Font = GEngine->GetSmallFont();
-		const FNavMeshSceneProxyData::FDebugText* DebugText = ProxyData->DebugLabels.GetTypedData();
-		for (int32 i = 0 ; i < ProxyData->DebugLabels.Num(); ++i, ++DebugText)
+		const FNavMeshSceneProxyData::FDebugText* DebugText = ProxyData.DebugLabels.GetTypedData();
+		for (int32 i = 0 ; i < ProxyData.DebugLabels.Num(); ++i, ++DebugText)
 		{
 			if (PointInView(DebugText->Location, View))
 			{
@@ -363,18 +364,17 @@ public:
 	uint32 GetAllocatedSize( void ) const 
 	{ 
 		return FDebugRenderSceneProxy::GetAllocatedSize() + 
-			(sizeof(FNavMeshSceneProxyData) + ProxyData->NavMeshGeometry.GetAllocatedSize()
-			+ ProxyData->TileEdgeLines.GetAllocatedSize() + ProxyData->NavMeshEdgeLines.GetAllocatedSize() + ProxyData->NavLinkLines.GetAllocatedSize()
-			+ ProxyData->PathCollidingGeomIndices.GetAllocatedSize() + ProxyData->PathCollidingGeomVerts.GetAllocatedSize()
-			+ ProxyData->DebugLabels.GetAllocatedSize() + ProxyData->ClusterLinkLines.GetAllocatedSize() + ProxyData->MeshBuilders.GetAllocatedSize()
-			+ ProxyData->BatchedElements.GetAllocatedSize());
+			(sizeof(FNavMeshSceneProxyData) + ProxyData.NavMeshGeometry.GetAllocatedSize()
+			+ ProxyData.TileEdgeLines.GetAllocatedSize() + ProxyData.NavMeshEdgeLines.GetAllocatedSize() + ProxyData.NavLinkLines.GetAllocatedSize()
+			+ ProxyData.PathCollidingGeomIndices.GetAllocatedSize() + ProxyData.PathCollidingGeomVerts.GetAllocatedSize()
+			+ ProxyData.DebugLabels.GetAllocatedSize() + ProxyData.ClusterLinkLines.GetAllocatedSize() + ProxyData.MeshBuilders.GetAllocatedSize()
+			+ ProxyData.BatchedElements.GetAllocatedSize());
 	}
 
 private:
-	TSharedPtr<FNavMeshSceneProxyData, ESPMode::ThreadSafe>	ProxyData;
+	FNavMeshSceneProxyData ProxyData;
 	FDebugDrawDelegate DebugTextDrawingDelegate;
 	TWeakObjectPtr<UNavMeshRenderingComponent> RenderingComponent;
-	FSimpleDelegateGraphTask::FDelegate TaskDeletegate;
 	uint32 bRequestedData : 1;
 	uint32 bForceRendering : 1;
 };
