Fix Foliage Navmesh generation
Fix potential GC crash with foliage
Workaround for bugs in navigation system when using PerComponentNavigation actors
Workaround for race condition in FRecastDebugGeometry::GetAllocatedSize()

diff --git a/Engine/Source/Runtime/Engine/Classes/Foliage/InstancedFoliageActor.h b/Engine/Source/Runtime/Engine/Classes/Foliage/InstancedFoliageActor.h
--- a/Engine/Source/Runtime/Engine/Classes/Foliage/InstancedFoliageActor.h
+++ b/Engine/Source/Runtime/Engine/Classes/Foliage/InstancedFoliageActor.h
@@ -13,7 +13,7 @@ struct FFoliageInstancePlacementInfo;
 struct FFoliageMeshInfo;
 
 UCLASS(notplaceable, hidecategories = Object, MinimalAPI, NotBlueprintable)
-class AInstancedFoliageActor : public AActor
+class AInstancedFoliageActor : public AActor, public INavRelevantActorInterface
 {
 	GENERATED_UCLASS_BODY()
 
@@ -106,7 +106,8 @@ public:
 	// Get the instanced foliage actor for the specified streaming level. Never creates a new IFA.
 	static ENGINE_API AInstancedFoliageActor* GetInstancedFoliageActorForLevel(ULevel* Level);
 #endif	//WITH_EDITOR
-};
-
-
 
+	// Begin INavRelevantActorInterface Interface
+	virtual bool DoesSupplyPerComponentNavigationCollision() const override { return true; }
+	// End INavRelevantActorInterface Interface
+};
diff --git a/Engine/Source/Runtime/Engine/Private/AI/Navigation/NavigationSystem.cpp b/Engine/Source/Runtime/Engine/Private/AI/Navigation/NavigationSystem.cpp
--- a/Engine/Source/Runtime/Engine/Private/AI/Navigation/NavigationSystem.cpp
+++ b/Engine/Source/Runtime/Engine/Private/AI/Navigation/NavigationSystem.cpp
@@ -1811,6 +1811,28 @@ FSetElementId UNavigationSystem::RegisterNavOctreeElement(class UObject* Element
 		return SetId;
 	}
 
+	//
+	// HACK: This is necessary to work around the case where a PerComponentNavigation actor has its components registered,
+	// which registers the actor, and also UpdateNavOctree(Component) is called in the same frame.
+	//
+	// A better solution would be to make the RegisterNavigationRelevantActor function add only individual components if
+	// the actor is using PerComponentNavigation.
+	//
+	UPrimitiveComponent* Component = Cast<UPrimitiveComponent>(ElementOwner);
+	if (Component && Component->GetOwner())
+	{
+		INavRelevantActorInterface* NavRelevantActor = InterfaceCast<INavRelevantActorInterface>(Component->GetOwner());
+		const bool bPerComponentNavigation = NavRelevantActor ? NavRelevantActor->DoesSupplyPerComponentNavigationCollision() : false;
+		if (bPerComponentNavigation)
+		{
+			FSetElementId ActorSetId = PendingOctreeUpdates.FindId(Component->GetOwner());
+			if (ActorSetId.IsValidId())
+			{
+				return ActorSetId;
+			}
+		}
+	}
+
 	const FOctreeElementId* ElementId = GetObjectsNavOctreeId(ElementOwner);
 	if (ElementId == NULL)
 	{
diff --git a/Engine/Source/Runtime/Engine/Private/AI/Navigation/RecastNavMesh.cpp b/Engine/Source/Runtime/Engine/Private/AI/Navigation/RecastNavMesh.cpp
--- a/Engine/Source/Runtime/Engine/Private/AI/Navigation/RecastNavMesh.cpp
+++ b/Engine/Source/Runtime/Engine/Private/AI/Navigation/RecastNavMesh.cpp
@@ -158,7 +158,8 @@ uint32 FRecastDebugGeometry::GetAllocatedSize() const
 
 	for (int i = 0; i < Clusters.Num(); ++i)
 	{
-		Size += Clusters[i].MeshIndices.GetAllocatedSize();
+//		Removed this as a race condition causes crash: TTP#343021
+//		Size += Clusters[i].MeshIndices.GetAllocatedSize();
 	}
 
 	return Size;
diff --git a/Engine/Source/Runtime/Engine/Private/InstancedFoliage.cpp b/Engine/Source/Runtime/Engine/Private/InstancedFoliage.cpp
--- a/Engine/Source/Runtime/Engine/Private/InstancedFoliage.cpp
+++ b/Engine/Source/Runtime/Engine/Private/InstancedFoliage.cpp
@@ -422,6 +422,11 @@ void FFoliageMeshInfo::RemoveInstances(AInstancedFoliageActor* InIFA, const TArr
 		{
 			//!! need to update bounds for this cluster
 			ClusterComponent->RegisterComponent();
+
+			if (ClusterComponent->IsCollisionEnabled() && ClusterComponent->GetWorld()->GetNavigationSystem() != nullptr)
+			{
+				ClusterComponent->GetWorld()->GetNavigationSystem()->UpdateNavOctree(ClusterComponent);
+			}
 			bRemoved = true;
 		}
 	}
@@ -548,6 +553,7 @@ void FFoliageMeshInfo::ReallocateClusters(AInstancedFoliageActor* InIFA, UFoliag
 		UInstancedStaticMeshComponent* Component = Cluster.ClusterComponent;
 		if (Component != nullptr)
 		{
+			Component->PerInstanceSMData.Empty();
 			Component->bAutoRegister = false;
 		}
 	}
@@ -1706,10 +1712,9 @@ void AInstancedFoliageActor::AddReferencedObjects(UObject* InThis, FReferenceCol
 #if WITH_EDITORONLY_DATA
 		for (FFoliageInstance& Instance : MeshInfo.Instances)
 		{
-			UActorComponent* Base = Instance.Base;
-			if (Base != nullptr)
+			if (Instance.Base != nullptr)
 			{
-				Collector.AddReferencedObject(Base, This);
+				Collector.AddReferencedObject(Instance.Base, This);
 			}
 		}
 #endif
diff --git a/Engine/Source/Runtime/Engine/Private/InstancedStaticMesh.cpp b/Engine/Source/Runtime/Engine/Private/InstancedStaticMesh.cpp
--- a/Engine/Source/Runtime/Engine/Private/InstancedStaticMesh.cpp
+++ b/Engine/Source/Runtime/Engine/Private/InstancedStaticMesh.cpp
@@ -1507,6 +1507,11 @@ void UInstancedStaticMeshComponent::AddInstance(const FTransform& InstanceTransf
 	SetupNewInstanceData(*NewInstanceData, PerInstanceSMData.Num() - 1, InstanceTransform);
 
 	MarkRenderStateDirty();
+
+	if (IsCollisionEnabled() && GetWorld()->GetNavigationSystem() != nullptr)
+	{
+		GetWorld()->GetNavigationSystem()->UpdateNavOctree(this);
+	}
 }
 
 void UInstancedStaticMeshComponent::AddInstanceWorldSpace(const FTransform& WorldTransform)
@@ -1597,6 +1602,11 @@ bool UInstancedStaticMeshComponent::RemoveInstance(int32 InstanceIndex)
 	// Indicate we need to update render state to reflect changes
 	MarkRenderStateDirty();
 
+	if (IsCollisionEnabled() && GetWorld()->GetNavigationSystem() != nullptr)
+	{
+		GetWorld()->GetNavigationSystem()->UpdateNavOctree(this);
+	}
+
 	return true;
 }
 
@@ -1609,6 +1619,11 @@ void UInstancedStaticMeshComponent::ClearInstances()
 
 	// Indicate we need to update render state to reflect changes
 	MarkRenderStateDirty();
+
+	if (IsCollisionEnabled() && GetWorld()->GetNavigationSystem() != nullptr)
+	{
+		GetWorld()->GetNavigationSystem()->UpdateNavOctree(this);
+	}
 }
 
 int32 UInstancedStaticMeshComponent::GetInstanceCount() const
@@ -1653,10 +1668,10 @@ bool UInstancedStaticMeshComponent::DoCustomNavigableGeometryExport(struct FNavi
 				if (!Scale3D.IsZero())
 				{
 					GeomExport->ExportCustomMesh(NavCollision->ConvexCollision.VertexBuffer.GetData(), NavCollision->ConvexCollision.VertexBuffer.Num(),
-						NavCollision->ConvexCollision.IndexBuffer.GetData(), NavCollision->ConvexCollision.IndexBuffer.Num(), FTransform(InstanceData.Transform));
+						NavCollision->ConvexCollision.IndexBuffer.GetData(), NavCollision->ConvexCollision.IndexBuffer.Num(), FTransform(InstanceData.Transform) * ComponentToWorld);
 
 					GeomExport->ExportCustomMesh(NavCollision->TriMeshCollision.VertexBuffer.GetData(), NavCollision->TriMeshCollision.VertexBuffer.Num(),
-						NavCollision->TriMeshCollision.IndexBuffer.GetData(), NavCollision->TriMeshCollision.IndexBuffer.Num(), FTransform(InstanceData.Transform));
+						NavCollision->TriMeshCollision.IndexBuffer.GetData(), NavCollision->TriMeshCollision.IndexBuffer.Num(), FTransform(InstanceData.Transform) * ComponentToWorld);
 				}
 			}
 		}
@@ -1671,7 +1686,7 @@ bool UInstancedStaticMeshComponent::DoCustomNavigableGeometryExport(struct FNavi
 					// if any of scales is 0 there's no point in exporting it
 					if (!Scale3D.IsZero())
 					{
-						GeomExport->ExportRigidBodySetup(*BodySetup, FTransform(InstanceData.Transform));
+						GeomExport->ExportRigidBodySetup(*BodySetup, FTransform(InstanceData.Transform) * ComponentToWorld);
 					}
 				}
 
