ExternalToolTest.cpp

Flagged the ExternalToolTest to be ATF_NonNullRHI

diff --git a/Engine/Source/Runtime/Engine/Private/Tests/ExternalToolTest.cpp b/Engine/Source/Runtime/Engine/Private/Tests/ExternalToolTest.cpp
--- a/Engine/Source/Runtime/Engine/Private/Tests/ExternalToolTest.cpp
+++ b/Engine/Source/Runtime/Engine/Private/Tests/ExternalToolTest.cpp
@@ -77,7 +77,7 @@ bool FWaitForProcessToCompleteLatentCommand::Update()
 /**
  * RunExternalToolTest
  */
-IMPLEMENT_COMPLEX_AUTOMATION_TEST( FRunExternalToolTest, "External", EAutomationTestFlags::ATF_None )
+IMPLEMENT_COMPLEX_AUTOMATION_TEST(FRunExternalToolTest, "External", EAutomationTestFlags::ATF_NonNullRHI)
 
 /** 
  * Find all the external too commands
