Header cleanup: Engine module, Landscape headers.

diff --git a/Engine/Source/Editor/LandscapeEditor/Private/LandscapeEdMode.cpp b/Engine/Source/Editor/LandscapeEditor/Private/LandscapeEdMode.cpp
--- a/Engine/Source/Editor/LandscapeEditor/Private/LandscapeEdMode.cpp
+++ b/Engine/Source/Editor/LandscapeEditor/Private/LandscapeEdMode.cpp
@@ -26,6 +26,7 @@
 
 // Classes
 #include "Landscape/Landscape.h"
+#include "Landscape/LandscapeLayerInfoObject.h"
 #include "Landscape/LandscapeHeightfieldCollisionComponent.h"
 #include "Landscape/LandscapeMaterialInstanceConstant.h"
 #include "Landscape/LandscapeSplinesComponent.h"
@@ -2651,4 +2652,9 @@ bool LandscapeEditorUtils::SetWeightmapData(ALandscapeProxy* Landscape, ULandsca
 	return false;
 }
 
+FName FLandscapeTargetListInfo::GetLayerName() const
+{
+	return LayerInfoObj.IsValid() ? LayerInfoObj->LayerName : LayerName;
+}
+
 #undef LOCTEXT_NAMESPACE
diff --git a/Engine/Source/Editor/LandscapeEditor/Private/LandscapeEdModeComponentTools.cpp b/Engine/Source/Editor/LandscapeEditor/Private/LandscapeEdModeComponentTools.cpp
--- a/Engine/Source/Editor/LandscapeEditor/Private/LandscapeEdModeComponentTools.cpp
+++ b/Engine/Source/Editor/LandscapeEditor/Private/LandscapeEdModeComponentTools.cpp
@@ -6,6 +6,7 @@
 #include "ScopedTransaction.h"
 #include "Landscape/Landscape.h"
 #include "Landscape/LandscapeEdit.h"
+#include "Landscape/LandscapeLayerInfoObject.h"
 #include "Landscape/LandscapeRender.h"
 #include "Landscape/LandscapeDataAccess.h"
 #include "Landscape/LandscapeSplineProxies.h"
diff --git a/Engine/Source/Editor/LandscapeEditor/Private/LandscapeEdModeErosionTools.cpp b/Engine/Source/Editor/LandscapeEditor/Private/LandscapeEdModeErosionTools.cpp
--- a/Engine/Source/Editor/LandscapeEditor/Private/LandscapeEdModeErosionTools.cpp
+++ b/Engine/Source/Editor/LandscapeEditor/Private/LandscapeEdModeErosionTools.cpp
@@ -12,6 +12,7 @@
 #include "Editor/PropertyEditor/Public/PropertyEditorModule.h"
 #include "LandscapeEdModeTools.h"
 #include "Landscape/Landscape.h"
+#include "Landscape/LandscapeLayerInfoObject.h"
 
 //
 // FLandscapeToolErosionBase
diff --git a/Engine/Source/Editor/LandscapeEditor/Private/LandscapeEditorDetailCustomization_CopyPaste.cpp b/Engine/Source/Editor/LandscapeEditor/Private/LandscapeEditorDetailCustomization_CopyPaste.cpp
--- a/Engine/Source/Editor/LandscapeEditor/Private/LandscapeEditorDetailCustomization_CopyPaste.cpp
+++ b/Engine/Source/Editor/LandscapeEditor/Private/LandscapeEditorDetailCustomization_CopyPaste.cpp
@@ -18,6 +18,8 @@
 #include "DesktopPlatformModule.h"
 #include "MainFrame.h"
 
+#include "Landscape/LandscapeLayerInfoObject.h"
+
 #define LOCTEXT_NAMESPACE "LandscapeEditor.CopyPaste"
 
 TSharedRef<IDetailCustomization> FLandscapeEditorDetailCustomization_CopyPaste::MakeInstance()
diff --git a/Engine/Source/Editor/LandscapeEditor/Private/LandscapeEditorDetailCustomization_NewLandscape.cpp b/Engine/Source/Editor/LandscapeEditor/Private/LandscapeEditorDetailCustomization_NewLandscape.cpp
--- a/Engine/Source/Editor/LandscapeEditor/Private/LandscapeEditorDetailCustomization_NewLandscape.cpp
+++ b/Engine/Source/Editor/LandscapeEditor/Private/LandscapeEditorDetailCustomization_NewLandscape.cpp
@@ -28,6 +28,7 @@
 #include "ImageWrapper.h"
 
 #include "Landscape/Landscape.h"
+#include "Landscape/LandscapeLayerInfoObject.h"
 
 #define LOCTEXT_NAMESPACE "LandscapeEditor.NewLandscape"
 
diff --git a/Engine/Source/Editor/LandscapeEditor/Private/LandscapeEditorDetailCustomization_TargetLayers.cpp b/Engine/Source/Editor/LandscapeEditor/Private/LandscapeEditorDetailCustomization_TargetLayers.cpp
--- a/Engine/Source/Editor/LandscapeEditor/Private/LandscapeEditorDetailCustomization_TargetLayers.cpp
+++ b/Engine/Source/Editor/LandscapeEditor/Private/LandscapeEditorDetailCustomization_TargetLayers.cpp
@@ -24,6 +24,8 @@
 #include "AssetRegistryModule.h"
 #include "LandscapeRender.h"
 
+#include "Landscape/LandscapeLayerInfoObject.h"
+
 #define LOCTEXT_NAMESPACE "LandscapeEditor.TargetLayers"
 
 TSharedRef<IDetailCustomization> FLandscapeEditorDetailCustomization_TargetLayers::MakeInstance()
diff --git a/Engine/Source/Editor/LandscapeEditor/Public/LandscapeEdMode.h b/Engine/Source/Editor/LandscapeEditor/Public/LandscapeEdMode.h
--- a/Engine/Source/Editor/LandscapeEditor/Public/LandscapeEdMode.h
+++ b/Engine/Source/Editor/LandscapeEditor/Public/LandscapeEdMode.h
@@ -14,6 +14,7 @@ struct FHeightmapToolTarget;
 template<typename TargetType> class FLandscapeToolCopyPaste;
 
 #include "Landscape/LandscapeInfo.h"
+#include "Landscape/LandscapeLayerInfoObject.h"
 #include "Landscape/LandscapeProxy.h"
 #include "Landscape/LandscapeGizmoActiveActor.h"
 
@@ -311,10 +312,7 @@ struct FLandscapeTargetListInfo
 		return NULL;
 	}
 
-	FName GetLayerName() const
-	{
-		return LayerInfoObj.IsValid() ? LayerInfoObj->LayerName : LayerName;
-	}
+	FName GetLayerName() const;
 
 	FString& ReimportFilePath() const
 	{
diff --git a/Engine/Source/Editor/UnrealEd/Private/EditorExporters.cpp b/Engine/Source/Editor/UnrealEd/Private/EditorExporters.cpp
--- a/Engine/Source/Editor/UnrealEd/Private/EditorExporters.cpp
+++ b/Engine/Source/Editor/UnrealEd/Private/EditorExporters.cpp
@@ -19,6 +19,8 @@
 #include "ImageUtils.h"
 #include "Foliage/InstancedFoliageActor.h"
 #include "Landscape/Landscape.h"
+#include "Landscape/LandscapeComponent.h"
+#include "Landscape/LandscapeInfo.h"
 #include "Foliage/FoliageType.h"
 #include "UnrealExporter.h"
 #include "EngineModule.h"
diff --git a/Engine/Source/Editor/UnrealEd/Private/Fbx/FbxMainExport.cpp b/Engine/Source/Editor/UnrealEd/Private/Fbx/FbxMainExport.cpp
--- a/Engine/Source/Editor/UnrealEd/Private/Fbx/FbxMainExport.cpp
+++ b/Engine/Source/Editor/UnrealEd/Private/Fbx/FbxMainExport.cpp
@@ -50,6 +50,8 @@
 #include "StaticMeshResources.h"
 #include "LandscapeDataAccess.h"
 #include "Landscape/LandscapeProxy.h"
+#include "Landscape/LandscapeInfo.h"
+#include "Landscape/LandscapeComponent.h"
 #include "Components/SplineMeshComponent.h"
 
 #include "FbxExporter.h"
diff --git a/Engine/Source/Editor/UnrealEd/Private/UnrealEdSrv.cpp b/Engine/Source/Editor/UnrealEd/Private/UnrealEdSrv.cpp
--- a/Engine/Source/Editor/UnrealEd/Private/UnrealEdSrv.cpp
+++ b/Engine/Source/Editor/UnrealEd/Private/UnrealEdSrv.cpp
@@ -28,6 +28,7 @@
 #include "Matinee/InterpData.h"
 #include "Animation/SkeletalMeshActor.h"
 #include "Landscape/LandscapeInfo.h"
+#include "Landscape/LandscapeLayerInfoObject.h"
 #include "Landscape/LandscapeProxy.h"
 #include "Landscape/LandscapeGizmoActiveActor.h"
 #include "Landscape/LandscapeComponent.h"
diff --git a/Engine/Source/Runtime/Engine/Classes/Landscape/Landscape.h b/Engine/Source/Runtime/Engine/Classes/Landscape/Landscape.h
--- a/Engine/Source/Runtime/Engine/Classes/Landscape/Landscape.h
+++ b/Engine/Source/Runtime/Engine/Classes/Landscape/Landscape.h
@@ -2,10 +2,7 @@
 
 #pragma once
 
-#include "LandscapeInfo.h"
 #include "LandscapeProxy.h"
-#include "Components/LightComponent.h"
-
 #include "Landscape.generated.h"
 
 UENUM()
diff --git a/Engine/Source/Runtime/Engine/Classes/Landscape/LandscapeComponent.h b/Engine/Source/Runtime/Engine/Classes/Landscape/LandscapeComponent.h
--- a/Engine/Source/Runtime/Engine/Classes/Landscape/LandscapeComponent.h
+++ b/Engine/Source/Runtime/Engine/Classes/Landscape/LandscapeComponent.h
@@ -1,10 +1,11 @@
 // Copyright 1998-2014 Epic Games, Inc. All Rights Reserved.
 
 #pragma once
-#include "LandscapeLayerInfoObject.h"
 #include "SceneTypes.h"
 #include "LandscapeComponent.generated.h"
 
+class ULandscapeLayerInfoObject;
+
 class FLandscapeComponentDerivedData
 {
 	/** The compressed Landscape component data for mobile rendering. Serialized to disk. 
@@ -90,14 +91,7 @@ struct FWeightmapLayerAllocationInfo
 	{
 	}
 	
-	FName GetLayerName() const
-	{
-		if (LayerInfo)
-		{
-			return LayerInfo->LayerName;
-		}
-		return NAME_None;
-	}
+	FName GetLayerName() const;
 };
 
 UCLASS(hidecategories=(Display, Attachment, Physics, Debug, Collision, Movement, Rendering, PrimitiveComponent, Object, Transform), showcategories=("Rendering|Material"), MinimalAPI)
diff --git a/Engine/Source/Runtime/Engine/Classes/Landscape/LandscapeInfo.h b/Engine/Source/Runtime/Engine/Classes/Landscape/LandscapeInfo.h
--- a/Engine/Source/Runtime/Engine/Classes/Landscape/LandscapeInfo.h
+++ b/Engine/Source/Runtime/Engine/Classes/Landscape/LandscapeInfo.h
@@ -1,9 +1,6 @@
 // Copyright 1998-2014 Epic Games, Inc. All Rights Reserved.
 
 #pragma once
-
-#include "LandscapeLayerInfoObject.h"
-
 #include "LandscapeInfo.generated.h"
 
 /** Structure storing Collision for LandscapeComponent Add */
@@ -65,17 +62,7 @@ struct FLandscapeInfoLayerSettings
 	{
 	}
 
-	FLandscapeInfoLayerSettings(ULandscapeLayerInfoObject* InLayerInfo, class ALandscapeProxy* InProxy)
-		: LayerInfoObj(InLayerInfo)
-		, LayerName((InLayerInfo != NULL) ? InLayerInfo->LayerName : NAME_None)
-#if WITH_EDITORONLY_DATA
-		, ThumbnailMIC(NULL)
-		, Owner(InProxy)
-		, DebugColorChannel(0)
-		, bValid(false)
-#endif
-	{
-	}
+	ENGINE_API FLandscapeInfoLayerSettings(ULandscapeLayerInfoObject* InLayerInfo, class ALandscapeProxy* InProxy);
 
 	FLandscapeInfoLayerSettings(FName InPlaceholderLayerName, class ALandscapeProxy* InProxy)
 		: LayerInfoObj(NULL)
diff --git a/Engine/Source/Runtime/Engine/Classes/Landscape/LandscapeProxy.h b/Engine/Source/Runtime/Engine/Classes/Landscape/LandscapeProxy.h
--- a/Engine/Source/Runtime/Engine/Classes/Landscape/LandscapeProxy.h
+++ b/Engine/Source/Runtime/Engine/Classes/Landscape/LandscapeProxy.h
@@ -1,20 +1,21 @@
 // Copyright 1998-2014 Epic Games, Inc. All Rights Reserved.
 
 #pragma once
-
+#include "AI/Navigation/NavRelevantActorInterface.h"
 #include "Engine/EngineTypes.h"
-#include "LandscapeInfo.h"
+#include "GameFramework/Actor.h"
 #include "PhysicsEngine/BodyInstance.h"
-
 #include "LandscapeProxy.generated.h"
 
-//
-// Forward declarations.
-//
 class ULandscapeMaterialInstanceConstant;
 class ULandscapeLayerInfoObject;
 class ULandscapeSplinesComponent;
 class ULandscapeHeightfieldCollisionComponent;
+class UMaterialInterface;
+class UTexture2D;
+class ALandscape;
+class ALandscapeProxy;
+class ULandscapeComponent;
 
 /** Structure storing channel usage for weightmap textures */
 USTRUCT()
@@ -148,13 +149,7 @@ struct FLandscapeImportLayerInfo
 	{
 	}
 
-	FLandscapeImportLayerInfo(const struct FLandscapeInfoLayerSettings& InLayerSettings)
-	:	LayerName(InLayerSettings.GetLayerName())
-	,	LayerInfo(InLayerSettings.LayerInfoObj)
-	,	ThumbnailMIC(NULL)
-	,	SourceFilePath(InLayerSettings.GetEditorSettings().ReimportLayerFilePath)
-	{
-	}
+	ENGINE_API FLandscapeImportLayerInfo(const struct FLandscapeInfoLayerSettings& InLayerSettings);
 #endif
 };
 
diff --git a/Engine/Source/Runtime/Engine/Classes/Landscape/LandscapeSplineControlPoint.h b/Engine/Source/Runtime/Engine/Classes/Landscape/LandscapeSplineControlPoint.h
--- a/Engine/Source/Runtime/Engine/Classes/Landscape/LandscapeSplineControlPoint.h
+++ b/Engine/Source/Runtime/Engine/Classes/Landscape/LandscapeSplineControlPoint.h
@@ -2,10 +2,11 @@
 
 
 #pragma once
-#include "LandscapeSplineSegment.h"
-
+#include "Landscape/LandscapeSplineSegment.h"
 #include "LandscapeSplineControlPoint.generated.h"
 
+struct FLandscapeSplineInterpPoint;
+
 USTRUCT()
 struct FLandscapeSplineConnection
 {
diff --git a/Engine/Source/Runtime/Engine/Private/Landscape/Landscape.cpp b/Engine/Source/Runtime/Engine/Private/Landscape/Landscape.cpp
--- a/Engine/Source/Runtime/Engine/Private/Landscape/Landscape.cpp
+++ b/Engine/Source/Runtime/Engine/Private/Landscape/Landscape.cpp
@@ -17,6 +17,8 @@ Landscape.cpp: Terrain rendering
 #include "Landscape/Landscape.h"
 #include "Landscape/LandscapeMeshCollisionComponent.h"
 #include "Landscape/LandscapeSplinesComponent.h"
+#include "Landscape/LandscapeInfo.h"
+#include "Landscape/LandscapeLayerInfoObject.h"
 #include "LightMap.h"
 #include "ShadowMap.h"
 
diff --git a/Engine/Source/Runtime/Engine/Private/Landscape/LandscapeCollision.cpp b/Engine/Source/Runtime/Engine/Private/Landscape/LandscapeCollision.cpp
--- a/Engine/Source/Runtime/Engine/Private/Landscape/LandscapeCollision.cpp
+++ b/Engine/Source/Runtime/Engine/Private/Landscape/LandscapeCollision.cpp
@@ -15,6 +15,8 @@ LandscapeCollision.cpp: Landscape collision
 #include "../PhysicsEngine/PhysDerivedData.h"
 #include "Landscape/LandscapeHeightfieldCollisionComponent.h"
 #include "Landscape/LandscapeMeshCollisionComponent.h"
+#include "Landscape/LandscapeLayerInfoObject.h"
+#include "Landscape/LandscapeInfo.h"
 
 TMap<FGuid, ULandscapeHeightfieldCollisionComponent::FPhysXHeightfieldRef* > GSharedHeightfieldRefs;
 
diff --git a/Engine/Source/Runtime/Engine/Private/Landscape/LandscapeComponent.cpp b/Engine/Source/Runtime/Engine/Private/Landscape/LandscapeComponent.cpp
new file mode 100644
--- /dev/null
+++ b/Engine/Source/Runtime/Engine/Private/Landscape/LandscapeComponent.cpp
@@ -0,0 +1,12 @@
+#include "EnginePrivate.h"
+#include "Landscape/LandscapeComponent.h"
+#include "Landscape/LandscapeLayerInfoObject.h"
+
+FName FWeightmapLayerAllocationInfo::GetLayerName() const
+{
+	if (LayerInfo)
+	{
+		return LayerInfo->LayerName;
+	}
+	return NAME_None;
+}
\ No newline at end of file
diff --git a/Engine/Source/Runtime/Engine/Private/Landscape/LandscapeDataAccess.cpp b/Engine/Source/Runtime/Engine/Private/Landscape/LandscapeDataAccess.cpp
new file mode 100644
--- /dev/null
+++ b/Engine/Source/Runtime/Engine/Private/Landscape/LandscapeDataAccess.cpp
@@ -0,0 +1,196 @@
+// Copyright 1998-2014 Epic Games, Inc. All Rights Reserved.
+
+#include "EnginePrivate.h"
+#include "Landscape/LandscapeDataAccess.h"
+#include "Landscape/LandscapeComponent.h"
+#include "Landscape/LandscapeInfo.h"
+
+ENGINE_API FLandscapeComponentDataInterface::FLandscapeComponentDataInterface(ULandscapeComponent* InComponent, int32 InMipLevel) :
+	Component(InComponent),
+	HeightMipData(NULL),
+	XYOffsetMipData(NULL),
+	bNeedToDeleteDataInterface(false),
+	MipLevel(InMipLevel)
+{
+	if (Component->GetLandscapeInfo())
+	{
+		DataInterface = Component->GetLandscapeInfo()->GetDataInterface();
+	}
+	else
+	{
+		DataInterface = new FLandscapeDataInterface;
+		bNeedToDeleteDataInterface = true;
+	}
+
+	// Offset and stride for this component's data in heightmap texture
+	HeightmapStride = Component->HeightmapTexture->Source.GetSizeX() >> MipLevel;
+	HeightmapComponentOffsetX = FMath::RoundToInt((float)(Component->HeightmapTexture->Source.GetSizeX() >> MipLevel) * Component->HeightmapScaleBias.Z);
+	HeightmapComponentOffsetY = FMath::RoundToInt((float)(Component->HeightmapTexture->Source.GetSizeY() >> MipLevel) * Component->HeightmapScaleBias.W);
+	HeightmapSubsectionOffset = (Component->SubsectionSizeQuads + 1) >> MipLevel;
+
+	ComponentSizeVerts = (Component->ComponentSizeQuads + 1) >> MipLevel;
+	SubsectionSizeVerts = (Component->SubsectionSizeQuads + 1) >> MipLevel;
+
+	if (MipLevel < Component->HeightmapTexture->Source.GetNumMips())
+	{
+		HeightMipData = (FColor*)DataInterface->LockMip(Component->HeightmapTexture, MipLevel);
+		if (Component->XYOffsetmapTexture)
+		{
+			XYOffsetMipData = (FColor*)DataInterface->LockMip(Component->XYOffsetmapTexture, MipLevel);
+		}
+	}
+}
+
+ENGINE_API FLandscapeComponentDataInterface::~FLandscapeComponentDataInterface()
+{
+	if (HeightMipData)
+	{
+		DataInterface->UnlockMip(Component->HeightmapTexture, MipLevel);
+		if (Component->XYOffsetmapTexture)
+		{
+			DataInterface->UnlockMip(Component->XYOffsetmapTexture, MipLevel);
+		}
+	}
+
+	if (bNeedToDeleteDataInterface)
+	{
+		delete DataInterface;
+		DataInterface = NULL;
+	}
+}
+
+ENGINE_API void FLandscapeComponentDataInterface::UnlockRawHeightData() const
+{
+	DataInterface->UnlockMip(Component->HeightmapTexture, MipLevel);
+}
+
+ENGINE_API void FLandscapeComponentDataInterface::GetHeightmapTextureData(TArray<FColor>& OutData)
+{
+#if LANDSCAPE_VALIDATE_DATA_ACCESS
+	check(HeightMipData);
+#endif
+	int32 HeightmapSize = ((Component->SubsectionSizeQuads + 1) * Component->NumSubsections) >> MipLevel;
+	OutData.Empty(FMath::Square(HeightmapSize));
+	OutData.AddUninitialized(FMath::Square(HeightmapSize));
+
+	for (int32 SubY = 0; SubY < HeightmapSize; SubY++)
+	{
+		// X/Y of the vertex we're looking at in component's coordinates.
+		int32 CompY = SubY;
+
+		// UV coordinates of the data offset into the texture
+		int32 TexV = SubY + HeightmapComponentOffsetY;
+
+		// Copy the data
+		FMemory::Memcpy(&OutData[CompY * HeightmapSize], &HeightMipData[HeightmapComponentOffsetX + TexV * HeightmapStride], HeightmapSize * sizeof(FColor));
+	}
+}
+
+ENGINE_API bool FLandscapeComponentDataInterface::GetWeightmapTextureData(ULandscapeLayerInfoObject* LayerInfo, TArray<uint8>& OutData)
+{
+	int32 LayerIdx = INDEX_NONE;
+	for (int32 Idx = 0; Idx < Component->WeightmapLayerAllocations.Num(); Idx++)
+	{
+		if (Component->WeightmapLayerAllocations[Idx].LayerInfo == LayerInfo)
+		{
+			LayerIdx = Idx;
+			break;
+		}
+	}
+	if (LayerIdx < 0)
+	{
+		return false;
+	}
+	if (Component->WeightmapLayerAllocations[LayerIdx].WeightmapTextureIndex >= Component->WeightmapTextures.Num())
+	{
+		return false;
+	}
+	if (Component->WeightmapLayerAllocations[LayerIdx].WeightmapTextureChannel >= 4)
+	{
+		return false;
+	}
+
+	int32 WeightmapSize = ((Component->SubsectionSizeQuads + 1) * Component->NumSubsections) >> MipLevel;
+	OutData.Empty(FMath::Square(WeightmapSize));
+	OutData.AddUninitialized(FMath::Square(WeightmapSize));
+
+	FColor* WeightMipData = (FColor*)DataInterface->LockMip(Component->WeightmapTextures[Component->WeightmapLayerAllocations[LayerIdx].WeightmapTextureIndex], MipLevel);
+
+	// Channel remapping
+	int32 ChannelOffsets[4] = { (int32)STRUCT_OFFSET(FColor, R), (int32)STRUCT_OFFSET(FColor, G), (int32)STRUCT_OFFSET(FColor, B), (int32)STRUCT_OFFSET(FColor, A) };
+
+	uint8* SrcTextureData = (uint8*)WeightMipData + ChannelOffsets[Component->WeightmapLayerAllocations[LayerIdx].WeightmapTextureChannel];
+
+	for (int32 i = 0; i < FMath::Square(WeightmapSize); i++)
+	{
+		OutData[i] = SrcTextureData[i * 4];
+	}
+
+	DataInterface->UnlockMip(Component->WeightmapTextures[Component->WeightmapLayerAllocations[LayerIdx].WeightmapTextureIndex], MipLevel);
+	return true;
+}
+
+ENGINE_API FColor* FLandscapeComponentDataInterface::GetXYOffsetData(int32 LocalX, int32 LocalY) const
+{
+#if LANDSCAPE_VALIDATE_DATA_ACCESS
+	check(Component);
+	check(LocalX >= 0 && LocalY >= 0 && LocalX < Component->ComponentSizeQuads + 1 && LocalY < Component->ComponentSizeQuads + 1);
+#endif
+
+	const int32 WeightmapSize = (Component->SubsectionSizeQuads + 1) * Component->NumSubsections;
+	int32 SubNumX;
+	int32 SubNumY;
+	int32 SubX;
+	int32 SubY;
+	ComponentXYToSubsectionXY(LocalX, LocalY, SubNumX, SubNumY, SubX, SubY);
+
+	return &XYOffsetMipData[SubX + SubNumX*SubsectionSizeVerts + (SubY + SubNumY*SubsectionSizeVerts)*WeightmapSize];
+}
+
+ENGINE_API FVector FLandscapeComponentDataInterface::GetLocalVertex(int32 LocalX, int32 LocalY) const
+{
+	const float ScaleFactor = (float)Component->ComponentSizeQuads / (float)(ComponentSizeVerts - 1);
+	float XOffset, YOffset;
+	GetXYOffset(LocalX, LocalY, XOffset, YOffset);
+	return FVector(LocalX * ScaleFactor + XOffset, LocalY * ScaleFactor + YOffset, LandscapeDataAccess::GetLocalHeight(GetHeight(LocalX, LocalY)));
+}
+
+ENGINE_API FVector FLandscapeComponentDataInterface::GetWorldVertex(int32 LocalX, int32 LocalY) const
+{
+	return Component->ComponentToWorld.TransformPosition(GetLocalVertex(LocalX, LocalY));
+}
+
+ENGINE_API void FLandscapeComponentDataInterface::GetWorldTangentVectors(int32 LocalX, int32 LocalY, FVector& WorldTangentX, FVector& WorldTangentY, FVector& WorldTangentZ) const
+{
+	FColor* Data = GetHeightData(LocalX, LocalY);
+	WorldTangentZ.X = 2.f * (float)Data->B / 255.f - 1.f;
+	WorldTangentZ.Y = 2.f * (float)Data->A / 255.f - 1.f;
+	WorldTangentZ.Z = FMath::Sqrt(1.f - (FMath::Square(WorldTangentZ.X) + FMath::Square(WorldTangentZ.Y)));
+	WorldTangentX = FVector(-WorldTangentZ.Z, 0.f, WorldTangentZ.X);
+	WorldTangentY = FVector(0.f, WorldTangentZ.Z, -WorldTangentZ.Y);
+
+	WorldTangentX = Component->ComponentToWorld.TransformVectorNoScale(WorldTangentX);
+	WorldTangentY = Component->ComponentToWorld.TransformVectorNoScale(WorldTangentY);
+	WorldTangentZ = Component->ComponentToWorld.TransformVectorNoScale(WorldTangentZ);
+}
+
+ENGINE_API void FLandscapeComponentDataInterface::GetWorldPositionTangents(int32 LocalX, int32 LocalY, FVector& WorldPos, FVector& WorldTangentX, FVector& WorldTangentY, FVector& WorldTangentZ) const
+{
+	FColor* Data = GetHeightData(LocalX, LocalY);
+
+	WorldTangentZ.X = 2.f * (float)Data->B / 255.f - 1.f;
+	WorldTangentZ.Y = 2.f * (float)Data->A / 255.f - 1.f;
+	WorldTangentZ.Z = FMath::Sqrt(1.f - (FMath::Square(WorldTangentZ.X) + FMath::Square(WorldTangentZ.Y)));
+	WorldTangentX = FVector(WorldTangentZ.Z, 0.f, -WorldTangentZ.X);
+	WorldTangentY = WorldTangentZ ^ WorldTangentX;
+
+	uint16 Height = (Data->R << 8) + Data->G;
+
+	const float ScaleFactor = (float)Component->ComponentSizeQuads / (float)(ComponentSizeVerts - 1);
+	float XOffset, YOffset;
+	GetXYOffset(LocalX, LocalY, XOffset, YOffset);
+	WorldPos = Component->ComponentToWorld.TransformPosition(FVector(LocalX * ScaleFactor + XOffset, LocalY * ScaleFactor + YOffset, LandscapeDataAccess::GetLocalHeight(Height)));
+	WorldTangentX = Component->ComponentToWorld.TransformVectorNoScale(WorldTangentX);
+	WorldTangentY = Component->ComponentToWorld.TransformVectorNoScale(WorldTangentY);
+	WorldTangentZ = Component->ComponentToWorld.TransformVectorNoScale(WorldTangentZ);
+}
\ No newline at end of file
diff --git a/Engine/Source/Runtime/Engine/Private/Landscape/LandscapeEdit.cpp b/Engine/Source/Runtime/Engine/Private/Landscape/LandscapeEdit.cpp
--- a/Engine/Source/Runtime/Engine/Private/Landscape/LandscapeEdit.cpp
+++ b/Engine/Source/Runtime/Engine/Private/Landscape/LandscapeEdit.cpp
@@ -14,6 +14,8 @@ LandscapeEdit.cpp: Landscape editing
 #include "Landscape/LandscapeRender.h"
 #include "Landscape/LandscapeRenderMobile.h"
 #include "Landscape/Landscape.h"
+#include "Landscape/LandscapeInfo.h"
+#include "Landscape/LandscapeLayerInfoObject.h"
 #include "Landscape/LandscapeMaterialInstanceConstant.h"
 #include "Landscape/LandscapeHeightfieldCollisionComponent.h"
 #include "Landscape/LandscapeMeshCollisionComponent.h"
diff --git a/Engine/Source/Runtime/Engine/Private/Landscape/LandscapeEditInterface.cpp b/Engine/Source/Runtime/Engine/Private/Landscape/LandscapeEditInterface.cpp
--- a/Engine/Source/Runtime/Engine/Private/Landscape/LandscapeEditInterface.cpp
+++ b/Engine/Source/Runtime/Engine/Private/Landscape/LandscapeEditInterface.cpp
@@ -9,6 +9,8 @@ LandscapeEditInterface.cpp: Landscape editing interface
 #include "Landscape/LandscapeEdit.h"
 #include "Landscape/LandscapeRender.h"
 #include "Landscape/Landscape.h"
+#include "Landscape/LandscapeInfo.h"
+#include "Landscape/LandscapeLayerInfoObject.h"
 
 #if WITH_EDITOR
 
diff --git a/Engine/Source/Runtime/Engine/Private/Landscape/LandscapeGizmoActor.cpp b/Engine/Source/Runtime/Engine/Private/Landscape/LandscapeGizmoActor.cpp
--- a/Engine/Source/Runtime/Engine/Private/Landscape/LandscapeGizmoActor.cpp
+++ b/Engine/Source/Runtime/Engine/Private/Landscape/LandscapeGizmoActor.cpp
@@ -7,6 +7,8 @@
 #include "Landscape/LandscapeGizmoActiveActor.h"
 #include "Landscape/LandscapeGizmoRenderComponent.h"
 #include "Landscape/Landscape.h"
+#include "Landscape/LandscapeInfo.h"
+#include "Landscape/LandscapeLayerInfoObject.h"
 #include "DynamicMeshBuilder.h"
 
 namespace
diff --git a/Engine/Source/Runtime/Engine/Private/Landscape/LandscapeInfo.cpp b/Engine/Source/Runtime/Engine/Private/Landscape/LandscapeInfo.cpp
new file mode 100644
--- /dev/null
+++ b/Engine/Source/Runtime/Engine/Private/Landscape/LandscapeInfo.cpp
@@ -0,0 +1,14 @@
+#include "EnginePrivate.h"
+#include "Landscape/LandscapeInfo.h"
+#include "Landscape/LandscapeLayerInfoObject.h"
+
+ENGINE_API FLandscapeInfoLayerSettings::FLandscapeInfoLayerSettings(ULandscapeLayerInfoObject* InLayerInfo, class ALandscapeProxy* InProxy)
+: LayerInfoObj(InLayerInfo)
+, LayerName((InLayerInfo != NULL) ? InLayerInfo->LayerName : NAME_None)
+#if WITH_EDITORONLY_DATA
+, ThumbnailMIC(NULL)
+, Owner(InProxy)
+, DebugColorChannel(0)
+, bValid(false)
+#endif
+{ }
\ No newline at end of file
diff --git a/Engine/Source/Runtime/Engine/Private/Landscape/LandscapeLight.cpp b/Engine/Source/Runtime/Engine/Private/Landscape/LandscapeLight.cpp
--- a/Engine/Source/Runtime/Engine/Private/Landscape/LandscapeLight.cpp
+++ b/Engine/Source/Runtime/Engine/Private/Landscape/LandscapeLight.cpp
@@ -6,6 +6,7 @@ LandscapeLight.cpp: Static lighting for LandscapeComponents
 
 #include "EnginePrivate.h"
 #include "Landscape/LandscapeLight.h"
+#include "Landscape/LandscapeInfo.h"
 #include "Landscape/LandscapeRender.h"
 #include "Landscape/LandscapeDataAccess.h"
 
diff --git a/Engine/Source/Runtime/Engine/Private/Landscape/LandscapeProxy.cpp b/Engine/Source/Runtime/Engine/Private/Landscape/LandscapeProxy.cpp
new file mode 100644
--- /dev/null
+++ b/Engine/Source/Runtime/Engine/Private/Landscape/LandscapeProxy.cpp
@@ -0,0 +1,16 @@
+// Copyright 1998-2014 Epic Games, Inc. All Rights Reserved.
+
+#include "EnginePrivate.h"
+#include "Landscape/LandscapeProxy.h"
+#include "Landscape/LandscapeInfo.h"
+
+#if WITH_EDITOR
+
+ENGINE_API FLandscapeImportLayerInfo::FLandscapeImportLayerInfo(const struct FLandscapeInfoLayerSettings& InLayerSettings)
+:	LayerName(InLayerSettings.GetLayerName())
+,	LayerInfo(InLayerSettings.LayerInfoObj)
+,	ThumbnailMIC(NULL)
+,	SourceFilePath(InLayerSettings.GetEditorSettings().ReimportLayerFilePath)
+{ }
+
+#endif // WITH_EDITOR
\ No newline at end of file
diff --git a/Engine/Source/Runtime/Engine/Private/Landscape/LandscapeSplineProxies.cpp b/Engine/Source/Runtime/Engine/Private/Landscape/LandscapeSplineProxies.cpp
new file mode 100644
--- /dev/null
+++ b/Engine/Source/Runtime/Engine/Private/Landscape/LandscapeSplineProxies.cpp
@@ -0,0 +1,10 @@
+// Copyright 1998-2014 Epic Games, Inc. All Rights Reserved.
+
+#include "EnginePrivate.h"
+#include "Landscape/LandscapeSplineProxies.h"
+#include "Landscape/LandscapeSplineSegment.h"
+
+ENGINE_API void HLandscapeSplineProxy_Tangent::Serialize(FArchive& Ar)
+{
+	Ar << SplineSegment;
+}
\ No newline at end of file
diff --git a/Engine/Source/Runtime/Engine/Private/Landscape/LandscapeSplines.cpp b/Engine/Source/Runtime/Engine/Private/Landscape/LandscapeSplines.cpp
--- a/Engine/Source/Runtime/Engine/Private/Landscape/LandscapeSplines.cpp
+++ b/Engine/Source/Runtime/Engine/Private/Landscape/LandscapeSplines.cpp
@@ -5,6 +5,7 @@
 =============================================================================*/
 
 #include "EnginePrivate.h"
+#include "Components/SplineMeshComponent.h"
 #include "ShaderParameterUtils.h"
 #include "StaticMeshResources.h"
 #include "AI/Navigation/RecastHelpers.h"
diff --git a/Engine/Source/Runtime/Engine/Public/Landscape/LandscapeDataAccess.h b/Engine/Source/Runtime/Engine/Public/Landscape/LandscapeDataAccess.h
--- a/Engine/Source/Runtime/Engine/Public/Landscape/LandscapeDataAccess.h
+++ b/Engine/Source/Runtime/Engine/Public/Landscape/LandscapeDataAccess.h
@@ -17,10 +17,8 @@ LandscapeDataAccess.h: Classes for the editor to access to Landscape data
 
 #if WITH_EDITOR
 
-// Forward declarations
+class ULandscapeComponent;
 class ULandscapeLayerInfoObject;
-#include "Landscape/LandscapeInfo.h"
-#include "Landscape/LandscapeComponent.h"
 
 namespace LandscapeDataAccess
 {
@@ -125,59 +123,8 @@ struct FLandscapeComponentDataInterface
 	friend struct FLandscapeDataInterface;
 
 	// tors
-	FLandscapeComponentDataInterface(ULandscapeComponent* InComponent, int32 InMipLevel = 0)
-	:	Component(InComponent),
-		HeightMipData(NULL),
-		XYOffsetMipData(NULL),
-		bNeedToDeleteDataInterface(false),
-		MipLevel(InMipLevel)
-	{
-		if (Component->GetLandscapeInfo())
-		{
-			DataInterface = Component->GetLandscapeInfo()->GetDataInterface();
-		}
-		else
-		{
-			DataInterface = new FLandscapeDataInterface;
-			bNeedToDeleteDataInterface = true;
-		}
-
-		// Offset and stride for this component's data in heightmap texture
-		HeightmapStride = Component->HeightmapTexture->Source.GetSizeX() >> MipLevel;
-		HeightmapComponentOffsetX = FMath::RoundToInt( (float)(Component->HeightmapTexture->Source.GetSizeX() >> MipLevel) * Component->HeightmapScaleBias.Z );
-		HeightmapComponentOffsetY = FMath::RoundToInt( (float)(Component->HeightmapTexture->Source.GetSizeY() >> MipLevel) * Component->HeightmapScaleBias.W );
-		HeightmapSubsectionOffset = (Component->SubsectionSizeQuads + 1) >> MipLevel;
-
-		ComponentSizeVerts = (Component->ComponentSizeQuads + 1) >> MipLevel;
-		SubsectionSizeVerts = (Component->SubsectionSizeQuads + 1) >> MipLevel;
-
-		if (MipLevel < Component->HeightmapTexture->Source.GetNumMips())
-		{
-			HeightMipData = (FColor*)DataInterface->LockMip(Component->HeightmapTexture, MipLevel);
-			if (Component->XYOffsetmapTexture)
-			{
-				XYOffsetMipData = (FColor*)DataInterface->LockMip(Component->XYOffsetmapTexture, MipLevel);
-			}
-		}
-	}
-
-	~FLandscapeComponentDataInterface()
-	{
-		if (HeightMipData)
-		{
-			DataInterface->UnlockMip(Component->HeightmapTexture,MipLevel);
-			if (Component->XYOffsetmapTexture)
-			{
-				DataInterface->UnlockMip(Component->XYOffsetmapTexture,MipLevel);
-			}
-		}
-
-		if (bNeedToDeleteDataInterface)
-		{
-			delete DataInterface;
-			DataInterface = NULL;
-		}
-	}
+	ENGINE_API FLandscapeComponentDataInterface(ULandscapeComponent* InComponent, int32 InMipLevel = 0);
+	ENGINE_API ~FLandscapeComponentDataInterface();
 
 	// Accessors
 	void VertexIndexToXY(int32 VertexIndex, int32& OutX, int32& OutY) const
@@ -243,77 +190,12 @@ struct FLandscapeComponentDataInterface
 		XYOffsetMipData = NewXYOffsetData;
 	}
 
-	void UnlockRawHeightData() const
-	{
-		DataInterface->UnlockMip(Component->HeightmapTexture, MipLevel);
-	}
+	ENGINE_API void UnlockRawHeightData() const;
 
 	/* Return the raw heightmap data exactly same size for Heightmap texture which belong to only this component */
-	void GetHeightmapTextureData( TArray<FColor>& OutData )
-	{
-#if LANDSCAPE_VALIDATE_DATA_ACCESS
-		check(HeightMipData);
-#endif
-		int32 HeightmapSize = ((Component->SubsectionSizeQuads+1) * Component->NumSubsections) >> MipLevel;
-		OutData.Empty(FMath::Square(HeightmapSize));
-		OutData.AddUninitialized(FMath::Square(HeightmapSize));
-
-		for( int32 SubY=0;SubY<HeightmapSize;SubY++ )
-		{
-			// X/Y of the vertex we're looking at in component's coordinates.
-			int32 CompY = SubY;
-
-			// UV coordinates of the data offset into the texture
-			int32 TexV = SubY + HeightmapComponentOffsetY;
-
-			// Copy the data
-			FMemory::Memcpy( &OutData[CompY * HeightmapSize], &HeightMipData[HeightmapComponentOffsetX + TexV * HeightmapStride], HeightmapSize * sizeof(FColor));
-		}
-	}
-
-	bool GetWeightmapTextureData( ULandscapeLayerInfoObject* LayerInfo, TArray<uint8>& OutData )
-	{
-		int32 LayerIdx = INDEX_NONE;
-		for (int32 Idx = 0; Idx < Component->WeightmapLayerAllocations.Num(); Idx++)
-		{
-			if ( Component->WeightmapLayerAllocations[Idx].LayerInfo == LayerInfo )
-			{
-				LayerIdx = Idx;
-				break;
-			}
-		}
-		if (LayerIdx < 0)
-		{
-			return false;
-		}
-		if ( Component->WeightmapLayerAllocations[LayerIdx].WeightmapTextureIndex >= Component->WeightmapTextures.Num() )
-		{
-			return false;
-		}
-		if ( Component->WeightmapLayerAllocations[LayerIdx].WeightmapTextureChannel >= 4 )
-		{
-			return false;
-		}
-
-		int32 WeightmapSize = ((Component->SubsectionSizeQuads + 1) * Component->NumSubsections) >> MipLevel;
-		OutData.Empty(FMath::Square(WeightmapSize));
-		OutData.AddUninitialized(FMath::Square(WeightmapSize));
-
-		FColor* WeightMipData = (FColor*)DataInterface->LockMip(Component->WeightmapTextures[Component->WeightmapLayerAllocations[LayerIdx].WeightmapTextureIndex], MipLevel);
+	ENGINE_API void GetHeightmapTextureData(TArray<FColor>& OutData);
 
-		// Channel remapping
-		int32 ChannelOffsets[4] = {(int32)STRUCT_OFFSET(FColor,R),(int32)STRUCT_OFFSET(FColor,G),(int32)STRUCT_OFFSET(FColor,B),(int32)STRUCT_OFFSET(FColor,A)};
-
-		uint8* SrcTextureData = (uint8*)WeightMipData + ChannelOffsets[Component->WeightmapLayerAllocations[LayerIdx].WeightmapTextureChannel];
-
-		for( int32 i=0;i<FMath::Square(WeightmapSize);i++ )
-		{
-			OutData[i] = SrcTextureData[i*4];
-		}
-
-		DataInterface->UnlockMip(Component->WeightmapTextures[Component->WeightmapLayerAllocations[LayerIdx].WeightmapTextureIndex], MipLevel);
-		return true;
-	}
+	ENGINE_API bool GetWeightmapTextureData(ULandscapeLayerInfoObject* LayerInfo, TArray<uint8>& OutData);
 
 	FColor* GetHeightData( int32 LocalX, int32 LocalY ) const
 	{
@@ -332,22 +214,7 @@ struct FLandscapeComponentDataInterface
 		return &HeightMipData[SubX + SubNumX*SubsectionSizeVerts + HeightmapComponentOffsetX + (SubY+SubNumY*SubsectionSizeVerts+HeightmapComponentOffsetY)*HeightmapStride];
 	}
 
-	FColor* GetXYOffsetData( int32 LocalX, int32 LocalY ) const
-	{
-#if LANDSCAPE_VALIDATE_DATA_ACCESS
-		check(Component);
-		check(LocalX >=0 && LocalY >=0 && LocalX < Component->ComponentSizeQuads+1 && LocalY < Component->ComponentSizeQuads+1);
-#endif
-
-		const int32 WeightmapSize = (Component->SubsectionSizeQuads+1) * Component->NumSubsections;
-		int32 SubNumX;
-		int32 SubNumY;
-		int32 SubX;
-		int32 SubY;
-		ComponentXYToSubsectionXY(LocalX, LocalY, SubNumX, SubNumY, SubX, SubY );
-
-		return &XYOffsetMipData[SubX + SubNumX*SubsectionSizeVerts + (SubY+SubNumY*SubsectionSizeVerts)*WeightmapSize];
-	}
+	ENGINE_API FColor* GetXYOffsetData(int32 LocalX, int32 LocalY) const;
 
 	uint16 GetHeight( int32 LocalX, int32 LocalY ) const
 	{
@@ -383,13 +250,7 @@ struct FLandscapeComponentDataInterface
 		GetXYOffset( X, Y, XOffset, YOffset );
 	}
 
-	FVector GetLocalVertex( int32 LocalX, int32 LocalY ) const
-	{
-		const float ScaleFactor = (float)Component->ComponentSizeQuads / (float)(ComponentSizeVerts-1);
-		float XOffset, YOffset;
-		GetXYOffset(LocalX, LocalY, XOffset, YOffset);
-		return FVector( LocalX * ScaleFactor + XOffset, LocalY * ScaleFactor + YOffset, LandscapeDataAccess::GetLocalHeight( GetHeight(LocalX, LocalY) ) );
-	}
+	ENGINE_API FVector GetLocalVertex(int32 LocalX, int32 LocalY) const;
 
 	void GetLocalTangentVectors( int32 LocalX, int32 LocalY, FVector& LocalTangentX, FVector& LocalTangentY, FVector& LocalTangentZ ) const
 	{
@@ -417,46 +278,11 @@ struct FLandscapeComponentDataInterface
 		GetLocalTangentVectors( X, Y, LocalTangentX, LocalTangentY, LocalTangentZ );
 	}
 
+	ENGINE_API FVector GetWorldVertex(int32 LocalX, int32 LocalY) const;
 
-	FVector GetWorldVertex( int32 LocalX, int32 LocalY ) const
-	{
-		return Component->ComponentToWorld.TransformPosition( GetLocalVertex( LocalX, LocalY ) );
-	}
-
-	void GetWorldTangentVectors( int32 LocalX, int32 LocalY, FVector& WorldTangentX, FVector& WorldTangentY, FVector& WorldTangentZ ) const
-	{
-		FColor* Data = GetHeightData( LocalX, LocalY );
-		WorldTangentZ.X = 2.f * (float)Data->B / 255.f - 1.f;
-		WorldTangentZ.Y = 2.f * (float)Data->A / 255.f - 1.f;
-		WorldTangentZ.Z = FMath::Sqrt(1.f - (FMath::Square(WorldTangentZ.X)+FMath::Square(WorldTangentZ.Y)));
-		WorldTangentX = FVector(-WorldTangentZ.Z, 0.f, WorldTangentZ.X);
-		WorldTangentY = FVector(0.f, WorldTangentZ.Z, -WorldTangentZ.Y);
-
-		WorldTangentX = Component->ComponentToWorld.TransformVectorNoScale(WorldTangentX);
-		WorldTangentY = Component->ComponentToWorld.TransformVectorNoScale(WorldTangentY);
-		WorldTangentZ = Component->ComponentToWorld.TransformVectorNoScale(WorldTangentZ);
-	}
+	ENGINE_API void GetWorldTangentVectors(int32 LocalX, int32 LocalY, FVector& WorldTangentX, FVector& WorldTangentY, FVector& WorldTangentZ) const;
 
-	void GetWorldPositionTangents( int32 LocalX, int32 LocalY, FVector& WorldPos, FVector& WorldTangentX, FVector& WorldTangentY, FVector& WorldTangentZ ) const
-	{
-		FColor* Data = GetHeightData( LocalX, LocalY );
-
-		WorldTangentZ.X = 2.f * (float)Data->B / 255.f - 1.f;
-		WorldTangentZ.Y = 2.f * (float)Data->A / 255.f - 1.f;
-		WorldTangentZ.Z = FMath::Sqrt(1.f - (FMath::Square(WorldTangentZ.X)+FMath::Square(WorldTangentZ.Y)));
-		WorldTangentX = FVector(WorldTangentZ.Z, 0.f, -WorldTangentZ.X);
-		WorldTangentY = WorldTangentZ ^ WorldTangentX;
-
-		uint16 Height = (Data->R << 8) + Data->G;
-
-		const float ScaleFactor = (float)Component->ComponentSizeQuads / (float)(ComponentSizeVerts-1);
-		float XOffset, YOffset;
-		GetXYOffset(LocalX, LocalY, XOffset, YOffset);
-		WorldPos = Component->ComponentToWorld.TransformPosition(FVector(LocalX * ScaleFactor + XOffset, LocalY * ScaleFactor + YOffset, LandscapeDataAccess::GetLocalHeight(Height)));
-		WorldTangentX = Component->ComponentToWorld.TransformVectorNoScale(WorldTangentX);
-		WorldTangentY = Component->ComponentToWorld.TransformVectorNoScale(WorldTangentY);
-		WorldTangentZ = Component->ComponentToWorld.TransformVectorNoScale(WorldTangentZ);
-	}
+	ENGINE_API void GetWorldPositionTangents(int32 LocalX, int32 LocalY, FVector& WorldPos, FVector& WorldTangentX, FVector& WorldTangentY, FVector& WorldTangentZ) const;
 
 	FVector GetWorldVertex( int32 VertexIndex ) const
 	{
diff --git a/Engine/Source/Runtime/Engine/Public/Landscape/LandscapeRender.h b/Engine/Source/Runtime/Engine/Public/Landscape/LandscapeRender.h
--- a/Engine/Source/Runtime/Engine/Public/Landscape/LandscapeRender.h
+++ b/Engine/Source/Runtime/Engine/Public/Landscape/LandscapeRender.h
@@ -13,11 +13,6 @@ LandscapeRender.h: New terrain rendering
 #include "LightMap.h"
 #include "ShadowMap.h"
 
-
-class FLandscapeVertexFactory;
-class FLandscapeVertexBuffer;
-class FLandscapeComponentSceneProxy;
-
 // This defines the number of border blocks to surround terrain by when generating lightmaps
 #define TERRAIN_PATCH_EXPAND_SCALAR	1
 
diff --git a/Engine/Source/Runtime/Engine/Public/Landscape/LandscapeSplineProxies.h b/Engine/Source/Runtime/Engine/Public/Landscape/LandscapeSplineProxies.h
--- a/Engine/Source/Runtime/Engine/Public/Landscape/LandscapeSplineProxies.h
+++ b/Engine/Source/Runtime/Engine/Public/Landscape/LandscapeSplineProxies.h
@@ -2,8 +2,6 @@
 
 #pragma once
 
-#include "Landscape/LandscapeSplineSegment.h"
-
 //////////////////////////////////////////////////////////////////////////
 // LANDSCAPE SPLINES HIT PROXY
 
@@ -68,10 +66,8 @@ struct HLandscapeSplineProxy_Tangent : public HLandscapeSplineProxy
 		End(InEnd)
 	{
 	}
-	virtual void Serialize(FArchive& Ar)
-	{
-		Ar << SplineSegment;
-	}
+	ENGINE_API virtual void Serialize(FArchive& Ar);
+
 	virtual EMouseCursor::Type GetMouseCursor()
 	{
 		return EMouseCursor::CardinalCross;
