EnableInput will always create an InputComponent if one doesn't exist even if there are no BP bindings to add.
PushInputComponent will put the input component in to the stack even if there are no bindings.

diff --git a/Engine/Source/Runtime/Engine/Private/Actor.cpp b/Engine/Source/Runtime/Engine/Private/Actor.cpp
--- a/Engine/Source/Runtime/Engine/Private/Actor.cpp
+++ b/Engine/Source/Runtime/Engine/Private/Actor.cpp
@@ -2295,20 +2295,23 @@ void AActor::EnableInput(APlayerController* PlayerController)
 		// If it doesn't exist create it and bind delegates
 		if (!InputComponent)
 		{
+			InputComponent = ConstructObject<UInputComponent>(UInputComponent::StaticClass(), this);
+			InputComponent->RegisterComponent();
+			InputComponent->bBlockInput = bBlockInput;
+
 			// Only do this if this actor is of a blueprint class
 			UBlueprintGeneratedClass* BGClass = Cast<UBlueprintGeneratedClass>(GetClass());
 			if(BGClass != NULL)
 			{
-				InputComponent = ConstructObject<UInputComponent>(UInputComponent::StaticClass(), this);
-				InputComponent->RegisterComponent();
-				InputComponent->bBlockInput = bBlockInput;
-
 				UInputDelegateBinding::BindInputDelegates(BGClass, InputComponent);
 			}
 		}
+		else
+		{
+			// Make sure we only have one instance of the InputComponent on the stack
+			PlayerController->PopInputComponent(InputComponent);
+		}
 
-		// Make sure we only have one instance of the InputComponent on the stack
-		PlayerController->PopInputComponent(InputComponent);
 		PlayerController->PushInputComponent(InputComponent);
 	}
 }
diff --git a/Engine/Source/Runtime/Engine/Private/PlayerController.cpp b/Engine/Source/Runtime/Engine/Private/PlayerController.cpp
--- a/Engine/Source/Runtime/Engine/Private/PlayerController.cpp
+++ b/Engine/Source/Runtime/Engine/Private/PlayerController.cpp
@@ -3996,14 +3996,7 @@ void APlayerController::PushInputComponent(UInputComponent* InputComponent)
 {
 	if (InputComponent)
 	{
-		if (InputComponent->HasBindings())
-		{
-			CurrentInputStack.Push(InputComponent);
-		}
-		else
-		{
-			UE_LOG(LogPlayerController, Warning, TEXT("InputComponent '%s' with no bindings pushed on the stack"), *InputComponent->GetFullName());
-		}
+		CurrentInputStack.Push(InputComponent);
 	}
 }
 
