CIS fix.

diff --git a/Engine/Source/Runtime/Engine/Private/Landscape/LandscapeDataAccess.cpp b/Engine/Source/Runtime/Engine/Private/Landscape/LandscapeDataAccess.cpp
--- a/Engine/Source/Runtime/Engine/Private/Landscape/LandscapeDataAccess.cpp
+++ b/Engine/Source/Runtime/Engine/Private/Landscape/LandscapeDataAccess.cpp
@@ -1,6 +1,9 @@
 // Copyright 1998-2014 Epic Games, Inc. All Rights Reserved.
 
 #include "EnginePrivate.h"
+
+#if WITH_EDITOR
+
 #include "Landscape/LandscapeDataAccess.h"
 #include "Landscape/LandscapeComponent.h"
 #include "Landscape/LandscapeInfo.h"
@@ -193,4 +196,6 @@ ENGINE_API void FLandscapeComponentDataInterface::GetWorldPositionTangents(int32
 	WorldTangentX = Component->ComponentToWorld.TransformVectorNoScale(WorldTangentX);
 	WorldTangentY = Component->ComponentToWorld.TransformVectorNoScale(WorldTangentY);
 	WorldTangentZ = Component->ComponentToWorld.TransformVectorNoScale(WorldTangentZ);
-}
\ No newline at end of file
+}
+
+#endif // WITH_EDITOR
\ No newline at end of file
