Don't unregister tick functions in EndPlay as it will break the case where a level has just been hidden, not unloaded from memory

diff --git a/Engine/Source/Runtime/Engine/Private/Actor.cpp b/Engine/Source/Runtime/Engine/Private/Actor.cpp
--- a/Engine/Source/Runtime/Engine/Private/Actor.cpp
+++ b/Engine/Source/Runtime/Engine/Private/Actor.cpp
@@ -1452,8 +1452,6 @@ void AActor::EndPlay(const EEndPlayReason::Type EndPlayReason)
 	ReceiveEndPlay(EndPlayReason);
 	OnEndPlay.Broadcast(EndPlayReason);
 
-	RegisterAllActorTickFunctions(false, true);
-
 	// Behaviors specific to an actor being unloaded due to a streaming level removal
 	if (EndPlayReason == EEndPlayReason::RemovedFromWorld)
 	{
