Fix spelling errors in comment

diff --git a/Engine/Source/Runtime/Engine/Classes/GameFramework/Actor.h b/Engine/Source/Runtime/Engine/Classes/GameFramework/Actor.h
--- a/Engine/Source/Runtime/Engine/Classes/GameFramework/Actor.h
+++ b/Engine/Source/Runtime/Engine/Classes/GameFramework/Actor.h
@@ -1690,7 +1690,7 @@ public:
 	virtual void OnConstruction(const FTransform& Transform) {}
 
 	/**
-	 * Helper function to regoster tje specified component, and add it to the serialized components array
+	 * Helper function to register the specified component, and add it to the serialized components array
 	 * @param	Component	Component to be finalized
 	 */
 	void FinishAndRegisterComponent(UActorComponent* Component);
