Don't clear the input stack when initializing player input as other actors may have already pushed their input component on to the stack

diff --git a/Engine/Source/Runtime/Engine/Private/PlayerController.cpp b/Engine/Source/Runtime/Engine/Private/PlayerController.cpp
--- a/Engine/Source/Runtime/Engine/Private/PlayerController.cpp
+++ b/Engine/Source/Runtime/Engine/Private/PlayerController.cpp
@@ -630,9 +630,6 @@ void APlayerController::InitInputSystem()
 		PlayerInput = ConstructObject<UPlayerInput>(UPlayerInput::StaticClass(), this);
 	}
 
-	// initialize input stack
-	CurrentInputStack.Empty();
-
 	SetupInputComponent();
 
 	CurrentMouseCursor = DefaultMouseCursor;
