Missing cvar for rendering settings is fatal, too easy to miss otherwise

diff --git a/Engine/Source/Runtime/Engine/Private/RenderingSettings.cpp b/Engine/Source/Runtime/Engine/Private/RenderingSettings.cpp
--- a/Engine/Source/Runtime/Engine/Private/RenderingSettings.cpp
+++ b/Engine/Source/Runtime/Engine/Private/RenderingSettings.cpp
@@ -36,7 +36,7 @@ void URendererSettings::PostInitProperties()
 				}
 				else
 				{
-					UE_LOG(LogTemp, Error, TEXT("URendererSettings failed to find console variable %s for %s"), *CVarName, *Property->GetName());
+					UE_LOG(LogTemp, Fatal, TEXT("URendererSettings failed to find console variable %s for %s"), *CVarName, *Property->GetName());
 				}
 			}
 		}
