UE4: Added null checking for missing expressions and comments when fixing up materials on load.

diff --git a/Engine/Source/Runtime/Engine/Private/Materials/Material.cpp b/Engine/Source/Runtime/Engine/Private/Materials/Material.cpp
--- a/Engine/Source/Runtime/Engine/Private/Materials/Material.cpp
+++ b/Engine/Source/Runtime/Engine/Private/Materials/Material.cpp
@@ -3611,16 +3611,22 @@ void UMaterial::FlipExpressionPositions(const TArray<UMaterialExpression*>& Expr
 	for (int32 ExpressionIndex = 0; ExpressionIndex < Expressions.Num(); ExpressionIndex++)
 	{
 		UMaterialExpression* Expression = Expressions[ExpressionIndex];
-		Expression->MaterialExpressionEditorX = -Expression->MaterialExpressionEditorX * PosScaling;
-		Expression->MaterialExpressionEditorY *= PosScaling;
+		if (Expression)
+		{
+			Expression->MaterialExpressionEditorX = -Expression->MaterialExpressionEditorX * PosScaling;
+			Expression->MaterialExpressionEditorY *= PosScaling;
+		}
 	}
 	for (int32 ExpressionIndex = 0; ExpressionIndex < Comments.Num(); ExpressionIndex++)
 	{
 		UMaterialExpressionComment* Comment = Comments[ExpressionIndex];
-		Comment->MaterialExpressionEditorX = (-Comment->MaterialExpressionEditorX - Comment->SizeX) * PosScaling;
-		Comment->MaterialExpressionEditorY *= PosScaling;
-		Comment->SizeX *= PosScaling;
-		Comment->SizeY *= PosScaling;
+		if (Comment)
+		{
+			Comment->MaterialExpressionEditorX = (-Comment->MaterialExpressionEditorX - Comment->SizeX) * PosScaling;
+			Comment->MaterialExpressionEditorY *= PosScaling;
+			Comment->SizeX *= PosScaling;
+			Comment->SizeY *= PosScaling;
+		}
 	}
 }
 
