#ue4 - Remove delay when characters spawn in, and start with the correct movement mode.

- DefaultLandMovementMode is now MOVE_Walking instead of MOVE_Flying.
- Add UCharacterMovementComponent::GetMovementBase() helper.

TTP 338877

diff --git a/Engine/Source/Runtime/Engine/Classes/GameFramework/CharacterMovementComponent.h b/Engine/Source/Runtime/Engine/Classes/GameFramework/CharacterMovementComponent.h
--- a/Engine/Source/Runtime/Engine/Classes/GameFramework/CharacterMovementComponent.h
+++ b/Engine/Source/Runtime/Engine/Classes/GameFramework/CharacterMovementComponent.h
@@ -720,6 +720,10 @@ public:
 	/** Adjust distance from floor, trying to maintain a slight offset from the floor when walking (based on CurrentFloor). */
 	virtual void AdjustFloorHeight();
 
+	/** Return PrimitiveComponent we are based on (standing and walking on). */
+	UFUNCTION(BlueprintCallable, Category="Pawn|Components|CharacterMovement")
+	UPrimitiveComponent* GetMovementBase() const;
+
 	/** Update or defer updating of position based on Base movement */
 	virtual void MaybeUpdateBasedMovement(float DeltaSeconds);
 
diff --git a/Engine/Source/Runtime/Engine/Classes/GameFramework/Pawn.h b/Engine/Source/Runtime/Engine/Classes/GameFramework/Pawn.h
--- a/Engine/Source/Runtime/Engine/Classes/GameFramework/Pawn.h
+++ b/Engine/Source/Runtime/Engine/Classes/GameFramework/Pawn.h
@@ -34,7 +34,7 @@ class ENGINE_API APawn : public AActor, public INavAgentInterface
 	UFUNCTION(BlueprintCallable, meta=(Tooltip="Return our PawnMovementComponent, if we have one."), Category="Pawn")
 	virtual UPawnMovementComponent* GetMovementComponent() const;
 
-	/** Return Actor we are based on (standing on, attached to, and moving on). */
+	/** Return PrimitiveComponent we are based on (standing on, attached to, and moving on). */
 	virtual UPrimitiveComponent* GetMovementBase() const { return NULL; }
 
 public:
@@ -244,9 +244,11 @@ public:
 	 */
 	virtual void PossessedBy(AController* NewController);
 
+	/** Event called when the Pawn is possessed by a Controller (normally only occurs on the server/standalone). */
 	UFUNCTION(BlueprintImplementableEvent, meta=(FriendlyName = "Possessed"))
 	void ReceivePossessed(AController* NewController);
 
+	/** Event called when the Pawn is no longer possessed by a Controller. */
 	UFUNCTION(BlueprintImplementableEvent, meta=(FriendlyName = "Unpossessed"))
 	void ReceiveUnpossessed(AController* OldController);
 
diff --git a/Engine/Source/Runtime/Engine/Private/Character.cpp b/Engine/Source/Runtime/Engine/Private/Character.cpp
--- a/Engine/Source/Runtime/Engine/Private/Character.cpp
+++ b/Engine/Source/Runtime/Engine/Private/Character.cpp
@@ -626,9 +626,16 @@ void ACharacter::TurnOff()
 void ACharacter::Restart()
 {
 	Super::Restart();
+
 	bPressedJump = false;
 	JumpKeyHoldTime = 0.0f;
 	ClearJumpInput();
+	UnCrouch(true);
+
+	if (CharacterMovement)
+	{
+		CharacterMovement->SetDefaultMovementMode();
+	}
 }
 
 void ACharacter::PawnClientRestart()
diff --git a/Engine/Source/Runtime/Engine/Private/CharacterMovementComponent.cpp b/Engine/Source/Runtime/Engine/Private/CharacterMovementComponent.cpp
--- a/Engine/Source/Runtime/Engine/Private/CharacterMovementComponent.cpp
+++ b/Engine/Source/Runtime/Engine/Private/CharacterMovementComponent.cpp
@@ -196,7 +196,8 @@ UCharacterMovementComponent::UCharacterMovementComponent(const class FPostConstr
 	CrouchedHalfHeight = 40.0f;
 	PendingLaunchVelocity = FVector::ZeroVector;
 	DefaultWaterMovementMode = MOVE_Swimming;
-	DefaultLandMovementMode = MOVE_Falling;
+	DefaultLandMovementMode = MOVE_Walking;
+	bForceNextFloorCheck = true;
 	bForceBraking_DEPRECATED = false;
 	bShrinkProxyCapsule = true;
 	bCanWalkOffLedges = true;
@@ -560,9 +561,15 @@ void UCharacterMovementComponent::SetDefaultMovementMode()
 	{
 		SetMovementMode(DefaultWaterMovementMode);
 	}
-	else if ( !CharacterOwner || !IsFalling() )
+	else if ( !CharacterOwner || MovementMode != DefaultLandMovementMode )
 	{
 		SetMovementMode(DefaultLandMovementMode);
+
+		// Avoid 1-frame delay if trying to walk but walking fails at this location.
+		if (MovementMode == MOVE_Walking && GetMovementBase() == NULL)
+		{
+			SetMovementMode(MOVE_Falling);
+		}
 	}
 }
 
@@ -1133,6 +1140,11 @@ void UCharacterMovementComponent::SimulateMovement(float DeltaSeconds)
 	LastUpdateLocation = UpdatedComponent ? UpdatedComponent->GetComponentLocation() : FVector::ZeroVector;
 }
 
+UPrimitiveComponent* UCharacterMovementComponent::GetMovementBase() const
+{
+	return CharacterOwner ? CharacterOwner->GetMovementBase() : NULL;
+}
+
 void UCharacterMovementComponent::SetBase( UPrimitiveComponent* NewBase, FName BoneName, bool bNotifyActor )
 {
 	if (CharacterOwner)
@@ -3561,7 +3573,14 @@ void UCharacterMovementComponent::OnTeleported()
 	// If we were walking but no longer have a valid base or floor, start falling.
 	if (!CurrentFloor.IsWalkableFloor() || (OldBase && !NewBase))
 	{
-		SetDefaultMovementMode();
+		if (DefaultLandMovementMode == MOVE_Walking)
+		{
+			SetMovementMode(MOVE_Falling);
+		}
+		else
+		{
+			SetDefaultMovementMode();
+		}
 	}
 }
 
diff --git a/Engine/Source/Runtime/Engine/Private/Controller.cpp b/Engine/Source/Runtime/Engine/Private/Controller.cpp
--- a/Engine/Source/Runtime/Engine/Private/Controller.cpp
+++ b/Engine/Source/Runtime/Engine/Private/Controller.cpp
@@ -565,13 +565,6 @@ void AController::UpdateNavigationComponents()
 	{
 		PathFollowingComp->UpdateCachedComponents();
 	}
-
-	// initialize movement mode in characters
-	ACharacter* MyCharacter = Cast<ACharacter>(GetPawn());
-	if (MyCharacter && MyCharacter->CharacterMovement)
-	{
-		MyCharacter->CharacterMovement->SetDefaultMovementMode();
-	}
 }
 
 void AController::InitNavigationControl(UNavigationComponent*& PathFindingComp, UPathFollowingComponent*& PathFollowingComp)
diff --git a/Engine/Source/Runtime/Engine/Private/PlayerController.cpp b/Engine/Source/Runtime/Engine/Private/PlayerController.cpp
--- a/Engine/Source/Runtime/Engine/Private/PlayerController.cpp
+++ b/Engine/Source/Runtime/Engine/Private/PlayerController.cpp
@@ -722,13 +722,12 @@ void APlayerController::ClientRestart_Implementation(APawn* NewPawn)
 	
 	if (Role < ROLE_Authority)
 	{
+		ChangeState( NAME_Playing );
 		if (bAutoManageActiveCameraTarget)
 		{
 			AutoManageActiveCameraTarget(GetPawn());
 			ResetCameraMode();
 		}
-		
-		ChangeState( NAME_Playing );
 	}
 }
 
@@ -756,7 +755,6 @@ void APlayerController::Possess(APawn* PawnToPossess)
 		check(GetPawn() != NULL);
 
 		GetPawn()->SetActorTickEnabled(true);
-		GetPawn()->Restart();
 
 		INetworkPredictionInterface* NetworkPredictionInterface = GetPawn() ? InterfaceCast<INetworkPredictionInterface>(GetPawn()->GetMovementComponent()) : NULL;
 		if (NetworkPredictionInterface)
@@ -764,10 +762,18 @@ void APlayerController::Possess(APawn* PawnToPossess)
 			NetworkPredictionInterface->ResetPredictionData_Server();
 		}
 
-		ChangeState( NAME_Playing );
 		AcknowledgedPawn = NULL;
+		
+		// Local PCs will have the Restart() triggered right away in ClientRestart (via PawnClientRestart()), but the server should call Restart() locally for remote PCs.
+		// We're really just trying to avoid calling Restart() multiple times.
+		if (!IsLocalPlayerController())
+		{
+			GetPawn()->Restart();
+		}
+
 		ClientRestart(GetPawn());
 		
+		ChangeState( NAME_Playing );
 		if (bAutoManageActiveCameraTarget)
 		{
 			AutoManageActiveCameraTarget(GetPawn());
@@ -3932,18 +3938,6 @@ void APlayerController::ChangeState(FName NewState)
 
 void APlayerController::BeginPlayingState()
 {
-	if (GetPawn() != NULL)
-	{
-		if (GetCharacter() != NULL)
-		{
-			GetCharacter()->UnCrouch(false);
-			UCharacterMovementComponent* CharacterMovement = GetCharacter()->CharacterMovement;
-			if (CharacterMovement && !CharacterMovement->IsFalling() && GetCharacter()->GetRootComponent() && !GetCharacter()->GetRootComponent()->IsSimulatingPhysics()) // FIXME HACK!!!
-			{
-				CharacterMovement->SetMovementMode(MOVE_Walking);
-			}
-		}		
-	}
 }
 
 void APlayerController::EndSpectatingState()
diff --git a/Engine/Source/Runtime/Engine/Private/SpectatorPawn.cpp b/Engine/Source/Runtime/Engine/Private/SpectatorPawn.cpp
--- a/Engine/Source/Runtime/Engine/Private/SpectatorPawn.cpp
+++ b/Engine/Source/Runtime/Engine/Private/SpectatorPawn.cpp
@@ -25,6 +25,13 @@ ASpectatorPawn::ASpectatorPawn(const class FPostConstructInitializeProperties& P
 
 void ASpectatorPawn::PossessedBy(class AController* NewController)
 {
+	AController* const OldController = Controller;
 	Controller = NewController;
+
+	// dispatch Blueprint event if necessary
+	if (OldController != NewController)
+	{
+		ReceivePossessed(Controller);
+	}
 }
 
