#ue4 - CharacterMovementComponent Buoyancy default changed to be neutral (1.0).

diff --git a/Engine/Source/Runtime/Engine/Private/CharacterMovementComponent.cpp b/Engine/Source/Runtime/Engine/Private/CharacterMovementComponent.cpp
--- a/Engine/Source/Runtime/Engine/Private/CharacterMovementComponent.cpp
+++ b/Engine/Source/Runtime/Engine/Private/CharacterMovementComponent.cpp
@@ -194,6 +194,7 @@ UCharacterMovementComponent::UCharacterMovementComponent(const class FPostConstr
 	Mass = 100.0f;
 	bJustTeleported = true;
 	CrouchedHalfHeight = 40.0f;
+	Buoyancy = 1.0f;
 	PendingLaunchVelocity = FVector::ZeroVector;
 	DefaultWaterMovementMode = MOVE_Swimming;
 	DefaultLandMovementMode = MOVE_Walking;
@@ -225,7 +226,7 @@ UCharacterMovementComponent::UCharacterMovementComponent(const class FPostConstr
 	bImpartBaseAngularVelocity = true;
 	bAlwaysCheckFloor = true;
 
-	// default character can jump and walk
+	// default character can jump, walk, and swim
 	NavAgentProps.bCanJump = true;
 	NavAgentProps.bCanWalk = true;
 	NavAgentProps.bCanSwim = true;
