Remove workaround for navmesh debug rendering race condition

diff --git a/Engine/Source/Runtime/Engine/Private/AI/Navigation/RecastNavMesh.cpp b/Engine/Source/Runtime/Engine/Private/AI/Navigation/RecastNavMesh.cpp
--- a/Engine/Source/Runtime/Engine/Private/AI/Navigation/RecastNavMesh.cpp
+++ b/Engine/Source/Runtime/Engine/Private/AI/Navigation/RecastNavMesh.cpp
@@ -158,8 +158,7 @@ uint32 FRecastDebugGeometry::GetAllocatedSize() const
 
 	for (int i = 0; i < Clusters.Num(); ++i)
 	{
-//		Removed this as a race condition causes crash: TTP#343021
-//		Size += Clusters[i].MeshIndices.GetAllocatedSize();
+		Size += Clusters[i].MeshIndices.GetAllocatedSize();
 	}
 
 	return Size;
