Fix registration of landscape newly-added collision components

diff --git a/Engine/Source/Runtime/Engine/Private/Landscape/LandscapeEdit.cpp b/Engine/Source/Runtime/Engine/Private/Landscape/LandscapeEdit.cpp
--- a/Engine/Source/Runtime/Engine/Private/Landscape/LandscapeEdit.cpp
+++ b/Engine/Source/Runtime/Engine/Private/Landscape/LandscapeEdit.cpp
@@ -716,7 +716,7 @@ void ULandscapeComponent::UpdateCollisionHeightData(const FColor* HeightmapTextu
 	// Set new collision component to pointer
 	CollisionComponent = CollisionComp;
 
-	if (ChangeType && CreatedNew)
+	if (ChangeType || CreatedNew)
 	{
 		Proxy->RegisterAllComponents();
 	}
