Added a new feature to support morph targets for clothing

diff --git a/Engine/Source/Editor/Persona/Private/PersonaMeshDetails.cpp b/Engine/Source/Editor/Persona/Private/PersonaMeshDetails.cpp
--- a/Engine/Source/Editor/Persona/Private/PersonaMeshDetails.cpp
+++ b/Engine/Source/Editor/Persona/Private/PersonaMeshDetails.cpp
@@ -2004,6 +2004,14 @@ void FPersonaMeshDetails::HandleSectionsComboBoxSelectionChanged( TSharedPtr<FSt
 			int32 SelectedIndex = ComboSelectedIndices.IsValidIndex(SectionIndex) ? ComboSelectedIndices[SectionIndex] : 0;
 			ComboBoxes[SectionIndex]->SetSelectedItem(ComboStrings[SelectedIndex]);
 		}
+		else
+		{
+			// when clothing is mapped, need to re-prepare cloth morph targets
+			if (PersonaPtr->PreviewComponent)
+			{
+				PersonaPtr->PreviewComponent->bPreparedClothMorphTargets = false;
+			}
+		}
 	}
 	else
 	{
diff --git a/Engine/Source/Editor/UnrealEd/Private/Animation/DebugSkelMeshComponent.cpp b/Engine/Source/Editor/UnrealEd/Private/Animation/DebugSkelMeshComponent.cpp
--- a/Engine/Source/Editor/UnrealEd/Private/Animation/DebugSkelMeshComponent.cpp
+++ b/Engine/Source/Editor/UnrealEd/Private/Animation/DebugSkelMeshComponent.cpp
@@ -93,6 +93,8 @@ UDebugSkelMeshComponent::UDebugSkelMeshComponent(const class FPostConstructIniti
 
 #if WITH_APEX_CLOTHING
 	SectionsDisplayMode = ESectionDisplayMode::None;
+	// always shows cloth morph target when previewing in editor
+	bClothMorphTarget = true;
 #endif //#if WITH_APEX_CLOTHING
 }
 
@@ -292,7 +294,6 @@ bool UDebugSkelMeshComponent::IsWindEnabled() const
 	return bEnableWind;
 }
 
-
 void UDebugSkelMeshComponent::EnablePreview(bool bEnable, UAnimationAsset* PreviewAsset, UVertexAnimation* PreviewVertexAnim)
 {
 	if (PreviewInstance)
diff --git a/Engine/Source/Runtime/Engine/Classes/Components/SkeletalMeshComponent.h b/Engine/Source/Runtime/Engine/Classes/Components/SkeletalMeshComponent.h
--- a/Engine/Source/Runtime/Engine/Classes/Components/SkeletalMeshComponent.h
+++ b/Engine/Source/Runtime/Engine/Classes/Components/SkeletalMeshComponent.h
@@ -59,7 +59,7 @@ namespace physx
 
 class FPhysScene;
 
-/** thie is a class to manage an APEX clothing actor */
+/** a class to manage an APEX clothing actor */
 class FClothingActor
 {
 public:
@@ -100,6 +100,20 @@ struct FClothSimulData
 	TArray<FVector4> ClothSimulNormals;
 };
 
+/**  for storing precomputed cloth morph target data */
+struct FClothMorphTargetData
+{
+	FName MorphTargetName;
+	// save a previous weight to compare whether weight was changed or not
+	float PrevWeight;
+	// an index of clothing assets which this morph target data is used in
+	int32 ClothAssetIndex;
+	// original positions which only this cloth section is including / extracted from a morph target
+	TArray<FVector> OriginPos;
+	// delta positions to morph this cloth section
+	TArray<FVector> PosDelta;
+};
+
 #if WITH_CLOTH_COLLISION_DETECTION
 
 class FClothCollisionPrimitive
@@ -349,6 +363,15 @@ public:
 	 */
 	UPROPERTY(EditAnywhere, BlueprintReadWrite, Category = Clothing)
 	uint32 bLocalSpaceSimulation : 1;
+
+	/**
+	 * cloth morph target option
+	 * This option will be applied only before playing because should do pre-calculation to reduce computation time for run-time play
+	 * so it's impossible to change this option in run-time
+	 */
+	UPROPERTY(EditAnywhere, Category = Clothing)
+	uint32 bClothMorphTarget : 1;
+
 	/** reset the clothing after moving the clothing position (called teleport) */
 	UPROPERTY(EditAnywhere, BlueprintReadWrite, Category=Clothing)
 	uint32 bResetAfterTeleport:1;
@@ -646,6 +669,10 @@ public:
 	 * use until Apex clothing bug is resolved 
 	 */
 	bool bNeedTeleportAndResetOnceMore;
+	/** used for checking whether cloth morph target data were pre-computed or not */
+	bool bPreparedClothMorphTargets;
+	/** precomputed actual cloth morph target data */
+	TArray<FClothMorphTargetData> ClothMorphTargets;
 
 #if WITH_CLOTH_COLLISION_DETECTION
 	/** increase every tick to update clothing collision  */
@@ -1005,8 +1032,9 @@ public:
 	/** 
 	* APEX clothing actor is created from APEX clothing asset for cloth simulation 
 	* create only if became invalid
+	* BlendedData : added for cloth morph target but not used commonly
 	*/
-	bool CreateClothingActor(int32 AssetIndex, TSharedPtr<FClothingAssetWrapper> ClothingAsset);
+	bool CreateClothingActor(int32 AssetIndex, TSharedPtr<FClothingAssetWrapper> ClothingAsset, TArray<FVector>* BlendedDelta=NULL);
 	/** should call this method if occurred any changes in clothing assets */
 	void ValidateClothingActors();
 	/** add bounding box for cloth */
@@ -1015,6 +1043,17 @@ public:
 	void SetClothingLOD(int32 LODIndex);
 	/** check whether clothing teleport is needed or not to avoid a weird simulation result */
 	virtual void CheckClothTeleport(float DeltaTime);
+
+	/** 
+	* methods for cloth morph targets 
+	*/
+	/** pre-compute morph target data for clothing */
+	void PrepareClothMorphTargets();
+	/** change morph target mapping when active morph target is changed */
+	void ChangeClothMorphTargetMapping(FClothMorphTargetData& MorphData, FName CurrentActivatedMorphName);
+	/** update active morph target's blending data when morph weight is changed */
+	void UpdateClothMorphTarget();
+
 	/** 
 	 * Updates all clothing animation states including ComponentToWorld-related states.
 	 */
diff --git a/Engine/Source/Runtime/Engine/Classes/Engine/SkeletalMesh.h b/Engine/Source/Runtime/Engine/Classes/Engine/SkeletalMesh.h
--- a/Engine/Source/Runtime/Engine/Classes/Engine/SkeletalMesh.h
+++ b/Engine/Source/Runtime/Engine/Classes/Engine/SkeletalMesh.h
@@ -501,7 +501,7 @@ struct FClothingAssetData
 #if WITH_APEX_CLOTHING
 	TSharedPtr<FClothingAssetWrapper> ApexClothingAsset;
 
-	/* Collision volume data for showing to the users whether collision shape is correct or not */
+	/** Collision volume data for showing to the users whether collision shape is correct or not */
 	TArray<FApexClothCollisionVolumeData> ClothCollisionVolumes;
 	TArray<uint32> ClothCollisionConvexPlaneIndices;
 	TArray<FClothBonePlane> ClothCollisionVolumePlanes;
@@ -513,7 +513,8 @@ struct FClothingAssetData
 	 * Num of this array means LOD number of clothing physical meshes 
 	 */
 	TArray<FClothVisualizationInfo> ClothVisualizationInfos;
-
+	/** currently mapped morph target name */
+	FName PreparedMorphTargetName;
 #endif// #if WITH_APEX_CLOTHING
 
 	// serialization
diff --git a/Engine/Source/Runtime/Engine/Private/SkeletalMeshComponent.cpp b/Engine/Source/Runtime/Engine/Private/SkeletalMeshComponent.cpp
--- a/Engine/Source/Runtime/Engine/Private/SkeletalMeshComponent.cpp
+++ b/Engine/Source/Runtime/Engine/Private/SkeletalMeshComponent.cpp
@@ -118,6 +118,7 @@ USkeletalMeshComponent::USkeletalMeshComponent(const class FPostConstructInitial
 	ClothTeleportDistThresholdSquared = TeleportDistanceThreshold * TeleportDistanceThreshold;
 	bNeedTeleportAndResetOnceMore = false;
 	ClothBlendWeight = 1.0f;
+	bPreparedClothMorphTargets = false;
 #if WITH_CLOTH_COLLISION_DETECTION
 	ClothingCollisionRevision = 0;
 #endif// #if WITH_CLOTH_COLLISION_DETECTION
diff --git a/Engine/Source/Runtime/Engine/Private/SkeletalMeshComponentPhysics.cpp b/Engine/Source/Runtime/Engine/Private/SkeletalMeshComponentPhysics.cpp
--- a/Engine/Source/Runtime/Engine/Private/SkeletalMeshComponentPhysics.cpp
+++ b/Engine/Source/Runtime/Engine/Private/SkeletalMeshComponentPhysics.cpp
@@ -23,6 +23,9 @@
 	#include "NxClothingAsset.h"
 	#include "NxClothingActor.h"
 	#include "NxClothingCollision.h"
+	// for cloth morph target	
+	#include "Animation/VertexAnim/VertexAnimBase.h"
+
 #endif// #if WITH_APEX_CLOTHING
 
 #endif//#if WITH_APEX
@@ -314,12 +317,12 @@ bool USkeletalMesh::IsMappedClothingLOD(int32 InLODIndex, int32 InAssetIndex)
 
 		// loop reversely for optimized search
 		for (int32 SecIdx = NumSections-1; SecIdx >= 0; SecIdx--)
-		{
+	{
 			int32 ClothAssetIndex = LODModel.Chunks[LODModel.Sections[SecIdx].ChunkIndex].CorrespondClothAssetIndex;
 			if (ClothAssetIndex == InAssetIndex)
-			{
-				return true;
-			}
+		{
+			return true;
+		}
 			else if (ClothAssetIndex == INDEX_NONE) // no more cloth sections
 			{
 				return false;
@@ -341,14 +344,14 @@ int32 USkeletalMesh::GetClothAssetIndex(int32 LODIndex, int32 SectionIndex)
 	}
 	FStaticLODModel& LODModel = Resource->LODModels[LODIndex];
 
-	// no sections
+	//no sections
 	if(!LODModel.Sections.IsValidIndex(SectionIndex))
 	{
 		return INDEX_NONE;
 	}
 	int16 ClothSecIdx = LODModel.Sections[SectionIndex].CorrespondClothSectionIndex;
 
-	// no mapping
+	//no mapping
 	if(ClothSecIdx < 0)
 	{
 		return INDEX_NONE;
@@ -1983,7 +1986,7 @@ void USkeletalMeshComponent::ValidateClothingActors()
  * APEX clothing actor is created from APEX clothing asset for cloth simulation 
  * If this is invalid, re-create actor , but if valid ,just skip to create
 */
-bool USkeletalMeshComponent::CreateClothingActor(int32 AssetIndex, TSharedPtr<FClothingAssetWrapper> ClothingAssetWrapper)
+bool USkeletalMeshComponent::CreateClothingActor(int32 AssetIndex, TSharedPtr<FClothingAssetWrapper> ClothingAssetWrapper, TArray<FVector>* BlendedDelta)
 {	
 	int32 NumActors = ClothingActors.Num();
 	int32 ActorIndex = -1;
@@ -2053,6 +2056,21 @@ bool USkeletalMeshComponent::CreateClothingActor(int32 AssetIndex, TSharedPtr<FC
 	verify(NxParameterized::setParamBool(*ActorDesc, "maxDistanceScale.Multipliable", true));
 	verify(NxParameterized::setParamF32(*ActorDesc, "maxDistanceScale.Scale", ClothMaxDistanceScale));
 
+	// apply delta positions of cloth morph target
+	if (BlendedDelta)
+	{
+		int32 NumBlendData = BlendedDelta->Num();
+		TArray<PxVec3> PxBlendedData;
+		PxBlendedData.AddUninitialized(NumBlendData);
+		for (int32 Index = 0; Index < NumBlendData; Index++)
+		{
+			PxBlendedData[Index] = U2PVector((*BlendedDelta)[Index]);
+		}
+		NxParameterized::Handle md(*ActorDesc, "morphDisplacements");
+		md.resizeArray(PxBlendedData.Num());
+		md.setParamVec3Array(PxBlendedData.GetTypedData(), PxBlendedData.Num());
+	}
+
 	FPhysScene* PhysScene = NULL;
 
 	if(GetWorld())
@@ -2143,7 +2161,7 @@ void USkeletalMeshComponent::SetClothingLOD(int32 LODIndex)
 			// decide whether should enable or disable
 			if (!IsMappedClothLOD || (LODIndex >= NumClothLODs))
 			{
-				// disable clothing simulation
+				//disable clothing simulation
 				Actor.ApexClothingActor->forcePhysicalLod(0);
 			}
 			else
@@ -3157,9 +3175,259 @@ void USkeletalMeshComponent::CheckClothTeleport(float DeltaTime)
 	PrevRootBoneMatrix = CurRootBoneMat;
 }
 
+void USkeletalMeshComponent::ChangeClothMorphTargetMapping(FClothMorphTargetData& MorphData, FName CurrentActivatedMorphName)
+{
+	int32 AssetIndex = MorphData.ClothAssetIndex;
+
+	if(SkeletalMesh && SkeletalMesh->ClothingAssets.IsValidIndex(AssetIndex))
+	{
+		FClothingAssetData& Asset = SkeletalMesh->ClothingAssets[AssetIndex];
+		// if different from prepared mapping, should do re-mapping
+		if (Asset.PreparedMorphTargetName != CurrentActivatedMorphName)
+		{
+			TArray<PxVec3>	ClothOriginalPosArray;
+			int32 NumOriginPos = MorphData.OriginPos.Num();
+			ClothOriginalPosArray.AddUninitialized(NumOriginPos);
+			for (int32 Index = 0; Index < NumOriginPos; Index++)
+			{
+				ClothOriginalPosArray[Index] = U2PVector(MorphData.OriginPos[Index]);
+			}
+
+			NxClothingAsset* ClothingAsset = Asset.ApexClothingAsset->GetAsset();
+			float Epsilon = 0.0f;
+			uint32 NumMapped = ClothingAsset->prepareMorphTargetMapping(ClothOriginalPosArray.GetTypedData(), NumOriginPos, Epsilon);
+
+			if ((int32)NumMapped < NumOriginPos)
+			{
+				// @TODO : 
+				// The mapping still worked, but some vertices were mapped to an original vertex with more than epsilon differences.
+				// Either bump the epsilon or, if too many failed, emit an error message about a potentially bad mapping
+
+				// if more than half failed
+				if ((int32)NumMapped < NumOriginPos / 2)
+				{
+					FMessageDialog::Open(EAppMsgType::Ok, NSLOCTEXT("SkelmeshComponent", "Warning_ClothMorphTargetMapping", "Mapping vertices for Cloth Morph Target too many failed! It could introduce a bad morphing result."));
+				}			
+			}
+
+			// change prepared morph target name
+			Asset.PreparedMorphTargetName = CurrentActivatedMorphName;
+		}
+
+	}
+}
+
+void USkeletalMeshComponent::PrepareClothMorphTargets()
+{	
+	// if didn't turn on the cloth morph target option or already precomputed
+	if (!bClothMorphTarget || bPreparedClothMorphTargets)
+	{ 
+		return;
+	}
+
+	ClothMorphTargets.Empty();
+
+	for (UMorphTarget* MorphTarget : SkeletalMesh->MorphTargets)
+	{
+		if (MorphTarget)
+		{
+			FName MorphTargetName = MorphTarget->GetFName();
+
+			int32 NumVerts;
+			FVertexAnimDelta* Vertices = MorphTarget->GetDeltasAtTime(0.0f, 0, NULL, NumVerts);
+
+			check(MeshObject);
+			// should exist at least 1 LODModel 
+			check(MeshObject->GetSkeletalMeshResource().LODModels.Num() > 0);
+
+			FStaticLODModel& Model = MeshObject->GetSkeletalMeshResource().LODModels[0];
+
+			// Find the chunk and vertex within that chunk, and skinning type, for this vertex.
+			int32 ChunkIndex;
+			int32 VertIndexInChunk;
+			bool bSoftVertex;
+			bool bHasExtraBoneInfluences;
+			int32 SectionIndex = INDEX_NONE;
+
+			int32 NumAssets = SkeletalMesh->ClothingAssets.Num();
+			TArray<TArray<PxVec3>> ClothOriginalPosArray;
+			TArray<TArray<FVector>> ClothPositionDeltaArray;
+			ClothOriginalPosArray.AddZeroed(NumAssets);
+			ClothPositionDeltaArray.AddZeroed(NumAssets);
+
+			int ClothChunkIndex = INDEX_NONE;
+
+			for (int32 VertIdx = 0; VertIdx < NumVerts; VertIdx++)
+			{
+				Model.GetChunkAndSkinType(Vertices[VertIdx].SourceIdx, ChunkIndex, VertIndexInChunk, bSoftVertex, bHasExtraBoneInfluences);
+
+				FSkelMeshChunk& Chunk = Model.Chunks[ChunkIndex];
+
+				for (int32 SecIdx = 0; SecIdx < Model.Sections.Num(); SecIdx++)
+				{
+					FSkelMeshSection& Section = Model.Sections[SecIdx];
+					if (Section.ChunkIndex == ChunkIndex)
+					{
+						// if current section is disabled and the corresponding cloth section is visible
+						if (Section.bDisabled && Section.CorrespondClothSectionIndex >= 0)
+						{
+							SectionIndex = SecIdx;
+							ClothChunkIndex = Model.Sections[Section.CorrespondClothSectionIndex].ChunkIndex;
+						}
+						break;
+					}
+				}
+
+				if (SectionIndex != INDEX_NONE)
+				{
+					FVector Position;
+					if (bSoftVertex)
+					{
+						Position = Chunk.SoftVertices[VertIndexInChunk].Position;
+					}
+					else
+					{
+						Position = Chunk.RigidVertices[VertIndexInChunk].Position;
+					}
+
+					int32 AssetIndex = Model.Chunks[ClothChunkIndex].CorrespondClothAssetIndex;
+					// save to original positions
+					ClothOriginalPosArray[AssetIndex].Add(U2PVector(Position));
+					ClothPositionDeltaArray[AssetIndex].Add(Vertices[VertIdx].PositionDelta);
+				}
+			}
+
+			// fill in FClothMorphTargetData array
+			for (int32 AssetIdx = 0; AssetIdx < NumAssets; AssetIdx++)
+			{
+				if (ClothOriginalPosArray[AssetIdx].Num() > 0)
+				{
+					NxClothingAsset* ClothingAsset = SkeletalMesh->ClothingAssets[AssetIdx].ApexClothingAsset->GetAsset();
+					float Epsilon = 0.0f;
+					uint32 NumMapped = ClothingAsset->prepareMorphTargetMapping(ClothOriginalPosArray[AssetIdx].GetTypedData(), ClothOriginalPosArray[AssetIdx].Num(), Epsilon);
+
+					int32 NumOriginPos = ClothOriginalPosArray[AssetIdx].Num();
+					if ((int32)NumMapped < NumOriginPos)
+					{
+						// @TODO : 
+						// The mapping still worked, but some vertices were mapped to an original vertex with more than epsilon differences.
+						// Either bump the epsilon or, if too many failed, emit an error message about a potentially bad mapping
+
+						// if more than half failed
+						if ((int32)NumMapped < NumOriginPos / 2)
+						{
+							FMessageDialog::Open(EAppMsgType::Ok, NSLOCTEXT("SkelmeshComponent", "Warning_ClothMorphTargetMapping", "Mapping vertices for Cloth Morph Target too many failed! It could introduce a bad morphing result."));
+						}
+					}
+					
+					SkeletalMesh->ClothingAssets[AssetIdx].PreparedMorphTargetName = MorphTargetName;
+
+					FClothMorphTargetData *Data = new(ClothMorphTargets)FClothMorphTargetData;
+					Data->MorphTargetName = MorphTargetName;
+					Data->ClothAssetIndex = AssetIdx;
+					Data->PrevWeight = 0.0f;
+
+					// stores original positions
+					TArray<PxVec3>& OriginPosArray = ClothOriginalPosArray[AssetIdx];
+					NumOriginPos = OriginPosArray.Num();
+					Data->OriginPos.AddUninitialized(NumOriginPos);
+					for (int32 Index = 0; Index < NumOriginPos; Index++)
+					{
+						Data->OriginPos[Index] = P2UVector(OriginPosArray[Index]);
+					}
+
+					Data->PosDelta = ClothPositionDeltaArray[AssetIdx];
+				}
+			}
+		}
+	}
+
+	bPreparedClothMorphTargets = true;
+}
+
+void USkeletalMeshComponent::UpdateClothMorphTarget()
+{
+	if (!bClothMorphTarget)
+	{
+		return;
+	}
+
+	for (int32 MorphIdx = 0; MorphIdx < ActiveVertexAnims.Num(); MorphIdx++)
+	{
+		if (ActiveVertexAnims[MorphIdx].Weight > 0.0f)
+		{
+			if (ActiveVertexAnims[MorphIdx].VertAnim)
+			{
+				FName MorphTargetName = ActiveVertexAnims[MorphIdx].VertAnim->GetFName();
+
+				int32 SelectedClothMorphIndex = INDEX_NONE;
+				for (int32 ClothMorphIdx = 0; ClothMorphIdx < ClothMorphTargets.Num(); ClothMorphIdx++)
+				{
+					if (ClothMorphTargets[ClothMorphIdx].MorphTargetName == MorphTargetName)
+					{
+						SelectedClothMorphIndex = ClothMorphIdx;
+						break;
+					}
+				}
+
+				// if this morph target is not cloth morph target, skip this index
+				if (SelectedClothMorphIndex == INDEX_NONE)
+				{
+					continue;
+				}
+
+				FClothMorphTargetData& MorphData = ClothMorphTargets[SelectedClothMorphIndex];
+
+				// if a currently mapped morph target name is same as MorphTargetName, doesn't change mapping. Otherwise, change morph target mapping
+				ChangeClothMorphTargetMapping(MorphData, MorphTargetName);
+
+				float CurWeight = ActiveVertexAnims[MorphIdx].Weight;
+
+				if (CurWeight != MorphData.PrevWeight)
+				{
+					MorphData.PrevWeight = CurWeight;
+					TArray<FVector> BlendedDelta;
+					TArray<FVector>& OriginDelta = MorphData.PosDelta;
+
+					if (OriginDelta.Num() > 0)
+					{
+						int32 ActorIndex = MorphData.ClothAssetIndex;
+
+						if (!IsValidClothingActor(ActorIndex))
+						{
+							continue;
+						}
+
+						BlendedDelta.AddUninitialized(OriginDelta.Num());
+
+						for (int32 DeltaIdx = 0; DeltaIdx < BlendedDelta.Num(); DeltaIdx++)
+						{
+							BlendedDelta[DeltaIdx] = (OriginDelta[DeltaIdx] * CurWeight);
+						}
+
+						// make current actor invalid
+						ClothingActors[ActorIndex].Clear(true);
+						CreateClothingActor(ActorIndex, SkeletalMesh->ClothingAssets[ActorIndex].ApexClothingAsset, &BlendedDelta);
+					}
+				}
+			}
+		}
+	}
+}
+
 void USkeletalMeshComponent::UpdateClothState(float DeltaTime)
 {
-	int32 NumActors = ClothingActors.Num();
+	// if turned on bClothMorphTarget option
+	if (bClothMorphTarget)
+	{
+		// @TODO : if not in editor, doesn't need to check preparation of cloth morph targets in Tick
+		// if pre-computation is not conducted yet
+		if (!bPreparedClothMorphTargets)
+		{
+			PrepareClothMorphTargets();
+		}
+		UpdateClothMorphTarget();
+	}
 
 #if WITH_CLOTH_COLLISION_DETECTION
 
@@ -3170,6 +3438,7 @@ void USkeletalMeshComponent::UpdateClothState(float DeltaTime)
 	}
 #endif // WITH_CLOTH_COLLISION_DETECTION
 
+	int32 NumActors = ClothingActors.Num();
 	if(NumActors == 0)
 	{
 		return;
