Use incremental components registration in PIE. Do not run constuction script for a duplicated levels in PIE.

diff --git a/Engine/Source/Editor/UnrealEd/Private/PlayLevel.cpp b/Engine/Source/Editor/UnrealEd/Private/PlayLevel.cpp
--- a/Engine/Source/Editor/UnrealEd/Private/PlayLevel.cpp
+++ b/Engine/Source/Editor/UnrealEd/Private/PlayLevel.cpp
@@ -2875,12 +2875,15 @@ UWorld* UEditorEngine::CreatePIEWorldByDuplication(FWorldContext &WorldContext,
 		// NULL GWorld before various PostLoad functions are called, this makes it easier to debug invalid GWorld accesses
 		GWorld = NULL;
 
+		EObjectFlags FlagMask = RF_AllFlags;
+		FlagMask&= ~RF_WasLoaded; // clear "was loaded" flag for duplicated world, so editor will not run construction script for it
+		
 		// Duplicate the editor world to create the PIE world
 		NewPIEWorld = CastChecked<UWorld>( StaticDuplicateObject(
 			EditorWorld,			// Source root
 			PlayWorldPackage,		// Destination root
 			*EditorWorld->GetName(),// Name for new object
-			RF_AllFlags,			// FlagMask
+			FlagMask,				// FlagMask
 			NULL,					// DestClass
 			SDO_DuplicateForPie		// bDuplicateForPIE
 			) );
diff --git a/Engine/Source/Runtime/Engine/Private/Level.cpp b/Engine/Source/Runtime/Engine/Private/Level.cpp
--- a/Engine/Source/Runtime/Engine/Private/Level.cpp
+++ b/Engine/Source/Runtime/Engine/Private/Level.cpp
@@ -743,7 +743,7 @@ void ULevel::IncrementalUpdateComponents(int32 NumComponentsToUpdate, bool bReru
 	if (NumComponentsToUpdate != 0)
 	{
 		// Only the game can use incremental update functionality.
-		checkf(!GIsEditor && OwningWorld->IsGameWorld(),TEXT("Cannot call IncrementalUpdateComponents with non 0 argument in the Editor/ commandlets."));
+		checkf(OwningWorld->IsGameWorld(), TEXT("Cannot call IncrementalUpdateComponents with non 0 argument in the Editor/ commandlets."));
 	}
 
 	// Do BSP on the first pass.
diff --git a/Engine/Source/Runtime/Engine/Private/World.cpp b/Engine/Source/Runtime/Engine/Private/World.cpp
--- a/Engine/Source/Runtime/Engine/Private/World.cpp
+++ b/Engine/Source/Runtime/Engine/Private/World.cpp
@@ -1682,8 +1682,8 @@ void UWorld::AddToWorld( ULevel* Level, const FTransform& LevelTransform )
 #endif
 			// We don't need to rerun construction scripts if we have cooked data or we are playing in editor unless the PIE world was loaded
 			// from disk rather than duplicated
-			const bool bRerunConstructionScript = !FPlatformProperties::RequiresCookedData() && (!IsPlayInEditor() || HasAnyFlags(RF_WasLoaded));
-			Level->IncrementalUpdateComponents( (GIsEditor || IsRunningCommandlet()) ? 0 : NumComponentsToUpdate, bRerunConstructionScript );
+			const bool bRerunConstructionScript = !FPlatformProperties::RequiresCookedData() && (!IsPlayInEditor() || Level->HasAnyFlags(RF_WasLoaded));
+			Level->IncrementalUpdateComponents( (!IsGameWorld() || IsRunningCommandlet()) ? 0 : NumComponentsToUpdate, bRerunConstructionScript );
 		}
 		while( (!bConsiderTimeLimit || !IsTimeLimitExceeded( TEXT("updating components"), StartTime, Level )) && !Level->bAreComponentsCurrentlyRegistered );
 
@@ -2139,8 +2139,11 @@ UWorld* UWorld::DuplicateWorldForPIE(const FString& PackageName, UWorld* OwningW
 	FStringAssetReference::SetPackageNamesBeingDuplicatedForPIE(PackageNamesBeingDuplicatedForPIE);
 
 	ULevel::StreamedLevelsOwningWorld.Add(PIELevelPackage->GetFName(), OwningWorld);
-	UWorld* PIELevelWorld = CastChecked<UWorld>(StaticDuplicateObject(EditorLevelWorld, PIELevelPackage, *EditorLevelWorld->GetName(), RF_AllFlags, NULL, SDO_DuplicateForPie));
+	EObjectFlags FlagMask = RF_AllFlags;
+	FlagMask&= ~RF_WasLoaded; // clear "was loaded" flag for duplicated world, so editor will not run construction script for it
 
+	UWorld* PIELevelWorld = CastChecked<UWorld>(StaticDuplicateObject(EditorLevelWorld, PIELevelPackage, *EditorLevelWorld->GetName(), FlagMask, NULL, SDO_DuplicateForPie));
+	
 	// Clean up string asset reference fixups
 	FStringAssetReference::ClearPackageNamesBeingDuplicatedForPIE();
 
