Fix for transient blueprints generating warnings
#TTP 343062 - FIXIF: AUTOMATION: TappyChicken: BP_MainGame2: Lots of warnings about timeline tracks with invalid curve.
#Branch UE4
#Proj Engine

#Change Added a guard to prevent log warnings being generated on transient blueprints during migration

diff --git a/Engine/Source/Runtime/Engine/Private/Blueprint.cpp b/Engine/Source/Runtime/Engine/Private/Blueprint.cpp
--- a/Engine/Source/Runtime/Engine/Private/Blueprint.cpp
+++ b/Engine/Source/Runtime/Engine/Private/Blueprint.cpp
@@ -1138,16 +1138,19 @@ bool UBlueprint::ChangeOwnerOfTemplates()
 
 		if (bMigratedOwner)
 		{
-			// alert the user that blueprints gave been migrated and require re-saving to enable them to locate and fix them without nagging them.
-			FMessageLog("BlueprintLog").Warning( FText::Format( NSLOCTEXT( "Blueprint", "MigrationWarning", "Blueprint {0} has been migrated and requires re-saving to avoid import errors" ), FText::FromString( *GetName() )));
-
-			if( GetDefault<UEditorLoadingSavingSettings>()->bDirtyMigratedBlueprints )
+			if( !HasAnyFlags( RF_Transient ))
 			{
-				UPackage* BPPackage = GetOutermost();
+				// alert the user that blueprints gave been migrated and require re-saving to enable them to locate and fix them without nagging them.
+				FMessageLog("BlueprintLog").Warning( FText::Format( NSLOCTEXT( "Blueprint", "MigrationWarning", "Blueprint {0} has been migrated and requires re-saving to avoid import errors" ), FText::FromString( *GetName() )));
 
-				if( BPPackage )
+				if( GetDefault<UEditorLoadingSavingSettings>()->bDirtyMigratedBlueprints )
 				{
-					BPPackage->SetDirtyFlag( true );
+					UPackage* BPPackage = GetOutermost();
+
+					if( BPPackage )
+					{
+						BPPackage->SetDirtyFlag( true );
+					}
 				}
 			}
 
