Header cleanup: AnimationRuntime.h, AnimationUtils.h, Audio.h, AudioDerivedData.h, BatchedElements.h.

diff --git a/Engine/Source/Runtime/Engine/Private/Animation/AnimCompositeBase.cpp b/Engine/Source/Runtime/Engine/Private/Animation/AnimCompositeBase.cpp
--- a/Engine/Source/Runtime/Engine/Private/Animation/AnimCompositeBase.cpp
+++ b/Engine/Source/Runtime/Engine/Private/Animation/AnimCompositeBase.cpp
@@ -7,6 +7,7 @@
 #include "EnginePrivate.h"
 #include "AnimationUtils.h"
 #include "Animation/AnimCompositeBase.h"
+#include "Animation/AnimSequence.h"
 
 ///////////////////////////////////////////////////////
 // FAnimSegment
diff --git a/Engine/Source/Runtime/Engine/Private/Animation/AnimCompress_PerTrackCompression.cpp b/Engine/Source/Runtime/Engine/Private/Animation/AnimCompress_PerTrackCompression.cpp
--- a/Engine/Source/Runtime/Engine/Private/Animation/AnimCompress_PerTrackCompression.cpp
+++ b/Engine/Source/Runtime/Engine/Private/Animation/AnimCompress_PerTrackCompression.cpp
@@ -11,6 +11,7 @@
 #include "AnimationCompression.h"
 #include "AnimEncoding.h"
 
+struct FAnimSetMeshLinkup;
 
 struct FPerTrackCachedInfo
 {
diff --git a/Engine/Source/Runtime/Engine/Private/Animation/AnimSequenceBase.cpp b/Engine/Source/Runtime/Engine/Private/Animation/AnimSequenceBase.cpp
--- a/Engine/Source/Runtime/Engine/Private/Animation/AnimSequenceBase.cpp
+++ b/Engine/Source/Runtime/Engine/Private/Animation/AnimSequenceBase.cpp
@@ -4,6 +4,8 @@
 #include "AnimationUtils.h"
 #include "AnimationRuntime.h"
 #include "Animation/AnimNotifies/AnimNotify.h"
+#include "Animation/AnimSequenceBase.h"
+#include "Animation/AnimInstance.h"
 
 #define NOTIFY_TRIGGER_OFFSET KINDA_SMALL_NUMBER;
 
diff --git a/Engine/Source/Runtime/Engine/Private/Animation/AnimSet.cpp b/Engine/Source/Runtime/Engine/Private/Animation/AnimSet.cpp
--- a/Engine/Source/Runtime/Engine/Private/Animation/AnimSet.cpp
+++ b/Engine/Source/Runtime/Engine/Private/Animation/AnimSet.cpp
@@ -3,6 +3,7 @@
 #include "EnginePrivate.h"
 #include "Animation/AnimSet.h"
 #include "AnimationUtils.h"
+#include "Animation/AnimSequence.h"
 
 /////////////////////////////////////////////////////
 // UAnimSet
diff --git a/Engine/Source/Runtime/Engine/Private/Animation/AnimationRuntime.cpp b/Engine/Source/Runtime/Engine/Private/Animation/AnimationRuntime.cpp
--- a/Engine/Source/Runtime/Engine/Private/Animation/AnimationRuntime.cpp
+++ b/Engine/Source/Runtime/Engine/Private/Animation/AnimationRuntime.cpp
@@ -12,6 +12,7 @@
 #include "Animation/AnimData/BoneMaskFilter.h"
 #include "Animation/BlendSpaceBase.h"
 #include "Animation/AnimCompositeBase.h"
+#include "Animation/AnimInstance.h"
 
 DEFINE_LOG_CATEGORY(LogAnimation);
 DEFINE_LOG_CATEGORY(LogRootMotion);
diff --git a/Engine/Source/Runtime/Engine/Private/Animation/BlendSpaceBase.cpp b/Engine/Source/Runtime/Engine/Private/Animation/BlendSpaceBase.cpp
--- a/Engine/Source/Runtime/Engine/Private/Animation/BlendSpaceBase.cpp
+++ b/Engine/Source/Runtime/Engine/Private/Animation/BlendSpaceBase.cpp
@@ -6,6 +6,7 @@
 
 #include "EnginePrivate.h"
 #include "AnimationRuntime.h"
+#include "Animation/AnimInstance.h"
 #include "Animation/BlendSpaceBase.h"
 
 // global variable that's used in editor when you change the property only
diff --git a/Engine/Source/Runtime/Engine/Private/Animation/Skeleton.cpp b/Engine/Source/Runtime/Engine/Private/Animation/Skeleton.cpp
--- a/Engine/Source/Runtime/Engine/Private/Animation/Skeleton.cpp
+++ b/Engine/Source/Runtime/Engine/Private/Animation/Skeleton.cpp
@@ -8,6 +8,7 @@
 #include "EngineUtils.h"
 #include "AnimationRuntime.h"
 #include "AssetRegistryModule.h"
+#include "Animation/AnimSequence.h"
 
 #define LOCTEXT_NAMESPACE "Skeleton"
 #define ROOT_BONE_PARENT	INDEX_NONE
diff --git a/Engine/Source/Runtime/Engine/Private/AudioDerivedData.cpp b/Engine/Source/Runtime/Engine/Private/AudioDerivedData.cpp
--- a/Engine/Source/Runtime/Engine/Private/AudioDerivedData.cpp
+++ b/Engine/Source/Runtime/Engine/Private/AudioDerivedData.cpp
@@ -3,6 +3,7 @@
 #include "AudioDerivedData.h"
 #include "TargetPlatform.h"
 #include "Sound/SoundWave.h"
+#include "DerivedDataCacheInterface.h"
 
 DEFINE_LOG_CATEGORY_STATIC(LogAudioDerivedData, Log, All);
 
diff --git a/Engine/Source/Runtime/Engine/Private/InterpolationDraw.cpp b/Engine/Source/Runtime/Engine/Private/InterpolationDraw.cpp
--- a/Engine/Source/Runtime/Engine/Private/InterpolationDraw.cpp
+++ b/Engine/Source/Runtime/Engine/Private/InterpolationDraw.cpp
@@ -32,6 +32,7 @@
 #include "InterpolationHitProxy.h"
 #include "AnimationUtils.h"
 #include "InterpolationHitProxy.h"
+#include "Animation/AnimSequence.h"
 
 
 static const int32	KeyHalfTriSize = 6;
diff --git a/Engine/Source/Runtime/Engine/Private/SoundWave.cpp b/Engine/Source/Runtime/Engine/Private/SoundWave.cpp
--- a/Engine/Source/Runtime/Engine/Private/SoundWave.cpp
+++ b/Engine/Source/Runtime/Engine/Private/SoundWave.cpp
@@ -9,7 +9,7 @@
 #include "TargetPlatform.h"
 #include "AudioDerivedData.h"
 #include "SubtitleManager.h"
-
+#include "DerivedDataCacheInterface.h"
 /*-----------------------------------------------------------------------------
 	FStreamedAudioChunk
 -----------------------------------------------------------------------------*/
diff --git a/Engine/Source/Runtime/Engine/Public/AnimationRuntime.h b/Engine/Source/Runtime/Engine/Public/AnimationRuntime.h
--- a/Engine/Source/Runtime/Engine/Public/AnimationRuntime.h
+++ b/Engine/Source/Runtime/Engine/Public/AnimationRuntime.h
@@ -6,13 +6,16 @@
 =============================================================================*/ 
 
 #pragma once
-#include "Animation/AnimInstance.h"
+#include "Animation/AnimSequenceBase.h"
 
 struct FInputBlendPose;
+struct FA2CSPose;
 struct FA2Pose;
 struct FAnimTrack;
+struct FPerBoneBlendWeight;
 class UBlendSpaceBase;
-
+class UAnimSequenceBase;
+typedef TArray<FTransform> FTransformArrayA2;
 /** In AnimationRunTime Library, we extract animation data based on Skeleton hierarchy, not ref pose hierarchy. 
 	Ref pose will need to be re-mapped later **/
 
@@ -205,7 +208,7 @@ public:
 	 *
 	 * return ETypeAdvanceAnim type
 	 */
-	static ETypeAdvanceAnim AdvanceTime(const bool & bAllowLooping, const float & MoveDelta, float & InOutTime, const float & EndTime);
+	static enum ETypeAdvanceAnim AdvanceTime(const bool & bAllowLooping, const float & MoveDelta, float & InOutTime, const float & EndTime);
 
 	static void TickBlendWeight(float DeltaTime, float DesiredWeight, float & Weight, float& BlendTime);
 	/** 
diff --git a/Engine/Source/Runtime/Engine/Public/AnimationUtils.h b/Engine/Source/Runtime/Engine/Public/AnimationUtils.h
--- a/Engine/Source/Runtime/Engine/Public/AnimationUtils.h
+++ b/Engine/Source/Runtime/Engine/Public/AnimationUtils.h
@@ -7,20 +7,16 @@
 #ifndef __ANIMATIONUTILS_H__
 #define __ANIMATIONUTILS_H__
 
+#include "Animation/AnimSequence.h"
+
 ENGINE_API DECLARE_LOG_CATEGORY_EXTERN(LogMatinee, Warning, All);
 ENGINE_API DECLARE_LOG_CATEGORY_EXTERN(LogAnimNotify, Warning, All);
 
-#include "Animation/AnimSequence.h"
 
 // Forward declarations.
 class UAnimCompress;
 class UAnimSet;
-class USkeletalMesh;
-struct FAnimSetMeshLinkup;
 class FTransform;
-struct FMeshBone;
-struct FRotationTrack;
-struct FTranslationTrack;
 
 /**
  * Encapsulates commonly useful data about bones.
@@ -230,7 +226,7 @@ public:
 	 * @param	InFormat	The compression format to convert into a string
 	 * @return				The format as a string
 	 */
-	ENGINE_API static FString GetAnimationCompressionFormatString(AnimationCompressionFormat InFormat);
+	ENGINE_API static FString GetAnimationCompressionFormatString(enum AnimationCompressionFormat InFormat);
 
 	/**
 	 * Computes the 'height' of each track, relative to a given animation linkup.
diff --git a/Engine/Source/Runtime/Engine/Public/AudioDerivedData.h b/Engine/Source/Runtime/Engine/Public/AudioDerivedData.h
--- a/Engine/Source/Runtime/Engine/Public/AudioDerivedData.h
+++ b/Engine/Source/Runtime/Engine/Public/AudioDerivedData.h
@@ -1,11 +1,9 @@
 // Copyright 1998-2014 Epic Games, Inc. All Rights Reserved.
 
 #pragma once
-
 #include "DerivedDataPluginInterface.h"
-#include "DerivedDataCacheInterface.h"
-#include "TargetPlatform.h"
 
+class IAudioFormat;
 
 class FDerivedAudioDataCompressor : public FDerivedDataPluginInterface
 {
diff --git a/Engine/Source/Runtime/UMG/Public/Blueprint/WidgetTree.h b/Engine/Source/Runtime/UMG/Public/Blueprint/WidgetTree.h
--- a/Engine/Source/Runtime/UMG/Public/Blueprint/WidgetTree.h
+++ b/Engine/Source/Runtime/UMG/Public/Blueprint/WidgetTree.h
@@ -1,7 +1,7 @@
 // Copyright 1998-2014 Epic Games, Inc. All Rights Reserved.
 
 #pragma once
-
+#include "Blueprint/UserWidget.h"
 #include "WidgetTree.generated.h"
 
 /** The widget tree manages the collection of widgets in a blueprint widget. */
