Fix Clang compilation on OS X - UMorphTarget needs to be fully declared.

diff --git a/Engine/Source/Runtime/Engine/Private/SkeletalMeshComponentPhysics.cpp b/Engine/Source/Runtime/Engine/Private/SkeletalMeshComponentPhysics.cpp
--- a/Engine/Source/Runtime/Engine/Private/SkeletalMeshComponentPhysics.cpp
+++ b/Engine/Source/Runtime/Engine/Private/SkeletalMeshComponentPhysics.cpp
@@ -25,6 +25,7 @@
 	#include "NxClothingCollision.h"
 	// for cloth morph target	
 	#include "Animation/VertexAnim/VertexAnimBase.h"
+	#include "Animation/VertexAnim/MorphTarget.h"
 
 #endif// #if WITH_APEX_CLOTHING
 
