Fixed stall when landing on Foliage Actors.

We now check the primary actor bCanEverTick to fix issues with landing on Foliage actors with thousands of components. We need to depend on other component's Ticks in the general case to deal with things like Timelines driving the position of moving platforms etc.
This will prevent dependencies from being created for Blueprints that use timelines but do not implement a Tick() event, but this is a fairly simple workaround and the side effect is minor (a one frame lag in movement between base and character).

TTP# 341962 - EDITOR: FOLIAGE: Editor freezes temporarily when jumping on foliage with collision in user level

diff --git a/Engine/Source/Runtime/Engine/Private/Character.cpp b/Engine/Source/Runtime/Engine/Private/Character.cpp
--- a/Engine/Source/Runtime/Engine/Private/Character.cpp
+++ b/Engine/Source/Runtime/Engine/Private/Character.cpp
@@ -387,7 +387,14 @@ namespace MovementBaseUtility
 		{
 			BasedObjectTick.AddPrerequisite(NewBase, NewBase->PrimaryComponentTick);
 			AActor* NewBaseOwner = NewBase->GetOwner();
-			if (NewBaseOwner)
+
+			// NOTE: TTP# 341962: We check the primary actor bCanEverTick here to fix issues with landing on Foliage actors with thousands of components.
+			// We need to depend on other component's Tick()s in the general case to deal with things like Timelines driving the position
+			// of moving platforms etc.
+			// This will prevent dependencies from being created for Blueprints that use timelines but do not implement a Tick() event,
+			// but this is a fairly simple workaround and the side effect is minor (a one frame lag in movement between base and character).
+			// @TODO: We need to find a more efficient way of finding all ticking components in an actor.
+			if (NewBaseOwner &&	NewBaseOwner->PrimaryActorTick.bCanEverTick)
 			{
 				BasedObjectTick.AddPrerequisite(NewBaseOwner, NewBaseOwner->PrimaryActorTick);
 
@@ -407,7 +414,7 @@ namespace MovementBaseUtility
 		{
 			BasedObjectTick.RemovePrerequisite(OldBase, OldBase->PrimaryComponentTick);
 			AActor* OldBaseOwner = OldBase->GetOwner();
-			if (OldBaseOwner)
+			if (OldBaseOwner && OldBaseOwner->PrimaryActorTick.bCanEverTick)
 			{
 				BasedObjectTick.RemovePrerequisite(OldBaseOwner, OldBaseOwner->PrimaryActorTick);
 
diff --git a/Engine/Source/Runtime/Engine/Private/InstancedFoliage.cpp b/Engine/Source/Runtime/Engine/Private/InstancedFoliage.cpp
--- a/Engine/Source/Runtime/Engine/Private/InstancedFoliage.cpp
+++ b/Engine/Source/Runtime/Engine/Private/InstancedFoliage.cpp
@@ -719,6 +719,7 @@ AInstancedFoliageActor::AInstancedFoliageActor(const FPostConstructInitializePro
 #if WITH_EDITORONLY_DATA
 	bListedInSceneOutliner = false;
 #endif // WITH_EDITORONLY_DATA
+	PrimaryActorTick.bCanEverTick = false;
 }
 
 #if WITH_EDITOR
