ShooterGame and ULocalPlayer now pass an FUniqueNetId to GetPlayerNickname and GetLoginStatus
Fixes some problems with controller swapping on Xbox

diff --git a/Engine/Source/Runtime/Engine/Private/LocalPlayer.cpp b/Engine/Source/Runtime/Engine/Private/LocalPlayer.cpp
--- a/Engine/Source/Runtime/Engine/Private/LocalPlayer.cpp
+++ b/Engine/Source/Runtime/Engine/Private/LocalPlayer.cpp
@@ -1549,7 +1549,11 @@ FString ULocalPlayer::GetNickname() const
 		IOnlineIdentityPtr OnlineIdentityInt = Online::GetIdentityInterface(World);
 		if (OnlineIdentityInt.IsValid())
 		{
-			return OnlineIdentityInt->GetPlayerNickname(ControllerId);
+			auto UniqueId = GetPreferredUniqueNetId();
+			if (UniqueId.IsValid())
+			{
+				return OnlineIdentityInt->GetPlayerNickname(*UniqueId);
+			}
 		}
 	}
 
