Added Convert Space for animation remapping feature.

When remap to new skeleton, you can check the option to "Convert spaces", and it will convert it correctly. It only works if they both have same shape of bind pose.

diff --git a/Engine/Source/Developer/AssetTools/Private/AssetTypeActions/AssetTypeActions_AnimBlueprint.cpp b/Engine/Source/Developer/AssetTools/Private/AssetTypeActions/AssetTypeActions_AnimBlueprint.cpp
--- a/Engine/Source/Developer/AssetTools/Private/AssetTypeActions/AssetTypeActions_AnimBlueprint.cpp
+++ b/Engine/Source/Developer/AssetTools/Private/AssetTypeActions/AssetTypeActions_AnimBlueprint.cpp
@@ -150,13 +150,20 @@ void FAssetTypeActions_AnimBlueprint::ExecuteFindSkeleton(TArray<TWeakObjectPtr<
 void FAssetTypeActions_AnimBlueprint::RetargetAssets(const TArray<UObject*> InAnimBlueprints, bool bDuplicateAssets)
 {
 	bool bRemapReferencedAssets = false;
+	bool bConvertSpaces = true;
 	USkeleton * NewSkeleton = NULL;
+	USkeleton * OldSkeleton = NULL;
+	if ( InAnimBlueprints.Num() > 0 )
+	{
+		UAnimBlueprint * AnimBP = CastChecked<UAnimBlueprint>(InAnimBlueprints[0]);
+		OldSkeleton = AnimBP->TargetSkeleton;
+	}
 
 	const FText Message = LOCTEXT("RemapSkeleton_Warning", "Select the skeleton to remap this asset to.");
 
-	if (SAnimationRemapSkeleton::ShowModal(NULL, NewSkeleton, Message, &bRemapReferencedAssets))
+	if (SAnimationRemapSkeleton::ShowModal(OldSkeleton, NewSkeleton, Message, bConvertSpaces, &bRemapReferencedAssets))
 	{
-		EditorAnimUtils::RetargetAnimations(NewSkeleton, InAnimBlueprints, bRemapReferencedAssets, bDuplicateAssets);
+		EditorAnimUtils::RetargetAnimations(OldSkeleton, NewSkeleton, InAnimBlueprints, bRemapReferencedAssets, bDuplicateAssets, bConvertSpaces);
 	}
 }
 
diff --git a/Engine/Source/Developer/AssetTools/Private/AssetTypeActions/AssetTypeActions_AnimationAsset.cpp b/Engine/Source/Developer/AssetTools/Private/AssetTypeActions/AssetTypeActions_AnimationAsset.cpp
--- a/Engine/Source/Developer/AssetTools/Private/AssetTypeActions/AssetTypeActions_AnimationAsset.cpp
+++ b/Engine/Source/Developer/AssetTools/Private/AssetTypeActions/AssetTypeActions_AnimationAsset.cpp
@@ -156,13 +156,19 @@ void FAssetTypeActions_AnimationAsset::ExecuteFindSkeleton(TArray<TWeakObjectPtr
 void FAssetTypeActions_AnimationAsset::RetargetAssets(const TArray<UObject*> InAnimAssets, bool bDuplicateAssets)
 {
 	bool bRemapReferencedAssets = false;
+	bool bConvertSpaces = true;
 	USkeleton * NewSkeleton = NULL;
-
+	USkeleton * OldSkeleton = NULL;
+	if(InAnimAssets.Num() > 0)
+	{
+		UAnimationAsset * AnimAsset = CastChecked<UAnimationAsset>(InAnimAssets[0]);
+		OldSkeleton = AnimAsset->GetSkeleton();
+	}
 	const FText Message = LOCTEXT("SelectSkeletonToRemap", "Select the skeleton to remap this asset to.");
 
-	if (SAnimationRemapSkeleton::ShowModal(NULL, NewSkeleton, Message, &bRemapReferencedAssets))
+	if (SAnimationRemapSkeleton::ShowModal(OldSkeleton, NewSkeleton, Message, bConvertSpaces, &bRemapReferencedAssets))
 	{
-		EditorAnimUtils::RetargetAnimations(NewSkeleton, InAnimAssets, bRemapReferencedAssets, bDuplicateAssets);
+		EditorAnimUtils::RetargetAnimations(OldSkeleton, NewSkeleton, InAnimAssets, bRemapReferencedAssets, bDuplicateAssets, bConvertSpaces);
 	}
 }
 
diff --git a/Engine/Source/Developer/AssetTools/Private/AssetTypeActions/AssetTypeActions_Skeleton.cpp b/Engine/Source/Developer/AssetTools/Private/AssetTypeActions/AssetTypeActions_Skeleton.cpp
--- a/Engine/Source/Developer/AssetTools/Private/AssetTypeActions/AssetTypeActions_Skeleton.cpp
+++ b/Engine/Source/Developer/AssetTools/Private/AssetTypeActions/AssetTypeActions_Skeleton.cpp
@@ -179,8 +179,9 @@ void FAssetTypeActions_Skeleton::ExecuteRetargetSkeleton(TArray<TWeakObjectPtr<U
 
 			const FText Message = LOCTEXT("RetargetSkeleton_Warning", "This only converts animation data -i.e. animation assets and Anim Blueprints. \nIf you'd like to convert SkeletalMesh, use the context menu (Assign Skeleton) for each mesh. \n\nIf you'd like to convert mesh as well, please do so before converting animation data. \nOtherwise you will lose any extra track that is in the new mesh.");
 
+			bool bConvertSpaces = true;
 			// ask user what they'd like to change to 
-			if (SAnimationRemapSkeleton::ShowModal(OldSkeleton, NewSkeleton, Message ))
+			if (SAnimationRemapSkeleton::ShowModal(OldSkeleton, NewSkeleton, Message, bConvertSpaces ))
 			{
 				// find all assets who references old skeleton
 				TArray<FName> Packages;
@@ -193,12 +194,12 @@ void FAssetTypeActions_Skeleton::ExecuteRetargetSkeleton(TArray<TWeakObjectPtr<U
 				{
 					// Open a dialog asking the user to wait while assets are being discovered
 					SDiscoveringAssetsDialog::OpenDiscoveringAssetsDialog(
-						SDiscoveringAssetsDialog::FOnAssetsDiscovered::CreateSP(this, &FAssetTypeActions_Skeleton::PerformRetarget, OldSkeleton, NewSkeleton, Packages)
+						SDiscoveringAssetsDialog::FOnAssetsDiscovered::CreateSP(this, &FAssetTypeActions_Skeleton::PerformRetarget, OldSkeleton, NewSkeleton, Packages, bConvertSpaces)
 						);
 				}
 				else
 				{
-					PerformRetarget(OldSkeleton, NewSkeleton, Packages);
+					PerformRetarget(OldSkeleton, NewSkeleton, Packages, bConvertSpaces);
 				}
 				
 			}
@@ -206,7 +207,7 @@ void FAssetTypeActions_Skeleton::ExecuteRetargetSkeleton(TArray<TWeakObjectPtr<U
 	}
 }
 
-void FAssetTypeActions_Skeleton::PerformRetarget(USkeleton *OldSkeleton, USkeleton *NewSkeleton, TArray<FName> Packages) const
+void FAssetTypeActions_Skeleton::PerformRetarget(USkeleton *OldSkeleton, USkeleton *NewSkeleton, TArray<FName> Packages, bool bConvertSpaces) const
 {
 	TArray<FAssetToRemapSkeleton> AssetsToRemap;
 
@@ -238,7 +239,7 @@ void FAssetTypeActions_Skeleton::PerformRetarget(USkeleton *OldSkeleton, USkelet
 		DetectReadOnlyPackages(AssetsToRemap, PackagesToSave);
 
 		// retarget skeleton
-		RetargetSkeleton(AssetsToRemap, OldSkeleton, NewSkeleton);
+		RetargetSkeleton(AssetsToRemap, OldSkeleton, NewSkeleton, bConvertSpaces);
 
 		// Save all packages that were referencing any of the assets that were moved without redirectors
 		SavePackages(PackagesToSave);
@@ -439,7 +440,7 @@ void FAssetTypeActions_Skeleton::ReportFailures(const TArray<FAssetToRemapSkelet
 	}
 }
 
-void FAssetTypeActions_Skeleton::RetargetSkeleton(TArray<FAssetToRemapSkeleton>& AssetsToRemap, USkeleton* OldSkeleton, USkeleton* NewSkeleton) const
+void FAssetTypeActions_Skeleton::RetargetSkeleton(TArray<FAssetToRemapSkeleton>& AssetsToRemap, USkeleton* OldSkeleton, USkeleton* NewSkeleton, bool bConvertSpaces) const
 {
 	TArray<UAnimBlueprint*>	AnimBlueprints;
 
@@ -454,7 +455,7 @@ void FAssetTypeActions_Skeleton::RetargetSkeleton(TArray<FAssetToRemapSkeleton>&
 			{
 				if ( UAnimationAsset * AnimAsset = Cast<UAnimationAsset>(Asset) )
 				{
-					AnimAsset->ReplaceSkeleton(NewSkeleton);
+					AnimAsset->ReplaceSkeleton(NewSkeleton, bConvertSpaces);
 				}
 				else if ( UAnimBlueprint * AnimBlueprint = Cast<UAnimBlueprint>(Asset) )
 				{
diff --git a/Engine/Source/Developer/AssetTools/Private/AssetTypeActions/AssetTypeActions_Skeleton.h b/Engine/Source/Developer/AssetTools/Private/AssetTypeActions/AssetTypeActions_Skeleton.h
--- a/Engine/Source/Developer/AssetTools/Private/AssetTypeActions/AssetTypeActions_Skeleton.h
+++ b/Engine/Source/Developer/AssetTools/Private/AssetTypeActions/AssetTypeActions_Skeleton.h
@@ -57,15 +57,16 @@ private: // Helper functions
 	 * @param OldSkeleton : Old Skeleton that it changes from
 	 * @param NewSkeleton : New Skeleton that it would change to
 	 * @param Referencers : List of asset names that reference Old Skeleton
+	 * @param bConvertSpaces : Whether to convert spaces or not
 	 */
- 	void PerformRetarget(USkeleton *OldSkeleton, USkeleton *NewSkeleton, TArray<FName> Packages) const;
+ 	void PerformRetarget(USkeleton *OldSkeleton, USkeleton *NewSkeleton, TArray<FName> Packages, bool bConvertSpaces) const;
 
 	// utility functions for performing retargeting
 	// these codes are from AssetRenameManager workflow
 	void LoadPackages(TArray<FAssetToRemapSkeleton>& AssetsToRemap, TArray<UPackage*>& OutPackagesToSave) const;
 	bool CheckOutPackages(TArray<FAssetToRemapSkeleton>& AssetsToRemap, TArray<UPackage*>& InOutPackagesToSave) const;
 	void ReportFailures(const TArray<FAssetToRemapSkeleton>& AssetsToRemap) const;
-	void RetargetSkeleton(TArray<FAssetToRemapSkeleton>& AssetsToRemap, USkeleton* OldSkeleton, USkeleton* NewSkeleton) const;
+	void RetargetSkeleton(TArray<FAssetToRemapSkeleton>& AssetsToRemap, USkeleton* OldSkeleton, USkeleton* NewSkeleton, bool bConvertSpaces) const;
 	void SavePackages(const TArray<UPackage*> PackagesToSave) const;
 	void DetectReadOnlyPackages(TArray<FAssetToRemapSkeleton>& AssetsToRemap, TArray<UPackage*>& InOutPackagesToSave) const;
 	void FillCreateMenu(FMenuBuilder& MenuBuilder, TArray<TWeakObjectPtr<USkeleton>> Skeletons) const;
diff --git a/Engine/Source/Editor/Persona/Private/Persona.cpp b/Engine/Source/Editor/Persona/Private/Persona.cpp
--- a/Engine/Source/Editor/Persona/Private/Persona.cpp
+++ b/Engine/Source/Editor/Persona/Private/Persona.cpp
@@ -473,6 +473,7 @@ void FPersona::ExtendMenu()
 			{
 				MenuBuilder.AddMenuEntry(FPersonaCommands::Get().ChangeSkeletonPreviewMesh);
 				MenuBuilder.AddMenuEntry(FPersonaCommands::Get().RemoveUnusedBones);
+				MenuBuilder.AddMenuEntry( FPersonaCommands::Get().UpdateSkeletonRefPose);
 			}
 			MenuBuilder.EndSection();
 
@@ -997,6 +998,9 @@ void FPersona::CreateDefaultCommands()
 		FExecuteAction::CreateSP( this, &FPersona::RemoveUnusedBones ),
 		FCanExecuteAction::CreateSP( this, &FPersona::CanRemoveBones )
 		);
+	ToolkitCommands->MapAction(FPersonaCommands::Get().UpdateSkeletonRefPose,
+		FExecuteAction::CreateSP(this, &FPersona::UpdateSkeletonRefPose)
+		);
 
 	ToolkitCommands->MapAction(FPersonaCommands::Get().AnimNotifyWindow,
 		FExecuteAction::CreateSP(this, &FPersona::OnAnimNotifyWindow),
@@ -2567,6 +2571,15 @@ void FPersona::ChangeSkeletonPreviewMesh()
 	}
 }
 
+void FPersona::UpdateSkeletonRefPose()
+{
+	// update ref pose with current preview mesh
+	if (TargetSkeleton)
+	{
+		TargetSkeleton->UpdateReferencePoseFromMesh(GetPreviewMeshComponent()->SkeletalMesh);
+	}
+}
+
 void FPersona::RemoveUnusedBones()
 {
 	// Menu option cannot be called unless this is valid
diff --git a/Engine/Source/Editor/Persona/Private/Persona.h b/Engine/Source/Editor/Persona/Private/Persona.h
--- a/Engine/Source/Editor/Persona/Private/Persona.h
+++ b/Engine/Source/Editor/Persona/Private/Persona.h
@@ -574,6 +574,8 @@ private:
 
 	/** Extend menu and toolbar */
 	void ExtendMenu();
+	/** update skeleton ref pose based on current preview mesh */
+	void UpdateSkeletonRefPose();
 
 	/** Returns the editor objects that are applicable for our current mode (e.g mesh, animation etc) */
 	TArray<UObject*> GetEditorObjectsForMode(FName Mode) const;
diff --git a/Engine/Source/Editor/Persona/Private/PersonaCommands.cpp b/Engine/Source/Editor/Persona/Private/PersonaCommands.cpp
--- a/Engine/Source/Editor/Persona/Private/PersonaCommands.cpp
+++ b/Engine/Source/Editor/Persona/Private/PersonaCommands.cpp
@@ -29,4 +29,6 @@ void FPersonaCommands::RegisterCommands()
 	UI_COMMAND( RecordAnimation, "Record to new Animation", "Create new animation from currently playing", EUserInterfaceActionType::Button, FInputGesture() );
 	UI_COMMAND( ExportToFBX, "Export to FBX", "Export current animation to FBX", EUserInterfaceActionType::Button, FInputGesture() );
 	UI_COMMAND( AddLoopingInterpolation, "Add Looping Interpolation", "Add an extra first frame at the end of the animation to create interpolation when looping", EUserInterfaceActionType::Button, FInputGesture() );
+
+	UI_COMMAND( UpdateSkeletonRefPose, "Update Skeleton RefPose", "Update Skeleton ref pose based on current preview mesh", EUserInterfaceActionType::Button, FInputGesture() );
 }
diff --git a/Engine/Source/Editor/Persona/Private/PersonaCommands.h b/Engine/Source/Editor/Persona/Private/PersonaCommands.h
--- a/Engine/Source/Editor/Persona/Private/PersonaCommands.h
+++ b/Engine/Source/Editor/Persona/Private/PersonaCommands.h
@@ -59,4 +59,6 @@ public:
 	TSharedPtr<FUICommandInfo> CreateBlendSpace2D;
 	TSharedPtr<FUICommandInfo> CreateAnimComposite;
 	TSharedPtr<FUICommandInfo> CreateAanimMontage;
+	// Command to allow users to remove unused bones (not referenced by any skeletalmesh) from the skeleton
+	TSharedPtr<FUICommandInfo> UpdateSkeletonRefPose;
 };
diff --git a/Engine/Source/Editor/Persona/Private/SAnimationSequenceBrowser.cpp b/Engine/Source/Editor/Persona/Private/SAnimationSequenceBrowser.cpp
--- a/Engine/Source/Editor/Persona/Private/SAnimationSequenceBrowser.cpp
+++ b/Engine/Source/Editor/Persona/Private/SAnimationSequenceBrowser.cpp
@@ -291,14 +291,15 @@ void SAnimationSequenceBrowser::OnCreateCopy(TArray<FAssetData> Selected)
 		USkeleton * OldSkeleton = PersonaPtr.Pin()->GetSkeleton();
 		USkeleton * NewSkeleton = NULL;
 		bool		bRemapReferencedAssets = true;
+		bool		bConvertSpaces = true;
 		bool		bDuplicateAssets = true;
 
 		const FText Message = LOCTEXT("RemapSkeleton_Warning", "This will duplicate the asset and convert to new skeleton.");
 
 		// ask user what they'd like to change to 
-		if (SAnimationRemapSkeleton::ShowModal(OldSkeleton, NewSkeleton, Message, &bRemapReferencedAssets))
+		if (SAnimationRemapSkeleton::ShowModal(OldSkeleton, NewSkeleton, Message, bConvertSpaces, &bRemapReferencedAssets))
 		{
-			UObject* AssetToOpen = EditorAnimUtils::RetargetAnimations(NewSkeleton, Selected, bRemapReferencedAssets, bDuplicateAssets);
+			UObject* AssetToOpen = EditorAnimUtils::RetargetAnimations(OldSkeleton, NewSkeleton, Selected, bRemapReferencedAssets, bDuplicateAssets, bConvertSpaces);
 
 			if(UAnimationAsset* AnimAsset = Cast<UAnimationAsset>(AssetToOpen))
 			{
diff --git a/Engine/Source/Editor/UnrealEd/Private/EditorAnimUtils.cpp b/Engine/Source/Editor/UnrealEd/Private/EditorAnimUtils.cpp
--- a/Engine/Source/Editor/UnrealEd/Private/EditorAnimUtils.cpp
+++ b/Engine/Source/Editor/UnrealEd/Private/EditorAnimUtils.cpp
@@ -19,7 +19,9 @@ namespace EditorAnimUtils
 	//////////////////////////////////////////////////////////////////
 	// FAnimationRetargetContext
 
-	FAnimationRetargetContext::FAnimationRetargetContext(const TArray<FAssetData>& AssetsToRetarget, bool bRetargetReferredAssets) : SingleTargetObject(NULL)
+	FAnimationRetargetContext::FAnimationRetargetContext(const TArray<FAssetData>& AssetsToRetarget, bool bRetargetReferredAssets, bool bInConvertAnimationDataInComponentSpaces) 
+		: SingleTargetObject(NULL)
+		, bConvertAnimationDataInComponentSpaces(bInConvertAnimationDataInComponentSpaces)
 	{
 		TArray<UObject*> Objects;
 		for(auto Iter = AssetsToRetarget.CreateConstIterator(); Iter; ++Iter)
@@ -29,7 +31,9 @@ namespace EditorAnimUtils
 		Initialize(Objects,bRetargetReferredAssets);
 	}
 
-	FAnimationRetargetContext::FAnimationRetargetContext(const TArray<UObject*>& AssetsToRetarget, bool bRetargetReferredAssets) : SingleTargetObject(NULL)
+	FAnimationRetargetContext::FAnimationRetargetContext(const TArray<UObject*>& AssetsToRetarget, bool bRetargetReferredAssets, bool bInConvertAnimationDataInComponentSpaces) 
+		: SingleTargetObject(NULL)
+		, bConvertAnimationDataInComponentSpaces(bInConvertAnimationDataInComponentSpaces)
 	{
 		Initialize(AssetsToRetarget,bRetargetReferredAssets);
 	}
@@ -135,12 +139,23 @@ namespace EditorAnimUtils
 		}
 	}
 
-	void FAnimationRetargetContext::RetargetAnimations(USkeleton* NewSkeleton)
+	void FAnimationRetargetContext::RetargetAnimations(USkeleton * OldSkeleton, USkeleton* NewSkeleton)
 	{
+		check (OldSkeleton);
+		check (NewSkeleton);
+
+		if (bConvertAnimationDataInComponentSpaces)
+		{
+			// we need to update reference pose before retargeting. 
+			// this is to ensure the skeleton has the latest pose you're looking at. 
+			OldSkeleton->UpdateReferencePoseFromMesh(OldSkeleton->GetPreviewMesh(true));
+			NewSkeleton->UpdateReferencePoseFromMesh(NewSkeleton->GetPreviewMesh(true));
+		}
+
 		for(auto Iter = AnimSequencesToRetarget.CreateIterator(); Iter; ++Iter)
 		{
 			UAnimSequence* AssetToRetarget = (*Iter);
-			AssetToRetarget->ReplaceSkeleton(NewSkeleton);
+			AssetToRetarget->ReplaceSkeleton(NewSkeleton, bConvertAnimationDataInComponentSpaces);
 		}
 
 		for(auto Iter = ComplexAnimsToRetarget.CreateIterator(); Iter; ++Iter)
@@ -150,7 +165,7 @@ namespace EditorAnimUtils
 			{
 				AssetToRetarget->ReplaceReferredAnimations(DuplicatedSequences);
 			}
-			AssetToRetarget->ReplaceSkeleton(NewSkeleton);
+			AssetToRetarget->ReplaceSkeleton(NewSkeleton, bConvertAnimationDataInComponentSpaces);
 		}
 
 		// convert all Animation Blueprints and compile 
@@ -189,19 +204,19 @@ namespace EditorAnimUtils
 	}
 
 	//////////////////////////////////////////////////////////////////
-	UObject* RetargetAnimations(USkeleton* NewSkeleton, const TArray<UObject*>& AssetsToRetarget, bool bRetargetReferredAssets, bool bDuplicateAssetsBeforeRetarget)
+	UObject* RetargetAnimations(USkeleton * OldSkeleton, USkeleton* NewSkeleton, const TArray<UObject*>& AssetsToRetarget, bool bRetargetReferredAssets, bool bDuplicateAssetsBeforeRetarget, bool bConvertSpace)
 	{
-		FAnimationRetargetContext RetargetContext(AssetsToRetarget, bRetargetReferredAssets);
-		return RetargetAnimations(NewSkeleton, RetargetContext, bRetargetReferredAssets, bDuplicateAssetsBeforeRetarget);
+		FAnimationRetargetContext RetargetContext(AssetsToRetarget, bRetargetReferredAssets, bConvertSpace);
+		return RetargetAnimations(OldSkeleton, NewSkeleton, RetargetContext, bRetargetReferredAssets, bDuplicateAssetsBeforeRetarget);
 	}
 
-	UObject* RetargetAnimations(USkeleton* NewSkeleton, const TArray<FAssetData>& AssetsToRetarget, bool bRetargetReferredAssets, bool bDuplicateAssetsBeforeRetarget)
+	UObject* RetargetAnimations(USkeleton * OldSkeleton, USkeleton* NewSkeleton, const TArray<FAssetData>& AssetsToRetarget, bool bRetargetReferredAssets, bool bDuplicateAssetsBeforeRetarget, bool bConvertSpace)
 	{
-		FAnimationRetargetContext RetargetContext(AssetsToRetarget, bRetargetReferredAssets);
-		return RetargetAnimations(NewSkeleton, RetargetContext, bRetargetReferredAssets, bDuplicateAssetsBeforeRetarget);
+		FAnimationRetargetContext RetargetContext(AssetsToRetarget, bRetargetReferredAssets, bConvertSpace);
+		return RetargetAnimations(OldSkeleton, NewSkeleton, RetargetContext, bRetargetReferredAssets, bDuplicateAssetsBeforeRetarget);
 	}
 
-	UObject* RetargetAnimations(USkeleton* NewSkeleton, FAnimationRetargetContext& RetargetContext, bool bRetargetReferredAssets, bool bDuplicateAssetsBeforeRetarget)
+	UObject* RetargetAnimations(USkeleton * OldSkeleton, USkeleton* NewSkeleton, FAnimationRetargetContext& RetargetContext, bool bRetargetReferredAssets, bool bDuplicateAssetsBeforeRetarget)
 	{
 		check(NewSkeleton);
 		UObject* OriginalObject  = RetargetContext.GetSingleTargetObject();
@@ -213,7 +228,7 @@ namespace EditorAnimUtils
 			{
 				RetargetContext.DuplicateAssetsToRetarget(DuplicationDestPackage);
 			}
-			RetargetContext.RetargetAnimations(NewSkeleton);
+			RetargetContext.RetargetAnimations(OldSkeleton, NewSkeleton);
 		}
 
 		FNotificationInfo Notification(FText::GetEmpty());
diff --git a/Engine/Source/Editor/UnrealEd/Private/SSkeletonWidget.cpp b/Engine/Source/Editor/UnrealEd/Private/SSkeletonWidget.cpp
--- a/Engine/Source/Editor/UnrealEd/Private/SSkeletonWidget.cpp
+++ b/Engine/Source/Editor/UnrealEd/Private/SSkeletonWidget.cpp
@@ -341,17 +341,33 @@ void SAnimationRemapSkeleton::Construct( const FArguments& InArgs )
 	AssetPickerConfig.InitialAssetViewType = EAssetViewType::Column;
 	AssetPickerConfig.ThumbnailScale = 0.0f;
 
-	TSharedRef<SWidget> Widget = SNew(SBox);
+	TSharedRef<SVerticalBox> Widget = SNew(SVerticalBox);
 	if(InArgs._ShowRemapOption)
 	{
-		Widget = SNew(SCheckBox)
-				.IsChecked(this, &SAnimationRemapSkeleton::IsRemappingReferencedAssets)
-				.OnCheckStateChanged(this, &SAnimationRemapSkeleton::OnRemappingReferencedAssetsChanged)
-				[
-					SNew(STextBlock).Text(LOCTEXT("RemapSkeleton_RemapAssets", "Remap referenced assets "))
-				];
+		Widget->AddSlot()
+		[
+			SNew(SCheckBox)
+			.IsChecked(this, &SAnimationRemapSkeleton::IsRemappingReferencedAssets)
+			.OnCheckStateChanged(this, &SAnimationRemapSkeleton::OnRemappingReferencedAssetsChanged)
+			[
+				SNew(STextBlock).Text(LOCTEXT("RemapSkeleton_RemapAssets", "Remap referenced assets "))
+			]
+		];
 	}
 
+	TSharedPtr<SToolTip> ConvertSpaceTooltip = IDocumentation::Get()->CreateToolTip(FText::FromString("Check if you'd like to convert animation data to new skeleton space. If this is false, it won't convert any animation data to new space."), 
+											NULL, FString("Shared/Editors/Persona"), FString("AnimRemapSkeleton_ConvertSpace"));
+	Widget->AddSlot()
+	[
+		SNew(SCheckBox)
+		.IsChecked(this, &SAnimationRemapSkeleton::IsConvertSpacesChecked)
+		.OnCheckStateChanged(this, &SAnimationRemapSkeleton::OnConvertSpacesCheckChanged)
+		[
+			SNew(STextBlock).Text(LOCTEXT("RemapSkeleton_ConvertSpaces", "Convert Spaces to new Skeleton"))
+			.ToolTip(ConvertSpaceTooltip)
+		]
+	];
+
 	TSharedPtr<SToolTip> SkeletonTooltip = IDocumentation::Get()->CreateToolTip(FText::FromString("Pick a skeleton for this mesh"), NULL, FString("Shared/Editors/Persona"), FString("Skeleton"));
 
 	this->ChildSlot
@@ -483,6 +499,16 @@ void SAnimationRemapSkeleton::OnRemappingReferencedAssetsChanged(ESlateCheckBoxS
 	bRemapReferencedAssets = (InNewRadioState == ESlateCheckBoxState::Checked);
 }
 
+ESlateCheckBoxState::Type SAnimationRemapSkeleton::IsConvertSpacesChecked() const
+{
+	return bConvertSpaces ? ESlateCheckBoxState::Checked : ESlateCheckBoxState::Unchecked;
+}
+
+void SAnimationRemapSkeleton::OnConvertSpacesCheckChanged(ESlateCheckBoxState::Type InNewRadioState)
+{
+	bConvertSpaces = (InNewRadioState == ESlateCheckBoxState::Checked);
+}
+
 bool SAnimationRemapSkeleton::CanApply() const
 {
 	return (NewSkeleton!=NULL && NewSkeleton!=OldSkeleton);
@@ -516,7 +542,7 @@ void SAnimationRemapSkeleton::CloseWindow()
 	}
 }
 
-bool SAnimationRemapSkeleton::ShowModal(USkeleton * OldSkeleton, USkeleton * & NewSkeleton, const FText& WarningMessage, bool * bRemapReferencedAssets)
+bool SAnimationRemapSkeleton::ShowModal(USkeleton * OldSkeleton, USkeleton * & NewSkeleton, const FText& WarningMessage, bool & bConvertSpace, bool * bRemapReferencedAssets)
 {
 	TSharedPtr<class SAnimationRemapSkeleton> DialogWidget;
 
@@ -548,6 +574,8 @@ bool SAnimationRemapSkeleton::ShowModal(USkeleton * OldSkeleton, USkeleton * & N
 		*bRemapReferencedAssets = DialogWidget.Get()->bRemapReferencedAssets;
 	}
 
+	bConvertSpace = DialogWidget.Get()->bConvertSpaces;
+
 	return (NewSkeleton != NULL && NewSkeleton != OldSkeleton);
 }
 
diff --git a/Engine/Source/Editor/UnrealEd/Public/EditorAnimUtils.h b/Engine/Source/Editor/UnrealEd/Public/EditorAnimUtils.h
--- a/Engine/Source/Editor/UnrealEd/Public/EditorAnimUtils.h
+++ b/Engine/Source/Editor/UnrealEd/Public/EditorAnimUtils.h
@@ -7,8 +7,8 @@ namespace EditorAnimUtils
 	class FAnimationRetargetContext
 	{
 	public:
-		FAnimationRetargetContext(const TArray<FAssetData>& AssetsToRetarget, bool bRetargetReferredAssets);
-		FAnimationRetargetContext(const TArray<UObject*>& AssetsToRetarget, bool bRetargetReferredAssets);
+		FAnimationRetargetContext(const TArray<FAssetData>& AssetsToRetarget, bool bRetargetReferredAssets, bool bInConvertAnimationDataInComponentSpaces);
+		FAnimationRetargetContext(const TArray<UObject*>& AssetsToRetarget, bool bRetargetReferredAssets, bool bInConvertAnimationDataInComponentSpaces);
 
 		/** Were we supplied anything that we can retarget */
 		bool HasAssetsToRetarget() const;
@@ -26,7 +26,7 @@ namespace EditorAnimUtils
 		void DuplicateAssetsToRetarget(UPackage* DestinationPackage);
 
 		/** Retarget the contained assets */
-		void RetargetAnimations(USkeleton* NewSkeleton);
+		void RetargetAnimations(USkeleton * OldSkeleton, USkeleton* NewSkeleton);
 	
 	private:
 		/** Lists of assets to retarget. Populated from FAssetData supplied to constructor */
@@ -42,6 +42,9 @@ namespace EditorAnimUtils
 		/** If we only chose one object to retarget store it here */
 		UObject* SingleTargetObject;
 
+		/** whether to convert animation data in component spaces */
+		bool bConvertAnimationDataInComponentSpaces;
+
 		/** Initialize the object, only to be called by constructors */
 		void Initialize(const TArray<UObject*>& AssetsToRetarget, bool bRetargetReferredAssets);
 	};
@@ -53,8 +56,9 @@ namespace EditorAnimUtils
 	 * @param AssetsToRetarget					The animation assets to copy/retarget
 	 * @param bRetargetReferredAssets			If true retargets any assets referred to by assets in AssetsToRetarget. If false then the references are cleared.
 	 * @param bDuplicatedAssetsBeforeRetarget	If true the assets are duplicated and then the duplicates are retargetted
+	 * @param bConvertSpace						Do the conversion in component space of the animation to match new target
 	 */
-	UNREALED_API UObject* RetargetAnimations(USkeleton* NewSkeleton, const TArray<UObject*>& AssetsToRetarget, bool bRetargetReferredAssets, bool bDuplicateAssetsBeforeRetarget);
+	UNREALED_API UObject* RetargetAnimations(USkeleton * OldSkeleton, USkeleton* NewSkeleton, const TArray<UObject*>& AssetsToRetarget, bool bRetargetReferredAssets, bool bDuplicateAssetsBeforeRetarget, bool bConvertSpace);
 
 	/**
 	 * Retargets the supplied FAssetDatas (as long as they are an animation asset), optionally duplicating them and retargetting their reference assets too
@@ -63,8 +67,9 @@ namespace EditorAnimUtils
 	 * @param AssetsToRetarget					The animation assets to copy/retarget
 	 * @param bRetargetReferredAssets			If true retargets any assets referred to by assets in AssetsToRetarget. If false then the references are cleared.
 	 * @param bDuplicatedAssetsBeforeRetarget	If true the assets are duplicated and then the duplicates are retargetted
+	 * @param bConvertSpace						Do the conversion in component space of the animation to match new target
 	 */
-	UNREALED_API UObject* RetargetAnimations(USkeleton* NewSkeleton, const TArray<FAssetData>& AssetsToRetarget, bool bRetargetReferredAssets, bool bDuplicateAssetsBeforeRetarget);
+	UNREALED_API UObject* RetargetAnimations(USkeleton * OldSkeleton, USkeleton* NewSkeleton, const TArray<FAssetData>& AssetsToRetarget, bool bRetargetReferredAssets, bool bDuplicateAssetsBeforeRetarget, bool bConvertSpace);
 
 	/**
 	 * Retargets the supplied FAnimationRetargetContext, optionally duplicating the assets and retargetting the assets reference assets too. Is called by other overloads of RetargetAnimations
@@ -73,8 +78,9 @@ namespace EditorAnimUtils
 	 * @param AssetsToRetarget					The animation assets to copy/retarget
 	 * @param bRetargetReferredAssets			If true retargets any assets referred to by assets in AssetsToRetarget. If false then the references are cleared.
 	 * @param bDuplicatedAssetsBeforeRetarget	If true the assets are duplicated and then the duplicates are retargetted
+	 * @param bConvertSpace						Do the conversion in component space of the animation to match new target
 	 */
-	UObject* RetargetAnimations(USkeleton* NewSkeleton, FAnimationRetargetContext& RetargetContext, bool bRetargetReferredAssets, bool bDuplicateAssetsBeforeRetarget);
+	UObject* RetargetAnimations(USkeleton * OldSkeleton, USkeleton* NewSkeleton, FAnimationRetargetContext& RetargetContext, bool bRetargetReferredAssets, bool bDuplicateAssetsBeforeRetarget);
 
 	// Populates the supplied TArrays with any animation assets that this blueprint refers too
 	void GetAllAnimationSequencesReferredInBlueprint(UAnimBlueprint* AnimBlueprint, TArray<UAnimationAsset*>& ComplexAnims, TArray<UAnimSequence*>& AnimSequences);
diff --git a/Engine/Source/Editor/UnrealEd/Public/SSkeletonWidget.h b/Engine/Source/Editor/UnrealEd/Public/SSkeletonWidget.h
--- a/Engine/Source/Editor/UnrealEd/Public/SSkeletonWidget.h
+++ b/Engine/Source/Editor/UnrealEd/Public/SSkeletonWidget.h
@@ -295,12 +295,21 @@ public:
 	 */
 	bool bRemapReferencedAssets;
 
+	/**
+	 * Whether we are remapping assets that are referenced by the assets the user selects to remap
+	 */
+	bool bConvertSpaces;
+
 	TWeakPtr<SWindow> WidgetWindow;
 
 	/** Handlers for check box for remapping assets option */
 	ESlateCheckBoxState::Type IsRemappingReferencedAssets() const;
 	void OnRemappingReferencedAssetsChanged(ESlateCheckBoxState::Type InNewRadioState);
 
+	/** Handlers for check box for converting spaces*/
+	ESlateCheckBoxState::Type IsConvertSpacesChecked() const;
+	void OnConvertSpacesCheckChanged(ESlateCheckBoxState::Type InNewRadioState);
+
 	/**
 	 * return true if it can apply 
 	 */
@@ -324,7 +333,7 @@ public:
 	 *
 	 * @return true if successfully selected new skeleton
 	 */
-	static UNREALED_API bool ShowModal(USkeleton * OldSkeleton, USkeleton * & NewSkeleton, const FText& WarningMessage, bool * bRemapReferencedAssets=NULL);
+	static UNREALED_API bool ShowModal(USkeleton * OldSkeleton, USkeleton * & NewSkeleton, const FText& WarningMessage, bool & bConvertSpace, bool * bRemapReferencedAssets=NULL);
 };
 
 ////////////////////////////////////////////////////
diff --git a/Engine/Source/Runtime/Engine/Classes/Animation/AnimSequence.h b/Engine/Source/Runtime/Engine/Classes/Animation/AnimSequence.h
--- a/Engine/Source/Runtime/Engine/Classes/Animation/AnimSequence.h
+++ b/Engine/Source/Runtime/Engine/Classes/Animation/AnimSequence.h
@@ -583,7 +583,17 @@ public:
 	 * @param	bUseRawData		If true, use raw animation data instead of compressed data.
 	 */
 	ENGINE_API void GetBoneTransform(FTransform& OutAtom, int32 TrackIndex, float Time, bool bUseRawData) const;
-	
+
+	/**
+	 * Extract Bone Transform of the Time given, from InRawAnimationData
+	 *
+	 * @param	InRawAnimationData	RawAnimationData it extracts bone transform from
+	 * @param	OutAtom				[out] Output bone transform.
+	 * @param	TrackIndex			Index of track to interpolate.
+	 * @param	Time				Time on track to interpolate to.
+	 */
+	ENGINE_API void ExtractBoneTransform(const TArray<struct FRawAnimSequenceTrack> & InRawAnimationData, FTransform& OutAtom, int32 TrackIndex, float Time) const;
+
 	// End Transform related functions 
 
 	// Begin Memory related functions
@@ -704,7 +714,7 @@ private:
 	/**
 	 * Remap Tracks to New Skeleton
 	 */
-	void RemapTracksToNewSkeleton( USkeleton * NewSkeleton );
+	void RemapTracksToNewSkeleton( USkeleton * NewSkeleton, bool bConvertSpaces );
 	/**
 	 * Remap NaN tracks from the RawAnimation data and recompress
 	 */	
diff --git a/Engine/Source/Runtime/Engine/Classes/Animation/AnimationAsset.h b/Engine/Source/Runtime/Engine/Classes/Animation/AnimationAsset.h
--- a/Engine/Source/Runtime/Engine/Classes/Animation/AnimationAsset.h
+++ b/Engine/Source/Runtime/Engine/Classes/Animation/AnimationAsset.h
@@ -400,7 +400,7 @@ public:
 	 * 
 	 * @param NewSkeleton	NewSkeleton to change to 
 	 */
-	ENGINE_API bool ReplaceSkeleton(USkeleton* NewSkeleton);
+	ENGINE_API bool ReplaceSkeleton(USkeleton* NewSkeleton, bool bConvertSpaces=false);
 
 	/** Retrieve all animations that are used by this asset 
 	 * 
diff --git a/Engine/Source/Runtime/Engine/Private/Animation/AnimSequence.cpp b/Engine/Source/Runtime/Engine/Private/Animation/AnimSequence.cpp
--- a/Engine/Source/Runtime/Engine/Private/Animation/AnimSequence.cpp
+++ b/Engine/Source/Runtime/Engine/Private/Animation/AnimSequence.cpp
@@ -522,21 +522,25 @@ void UAnimSequence::GetBoneTransform(FTransform& OutAtom, int32 TrackIndex, floa
 		}
 	}
 
+	ExtractBoneTransform(RawAnimationData, OutAtom, TrackIndex, Time);
+}
 
+void UAnimSequence::ExtractBoneTransform(const TArray<struct FRawAnimSequenceTrack> & InRawAnimationData, FTransform& OutAtom, int32 TrackIndex, float Time) const
+{
 	// Bail out if the animation data doesn't exists (e.g. was stripped by the cooker).
-	if ( RawAnimationData.Num() == 0 )
+	if(InRawAnimationData.Num() == 0)
 	{
-		UE_LOG(LogAnimation, Log, TEXT("UAnimSequence::GetBoneTransform : No anim data in AnimSequence!") );
+		UE_LOG(LogAnimation, Log, TEXT("UAnimSequence::GetBoneTransform : No anim data in AnimSequence!"));
 		OutAtom.SetIdentity();
 		return;
 	}
 
-	const FRawAnimSequenceTrack & RawTrack = RawAnimationData[TrackIndex];
+	const FRawAnimSequenceTrack & RawTrack = InRawAnimationData[TrackIndex];
 
 	// Bail out (with rather wacky data) if data is empty for some reason.
-	if( RawTrack.PosKeys.Num() == 0 || RawTrack.RotKeys.Num() == 0 )
+	if(RawTrack.PosKeys.Num() == 0 || RawTrack.RotKeys.Num() == 0)
 	{
-		UE_LOG(LogAnimation, Log, TEXT("UAnimSequence::GetBoneTransform : No anim data in AnimSequence!") );
+		UE_LOG(LogAnimation, Log, TEXT("UAnimSequence::GetBoneTransform : No anim data in AnimSequence!"));
 		OutAtom.SetIdentity();
 		return;
 	}
@@ -548,33 +552,33 @@ void UAnimSequence::GetBoneTransform(FTransform& OutAtom, int32 TrackIndex, floa
 	const bool bHasScaleKey = (RawTrack.ScaleKeys.Num() > 0);
 	static const FVector DefaultScale3D = FVector(1.f);
 
- 	if( Alpha <= 0.f )
+	if(Alpha <= 0.f)
 	{
 		const int32 PosKeyIndex1 = FMath::Min(KeyIndex1, RawTrack.PosKeys.Num()-1);
 		const int32 RotKeyIndex1 = FMath::Min(KeyIndex1, RawTrack.RotKeys.Num()-1);
-		if ( bHasScaleKey )
+		if(bHasScaleKey)
 		{
 			const int32 ScaleKeyIndex1 = FMath::Min(KeyIndex1, RawTrack.ScaleKeys.Num()-1);
-			OutAtom = FTransform(RawTrack.RotKeys[RotKeyIndex1], RawTrack.PosKeys[PosKeyIndex1], RawTrack.ScaleKeys[ScaleKeyIndex1]);		
+			OutAtom = FTransform(RawTrack.RotKeys[RotKeyIndex1], RawTrack.PosKeys[PosKeyIndex1], RawTrack.ScaleKeys[ScaleKeyIndex1]);
 		}
 		else
 		{
-			OutAtom = FTransform(RawTrack.RotKeys[RotKeyIndex1], RawTrack.PosKeys[PosKeyIndex1], DefaultScale3D);		
+			OutAtom = FTransform(RawTrack.RotKeys[RotKeyIndex1], RawTrack.PosKeys[PosKeyIndex1], DefaultScale3D);
 		}
 		return;
 	}
-	else if( Alpha >= 1.f )
+	else if(Alpha >= 1.f)
 	{
 		const int32 PosKeyIndex2 = FMath::Min(KeyIndex2, RawTrack.PosKeys.Num()-1);
 		const int32 RotKeyIndex2 = FMath::Min(KeyIndex2, RawTrack.RotKeys.Num()-1);
-		if ( bHasScaleKey )
+		if(bHasScaleKey)
 		{
 			const int32 ScaleKeyIndex2 = FMath::Min(KeyIndex2, RawTrack.ScaleKeys.Num()-1);
 			OutAtom = FTransform(RawTrack.RotKeys[RotKeyIndex2], RawTrack.PosKeys[PosKeyIndex2], RawTrack.ScaleKeys[ScaleKeyIndex2]);
 		}
 		else
 		{
-			OutAtom = FTransform(RawTrack.RotKeys[RotKeyIndex2], RawTrack.PosKeys[PosKeyIndex2], DefaultScale3D);		
+			OutAtom = FTransform(RawTrack.RotKeys[RotKeyIndex2], RawTrack.PosKeys[PosKeyIndex2], DefaultScale3D);
 		}
 		return;
 	}
@@ -584,25 +588,25 @@ void UAnimSequence::GetBoneTransform(FTransform& OutAtom, int32 TrackIndex, floa
 
 	const int32 PosKeyIndex2 = FMath::Min(KeyIndex2, RawTrack.PosKeys.Num()-1);
 	const int32 RotKeyIndex2 = FMath::Min(KeyIndex2, RawTrack.RotKeys.Num()-1);
-	
+
 	FTransform KeyAtom1, KeyAtom2;
-	
-	if ( bHasScaleKey )
+
+	if(bHasScaleKey)
 	{
 		const int32 ScaleKeyIndex1 = FMath::Min(KeyIndex1, RawTrack.ScaleKeys.Num()-1);
 		const int32 ScaleKeyIndex2 = FMath::Min(KeyIndex2, RawTrack.ScaleKeys.Num()-1);
-		
-		KeyAtom1 = FTransform(RawTrack.RotKeys[RotKeyIndex1], RawTrack.PosKeys[PosKeyIndex1], RawTrack.ScaleKeys[ScaleKeyIndex1]);	
+
+		KeyAtom1 = FTransform(RawTrack.RotKeys[RotKeyIndex1], RawTrack.PosKeys[PosKeyIndex1], RawTrack.ScaleKeys[ScaleKeyIndex1]);
 		KeyAtom2 = FTransform(RawTrack.RotKeys[RotKeyIndex2], RawTrack.PosKeys[PosKeyIndex2], RawTrack.ScaleKeys[ScaleKeyIndex2]);
 	}
 	else
 	{
-		KeyAtom1 = FTransform(RawTrack.RotKeys[RotKeyIndex1], RawTrack.PosKeys[PosKeyIndex1], DefaultScale3D);	
-		KeyAtom2 = FTransform(RawTrack.RotKeys[RotKeyIndex2], RawTrack.PosKeys[PosKeyIndex2], DefaultScale3D);	
+		KeyAtom1 = FTransform(RawTrack.RotKeys[RotKeyIndex1], RawTrack.PosKeys[PosKeyIndex1], DefaultScale3D);
+		KeyAtom2 = FTransform(RawTrack.RotKeys[RotKeyIndex2], RawTrack.PosKeys[PosKeyIndex2], DefaultScale3D);
 	}
 
-// 	UE_LOG(LogAnimation, Log, TEXT(" *  *  *  Position. PosKeyIndex1: %3d, PosKeyIndex2: %3d, Alpha: %f"), PosKeyIndex1, PosKeyIndex2, Alpha);
-// 	UE_LOG(LogAnimation, Log, TEXT(" *  *  *  Rotation. RotKeyIndex1: %3d, RotKeyIndex2: %3d, Alpha: %f"), RotKeyIndex1, RotKeyIndex2, Alpha);
+	// 	UE_LOG(LogAnimation, Log, TEXT(" *  *  *  Position. PosKeyIndex1: %3d, PosKeyIndex2: %3d, Alpha: %f"), PosKeyIndex1, PosKeyIndex2, Alpha);
+	// 	UE_LOG(LogAnimation, Log, TEXT(" *  *  *  Rotation. RotKeyIndex1: %3d, RotKeyIndex2: %3d, Alpha: %f"), RotKeyIndex1, RotKeyIndex2, Alpha);
 
 	OutAtom.Blend(KeyAtom1, KeyAtom2, Alpha);
 	OutAtom.NormalizeRotation();
@@ -1202,7 +1206,7 @@ bool UAnimSequence::CompressRawAnimSequenceTrack(FRawAnimSequenceTrack& RawTrack
 	
 	float MaxScaleDiff = 0.0001f;
 
-	// Check variation of Scale keys
+	// Check variation of Scaleition keys
 	if( (RawTrack.ScaleKeys.Num() > 1) && (MaxScaleDiff >= 0.0f) )
 	{
 		FVector FirstScale = RawTrack.ScaleKeys[0];
@@ -1554,8 +1558,32 @@ bool UAnimSequence::IsValidAdditive() const
 }
 
 #if WITH_EDITOR
-void UAnimSequence::RemapTracksToNewSkeleton( USkeleton * NewSkeleton )
+void FillUpSpaceBasesRefPose(USkeleton * Skeleton, TArray<FTransform> &SpaceBaseRefPose)
 {
+	check (Skeleton);
+
+	const TArray<FTransform> & ReferencePose = Skeleton->GetReferenceSkeleton().GetRefBonePose();
+	SpaceBaseRefPose.Empty(ReferencePose.Num());
+	SpaceBaseRefPose.AddUninitialized(ReferencePose.Num());
+
+	// initialize to identity since some of them don't have tracks
+	for(int Index=0; Index <SpaceBaseRefPose.Num(); ++Index)
+	{
+		int32 ParentIndex = Skeleton->GetReferenceSkeleton().GetParentIndex(Index);
+		if(ParentIndex != INDEX_NONE)
+		{
+			SpaceBaseRefPose[Index] = ReferencePose[Index] * SpaceBaseRefPose[ParentIndex];
+		}
+		else
+		{
+			SpaceBaseRefPose[Index] = ReferencePose[Index];
+		}
+	}
+}
+
+void UAnimSequence::RemapTracksToNewSkeleton( USkeleton * NewSkeleton, bool bConvertSpaces )
+{
+
 	// this only replaces the primary one, it doesn't replace old ones
 	TArray<struct FTrackToSkeletonMap> NewTrackToSkeletonMapTable;
 	NewTrackToSkeletonMapTable.Empty(AnimationTrackNames.Num());
@@ -1587,6 +1615,202 @@ void UAnimSequence::RemapTracksToNewSkeleton( USkeleton * NewSkeleton )
 		}
 	}
 
+	// this is not cheap, so make sure it only happens in editor
+
+	// @Todo : currently additive will work fine since we don't bake anything except when we extract
+	// but in the future if we bake this can be problem
+	if (bConvertSpaces)
+	{
+		USkeleton * OldSkeleton = GetSkeleton();
+
+		// @IMPORTANT: this converts animation data to new skeleton
+		// this requires reference pose to be relative in same position
+		// if reference pose is totally different - i.e. A pose vs T pose then it won't work
+		if(OldSkeleton)
+		{
+			// make sure you do update reference pose before coming here
+			
+			// first calculate component space ref pose to get the relative transform between
+			// two ref poses. It is very important update ref pose before getting here. 
+			TArray<FTransform> NewSpaceBaseRefPose, OldSpaceBaseRefPose, RelativeToNewTransform;
+			// get the spacebases transform
+			FillUpSpaceBasesRefPose(NewSkeleton, NewSpaceBaseRefPose);
+			FillUpSpaceBasesRefPose(OldSkeleton, OldSpaceBaseRefPose);
+
+			const TArray<FTransform> & OldRefPose = OldSkeleton->GetReferenceSkeleton().GetRefBonePose();
+			const TArray<FTransform> & NewRefPose = NewSkeleton->GetReferenceSkeleton().GetRefBonePose();
+
+			// now we'd like to get the relative transform from old to new ref pose in component space
+			// PK2*K2 = PK1*K1*theta where theta => P1*R1*theta = P2*R2 
+			// where	P1 - parent transform in component space for original skeleton
+			//			R1 - local space of the current bone for original skeleton
+			//			P2 - parent transform in component space for new skeleton
+			//			R2 - local space of the current bone for new skeleton
+			// what we're looking for is theta, so that we can apply that to animated transform
+			int32 NumBones = NewSpaceBaseRefPose.Num();
+			// saves the theta data per bone
+			RelativeToNewTransform.AddUninitialized(NumBones);
+			TArray<float> OldToNewTranslationRatio;
+			// saves the translation conversion data
+			OldToNewTranslationRatio.AddUninitialized(NumBones);
+
+			// calculate the relative transform to new skeleton
+			// so that we can apply the delta in component space
+			for(int32 BoneIndex=0; BoneIndex<NumBones; ++BoneIndex)
+			{
+				// first find bone name of the idnex
+				FName BoneName = NewSkeleton->GetReferenceSkeleton().GetRefBoneInfo()[BoneIndex].Name;
+				// find it in old index
+				int32 OldBoneIndex = OldSkeleton->GetReferenceSkeleton().FindBoneIndex(BoneName);
+
+				// get old bone index
+				if(OldBoneIndex != INDEX_NONE)
+				{
+					// theta (RelativeToNewTransform) = (P1*R1)^(-1) * P2*R2 where theta => P1*R1*theta = P2*R2
+					RelativeToNewTransform[BoneIndex] = NewSpaceBaseRefPose[BoneIndex].GetRelativeTransform(OldSpaceBaseRefPose[OldBoneIndex]);
+
+					// also savees the translation difference between old to new
+					FVector OldTranslation = OldRefPose[OldBoneIndex].GetTranslation();
+					FVector NewTranslation = NewRefPose[BoneIndex].GetTranslation();
+
+					float OldTranslationSize = OldTranslation.Size();
+					float NewTranslationSize = NewTranslation.Size();
+					OldToNewTranslationRatio[BoneIndex] = (FMath::IsNearlyZero(OldTranslationSize))? 1.f/*do not touch new translation size*/ : NewTranslationSize/OldTranslationSize;
+				}
+				else
+				{
+					RelativeToNewTransform[BoneIndex].SetIdentity();
+				}
+			}
+
+			// 2d array of animated time [boneindex][time key]
+			TArray< TArray<FTransform> > AnimatedSpaceBases, ConvertedLocalSpaces, ConvertedSpaceBases;
+			AnimatedSpaceBases.AddZeroed(NumBones);
+			ConvertedLocalSpaces.AddZeroed(NumBones);
+			ConvertedSpaceBases.AddZeroed(NumBones);
+
+			int32 NumKeys = NumFrames;
+			float Interval = (NumKeys)? SequenceLength/NumKeys : 0.f;
+
+			// allocate arrays
+			for(int32 BoneIndex=0; BoneIndex<NumBones; ++BoneIndex)
+			{
+				AnimatedSpaceBases[BoneIndex].AddUninitialized(NumKeys);
+				ConvertedLocalSpaces[BoneIndex].AddUninitialized(NumKeys);
+				ConvertedSpaceBases[BoneIndex].AddUninitialized(NumKeys);
+			}
+
+			// now calculating old animated space bases
+			// this one calculates aniamted space per bones and per key
+			for(int32 BoneIndex=0; BoneIndex<NumBones; ++BoneIndex)
+			{
+				FName BoneName = NewSkeleton->GetReferenceSkeleton().GetBoneName(BoneIndex);
+				int32 OldBoneIndex = OldSkeleton->GetReferenceSkeleton().FindBoneIndex(BoneName);
+				int32 TrackIndex = AnimationTrackNames.Find(BoneName);
+				int32 ParentBoneIndex = NewSkeleton->GetReferenceSkeleton().GetParentIndex(BoneIndex);
+
+				if(TrackIndex != INDEX_NONE)
+				{
+					auto & RawAnimation = RawAnimationData[TrackIndex];
+					// fill up keys - calculate PK1 * K1
+					for(int32 Key=0; Key<NumKeys; ++Key)
+					{
+						FTransform AnimatedLocalKey;
+						ExtractBoneTransform(RawAnimationData, AnimatedLocalKey, TrackIndex, Interval*Key);
+
+						// note that we apply scale in the animated space
+						// at this point, you should have scaled version of animated skeleton
+						AnimatedLocalKey.ScaleTranslation(OldToNewTranslationRatio[BoneIndex]);
+
+						if(ParentBoneIndex != INDEX_NONE)
+						{
+							AnimatedSpaceBases[BoneIndex][Key] = AnimatedLocalKey * AnimatedSpaceBases[ParentBoneIndex][Key];
+						}
+						else
+						{
+							AnimatedSpaceBases[BoneIndex][Key] = AnimatedLocalKey;
+						}
+					}
+				}
+				else
+				{
+					// get local spaces from refpose and use that to fill it up
+					FTransform LocalTransform = (OldBoneIndex != INDEX_NONE)? OldSkeleton->GetReferenceSkeleton().GetRefBonePose()[OldBoneIndex] : FTransform::Identity;
+
+					for(int32 Key=0; Key<NumKeys; ++Key)
+					{
+						if(ParentBoneIndex != INDEX_NONE)
+						{
+							AnimatedSpaceBases[BoneIndex][Key] = LocalTransform * AnimatedSpaceBases[ParentBoneIndex][Key];
+						}
+						else
+						{
+							AnimatedSpaceBases[BoneIndex][Key] = LocalTransform;
+						}
+					}
+				}
+			}
+
+			// now apply the theta back to the animated space bases
+			TArray<struct FRawAnimSequenceTrack> NewRawAnimationData = RawAnimationData;
+			for(int32 BoneIndex=0; BoneIndex<NumBones; ++BoneIndex)
+			{
+				FName BoneName = NewSkeleton->GetReferenceSkeleton().GetBoneName(BoneIndex);
+				int32 TrackIndex = AnimationTrackNames.Find(BoneName);
+				int32 ParentBoneIndex = NewSkeleton->GetReferenceSkeleton().GetParentIndex(BoneIndex);
+
+				for(int32 Key=0; Key<NumKeys; ++Key)
+				{
+					// thus PK2 & K2 =  PK1 * K1 * theta where theta = (P1*R1)^(-1) * P2*R2
+					// where PK2	: parent transform in component space of animated key for new skeleton
+					//		 K2		: local transform of animated key for new skeleton
+					//		 PK1	: parent transform in component space of animated key for old skeleton
+					//		 K1		: local transform of animated key for old skeleton
+					FTransform SpaceBase;
+					// we don't just apply it because translation is sensitive
+					// we don't like to apply relative transform to tranlsation directly
+					// rotation and scale we can, but translation we'd like to use scaled translation instead of transformed location
+					// as their relative translation can be different
+					SpaceBase.SetRotation(AnimatedSpaceBases[BoneIndex][Key].GetRotation() * RelativeToNewTransform[BoneIndex].GetRotation());
+					SpaceBase.SetScale3D(AnimatedSpaceBases[BoneIndex][Key].GetScale3D() * RelativeToNewTransform[BoneIndex].GetScale3D());
+					// use animated scaled translation directly
+					SpaceBase.SetTranslation(AnimatedSpaceBases[BoneIndex][Key].GetTranslation());
+					ConvertedSpaceBases[BoneIndex][Key] = SpaceBase;
+					// now calculate local space for animation
+					if(ParentBoneIndex != INDEX_NONE)
+					{
+						// K2 = PK2^(-1) * PK1 * K1 * (P1*R1)^(-1) * P2*R2
+						ConvertedLocalSpaces[BoneIndex][Key] = SpaceBase.GetRelativeTransform(ConvertedSpaceBases[ParentBoneIndex][Key]);
+					}
+					else
+					{
+						ConvertedLocalSpaces[BoneIndex][Key] = SpaceBase;
+					}
+				}
+
+				// now save back to animation data
+				if(TrackIndex != INDEX_NONE)
+				{
+					auto & RawAnimation = NewRawAnimationData[TrackIndex];
+					RawAnimation.PosKeys.Empty(NumKeys);
+					RawAnimation.PosKeys.AddUninitialized(NumKeys);
+					RawAnimation.RotKeys.Empty(NumKeys);
+					RawAnimation.RotKeys.AddUninitialized(NumKeys);
+					RawAnimation.ScaleKeys.Empty(NumKeys);
+					RawAnimation.ScaleKeys.AddUninitialized(NumKeys);
+
+					for(int32 Key=0; Key<NumKeys; ++Key)
+					{
+						RawAnimation.PosKeys[Key] = ConvertedLocalSpaces[BoneIndex][Key].GetLocation();
+						RawAnimation.RotKeys[Key] = ConvertedLocalSpaces[BoneIndex][Key].GetRotation();
+						RawAnimation.ScaleKeys[Key] = ConvertedLocalSpaces[BoneIndex][Key].GetScale3D();
+					}
+				}
+			}
+			RawAnimationData = NewRawAnimationData;
+		}
+	}
+
 	// I have to set this here in order for compression
 	// that has to happen outside of this after Skeleton changes
 	SetSkeleton(NewSkeleton);
diff --git a/Engine/Source/Runtime/Engine/Private/AnimationAsset.cpp b/Engine/Source/Runtime/Engine/Private/AnimationAsset.cpp
--- a/Engine/Source/Runtime/Engine/Private/AnimationAsset.cpp
+++ b/Engine/Source/Runtime/Engine/Private/AnimationAsset.cpp
@@ -54,7 +54,7 @@ void UAnimationAsset::SetSkeleton(USkeleton* NewSkeleton)
 }
 
 #if WITH_EDITOR
-bool UAnimationAsset::ReplaceSkeleton(USkeleton* NewSkeleton)
+bool UAnimationAsset::ReplaceSkeleton(USkeleton* NewSkeleton, bool bConvertSpaces/*=false*/)
 {
 	// if it's not same 
 	if (NewSkeleton != Skeleton)
@@ -68,7 +68,7 @@ bool UAnimationAsset::ReplaceSkeleton(USkeleton* NewSkeleton)
 				UAnimSequence* AnimSeq = *Iter;
 				if (AnimSeq && AnimSeq->Skeleton != NewSkeleton)
 				{
-					AnimSeq->RemapTracksToNewSkeleton(NewSkeleton);
+					AnimSeq->RemapTracksToNewSkeleton(NewSkeleton, bConvertSpaces);
 				}
 			}
 		}
