#ue4 - Fix issues with landing on Foliage actors with thousands of components.

- Also only add/remove tick prereqs if the target tick function can ever tick.

TTP 341962

diff --git a/Engine/Source/Runtime/Engine/Private/Character.cpp b/Engine/Source/Runtime/Engine/Private/Character.cpp
--- a/Engine/Source/Runtime/Engine/Private/Character.cpp
+++ b/Engine/Source/Runtime/Engine/Private/Character.cpp
@@ -13,6 +13,7 @@
 #include "DisplayDebugHelpers.h"
 #include "Animation/AnimInstance.h"
 #include "Animation/AnimMontage.h"
+#include "Foliage/InstancedFoliageActor.h"
 
 DEFINE_LOG_CATEGORY_STATIC(LogCharacter, Log, All);
 DEFINE_LOG_CATEGORY_STATIC(LogAvatar, Log, All);
@@ -388,21 +389,28 @@ namespace MovementBaseUtility
 			BasedObjectTick.AddPrerequisite(NewBase, NewBase->PrimaryComponentTick);
 			AActor* NewBaseOwner = NewBase->GetOwner();
 
-			// NOTE: TTP# 341962: We check the primary actor bCanEverTick here to fix issues with landing on Foliage actors with thousands of components.
-			// We need to depend on other component's Tick()s in the general case to deal with things like Timelines driving the position
-			// of moving platforms etc.
-			// This will prevent dependencies from being created for Blueprints that use timelines but do not implement a Tick() event,
-			// but this is a fairly simple workaround and the side effect is minor (a one frame lag in movement between base and character).
-			// @TODO: We need to find a more efficient way of finding all ticking components in an actor.
-			if (NewBaseOwner &&	NewBaseOwner->PrimaryActorTick.bCanEverTick)
+			if (NewBaseOwner)
 			{
-				BasedObjectTick.AddPrerequisite(NewBaseOwner, NewBaseOwner->PrimaryActorTick);
+				// NOTE: TTP# 341962: Temp hack to fix issues with landing on Foliage actors with thousands of components.
+				if (Cast<AInstancedFoliageActor>(NewBaseOwner))
+				{
+					return;
+				}
 
+				if (NewBaseOwner->PrimaryActorTick.bCanEverTick)
+				{
+					BasedObjectTick.AddPrerequisite(NewBaseOwner, NewBaseOwner->PrimaryActorTick);
+				}
+
+				// @TODO: We need to find a more efficient way of finding all ticking components in an actor.
 				TArray<UActorComponent*> Components;
 				NewBaseOwner->GetComponents(Components);
 				for (auto& Component : Components)
 				{
-					BasedObjectTick.AddPrerequisite(Component, Component->PrimaryComponentTick);
+					if (Component->PrimaryComponentTick.bCanEverTick)
+					{
+						BasedObjectTick.AddPrerequisite(Component, Component->PrimaryComponentTick);
+					}
 				}
 			}
 		}
@@ -414,15 +422,28 @@ namespace MovementBaseUtility
 		{
 			BasedObjectTick.RemovePrerequisite(OldBase, OldBase->PrimaryComponentTick);
 			AActor* OldBaseOwner = OldBase->GetOwner();
-			if (OldBaseOwner && OldBaseOwner->PrimaryActorTick.bCanEverTick)
+			if (OldBaseOwner)
 			{
-				BasedObjectTick.RemovePrerequisite(OldBaseOwner, OldBaseOwner->PrimaryActorTick);
+				// NOTE: TTP# 341962: Temp hack to fix issues with landing on Foliage actors with thousands of components.
+				if (Cast<AInstancedFoliageActor>(OldBaseOwner))
+				{
+					return;
+				}
 
+				if (OldBaseOwner->PrimaryActorTick.bCanEverTick)
+				{
+					BasedObjectTick.RemovePrerequisite(OldBaseOwner, OldBaseOwner->PrimaryActorTick);
+				}
+
+				// @TODO: We need to find a more efficient way of finding all ticking components in an actor.
 				TArray<UActorComponent*> Components;
 				OldBaseOwner->GetComponents(Components);
 				for (auto& Component : Components)
 				{
-					BasedObjectTick.RemovePrerequisite(Component, Component->PrimaryComponentTick);
+					if (Component->PrimaryComponentTick.bCanEverTick)
+					{
+						BasedObjectTick.RemovePrerequisite(Component, Component->PrimaryComponentTick);
+					}
 				}
 			}
 		}
