ttp322381 Ensure condition failed: TimelineTemplate != NULL Function: FSCSEditorViewportClient::UpdatePreviewActorForBlueprint()

diff --git a/Engine/Source/Runtime/Engine/Private/BlueprintGeneratedClass.cpp b/Engine/Source/Runtime/Engine/Private/BlueprintGeneratedClass.cpp
--- a/Engine/Source/Runtime/Engine/Private/BlueprintGeneratedClass.cpp
+++ b/Engine/Source/Runtime/Engine/Private/BlueprintGeneratedClass.cpp
@@ -285,7 +285,7 @@ void UBlueprintGeneratedClass::CreateComponentsForActor(AActor* Actor) const
 		const UTimelineTemplate* TimelineTemplate = Timelines[i];
 
 		// Not fatal if NULL, but shouldn't happen
-		if(!ensure(TimelineTemplate != NULL))
+		if(!TimelineTemplate)
 		{
 			continue;
 		}
