EvalutecurveData shouldn't be editor only.

Reported by licensee.

diff --git a/Engine/Source/Runtime/Engine/Private/Animation/AnimMontage.cpp b/Engine/Source/Runtime/Engine/Private/Animation/AnimMontage.cpp
--- a/Engine/Source/Runtime/Engine/Private/Animation/AnimMontage.cpp
+++ b/Engine/Source/Runtime/Engine/Private/Animation/AnimMontage.cpp
@@ -622,7 +622,6 @@ FTransform UAnimMontage::ExtractRootMotionFromTrackRange(float StartTrackPositio
 	return RootMotion.RootMotionTransform;
 }
 
-#if WITH_EDITOR
 void UAnimMontage::EvaluateCurveData(class UAnimInstance* Instance, float CurrentTime, float BlendWeight ) const
 {
 	Super::EvaluateCurveData(Instance, CurrentTime, BlendWeight);
@@ -654,7 +653,6 @@ void UAnimMontage::EvaluateCurveData(class UAnimInstance* Instance, float Curren
 		}	
 	}
 }
-#endif	//WITH_EDITOR
 
 //////////////////////////////////////////////////////////////////////////////////////////////
 // MontageInstance
