#ue4 - Optimize FTickFunction tick prerequisite add/remove use.

- AddPrerequisite() only adds the TargetTickFunction if it can ever tick, or if already registered.
- RemovePrerequisite() uses RemoveSwap because order doesn't matter.

diff --git a/Engine/Source/Runtime/Engine/Private/Actor.cpp b/Engine/Source/Runtime/Engine/Private/Actor.cpp
--- a/Engine/Source/Runtime/Engine/Private/Actor.cpp
+++ b/Engine/Source/Runtime/Engine/Private/Actor.cpp
@@ -339,11 +339,15 @@ void AActor::SetTickPrerequisite(AActor * DependentActor)
 		PrimaryActorTick.AddPrerequisite(DependentActor, DependentActor->PrimaryActorTick);
 		// We can't just rely on rootcomponent. If one of chain is invalid tick, it will invalidate the remaining chains
 		// so here I manually set prerequisite to be the current actor, so it ticks after this actor ticks
-		TArray<UActorComponent*> Components;
-		GetComponents(Components);
-		for (int32 Index = 0; Index < Components.Num(); ++Index)
+
+		if (PrimaryActorTick.bCanEverTick)
 		{
-			Components[Index]->PrimaryComponentTick.AddPrerequisite(this, PrimaryActorTick);
+			TArray<UActorComponent*> Components;
+			GetComponents(Components);
+			for (int32 Index = 0; Index < Components.Num(); ++Index)
+			{
+				Components[Index]->PrimaryComponentTick.AddPrerequisite(this, PrimaryActorTick);
+			}
 		}
 	}
 }
diff --git a/Engine/Source/Runtime/Engine/Private/Character.cpp b/Engine/Source/Runtime/Engine/Private/Character.cpp
--- a/Engine/Source/Runtime/Engine/Private/Character.cpp
+++ b/Engine/Source/Runtime/Engine/Private/Character.cpp
@@ -386,9 +386,12 @@ namespace MovementBaseUtility
 	{
 		if (NewBase && MovementBaseUtility::UseRelativeLocation(NewBase))
 		{
-			BasedObjectTick.AddPrerequisite(NewBase, NewBase->PrimaryComponentTick);
-			AActor* NewBaseOwner = NewBase->GetOwner();
+			if (NewBase->PrimaryComponentTick.bCanEverTick)
+			{
+				BasedObjectTick.AddPrerequisite(NewBase, NewBase->PrimaryComponentTick);
+			}
 
+			AActor* NewBaseOwner = NewBase->GetOwner();
 			if (NewBaseOwner)
 			{
 				// NOTE: TTP# 341962: Temp hack to fix issues with landing on Foliage actors with thousands of components.
@@ -430,10 +433,7 @@ namespace MovementBaseUtility
 					return;
 				}
 
-				if (OldBaseOwner->PrimaryActorTick.bCanEverTick)
-				{
-					BasedObjectTick.RemovePrerequisite(OldBaseOwner, OldBaseOwner->PrimaryActorTick);
-				}
+				BasedObjectTick.RemovePrerequisite(OldBaseOwner, OldBaseOwner->PrimaryActorTick);
 
 				// @TODO: We need to find a more efficient way of finding all ticking components in an actor.
 				TArray<UActorComponent*> Components;
diff --git a/Engine/Source/Runtime/Engine/Private/MovementComponent.cpp b/Engine/Source/Runtime/Engine/Private/MovementComponent.cpp
--- a/Engine/Source/Runtime/Engine/Private/MovementComponent.cpp
+++ b/Engine/Source/Runtime/Engine/Private/MovementComponent.cpp
@@ -104,7 +104,7 @@ void UMovementComponent::RegisterComponentTickFunctions(bool bRegister)
 
 	// If the owner ticks, make sure we tick first
 	AActor* Owner = GetOwner();
-	if (bRegister && Owner && Owner->CanEverTick())
+	if (bRegister && PrimaryComponentTick.bCanEverTick && Owner && Owner->CanEverTick())
 	{
 		Owner->PrimaryActorTick.AddPrerequisite(this, PrimaryComponentTick);
 	}
diff --git a/Engine/Source/Runtime/Engine/Private/PhysicsEngine/PhysLevel.cpp b/Engine/Source/Runtime/Engine/Private/PhysicsEngine/PhysLevel.cpp
--- a/Engine/Source/Runtime/Engine/Private/PhysicsEngine/PhysLevel.cpp
+++ b/Engine/Source/Runtime/Engine/Private/PhysicsEngine/PhysLevel.cpp
@@ -74,9 +74,16 @@ static TAutoConsoleVariable<int32> CVarAPEXSortDynamicChunksByBenefit(
 
 void UWorld::SetupPhysicsTickFunctions(float DeltaSeconds)
 {
+	StartPhysicsTickFunction.bCanEverTick = true;
 	StartPhysicsTickFunction.Target = this;
+	
+	EndPhysicsTickFunction.bCanEverTick = true;
 	EndPhysicsTickFunction.Target = this;
+	
+	StartClothTickFunction.bCanEverTick = true;
 	StartClothTickFunction.Target = this;
+	
+	EndClothTickFunction.bCanEverTick = true;
 	EndClothTickFunction.Target = this;
 	
 	
diff --git a/Engine/Source/Runtime/Engine/Private/PhysicsEngine2D/Box2DIntegration.cpp b/Engine/Source/Runtime/Engine/Private/PhysicsEngine2D/Box2DIntegration.cpp
--- a/Engine/Source/Runtime/Engine/Private/PhysicsEngine2D/Box2DIntegration.cpp
+++ b/Engine/Source/Runtime/Engine/Private/PhysicsEngine2D/Box2DIntegration.cpp
@@ -108,10 +108,12 @@ FPhysicsScene2D::FPhysicsScene2D(UWorld* AssociatedWorld)
 	World = new b2World(DefaultGravityBox);
 	World->SetDebugDraw(&FBox2DVisualizer::GlobalVisualizer);
 
+	StartPhysicsTickFunction.bCanEverTick = true;
 	StartPhysicsTickFunction.Target = this;
 	StartPhysicsTickFunction.TickGroup = TG_StartPhysics;
 	StartPhysicsTickFunction.RegisterTickFunction(AssociatedWorld->PersistentLevel);
 
+	EndPhysicsTickFunction.bCanEverTick = true;
 	EndPhysicsTickFunction.Target = this;
 	EndPhysicsTickFunction.TickGroup = TG_EndPhysics;
 	EndPhysicsTickFunction.RegisterTickFunction(AssociatedWorld->PersistentLevel);
diff --git a/Engine/Source/Runtime/Engine/Private/TickTaskManager.cpp b/Engine/Source/Runtime/Engine/Private/TickTaskManager.cpp
--- a/Engine/Source/Runtime/Engine/Private/TickTaskManager.cpp
+++ b/Engine/Source/Runtime/Engine/Private/TickTaskManager.cpp
@@ -1017,7 +1017,10 @@ void FTickFunction::SetTickFunctionEnable(bool bInEnabled)
 **/
 void FTickFunction::AddPrerequisite(UObject* TargetObject, struct FTickFunction& TargetTickFunction)
 {
-	Prerequisites.AddUnique(FTickPrerequisite(TargetObject, TargetTickFunction));
+	if (TargetTickFunction.bCanEverTick || TargetTickFunction.IsTickFunctionRegistered())
+	{
+		Prerequisites.AddUnique(FTickPrerequisite(TargetObject, TargetTickFunction));
+	}
 }
 
 /** 
@@ -1027,7 +1030,7 @@ void FTickFunction::AddPrerequisite(UObject* TargetObject, struct FTickFunction&
 **/
 void FTickFunction::RemovePrerequisite(UObject* TargetObject, struct FTickFunction& TargetTickFunction)
 {
-	Prerequisites.Remove(FTickPrerequisite(TargetObject, TargetTickFunction));
+	Prerequisites.RemoveSwap(FTickPrerequisite(TargetObject, TargetTickFunction));
 }
 
 /**
