Catch errors with invalid SceneCanvas earlier.

diff --git a/Engine/Source/Runtime/Engine/Private/GameViewportClient.cpp b/Engine/Source/Runtime/Engine/Private/GameViewportClient.cpp
--- a/Engine/Source/Runtime/Engine/Private/GameViewportClient.cpp
+++ b/Engine/Source/Runtime/Engine/Private/GameViewportClient.cpp
@@ -610,12 +610,15 @@ static UCanvas* GetCanvasByName(FName CanvasName)
 
 void UGameViewportClient::Draw(FViewport* InViewport, FCanvas* SceneCanvas)
 {
+	//Valid SceneCanvas is required.  Make this explicit.
+	check(SceneCanvas);
+
 	FCanvas* DebugCanvas = InViewport->GetDebugCanvas();
 
 	// Create a temporary canvas if there isn't already one.
 	static FName CanvasObjectName(TEXT("CanvasObject"));
 	UCanvas* CanvasObject = GetCanvasByName(CanvasObjectName);
-	CanvasObject->Canvas = SceneCanvas;	
+	CanvasObject->Canvas = SceneCanvas;		
 
 	// Create temp debug canvas object
 	static FName DebugCanvasObjectName(TEXT("DebugCanvasObject"));
@@ -933,8 +936,8 @@ void UGameViewportClient::Draw(FViewport* InViewport, FCanvas* SceneCanvas)
 						CanvasObject->Init(View->UnscaledViewRect.Width(), View->UnscaledViewRect.Height(), View);
 
 						// Set the canvas transform for the player's view rectangle.
-						SceneCanvas->PushAbsoluteTransform(FTranslationMatrix(CanvasOrigin));						
-						CanvasObject->ApplySafeZoneTransform();
+						SceneCanvas->PushAbsoluteTransform(FTranslationMatrix(CanvasOrigin));
+						CanvasObject->ApplySafeZoneTransform();						
 
 						// Render the player's HUD.
 						if( PlayerController->MyHUD )
@@ -965,7 +968,7 @@ void UGameViewportClient::Draw(FViewport* InViewport, FCanvas* SceneCanvas)
 						}
 
 						CanvasObject->PopSafeZoneTransform();
-						SceneCanvas->PopTransform();
+						SceneCanvas->PopTransform();						
 
 						// draw subtitles
 						if (!bDisplayedSubtitles)
