#ue4 - Handle NaN results from PhyX computePenetration calls.

TTP 343169

diff --git a/Engine/Source/Runtime/Engine/Private/PrimitiveComponent.cpp b/Engine/Source/Runtime/Engine/Private/PrimitiveComponent.cpp
--- a/Engine/Source/Runtime/Engine/Private/PrimitiveComponent.cpp
+++ b/Engine/Source/Runtime/Engine/Private/PrimitiveComponent.cpp
@@ -1729,8 +1729,16 @@ bool UPrimitiveComponent::ComputePenetration(FMTDResult & OutMTD, const FCollisi
 			bool bSuccess = PxGeometryQuery::computePenetration(POutDirection, OutMTD.Distance, *PGeom0, PGeomPose0, *PGeom1, PGeomPose1);
 			if (bSuccess)
 			{
-				OutMTD.Direction = P2UVector(POutDirection);
-				return true;
+				if (POutDirection.isFinite())
+				{
+					OutMTD.Direction = P2UVector(POutDirection);
+					return true;
+				}
+				else
+				{
+					UE_LOG(LogPhysics, Warning, TEXT("UPrimitiveComponent::ComputePenetration: MTD returned NaN"));
+					return false;
+				}
 			}
 		}
 	}
