Update AdBanner Blueprint nodes to reflect not just iOS and out of experimental status
#Android
#ttp  343118

diff --git a/Engine/Config/BaseEngine.ini b/Engine/Config/BaseEngine.ini
--- a/Engine/Config/BaseEngine.ini
+++ b/Engine/Config/BaseEngine.ini
@@ -1278,6 +1278,10 @@ bUseStreamingPause=false
 ;Removal of EmitterSpawnable
 +ActiveClassRedirects=(OldClassName="EmitterSpawnable", NewClassName="Emitter")
 
+;AdBanner redirects
++K2FieldRedirects=(OldFieldName="KismetSystemLibrary.EXPERIMENTAL_ShowAdBanner",NewFieldName="KismetSystemLibrary.ShowAdBanner")
++K2FieldRedirects=(OldFieldName="KismetSystemLibrary.EXPERIMENTAL_HideAdBanner",NewFieldName="KismetSystemLibrary.HideAdBanner")
++K2FieldRedirects=(OldFieldName="KismetSystemLibrary.EXPERIMENTAL_CloseAdBanner",NewFieldName="KismetSystemLibrary.ForceCloseAdBanner")
 
 [CoreUObject.Metadata]
 ; Note: UnrealHeaderTool should be rerun after making changes to MetadataRedirects to catch any keys specified in class headers
diff --git a/Engine/Source/Runtime/Engine/Classes/Kismet/KismetSystemLibrary.h b/Engine/Source/Runtime/Engine/Classes/Kismet/KismetSystemLibrary.h
--- a/Engine/Source/Runtime/Engine/Classes/Kismet/KismetSystemLibrary.h
+++ b/Engine/Source/Runtime/Engine/Classes/Kismet/KismetSystemLibrary.h
@@ -1124,27 +1124,27 @@ class UKismetSystemLibrary : public UBlueprintFunctionLibrary
 	static void CollectGarbage();
 
 	/**
-	 * Will show an iAd on the top or bottom of screen, on top of the GL view (doesn't resize the view)
-	 * (iOS only; this function will be renamed or moved in a future release)
+	 * Will show an ad banner (iAd on iOS, or AdMob on Android) on the top or bottom of screen, on top of the GL view (doesn't resize the view)
+	 * (iOS and Android only)
 	 *
 	 * @param bShowOnBottomOfScreen If true, the iAd will be shown at the bottom of the screen, top otherwise
 	 */
 	UFUNCTION(BlueprintCallable, Category = "Utilities|Platform")
-	static void EXPERIMENTAL_ShowAdBanner(bool bShowOnBottomOfScreen);
+	static void ShowAdBanner(bool bShowOnBottomOfScreen);
 
 	/**
-	 * Hides the iAd banner shows with IPhoneShowAdBanner. Will force close the ad if it's open
-	 * (iOS only; this function will be renamed or moved in a future release)
+	 * Hides the ad banner (iAd on iOS, or AdMob on Android). Will force close the ad if it's open
+	 * (iOS and Android only)
 	 */
 	UFUNCTION(BlueprintCallable, Category = "Utilities|Platform")
-	static void EXPERIMENTAL_HideAdBanner();
+	static void HideAdBanner();
 
 	/**
 	 * Forces closed any displayed ad. Can lead to loss of revenue
-	 * (iOS only; this function will be renamed or moved in a future release)
+	 * (iOS and Android only)
 	 */
 	UFUNCTION(BlueprintCallable, Category = "Utilities|Platform")
-	static void EXPERIMENTAL_CloseAdBanner();
+	static void ForceCloseAdBanner();
 
 	/**
 	 * Displays the built-in leaderboard GUI (iOS and Android only; this function may be renamed or moved in a future release)
diff --git a/Engine/Source/Runtime/Engine/Private/KismetSystemLibrary.cpp b/Engine/Source/Runtime/Engine/Private/KismetSystemLibrary.cpp
--- a/Engine/Source/Runtime/Engine/Private/KismetSystemLibrary.cpp
+++ b/Engine/Source/Runtime/Engine/Private/KismetSystemLibrary.cpp
@@ -2606,7 +2606,7 @@ void UKismetSystemLibrary::CollectGarbage()
 	GEngine->DeferredCommands.Add(TEXT("obj gc"));
 }
 
-void UKismetSystemLibrary::EXPERIMENTAL_ShowAdBanner(bool bShowOnBottomOfScreen)
+void UKismetSystemLibrary::ShowAdBanner(bool bShowOnBottomOfScreen)
 {
 	if (IAdvertisingProvider* Provider = FAdvertising::Get().GetDefaultProvider())
 	{
@@ -2614,7 +2614,7 @@ void UKismetSystemLibrary::EXPERIMENTAL_ShowAdBanner(bool bShowOnBottomOfScreen)
 	}
 }
 
-void UKismetSystemLibrary::EXPERIMENTAL_HideAdBanner()
+void UKismetSystemLibrary::HideAdBanner()
 {
 	if (IAdvertisingProvider* Provider = FAdvertising::Get().GetDefaultProvider())
 	{
@@ -2622,7 +2622,7 @@ void UKismetSystemLibrary::EXPERIMENTAL_HideAdBanner()
 	}
 }
 
-void UKismetSystemLibrary::EXPERIMENTAL_CloseAdBanner()
+void UKismetSystemLibrary::ForceCloseAdBanner()
 {
 	if (IAdvertisingProvider* Provider = FAdvertising::Get().GetDefaultProvider())
 	{
