#ue4 - ProjectileMovementComponent should still update ComponentVelocity when HandleHitWall aborts the tick.

diff --git a/Engine/Source/Runtime/Engine/Private/ProjectileMovementComponent.cpp b/Engine/Source/Runtime/Engine/Private/ProjectileMovementComponent.cpp
--- a/Engine/Source/Runtime/Engine/Private/ProjectileMovementComponent.cpp
+++ b/Engine/Source/Runtime/Engine/Private/ProjectileMovementComponent.cpp
@@ -148,7 +148,7 @@ void UProjectileMovementComponent::TickComponent(float DeltaTime, enum ELevelTic
 			}
 			if ( HandleHitWall(Hit, TimeTick*Hit.Time, MoveDelta) )
 			{
-				return;
+				break;
 			}
 			if( NumBounces < 2 )
 			{
