Fix crash that can happen when leaving a splitscreen game as a client

diff --git a/Engine/Source/Runtime/Engine/Private/NetConnection.cpp b/Engine/Source/Runtime/Engine/Private/NetConnection.cpp
--- a/Engine/Source/Runtime/Engine/Private/NetConnection.cpp
+++ b/Engine/Source/Runtime/Engine/Private/NetConnection.cpp
@@ -343,8 +343,12 @@ void UChildConnection::CleanUp()
 	{
 		if (OwningActor != NULL)
 		{
-			// Cleanup/Destroy the connection actor & controller	
-			OwningActor->OnNetCleanup(this);
+			if ( !OwningActor->HasAnyFlags( RF_BeginDestroyed | RF_FinishDestroyed ) )
+			{
+				// Cleanup/Destroy the connection actor & controller	
+				OwningActor->OnNetCleanup(this);
+			}
+
 			OwningActor = NULL;
 			PlayerController = NULL;
 		}
