#ue4 - Added functions on Actor and ActorComponent to add/remove tick order dependencies.

- Added AddTickPrerequisiteActor, AddTickPrerequisiteComponent, and matching Remove functions.
- Deprecated old SetTickPrerequisite.***  Note that this function used to also change the components of the actor to tick after the actor but this is no longer the case. ***

diff --git a/Engine/Config/BaseEngine.ini b/Engine/Config/BaseEngine.ini
--- a/Engine/Config/BaseEngine.ini
+++ b/Engine/Config/BaseEngine.ini
@@ -940,6 +940,9 @@ bUseStreamingPause=false
 +K2FieldRedirects=(OldFieldName="SceneComponent.AddRelativeTranslation",NewFieldName="SceneComponent.AddRelativeLocation")
 +K2ParamRedirects=(NodeName="/Script/BlueprintGraph.K2Node_CallFunction", OldParamName="Actor.AddRelativeLocation.DeltaTranslation", NewParamName="DeltaLocation")
 
++K2FieldRedirects=(OldFieldName="Actor.SetTickPrerequisite",NewFieldName="Actor.AddTickPrerequisiteActor")
++K2ParamRedirects=(NodeName="/Script/BlueprintGraph.K2Node_CallFunction", OldParamName="Actor.AddTickPrerequisiteActor.DependentActor", NewParamName="PrerequisiteActor")
+
 +K2FieldRedirects=(OldFieldName="SceneComponent.bAbsoluteTranslation",NewFieldName="SceneComponent.bAbsoluteLocation")
 +K2FieldRedirects=(OldFieldName="SceneComponent.RelativeTranslation",NewFieldName="SceneComponent.RelativeLocation")
 
diff --git a/Engine/Source/Runtime/Engine/Classes/Components/ActorComponent.h b/Engine/Source/Runtime/Engine/Classes/Components/ActorComponent.h
--- a/Engine/Source/Runtime/Engine/Classes/Components/ActorComponent.h
+++ b/Engine/Source/Runtime/Engine/Classes/Components/ActorComponent.h
@@ -481,6 +481,22 @@ public:
 	/** Changes the ticking group for this component */
 	void SetTickGroup(ETickingGroup NewTickGroup);
 
+	/** Make this component tick after PrerequisiteActor */
+	UFUNCTION(BlueprintCallable, Category="Utilities", meta=(Keywords = "dependency"))
+	virtual void AddTickPrerequisiteActor(AActor* PrerequisiteActor);
+
+	/** Make this component tick after PrerequisiteComponent. */
+	UFUNCTION(BlueprintCallable, Category="Utilities", meta=(Keywords = "dependency"))
+	virtual void AddTickPrerequisiteComponent(UActorComponent* PrerequisiteComponent);
+
+	/** Remove tick dependency on PrerequisiteActor. */
+	UFUNCTION(BlueprintCallable, Category="Utilities", meta=(Keywords = "dependency"))
+	virtual void RemoveTickPrerequisiteActor(AActor* PrerequisiteActor);
+
+	/** Remove tick dependency on PrerequisiteComponent. */
+	UFUNCTION(BlueprintCallable, Category="Utilities", meta=(Keywords = "dependency"))
+	virtual void RemoveTickPrerequisiteComponent(UActorComponent* PrerequisiteComponent);
+
 	/** 
 	 *  Called by owner actor on position shifting
 	 *  Component should update all relevant data structures to reflect new actor location
diff --git a/Engine/Source/Runtime/Engine/Classes/GameFramework/Actor.h b/Engine/Source/Runtime/Engine/Classes/GameFramework/Actor.h
--- a/Engine/Source/Runtime/Engine/Classes/GameFramework/Actor.h
+++ b/Engine/Source/Runtime/Engine/Classes/GameFramework/Actor.h
@@ -776,10 +776,26 @@ public:
 	UFUNCTION(BlueprintCallable, Category="Utilities|Time")
 	float GetActorTimeDilation() const;
 
-	/** Adds a dependent actor, so DependentActor ticks before this actor. */
-	UFUNCTION(BlueprintCallable, Category="Utilities", meta=(FriendlyName = "SetTickDependency"))
-	void SetTickPrerequisite(AActor* DependentActor);
+	DEPRECATED(4.5, "Actor::SetTickPrerequisite() will be removed, use AddTickPrerequisiteActor().")
+	void SetTickPrerequisite(AActor* PrerequisiteActor);
 
+	/** Make this actor tick after PrerequisiteActor */
+	UFUNCTION(BlueprintCallable, Category="Utilities", meta=(Keywords = "dependency"))
+	virtual void AddTickPrerequisiteActor(AActor* PrerequisiteActor);
+
+	/** Make this actor tick after PrerequisiteComponent. */
+	UFUNCTION(BlueprintCallable, Category="Utilities", meta=(Keywords = "dependency"))
+	virtual void AddTickPrerequisiteComponent(UActorComponent* PrerequisiteComponent);
+
+	/** Remove tick dependency on PrerequisiteActor. */
+	UFUNCTION(BlueprintCallable, Category="Utilities", meta=(Keywords = "dependency"))
+	virtual void RemoveTickPrerequisiteActor(AActor* PrerequisiteActor);
+
+	/** Remove tick dependency on PrerequisiteComponent. */
+	UFUNCTION(BlueprintCallable, Category="Utilities", meta=(Keywords = "dependency"))
+	virtual void RemoveTickPrerequisiteComponent(UActorComponent* PrerequisiteComponent);
+
+	/** Sets whether this actor can tick when paused. */
 	UFUNCTION(BlueprintCallable, Category="Utilities")
 	void SetTickableWhenPaused(bool bTickableWhenPaused);
 
diff --git a/Engine/Source/Runtime/Engine/Private/Actor.cpp b/Engine/Source/Runtime/Engine/Private/Actor.cpp
--- a/Engine/Source/Runtime/Engine/Private/Actor.cpp
+++ b/Engine/Source/Runtime/Engine/Private/Actor.cpp
@@ -332,23 +332,40 @@ bool AActor::K2_TeleportTo( FVector DestLocation, FRotator DestRotation )
 	return TeleportTo(DestLocation, DestRotation, false, false);
 }
 
-void AActor::SetTickPrerequisite(AActor * DependentActor)
+void AActor::SetTickPrerequisite(AActor* PrerequisiteActor)
 {
-	if (DependentActor)
+	AddTickPrerequisiteActor(PrerequisiteActor);
+}
+
+void AActor::AddTickPrerequisiteActor(AActor* PrerequisiteActor)
+{
+	if (PrimaryActorTick.bCanEverTick && PrerequisiteActor && PrerequisiteActor->PrimaryActorTick.bCanEverTick)
 	{
-		PrimaryActorTick.AddPrerequisite(DependentActor, DependentActor->PrimaryActorTick);
-		// We can't just rely on rootcomponent. If one of chain is invalid tick, it will invalidate the remaining chains
-		// so here I manually set prerequisite to be the current actor, so it ticks after this actor ticks
+		PrimaryActorTick.AddPrerequisite(PrerequisiteActor, PrerequisiteActor->PrimaryActorTick);
+	}
+}
 
-		if (PrimaryActorTick.bCanEverTick)
-		{
-			TArray<UActorComponent*> Components;
-			GetComponents(Components);
-			for (int32 Index = 0; Index < Components.Num(); ++Index)
-			{
-				Components[Index]->PrimaryComponentTick.AddPrerequisite(this, PrimaryActorTick);
-			}
-		}
+void AActor::AddTickPrerequisiteComponent(UActorComponent* PrerequisiteComponent)
+{
+	if (PrimaryActorTick.bCanEverTick && PrerequisiteComponent && PrerequisiteComponent->PrimaryComponentTick.bCanEverTick)
+	{
+		PrimaryActorTick.AddPrerequisite(PrerequisiteComponent, PrerequisiteComponent->PrimaryComponentTick);
+	}
+}
+
+void AActor::RemoveTickPrerequisiteActor(AActor* PrerequisiteActor)
+{
+	if (PrerequisiteActor)
+	{
+		PrimaryActorTick.RemovePrerequisite(PrerequisiteActor, PrerequisiteActor->PrimaryActorTick);
+	}
+}
+
+void AActor::RemoveTickPrerequisiteComponent(UActorComponent* PrerequisiteComponent)
+{
+	if (PrerequisiteComponent)
+	{
+		PrimaryActorTick.RemovePrerequisite(PrerequisiteComponent, PrerequisiteComponent->PrimaryComponentTick);
 	}
 }
 
diff --git a/Engine/Source/Runtime/Engine/Private/ActorComponent.cpp b/Engine/Source/Runtime/Engine/Private/ActorComponent.cpp
--- a/Engine/Source/Runtime/Engine/Private/ActorComponent.cpp
+++ b/Engine/Source/Runtime/Engine/Private/ActorComponent.cpp
@@ -866,6 +866,39 @@ void UActorComponent::SetTickGroup(ETickingGroup NewTickGroup)
 	PrimaryComponentTick.TickGroup = NewTickGroup;
 }
 
+
+void UActorComponent::AddTickPrerequisiteActor(AActor* PrerequisiteActor)
+{
+	if (PrimaryComponentTick.bCanEverTick && PrerequisiteActor && PrerequisiteActor->PrimaryActorTick.bCanEverTick)
+	{
+		PrimaryComponentTick.AddPrerequisite(PrerequisiteActor, PrerequisiteActor->PrimaryActorTick);
+	}
+}
+
+void UActorComponent::AddTickPrerequisiteComponent(UActorComponent* PrerequisiteComponent)
+{
+	if (PrimaryComponentTick.bCanEverTick && PrerequisiteComponent && PrerequisiteComponent->PrimaryComponentTick.bCanEverTick)
+	{
+		PrimaryComponentTick.AddPrerequisite(PrerequisiteComponent, PrerequisiteComponent->PrimaryComponentTick);
+	}
+}
+
+void UActorComponent::RemoveTickPrerequisiteActor(AActor* PrerequisiteActor)
+{
+	if (PrerequisiteActor)
+	{
+		PrimaryComponentTick.RemovePrerequisite(PrerequisiteActor, PrerequisiteActor->PrimaryActorTick);
+	}
+}
+
+void UActorComponent::RemoveTickPrerequisiteComponent(UActorComponent* PrerequisiteComponent)
+{
+	if (PrerequisiteComponent)
+	{
+		PrimaryComponentTick.RemovePrerequisite(PrerequisiteComponent, PrerequisiteComponent->PrimaryComponentTick);
+	}
+}
+
 void UActorComponent::DoDeferredRenderUpdates_Concurrent()
 {
 	checkf(!HasAnyFlags(RF_Unreachable), TEXT("%s"), *GetFullName());
