changed the release way of a clothing actor

diff --git a/Engine/Source/Runtime/Engine/Private/SkeletalMeshComponentPhysics.cpp b/Engine/Source/Runtime/Engine/Private/SkeletalMeshComponentPhysics.cpp
--- a/Engine/Source/Runtime/Engine/Private/SkeletalMeshComponentPhysics.cpp
+++ b/Engine/Source/Runtime/Engine/Private/SkeletalMeshComponentPhysics.cpp
@@ -56,7 +56,7 @@ void FClothingActor::Clear(bool bReleaseResource)
 {
 	if(bReleaseResource)
 	{
-		GPhysCommandHandler->DeferredRelease(ApexClothingActor);
+		PhysScene->DeferredCommandHandler.DeferredRelease(ApexClothingActor);
 	}
 
 	ParentClothingAsset = NULL;
