Fixed crash when undoing brush add operation
CRITICAL: Regression: EDITOR: CRASH: Undoing the placement of geometry after switching from Landscape mode
TTP# 343161

diff --git a/Engine/Source/Runtime/Engine/Private/Level.cpp b/Engine/Source/Runtime/Engine/Private/Level.cpp
--- a/Engine/Source/Runtime/Engine/Private/Level.cpp
+++ b/Engine/Source/Runtime/Engine/Private/Level.cpp
@@ -737,6 +737,32 @@ void ULevel::UpdateLevelComponents(bool bRerunConstructionScripts)
 	IncrementalUpdateComponents( 0, bRerunConstructionScripts );
 }
 
+/**
+*	Sorts actors such that parent actors will appear before children actors in the list
+*	Stable sort
+*/
+static void SortActorsHierarchy(TTransArray<AActor*>& Actors)
+{
+	auto CalcAttachDepth = [](AActor* InActor) -> int32 {
+		int32 Depth = MAX_int32;
+		if (InActor)
+		{
+			Depth = 0;
+			if (InActor->GetRootComponent())
+			{
+				for (const USceneComponent* Test = InActor->GetRootComponent()->AttachParent; Test != nullptr; Test = Test->AttachParent, Depth++);
+			}
+		}
+		return Depth;
+	};
+	
+	// Unfortunately TArray.StableSort assumes no null entries in the array
+	// So it forces me to use internal unrestricted version
+	StableSortInternal(Actors.GetTypedData(), Actors.Num(), [&](AActor* L, AActor* R) {
+			return CalcAttachDepth(L) < CalcAttachDepth(R);
+	});
+}
+
 void ULevel::IncrementalUpdateComponents(int32 NumComponentsToUpdate, bool bRerunConstructionScripts)
 {
 	// A value of 0 means that we want to update all components.
@@ -750,22 +776,8 @@ void ULevel::IncrementalUpdateComponents(int32 NumComponentsToUpdate, bool bReru
 	if (CurrentActorIndexForUpdateComponents == 0)
 	{
 		UpdateModelComponents();
-		
 		// Sort actors to ensure that parent actors will be registered before child actors
-		Actors.Remove(nullptr);
-		
-		auto CalcAttachDepth = [](AActor& InActor) -> int32 {
-			int32 Depth = 0;
-			if (InActor.GetRootComponent())
-			{
-				for (const USceneComponent* Test = InActor.GetRootComponent()->AttachParent; Test != nullptr; Test = Test->AttachParent, Depth++);
-			}
-			return Depth;
-		};
-
-		Actors.StableSort([&](AActor& L, AActor& R) {
-			return CalcAttachDepth(L) < CalcAttachDepth(R);
-		});
+		SortActorsHierarchy(Actors);
 	}
 
 	// Find next valid actor to process components registration
