Fixed wrong landscape rect calculation when landscape has only 1 component

diff --git a/Engine/Source/Runtime/Engine/Private/Landscape/LandscapeEdit.cpp b/Engine/Source/Runtime/Engine/Private/Landscape/LandscapeEdit.cpp
--- a/Engine/Source/Runtime/Engine/Private/Landscape/LandscapeEdit.cpp
+++ b/Engine/Source/Runtime/Engine/Private/Landscape/LandscapeEdit.cpp
@@ -2370,7 +2370,7 @@ FIntRect ALandscapeProxy::GetBoundingRect() const
 		Rect.Include(LandscapeComponents[CompIdx]->GetSectionBase());
 	}
 
-	if (Rect.Width() > 0 && Rect.Height() > 0)
+	if (LandscapeComponents.Num() > 0)
 	{
 		Rect.Max += FIntPoint(ComponentSizeQuads, ComponentSizeQuads);
 		Rect -= LandscapeSectionOffset;
