Fix for smartname mapping not being correctly populated on importing a new skeletal mesh.

diff --git a/Engine/Source/Runtime/Engine/Private/Animation/Skeleton.cpp b/Engine/Source/Runtime/Engine/Private/Animation/Skeleton.cpp
--- a/Engine/Source/Runtime/Engine/Private/Animation/Skeleton.cpp
+++ b/Engine/Source/Runtime/Engine/Private/Animation/Skeleton.cpp
@@ -46,6 +46,9 @@ USkeleton::USkeleton(const class FPostConstructInitializeProperties& PCIP)
 #if WITH_EDITORONLY_DATA
 	SlotGroupNames.Empty();
 	SlotGroupNames.Add(DefaultSlotGroupName);
+
+	// Make sure we have somewhere for curve names.
+	SmartNames.AddContainer(AnimCurveMappingName);
 #endif
 }
 
@@ -74,9 +77,6 @@ void USkeleton::PostLoad()
 	// catch any case if guid isn't valid
 	check(Guid.IsValid());
 
-	// Set up smart name mappings, if these have been loaded they will already be present
-	SmartNames.AddContainer(AnimCurveMappingName);
-
 #if WITH_EDITOR
 	if ( GIsEditor )
 	{
