Don't free PlayerController if we won't own it when removing local players + sanity checks

diff --git a/Engine/Source/Runtime/Engine/Private/GameInstance.cpp b/Engine/Source/Runtime/Engine/Private/GameInstance.cpp
--- a/Engine/Source/Runtime/Engine/Private/GameInstance.cpp
+++ b/Engine/Source/Runtime/Engine/Private/GameInstance.cpp
@@ -438,14 +438,18 @@ bool UGameInstance::RemoveLocalPlayer(ULocalPlayer* ExistingPlayer)
 	{
 		// FIXME: Do this all inside PlayerRemoved?
 		ExistingPlayer->PlayerController->CleanupGameViewport();
+
 		// Destroy the player's actors.
-		ExistingPlayer->PlayerController->Destroy();
+		if ( ExistingPlayer->PlayerController->Role == ROLE_Authority )
+		{
+			ExistingPlayer->PlayerController->Destroy();
+		}
 	}
 
 	// Remove the player from the context list
 	const int32 OldIndex = LocalPlayers.Find(ExistingPlayer);
 
-	if (OldIndex != INDEX_NONE)
+	if (ensure(OldIndex != INDEX_NONE))
 	{
 		ExistingPlayer->PlayerRemoved();
 		LocalPlayers.RemoveAt(OldIndex);
