Fix crash that can happen when splitscreen players leave

diff --git a/Engine/Source/Runtime/Engine/Private/HUD.cpp b/Engine/Source/Runtime/Engine/Private/HUD.cpp
--- a/Engine/Source/Runtime/Engine/Private/HUD.cpp
+++ b/Engine/Source/Runtime/Engine/Private/HUD.cpp
@@ -131,7 +131,7 @@ void AHUD::PostRender()
 		
 		ULocalPlayer* LocalPlayer = GetOwningPlayerController() ? Cast<ULocalPlayer>(GetOwningPlayerController()->Player) : NULL;
 
-		if (LocalPlayer)
+		if (LocalPlayer && LocalPlayer->ViewportClient)
 		{
 			TArray<FVector2D> ContactPoints;
 
