- Fix to FBodyInstance::SetMaxDepenetrationVelocity, fixes boxes penetrating

diff --git a/Engine/Source/Runtime/Engine/Private/PhysicsEngine/BodyInstance.cpp b/Engine/Source/Runtime/Engine/Private/PhysicsEngine/BodyInstance.cpp
--- a/Engine/Source/Runtime/Engine/Private/PhysicsEngine/BodyInstance.cpp
+++ b/Engine/Source/Runtime/Engine/Private/PhysicsEngine/BodyInstance.cpp
@@ -2745,7 +2745,7 @@ void FBodyInstance::SetMaxDepenetrationVelocity(float MaxVelocity)
 
 	if (PxRigidDynamic* PRigidDynamic = GetPxRigidDynamic())
 	{
-		PRigidDynamic->setMaxDepenetrationVelocity(MaxVelocity);
+		PRigidDynamic->setMaxDepenetrationVelocity(UseMaxVelocity);
 	}
 #endif
 }
