#UE4 Expose DegreesToRadians, RadiansToDegrees and PI to Blueprints

GitHub pull request #330 from user tommybear

diff --git a/Engine/Source/Runtime/Engine/Classes/Kismet/KismetMathLibrary.h b/Engine/Source/Runtime/Engine/Classes/Kismet/KismetMathLibrary.h
--- a/Engine/Source/Runtime/Engine/Classes/Kismet/KismetMathLibrary.h
+++ b/Engine/Source/Runtime/Engine/Classes/Kismet/KismetMathLibrary.h
@@ -297,6 +297,18 @@ class UKismetMathLibrary : public UBlueprintFunctionLibrary
 	UFUNCTION(BlueprintPure, Category="Math|Random")
 	static float RandomFloatInRange(float Min, float Max);
 
+	/* Returns the value of PI */
+	UFUNCTION(BlueprintPure, meta=(FriendlyName = "Get PI", CompactNodeTitle = "PI"), Category="Math|Trig")
+	static float GetPI();
+
+	/* Returns radians value based on the input degrees */
+	UFUNCTION(BlueprintPure, meta=(FriendlyName = "Degrees To Radians", CompactNodeTitle = "D2R"), Category="Math|Trig")
+	static float DegreesToRadians(float A);
+
+	/* Returns degrees value based on the input radians */
+	UFUNCTION(BlueprintPure, meta=(FriendlyName = "Radians To Degrees", CompactNodeTitle = "R2D"), Category="Math|Trig")
+	static float RadiansToDegrees(float A);
+
 	/* Returns the sin of A (expects Degrees)*/
 	UFUNCTION(BlueprintPure, meta=(FriendlyName = "Sin (Degrees)", CompactNodeTitle = "SINd"), Category="Math|Trig")
 	static float DegSin(float A);
diff --git a/Engine/Source/Runtime/Engine/Private/KismetMathLibrary.cpp b/Engine/Source/Runtime/Engine/Private/KismetMathLibrary.cpp
--- a/Engine/Source/Runtime/Engine/Private/KismetMathLibrary.cpp
+++ b/Engine/Source/Runtime/Engine/Private/KismetMathLibrary.cpp
@@ -323,6 +323,22 @@ bool UKismetMathLibrary::InRange_FloatFloat(float Value, float Min, float Max, b
 	return ((InclusiveMin ? (Value >= Min) : (Value > Min)) && (InclusiveMax ? (Value <= Max) : (Value < Max)));
 }	
 
+float UKismetMathLibrary::GetPI()
+{
+    return PI;
+}
+
+float UKismetMathLibrary::DegreesToRadians(float A)
+{
+    return FMath::DegreesToRadians(A);
+}
+
+float UKismetMathLibrary::RadiansToDegrees(float A)
+{
+    return FMath::RadiansToDegrees(A);
+}
+
+
 float UKismetMathLibrary::Abs(float A)
 {
 	return FMath::Abs(A);
