Added BlueprintReadWrite to FFormatTextArgument to fix warnings about deprecation when using Format Text nodes in Blueprints.

diff --git a/Engine/Source/Runtime/Engine/Classes/Kismet/KismetTextLibrary.h b/Engine/Source/Runtime/Engine/Classes/Kismet/KismetTextLibrary.h
--- a/Engine/Source/Runtime/Engine/Classes/Kismet/KismetTextLibrary.h
+++ b/Engine/Source/Runtime/Engine/Classes/Kismet/KismetTextLibrary.h
@@ -33,10 +33,10 @@ struct FFormatTextArgument
 {
 	GENERATED_USTRUCT_BODY()
 
-	UPROPERTY(EditInstanceOnly, Category=ArgumentValue)
+	UPROPERTY(EditInstanceOnly, BlueprintReadWrite, Category=ArgumentValue)
 	FText ArgumentName;
 
-	UPROPERTY(EditInstanceOnly, Category=ArgumentValue)
+	UPROPERTY(EditInstanceOnly, BlueprintReadWrite, Category=ArgumentValue)
 	FText TextValue;
 };
 
