TTP #342124 : SkelControl Trail Node integration

diff --git a/Engine/Source/Editor/AnimGraph/Classes/AnimGraphNode_Trail.h b/Engine/Source/Editor/AnimGraph/Classes/AnimGraphNode_Trail.h
new file mode 100644
--- /dev/null
+++ b/Engine/Source/Editor/AnimGraph/Classes/AnimGraphNode_Trail.h
@@ -0,0 +1,27 @@
+// Copyright 1998-2014 Epic Games, Inc. All Rights Reserved.
+
+#pragma once
+#include "AnimGraphNode_SkeletalControlBase.h"
+#include "Animation/BoneControllers/AnimNode_Trail.h"
+#include "AnimGraphNode_Trail.generated.h"
+
+UCLASS(MinimalAPI)
+class UAnimGraphNode_Trail : public UAnimGraphNode_SkeletalControlBase
+{
+	GENERATED_UCLASS_BODY()
+
+		UPROPERTY(EditAnywhere, Category = Settings)
+		FAnimNode_Trail Node;
+
+public:
+	// UEdGraphNode interface
+	virtual FText GetNodeTitle(ENodeTitleType::Type TitleType) const override;
+//	virtual FString GetNodeNativeTitle(ENodeTitleType::Type TitleType) const override;
+	virtual FString GetTooltip() const override;
+	// End of UEdGraphNode interface
+
+protected:
+	// UAnimGraphNode_SkeletalControlBase interface
+	virtual FText GetControllerDescription() const override;
+	// End of UAnimGraphNode_SkeletalControlBase interface
+};
diff --git a/Engine/Source/Editor/AnimGraph/Private/AnimGraphNode_Trail.cpp b/Engine/Source/Editor/AnimGraph/Private/AnimGraphNode_Trail.cpp
new file mode 100644
--- /dev/null
+++ b/Engine/Source/Editor/AnimGraph/Private/AnimGraphNode_Trail.cpp
@@ -0,0 +1,42 @@
+// Copyright 1998-2014 Epic Games, Inc. All Rights Reserved.
+
+#include "AnimGraphPrivatePCH.h"
+#include "AnimGraphNode_Trail.h"
+
+/////////////////////////////////////////////////////
+// UAnimGraphNode_TrailBone
+
+#define LOCTEXT_NAMESPACE "A3Nodes"
+
+UAnimGraphNode_Trail::UAnimGraphNode_Trail(const FPostConstructInitializeProperties& PCIP)
+	: Super(PCIP)
+{
+}
+
+FText UAnimGraphNode_Trail::GetControllerDescription() const
+{
+	return LOCTEXT("TrailController", "Trail controller");
+}
+
+FString UAnimGraphNode_Trail::GetTooltip() const
+{
+	return LOCTEXT("AnimGraphNode_Trail_Tooltip", "The Trail Controller.").ToString();
+}
+
+FText UAnimGraphNode_Trail::GetNodeTitle(ENodeTitleType::Type TitleType) const
+{
+	FFormatNamedArguments Args;
+	Args.Add(TEXT("ControllerDescription"), GetControllerDescription());
+	Args.Add(TEXT("BoneName"), FText::FromName(Node.TrailBone.BoneName));
+
+	if (TitleType == ENodeTitleType::ListView)
+	{
+		return FText::Format(LOCTEXT("AnimGraphNode_Trail_Title", "{ControllerDescription} - Bone: {BoneName}"), Args);
+	}
+	else
+	{
+		return FText::Format(LOCTEXT("AnimGraphNode_Trail_Title", "{ControllerDescription}\nBone: {BoneName}"), Args);
+	}
+}
+
+#undef LOCTEXT_NAMESPACE
\ No newline at end of file
diff --git a/Engine/Source/Runtime/Engine/Classes/Animation/BoneControllers/AnimNode_Trail.h b/Engine/Source/Runtime/Engine/Classes/Animation/BoneControllers/AnimNode_Trail.h
new file mode 100644
--- /dev/null
+++ b/Engine/Source/Runtime/Engine/Classes/Animation/BoneControllers/AnimNode_Trail.h
@@ -0,0 +1,84 @@
+// Copyright 1998-2014 Epic Games, Inc. All Rights Reserved.
+
+#pragma once
+
+#include "AnimNode_SkeletalControlBase.h"
+#include "AnimNode_Trail.generated.h"
+
+/**
+ * Trail Controller
+ */
+
+USTRUCT()
+struct ENGINE_API FAnimNode_Trail : public FAnimNode_SkeletalControlBase
+{
+	GENERATED_USTRUCT_BODY()
+
+	/** Reference to the active bone in the hierarchy to modify. */
+	UPROPERTY(EditAnywhere, BlueprintReadWrite, Category=Trail)
+	FBoneReference TrailBone;
+
+	/** Number of bones above the active one in the hierarchy to modify. ChainLength should be at least 2. */
+	UPROPERTY(EditAnywhere, BlueprintReadWrite, Category = Trail, meta = (ClampMin = "2", UIMin = "2"))
+	int32	ChainLength;
+
+	/** Axis of the bones to point along trail. */
+	UPROPERTY(EditAnywhere, BlueprintReadWrite, Category=Trail)
+	TEnumAsByte<EAxis::Type>	ChainBoneAxis;
+
+	/** Invert the direction specified in ChainBoneAxis. */
+	UPROPERTY(EditAnywhere, BlueprintReadWrite, Category=Trail)
+	bool	bInvertChainBoneAxis;
+
+	/** Limit the amount that a bone can stretch from its ref-pose length. */
+	UPROPERTY(EditAnywhere, BlueprintReadWrite, Category=Trail)
+	bool	bLimitStretch;
+
+	/** How quickly we 'relax' the bones to their animated positions. */
+	UPROPERTY(EditAnywhere, BlueprintReadWrite, Category=Trail)
+	float	TrailRelaxation;
+
+	/** If bLimitStretch is true, this indicates how long a bone can stretch beyond its length in the ref-pose. */
+	UPROPERTY(EditAnywhere, BlueprintReadWrite, Category=Trail)
+	float	StretchLimit;
+
+	/** 'Fake' velocity applied to bones. */
+	UPROPERTY(EditAnywhere, BlueprintReadWrite, Category=Trail)
+	FVector	FakeVelocity;
+
+	/** Whether 'fake' velocity should be applied in actor or world space. */
+	UPROPERTY(EditAnywhere, BlueprintReadWrite, Category=Trail)
+	bool	bActorSpaceFakeVel;
+
+	/** Internal use - we need the timestep to do the relaxation in CalculateNewBoneTransforms. */
+	float	ThisTimstep;
+
+	/** Did we have a non-zero ControlStrength last frame. */
+	bool	bHadValidStrength;
+
+	/** Component-space locations of the bones from last frame. Each frame these are moved towards their 'animated' locations. */
+	TArray<FVector>	TrailBoneLocations;
+
+	/** LocalToWorld used last frame, used for building transform between frames. */
+	FMatrix		OldLocalToWorld;
+
+
+	FAnimNode_Trail();
+
+	// FAnimNode_Base interface
+	virtual void Update(const FAnimationUpdateContext& Context) override;
+	virtual void GatherDebugData(FNodeDebugData& DebugData) override;
+	// End of FAnimNode_Base interface
+
+	// FAnimNode_SkeletalControlBase interface
+	virtual void EvaluateBoneTransforms(USkeletalMeshComponent* SkelComp, const FBoneContainer & RequiredBones, FA2CSPose& MeshBases, TArray<FBoneTransform>& OutBoneTransforms) override;
+	virtual bool IsValidToEvaluate(const USkeleton * Skeleton, const FBoneContainer & RequiredBones) override;
+	// End of FAnimNode_SkeletalControlBase interface
+
+private:
+	// FAnimNode_SkeletalControlBase interface
+	virtual void InitializeBoneReferences(const FBoneContainer & RequiredBones) override;
+	// End of FAnimNode_SkeletalControlBase interface
+
+	FVector GetAlignVector(EAxis::Type AxisOption, bool bInvert);
+};
\ No newline at end of file
diff --git a/Engine/Source/Runtime/Engine/Private/Animation/AnimNode_Trail.cpp b/Engine/Source/Runtime/Engine/Private/Animation/AnimNode_Trail.cpp
new file mode 100644
--- /dev/null
+++ b/Engine/Source/Runtime/Engine/Private/Animation/AnimNode_Trail.cpp
@@ -0,0 +1,234 @@
+// Copyright 1998-2014 Epic Games, Inc. All Rights Reserved.
+
+#include "EnginePrivate.h"
+#include "Animation/BoneControllers/AnimNode_Trail.h"
+
+/////////////////////////////////////////////////////
+// FAnimNode_Trail
+
+FAnimNode_Trail::FAnimNode_Trail()
+	: FakeVelocity(FVector::ZeroVector)
+	, ChainLength(2)
+	, ChainBoneAxis(EAxis::X)
+	, bInvertChainBoneAxis(false)
+	, bLimitStretch(false)
+	, bActorSpaceFakeVel(false)
+	, bHadValidStrength(false)
+	, TrailRelaxation(10)
+	, StretchLimit(0)
+{
+}
+
+void FAnimNode_Trail::Update(const FAnimationUpdateContext& Context)
+{
+	FAnimNode_SkeletalControlBase::Update(Context);
+
+	ThisTimstep = Context.GetDeltaTime();
+}
+
+void FAnimNode_Trail::GatherDebugData(FNodeDebugData& DebugData)
+{
+	FString DebugLine = DebugData.GetNodeName(this);
+
+	DebugLine += "(";
+	AddDebugNodeData(DebugLine);
+	DebugLine += FString::Printf(TEXT(" Active: %s)"), *TrailBone.BoneName.ToString());
+	DebugData.AddDebugItem(DebugLine);
+
+	ComponentPose.GatherDebugData(DebugData);
+}
+
+void FAnimNode_Trail::EvaluateBoneTransforms(USkeletalMeshComponent* SkelComp, const FBoneContainer & RequiredBones, FA2CSPose& MeshBases, TArray<FBoneTransform>& OutBoneTransforms)
+{
+	check(OutBoneTransforms.Num() == 0);
+
+	if( ChainLength < 2 )
+	{
+		return;
+	}
+
+	// The incoming BoneIndex is the 'end' of the spline chain. We need to find the 'start' by walking SplineLength bones up hierarchy.
+	// Fail if we walk past the root bone.
+
+	int32 WalkBoneIndex = TrailBone.BoneIndex;
+
+	TArray<int32> ChainBoneIndices;
+	ChainBoneIndices.AddZeroed(ChainLength);
+
+	ChainBoneIndices[ChainLength - 1] = WalkBoneIndex;
+
+	for (int32 i = 1; i < ChainLength; i++)
+	{
+		// returns to avoid a crash
+		// @TODO : shows an error message why failed
+		if (WalkBoneIndex == 0)
+		{
+			return;
+		}
+
+		// Get parent bone.
+		WalkBoneIndex = RequiredBones.GetParentBoneIndex(WalkBoneIndex);
+
+		//Insert indices at the start of array, so that parents are before children in the array.
+		int32 TransformIndex = ChainLength - (i + 1);
+		//ChainBoneIndicesTemp.Add(WalkBoneIndex);
+		ChainBoneIndices[TransformIndex] = WalkBoneIndex;
+	}
+
+	OutBoneTransforms.AddZeroed(ChainLength);
+
+	// If we have >0 this frame, but didn't last time, record positions of all the bones.
+	// Also do this if number has changed or array is zero.
+	bool bHasValidStrength = (Alpha > 0.f);
+	if(TrailBoneLocations.Num() != ChainLength || (bHasValidStrength && !bHadValidStrength))
+	{
+		TrailBoneLocations.Empty();
+		TrailBoneLocations.AddZeroed(ChainLength);
+
+		for(int32 i=0; i<ChainBoneIndices.Num(); i++)
+		{
+			int32 ChildIndex = ChainBoneIndices[i];
+			FTransform ChainTransform = MeshBases.GetComponentSpaceTransform(ChildIndex);
+			TrailBoneLocations[i] = ChainTransform.GetTranslation();
+		}
+		OldLocalToWorld = SkelComp->GetTransformMatrix();
+	}
+	bHadValidStrength = bHasValidStrength;
+
+	// transform between last frame and now.
+	FMatrix OldToNewTM = OldLocalToWorld * SkelComp->GetTransformMatrix().Inverse();
+
+	// Add fake velocity if present to all but root bone
+	if(!FakeVelocity.IsZero())
+	{
+		FVector FakeMovement = -FakeVelocity * ThisTimstep;
+
+		if (bActorSpaceFakeVel && SkelComp->GetOwner())
+		{
+			const FTransform BoneToWorld(SkelComp->GetOwner()->GetActorRotation(), SkelComp->GetOwner()->GetActorLocation());
+			FakeMovement = BoneToWorld.TransformVector(FakeMovement);
+		}
+
+		FakeMovement = SkelComp->GetTransformMatrix().InverseTransformVector(FakeMovement);
+		// Then add to each bone
+		for(int32 i=1; i<TrailBoneLocations.Num(); i++)
+		{
+			TrailBoneLocations[i] += FakeMovement;
+		}
+	}
+
+	// Root bone of trail is not modified.
+	int32 RootIndex = ChainBoneIndices[0]; 
+	FTransform ChainTransform = MeshBases.GetComponentSpaceTransform(RootIndex);
+	OutBoneTransforms[0] = FBoneTransform(RootIndex, ChainTransform);
+	TrailBoneLocations[0] = ChainTransform.GetTranslation();
+
+	// Starting one below head of chain, move bones.
+	for(int32 i=1; i<ChainBoneIndices.Num(); i++)
+	{
+		// Parent bone position in component space.
+		int32 ParentIndex = ChainBoneIndices[i-1];
+		FVector ParentPos = TrailBoneLocations[i-1];
+		FVector ParentAnimPos = MeshBases.GetComponentSpaceTransform(ParentIndex).GetTranslation();
+
+		// Child bone position in component space.
+		int32 ChildIndex = ChainBoneIndices[i];
+		FVector ChildPos = OldToNewTM.TransformPosition(TrailBoneLocations[i]); // move from 'last frames component' frame to 'this frames component' frame
+		FVector ChildAnimPos = MeshBases.GetComponentSpaceTransform(ChildIndex).GetTranslation();
+
+		// Desired parent->child offset.
+		FVector TargetDelta = (ChildAnimPos - ParentAnimPos);
+
+		// Desired child position.
+		FVector ChildTarget = ParentPos + TargetDelta;
+
+		// Find vector from child to target
+		FVector Error = ChildTarget - ChildPos;
+
+		// Calculate how much to push the child towards its target
+		float Correction = FMath::Clamp<float>(ThisTimstep * TrailRelaxation, 0.f, 1.f);
+
+		// Scale correction vector and apply to get new world-space child position.
+		TrailBoneLocations[i] = ChildPos + (Error * Correction);
+
+		// If desired, prevent bones stretching too far.
+		if(bLimitStretch)
+		{
+			float RefPoseLength = TargetDelta.Size();
+			FVector CurrentDelta = TrailBoneLocations[i] - TrailBoneLocations[i-1];
+			float CurrentLength = CurrentDelta.Size();
+
+			// If we are too far - cut it back (just project towards parent particle).
+			if( (CurrentLength - RefPoseLength > StretchLimit) && CurrentLength > SMALL_NUMBER )
+			{
+				FVector CurrentDir = CurrentDelta / CurrentLength;
+				TrailBoneLocations[i] = TrailBoneLocations[i-1] + (CurrentDir * (RefPoseLength + StretchLimit));
+			}
+		}
+
+		// Modify child matrix
+		OutBoneTransforms[i] = FBoneTransform(ChildIndex, MeshBases.GetComponentSpaceTransform(ChildIndex));
+		OutBoneTransforms[i].Transform.SetTranslation(TrailBoneLocations[i]);
+
+		// Modify rotation of parent matrix to point at this one.
+
+		// Calculate the direction that parent bone is currently pointing.
+		FVector CurrentBoneDir = OutBoneTransforms[i-1].Transform.TransformVector( GetAlignVector(ChainBoneAxis, bInvertChainBoneAxis) );
+		CurrentBoneDir = CurrentBoneDir.SafeNormal(SMALL_NUMBER);
+
+		// Calculate vector from parent to child.
+		FVector NewBoneDir = FVector(OutBoneTransforms[i].Transform.GetTranslation() - OutBoneTransforms[i - 1].Transform.GetTranslation()).SafeNormal(SMALL_NUMBER);
+
+		// Calculate a quaternion that gets us from our current rotation to the desired one.
+		FQuat DeltaLookQuat = FQuat::FindBetween(CurrentBoneDir, NewBoneDir);
+		FTransform DeltaTM( DeltaLookQuat, FVector(0.f) );
+
+		// Apply to the current parent bone transform.
+		FTransform TmpMatrix = FTransform::Identity;
+		TmpMatrix.CopyRotationPart(OutBoneTransforms[i - 1].Transform);
+		TmpMatrix = TmpMatrix * DeltaTM;
+		OutBoneTransforms[i - 1].Transform.CopyRotationPart(TmpMatrix);
+	}
+
+	// For the last bone in the chain, use the rotation from the bone above it.
+	OutBoneTransforms[ChainLength - 1].Transform.CopyRotationPart(OutBoneTransforms[ChainLength - 2].Transform);
+
+	// Update OldLocalToWorld
+	OldLocalToWorld = SkelComp->GetTransformMatrix();
+}
+
+bool FAnimNode_Trail::IsValidToEvaluate(const USkeleton * Skeleton, const FBoneContainer & RequiredBones) 
+{
+	// if bones are valid
+	return (TrailBone.IsValid(RequiredBones));
+}
+
+void FAnimNode_Trail::InitializeBoneReferences(const FBoneContainer & RequiredBones) 
+{
+	TrailBone.Initialize(RequiredBones);
+}
+
+FVector FAnimNode_Trail::GetAlignVector(EAxis::Type AxisOption, bool bInvert)
+{
+	FVector AxisDir;
+
+	if (AxisOption == EAxis::X)
+	{
+		AxisDir = FVector(1, 0, 0);
+	}
+	else if (AxisOption == EAxis::Y)
+	{
+		AxisDir = FVector(0, 1, 0);
+	}
+	else
+	{
+		AxisDir = FVector(0, 0, 1);
+	}
+
+	if (bInvert)
+	{
+		AxisDir *= -1.f;
+	}
+
+	return AxisDir;
+}
