Fix saver error after moving landscape components to streaming levels when there is foliage with non-shared settings placed.
TTP#343208

diff --git a/Engine/Source/Runtime/Engine/Private/InstancedFoliage.cpp b/Engine/Source/Runtime/Engine/Private/InstancedFoliage.cpp
--- a/Engine/Source/Runtime/Engine/Private/InstancedFoliage.cpp
+++ b/Engine/Source/Runtime/Engine/Private/InstancedFoliage.cpp
@@ -963,16 +963,23 @@ void AInstancedFoliageActor::MoveInstancesForComponentToCurrentLevel(UActorCompo
 	for (auto& MeshPair : FoliageMeshes)
 	{
 		FFoliageMeshInfo& MeshInfo = *MeshPair.Value;
+		UFoliageType* FoliageType = MeshPair.Key;
 
+		// Duplicate the foliage type if it's not shared
+		if (FoliageType->GetOutermost() == GetOutermost())
+		{
+			FoliageType = (UFoliageType*)StaticDuplicateObject(FoliageType, NewIFA, nullptr, RF_AllFlags & ~(RF_Standalone | RF_Public));
+		}
+		
 		const FFoliageComponentHashInfo* ComponentHashInfo = MeshInfo.ComponentHash.Find(InComponent);
 		if (ComponentHashInfo)
 		{
-			FFoliageMeshInfo* NewMeshInfo = NewIFA->FindOrAddMesh(MeshPair.Key);
+			FFoliageMeshInfo* NewMeshInfo = NewIFA->FindOrAddMesh(FoliageType);
 
 			// Add the foliage to the new level
 			for (int32 InstanceIndex : ComponentHashInfo->Instances)
 			{
-				NewMeshInfo->AddInstance(NewIFA, MeshPair.Key, MeshInfo.Instances[InstanceIndex]);
+				NewMeshInfo->AddInstance(NewIFA, FoliageType, MeshInfo.Instances[InstanceIndex]);
 			}
 
 			// Remove from old level
@@ -993,6 +1000,13 @@ void AInstancedFoliageActor::MoveInstancesToNewComponent(UPrimitiveComponent* In
 	for (auto& MeshPair : FoliageMeshes)
 	{
 		FFoliageMeshInfo& MeshInfo = *MeshPair.Value;
+		UFoliageType* FoliageType = MeshPair.Key;
+
+		// Duplicate the foliage type if it's not shared
+		if (FoliageType->GetOutermost() == GetOutermost())
+		{
+			FoliageType = (UFoliageType*)StaticDuplicateObject(FoliageType, NewIFA, nullptr, RF_AllFlags & ~(RF_Standalone | RF_Public));
+		}
 
 		const FFoliageComponentHashInfo* ComponentHashInfo = MeshInfo.ComponentHash.Find(InOldComponent);
 		if (ComponentHashInfo)
@@ -1015,7 +1029,7 @@ void AInstancedFoliageActor::MoveInstancesToNewComponent(UPrimitiveComponent* In
 			}
 			else
 			{
-				FFoliageMeshInfo* NewMeshInfo = NewIFA->FindOrAddMesh(MeshPair.Key);
+				FFoliageMeshInfo* NewMeshInfo = NewIFA->FindOrAddMesh(FoliageType);
 
 				// Add the foliage to the new level
 				for (int32 InstanceIndex : ComponentHashInfo->Instances)
@@ -1023,7 +1037,7 @@ void AInstancedFoliageActor::MoveInstancesToNewComponent(UPrimitiveComponent* In
 					FFoliageInstance NewInstance = MeshInfo.Instances[InstanceIndex];
 					NewInstance.Base = InNewComponent;
 					NewInstance.ClusterIndex = INDEX_NONE;
-					NewMeshInfo->AddInstance(NewIFA, MeshPair.Key, NewInstance);
+					NewMeshInfo->AddInstance(NewIFA, FoliageType, NewInstance);
 				}
 
 				// Remove from old level
