Made navigation generators public and expandable in external modules #UE4

diff --git a/Engine/Source/Runtime/Engine/Private/AI/Navigation/NavCollision.cpp b/Engine/Source/Runtime/Engine/Private/AI/Navigation/NavCollision.cpp
--- a/Engine/Source/Runtime/Engine/Private/AI/Navigation/NavCollision.cpp
+++ b/Engine/Source/Runtime/Engine/Private/AI/Navigation/NavCollision.cpp
@@ -4,7 +4,7 @@
 #include "AI/Navigation/NavCollision.h"
 #include "AI/NavigationSystemHelpers.h"
 #include "NavigationModifier.h"
-#include "../Private/AI/Navigation/RecastNavMeshGenerator.h"
+#include "AI/Navigation/RecastNavMeshGenerator.h"
 #include "DerivedDataPluginInterface.h"
 #include "DerivedDataCacheInterface.h"
 #include "TargetPlatform.h"
diff --git a/Engine/Source/Runtime/Engine/Private/AI/Navigation/RecastNavMeshGenerator.h b/Engine/Source/Runtime/Engine/Private/AI/Navigation/RecastNavMeshGenerator.h
deleted file mode 100644
diff --git a/Engine/Source/Runtime/Engine/Public/AI/NavDataGenerator.h b/Engine/Source/Runtime/Engine/Public/AI/NavDataGenerator.h
--- a/Engine/Source/Runtime/Engine/Public/AI/NavDataGenerator.h
+++ b/Engine/Source/Runtime/Engine/Public/AI/NavDataGenerator.h
@@ -4,7 +4,7 @@
 
 struct FNavigationDirtyArea;
 
-class FNavDataGenerator : public TSharedFromThis<FNavDataGenerator, ESPMode::ThreadSafe>
+class ENGINE_API FNavDataGenerator : public TSharedFromThis<FNavDataGenerator, ESPMode::ThreadSafe>
 {
 public:
 	virtual ~FNavDataGenerator() {}
diff --git a/Engine/Source/Runtime/Engine/Public/AI/Navigation/RecastNavMeshGenerator.h b/Engine/Source/Runtime/Engine/Public/AI/Navigation/RecastNavMeshGenerator.h
new file mode 100644
--- /dev/null
+++ b/Engine/Source/Runtime/Engine/Public/AI/Navigation/RecastNavMeshGenerator.h
@@ -0,0 +1,586 @@
+// Copyright 1998-2014 Epic Games, Inc. All Rights Reserved.
+#pragma once 
+
+#if WITH_RECAST
+
+#include "Recast.h"
+#include "DetourNavMesh.h"
+#include "AI/NavDataGenerator.h"
+#include "AI/NavigationModifier.h"
+
+#define MAX_VERTS_PER_POLY	6
+
+struct FNavigationRelevantData;
+class FNavigationOctree;
+
+struct FRecastBuildConfig : public rcConfig
+{
+	/** controls whether voxel filterring will be applied (via FRecastTileGenerator::ApplyVoxelFilter) */
+	uint32 bPerformVoxelFiltering:1;
+	/** generate detailed mesh (additional tessellation to match heights of geometry) */
+	uint32 bGenerateDetailedMesh:1;
+	/** generate BV tree (space partitioning for queries) */
+	uint32 bGenerateBVTree:1;
+
+	/** region partitioning method used by tile cache */
+	int32 TileCachePartitionType;
+	/** chunk size for ChunkyMonotone partitioning */
+	int32 TileCacheChunkSize;
+
+	int32 PolyMaxHeight;
+	/** indicates what's the limit of navmesh polygons per tile. This value is calculated from other
+	 *	factors - DO NOT SET IT TO ARBITRARY VALUE */
+	int32 MaxPolysPerTile;
+
+	/** Actual agent height (in uu)*/
+	float AgentHeight;
+	/** Actual agent climb (in uu)*/
+	float AgentMaxClimb;
+	/** Actual agent radius (in uu)*/
+	float AgentRadius;
+	/** Agent index for filtering links */
+	int32 AgentIndex;
+
+	FRecastBuildConfig()
+	{
+		Reset();
+	}
+
+	void Reset()
+	{
+		FMemory::MemZero(*this);
+		bPerformVoxelFiltering = true;
+		bGenerateDetailedMesh = true;
+		bGenerateBVTree = true;
+		PolyMaxHeight = 10;
+		MaxPolysPerTile = -1;
+		AgentIndex = 0;
+	}
+};
+
+struct FRecastVoxelCache
+{
+	struct FTileInfo
+	{
+		int16 TileX;
+		int16 TileY;
+		int32 NumSpans;
+		FTileInfo* NextTile;
+		rcSpanCache* SpanData;
+	};
+
+	int32 NumTiles;
+
+	/** tile info */
+	FTileInfo* Tiles;
+
+	FRecastVoxelCache() {}
+	FRecastVoxelCache(const uint8* Memory);
+};
+
+struct FRecastGeometryCache
+{
+	struct FHeader
+	{
+		int32 NumVerts;
+		int32 NumFaces;
+		struct FWalkableSlopeOverride SlopeOverride;
+	};
+
+	FHeader Header;
+
+	/** recast coords of vertices (size: NumVerts * 3) */
+	float* Verts;
+
+	/** vert indices for triangles (size: NumFaces * 3) */
+	int32* Indices;
+
+	FRecastGeometryCache() {}
+	FRecastGeometryCache(const uint8* Memory);
+};
+
+struct FRecastTileDirtyState
+{
+	/** dynamic array used when DirtyLayers is not enough */
+	TArray<uint8> FallbackDirtyLayers;
+	/** bitfields for marking layers to rebuild */
+	uint8 DirtyLayers[7];
+
+	/** set when any layer needs to be rebuild */
+	uint8 bRebuildLayers : 1;
+	/** set when all layers needs to be rebuild */
+	uint8 bRebuildAllLayers : 1;
+	/** set when geometry needs to be rebuild */
+	uint8 bRebuildGeometry : 1;
+
+	FRecastTileDirtyState() : bRebuildLayers(false), bRebuildAllLayers(false), bRebuildGeometry(false) { FMemory::Memzero(DirtyLayers, sizeof(DirtyLayers)); }
+	FRecastTileDirtyState(const class FRecastTileGenerator* DirtyGenerator);
+	
+	void Append(const FRecastTileDirtyState& Other);
+	void Clear();
+
+	void MarkDirtyLayer(int32 LayerIdx);
+	bool HasDirtyLayer(int32 LayerIdx) const;
+};
+
+struct FRecastDirtyGenerator
+{
+	class FRecastTileGenerator* Generator;
+	struct FRecastTileDirtyState State;
+
+	FRecastDirtyGenerator() : Generator(NULL) {}
+	FRecastDirtyGenerator(class FRecastTileGenerator* InGenerator) : Generator(InGenerator), State(InGenerator) {}
+};
+
+/**
+ * Class handling generation of a single tile, caching data that can speed up subsequent tile generations
+ */
+class ENGINE_API FRecastTileGenerator
+{
+public:
+	FRecastTileGenerator();
+	~FRecastTileGenerator();
+
+	// sets up tile generator instance with initial data
+	void Init(class FRecastNavMeshGenerator* ParentGenerator, const int32 X, const int32 Y, const float TileBmin[3], const float TileBmax[3], const TNavStatArray<FBox>& BoundingBoxes);
+
+	// this is the only legit way of triggering navmesh tile generation. 
+	void TriggerAsyncBuild();
+
+	/** Does the actual tile generations. 
+	 *	@note always trigger tile generation only via TriggerAsyncBuild. This is a worker function
+	 *	@return true if new tile navigation data has been generated and is ready to be added to navmesh instance, 
+	 *	false if failed or no need to generate (still valid)
+	 */
+	bool GenerateTile();
+
+	// free allocation of navmesh data
+	void ClearNavigationData() { NavigationData.Reset(); }
+	// copy generated navmesh data for caller and free allocation
+	void TransferNavigationData(TArray<FNavMeshTileData>& TileLayers) { TileLayers = NavigationData; ClearNavigationData(); }
+
+	FORCEINLINE int32 GetTileX() const { return TileX; }
+	FORCEINLINE int32 GetTileY() const { return TileY; }
+	FORCEINLINE int32 GetId() const { return TileId; }
+	FORCEINLINE int32 GetVersion() const { return Version; }
+
+	FORCEINLINE const FBox& GetUnrealBB() const { return TileBB; }
+	
+	void SetDirty(const FNavigationDirtyArea& DirtyArea, const FBox& AreaBounds);
+	FORCEINLINE void GetDirtyState(FRecastTileDirtyState& Data) const { Data = DirtyState; }
+	FORCEINLINE void SetDirtyState(const FRecastTileDirtyState& Data) { DirtyState.Append(Data); }
+	FORCEINLINE bool HasDirtyGeometry() const { return DirtyState.bRebuildGeometry; }
+	FORCEINLINE bool HasDirtyLayers() const { return DirtyState.bRebuildLayers; }
+	FORCEINLINE bool IsDirty() const { return DirtyState.bRebuildGeometry || DirtyState.bRebuildLayers; }
+	FORCEINLINE bool IsRebuildingGeometry() const { return GeneratingState.bRebuildGeometry; }
+	FORCEINLINE bool IsRebuildingLayers() const { return GeneratingState.bRebuildLayers; }
+	FORCEINLINE bool ShouldBeBuilt() const { return bOutsideOfInclusionBounds == 0; }
+	FORCEINLINE bool IsBeingRebuild() const { return bBeingRebuild; }
+	FORCEINLINE bool IsPendingRebuild() const { return bRebuildPending; }
+	FORCEINLINE bool IsAsyncBuildInProgress() const { return bAsyncBuildInProgress; }
+
+	/** Marks start of rebuilding process. */
+	void InitiateRebuild();
+	/** Tile remains dirty and will be built next time around */
+	void AbortRebuild();
+	/** called when tile generator prepared for generation signals there's actually
+	 *	no point in doing so, for example when ShouldBeBuilt() == false */
+	void AbadonGeneration();
+	/** finish rebuild process and update dirty state */
+	void FinishRebuild();
+	
+	/** mark start of async worker */
+	void StartAsyncBuild();
+	/** mark finish of async worker */
+	void FinishAsyncBuild();
+
+	/** mark pending rebuild: generator has all data and waits in queue for free worker thread */
+	void MarkPendingRebuild();
+
+	void AppendModifier(const FCompositeNavModifier& Modifier, bool bStatic);
+	/** Appends specified geometry to tile's geometry */
+	void AppendGeometry(const TNavStatArray<uint8>& RawCollisionCache);
+	void AppendGeometry(const TNavStatArray<FVector>& Verts, const TNavStatArray<int32>& Faces);
+	void AppendVoxels(rcSpanCache* SpanData, int32 NumSpans);
+	
+	uint32 GetUsedMemCount() const;
+	static int32 GetStaticVoxelSpanCount() { return StaticGeomSpans.Num(); }
+
+#if !(UE_BUILD_SHIPPING || UE_BUILD_TEST)
+	float GetLastBuildTimeStamp() const { return LastBuildTimeStamp; }
+#endif // !(UE_BUILD_SHIPPING || UE_BUILD_TEST)
+
+	void ClearGeometry();
+	void ClearModifiers();
+	static void ClearStaticData();
+
+	/** prepare voxel cache from collision data */
+	void PrepareVoxelCache(const TNavStatArray<uint8>& RawCollisionCache, TNavStatArray<rcSpanCache>& SpanData);
+	bool HasVoxelCache(const TNavStatArray<uint8>& RawVoxelCache, rcSpanCache*& CachedVoxels, int32& NumCachedVoxels) const;
+	void AddVoxelCache(TNavStatArray<uint8>& RawVoxelCache, const rcSpanCache* CachedVoxels, const int32 NumCachedVoxels) const;
+
+protected:
+
+	uint32 bInitialized : 1;
+	uint32 bOutsideOfInclusionBounds : 1;
+	uint32 bFullyEncapsulatedByInclusionBounds : 1;
+	uint32 bBeingRebuild : 1;
+	uint32 bRebuildPending : 1;
+	uint32 bAsyncBuildInProgress : 1;
+
+	/** dirty state exposed for navigation system */
+	FRecastTileDirtyState DirtyState;
+	/** dirty state used in active processing */
+	FRecastTileDirtyState GeneratingState;
+
+	int32 TileX;
+	int32 TileY;
+	int32 TileId;
+	uint32 Version;
+
+	/** bounding box minimum, Recast coords */
+	float BMin[3];
+	/** bounding box maximum, Recast coords */
+	float BMax[3];
+
+	/** Tile's bounding box, Unreal coords */
+	FBox TileBB;
+	/** Layer bounds and dirty flags, Unreal coords */
+	TArray<FBox> LayerBB;
+
+	/** NavMesh generator owning this tile generator */
+	ANavigationData::FNavDataGeneratorWeakPtr NavMeshGenerator;
+	/** Additional config */
+	TSharedPtr<struct FRecastNavMeshCachedData, ESPMode::ThreadSafe> AdditionalCachedData;
+
+	// generated tiles
+	TArray<FNavMeshTileData> CompressedLayers;
+	TArray<FNavMeshTileData> NavigationData;
+
+	// tile's geometry: voxel cache (only for synchronous rebuilds)
+	float RasterizationPadding;
+	int32 WalkableClimbVX;
+	float WalkableSlopeCos;
+	static TNavStatArray<rcSpanCache> StaticGeomSpans;
+
+	// tile's geometry: without voxel cache
+	TNavStatArray<float> GeomCoords;
+	TNavStatArray<int32> GeomIndices;
+	static TNavStatArray<float> StaticGeomCoords;
+	static TNavStatArray<int32> StaticGeomIndices;
+
+	TNavStatArray<FBox> InclusionBounds;
+	// areas used for creating compressed layers: static zones
+	TArray<FAreaNavModifier> StaticAreas;
+	// areas used for creating navigation data: obstacles
+	TArray<FAreaNavModifier> DynamicAreas;
+	// navigation links
+	TArray<FSimpleLinkNavModifier> OffmeshLinks;
+
+	static TArray<FAreaNavModifier> StaticStaticAreas;
+	static TArray<FAreaNavModifier> StaticDynamicAreas;
+	static TArray<FSimpleLinkNavModifier> StaticOffmeshLinks;
+
+	FCriticalSection GenerationLock;
+
+	typedef FAutoDeleteAsyncTask<class FAsyncNavTileBuildWorker> FAsyncNavTileBuildTask;
+
+#if !(UE_BUILD_SHIPPING || UE_BUILD_TEST)
+	// some statistics - not shippable
+	float LastBuildTimeCost;
+	float LastBuildTimeStamp;
+#endif // !(UE_BUILD_SHIPPING || UE_BUILD_TEST)
+
+	/** builds CompressedLayers array (geometry + modifiers) */
+	bool GenerateCompressedLayers(class FNavMeshBuildContext* BuildContext, FRecastNavMeshGenerator* Generator);
+
+	/** builds NavigationData array (layers + obstacles) */
+	bool GenerateNavigationData(class FNavMeshBuildContext* BuildContext, FRecastNavMeshGenerator* Generator);
+
+	void ApplyVoxelFilter(struct rcHeightfield* SolidHF, float WalkableRadius);
+
+	/** apply areas from StaticAreas to heightfield */
+	void MarkStaticAreas(class FNavMeshBuildContext* BuildContext, struct rcCompactHeightfield& CompactHF, const struct FRecastBuildConfig& TileConfig);
+
+	/** apply areas from DynamicAreas to layer */
+	void MarkDynamicAreas(struct dtTileCacheLayer& Layer, const struct FRecastBuildConfig& TileConfig);
+};
+
+struct FNavMeshDirtInfo
+{	
+	FBox Box;
+	const AActor* Actor;
+
+	FNavMeshDirtInfo(const FBox& InBox, const AActor* InActor = 0) : Box(InBox), Actor(InActor) {}
+};
+
+struct FNavMeshGenerationResult
+{
+	dtTileRef OldTileRef;
+	uint8* OldRawNavData;
+	uint32 TileIndex;
+	FNavMeshTileData NewNavData;
+
+	FNavMeshGenerationResult() : OldTileRef(INVALID_NAVNODEREF), OldRawNavData(NULL) {}
+};
+
+/**
+ * Class that handles generation of the whole Recast-based navmesh.
+ */
+class FRecastNavMeshGenerator : public FNavDataGenerator
+{
+public:
+	FRecastNavMeshGenerator(class ARecastNavMesh* InDestNavMesh);
+	virtual ~FRecastNavMeshGenerator();
+
+private:
+	/** Prevent copying. */
+	FRecastNavMeshGenerator(FRecastNavMeshGenerator const& NoCopy) { check(0); };
+	FRecastNavMeshGenerator& operator=(FRecastNavMeshGenerator const& NoCopy) { check(0); return *this; }
+
+public:
+
+	// Triggers navmesh building process
+	virtual bool Generate() override;
+
+	/** Asks generator to update navigation affected by DirtyAreas */
+	virtual void RebuildDirtyAreas(const TArray<FNavigationDirtyArea>& DirtyAreas) override;
+
+	virtual void TriggerGeneration() override;
+
+	virtual bool IsBuildInProgress(bool bCheckDirtyToo = false) const override;
+
+	bool AreAnyTilesBeingBuilt(bool bCheckDirtyToo = false) const;
+	bool IsAsyncBuildInProgress() const;
+
+	/** Checks if a given tile is being build, is marked dirty, or has just finished building
+	 *	If TileIndex is out of bounds function returns false (i.e. not being built) */
+	bool IsTileFresh(int32 X, int32 Y, float FreshnessTime = FRecastNavMeshGenerator::DefaultFreshness) const;
+
+	virtual void OnWorldInitDone(bool bAllowedToRebuild) override;
+
+	/** Moves data from this generator to it's destination navmesh object. 
+	 *	@return fail if DestNavMesh is not valid in some way, true otherwise */
+	bool TransferGeneratedData();
+	
+	FORCEINLINE int32 GetTileIdAt(int32 X, int32 Y) const { return Y * TilesWidth + X; }
+	FORCEINLINE uint32 GetVersion() const { return Version; }
+
+	/** Retrieves FCriticalSection used when adding tiles to navmesh instance, so that it can be locked
+	 *	while retrieving data for drawing for example. 
+	 *	@todo this is awkward way of doing this, but will be removed once FRecastNavMeshGenerator instance is owned by 
+	 *		FPImplRecastNavMesh, not ARecastNavMesh. Will be done soon. */
+	FCriticalSection* GetNavMeshModificationLock() const { return &TileAddingLock; }
+
+	const ARecastNavMesh* GetOwner() const { return DestNavMesh.Get(); }
+	
+	/** 
+	 *	@param Actor is a reference to make callee responsible for assuring it's valid
+	 */
+	static void ExportActorGeometry(AActor& Actor, FNavigationRelevantData& Data);
+	static void ExportComponentGeometry(UActorComponent& Component, FNavigationRelevantData& Data);
+	static void ExportVertexSoupGeometry(const TArray<FVector>& Verts, FNavigationRelevantData& Data);
+
+	static void ExportRigidBodyGeometry(UBodySetup& BodySetup, TNavStatArray<FVector>& OutVertexBuffer, TNavStatArray<int32>& OutIndexBuffer, const FTransform& LocalToWorld = FTransform::Identity);
+	static void ExportRigidBodyGeometry(UBodySetup& BodySetup, TNavStatArray<FVector>& OutTriMeshVertexBuffer, TNavStatArray<int32>& OutTriMeshIndexBuffer, TNavStatArray<FVector>& OutConvexVertexBuffer, TNavStatArray<int32>& OutConvexIndexBuffer, TNavStatArray<int32>& OutShapeBuffer, const FTransform& LocalToWorld = FTransform::Identity);
+
+	/** update workers to (re)generate new tiles */
+	void UpdateTileGenerationWorkers(int32 TileId);
+
+	//--- accessors --- //
+	FORCEINLINE class UWorld* GetWorld() const { return DestNavMesh.IsValid() ? DestNavMesh->GetWorld() : NULL; }
+
+	class FNavMeshBuildContext* GetBuildContext()	{ return BuildContext; }
+	const FRecastBuildConfig& GetConfig() const { return Config; }
+	TSharedPtr<struct FRecastNavMeshCachedData, ESPMode::ThreadSafe> GetAdditionalCachedData() { return AdditionalCachedData; }
+private:
+	// Performs initial setup of member variables so that generator is ready to do its thing from this point on
+	void Init();
+
+	bool IsBuildingLocked() const;
+
+	/** @NOTE may produce false-negatives if outside of game thread (due to GC) */
+	FORCEINLINE bool ShouldContinueBuilding() const
+	{
+		return DestNavMesh.IsValid()
+			&& DestNavMesh != NULL
+			// using this construct rather than "DestNavMesh->GetWorld()" to 
+			// avoid unreachability checks. Note that this will be called
+			// outside of game thread and may stumble upon GC, when everything is
+			// unreachanble, while here we just care about "outer systems' existance"
+			&& Cast<ULevel>(DestNavMesh->GetOuter())
+			&& ((ULevel*)(DestNavMesh->GetOuter()))->OwningWorld != NULL
+			&& ((ULevel*)(DestNavMesh->GetOuter()))->OwningWorld->GetNavigationSystem() != NULL;
+	}
+
+	// Loads navmesh generation config
+	// Params:
+	//		cellSize - (in) size of voxel 2D grid
+	//		cellHeight - (in) height of voxel cell
+	//		agentMinHeight - (in) min walkable height
+	//		agentMaxHeight - (in) max walkable height
+	//		agentMaxSlope - (in) max angle of walkable slope (degrees)
+	//		agentMaxClimb - (in) max height of walkable step
+	void SetUpGeneration(float CellSize, float CellHeight, float AgentMinHeight, float AgentMaxHeight, float AgentMaxSlope, float AgentMaxClimb, float AgentRadius);
+
+	// Generates tiled navmesh. NavMesh tile size defined by FRecastNavMeshGenerator::TileSize
+	// @todo considering passing TileSize as a parameter to this function
+	bool GenerateTiledNavMesh();
+
+public:
+	void GenerateTile(const int32 TileId, const uint32 TileGeneratorVersion);
+	bool AddTile(FRecastTileGenerator* TileGenerator, ARecastNavMesh* CachedRecastNavMeshInstance);
+	void RegenerateDirtyTiles();
+	bool HasDirtyTiles() const;
+
+	bool HasDirtyAreas() const { return DirtyAreas.Num() > 0; }
+	
+	FBox GrowBoundingBox(const FBox& BBox, bool bIncludeAgentHeight) const
+	{
+		const FVector BBoxGrowOffsetBoth = FVector(2.0f * Config.borderSize * Config.cs);
+		const FVector BBoxGrowOffsetMin = FVector(0, 0, bIncludeAgentHeight ? Config.AgentHeight : 0.0f);
+
+		return FBox(BBox.Min - BBoxGrowOffsetBoth - BBoxGrowOffsetMin, BBox.Max + BBoxGrowOffsetBoth);
+	}
+
+	FORCEINLINE bool HasResultsPending() const { return AsyncGenerationResultContainer.Num() > 0; }
+	void GetAsyncResultsCopy(TNavStatArray<FNavMeshGenerationResult>& Dest, bool bClearSource);
+
+	void StoreAsyncResults(TArray<FNavMeshGenerationResult> AsyncResults);
+
+#if !(UE_BUILD_SHIPPING || UE_BUILD_TEST)
+	virtual void ExportNavigationData(const FString& FileName) const;
+#endif
+private:
+	/** Instantiates dtNavMesh and configures it for tiles generation. Returns false if failed */
+	bool ConstructTiledNavMesh();
+
+	void UpdateBuilding();
+
+	void RemoveTileLayers(const int32 TileX, const int32 TileY, TArray<FNavMeshGenerationResult>& AsyncResults);
+	
+private:
+	void RefreshParentReference();
+	/** calls MarkBoxDirty on non-main thread */
+	//void NotifyBoxDirty(FBox Bounds, const AActor* Actor);
+	//void MarkBoxDirty(FBox Bounds, const AActor* Actor);
+
+	virtual void RebuildAll() override;
+
+	void RequestDirtyTilesRebuild();
+	
+	/** mark dirty area for aborted generator */
+	void MarkAbortedGenerator(const FRecastTileGenerator* TileGenerator);
+
+	/** mark tile generators using dirty areas and clear dirty areas array */
+	void MarkDirtyGenerators();
+
+	/** start tile generators marked as dirty, up to MaxActiveGenerators running at the same time */
+	void StartDirtyGenerators();
+
+	/** fill generator from navigation octree */
+	void FillGeneratorData(FRecastTileGenerator* TileGenerator, const FNavigationOctree* NavOctree);
+
+	enum EDataOwnership
+	{
+		DO_ForeignData = 0,
+		DO_OwnsData,
+	};
+	
+	void SetDetourMesh(class dtNavMesh* NewDetourMesh, EDataOwnership OwnsData = FRecastNavMeshGenerator::DO_OwnsData);
+
+public:
+
+	/** update area data */
+	void OnAreaAdded(const UClass* AreaClass, int32 AreaID);
+
+	/** Called to trigger generation at earliest convenience */
+	void RequestGeneration();
+
+	virtual void OnNavigationBuildingLocked() override;
+	virtual void OnNavigationBuildingUnlocked(bool bForce) override;
+	virtual void OnNavigationBoundsUpdated(class AVolume* Volume) override;
+
+	virtual void OnNavigationDataDestroyed(class ANavigationData* NavData) override;
+
+	//----------------------------------------------------------------------//
+	// debug
+	//----------------------------------------------------------------------//
+	virtual uint32 LogMemUsed() const override;
+
+private:
+	/** Parameters defining navmesh tiles */
+	struct dtNavMeshParams TiledMeshParams;
+	struct FRecastBuildConfig Config;
+	class dtNavMesh* DetourMesh;
+
+	int32 MaxActiveTiles;
+	int32 NumActiveTiles;
+	int32 MaxActiveGenerators;
+
+	int32 TilesWidth;
+	int32 TilesHeight;
+	int32 GridWidth;
+	int32 GridHeight;
+
+	TSharedPtr<struct FRecastNavMeshCachedData, ESPMode::ThreadSafe> AdditionalCachedData;
+
+	// Additional configurable values
+	// @todo These settings need to be tweakable per navmesh at some point
+	int32 TileSize;
+
+	/** bounding box of the whole geometry data used to generate navmesh - in recast coords*/
+	FBox RCNavBounds;
+	/** bounding box of the whole geometry data used to generate navmesh - in unreal coords*/
+	FBox UnrealNavBounds;
+
+	class FNavMeshBuildContext* BuildContext;
+
+	/** Bounding geometry definition. */
+	TNavStatArray<FBox> InclusionBounds;
+
+	TNavStatArray<FRecastTileGenerator> TileGenerators;
+
+	mutable FCriticalSection TileAddingLock;
+	mutable FCriticalSection TileGenerationLock;
+	FCriticalSection NavMeshDirtyLock;
+	FCriticalSection InitLock;
+	FCriticalSection GatheringDataLock;
+
+	TWeakObjectPtr<ARecastNavMesh> DestNavMesh;
+
+	/** The output navmesh we're generating. */
+	class FPImplRecastNavMesh* OutNavMesh;
+
+	TMapBase<const AActor*, FBox, false> ActorToAreaMap;
+	/** stores information about areas needing rebuilding */
+	TNavStatArray<FNavigationDirtyArea> DirtyAreas;
+	/** stores indices of tile generators that need rebuilding and their dirty state */
+	TMap<int32,FRecastTileDirtyState> DirtyGenerators;
+
+	/** */
+	TNavStatArray<FNavMeshGenerationResult> AsyncGenerationResultContainer;
+
+	uint32 bInitialized:1;
+	uint32 bBuildFromScratchRequested:1;
+	uint32 bRebuildDirtyTilesRequested:1;
+	/** gets set to true when RecastNavMesh owned has been destroyed and there's 
+	 *	no more point in generating navigation data */
+	uint32 bAbortAllTileGeneration:1;
+	uint32 bOwnsDetourMesh:1;
+	mutable uint32 bBuildingLocked:1;
+
+	/** Runtime generator's version, increased every time all tile generators get invalidated
+	 *	like when navmesh size changes */
+	uint32 Version;
+
+	static const float DefaultFreshness;
+
+	/** queue of working generators - size of this array depends on the number of CPU cores */
+	TArray<FRecastTileGenerator*> ActiveGenerators;
+
+	/** sorted array of waiting generators to work asap*/
+	TArray<FRecastTileGenerator*> GeneratorsQueue;
+};
+
+#endif // WITH_RECAST
