Changed world origin to a 3d location instead of 2d. Exposed Get/Set OriginLocation functions to blueprints.

diff --git a/Engine/Source/Editor/UnrealEd/Classes/Editor/EditorEngine.h b/Engine/Source/Editor/UnrealEd/Classes/Editor/EditorEngine.h
--- a/Engine/Source/Editor/UnrealEd/Classes/Editor/EditorEngine.h
+++ b/Engine/Source/Editor/UnrealEd/Classes/Editor/EditorEngine.h
@@ -2455,7 +2455,7 @@ private:
 	virtual void PostUndo (bool bSuccess);
 
 	/** Delegate callback: the world origin is going to be moved. */
-	void PreWorldOriginOffset(UWorld* InWorld, const FIntPoint& InSrcOrigin, const FIntPoint& InDstOrigin);
+	void PreWorldOriginOffset(UWorld* InWorld, FIntVector InSrcOrigin, FIntVector InDstOrigin);
 
 	/** Delegate callback for when a streaming level is added to world. */
 	void OnLevelAddedToWorld(ULevel* InLevel, UWorld* InWorld);
diff --git a/Engine/Source/Editor/UnrealEd/Private/Editor.cpp b/Engine/Source/Editor/UnrealEd/Private/Editor.cpp
--- a/Engine/Source/Editor/UnrealEd/Private/Editor.cpp
+++ b/Engine/Source/Editor/UnrealEd/Private/Editor.cpp
@@ -6215,7 +6215,7 @@ UActorFactory* UEditorEngine::FindActorFactoryByClassForActorClass( const UClass
 	return NULL;
 }
 
-void UEditorEngine::PreWorldOriginOffset(UWorld* InWorld, const FIntPoint& InSrcOrigin, const FIntPoint& InDstOrigin)
+void UEditorEngine::PreWorldOriginOffset(UWorld* InWorld, FIntVector InSrcOrigin, FIntVector InDstOrigin)
 {
 	// In case we simulating world in the editor, 
 	// we need to shift current viewport as well, 
diff --git a/Engine/Source/Editor/WorldBrowser/Private/LevelCollectionModel.h b/Engine/Source/Editor/WorldBrowser/Private/LevelCollectionModel.h
--- a/Engine/Source/Editor/WorldBrowser/Private/LevelCollectionModel.h
+++ b/Engine/Source/Editor/WorldBrowser/Private/LevelCollectionModel.h
@@ -46,6 +46,9 @@ public:
 	/** Current world size  */
 	FIntPoint GetWorldSize() const { return WorldSize; }
 	
+	/** Current world origin location on XY plane  */
+	FIntPoint GetWorldOriginLocationXY() const { return FIntPoint(GetWorld()->OriginLocation.X, GetWorld()->OriginLocation.Y); }
+	
 	/**	@return	Root list of levels in hierarchy */
 	FLevelModelList& GetRootLevelList();
 
diff --git a/Engine/Source/Editor/WorldBrowser/Private/Tiles/SWorldComposition.cpp b/Engine/Source/Editor/WorldBrowser/Private/Tiles/SWorldComposition.cpp
--- a/Engine/Source/Editor/WorldBrowser/Private/Tiles/SWorldComposition.cpp
+++ b/Engine/Source/Editor/WorldBrowser/Private/Tiles/SWorldComposition.cpp
@@ -770,14 +770,15 @@ protected:
 	}
 
 	/** Delegate callback: world origin is going to be moved. */
-	void PreWorldOriginOffset(UWorld* InWorld, const FIntPoint& InSrcOrigin, const FIntPoint& InDstOrigin)
+	void PreWorldOriginOffset(UWorld* InWorld, FIntVector InSrcOrigin, FIntVector InDstOrigin)
 	{
 		if (WorldModel->GetWorld() != InWorld)
 		{
 			return;
 		}
-				
-		RequestScrollBy(-FVector2D(InDstOrigin - InSrcOrigin));
+		
+		FIntVector Offset = InDstOrigin - InSrcOrigin;
+		RequestScrollBy(-FVector2D(Offset.X, Offset.Y));
 	}
 	
 	/** Handles new item added to data source */
@@ -1214,7 +1215,7 @@ FText SWorldComposition::GetZoomText() const
 FString SWorldComposition::GetCurrentOriginText() const
 {
 	UWorld* CurrentWorld = (TileWorldModel->IsSimulating() ? TileWorldModel->GetSimulationWorld() : TileWorldModel->GetWorld());
-	return FString::Printf(TEXT("%d, %d"), CurrentWorld->GlobalOriginOffset.X, CurrentWorld->GlobalOriginOffset.Y);
+	return FString::Printf(TEXT("%d, %d"), CurrentWorld->OriginLocation.X, CurrentWorld->OriginLocation.Y);
 }
 
 FString SWorldComposition::GetCurrentLevelText() const
diff --git a/Engine/Source/Editor/WorldBrowser/Private/Tiles/WorldTileCollectionModel.cpp b/Engine/Source/Editor/WorldBrowser/Private/Tiles/WorldTileCollectionModel.cpp
--- a/Engine/Source/Editor/WorldBrowser/Private/Tiles/WorldTileCollectionModel.cpp
+++ b/Engine/Source/Editor/WorldBrowser/Private/Tiles/WorldTileCollectionModel.cpp
@@ -989,16 +989,16 @@ void FWorldTileCollectionModel::Focus(FBox InArea, FocusStrategy InStrategy)
 	{
 	case OriginAtCenter:
 		{
-			FIntPoint NewWorldOrigin = FIntPoint(InArea.GetCenter().X, InArea.GetCenter().Y);
-			GetWorld()->SetNewWorldOrigin(NewWorldOrigin + GetWorld()->GlobalOriginOffset);
+			FIntVector OriginOffset = FIntVector(InArea.GetCenter().X, InArea.GetCenter().Y, 0);
+			GetWorld()->SetNewWorldOrigin(GetWorld()->OriginLocation + OriginOffset);
 		}
 		break;
 	
 	case EnsureEditableCentered:
 		if (!bIsEditable)
 		{
-			FIntPoint NewWorldOrigin = FIntPoint(InArea.GetCenter().X, InArea.GetCenter().Y);
-			GetWorld()->SetNewWorldOrigin(NewWorldOrigin + GetWorld()->GlobalOriginOffset);
+			FIntVector OriginOffset = FIntVector(InArea.GetCenter().X, InArea.GetCenter().Y, 0);
+			GetWorld()->SetNewWorldOrigin(GetWorld()->OriginLocation + OriginOffset);
 		}
 		break;
 
@@ -1032,14 +1032,14 @@ void FWorldTileCollectionModel::Focus(FBox InArea, FocusStrategy InStrategy)
 				NewWorldBounds.Min.Y = InArea.Max.Y - EditableAxisLength();
 			}
 
-			FIntPoint NewWorldOrigin = FIntPoint(NewWorldBounds.GetCenter().X, NewWorldBounds.GetCenter().Y);
-			GetWorld()->SetNewWorldOrigin(NewWorldOrigin + GetWorld()->GlobalOriginOffset);
+			FIntVector OriginOffset = FIntVector(NewWorldBounds.GetCenter().X, NewWorldBounds.GetCenter().Y, 0);
+			GetWorld()->SetNewWorldOrigin(GetWorld()->OriginLocation + OriginOffset);
 		}
 		break;
 	}
 }
 
-void FWorldTileCollectionModel::PreWorldOriginOffset(UWorld* InWorld, const FIntPoint& InSrcOrigin, const FIntPoint& InDstOrigin)
+void FWorldTileCollectionModel::PreWorldOriginOffset(UWorld* InWorld, FIntVector InSrcOrigin, FIntVector InDstOrigin)
 {
 	// Make sure we handle our world notifications
 	if (GetWorld() != InWorld)
@@ -1047,7 +1047,7 @@ void FWorldTileCollectionModel::PreWorldOriginOffset(UWorld* InWorld, const FInt
 		return;
 	}
 	
-	FBox NewWorldBounds = EditableWorldArea().ShiftBy(FVector(InDstOrigin-InSrcOrigin, 0));
+	FBox NewWorldBounds = EditableWorldArea().ShiftBy(FVector(InDstOrigin - InSrcOrigin));
 
 	// Shelve levels which do not fit to a new world bounds
 	for (auto It = AllLevelsList.CreateIterator(); It; ++It)
@@ -1060,7 +1060,7 @@ void FWorldTileCollectionModel::PreWorldOriginOffset(UWorld* InWorld, const FInt
 	}
 }
 
-void FWorldTileCollectionModel::PostWorldOriginOffset(UWorld* InWorld, const FIntPoint& InSrcOrigin, const FIntPoint& InDstOrigin)
+void FWorldTileCollectionModel::PostWorldOriginOffset(UWorld* InWorld, FIntVector InSrcOrigin, FIntVector InDstOrigin)
 {
 	// Make sure we handle our world notifications
 	if (GetWorld() != InWorld)
@@ -1691,7 +1691,7 @@ void FWorldTileCollectionModel::MoveWorldOrigin(const FIntPoint& InOrigin)
 		return;
 	}
 	
-	GetWorld()->SetNewWorldOrigin(InOrigin);
+	GetWorld()->SetNewWorldOrigin(FIntVector(InOrigin.X, InOrigin.Y, 0));
 	RequestUpdateAllLevels();
 }
 
@@ -1710,7 +1710,7 @@ void FWorldTileCollectionModel::ResetWorldOrigin_Executed()
 {
 	if (GetWorld()->GetWorldSettings()->bEnableWorldOriginRebasing)
 	{
-		FBox OriginArea = EditableWorldArea().ShiftBy(FVector(GetWorld()->GlobalOriginOffset, 0));
+		FBox OriginArea = EditableWorldArea().ShiftBy(FVector(GetWorld()->OriginLocation));
 		Focus(OriginArea, OriginAtCenter);
 	
 		MoveWorldOrigin(FIntPoint::ZeroValue);
@@ -1828,7 +1828,7 @@ bool FWorldTileCollectionModel::GenerateLODLevels(FLevelModelList InLevelList, i
 		LODPackage->FileName = FName(*LODLevelFileName);
 				
 		// This is current actors offset from their original position
-		FVector ActorsOffset = FVector(TileModel->GetAbsoluteLevelPosition() - CurrentWorld->GlobalOriginOffset);
+		FVector ActorsOffset = FVector(TileModel->GetAbsoluteLevelPosition() - GetWorldOriginLocationXY());
 		if (GetWorld()->WorldComposition->bTemporallyDisableOriginTracking)
 		{
 			ActorsOffset = FVector::ZeroVector;
diff --git a/Engine/Source/Editor/WorldBrowser/Private/Tiles/WorldTileCollectionModel.h b/Engine/Source/Editor/WorldBrowser/Private/Tiles/WorldTileCollectionModel.h
--- a/Engine/Source/Editor/WorldBrowser/Private/Tiles/WorldTileCollectionModel.h
+++ b/Engine/Source/Editor/WorldBrowser/Private/Tiles/WorldTileCollectionModel.h
@@ -169,10 +169,10 @@ private:
 	void PrepareToLoadLevels(FWorldTileModelList& InLevels);
 		
 	/** Delegate callback: the world origin is going to be moved. */
-	void PreWorldOriginOffset(UWorld* InWorld, const FIntPoint& InSrcOrigin, const FIntPoint& InDstOrigin);
+	void PreWorldOriginOffset(UWorld* InWorld, FIntVector InSrcOrigin, FIntVector InDstOrigin);
 	
 	/** Delegate callback: the world origin has been moved. */
-	void PostWorldOriginOffset(UWorld* InWorld, const FIntPoint& InSrcOrigin, const FIntPoint& InDstOrigin);
+	void PostWorldOriginOffset(UWorld* InWorld, FIntVector InSrcOrigin, FIntVector InDstOrigin);
 
 	/** Update list layers */
 	void PopulateLayersList();
diff --git a/Engine/Source/Editor/WorldBrowser/Private/Tiles/WorldTileModel.cpp b/Engine/Source/Editor/WorldBrowser/Private/Tiles/WorldTileModel.cpp
--- a/Engine/Source/Editor/WorldBrowser/Private/Tiles/WorldTileModel.cpp
+++ b/Engine/Source/Editor/WorldBrowser/Private/Tiles/WorldTileModel.cpp
@@ -389,7 +389,7 @@ FVector2D FWorldTileModel::GetLevelCurrentPosition() const
 		FVector2D LevelLocalPosition(TileDetails->Bounds.GetCenter());
 		FIntPoint LevelOffset = GetAbsoluteLevelPosition();
 			
-		return LevelLocalPosition + FVector2D(LevelOffset - CurrentWorld->GlobalOriginOffset); 
+		return LevelLocalPosition + FVector2D(LevelOffset - LevelCollectionModel.GetWorldOriginLocationXY()); 
 	}
 
 	return FVector2D(0, 0);
@@ -480,11 +480,10 @@ void FWorldTileModel::Update()
 			if (TileDetails->Bounds.IsValid)
 			{
 				FBox LevelWorldBounds = TileDetails->Bounds;
-				FIntPoint GlobalOriginOffset = LevelCollectionModel.GetWorld()->GlobalOriginOffset;
 				FIntPoint LevelAbolutePosition = GetAbsoluteLevelPosition();
-				FIntPoint LevelOffset = LevelAbolutePosition - GlobalOriginOffset;
+				FIntPoint LevelOffset = LevelAbolutePosition - LevelCollectionModel.GetWorldOriginLocationXY();
 
-				TileDetails->Bounds = LevelWorldBounds.ShiftBy(-FVector(LevelOffset, 0.f));
+				TileDetails->Bounds = LevelWorldBounds.ShiftBy(-FVector(LevelOffset));
 			}
 
 			OnLevelInfoUpdated();
@@ -763,7 +762,7 @@ bool FWorldTileModel::CreateAdjacentLandscapeProxy(ALandscapeProxy* SourceLandsc
 		}
 
 		// Add source level position
-		FIntPoint IntOffset = FIntPoint(ProxyOffset.X, ProxyOffset.Y) + LevelCollectionModel.GetWorld()->GlobalOriginOffset;
+		FIntPoint IntOffset = FIntPoint(ProxyOffset.X, ProxyOffset.Y) + LevelCollectionModel.GetWorldOriginLocationXY();
 
 		// Move level with landscape proxy to desired position
 		SetLevelPosition(IntOffset);
diff --git a/Engine/Source/Runtime/AIModule/Classes/AISystem.h b/Engine/Source/Runtime/AIModule/Classes/AISystem.h
--- a/Engine/Source/Runtime/AIModule/Classes/AISystem.h
+++ b/Engine/Source/Runtime/AIModule/Classes/AISystem.h
@@ -33,7 +33,7 @@ public:
 
 	// IAISystemInterface begin		
 	virtual void InitializeActorsForPlay(bool bTimeGotReset) override;
-	virtual void WorldOriginChanged(FIntPoint OldOrigin, FIntPoint NewOrigin) override;
+	virtual void WorldOriginLocationChanged(FIntVector OldOriginLocation, FIntVector NewOriginLocation) override;
 	virtual void CleanupWorld(bool bSessionEnded = true, bool bCleanupResources = true, UWorld* NewWorld = NULL) override;
 	// IAISystemInterface end
 	
diff --git a/Engine/Source/Runtime/AIModule/Private/AISystem.cpp b/Engine/Source/Runtime/AIModule/Private/AISystem.cpp
--- a/Engine/Source/Runtime/AIModule/Private/AISystem.cpp
+++ b/Engine/Source/Runtime/AIModule/Private/AISystem.cpp
@@ -33,7 +33,7 @@ void UAISystem::InitializeActorsForPlay(bool bTimeGotReset)
 
 }
 
-void UAISystem::WorldOriginChanged(FIntPoint OldOrigin, FIntPoint NewOrigin)
+void UAISystem::WorldOriginLocationChanged(FIntVector OldOriginLocation, FIntVector NewOriginLocation)
 {
 
 }
diff --git a/Engine/Source/Runtime/Core/Public/Misc/CallbackDevice.h b/Engine/Source/Runtime/Core/Public/Misc/CallbackDevice.h
--- a/Engine/Source/Runtime/Core/Public/Misc/CallbackDevice.h
+++ b/Engine/Source/Runtime/Core/Public/Misc/CallbackDevice.h
@@ -196,7 +196,7 @@ public:
 	DECLARE_DELEGATE_RetVal_OneParam( FString, FStringAssetReferenceSaving, FString const& /*SavingAssetLongPathname*/);
 	static FStringAssetReferenceSaving StringAssetReferenceSaving;
 		
-	DECLARE_MULTICAST_DELEGATE_ThreeParams(FWorldOriginOffset, class UWorld*, const FIntPoint&, const FIntPoint&);
+	DECLARE_MULTICAST_DELEGATE_ThreeParams(FWorldOriginOffset, class UWorld*, FIntVector, FIntVector);
 	/** called before world origin shifting */
 	static FWorldOriginOffset PreWorldOriginOffset;
 	/** called after world origin shifting */
diff --git a/Engine/Source/Runtime/CoreUObject/Classes/Object.h b/Engine/Source/Runtime/CoreUObject/Classes/Object.h
--- a/Engine/Source/Runtime/CoreUObject/Classes/Object.h
+++ b/Engine/Source/Runtime/CoreUObject/Classes/Object.h
@@ -322,6 +322,23 @@ struct FIntPoint
 
 };
 
+/**
+ *  An integer vector in 3D space
+ */
+
+USTRUCT(immutable, noexport)
+struct FIntVector
+{
+	UPROPERTY(EditAnywhere, BlueprintReadWrite, Category=IntVector, SaveGame)
+	int32 X;
+
+	UPROPERTY(EditAnywhere, BlueprintReadWrite, Category=IntVector, SaveGame)
+	int32 Y;
+
+	UPROPERTY(EditAnywhere, BlueprintReadWrite, Category=IntVector, SaveGame)
+	int32 Z;
+};
+
 
 // A Color.
 
diff --git a/Engine/Source/Runtime/Engine/Classes/AI/AISystemBase.h b/Engine/Source/Runtime/Engine/Classes/AI/AISystemBase.h
--- a/Engine/Source/Runtime/Engine/Classes/AI/AISystemBase.h
+++ b/Engine/Source/Runtime/Engine/Classes/AI/AISystemBase.h
@@ -20,12 +20,12 @@ class ENGINE_API UAISystemBase : public UObject
 	virtual void InitializeActorsForPlay(bool bTimeGotReset) PURE_VIRTUAL(UAISystemBase::InitializeActorsForPlay, );
 
 	/**
-	 * Event called on world origin changes
+	 * Event called on world origin location changes
 	 *
-	 * @param	OldOrigin			Previous world origin position
-	 * @param	NewOrigin			New world origin position
+	 * @param	OldOriginLocation			Previous world origin location
+	 * @param	NewOriginLocation			New world origin location
 	 */
-	virtual void WorldOriginChanged(FIntPoint OldOrigin, FIntPoint NewOrigin) PURE_VIRTUAL(UAISystemBase::WorldOriginChanged, );
+	virtual void WorldOriginLocationChanged(FIntVector OldOriginLocation, FIntVector NewOriginLocation) PURE_VIRTUAL(UAISystemBase::WorldOriginLocationChanged, );
 
 	/**
 	 * Called by UWorld::CleanupWorld.
diff --git a/Engine/Source/Runtime/Engine/Classes/Engine/LevelScriptActor.h b/Engine/Source/Runtime/Engine/Classes/Engine/LevelScriptActor.h
--- a/Engine/Source/Runtime/Engine/Classes/Engine/LevelScriptActor.h
+++ b/Engine/Source/Runtime/Engine/Classes/Engine/LevelScriptActor.h
@@ -33,13 +33,13 @@ class ENGINE_API ALevelScriptActor : public AActor
 	virtual void LevelReset();
 
 	/**
-	 * Event called on world origin changes
+	 * Event called on world origin location changes
 	 *
-	 * @param	OldOrigin			Previous world origin position
-	 * @param	NewOrigin			New world origin position
+	 * @param	OldOriginLocation	Previous world origin location
+	 * @param	NewOriginLocation	New world origin location
 	 */
 	UFUNCTION(BlueprintImplementableEvent)
-	virtual void WorldOriginChanged(FIntPoint OldOrigin, FIntPoint NewOrigin);
+	virtual void WorldOriginLocationChanged(FIntVector OldOriginLocation, FIntVector NewOriginLocation);
 	
 #if WITH_EDITOR
 	// Begin UObject Interface
diff --git a/Engine/Source/Runtime/Engine/Classes/Engine/World.h b/Engine/Source/Runtime/Engine/Classes/Engine/World.h
--- a/Engine/Source/Runtime/Engine/Classes/Engine/World.h
+++ b/Engine/Source/Runtime/Engine/Classes/Engine/World.h
@@ -866,11 +866,11 @@ public:
 	/** time at which to start pause **/
 	float PauseDelay;
 
-	/** Offset for all levels from world origin */
-	FIntPoint GlobalOriginOffset;
+	/** Current location of this world origin */
+	FIntVector OriginLocation;
 
-	/** Requested new world origin offset */
-	FIntPoint RequestedGlobalOriginOffset;
+	/** Requested new world origin location */
+	FIntVector RequestedOriginLocation;
 
 	/** Whether world origin was rebased this frame */
 	bool bOriginOffsetThisFrame;
@@ -2412,13 +2412,13 @@ public:
 	bool IsNavigationRebuilt() const;
 
 	/** Request to translate world origin to specified position on next tick */
-	void RequestNewWorldOrigin(const FIntPoint& InNewOrigin);
+	void RequestNewWorldOrigin(FIntVector InNewOriginLocation);
 	
 	/** Translate world origin to specified position  */
-	bool SetNewWorldOrigin(const FIntPoint& InNewOrigin);
+	bool SetNewWorldOrigin(FIntVector InNewOriginLocation);
 
 	/** Sets world origin at specified position and stream-in all relevant levels */
-	void NavigateTo(FIntPoint InPosition);
+	void NavigateTo(FIntVector InLocation);
 
 	/** Gets all matinee actors for the current level */
 	void GetMatineeActors( TArray<AMatineeActor*>& OutMatineeActors );
diff --git a/Engine/Source/Runtime/Engine/Classes/Kismet/GameplayStatics.h b/Engine/Source/Runtime/Engine/Classes/Kismet/GameplayStatics.h
--- a/Engine/Source/Runtime/Engine/Classes/Kismet/GameplayStatics.h
+++ b/Engine/Source/Runtime/Engine/Classes/Kismet/GameplayStatics.h
@@ -497,5 +497,12 @@ class ENGINE_API UGameplayStatics : public UBlueprintFunctionLibrary
 	/** Native version, has more options than the Blueprint version. */
 	static bool SuggestProjectileVelocity(UObject* WorldContextObject, FVector& TossVelocity, FVector StartLocation, FVector EndLocation, float TossSpeed, bool bHighArc=false, float CollisionRadius = 0.f, float OverrideGravityZ = 0, ESuggestProjVelocityTraceOption::Type TraceOption = ESuggestProjVelocityTraceOption::TraceFullPath, bool bDrawDebug=false);
 
+	/** Returns world origin current location */
+	UFUNCTION(BlueprintPure, Category="Game", meta=(HidePin="WorldContextObject", DefaultToSelf="WorldContextObject") )
+	static FIntVector GetWorldOriginLocation(UObject* WorldContextObject);
+	
+	/** Requests a new location for a world origin */
+	UFUNCTION(BlueprintCallable, Category="Game", meta=(HidePin="WorldContextObject", DefaultToSelf="WorldContextObject"))
+	static void SetWorldOriginLocation(UObject* WorldContextObject, FIntVector NewLocation);
 };
 
diff --git a/Engine/Source/Runtime/Engine/Private/CheatManager.cpp b/Engine/Source/Runtime/Engine/Private/CheatManager.cpp
--- a/Engine/Source/Runtime/Engine/Private/CheatManager.cpp
+++ b/Engine/Source/Runtime/Engine/Private/CheatManager.cpp
@@ -1078,7 +1078,7 @@ void UCheatManager::SetWorldOrigin()
 		ViewLocation = MyPlayerController->GetPawn()->GetActorLocation();
 	}
 	
-	FIntPoint NewOrigin = FIntPoint(ViewLocation.X, ViewLocation.Y) + World->GlobalOriginOffset;
+	FIntVector NewOrigin = FIntVector(ViewLocation.X, ViewLocation.Y, ViewLocation.Z) + World->OriginLocation;
 	World->RequestNewWorldOrigin(NewOrigin);
 }
 
diff --git a/Engine/Source/Runtime/Engine/Private/GameInstance.cpp b/Engine/Source/Runtime/Engine/Private/GameInstance.cpp
--- a/Engine/Source/Runtime/Engine/Private/GameInstance.cpp
+++ b/Engine/Source/Runtime/Engine/Private/GameInstance.cpp
@@ -192,7 +192,7 @@ bool UGameInstance::StartPIEGameInstance(ULocalPlayer* LocalPlayer, bool bInSimu
 		// and stream in all relevant levels around that location
 		if (PlayWorld->WorldComposition)
 		{
-			PlayWorld->NavigateTo(PlayWorld->GlobalOriginOffset);
+			PlayWorld->NavigateTo(PlayWorld->OriginLocation);
 		}
 		else
 		{
diff --git a/Engine/Source/Runtime/Engine/Private/GameplayStatics.cpp b/Engine/Source/Runtime/Engine/Private/GameplayStatics.cpp
--- a/Engine/Source/Runtime/Engine/Private/GameplayStatics.cpp
+++ b/Engine/Source/Runtime/Engine/Private/GameplayStatics.cpp
@@ -1277,4 +1277,16 @@ bool UGameplayStatics::SuggestProjectileVelocity(UObject* WorldContextObject, FV
 	}
 
 	return bFoundAValidSolution;
+}
+
+FIntVector UGameplayStatics::GetWorldOriginLocation(UObject* WorldContextObject)
+{
+	UWorld* World = GEngine->GetWorldFromContextObject(WorldContextObject);
+	return World->OriginLocation;
+}
+
+void UGameplayStatics::SetWorldOriginLocation(UObject* WorldContextObject, FIntVector NewLocation)
+{
+	UWorld* World = GEngine->GetWorldFromContextObject(WorldContextObject);
+	return World->RequestNewWorldOrigin(NewLocation);
 }
\ No newline at end of file
diff --git a/Engine/Source/Runtime/Engine/Private/LevelTick.cpp b/Engine/Source/Runtime/Engine/Private/LevelTick.cpp
--- a/Engine/Source/Runtime/Engine/Private/LevelTick.cpp
+++ b/Engine/Source/Runtime/Engine/Private/LevelTick.cpp
@@ -1061,9 +1061,9 @@ void UWorld::Tick( ELevelTick TickType, float DeltaSeconds )
 	}
 
 	// Translate world origin if requested
-	if (GlobalOriginOffset != RequestedGlobalOriginOffset)
+	if (OriginLocation != RequestedOriginLocation)
 	{
-		SetNewWorldOrigin(RequestedGlobalOriginOffset);
+		SetNewWorldOrigin(RequestedOriginLocation);
 		bOriginOffsetThisFrame = true;
 	}
 	else
diff --git a/Engine/Source/Runtime/Engine/Private/Particles/ParticleBeamModules.cpp b/Engine/Source/Runtime/Engine/Private/Particles/ParticleBeamModules.cpp
--- a/Engine/Source/Runtime/Engine/Private/Particles/ParticleBeamModules.cpp
+++ b/Engine/Source/Runtime/Engine/Private/Particles/ParticleBeamModules.cpp
@@ -1853,7 +1853,7 @@ bool UParticleModuleBeamSource::ResolveSourceData(FParticleBeam2EmitterInstance*
 				// Use the value as a world space position
 				BeamData->SourcePoint	= Source.GetValue(BeamInst->EmitterTime, BeamInst->Component);
 				// Take into account current world origin offset
-				BeamData->SourcePoint  -= FVector(BeamInst->Component->GetWorld()->GlobalOriginOffset);
+				BeamData->SourcePoint  -= FVector(BeamInst->Component->GetWorld()->OriginLocation);
 			}
 			else
 			{
@@ -2227,7 +2227,7 @@ bool UParticleModuleBeamTarget::ResolveTargetData(FParticleBeam2EmitterInstance*
 			{
 				BeamData->TargetPoint	= Target.GetValue(BeamInst->EmitterTime, BeamInst->Component);
 				// Take into account current world origin offset
-				BeamData->TargetPoint  -= FVector(BeamInst->Component->GetWorld()->GlobalOriginOffset);
+				BeamData->TargetPoint  -= FVector(BeamInst->Component->GetWorld()->OriginLocation);
 			}
 			else
 			{
diff --git a/Engine/Source/Runtime/Engine/Private/UnrealEngine.cpp b/Engine/Source/Runtime/Engine/Private/UnrealEngine.cpp
--- a/Engine/Source/Runtime/Engine/Private/UnrealEngine.cpp
+++ b/Engine/Source/Runtime/Engine/Private/UnrealEngine.cpp
@@ -8600,7 +8600,7 @@ bool UEngine::LoadMap( FWorldContext& WorldContext, FURL URL, class UPendingNetG
 	if (WorldContext.World()->WorldComposition)
 	{
 		// Set initial world origin and stream in levels
-		WorldContext.World()->NavigateTo(FIntPoint::ZeroValue);
+		WorldContext.World()->NavigateTo(FIntVector::ZeroValue);
 	}
 	else
 	{
diff --git a/Engine/Source/Runtime/Engine/Private/World.cpp b/Engine/Source/Runtime/Engine/Private/World.cpp
--- a/Engine/Source/Runtime/Engine/Private/World.cpp
+++ b/Engine/Source/Runtime/Engine/Private/World.cpp
@@ -159,8 +159,8 @@ void UWorld::Serialize( FArchive& Ar )
 	// Serialize World composition for PIE
 	if (Ar.GetPortFlags() & PPF_DuplicateForPIE)
 	{
-		Ar << GlobalOriginOffset;
-		Ar << RequestedGlobalOriginOffset;
+		Ar << OriginLocation;
+		Ar << RequestedOriginLocation;
 	}
 	
 	// UWorlds loaded/duplicated for PIE must lose RF_Public and RF_Standalone since they should not be referenced by objects in other packages and they should be GCed normally
@@ -2388,9 +2388,9 @@ void UWorld::EvaluateWorldOriginLocation(const FSceneViewFamily& ViewFamily)
 		CentroidLocation/= ViewFamily.Views.Num();
 	
 		// Request to shift world in case current view is quite far from current origin
-		if (CentroidLocation.Size2D() > HALF_WORLD_MAX1*0.5f)
+		if (CentroidLocation.Size() > HALF_WORLD_MAX1*0.5f)
 		{
-			RequestNewWorldOrigin(FIntPoint(CentroidLocation.X, CentroidLocation.Y) + GlobalOriginOffset);
+			RequestNewWorldOrigin(FIntVector(CentroidLocation.X, CentroidLocation.Y, CentroidLocation.Z) + OriginLocation);
 		}
 	}
 }
@@ -3896,14 +3896,14 @@ void UWorld::CommitMapChange()
 	}
 }
 
-void UWorld::RequestNewWorldOrigin(const FIntPoint& InNewOrigin)
+void UWorld::RequestNewWorldOrigin(FIntVector InNewOriginLocation)
 {
-	RequestedGlobalOriginOffset = InNewOrigin;
+	RequestedOriginLocation = InNewOriginLocation;
 }
 
-bool UWorld::SetNewWorldOrigin(const FIntPoint& InNewOrigin)
+bool UWorld::SetNewWorldOrigin(FIntVector InNewOriginLocation)
 {
-	if (GlobalOriginOffset == InNewOrigin) 
+	if (OriginLocation == InNewOriginLocation) 
 	{
 		return true;
 	}
@@ -3915,13 +3915,15 @@ bool UWorld::SetNewWorldOrigin(const FIntPoint& InNewOrigin)
 		return false;
 	}
 	
-	UE_LOG(LogLevel, Log, TEXT("WORLD TRANSLATION BEGIN {%d, %d} -> {%d, %d}"), GlobalOriginOffset.X, GlobalOriginOffset.Y, InNewOrigin.X, InNewOrigin.Y);
+	UE_LOG(LogLevel, Log, TEXT("WORLD TRANSLATION BEGIN {%d, %d, %d} -> {%d, %d, %d}"), 
+		OriginLocation.X, OriginLocation.Y, OriginLocation.Z, InNewOriginLocation.X, InNewOriginLocation.Y, InNewOriginLocation.Z);
+
 	const double MoveStartTime = FPlatformTime::Seconds();
 
 	// Broadcast that we going to shift world to the new origin
-	FCoreDelegates::PreWorldOriginOffset.Broadcast(this, GlobalOriginOffset, InNewOrigin);
+	FCoreDelegates::PreWorldOriginOffset.Broadcast(this, OriginLocation, InNewOriginLocation);
 
-	FVector Offset(GlobalOriginOffset - InNewOrigin, 0);
+	FVector Offset(OriginLocation - InNewOriginLocation);
 
 	// Send offset command to rendering thread
 	Scene->ApplyWorldOffset(Offset);
@@ -3965,10 +3967,10 @@ bool UWorld::SetNewWorldOrigin(const FIntPoint& InNewOrigin)
 		NavigationSystem->ApplyWorldOffset(Offset, true);
 	}
 
-	FIntPoint PreviosWorldOrigin = GlobalOriginOffset;
+	FIntVector PreviosWorldOriginLocation = OriginLocation;
 	// Set new world origin
-	GlobalOriginOffset = InNewOrigin;
-	RequestedGlobalOriginOffset = InNewOrigin;
+	OriginLocation = InNewOriginLocation;
+	RequestedOriginLocation = InNewOriginLocation;
 	
 	// Propagate event to a level blueprints
 	for(int32 LevelIndex = 0; LevelIndex < Levels.Num(); LevelIndex++)
@@ -3977,30 +3979,31 @@ bool UWorld::SetNewWorldOrigin(const FIntPoint& InNewOrigin)
 		if (Level->bIsVisible && 
 			Level->LevelScriptActor)
 		{
-			Level->LevelScriptActor->WorldOriginChanged(PreviosWorldOrigin, GlobalOriginOffset);
+			Level->LevelScriptActor->WorldOriginLocationChanged(PreviosWorldOriginLocation, OriginLocation);
 		}
 	}
 
 	if (AISystem != NULL)
 	{
-		AISystem->WorldOriginChanged(PreviosWorldOrigin, GlobalOriginOffset);
+		AISystem->WorldOriginLocationChanged(PreviosWorldOriginLocation, OriginLocation);
 	}
 
 	// Broadcast that have finished world shifting
-	FCoreDelegates::PostWorldOriginOffset.Broadcast(this, PreviosWorldOrigin, GlobalOriginOffset);
+	FCoreDelegates::PostWorldOriginOffset.Broadcast(this, PreviosWorldOriginLocation, OriginLocation);
 
 	const double CurrentTime = FPlatformTime::Seconds();
 	const float TimeTaken = CurrentTime - MoveStartTime;
-	UE_LOG(LogLevel, Log, TEXT("WORLD TRANSLATION END {%d, %d} took %.4f"), GlobalOriginOffset.X, GlobalOriginOffset.Y, TimeTaken);
+	UE_LOG(LogLevel, Log, TEXT("WORLD TRANSLATION END {%d, %d, %d} took %.4f ms"),
+		OriginLocation.X, OriginLocation.Y, OriginLocation.Z, TimeTaken);
 	
 	return true;
 }
 
-void UWorld::NavigateTo(FIntPoint InPosition)
+void UWorld::NavigateTo(FIntVector InLocation)
 {
 	check(WorldComposition != NULL);
 
-	SetNewWorldOrigin(InPosition);
+	SetNewWorldOrigin(InLocation);
 	WorldComposition->UpdateStreamingState(FVector::ZeroVector);
 	FlushLevelStreaming();
 }
diff --git a/Engine/Source/Runtime/Engine/Private/WorldComposition.cpp b/Engine/Source/Runtime/Engine/Private/WorldComposition.cpp
--- a/Engine/Source/Runtime/Engine/Private/WorldComposition.cpp
+++ b/Engine/Source/Runtime/Engine/Private/WorldComposition.cpp
@@ -441,7 +441,7 @@ void UWorldComposition::GetDistanceVisibleLevels(
 {
 	const UWorld* OwningWorld = GetWorld();
 
-	FIntPoint WorldOffset = OwningWorld->GlobalOriginOffset;
+	FIntPoint WorldOriginLocationXY = FIntPoint(OwningWorld->OriginLocation.X, OwningWorld->OriginLocation.Y);
 			
 	for (int32 TileIdx = 0; TileIdx < Tiles.Num(); TileIdx++)
 	{
@@ -472,8 +472,8 @@ void UWorldComposition::GetDistanceVisibleLevels(
 			//
 			// Check if tile bounding box intersects with a sphere with origin at provided location and with radius equal to tile layer distance settings
 			//
-			FIntPoint LevelOffset = Tile.Info.AbsolutePosition - WorldOffset;
-			FBox LevelBounds = Tile.Info.Bounds.ShiftBy(FVector(LevelOffset, 0));
+			FIntPoint LevelOffset = Tile.Info.AbsolutePosition - WorldOriginLocationXY;
+			FBox LevelBounds = Tile.Info.Bounds.ShiftBy(FVector(LevelOffset));
 			// We don't care about third dimension yet
 			LevelBounds.Min.Z = -WORLD_MAX;
 			LevelBounds.Max.Z = +WORLD_MAX;
@@ -682,7 +682,7 @@ FIntPoint UWorldComposition::GetLevelOffset(ULevel* InLevel) const
 		LevelPosition = LevelPackage->WorldTileInfo->AbsolutePosition;
 	}
 	
-	return LevelPosition - OwningWorld->GlobalOriginOffset;
+	return LevelPosition - FIntPoint(OwningWorld->OriginLocation.X, OwningWorld->OriginLocation.Y);
 }
 
 FBox UWorldComposition::GetLevelBounds(ULevel* InLevel) const
