Fixed some compile errors caused by out-of-order member initialization (on platforms with a Clang compiler)

diff --git a/Engine/Source/Runtime/Engine/Private/Animation/AnimNode_Trail.cpp b/Engine/Source/Runtime/Engine/Private/Animation/AnimNode_Trail.cpp
--- a/Engine/Source/Runtime/Engine/Private/Animation/AnimNode_Trail.cpp
+++ b/Engine/Source/Runtime/Engine/Private/Animation/AnimNode_Trail.cpp
@@ -7,15 +7,15 @@
 // FAnimNode_Trail
 
 FAnimNode_Trail::FAnimNode_Trail()
-	: FakeVelocity(FVector::ZeroVector)
-	, ChainLength(2)
+	: ChainLength(2)
 	, ChainBoneAxis(EAxis::X)
 	, bInvertChainBoneAxis(false)
 	, bLimitStretch(false)
-	, bActorSpaceFakeVel(false)
-	, bHadValidStrength(false)
 	, TrailRelaxation(10)
 	, StretchLimit(0)
+	, FakeVelocity(FVector::ZeroVector)
+	, bActorSpaceFakeVel(false)
+	, bHadValidStrength(false)
 {
 }
 
