Header cleanup: Engine module, public headers.

diff --git a/Engine/Source/Developer/GameplayDebugger/Private/GameplayDebuggingReplicator.cpp b/Engine/Source/Developer/GameplayDebugger/Private/GameplayDebuggingReplicator.cpp
--- a/Engine/Source/Developer/GameplayDebugger/Private/GameplayDebuggingReplicator.cpp
+++ b/Engine/Source/Developer/GameplayDebugger/Private/GameplayDebuggingReplicator.cpp
@@ -1,6 +1,7 @@
 // Copyright 1998-2014 Epic Games, Inc. All Rights Reserved.
 
 #include "GameplayDebuggerPrivate.h"
+#include "Debug/DebugDrawService.h"
 #include "GameFramework/HUD.h"
 #include "GameplayDebuggingComponent.h"
 #include "GameplayDebuggingHUDComponent.h"
diff --git a/Engine/Source/Runtime/Engine/Classes/Engine/Canvas.h b/Engine/Source/Runtime/Engine/Classes/Engine/Canvas.h
--- a/Engine/Source/Runtime/Engine/Classes/Engine/Canvas.h
+++ b/Engine/Source/Runtime/Engine/Classes/Engine/Canvas.h
@@ -5,8 +5,7 @@
 //=============================================================================
 
 #pragma once
-#include "WordWrapper.h"
-#include "Engine/EngineTypes.h"
+#include "Internationalization/WordWrapper.h"
 #include "Canvas.generated.h"
 
 class FSceneView;
diff --git a/Engine/Source/Runtime/Engine/Classes/GameFramework/PlayerInput.h b/Engine/Source/Runtime/Engine/Classes/GameFramework/PlayerInput.h
--- a/Engine/Source/Runtime/Engine/Classes/GameFramework/PlayerInput.h
+++ b/Engine/Source/Runtime/Engine/Classes/GameFramework/PlayerInput.h
@@ -10,6 +10,7 @@
 #include "Components/InputComponent.h"
 #include "Engine/EngineTypes.h"
 #include "GestureRecognizer.h"
+#include "KeyState.h"
 #include "PlayerInput.generated.h"
 
 ENGINE_API DECLARE_LOG_CATEGORY_EXTERN(LogInput, Log, All);
@@ -18,7 +19,6 @@ ENGINE_API DECLARE_LOG_CATEGORY_EXTERN(LogInput, Log, All);
 // Forward declarations.
 //
 struct FDelegateDispatchDetails;
-struct FKeyState;
 
 USTRUCT()
 struct FKeyBind
diff --git a/Engine/Source/Runtime/Engine/Private/NavMeshRenderingHelpers.h b/Engine/Source/Runtime/Engine/Private/NavMeshRenderingHelpers.h
--- a/Engine/Source/Runtime/Engine/Private/NavMeshRenderingHelpers.h
+++ b/Engine/Source/Runtime/Engine/Private/NavMeshRenderingHelpers.h
@@ -1,6 +1,7 @@
 // Copyright 1998-2014 Epic Games, Inc. All Rights Reserved.
 #pragma once
 #include "DebugRenderSceneProxy.h"
+#include "Debug/DebugDrawService.h"
 #include "DynamicMeshBuilder.h"
 #include "AI/Navigation/RecastHelpers.h"
 #include "BatchedElements.h"
diff --git a/Engine/Source/Runtime/Engine/Public/BlueprintUtilities.h b/Engine/Source/Runtime/Engine/Public/BlueprintUtilities.h
--- a/Engine/Source/Runtime/Engine/Public/BlueprintUtilities.h
+++ b/Engine/Source/Runtime/Engine/Public/BlueprintUtilities.h
@@ -2,12 +2,11 @@
 
 #pragma once
 
-#include "Delegate.h"
 #include "GraphEditAction.h"
 
 DECLARE_MULTICAST_DELEGATE_OneParam( FOnGraphChanged, const FEdGraphEditAction& );
-DECLARE_DELEGATE_OneParam( FSingleNodeEvent, class UEdGraphNode* );
-DECLARE_DELEGATE_OneParam( FEdGraphEvent, class UEdGraph* );
+DECLARE_DELEGATE_OneParam( FSingleNodeEvent, UEdGraphNode* );
+DECLARE_DELEGATE_OneParam( FEdGraphEvent, UEdGraph* );
 /** Delegate for notification when property changed */
 DECLARE_MULTICAST_DELEGATE_TwoParams( FOnPropertyChanged, const FPropertyChangedEvent&, const FString& );
 
diff --git a/Engine/Source/Runtime/Engine/Public/Components.h b/Engine/Source/Runtime/Engine/Public/Components.h
--- a/Engine/Source/Runtime/Engine/Public/Components.h
+++ b/Engine/Source/Runtime/Engine/Public/Components.h
@@ -1,9 +1,8 @@
 // Copyright 1998-2014 Epic Games, Inc. All Rights Reserved.
 
 #pragma once
-
-#include "RenderUtils.h"
-
+#include "PackedNormal.h"
+#include "RenderResource.h"
 /*=============================================================================
 	Components.h: Forward declarations of object components of actors
 =============================================================================*/
diff --git a/Engine/Source/Runtime/Engine/Public/DebugRenderSceneProxy.h b/Engine/Source/Runtime/Engine/Public/DebugRenderSceneProxy.h
--- a/Engine/Source/Runtime/Engine/Public/DebugRenderSceneProxy.h
+++ b/Engine/Source/Runtime/Engine/Public/DebugRenderSceneProxy.h
@@ -9,9 +9,10 @@
 #ifndef _INC_DEBUGRENDERSCENEPROXY
 #define _INC_DEBUGRENDERSCENEPROXY
 
-#include "Debug/DebugDrawService.h"
 #include "PrimitiveSceneProxy.h"
 
+DECLARE_DELEGATE_TwoParams(FDebugDrawDelegate, class UCanvas*, class APlayerController*);
+
 class FDebugRenderSceneProxy : public FPrimitiveSceneProxy
 {
 public:
diff --git a/Engine/Source/Runtime/Engine/Public/EngineModule.h b/Engine/Source/Runtime/Engine/Public/EngineModule.h
--- a/Engine/Source/Runtime/Engine/Public/EngineModule.h
+++ b/Engine/Source/Runtime/Engine/Public/EngineModule.h
@@ -4,6 +4,8 @@
 
 #pragma once
 
+class IRendererModule;
+
 /** Implements the engine module. */
 class FEngineModule : public FDefaultModuleImpl
 {
diff --git a/Engine/Source/Runtime/Engine/Public/EngineUtils.h b/Engine/Source/Runtime/Engine/Public/EngineUtils.h
--- a/Engine/Source/Runtime/Engine/Public/EngineUtils.h
+++ b/Engine/Source/Runtime/Engine/Public/EngineUtils.h
@@ -7,6 +7,7 @@
 #ifndef _INC_ENGINE_UTILS
 #define _INC_ENGINE_UTILS
 
+#include "Engine/Brush.h"
 #include "Engine/Level.h"
 #include "GameFramework/WorldSettings.h"
 #include "HitProxies.h"
diff --git a/Engine/Source/Runtime/Engine/Public/GPUSkinVertexFactory.h b/Engine/Source/Runtime/Engine/Public/GPUSkinVertexFactory.h
--- a/Engine/Source/Runtime/Engine/Public/GPUSkinVertexFactory.h
+++ b/Engine/Source/Runtime/Engine/Public/GPUSkinVertexFactory.h
@@ -27,7 +27,6 @@
 
 #include "GPUSkinPublicDefs.h"
 #include "ResourcePool.h"
-#include "SharedPointer.h"
 #include "UniformBuffer.h"
 
 /** for final bone matrices - this needs to move out of ifdef due to APEX using it*/
diff --git a/Engine/Source/Runtime/Engine/Public/GameDelegates.h b/Engine/Source/Runtime/Engine/Public/GameDelegates.h
--- a/Engine/Source/Runtime/Engine/Public/GameDelegates.h
+++ b/Engine/Source/Runtime/Engine/Public/GameDelegates.h
@@ -6,8 +6,6 @@
  * Collection of delegates for various components to call into game code
  */
 
-#include "Delegate.h"
-
 // enum class with size so it can be forward declared.
 enum class EGameDelegates_SaveGame : short
 {	
diff --git a/Engine/Source/Runtime/Engine/Public/GestureRecognizer.h b/Engine/Source/Runtime/Engine/Public/GestureRecognizer.h
--- a/Engine/Source/Runtime/Engine/Public/GestureRecognizer.h
+++ b/Engine/Source/Runtime/Engine/Public/GestureRecognizer.h
@@ -6,8 +6,6 @@
 
 #pragma once
 
-#include "KeyState.h"
-
 class FGestureRecognizer
 {
 public:
diff --git a/Engine/Source/Runtime/Engine/Public/GlobalShader.h b/Engine/Source/Runtime/Engine/Public/GlobalShader.h
--- a/Engine/Source/Runtime/Engine/Public/GlobalShader.h
+++ b/Engine/Source/Runtime/Engine/Public/GlobalShader.h
@@ -7,8 +7,6 @@
 #pragma once
 
 #include "Shader.h"
-#include "SceneManagement.h"
-#include "ShaderParameters.h"
 
 class FShaderCompileJob;
 
diff --git a/Engine/Source/Runtime/MoviePlayer/Private/DefaultGameMoviePlayer.cpp b/Engine/Source/Runtime/MoviePlayer/Private/DefaultGameMoviePlayer.cpp
--- a/Engine/Source/Runtime/MoviePlayer/Private/DefaultGameMoviePlayer.cpp
+++ b/Engine/Source/Runtime/MoviePlayer/Private/DefaultGameMoviePlayer.cpp
@@ -13,7 +13,7 @@
 #include "MoviePlayerThreading.h"
 #include "DefaultGameMoviePlayer.h"
 #include "MoviePlayerSettings.h"
-
+#include "ShaderCompiler.h"
 
 TSharedPtr<FDefaultGameMoviePlayer> FDefaultGameMoviePlayer::MoviePlayer;
 
