Move updates to the GPU particle simulation from PreRenderView to CreateRenderThreadResources

diff --git a/Engine/Source/Runtime/Engine/Private/Particles/FXSystem.cpp b/Engine/Source/Runtime/Engine/Private/Particles/FXSystem.cpp
--- a/Engine/Source/Runtime/Engine/Private/Particles/FXSystem.cpp
+++ b/Engine/Source/Runtime/Engine/Private/Particles/FXSystem.cpp
@@ -309,7 +309,6 @@ void FFXSystem::PreInitViews()
 {
 	if (RHISupportsGPUParticles(FeatureLevel))
 	{
-		ResetSimulationPhases();
 		ResetSortedGPUParticles();
 	}
 }
diff --git a/Engine/Source/Runtime/Engine/Private/Particles/FXSystemPrivate.h b/Engine/Source/Runtime/Engine/Private/Particles/FXSystemPrivate.h
--- a/Engine/Source/Runtime/Engine/Private/Particles/FXSystemPrivate.h
+++ b/Engine/Source/Runtime/Engine/Private/Particles/FXSystemPrivate.h
@@ -143,11 +143,6 @@ private:
 	 */
 	void SortGPUParticles(FRHICommandListImmediate& RHICmdList);
 
-	/**
-	 * Resets all simulations to be run in the main pass.
-	 */
-	void ResetSimulationPhases();
-
 	/**
 	 * Update particles simulated on the GPU.
 	 * @param Phase				Which emitters are being simulated.
diff --git a/Engine/Source/Runtime/Engine/Private/Particles/ParticleGpuSimulation.cpp b/Engine/Source/Runtime/Engine/Private/Particles/ParticleGpuSimulation.cpp
--- a/Engine/Source/Runtime/Engine/Private/Particles/ParticleGpuSimulation.cpp
+++ b/Engine/Source/Runtime/Engine/Private/Particles/ParticleGpuSimulation.cpp
@@ -2290,8 +2290,7 @@ public:
 	int32 SimulationIndex;
 
 	/**
-	 * The phase in which these particles should simulate. Note that this state
-	 * is TRANSIENT and is updated based on visibility.
+	 * The phase in which these particles should simulate.
 	 */
 	EParticleSimulatePhase::Type SimulationPhase;
 
@@ -2540,6 +2539,12 @@ struct FGPUSpriteDynamicEmitterData : FDynamicEmitterDataBase
 	{
 	}
 
+	bool RendersWithTranslucentMaterial() const
+	{
+		EBlendMode BlendMode = Material->GetRenderProxy(false)->GetMaterial(FXSystem->GetFeatureLevel())->GetBlendMode();
+		return IsTranslucentBlendMode(BlendMode);
+	}
+
 	/**
 	 * Called to create render thread resources.
 	 */
@@ -2549,6 +2554,40 @@ struct FGPUSpriteDynamicEmitterData : FDynamicEmitterDataBase
 
 		// Update the per-frame simulation parameters with those provided from the game thread.
 		Simulation->PerFrameSimulationParameters = PerFrameSimulationParameters;
+
+		// Local vector field parameters.
+		Simulation->LocalVectorField.Intensity = LocalVectorFieldIntensity;
+		Simulation->LocalVectorField.Tightness = LocalVectorFieldTightness;
+		Simulation->LocalVectorField.bTileX = bLocalVectorFieldTileX;
+		Simulation->LocalVectorField.bTileY = bLocalVectorFieldTileY;
+		Simulation->LocalVectorField.bTileZ = bLocalVectorFieldTileZ;
+		if (Simulation->LocalVectorField.Resource)
+		{
+			Simulation->LocalVectorField.UpdateTransforms(LocalVectorFieldToWorld);
+		}
+
+		// Update world bounds.
+		Simulation->Bounds = SimulationBounds;
+
+		// Transfer ownership of new data.
+		if (NewParticles.Num())
+		{
+			Exchange(Simulation->NewParticles, NewParticles);
+		}
+		if (TilesToClear.Num())
+		{
+			Exchange(Simulation->TilesToClear, TilesToClear);
+		}
+
+		const bool bTranslucent = RendersWithTranslucentMaterial();
+
+		// If the simulation wants to collide against the depth buffer
+		// and we're not rendering with an opaque material put the 
+		// simulation in the collision phase.
+		if (bTranslucent && Simulation->bWantsCollision)
+		{
+			Simulation->SimulationPhase = EParticleSimulatePhase::Collision;
+		}
 	}
 
 	/**
@@ -2582,48 +2621,16 @@ struct FGPUSpriteDynamicEmitterData : FDynamicEmitterDataBase
 				{
 					// Create per-emitter uniform buffer for dynamic parameters
 					DynamicUniformBuffer = FGPUSpriteEmitterDynamicUniformBufferRef::CreateUniformBufferImmediate(EmitterDynamicParameters, UniformBuffer_SingleFrame);
-
-					// Local vector field parameters.
-					Simulation->LocalVectorField.Intensity = LocalVectorFieldIntensity;
-					Simulation->LocalVectorField.Tightness = LocalVectorFieldTightness;
-					Simulation->LocalVectorField.bTileX = bLocalVectorFieldTileX;
-					Simulation->LocalVectorField.bTileY = bLocalVectorFieldTileY;
-					Simulation->LocalVectorField.bTileZ = bLocalVectorFieldTileZ;
-					if (Simulation->LocalVectorField.Resource)
-					{
-						Simulation->LocalVectorField.UpdateTransforms(LocalVectorFieldToWorld);
-					}
-
-					// Update world bounds.
-					Simulation->Bounds = SimulationBounds;
-
-					// Transfer ownership of new data.
-					if (NewParticles.Num())
-					{
-						Exchange(Simulation->NewParticles, NewParticles);
-					}
-					if (TilesToClear.Num())
-					{
-						Exchange(Simulation->TilesToClear, TilesToClear);
-					}
 				}
 
-				EBlendMode BlendMode = Material->GetRenderProxy(false)->GetMaterial(FeatureLevel)->GetBlendMode();
-				const bool bOpaque = (BlendMode == BLEND_Opaque || BlendMode == BLEND_Masked);
-
+				const bool bTranslucent = RendersWithTranslucentMaterial();
 				const bool bAllowSorting = FXConsoleVariables::bAllowGPUSorting
 					&& FeatureLevel == ERHIFeatureLevel::SM5
-					&& !bOpaque;
-
-				// If the simulation wants to collide against the depth buffer
-				// and we're not rendering with an opaque material put the 
-				// simulation in the collision phase.
-				if (!bOpaque && Simulation->bWantsCollision)
-				{
-					Simulation->SimulationPhase = EParticleSimulatePhase::Collision;
-				}
+					&& bTranslucent;
 
 				// Create vertex factories for any new views.
+				// TODO: We shouldn't modify the Simulation object here. This is just an optimization to cache vertex factories.
+				//       We should just create a temporary one here that lives for the duration of the frame.
 				const int32 ViewCount = ViewFamily->Views.Num();
 				while (Simulation->VertexFactories.Num() < ViewCount)
 				{
@@ -2641,6 +2648,9 @@ struct FGPUSpriteDynamicEmitterData : FDynamicEmitterDataBase
 						FGPUSpriteVertexFactory& VertexFactory = Simulation->VertexFactories[ViewIndex];
 						if (bAllowSorting && SortMode == PSORTMODE_DistanceToView)
 						{
+							// Extensibility TODO: This call to AddSortedGPUSimulation is very awkward. When rendering a frame we need to
+							// accumulate all GPU particle emitters that need to be sorted. That is so they can be sorted in one big radix
+							// sort for efficiency. Ideally that state is per-scene renderer but the renderer doesn't know anything about particles.
 							const FSceneView* View = ViewFamily->Views[ViewIndex];
 							const int32 SortedBufferOffset = FXSystem->AddSortedGPUSimulation(Simulation, View->ViewMatrices.ViewOrigin);
 							check(SimulationResources->SortedVertexBuffer.IsInitialized());
@@ -3964,15 +3974,6 @@ void FFXSystem::SortGPUParticles(FRHICommandListImmediate& RHICmdList)
 	}
 }
 
-void FFXSystem::ResetSimulationPhases()
-{
-	for (TSparseArray<FParticleSimulationGPU*>::TIterator It(GPUSimulations); It; ++It)
-	{
-		FParticleSimulationGPU* Simulation = *It;
-		Simulation->SimulationPhase = EParticleSimulatePhase::Main;
-	}
-}
-
 /**
  * Sets parameters for the vector field instance.
  * @param OutParameters - The uniform parameters structure.
