#ue4 - CharacterMovementComponent does nothing when ticking if the UpdatedComponent has static mobility.

diff --git a/Engine/Source/Runtime/Engine/Private/CharacterMovementComponent.cpp b/Engine/Source/Runtime/Engine/Private/CharacterMovementComponent.cpp
--- a/Engine/Source/Runtime/Engine/Private/CharacterMovementComponent.cpp
+++ b/Engine/Source/Runtime/Engine/Private/CharacterMovementComponent.cpp
@@ -746,7 +746,7 @@ void UCharacterMovementComponent::TickComponent(float DeltaTime, enum ELevelTick
 	Super::TickComponent(DeltaTime, TickType, ThisTickFunction);
 
 	const FVector InputVector = ConsumeInputVector();
-	if (!HasValidData() || UpdatedComponent->IsSimulatingPhysics())
+	if (!HasValidData() || UpdatedComponent->Mobility != EComponentMobility::Movable || UpdatedComponent->IsSimulatingPhysics())
 	{
 		return;
 	}
@@ -1063,7 +1063,7 @@ void UCharacterMovementComponent::SimulateMovement(float DeltaSeconds)
 			HandlePendingLaunch();
 		}
 
-		if (MovementMode == MOVE_None)
+		if (MovementMode == MOVE_None || UpdatedComponent->Mobility != EComponentMobility::Movable)
 		{
 			return;
 		}
@@ -1329,7 +1329,7 @@ void UCharacterMovementComponent::PerformMovement(float DeltaSeconds)
 	}
 	
 	// no movement if we can't move, or if currently doing physical simulation on UpdatedComponent
-	if (MovementMode == MOVE_None || UpdatedComponent->IsSimulatingPhysics())
+	if (MovementMode == MOVE_None || UpdatedComponent->Mobility != EComponentMobility::Movable || UpdatedComponent->IsSimulatingPhysics())
 	{
 		return;
 	}
@@ -4926,7 +4926,7 @@ bool UCharacterMovementComponent::ClientUpdatePositionAfterServerUpdate()
 
 void UCharacterMovementComponent::ForcePositionUpdate(float DeltaTime)
 {
-	if (!HasValidData() || MovementMode == MOVE_None)
+	if (!HasValidData() || MovementMode == MOVE_None || UpdatedComponent->Mobility != EComponentMobility::Movable)
 	{
 		return;
 	}
