#ue4 - Fix SkipUpdate() for MovementComponents when bUpdateOnlyIfRendered is true, and rename the function to ShouldSkipUpdate().

- SkipUpdate() has been deprecated in favor of ShouldSkipUpdate().
- Most components used by movement don't render (they are invisibile collision components), so now we check any attached components as well.
- ShouldSkipUpdate() now returns true if UpdatedComponent has static mobility.
- Added calls to ShouldSkipUpdate() to some MovementComponents that lacked it.

diff --git a/Engine/Source/Runtime/Engine/Classes/GameFramework/MovementComponent.h b/Engine/Source/Runtime/Engine/Classes/GameFramework/MovementComponent.h
--- a/Engine/Source/Runtime/Engine/Classes/GameFramework/MovementComponent.h
+++ b/Engine/Source/Runtime/Engine/Classes/GameFramework/MovementComponent.h
@@ -118,11 +118,15 @@ public:
 	virtual void StopMovementImmediately();
 
 	/**
-	 * Possibly skip update if moved component is not rendered.
+	 * Possibly skip update if moved component is not rendered or can't move.
 	 * @param DeltaTime @todo this parameter is not used in the function.
 	 * @return true if component movement update should be skipped
 	 */
-	virtual bool SkipUpdate(float DeltaTime);
+	virtual bool ShouldSkipUpdate(float DeltaTime) const;
+
+	DEPRECATED(4.5, "SkipUpdate has been renamed to ShouldSkipUpdate")
+	virtual bool SkipUpdate(float DeltaTime) { return ShouldSkipUpdate(DeltaTime); }
+
 
 	/** @return PhysicsVolume this MovementComponent is using, or the world's default physics volume if none. **/
 	UFUNCTION(BlueprintCallable, Category="Components|Movement")
diff --git a/Engine/Source/Runtime/Engine/Private/CharacterMovementComponent.cpp b/Engine/Source/Runtime/Engine/Private/CharacterMovementComponent.cpp
--- a/Engine/Source/Runtime/Engine/Private/CharacterMovementComponent.cpp
+++ b/Engine/Source/Runtime/Engine/Private/CharacterMovementComponent.cpp
@@ -746,7 +746,7 @@ void UCharacterMovementComponent::TickComponent(float DeltaTime, enum ELevelTick
 	Super::TickComponent(DeltaTime, TickType, ThisTickFunction);
 
 	const FVector InputVector = ConsumeInputVector();
-	if (!HasValidData() || UpdatedComponent->Mobility != EComponentMobility::Movable || UpdatedComponent->IsSimulatingPhysics())
+	if (!HasValidData() || ShouldSkipUpdate(DeltaTime) || UpdatedComponent->IsSimulatingPhysics())
 	{
 		return;
 	}
@@ -1011,7 +1011,7 @@ void UCharacterMovementComponent::SimulateRootMotion(float DeltaSeconds, const F
 
 void UCharacterMovementComponent::SimulateMovement(float DeltaSeconds)
 {
-	if (!HasValidData() || UpdatedComponent->IsSimulatingPhysics())
+	if (!HasValidData() || UpdatedComponent->Mobility != EComponentMobility::Movable || UpdatedComponent->IsSimulatingPhysics())
 	{
 		return;
 	}
@@ -1063,7 +1063,7 @@ void UCharacterMovementComponent::SimulateMovement(float DeltaSeconds)
 			HandlePendingLaunch();
 		}
 
-		if (MovementMode == MOVE_None || UpdatedComponent->Mobility != EComponentMobility::Movable)
+		if (MovementMode == MOVE_None)
 		{
 			return;
 		}
diff --git a/Engine/Source/Runtime/Engine/Private/MovementComponent.cpp b/Engine/Source/Runtime/Engine/Private/MovementComponent.cpp
--- a/Engine/Source/Runtime/Engine/Private/MovementComponent.cpp
+++ b/Engine/Source/Runtime/Engine/Private/MovementComponent.cpp
@@ -143,11 +143,48 @@ bool UMovementComponent::IsInWater() const
 	return GetPhysicsVolume() && GetPhysicsVolume()->bWaterVolume;
 }
 
-bool UMovementComponent::SkipUpdate(float DeltaTime)
+bool UMovementComponent::ShouldSkipUpdate(float DeltaTime) const
 {
-	return UpdatedComponent == NULL
-			|| (bUpdateOnlyIfRendered 
-				&& ((GetNetMode() == NM_DedicatedServer) || (GetWorld()->GetTimeSeconds() - UpdatedComponent->LastRenderTime > 0.2f)));
+	if (UpdatedComponent == NULL || UpdatedComponent->Mobility != EComponentMobility::Movable)
+	{
+		return true;
+	}
+
+	if (bUpdateOnlyIfRendered)
+	{
+		if (GetNetMode() == NM_DedicatedServer)
+		{
+			// Dedicated servers never render
+			return true;
+		}
+
+		const float RenderTimeThreshold = 0.41f;
+		UWorld* TheWorld = GetWorld();
+		if (TheWorld->TimeSince(UpdatedComponent->LastRenderTime) <= RenderTimeThreshold)
+		{
+			return false; // Rendered, don't skip it.
+		}
+
+		// Most components used with movement components don't actually render, so check attached children render times.
+		TArray<USceneComponent*> Children;
+		UpdatedComponent->GetChildrenComponents(true, Children);
+		for (auto Child : Children)
+		{
+			const UPrimitiveComponent* PrimitiveChild = Cast<UPrimitiveComponent>(Child);
+			if (PrimitiveChild)
+			{
+				if (PrimitiveChild->IsRegistered() && TheWorld->TimeSince(PrimitiveChild->LastRenderTime) <= RenderTimeThreshold)
+				{
+					return false; // Rendered, don't skip it.
+				}
+			}
+		}
+
+		// No children were recently rendered, safely skip the update.
+		return true;
+	}
+
+	return false;
 }
 
 
diff --git a/Engine/Source/Runtime/Engine/Private/ProjectileMovementComponent.cpp b/Engine/Source/Runtime/Engine/Private/ProjectileMovementComponent.cpp
--- a/Engine/Source/Runtime/Engine/Private/ProjectileMovementComponent.cpp
+++ b/Engine/Source/Runtime/Engine/Private/ProjectileMovementComponent.cpp
@@ -75,8 +75,8 @@ void UProjectileMovementComponent::TickComponent(float DeltaTime, enum ELevelTic
 {
 	QUICK_SCOPE_CYCLE_COUNTER( STAT_ProjectileMovementComponent_TickComponent );
 
-	// skip if don't want component updated when not rendered
-	if ( SkipUpdate(DeltaTime) )
+	// skip if don't want component updated when not rendered or updated component can't move
+	if ( ShouldSkipUpdate(DeltaTime) )
 	{
 		return;
 	}
diff --git a/Engine/Source/Runtime/Engine/Private/RotatingMovementComponent.cpp b/Engine/Source/Runtime/Engine/Private/RotatingMovementComponent.cpp
--- a/Engine/Source/Runtime/Engine/Private/RotatingMovementComponent.cpp
+++ b/Engine/Source/Runtime/Engine/Private/RotatingMovementComponent.cpp
@@ -13,8 +13,8 @@ URotatingMovementComponent::URotatingMovementComponent(const class FPostConstruc
 
 void URotatingMovementComponent::TickComponent(float DeltaTime, enum ELevelTick TickType, FActorComponentTickFunction *ThisTickFunction)
 {
-	// skip if don't want component updated when not rendered
-	if ( SkipUpdate(DeltaTime) )
+	// skip if don't want component updated when not rendered or updated component can't move
+	if ( ShouldSkipUpdate(DeltaTime) )
 	{
 		return;
 	}
