#ue4 - FloatingPawnMovement aborts if ShouldSkipUpdate() returns true.

diff --git a/Engine/Source/Runtime/Engine/Private/FloatingPawnMovement.cpp b/Engine/Source/Runtime/Engine/Private/FloatingPawnMovement.cpp
--- a/Engine/Source/Runtime/Engine/Private/FloatingPawnMovement.cpp
+++ b/Engine/Source/Runtime/Engine/Private/FloatingPawnMovement.cpp
@@ -19,7 +19,7 @@ void UFloatingPawnMovement::TickComponent(float DeltaTime, enum ELevelTick TickT
 {
 	Super::TickComponent(DeltaTime, TickType, ThisTickFunction);
 
-	if (!PawnOwner || !UpdatedComponent)
+	if (!PawnOwner || !UpdatedComponent || ShouldSkipUpdate(DeltaTime))
 	{
 		return;
 	}
