Don't crash in Game Mode Restart if Possessing the Pawn results in it being torn off (dying)

diff --git a/Engine/Source/Runtime/Engine/Private/GameMode.cpp b/Engine/Source/Runtime/Engine/Private/GameMode.cpp
--- a/Engine/Source/Runtime/Engine/Private/GameMode.cpp
+++ b/Engine/Source/Runtime/Engine/Private/GameMode.cpp
@@ -458,15 +458,23 @@ void AGameMode::RestartPlayer(AController* NewPlayer)
 		}
 		*/
 		NewPlayer->Possess(NewPlayer->GetPawn());
-		
-		// set initial control rotation to player start's rotation
-		NewPlayer->ClientSetRotation(NewPlayer->GetPawn()->GetActorRotation(), true);
 
-		FRotator NewControllerRot = StartSpot->GetActorRotation();
-		NewControllerRot.Roll = 0.f;
-		NewPlayer->SetControlRotation( NewControllerRot );
+		// If the Pawn is destroyed as part of possession we have to abort
+		if (NewPlayer->GetPawn() == nullptr)
+		{
+			NewPlayer->FailedToSpawnPawn();
+		}
+		else
+		{
+			// set initial control rotation to player start's rotation
+			NewPlayer->ClientSetRotation(NewPlayer->GetPawn()->GetActorRotation(), true);
+
+			FRotator NewControllerRot = StartSpot->GetActorRotation();
+			NewControllerRot.Roll = 0.f;
+			NewPlayer->SetControlRotation( NewControllerRot );
 
-		SetPlayerDefaults(NewPlayer->GetPawn());
+			SetPlayerDefaults(NewPlayer->GetPawn());
+		}
 	}
 
 #if !UE_WITH_PHYSICS
