Fix crashing on matinees with zero length that are looped

diff --git a/Engine/Source/Runtime/Engine/Private/Interpolation.cpp b/Engine/Source/Runtime/Engine/Private/Interpolation.cpp
--- a/Engine/Source/Runtime/Engine/Private/Interpolation.cpp
+++ b/Engine/Source/Runtime/Engine/Private/Interpolation.cpp
@@ -990,7 +990,7 @@ void AMatineeActor::StepInterp( float DeltaTime, bool bPreview )
 			if(NewPosition > MatineeData->InterpLength)
 			{
 				// If looping, play to end, jump to start, and set target to somewhere near the beginning.
-				if(bLooping)
+				if(bLooping && MatineeData->InterpLength > 0.0f )
 				{
 					UpdateInterp(MatineeData->InterpLength, bPreview);
 
