#ue4 - Improved comment on AActor::AddTickPrerequisiteActor and AddTickPrerequisiteComponent to make it clear that it does not include owned components.

diff --git a/Engine/Source/Runtime/Engine/Classes/GameFramework/Actor.h b/Engine/Source/Runtime/Engine/Classes/GameFramework/Actor.h
--- a/Engine/Source/Runtime/Engine/Classes/GameFramework/Actor.h
+++ b/Engine/Source/Runtime/Engine/Classes/GameFramework/Actor.h
@@ -779,11 +779,11 @@ public:
 	DEPRECATED(4.5, "Actor::SetTickPrerequisite() will be removed, use AddTickPrerequisiteActor().")
 	void SetTickPrerequisite(AActor* PrerequisiteActor);
 
-	/** Make this actor tick after PrerequisiteActor */
+	/** Make this actor tick after PrerequisiteActor. This only applies to this actor's tick function; dependencies for owned components must be set up separately if desired. */
 	UFUNCTION(BlueprintCallable, Category="Utilities", meta=(Keywords = "dependency"))
 	virtual void AddTickPrerequisiteActor(AActor* PrerequisiteActor);
 
-	/** Make this actor tick after PrerequisiteComponent. */
+	/** Make this actor tick after PrerequisiteComponent. This only applies to this actor's tick function; dependencies for owned components must be set up separately if desired. */
 	UFUNCTION(BlueprintCallable, Category="Utilities", meta=(Keywords = "dependency"))
 	virtual void AddTickPrerequisiteComponent(UActorComponent* PrerequisiteComponent);
 
