#ue4 - Expose GetBaseAimRotation() to BPs.

forum: https://forums.unrealengine.com/showthread.php?21255-AimOffset-Pitch-in-multiplayer

diff --git a/Engine/Source/Runtime/Engine/Classes/GameFramework/Pawn.h b/Engine/Source/Runtime/Engine/Classes/GameFramework/Pawn.h
--- a/Engine/Source/Runtime/Engine/Classes/GameFramework/Pawn.h
+++ b/Engine/Source/Runtime/Engine/Classes/GameFramework/Pawn.h
@@ -265,7 +265,12 @@ public:
 	/** @return	Pawn's eye location */
 	virtual FVector GetPawnViewLocation() const;
 
-	/** @return base Aim Rotation without any adjustment (no aim error, no autolock, no adhesion.. just clean initial aim rotation!) */
+	/**
+	 * Return the aim rotation for the Pawn.
+	 * If we have a controller, by default we aim at the player's 'eyes' direction
+	 * that is by default the Pawn rotation for AI, and camera (crosshair) rotation for human players.
+	 */
+	UFUNCTION(BlueprintCallable, Category="Pawn")
 	virtual FRotator GetBaseAimRotation() const;
 
 	/** return true if player is viewing this Pawn in FreeCam */
