#ue4 - PlayerController PlayerCameraManagerClass is BlueprintReadOnly (should not have been BlueprintReadWrite, as it is only used at startup).

diff --git a/Engine/Source/Runtime/Engine/Classes/GameFramework/PlayerController.h b/Engine/Source/Runtime/Engine/Classes/GameFramework/PlayerController.h
--- a/Engine/Source/Runtime/Engine/Classes/GameFramework/PlayerController.h
+++ b/Engine/Source/Runtime/Engine/Classes/GameFramework/PlayerController.h
@@ -104,8 +104,8 @@ class ENGINE_API APlayerController : public AController
 	UPROPERTY(EditInline, BlueprintReadOnly, Category=PlayerController)
 	class APlayerCameraManager* PlayerCameraManager;
 
-	/** PlayerCamera class should be set for each game, otherwise Engine.PlayerCamera is used */
-	UPROPERTY(EditAnywhere, BlueprintReadWrite, Category=PlayerController)
+	/** PlayerCamera class should be set for each game, otherwise Engine.PlayerCameraManager is used */
+	UPROPERTY(EditAnywhere, BlueprintReadOnly, Category=PlayerController)
 	TSubclassOf<class APlayerCameraManager> PlayerCameraManagerClass;
 
 	/** 
