EnginePrivate.h cleanup - VectorField headers.

diff --git a/Engine/Source/Runtime/Engine/Private/EnginePrivate.h b/Engine/Source/Runtime/Engine/Private/EnginePrivate.h
--- a/Engine/Source/Runtime/Engine/Private/EnginePrivate.h
+++ b/Engine/Source/Runtime/Engine/Private/EnginePrivate.h
@@ -279,9 +279,6 @@
 #include "GameFramework/TouchInterface.h"
 #include "Engine/UserDefinedEnum.h"
 #include "Engine/UserDefinedStruct.h"
-#include "VectorField/VectorField.h"
-#include "VectorField/VectorFieldAnimated.h"
-#include "VectorField/VectorFieldStatic.h"
 #include "Engine/World.h"
 
 #include "SystemSettings.h"					// Scalability options.
diff --git a/Engine/Source/Runtime/Engine/Private/Particles/FXSystem.cpp b/Engine/Source/Runtime/Engine/Private/Particles/FXSystem.cpp
--- a/Engine/Source/Runtime/Engine/Private/Particles/FXSystem.cpp
+++ b/Engine/Source/Runtime/Engine/Private/Particles/FXSystem.cpp
@@ -13,6 +13,7 @@
 #include "../VectorField.h"
 #include "../GPUSort.h"
 #include "ParticleCurveTexture.h"
+#include "VectorField/VectorField.h"
 
 /*-----------------------------------------------------------------------------
 	External FX system interface.
diff --git a/Engine/Source/Runtime/Engine/Private/Particles/ParticleGpuSimulation.cpp b/Engine/Source/Runtime/Engine/Private/Particles/ParticleGpuSimulation.cpp
--- a/Engine/Source/Runtime/Engine/Private/Particles/ParticleGpuSimulation.cpp
+++ b/Engine/Source/Runtime/Engine/Private/Particles/ParticleGpuSimulation.cpp
@@ -27,6 +27,7 @@
 #include "Particles/ParticleModuleRequired.h"
 #include "Particles/ParticleSpriteEmitter.h"
 #include "Particles/ParticleSystemComponent.h"
+#include "VectorField/VectorField.h"
 
 /*------------------------------------------------------------------------------
 	Constants to tune memory and performance for GPU particle simulation.
diff --git a/Engine/Source/Runtime/Engine/Private/VectorField.cpp b/Engine/Source/Runtime/Engine/Private/VectorField.cpp
--- a/Engine/Source/Runtime/Engine/Private/VectorField.cpp
+++ b/Engine/Source/Runtime/Engine/Private/VectorField.cpp
@@ -14,6 +14,9 @@
 #include "RHIStaticStates.h"
 #include "GlobalShader.h"
 #include "FXSystem.h"
+#include "VectorField/VectorField.h"
+#include "VectorField/VectorFieldAnimated.h"
+#include "VectorField/VectorFieldStatic.h"
 
 #define MAX_GLOBAL_VECTOR_FIELDS (16)
 DEFINE_LOG_CATEGORY(LogVectorField)
