EnginePrivate.h cleanup: Engine/Trigger*

diff --git a/Engine/Source/Runtime/Engine/Classes/Engine/TriggerBox.h b/Engine/Source/Runtime/Engine/Classes/Engine/TriggerBox.h
--- a/Engine/Source/Runtime/Engine/Classes/Engine/TriggerBox.h
+++ b/Engine/Source/Runtime/Engine/Classes/Engine/TriggerBox.h
@@ -2,6 +2,7 @@
 
 
 #pragma once
+#include "Engine/TriggerBase.h"
 #include "TriggerBox.generated.h"
 
 UCLASS(MinimalAPI)
diff --git a/Engine/Source/Runtime/Engine/Classes/Engine/TriggerCapsule.h b/Engine/Source/Runtime/Engine/Classes/Engine/TriggerCapsule.h
--- a/Engine/Source/Runtime/Engine/Classes/Engine/TriggerCapsule.h
+++ b/Engine/Source/Runtime/Engine/Classes/Engine/TriggerCapsule.h
@@ -2,6 +2,7 @@
 
 
 #pragma once
+#include "Engine/TriggerBase.h"
 #include "TriggerCapsule.generated.h"
 
 UCLASS(MinimalAPI)
diff --git a/Engine/Source/Runtime/Engine/Classes/Engine/TriggerSphere.h b/Engine/Source/Runtime/Engine/Classes/Engine/TriggerSphere.h
--- a/Engine/Source/Runtime/Engine/Classes/Engine/TriggerSphere.h
+++ b/Engine/Source/Runtime/Engine/Classes/Engine/TriggerSphere.h
@@ -2,6 +2,7 @@
 
 
 #pragma once
+#include "Engine/TriggerBase.h"
 #include "TriggerSphere.generated.h"
 
 UCLASS(MinimalAPI)
diff --git a/Engine/Source/Runtime/Engine/Private/EnginePrivate.h b/Engine/Source/Runtime/Engine/Private/EnginePrivate.h
--- a/Engine/Source/Runtime/Engine/Private/EnginePrivate.h
+++ b/Engine/Source/Runtime/Engine/Private/EnginePrivate.h
@@ -81,10 +81,6 @@
 #include "Components/SkeletalMeshComponent.h"
 #include "Engine/TargetPoint.h"
 #include "Engine/TextRenderActor.h"
-#include "Engine/TriggerBase.h"
-#include "Engine/TriggerBox.h"
-#include "Engine/TriggerCapsule.h"
-#include "Engine/TriggerSphere.h"
 #include "Components/ApplicationLifecycleComponent.h"
 #include "Engine/DataAsset.h"
 #include "Components/PawnNoiseEmitterComponent.h"
diff --git a/Engine/Source/Runtime/Engine/Private/Particles/ParticleModules_Collision.cpp b/Engine/Source/Runtime/Engine/Private/Particles/ParticleModules_Collision.cpp
--- a/Engine/Source/Runtime/Engine/Private/Particles/ParticleModules_Collision.cpp
+++ b/Engine/Source/Runtime/Engine/Private/Particles/ParticleModules_Collision.cpp
@@ -5,6 +5,7 @@
 	Collision-related particle module implementations.
 =============================================================================*/
 #include "EnginePrivate.h"
+#include "Engine/TriggerBase.h"
 #include "ParticleDefinitions.h"
 #include "../DistributionHelpers.h"
 #include "Particles/Collision/ParticleModuleCollision.h"
diff --git a/Engine/Source/Runtime/Engine/Private/TriggerActors.cpp b/Engine/Source/Runtime/Engine/Private/TriggerActors.cpp
--- a/Engine/Source/Runtime/Engine/Private/TriggerActors.cpp
+++ b/Engine/Source/Runtime/Engine/Private/TriggerActors.cpp
@@ -5,6 +5,9 @@
 =============================================================================*/
 
 #include "EnginePrivate.h"
+#include "Engine/TriggerBox.h"
+#include "Engine/TriggerCapsule.h"
+#include "Engine/TriggerSphere.h"
 #include "Net/UnrealNetwork.h"
 namespace
 {
diff --git a/Engine/Source/Runtime/Engine/Private/TriggerBase.cpp b/Engine/Source/Runtime/Engine/Private/TriggerBase.cpp
--- a/Engine/Source/Runtime/Engine/Private/TriggerBase.cpp
+++ b/Engine/Source/Runtime/Engine/Private/TriggerBase.cpp
@@ -1,6 +1,7 @@
 // Copyright 1998-2014 Epic Games, Inc. All Rights Reserved.
 
 #include "EnginePrivate.h"
+#include "Engine/TriggerBase.h"
 
 ATriggerBase::ATriggerBase(const class FPostConstructInitializeProperties& PCIP)
 	: Super(PCIP)
