Crash fixes to some NavigationSystem's extreme use cases #UE4

- TTP#343307, TTP#343277, TTP#343252 crashes related to Undo/Redo in the editor
- TTP#343278, TTP#342176 crashes related to stretching navmesh bounds to ridiculous sizes
- refactored some navigation system's functions to take reference parameters rather than pointers

diff --git a/Engine/Source/Runtime/Engine/Classes/AI/Navigation/NavigationData.h b/Engine/Source/Runtime/Engine/Classes/AI/Navigation/NavigationData.h
--- a/Engine/Source/Runtime/Engine/Classes/AI/Navigation/NavigationData.h
+++ b/Engine/Source/Runtime/Engine/Classes/AI/Navigation/NavigationData.h
@@ -400,6 +400,7 @@ class ENGINE_API ANavigationData : public AActor
 
 	FORCEINLINE bool IsRegistered() const { return bRegistered; }
 	void OnRegistered();
+	void OnUnregistered();
 	
 	FORCEINLINE uint16 GetNavDataUniqueID() const { return NavDataUniqueID; }
 
@@ -625,12 +626,12 @@ public:
 
 	/** new area was registered in navigation system */
 	virtual void OnNavAreaAdded(const UClass* NavAreaClass, int32 AgentIndex);
-	void OnNavAreaAddedNotify(const UClass* NavAreaClass);
 	
 	/** area was removed from navigation system */
 	virtual void OnNavAreaRemoved(const UClass* NavAreaClass);
-	void OnNavAreaRemovedNotify(const UClass* NavAreaClass);
-	
+		
+	void OnNavAreaEvent(const UClass* NavAreaClass, ENavAreaEvent::Type Event);
+
 	/** add all existing areas */
 	void ProcessNavAreas(const TArray<const UClass*>& AreaClasses, int32 AgentIndex);
 
diff --git a/Engine/Source/Runtime/Engine/Classes/AI/Navigation/NavigationSystem.h b/Engine/Source/Runtime/Engine/Classes/AI/Navigation/NavigationSystem.h
--- a/Engine/Source/Runtime/Engine/Classes/AI/Navigation/NavigationSystem.h
+++ b/Engine/Source/Runtime/Engine/Classes/AI/Navigation/NavigationSystem.h
@@ -23,6 +23,20 @@ class UWorld;
 class UCrowdManager;
 struct FNavDataConfig;
 struct FPathFindingResult;
+class UActorComponent;
+class AActor;
+class UEnum;
+class AController;
+class UCrowdManager;
+class UNavArea;
+class INavLinkCustomInterface;
+class FNavigationOctree;
+class AVolume;
+class FNavDataGenerator;
+class AWorldSettings;
+#if WITH_EDITOR
+class FEdMode;
+#endif // WITH_EDITOR
 
 ENGINE_API DECLARE_LOG_CATEGORY_EXTERN(LogNavigation, Warning, All);
 
@@ -160,30 +174,30 @@ public:
 	//----------------------------------------------------------------------//
 	/** Project a point onto the NavigationData */
 	UFUNCTION(BlueprintPure, Category="AI|Navigation", meta=(HidePin="WorldContext", DefaultToSelf="WorldContext" ) )
-	static FVector ProjectPointToNavigation(UObject* WorldContext, const FVector& Point, class ANavigationData* NavData = NULL, TSubclassOf<UNavigationQueryFilter> FilterClass = NULL);
+	static FVector ProjectPointToNavigation(UObject* WorldContext, const FVector& Point, ANavigationData* NavData = NULL, TSubclassOf<UNavigationQueryFilter> FilterClass = NULL);
 	
 	UFUNCTION(BlueprintPure, Category="AI|Navigation", meta=(HidePin="WorldContext", DefaultToSelf="WorldContext" ) )
-	static FVector GetRandomPoint(UObject* WorldContext, class ANavigationData* NavData = NULL, TSubclassOf<UNavigationQueryFilter> FilterClass = NULL);
+	static FVector GetRandomPoint(UObject* WorldContext, ANavigationData* NavData = NULL, TSubclassOf<UNavigationQueryFilter> FilterClass = NULL);
 
 	UFUNCTION(BlueprintPure, Category="AI|Navigation", meta=(HidePin="WorldContext", DefaultToSelf="WorldContext" ) )
-	static FVector GetRandomPointInRadius(UObject* WorldContext, const FVector& Origin, float Radius, class ANavigationData* NavData = NULL, TSubclassOf<UNavigationQueryFilter> FilterClass = NULL);
+	static FVector GetRandomPointInRadius(UObject* WorldContext, const FVector& Origin, float Radius, ANavigationData* NavData = NULL, TSubclassOf<UNavigationQueryFilter> FilterClass = NULL);
 
 	/** Potentially expensive. Use with caution. Consider using UPathFollowingComponent::GetRemainingPathCost instead */
 	UFUNCTION(BlueprintPure, Category="AI|Navigation", meta=(HidePin="WorldContext", DefaultToSelf="WorldContext" ) )
-	static ENavigationQueryResult::Type GetPathCost(UObject* WorldContext, const FVector& PathStart, const FVector& PathEnd, float& PathCost, class ANavigationData* NavData = NULL, TSubclassOf<UNavigationQueryFilter> FilterClass = NULL);
+	static ENavigationQueryResult::Type GetPathCost(UObject* WorldContext, const FVector& PathStart, const FVector& PathEnd, float& PathCost, ANavigationData* NavData = NULL, TSubclassOf<UNavigationQueryFilter> FilterClass = NULL);
 
 	/** Potentially expensive. Use with caution */
 	UFUNCTION(BlueprintPure, Category="AI|Navigation", meta=(HidePin="WorldContext", DefaultToSelf="WorldContext" ) )
-	static ENavigationQueryResult::Type GetPathLength(UObject* WorldContext, const FVector& PathStart, const FVector& PathEnd, float& PathLength, class ANavigationData* NavData = NULL, TSubclassOf<UNavigationQueryFilter> FilterClass = NULL);
+	static ENavigationQueryResult::Type GetPathLength(UObject* WorldContext, const FVector& PathStart, const FVector& PathEnd, float& PathLength, ANavigationData* NavData = NULL, TSubclassOf<UNavigationQueryFilter> FilterClass = NULL);
 
 	UFUNCTION(BlueprintPure, Category="AI|Navigation", meta=(HidePin="WorldContext", DefaultToSelf="WorldContext" ) )
 	static bool IsNavigationBeingBuilt(UObject* WorldContext);
 
 	UFUNCTION(BlueprintCallable, Category="AI|Navigation")
-	static void SimpleMoveToActor(class AController* Controller, const AActor* Goal);
+	static void SimpleMoveToActor(AController* Controller, const AActor* Goal);
 
 	UFUNCTION(BlueprintCallable, Category="AI|Navigation")
-	static void SimpleMoveToLocation(class AController* Controller, const FVector& Goal);
+	static void SimpleMoveToLocation(AController* Controller, const FVector& Goal);
 
 	/** Finds path instantly, in a FindPath Synchronously. 
 	 *	@param PathfindingContext could be one of following: NavigationData (like Navmesh actor), Pawn or Controller. This parameter determines parameters of specific pathfinding query */
@@ -201,7 +215,7 @@ public:
 	 *	@param HitLocation if line was obstructed this will be set to hit location. Otherwise it contains SegmentEnd
 	 *	@return true if line from RayStart to RayEnd was obstructed. Also, true when no navigation data present */
 	UFUNCTION(BlueprintCallable, Category="AI|Navigation", meta=(HidePin="WorldContext", DefaultToSelf="WorldContext" ))
-	static bool NavigationRaycast(UObject* WorldContext, const FVector& RayStart, const FVector& RayEnd, FVector& HitLocation, TSubclassOf<UNavigationQueryFilter> FilterClass = NULL, class AController* Querier = NULL);
+	static bool NavigationRaycast(UObject* WorldContext, const FVector& RayStart, const FVector& RayEnd, FVector& HitLocation, TSubclassOf<UNavigationQueryFilter> FilterClass = NULL, AController* Querier = NULL);
 
 	/** delegate type for events that dirty the navigation data ( Params: const FBox& DirtyBounds ) */
 	DECLARE_MULTICAST_DELEGATE_OneParam(FOnNavigationDirty, const FBox&);
@@ -240,7 +254,7 @@ public:
 
 	UWorld* GetWorld() const { return GetOuterUWorld(); }
 
-	class UCrowdManager* GetCrowdManager() const { return CrowdManager.Get(); }
+	UCrowdManager* GetCrowdManager() const { return CrowdManager.Get(); }
 
 	/** spawn new crowd manager */
 	virtual void CreateCrowdManager();
@@ -292,17 +306,17 @@ public:
 	 *	@param ResultLocation Found point is put here
 	 *	@param NavData If NavData == NULL then MainNavData is used.
 	 *	@return true if any location found, false otherwise */
-	bool GetRandomPoint(FNavLocation& ResultLocation, class ANavigationData* NavData = NULL, TSharedPtr<const FNavigationQueryFilter> QueryFilter = NULL);
+	bool GetRandomPoint(FNavLocation& ResultLocation, ANavigationData* NavData = NULL, TSharedPtr<const FNavigationQueryFilter> QueryFilter = NULL);
 
 	/** Finds random point in navigable space restricted to Radius around Origin
 	 *	@param ResultLocation Found point is put here
 	 *	@param NavData If NavData == NULL then MainNavData is used.
 	 *	@return true if any location found, false otherwise */
-	bool GetRandomPointInRadius(const FVector& Origin, float Radius, FNavLocation& ResultLocation, class ANavigationData* NavData = NULL, TSharedPtr<const FNavigationQueryFilter> QueryFilter = NULL) const;
+	bool GetRandomPointInRadius(const FVector& Origin, float Radius, FNavLocation& ResultLocation, ANavigationData* NavData = NULL, TSharedPtr<const FNavigationQueryFilter> QueryFilter = NULL) const;
 
 	/** Calculates a path from PathStart to PathEnd and retrieves its cost. 
 	 *	@NOTE potentially expensive, so use it with caution */
-	ENavigationQueryResult::Type GetPathCost(const FVector& PathStart, const FVector& PathEnd, float& PathCost, const class ANavigationData* NavData = NULL, TSharedPtr<const FNavigationQueryFilter> QueryFilter = NULL) const;
+	ENavigationQueryResult::Type GetPathCost(const FVector& PathStart, const FVector& PathEnd, float& PathCost, const ANavigationData* NavData = NULL, TSharedPtr<const FNavigationQueryFilter> QueryFilter = NULL) const;
 
 	/** Calculates a path from PathStart to PathEnd and retrieves its overestimated length.
 	 *	@NOTE potentially expensive, so use it with caution */
@@ -319,30 +333,30 @@ public:
 	}
 
 	// @todo document
-	bool ProjectPointToNavigation(const FVector& Point, FNavLocation& OutLocation, const FVector& Extent = INVALID_NAVEXTENT, const class ANavigationData* NavData = NULL, TSharedPtr<const FNavigationQueryFilter> QueryFilter = NULL) const;
+	bool ProjectPointToNavigation(const FVector& Point, FNavLocation& OutLocation, const FVector& Extent = INVALID_NAVEXTENT, const ANavigationData* NavData = NULL, TSharedPtr<const FNavigationQueryFilter> QueryFilter = NULL) const;
 
 	/** 
 	 * Looks for NavData generated for specified movement properties and returns it. NULL if not found;
 	 */
-	class ANavigationData* GetNavDataForProps(const FNavAgentProperties& AgentProperties);
+	ANavigationData* GetNavDataForProps(const FNavAgentProperties& AgentProperties);
 
 	/** 
 	 * Looks for NavData generated for specified movement properties and returns it. NULL if not found; Const version.
 	 */
-	const class ANavigationData* GetNavDataForProps(const FNavAgentProperties& AgentProperties) const;
+	const ANavigationData* GetNavDataForProps(const FNavAgentProperties& AgentProperties) const;
 
 	/** Returns the world nav mesh object.  Creates one if it doesn't exist. */
-	class ANavigationData* GetMainNavData(FNavigationSystem::ECreateIfEmpty CreateNewIfNoneFound);
+	ANavigationData* GetMainNavData(FNavigationSystem::ECreateIfEmpty CreateNewIfNoneFound);
 	/** Returns the world nav mesh object.  Creates one if it doesn't exist. */
-	const class ANavigationData* GetMainNavData() const { return MainNavData; }
+	const ANavigationData* GetMainNavData() const { return MainNavData; }
 
 	TSharedPtr<FNavigationQueryFilter> CreateDefaultQueryFilterCopy() const;
 
 	/** Super-hacky safety feature for threaded navmesh building. Will be gone once figure out why keeping TSharedPointer to Navigation Generator doesn't 
 	 *	guarantee its existence */
-	bool ShouldGeneratorRun(const class FNavDataGenerator* Generator) const;
+	bool ShouldGeneratorRun(const FNavDataGenerator* Generator) const;
 
-	bool IsNavigationBuilt(const class AWorldSettings* Settings) const;
+	bool IsNavigationBuilt(const AWorldSettings* Settings) const;
 
 	bool IsThereAnywhereToBuildNavigation() const;
 
@@ -353,10 +367,10 @@ public:
 	FBox GetLevelBounds(ULevel* InLevel) const;
 
 	/** @return default walkable area class */
-	FORCEINLINE static TSubclassOf<class UNavArea> GetDefaultWalkableArea() { return DefaultWalkableArea; }
+	FORCEINLINE static TSubclassOf<UNavArea> GetDefaultWalkableArea() { return DefaultWalkableArea; }
 
 	/** @return default obstacle area class */
-	FORCEINLINE static TSubclassOf<class UNavArea> GetDefaultObstacleArea() { return DefaultObstacleArea; }
+	FORCEINLINE static TSubclassOf<UNavArea> GetDefaultObstacleArea() { return DefaultObstacleArea; }
 
 	FORCEINLINE const FNavDataConfig& GetDefaultSupportedAgentConfig() const { check(SupportedAgents.Num() > 0);  return SupportedAgents[0]; }
 
@@ -376,10 +390,10 @@ public:
 	//----------------------------------------------------------------------//
 	
 	// @todo document
-	void UnregisterNavData(class ANavigationData* NavData);
+	void UnregisterNavData(ANavigationData* NavData);
 
 	/** adds NavData to registration candidates queue - NavDataRegistrationQueue */
-	void RequestRegistration(class ANavigationData* NavData, bool bTriggerRegistrationProcessing = true);
+	void RequestRegistration(ANavigationData* NavData, bool bTriggerRegistrationProcessing = true);
 
 	/** Processes registration of candidates queues via RequestRegistration and stored in NavDataRegistrationQueue */
 	void ProcessRegistrationCandidates();
@@ -388,26 +402,26 @@ public:
 	void ProcessNavAreaPendingRegistration();
 	
 	/** @return pointer to ANavigationData instance of given ID, or NULL if it was not found. Note it looks only through registered navigation data */
-	class ANavigationData* GetNavDataWithID(const uint16 NavDataID) const;
+	ANavigationData* GetNavDataWithID(const uint16 NavDataID) const;
 
 	void ReleaseInitialBuildingLock();
 
 	//----------------------------------------------------------------------//
 	// navigation octree related functions
 	//----------------------------------------------------------------------//
-	FSetElementId RegisterNavigationRelevantActor(class AActor* Actor, int32 UpdateFlags = OctreeUpdate_Default);
-	void UnregisterNavigationRelevantActor(class AActor* Actor, int32 UpdateFlags = OctreeUpdate_Default);
+	FSetElementId RegisterNavigationRelevantActor(AActor* Actor, int32 UpdateFlags = OctreeUpdate_Default);
+	void UnregisterNavigationRelevantActor(AActor* Actor, int32 UpdateFlags = OctreeUpdate_Default);
 
 #if WITH_NAVIGATION_GENERATOR
 	void AddDirtyArea(const FBox& NewArea, int32 Flags);
 	void AddDirtyAreas(const TArray<FBox>& NewAreas, int32 Flags);
 #endif // WITH_NAVIGATION_GENERATOR
 
-	const class FNavigationOctree* GetNavOctree() const { return NavOctree; }
+	const FNavigationOctree* GetNavOctree() const { return NavOctree; }
 
-	void UpdateNavOctree(class AActor* Actor, int32 UpdateFlags = OctreeUpdate_Default);
-	void UpdateNavOctree(class UActorComponent* ActorComp, int32 UpdateFlags = OctreeUpdate_Default);
-	void UpdateNavOctreeElement(class UObject* ElementOwner, int32 UpdateFlags, const FBox& Bounds);
+	void UpdateNavOctree(AActor* Actor, int32 UpdateFlags = OctreeUpdate_Default);
+	void UpdateNavOctree(UActorComponent* ActorComp, int32 UpdateFlags = OctreeUpdate_Default);
+	void UpdateNavOctreeElement(UObject* ElementOwner, int32 UpdateFlags, const FBox& Bounds);
 
 	/** called to gather navigation relevant actors that have been created while
 	 *	Navigation System was not present */
@@ -416,6 +430,7 @@ public:
 	FORCEINLINE void SetObjectsNavOctreeId(UObject* Object, FOctreeElementId Id) { ObjectToOctreeId.Add(Object, Id); }
 	FORCEINLINE const FOctreeElementId* GetObjectsNavOctreeId(const UObject* Object) const { return ObjectToOctreeId.Find(Object); }
 	FORCEINLINE	void RemoveObjectsNavOctreeId(UObject* Object) { ObjectToOctreeId.Remove(Object); }
+	void RemoveNavOctreeElementId(const FOctreeElementId& ElementId, int32 UpdateFlags);
 
 	/** find all elements in navigation octree within given box (intersection) */
 	void FindElementsInNavOctree(const FBox& QueryBox, const struct FNavigationOctreeFilter& Filter, TArray<struct FNavigationOctreeElement>& Elements);
@@ -423,14 +438,14 @@ public:
 	//----------------------------------------------------------------------//
 	// Custom navigation links
 	//----------------------------------------------------------------------//
-	void RegisterCustomLink(class INavLinkCustomInterface* CustomLink);
-	void UnregisterCustomLink(class INavLinkCustomInterface* CustomLink);
+	void RegisterCustomLink(INavLinkCustomInterface* CustomLink);
+	void UnregisterCustomLink(INavLinkCustomInterface* CustomLink);
 	
 	/** find custom link by unique ID */
-	class INavLinkCustomInterface* GetCustomLink(uint32 UniqueLinkId) const;
+	INavLinkCustomInterface* GetCustomLink(uint32 UniqueLinkId) const;
 
 	/** updates custom link for all active navigation data instances */
-	void UpdateCustomLink(const class INavLinkCustomInterface* CustomLink);
+	void UpdateCustomLink(const INavLinkCustomInterface* CustomLink);
 
 	//----------------------------------------------------------------------//
 	// Areas
@@ -449,7 +464,7 @@ public:
 	//----------------------------------------------------------------------//
 	
 	/** prepare descriptions of navigation flags in UNavigationQueryFilter class: using enum */
-	void DescribeFilterFlags(class UEnum* FlagsEnum) const;
+	void DescribeFilterFlags(UEnum* FlagsEnum) const;
 
 	/** prepare descriptions of navigation flags in UNavigationQueryFilter class: using array */
 	void DescribeFilterFlags(const TArray<FString>& FlagsDesc) const;
@@ -477,7 +492,7 @@ public:
 	FORCEINLINE bool IsNavigationBuildingLocked() const { return bNavigationBuildingLocked || bInitialBuildingLockActive; }
 
 	// @todo document
-	void OnNavigationBoundsUpdated(class AVolume* NavVolume);
+	void OnNavigationBoundsUpdated(AVolume* NavVolume);
 
 	/** Used to display "navigation building in progress" notify */
 	bool IsNavigationBuildInProgress(bool bCheckDirtyToo = true);
@@ -510,7 +525,7 @@ public:
 
 	void UpdateLevelCollision(ULevel* InLevel);
 
-	virtual void OnEditorModeChanged(class FEdMode* Mode, bool IsEntering);
+	virtual void OnEditorModeChanged(FEdMode* Mode, bool IsEntering);
 #endif // WITH_EDITOR
 
 	/** register actor important for generation (navigation data will be build around them first) */
@@ -518,13 +533,13 @@ public:
 	void UnregisterGenerationSeed(AActor* SeedActor);
 	void GetGenerationSeeds(TArray<FVector>& SeedLocations) const;
 
-	static UNavigationSystem* CreateNavigationSystem(class UWorld* WorldOwner);
+	static UNavigationSystem* CreateNavigationSystem(UWorld* WorldOwner);
 
-	static UNavigationSystem* GetCurrent(class UWorld* World);
-	static UNavigationSystem* GetCurrent(class UObject* WorldContextObject);
+	static UNavigationSystem* GetCurrent(UWorld* World);
+	static UNavigationSystem* GetCurrent(UObject* WorldContextObject);
 
 	/** try to create and setup navigation system */
-	static void InitializeForWorld(class UWorld* World, FNavigationSystem::EMode Mode);
+	static void InitializeForWorld(UWorld* World, FNavigationSystem::EMode Mode);
 
 	// Fetch the array of all nav-agent properties.
 	void GetNavAgentPropertiesArray(TArray<FNavAgentProperties>& OutNavAgentProperties) const;
@@ -552,7 +567,7 @@ protected:
 
 	FNavigationSystem::EMode OperationMode;
 
-	class FNavigationOctree* NavOctree;
+	FNavigationOctree* NavOctree;
 
 	TArray<FAsyncPathFindingQuery> AsyncPathFindingQueries;
 
@@ -563,7 +578,7 @@ protected:
 	TMap<const UObject*, FOctreeElementId> ObjectToOctreeId;
 
 	/** Map of all custom navigation links, that are relevant for path following */
-	TMap<uint32, class INavLinkCustomInterface*> CustomLinksMap;
+	TMap<uint32, INavLinkCustomInterface*> CustomLinksMap;
 
 	/** List of actors relevant to generation of navigation data */
 	TArray<TWeakObjectPtr<AActor> > GenerationSeeds;
@@ -596,7 +611,7 @@ protected:
 #endif // !UE_BUILD_SHIPPING
 
 	/** collection of delegates to create all available navigation data types */
-	static TArray<TSubclassOf<class ANavigationData> > NavDataClasses;
+	static TArray<TSubclassOf<ANavigationData> > NavDataClasses;
 	
 	/** whether seamless navigation building is enabled */
 	static bool bNavigationAutoUpdateEnabled;
@@ -605,11 +620,8 @@ protected:
 
 	static TArray<UClass*> PendingNavAreaRegistration;
 	static TArray<const UClass*> NavAreaClasses;
-	static TSubclassOf<class UNavArea> DefaultWalkableArea;
-	static TSubclassOf<class UNavArea> DefaultObstacleArea;
-
-	static FOnNavAreaChanged NavAreaAddedObservers;
-	static FOnNavAreaChanged NavAreaRemovedObservers;
+	static TSubclassOf<UNavArea> DefaultWalkableArea;
+	static TSubclassOf<UNavArea> DefaultObstacleArea;
 
 	/** delegate handler for PostLoadMap event */
 	void OnPostLoadMap();
@@ -624,26 +636,28 @@ protected:
 	 *	@return RegistrationSuccessful if registration was successful, other results mean it failed
 	 *	@see ERegistrationResult
 	 */
-	ERegistrationResult RegisterNavData(class ANavigationData* NavData);
+	ERegistrationResult RegisterNavData(ANavigationData* NavData);
 
 	/** tries to register navigation area */
 	void RegisterNavAreaClass(UClass* NavAreaClass);
+
+	void OnNavigationAreaEvent(UClass* AreaClass, ENavAreaEvent::Type Event);
 	
- 	FSetElementId RegisterNavOctreeElement(class UObject* ElementOwner, int32 UpdateFlags, const FBox& Bounds);
-	void UnregisterNavOctreeElement(class UObject* ElementOwner, int32 UpdateFlags, const FBox& Bounds);
+ 	FSetElementId RegisterNavOctreeElement(UObject* ElementOwner, int32 UpdateFlags, const FBox& Bounds);
+	void UnregisterNavOctreeElement(UObject* ElementOwner, int32 UpdateFlags, const FBox& Bounds);
 	
 	/** read element data from navigation octree */
-	bool GetNavOctreeElementData(class UObject* NodeOwner, int32& DirtyFlags, FBox& DirtyBounds);
+	bool GetNavOctreeElementData(UObject* NodeOwner, int32& DirtyFlags, FBox& DirtyBounds);
 
 	/** Adds given element to NavOctree. No check for owner's validity are performed, 
 	 *	nor its presence in NavOctree - function assumes callee responsibility 
 	 *	in this regard **/
 	void AddElementToNavOctree(const FNavigationDirtyElement& DirtyElement);
 
-	void AddActorElementToNavOctree(class AActor* Actor, const FNavigationDirtyElement& DirtyElement);
-	void AddComponentElementToNavOctree(class UActorComponent* ActorComp, const FNavigationDirtyElement& DirtyElement, const FBox& Bounds);
+	void AddActorElementToNavOctree(AActor& Actor, const FNavigationDirtyElement& DirtyElement);
+	void AddComponentElementToNavOctree(UActorComponent& ActorComp, const FNavigationDirtyElement& DirtyElement, const FBox& Bounds);
 
-	void SetCrowdManager(class UCrowdManager* NewCrowdManager); 
+	void SetCrowdManager(UCrowdManager* NewCrowdManager); 
 
 #if WITH_NAVIGATION_GENERATOR
 
@@ -666,7 +680,7 @@ private:
 
 	/** constructs a navigation data instance of specified NavDataClass, in passed World
 	 *	for supplied NavConfig */
-	virtual ANavigationData* CreateNavigationDataInstance(TSubclassOf<class ANavigationData> NavDataClass, UWorld* World, const FNavDataConfig& NavConfig);
+	virtual ANavigationData* CreateNavigationDataInstance(TSubclassOf<ANavigationData> NavDataClass, UWorld* World, const FNavDataConfig& NavConfig);
 
 	/** Triggers navigation building on all eligible navigation data. */
 	void RebuildAll();
diff --git a/Engine/Source/Runtime/Engine/Private/AI/Navigation/NavigationData.cpp b/Engine/Source/Runtime/Engine/Private/AI/Navigation/NavigationData.cpp
--- a/Engine/Source/Runtime/Engine/Private/AI/Navigation/NavigationData.cpp
+++ b/Engine/Source/Runtime/Engine/Private/AI/Navigation/NavigationData.cpp
@@ -273,6 +273,11 @@ void ANavigationData::OnRegistered()
 	bRegistered = true; 
 }
 
+void ANavigationData::OnUnregistered()
+{
+	bRegistered = false;
+}
+
 void ANavigationData::InstantiateAndRegisterRenderingComponent()
 {
 #if !UE_BUILD_SHIPPING
@@ -330,7 +335,14 @@ void ANavigationData::PostEditUndo()
 	UWorld* WorldOuter = GetWorld();
 	if (WorldOuter != NULL && WorldOuter->GetNavigationSystem() != NULL)
 	{
-		WorldOuter->GetNavigationSystem()->RequestRegistration(this);
+		if (IsPendingKillPending())
+		{
+			WorldOuter->GetNavigationSystem()->UnregisterNavData(this);
+		}
+		else
+		{
+			WorldOuter->GetNavigationSystem()->RequestRegistration(this);
+		}
 	}
 }
 #endif // WITH_EDITOR
@@ -469,12 +481,19 @@ void ANavigationData::OnNavAreaAdded(const UClass* NavAreaClass, int32 AgentInde
 	UE_LOG(LogNavigation, Verbose, TEXT("%s registered area %s with ID %d"), *GetName(), *AreaClassName, NewAgentData.AreaID);
 }
 
-void ANavigationData::OnNavAreaAddedNotify(const UClass* NavAreaClass)
+void ANavigationData::OnNavAreaEvent(const UClass* NavAreaClass, ENavAreaEvent::Type Event)
 {
-	const UNavigationSystem* NavSys = GetWorld()->GetNavigationSystem();
-	const int32 AgentIndex = NavSys->GetSupportedAgentIndex(this);
+	if (Event == ENavAreaEvent::Registered)
+	{
+		const UNavigationSystem* NavSys = GetWorld()->GetNavigationSystem();
+		const int32 AgentIndex = NavSys->GetSupportedAgentIndex(this);
 
-	OnNavAreaAdded(NavAreaClass, AgentIndex);
+		OnNavAreaAdded(NavAreaClass, AgentIndex);
+	}
+	else // Unregistered
+	{
+		OnNavAreaRemoved(NavAreaClass);
+	}
 }
 
 void ANavigationData::OnNavAreaRemoved(const UClass* NavAreaClass)
@@ -490,11 +509,6 @@ void ANavigationData::OnNavAreaRemoved(const UClass* NavAreaClass)
 	}
 }
 
-void ANavigationData::OnNavAreaRemovedNotify(const UClass* NavAreaClass)
-{
-	OnNavAreaRemoved(NavAreaClass);
-}
-
 void ANavigationData::ProcessNavAreas(const TArray<const UClass*>& AreaClasses, int32 AgentIndex)
 {
 	for (int32 i = 0; i < AreaClasses.Num(); i++)
diff --git a/Engine/Source/Runtime/Engine/Private/AI/Navigation/NavigationOctree.cpp b/Engine/Source/Runtime/Engine/Private/AI/Navigation/NavigationOctree.cpp
--- a/Engine/Source/Runtime/Engine/Private/AI/Navigation/NavigationOctree.cpp
+++ b/Engine/Source/Runtime/Engine/Private/AI/Navigation/NavigationOctree.cpp
@@ -26,20 +26,20 @@ FNavigationOctree::~FNavigationOctree()
 	DEC_MEMORY_STAT_BY(STAT_Navigation_CollisionTreeMemory, NodesMemory);
 }
 
-void FNavigationOctree::AddActor(class AActor* Actor, FNavigationOctreeElement& Data)
+void FNavigationOctree::AddActor(AActor& Actor, FNavigationOctreeElement& Data)
 {
-	if (Actor == NULL || Actor->IsPendingKill())
+	if (Actor.IsPendingKill())
 	{
 		return;
 	}
 
-	INavRelevantActorInterface* NavRelevantActor = InterfaceCast<INavRelevantActorInterface>(Actor);
+	INavRelevantActorInterface* NavRelevantActor = InterfaceCast<INavRelevantActorInterface>(&Actor);
 	UNavigationProxy* ProxyOb = NavRelevantActor ? NavRelevantActor->GetNavigationProxy() : NULL;
-	UObject* NodeOwner = ProxyOb ? (UObject*)ProxyOb : Actor;
+	UObject* NodeOwner = ProxyOb ? (UObject*)ProxyOb : &Actor;
 
 	Data.Owner = NodeOwner;
 
-	const FBox BBox = Actor->GetComponentsBoundingBox();
+	const FBox BBox = Actor.GetComponentsBoundingBox();
 	if (BBox.IsValid && !BBox.GetSize().IsNearlyZero())
 	{
 		Data.Bounds = BBox;
@@ -58,7 +58,7 @@ void FNavigationOctree::AddActor(class AActor* Actor, FNavigationOctreeElement&
 		if (bExportGeometry)
 		{
 			SCOPE_CYCLE_COUNTER(STAT_Navigation_ActorsGeometryExportSync);
-			NavigableGeometryExportDelegate.ExecuteIfBound(*Actor, Data.Data);
+			NavigableGeometryExportDelegate.ExecuteIfBound(Actor, Data.Data);
 		}
 #endif // NAVOCTREE_CONTAINS_COLLISION_DATA
 
@@ -66,23 +66,23 @@ void FNavigationOctree::AddActor(class AActor* Actor, FNavigationOctreeElement&
 	}
 }
 
-void FNavigationOctree::AddComponent(class UActorComponent* ActorComp, const FBox& Bounds, FNavigationOctreeElement& Data)
+void FNavigationOctree::AddComponent(UActorComponent& ActorComp, const FBox& Bounds, FNavigationOctreeElement& Data)
 {
-	if (ActorComp == NULL || ActorComp->IsPendingKill())
+	if (ActorComp.IsPendingKill())
 	{
 		return;
 	}
 
-	Data.Owner = ActorComp;
+	Data.Owner = &ActorComp;
 	Data.Bounds = Bounds;
 
 #if NAVOCTREE_CONTAINS_COLLISION_DATA
 	if (bGatherGeometry)
 	{
-		NavigableGeometryComponentExportDelegate.ExecuteIfBound(*ActorComp, Data.Data);
+		NavigableGeometryComponentExportDelegate.ExecuteIfBound(ActorComp, Data.Data);
 	}
 
-	UNavRelevantComponent* NavRelevantComponent = Cast<UNavRelevantComponent>(ActorComp);
+	UNavRelevantComponent* NavRelevantComponent = Cast<UNavRelevantComponent>(&ActorComp);
 	if (NavRelevantComponent)
 	{
 		NavRelevantComponent->OnApplyModifiers(Data.Data.Modifiers);
@@ -108,16 +108,16 @@ void FNavigationOctree::AddNode(FNavigationOctreeElement& Data)
 	AddElement(Data);
 }
 
-void FNavigationOctree::RemoveNode(const FOctreeElementId* Id)
+void FNavigationOctree::RemoveNode(const FOctreeElementId& Id)
 {
 #if NAVOCTREE_CONTAINS_COLLISION_DATA
-	FNavigationOctreeElement& Data = GetElementById(*Id);
+	FNavigationOctreeElement& Data = GetElementById(Id);
 	const int32 ElementMemory = Data.GetAllocatedSize();
 	NodesMemory -= ElementMemory;
 	DEC_MEMORY_STAT_BY(STAT_Navigation_CollisionTreeMemory, ElementMemory);
 #endif
 
-	RemoveElement(*Id);
+	RemoveElement(Id);
 }
 
 bool FNavigationRelevantData::IsMatchingFilter(const FNavigationOctreeFilter& Filter) const
diff --git a/Engine/Source/Runtime/Engine/Private/AI/Navigation/NavigationSystem.cpp b/Engine/Source/Runtime/Engine/Private/AI/Navigation/NavigationSystem.cpp
--- a/Engine/Source/Runtime/Engine/Private/AI/Navigation/NavigationSystem.cpp
+++ b/Engine/Source/Runtime/Engine/Private/AI/Navigation/NavigationSystem.cpp
@@ -161,13 +161,11 @@ void FNavigationLockContext::UnlockUpdates()
 // UNavigationSystem                                                                
 //----------------------------------------------------------------------//
 bool UNavigationSystem::bNavigationAutoUpdateEnabled = true;
-TArray<TSubclassOf<class ANavigationData> > UNavigationSystem::NavDataClasses;
+TArray<TSubclassOf<ANavigationData> > UNavigationSystem::NavDataClasses;
 TArray<const UClass*> UNavigationSystem::NavAreaClasses;
 TArray<UClass*> UNavigationSystem::PendingNavAreaRegistration;
-TSubclassOf<class UNavArea> UNavigationSystem::DefaultWalkableArea = NULL;
-TSubclassOf<class UNavArea> UNavigationSystem::DefaultObstacleArea = NULL;
-FOnNavAreaChanged UNavigationSystem::NavAreaAddedObservers;
-FOnNavAreaChanged UNavigationSystem::NavAreaRemovedObservers;
+TSubclassOf<UNavArea> UNavigationSystem::DefaultWalkableArea = NULL;
+TSubclassOf<UNavArea> UNavigationSystem::DefaultObstacleArea = NULL;
 #if !(UE_BUILD_SHIPPING || UE_BUILD_TEST)
 FNavigationSystemExec UNavigationSystem::ExecHandler;
 #endif // !UE_BUILD_SHIPPING
@@ -180,7 +178,7 @@ bool UNavigationSystem::bUpdateNavOctreeOnComponentChange = true;
 // life cycle stuff                                                                
 //----------------------------------------------------------------------//
 
-UNavigationSystem::UNavigationSystem(const class FPostConstructInitializeProperties& PCIP)
+UNavigationSystem::UNavigationSystem(const FPostConstructInitializeProperties& PCIP)
 	: Super(PCIP)
 	, bWholeWorldNavigable(false)
 	, bAddPlayersToGenerationSeeds(true)
@@ -250,7 +248,7 @@ void UNavigationSystem::DoInitialSetup()
 	NavDataClasses.Empty(RequiredNavigationDataClassNames.Num());
 	for (int32 Index = 0; Index < RequiredNavigationDataClassNames.Num(); ++Index)
 	{
-		TSubclassOf<class ANavigationData> NavDataClass = LoadClass<ANavigationData>(NULL, *RequiredNavigationDataClassNames[Index].ToString(), NULL, LOAD_None, NULL);
+		TSubclassOf<ANavigationData> NavDataClass = LoadClass<ANavigationData>(NULL, *RequiredNavigationDataClassNames[Index].ToString(), NULL, LOAD_None, NULL);
 		if (NavDataClass)
 		{
 			NavDataClasses.AddUnique(NavDataClass);
@@ -786,7 +784,7 @@ bool UNavigationSystem::GetRandomPointInRadius(const FVector& Origin, float Radi
 	return false;
 }
 
-ENavigationQueryResult::Type UNavigationSystem::GetPathCost(const FVector& PathStart, const FVector& PathEnd, float& OutPathCost, const class ANavigationData* NavData, TSharedPtr<const FNavigationQueryFilter> QueryFilter) const
+ENavigationQueryResult::Type UNavigationSystem::GetPathCost(const FVector& PathStart, const FVector& PathEnd, float& OutPathCost, const ANavigationData* NavData, TSharedPtr<const FNavigationQueryFilter> QueryFilter) const
 {
 	SCOPE_CYCLE_COUNTER(STAT_Navigation_QueriesTimeSync);
 
@@ -822,7 +820,7 @@ ENavigationQueryResult::Type UNavigationSystem::GetPathLengthAndCost(const FVect
 	return NavData != NULL ? NavData->CalcPathLengthAndCost(PathStart, PathEnd, OutPathLength, OutPathCost, QueryFilter) : ENavigationQueryResult::Error;
 }
 
-bool UNavigationSystem::ProjectPointToNavigation(const FVector& Point, FNavLocation& OutLocation, const FVector& Extent, const class ANavigationData* NavData, TSharedPtr<const FNavigationQueryFilter> QueryFilter) const
+bool UNavigationSystem::ProjectPointToNavigation(const FVector& Point, FNavLocation& OutLocation, const FVector& Extent, const ANavigationData* NavData, TSharedPtr<const FNavigationQueryFilter> QueryFilter) const
 {
 	SCOPE_CYCLE_COUNTER(STAT_Navigation_QueriesTimeSync);
 
@@ -964,7 +962,7 @@ UNavigationPath* UNavigationSystem::FindPathToLocationSynchronously(UObject* Wor
 	return ResultPath;
 }
 
-bool UNavigationSystem::NavigationRaycast(UObject* WorldContextObject, const FVector& RayStart, const FVector& RayEnd, FVector& HitLocation, TSubclassOf<class UNavigationQueryFilter> FilterClass, AController* Querier)
+bool UNavigationSystem::NavigationRaycast(UObject* WorldContextObject, const FVector& RayStart, const FVector& RayEnd, FVector& HitLocation, TSubclassOf<UNavigationQueryFilter> FilterClass, AController* Querier)
 {
 	UWorld* World = NULL;
 
@@ -1385,10 +1383,6 @@ UNavigationSystem::ERegistrationResult UNavigationSystem::RegisterNavData(ANavig
 	if (Result == RegistrationSuccessful)
 	{
 		NavDataSet.AddUnique(NavData);
-
-		NavAreaAddedObservers.AddUObject(NavData, &ANavigationData::OnNavAreaAddedNotify);
-		NavAreaRemovedObservers.AddUObject(NavData, &ANavigationData::OnNavAreaRemovedNotify);
-
 		NavData->OnRegistered();
 	}
 	// @todo else might consider modifying this NavData to implement navigation for one of the supported agents
@@ -1408,27 +1402,25 @@ void UNavigationSystem::UnregisterNavData(ANavigationData* NavData)
 	FScopeLock Lock(&NavDataRegistration);
 
 	NavDataSet.RemoveSingle(NavData);
-
-	NavAreaAddedObservers.RemoveUObject(NavData, &ANavigationData::OnNavAreaAddedNotify);
-	NavAreaRemovedObservers.RemoveUObject(NavData, &ANavigationData::OnNavAreaRemovedNotify);
+	NavData->OnUnregistered();
 }
 
-void UNavigationSystem::RegisterCustomLink(class INavLinkCustomInterface* CustomLink)
+void UNavigationSystem::RegisterCustomLink(INavLinkCustomInterface* CustomLink)
 {
 	CustomLinksMap.Add(CustomLink->GetLinkId(), CustomLink);
 }
 
-void UNavigationSystem::UnregisterCustomLink(class INavLinkCustomInterface* CustomLink)
+void UNavigationSystem::UnregisterCustomLink(INavLinkCustomInterface* CustomLink)
 {
 	CustomLinksMap.Remove(CustomLink->GetLinkId());
 }
 
-class INavLinkCustomInterface* UNavigationSystem::GetCustomLink(uint32 UniqueLinkId) const
+INavLinkCustomInterface* UNavigationSystem::GetCustomLink(uint32 UniqueLinkId) const
 {
 	return CustomLinksMap.FindRef(UniqueLinkId);
 }
 
-void UNavigationSystem::UpdateCustomLink(const class INavLinkCustomInterface* CustomLink)
+void UNavigationSystem::UpdateCustomLink(const INavLinkCustomInterface* CustomLink)
 {
 	for (TMap<FNavAgentProperties, ANavigationData*>::TIterator It(AgentToNavDataMap); It; ++It)
 	{
@@ -1448,7 +1440,14 @@ void UNavigationSystem::RequestAreaUnregistering(UClass* NavAreaClass)
 		PendingNavAreaRegistration.RemoveSingleSwap(NavAreaClass);
 
 		// notify navigation data
-		NavAreaRemovedObservers.Broadcast(NavAreaClass);
+		// notify existing nav data
+		for (TObjectIterator<UWorld> It; It; ++It)
+		{
+			if (It->GetNavigationSystem())
+			{
+				It->GetNavigationSystem()->OnNavigationAreaEvent(NavAreaClass, ENavAreaEvent::Unregistered);
+			}
+		}
 	}
 }
 
@@ -1503,7 +1502,13 @@ void UNavigationSystem::RegisterNavAreaClass(UClass* AreaClass)
 	NavAreaClasses.AddUnique(AreaClass);
 
 	// notify existing nav data
-	NavAreaAddedObservers.Broadcast(AreaClass);
+	for (TObjectIterator<UWorld> It; It; ++It)
+	{
+		if (It->GetNavigationSystem())
+		{
+			It->GetNavigationSystem()->OnNavigationAreaEvent(AreaClass, ENavAreaEvent::Registered);
+		}
+	}
 
 #if WITH_EDITOR
 	// update area properties
@@ -1511,6 +1516,18 @@ void UNavigationSystem::RegisterNavAreaClass(UClass* AreaClass)
 #endif
 }
 
+void UNavigationSystem::OnNavigationAreaEvent(UClass* AreaClass, ENavAreaEvent::Type Event)
+{
+	// notify existing nav data
+	for (auto NavigationData : NavDataSet)
+	{
+		if (NavigationData != NULL && NavigationData->IsPendingKillPending() == false)
+		{
+			NavigationData->OnNavAreaEvent(AreaClass, Event);
+		}
+	}
+}
+
 int32 UNavigationSystem::GetSupportedAgentIndex(const ANavigationData* NavData) const
 {
 	if (SupportedAgents.Num() < 2)
@@ -1548,7 +1565,7 @@ int32 UNavigationSystem::GetSupportedAgentIndex(const FNavAgentProperties& NavAg
 	return INDEX_NONE;
 }
 
-void UNavigationSystem::DescribeFilterFlags(class UEnum* FlagsEnum) const
+void UNavigationSystem::DescribeFilterFlags(UEnum* FlagsEnum) const
 {
 #if WITH_EDITOR
 	TArray<FString> FlagDesc;
@@ -1666,7 +1683,7 @@ void UNavigationSystem::GetGenerationSeeds(TArray<FVector>& SeedLocations) const
 	}
 }
 
-UNavigationSystem* UNavigationSystem::CreateNavigationSystem(class UWorld* WorldOwner)
+UNavigationSystem* UNavigationSystem::CreateNavigationSystem(UWorld* WorldOwner)
 {
 	UNavigationSystem* NavSys = NULL;
 
@@ -1686,7 +1703,7 @@ UNavigationSystem* UNavigationSystem::CreateNavigationSystem(class UWorld* World
 	return NavSys;
 }
 
-void UNavigationSystem::InitializeForWorld(class UWorld* World, FNavigationSystem::EMode Mode)
+void UNavigationSystem::InitializeForWorld(UWorld* World, FNavigationSystem::EMode Mode)
 {
 	if (World)
 	{
@@ -1708,7 +1725,7 @@ UNavigationSystem* UNavigationSystem::GetCurrent(UWorld* World)
 	return World ? World->GetNavigationSystem() : NULL;
 }
 
-UNavigationSystem* UNavigationSystem::GetCurrent(class UObject* WorldContextObject)
+UNavigationSystem* UNavigationSystem::GetCurrent(UObject* WorldContextObject)
 {
 	UWorld* World = NULL;
 
@@ -1799,7 +1816,7 @@ FSetElementId UNavigationSystem::RegisterNavigationRelevantActor(AActor* Actor,
 	return RequestId;
 }
 
-FSetElementId UNavigationSystem::RegisterNavOctreeElement(class UObject* ElementOwner, int32 UpdateFlags, const FBox& Bounds)
+FSetElementId UNavigationSystem::RegisterNavOctreeElement(UObject* ElementOwner, int32 UpdateFlags, const FBox& Bounds)
 {
 	FSetElementId SetId;
 
@@ -1912,7 +1929,7 @@ void UNavigationSystem::AddElementToNavOctree(const FNavigationDirtyElement& Dir
 				UPrimitiveComponent* PrimComp = Cast<UPrimitiveComponent>(Components[ComponentIndex]);
 				if (PrimComp && PrimComp->CanEverAffectNavigation())
 				{
-					AddComponentElementToNavOctree(PrimComp, DirtyElement, PrimComp->Bounds.GetBox());
+					AddComponentElementToNavOctree(*PrimComp, DirtyElement, PrimComp->Bounds.GetBox());
 				}
 				else
 				{
@@ -1920,7 +1937,7 @@ void UNavigationSystem::AddElementToNavOctree(const FNavigationDirtyElement& Dir
 					if (NavRelevantComponent && NavRelevantComponent->IsNavigationRelevant())
 					{
 						NavRelevantComponent->OnOwnerRegistered();
-						AddComponentElementToNavOctree(NavRelevantComponent, DirtyElement, NavRelevantComponent->Bounds);
+						AddComponentElementToNavOctree(*NavRelevantComponent, DirtyElement, NavRelevantComponent->Bounds);
 					}
 				}
 			}
@@ -1935,26 +1952,26 @@ void UNavigationSystem::AddElementToNavOctree(const FNavigationDirtyElement& Dir
 				if (NavRelevantComponent && NavRelevantComponent->IsNavigationRelevant())
 				{
 					NavRelevantComponent->OnOwnerRegistered();
-					AddComponentElementToNavOctree(NavRelevantComponent, DirtyElement, NavRelevantComponent->Bounds);
+					AddComponentElementToNavOctree(*NavRelevantComponent, DirtyElement, NavRelevantComponent->Bounds);
 				}
 			}
 
-			AddActorElementToNavOctree(ActorOwner, DirtyElement);
+			AddActorElementToNavOctree(*ActorOwner, DirtyElement);
 		}
 	}
-	else
+	else if (CompOwner)
 	{
-		AddComponentElementToNavOctree(CompOwner, DirtyElement, DirtyElement.BoundsOverride);
+		AddComponentElementToNavOctree(*CompOwner, DirtyElement, DirtyElement.BoundsOverride);
 	}
 }
 
-void UNavigationSystem::AddActorElementToNavOctree(class AActor* Actor, const FNavigationDirtyElement& DirtyElement)
+void UNavigationSystem::AddActorElementToNavOctree(AActor& Actor, const FNavigationDirtyElement& DirtyElement)
 {
-	ensure(GetObjectsNavOctreeId(Actor) == NULL);
+	ensure(GetObjectsNavOctreeId(&Actor) == NULL);
 
 	// notify relevant components about adding owner to octree
 	TArray<UNavRelevantComponent*> Components;
-	Actor->GetComponents(Components);
+	Actor.GetComponents(Components);
 	for (int32 ComponentIndex = 0; ComponentIndex < Components.Num(); ComponentIndex++)
 	{
 		UNavRelevantComponent* NavRelevantComponent = Components[ComponentIndex];
@@ -1978,9 +1995,14 @@ void UNavigationSystem::AddActorElementToNavOctree(class AActor* Actor, const FN
 #endif // WITH_NAVIGATION_GENERATOR
 }
 
-void UNavigationSystem::AddComponentElementToNavOctree(class UActorComponent* ActorComp, const FNavigationDirtyElement& DirtyElement, const FBox& Bounds)
+void UNavigationSystem::AddComponentElementToNavOctree(UActorComponent& ActorComp, const FNavigationDirtyElement& DirtyElement, const FBox& Bounds)
 {
-	ensure(GetObjectsNavOctreeId(ActorComp) == NULL);
+	const FOctreeElementId* ElementId = GetObjectsNavOctreeId(&ActorComp);
+	if (ElementId != NULL)
+	{
+		RemoveNavOctreeElementId(*ElementId, OctreeUpdate_Default);
+		RemoveObjectsNavOctreeId(&ActorComp);
+	}
 
 	FNavigationOctreeElement GeneratedData;
 	NavOctree->AddComponent(ActorComp, Bounds, GeneratedData);
@@ -1996,12 +2018,12 @@ void UNavigationSystem::AddComponentElementToNavOctree(class UActorComponent* Ac
 #endif // WITH_NAVIGATION_GENERATOR
 }
 
-bool UNavigationSystem::GetNavOctreeElementData(class UObject* NodeOwner, int32& DirtyFlags, FBox& DirtyBounds)
+bool UNavigationSystem::GetNavOctreeElementData(UObject* NodeOwner, int32& DirtyFlags, FBox& DirtyBounds)
 {
 	const FOctreeElementId* ElementId = GetObjectsNavOctreeId(NodeOwner);
 	if (ElementId != NULL)
 	{
-		if (ElementId->IsValidId() && NavOctree->IsValidElementId(*ElementId))
+		if (NavOctree->IsValidElementId(*ElementId))
 		{
 #if WITH_NAVIGATION_GENERATOR && NAVOCTREE_CONTAINS_COLLISION_DATA
 			// mark area occupied by given actor as dirty
@@ -2085,7 +2107,7 @@ void UNavigationSystem::UnregisterNavigationRelevantActor(AActor* Actor, int32 U
 	}
 }
 
-void UNavigationSystem::UnregisterNavOctreeElement(class UObject* ElementOwner, int32 UpdateFlags, const FBox& Bounds)
+void UNavigationSystem::UnregisterNavOctreeElement(UObject* ElementOwner, int32 UpdateFlags, const FBox& Bounds)
 {
 #if WITH_EDITOR
 	if (AreFakeComponentChangesBeingApplied())
@@ -2102,19 +2124,7 @@ void UNavigationSystem::UnregisterNavOctreeElement(class UObject* ElementOwner,
 	const FOctreeElementId* ElementId = GetObjectsNavOctreeId(ElementOwner);
 	if (ElementId != NULL)
 	{
-		if (ElementId->IsValidId() && NavOctree->IsValidElementId(*ElementId))
-		{
-#if WITH_NAVIGATION_GENERATOR && NAVOCTREE_CONTAINS_COLLISION_DATA
-			// mark area occupied by given actor as dirty
-			FNavigationOctreeElement& ElementData = NavOctree->GetElementById(*ElementId);
-			const int32 DirtyFlag = GetDirtyFlagHelper(UpdateFlags, ElementData.Data.GetDirtyFlag());
-			AddDirtyArea(ElementData.Bounds.GetBox(), DirtyFlag);
-#endif // WITH_NAVIGATION_GENERATOR
-			NavOctree->RemoveNode(ElementId);
-		}
-		// @todo could use a way to "pop" an element from ObjectToOctreeId map
-		// co that above I could get a copy of FOctreeElementId while destroying 
-		// map's copy.
+		RemoveNavOctreeElementId(*ElementId, UpdateFlags);
 		RemoveObjectsNavOctreeId(ElementOwner);
 	}
 
@@ -2130,7 +2140,21 @@ void UNavigationSystem::UnregisterNavOctreeElement(class UObject* ElementOwner,
 	}
 }
 
-void UNavigationSystem::UpdateNavOctree(class AActor* Actor, int32 UpdateFlags)
+void UNavigationSystem::RemoveNavOctreeElementId(const FOctreeElementId& ElementId, int32 UpdateFlags)
+{
+	if (NavOctree->IsValidElementId(ElementId))
+	{
+#if WITH_NAVIGATION_GENERATOR && NAVOCTREE_CONTAINS_COLLISION_DATA
+		// mark area occupied by given actor as dirty
+		FNavigationOctreeElement& ElementData = NavOctree->GetElementById(ElementId);
+		const int32 DirtyFlag = GetDirtyFlagHelper(UpdateFlags, ElementData.Data.GetDirtyFlag());
+		AddDirtyArea(ElementData.Bounds.GetBox(), DirtyFlag);
+#endif // WITH_NAVIGATION_GENERATOR
+		NavOctree->RemoveNode(ElementId);
+	}
+}
+
+void UNavigationSystem::UpdateNavOctree(AActor* Actor, int32 UpdateFlags)
 {
 	INC_DWORD_STAT(STAT_Navigation_UpdateNavOctree);
 
@@ -2181,7 +2205,7 @@ void UNavigationSystem::UpdateNavOctree(class AActor* Actor, int32 UpdateFlags)
 	}
 }
 
-void UNavigationSystem::UpdateNavOctree(class UActorComponent* ActorComp, int32 UpdateFlags)
+void UNavigationSystem::UpdateNavOctree(UActorComponent* ActorComp, int32 UpdateFlags)
 {
 	INC_DWORD_STAT(STAT_Navigation_UpdateNavOctree);
 
@@ -2198,7 +2222,7 @@ void UNavigationSystem::UpdateNavOctree(class UActorComponent* ActorComp, int32
 	}	
 }
 
-void UNavigationSystem::UpdateNavOctreeElement(class UObject* ElementOwner, int32 UpdateFlags, const FBox& Bounds)
+void UNavigationSystem::UpdateNavOctreeElement(UObject* ElementOwner, int32 UpdateFlags, const FBox& Bounds)
 {
 	// grab existing octree data
 	FBox CurrentBounds;
@@ -2349,6 +2373,20 @@ void UNavigationSystem::OnNavigationBoundsUpdated(AVolume* NavVolume)
 			PopulateNavOctree();
 			bNavDataRemovedDueToMissingNavBounds = false;
 		}
+
+//#if WITH_EDITOR
+//		if (GEditor)
+//		{
+//			// this is done just in case navigation data has been removed 
+//			for (int32 Index = NavDataSet.Num() - 1; Index >= 0; --Index)
+//			{
+//				if (NavDataSet[Index] == NULL || NavDataSet[Index]->IsPendingKillPending())
+//				{
+//					NavDataSet.RemoveAtSwap(Index, 1, false);
+//				}
+//			}
+//		}
+//#endif // WITH_EDITOR
 	
 		if (NavDataSet.Num() == 0)
 		{
@@ -2381,7 +2419,8 @@ void UNavigationSystem::OnNavigationBoundsUpdated(AVolume* NavVolume)
 			for (int32 NavDataIndex = 0; NavDataIndex < NavDataSet.Num(); ++NavDataIndex)
 			{
 				ANavigationData* NavData = NavDataSet[NavDataIndex];
-				if (NavData != NULL && NavData->GetGenerator() != NULL)
+				check(NavData);
+				if (NavData->GetGenerator() != NULL)
 				{
 					if ( Cast<ARecastNavMesh>(NavData) != NULL && (bIsInGame == false || ((ARecastNavMesh*)NavData)->bRebuildAtRuntime) )
 					{
@@ -2500,7 +2539,7 @@ void UNavigationSystem::SpawnMissingNavigationData()
 	}
 }
 
-ANavigationData* UNavigationSystem::CreateNavigationDataInstance(TSubclassOf<class ANavigationData> NavDataClass, UWorld* World, const FNavDataConfig& NavConfig)
+ANavigationData* UNavigationSystem::CreateNavigationDataInstance(TSubclassOf<ANavigationData> NavDataClass, UWorld* World, const FNavDataConfig& NavConfig)
 {
 	check(World);
 
@@ -2702,14 +2741,14 @@ void UNavigationSystem::RemoveLevelCollisionFromOctree(ULevel* Level)
 	const FOctreeElementId* ElementId = GetObjectsNavOctreeId(Level);
 	if (ElementId != NULL)
 	{
-		if (ElementId->IsValidId() && NavOctree->IsValidElementId(*ElementId))
+		if (NavOctree->IsValidElementId(*ElementId))
 		{
 			// mark area occupied by given actor as dirty
 			FNavigationOctreeElement& ElementData = NavOctree->GetElementById(*ElementId);
 			AddDirtyArea(ElementData.Bounds.GetBox(), ENavigationDirtyFlag::All);
 		}
 
-		NavOctree->RemoveNode(ElementId);
+		NavOctree->RemoveNode(*ElementId);
 		RemoveObjectsNavOctreeId(Level);
 	}
 #endif // NAVOCTREE_CONTAINS_COLLISION_DATA
@@ -2860,7 +2899,7 @@ ENavigationQueryResult::Type UNavigationSystem::GetPathCost(UObject* WorldContex
 	return ENavigationQueryResult::Error;
 }
 
-ENavigationQueryResult::Type UNavigationSystem::GetPathLength(UObject* WorldContextObject, const FVector& PathStart, const FVector& PathEnd, float& OutPathLength, class ANavigationData* NavData, TSubclassOf<UNavigationQueryFilter> FilterClass)
+ENavigationQueryResult::Type UNavigationSystem::GetPathLength(UObject* WorldContextObject, const FVector& PathStart, const FVector& PathEnd, float& OutPathLength, ANavigationData* NavData, TSubclassOf<UNavigationQueryFilter> FilterClass)
 {
 	float PathLength = 0.f;
 
@@ -2893,7 +2932,7 @@ bool UNavigationSystem::IsNavigationBeingBuilt(UObject* WorldContextObject)
 //----------------------------------------------------------------------//
 // HACKS!!!
 //----------------------------------------------------------------------//
-bool UNavigationSystem::ShouldGeneratorRun(const class FNavDataGenerator* Generator) const
+bool UNavigationSystem::ShouldGeneratorRun(const FNavDataGenerator* Generator) const
 {
 #if WITH_NAVIGATION_GENERATOR
 	if (Generator != NULL)
@@ -2988,7 +3027,7 @@ bool UNavigationSystem::IsNavigationDirty() const
 {
 	for (int32 NavDataIndex=0; NavDataIndex < NavDataSet.Num(); ++NavDataIndex)
 	{
-		if (NavDataSet[NavDataIndex]->NeedsRebuild())
+		if (NavDataSet[NavDataIndex] && NavDataSet[NavDataIndex]->NeedsRebuild())
 		{
 			return true;
 		}
diff --git a/Engine/Source/Runtime/Engine/Private/AI/Navigation/RecastNavMeshGenerator.cpp b/Engine/Source/Runtime/Engine/Private/AI/Navigation/RecastNavMeshGenerator.cpp
--- a/Engine/Source/Runtime/Engine/Private/AI/Navigation/RecastNavMeshGenerator.cpp
+++ b/Engine/Source/Runtime/Engine/Private/AI/Navigation/RecastNavMeshGenerator.cpp
@@ -3418,6 +3418,28 @@ void FRecastNavMeshGenerator::SetUpGeneration(float CellSize, float CellHeight,
 	// expand bounds a bit to support later inclusion tests
 	NavBounds.ExpandBy(CellSize);
 
+	bool bClampBounds = false;
+	const float ExtentLimit = float(MAX_int32);
+	FVector BoundsExtent = NavBounds.GetExtent();
+	if (BoundsExtent.X > ExtentLimit)
+	{
+		BoundsExtent.X = ExtentLimit;
+		bClampBounds = true;
+	}
+	if (BoundsExtent.Y > ExtentLimit)
+	{
+		BoundsExtent.Y = ExtentLimit;
+		bClampBounds = true;
+	}
+
+	if (bClampBounds)
+	{
+		const FVector BoundsCenter = NavBounds.GetCenter();
+		NavBounds = FBox(BoundsCenter - BoundsExtent, BoundsCenter + BoundsExtent);
+
+		UE_LOG(LogNavigation, Warning, TEXT("Navigation bounds are too large. Cutting down every dimention down to %.f"), ExtentLimit);
+	}
+
 	UnrealNavBounds = NavBounds;
 	// now move the box to Recast coords
 	RCNavBounds = Unreal2RecastBox(NavBounds);
@@ -3589,67 +3611,6 @@ bool FRecastNavMeshGenerator::GenerateTiledNavMesh()
 		return false;	
 	}
 
-	const float* BMin = Config.bmin;
-	const float* BMax = Config.bmax;
-	int32 NewGridWidth = 0;
-	int32 NewGridHeight = 0;
-	rcCalcGridSize(BMin, BMax, Config.cs, &NewGridWidth, &NewGridHeight);
-	const int32 ts = Config.tileSize;
-	const int32 tw = (NewGridWidth + ts-1) / ts;
-	const int32 th = (NewGridHeight + ts-1) / ts;
-	if (tw * th <= 0)
-	{
-		return false;
-	}
-	else if (TilesWidth != tw || TilesHeight != th || GridWidth != NewGridWidth || GridHeight != NewGridHeight)
-	{
-		// @todo there's also a low probability number of tiles is the same but other values are different
-		// need to detect that (just in case)
-		
-		TilesWidth = tw;
-		TilesHeight = th;
-		GridWidth = NewGridWidth;
-		GridHeight = NewGridHeight;
-
-		const float tcs = Config.tileSize * Config.cs;
-
-		// Start the build process.
-		float TileBmin[3];
-		float TileBmax[3];
-
-		TileGenerationLock.Lock();
-		// do fresh start
-		++Version;
-		TileGenerators.Reset();
-		TileGenerators.AddUninitialized(th*tw);
-		TileGenerationLock.Unlock();
-
-		DestNavMesh->ReserveTileSet(TilesWidth, TilesHeight);
-		FTileSetItem* TileData = DestNavMesh->GetTileSet();
-		for (int32 TileIndex = 0; TileIndex < TilesWidth * TilesHeight; ++TileIndex, ++TileData)
-		{
-			const int32 X = TileData->X;
-			const int32 Y = TileData->Y;
-
-			TileBmin[0] = BMin[0] + X*tcs;
-			TileBmin[1] = BMin[1];
-			TileBmin[2] = BMin[2] + Y*tcs;
-
-			TileBmax[0] = BMin[0] + (X+1)*tcs;
-			TileBmax[1] = BMax[1];
-			TileBmax[2] = BMin[2] + (Y+1)*tcs;
-
-			// @todo check if this tile overlaps current navmesh generation boundaries 
-			// if not remove it
-			FRecastTileGenerator& TileGenerator = TileGenerators[ Y*TilesWidth + X ];
-			new(&TileGenerator) FRecastTileGenerator();
-			TileGenerator.Init(this, X, Y, TileBmin, TileBmax, InclusionBounds);
-
-			// reinitialize with all data
-			new(TileData) FTileSetItem(X, Y, TileGenerator.GetUnrealBB());
-		}
-	}
-
 	RebuildAll();
 
 	UE_LOG(LogNavigation, Log, TEXT("RecastNavMeshGenerator: prepare tiles for generation took %.5fs")
diff --git a/Engine/Source/Runtime/Engine/Private/Landscape/Landscape.cpp b/Engine/Source/Runtime/Engine/Private/Landscape/Landscape.cpp
--- a/Engine/Source/Runtime/Engine/Private/Landscape/Landscape.cpp
+++ b/Engine/Source/Runtime/Engine/Private/Landscape/Landscape.cpp
@@ -945,6 +945,11 @@ void ALandscapeProxy::PreEditUndo()
 	if (GIsEditor && GetWorld() && !GetWorld()->IsPlayInEditor())
 	{
 		GEngine->DeferredCommands.AddUnique(TEXT("UpdateLandscapeEditorData"));
+
+		if (GetWorld()->GetNavigationSystem())
+		{
+			GetWorld()->GetNavigationSystem()->UnregisterNavigationRelevantActor(this);
+		}
 	}
 }
 
@@ -955,6 +960,11 @@ void ALandscapeProxy::PostEditUndo()
 	if (GIsEditor && GetWorld() && !GetWorld()->IsPlayInEditor())
 	{
 		GEngine->DeferredCommands.AddUnique(TEXT("UpdateLandscapeEditorData"));
+
+		if (GetWorld()->GetNavigationSystem())
+		{
+			GetWorld()->GetNavigationSystem()->RegisterNavigationRelevantActor(this);
+		}
 	}
 }
 
diff --git a/Engine/Source/Runtime/Engine/Public/AI/NavigationOctree.h b/Engine/Source/Runtime/Engine/Public/AI/NavigationOctree.h
--- a/Engine/Source/Runtime/Engine/Public/AI/NavigationOctree.h
+++ b/Engine/Source/Runtime/Engine/Public/AI/NavigationOctree.h
@@ -5,6 +5,9 @@
 #include "GenericOctree.h"
 #include "AI/Navigation/NavigationProxy.h"
 
+class AActor;
+class UActorComponent;
+
 struct ENGINE_API FNavigationOctreeFilter
 {
 	/** pass when actor has geometry */
@@ -179,14 +182,14 @@ public:
 
 	/** @param Data allow you to access data put into navoctree while not 
 	 *	having to query navoctree about it just after adding */
-	void AddActor(class AActor* Actor, FNavigationOctreeElement& Data);
-	void AddComponent(class UActorComponent* ActorComp, const FBox& Bounds, FNavigationOctreeElement& Data);
+	void AddActor(AActor& Actor, FNavigationOctreeElement& Data);
+	void AddComponent(UActorComponent& ActorComp, const FBox& Bounds, FNavigationOctreeElement& Data);
 
 	/** Add node and update memory stats */
 	void AddNode(FNavigationOctreeElement& Data);
 
 	/** Remove node and update memory stats */
-	void RemoveNode(const FOctreeElementId* Id);
+	void RemoveNode(const FOctreeElementId& Id);
 
 	void SetNavigableGeometryStoringMode(ENavGeometryStoringMode NavGeometryMode);
 
