UDS can be based on TableRowBase - first version.
Better handling CPF_BlueprintVisible.
DataTable uses DisplayName.

diff --git a/Engine/Source/Editor/BlueprintGraph/Private/K2Node_BreakStruct.cpp b/Engine/Source/Editor/BlueprintGraph/Private/K2Node_BreakStruct.cpp
--- a/Engine/Source/Editor/BlueprintGraph/Private/K2Node_BreakStruct.cpp
+++ b/Engine/Source/Editor/BlueprintGraph/Private/K2Node_BreakStruct.cpp
@@ -220,7 +220,7 @@ void UK2Node_BreakStruct::ValidateNodeDuringCompilation(class FCompilerResultsLo
 			const UProperty* Property = *It;
 			if (CanCreatePinForProperty(Property))
 			{
-				const bool bIsBlueprintVisible = Property->HasAnyPropertyFlags(CPF_BlueprintVisible);
+				const bool bIsBlueprintVisible = Property->HasAnyPropertyFlags(CPF_BlueprintVisible) || (Property->GetOwnerStruct() && Property->GetOwnerStruct()->IsA<UUserDefinedStruct>());
 				bHasAnyBlueprintVisibleProperty |= bIsBlueprintVisible;
 
 				const UEdGraphPin* Pin = Property ? FindPin(Property->GetName()) : NULL;
diff --git a/Engine/Source/Editor/BlueprintGraph/Private/K2Node_MakeStruct.cpp b/Engine/Source/Editor/BlueprintGraph/Private/K2Node_MakeStruct.cpp
--- a/Engine/Source/Editor/BlueprintGraph/Private/K2Node_MakeStruct.cpp
+++ b/Engine/Source/Editor/BlueprintGraph/Private/K2Node_MakeStruct.cpp
@@ -124,7 +124,7 @@ void UK2Node_MakeStruct::ValidateNodeDuringCompilation(class FCompilerResultsLog
 			const UProperty* Property = *It;
 			if (CanBeExposed(Property))
 			{
-				const bool bIsBlueprintVisible = Property->HasAnyPropertyFlags(CPF_BlueprintVisible);
+				const bool bIsBlueprintVisible = Property->HasAnyPropertyFlags(CPF_BlueprintVisible) || (Property->GetOwnerStruct() && Property->GetOwnerStruct()->IsA<UUserDefinedStruct>());
 				bHasAnyBlueprintVisibleProperty |= bIsBlueprintVisible;
 
 				const UEdGraphPin* Pin = Property ? FindPin(Property->GetName()) : NULL;
diff --git a/Engine/Source/Editor/KismetCompiler/Private/UserDefinedStructureCompilerUtils.cpp b/Engine/Source/Editor/KismetCompiler/Private/UserDefinedStructureCompilerUtils.cpp
--- a/Engine/Source/Editor/KismetCompiler/Private/UserDefinedStructureCompilerUtils.cpp
+++ b/Engine/Source/Editor/KismetCompiler/Private/UserDefinedStructureCompilerUtils.cpp
@@ -146,7 +146,7 @@ struct FUserDefinedStructureCompilerInner
 				continue;
 			}
 
-			NewProperty->SetPropertyFlags(CPF_Edit);
+			NewProperty->SetPropertyFlags(CPF_Edit | CPF_BlueprintVisible);
 			if (VarDesc.bDontEditoOnInstance)
 			{
 				NewProperty->SetPropertyFlags(CPF_DisableEditOnInstance);
@@ -176,6 +176,9 @@ struct FUserDefinedStructureCompilerInner
 
 		Struct->SetMetaData(FBlueprintMetadata::MD_Tooltip, *FStructureEditorUtils::GetTooltip(Struct));
 
+		auto EditorData = CastChecked<UUserDefinedStructEditorData>(Struct->EditorData);
+		Struct->SetSuperStruct(EditorData->NativeBase);
+
 		CreateVariables(Struct, K2Schema, MessageLog);
 
 		Struct->Bind();
diff --git a/Engine/Source/Editor/UnrealEd/Classes/Factories/StructureFactory.h b/Engine/Source/Editor/UnrealEd/Classes/Factories/StructureFactory.h
--- a/Engine/Source/Editor/UnrealEd/Classes/Factories/StructureFactory.h
+++ b/Engine/Source/Editor/UnrealEd/Classes/Factories/StructureFactory.h
@@ -9,7 +9,11 @@ class UStructureFactory : public UFactory
 {
 	GENERATED_UCLASS_BODY()
 
+	UPROPERTY()
+	class UScriptStruct* Base;
+
 	// Begin UFactory Interface
+	virtual bool ConfigureProperties() override;
 	virtual UObject* FactoryCreateNew(UClass* Class,UObject* InParent,FName Name,EObjectFlags Flags,UObject* Context,FFeedbackContext* Warn) override;
 	// Begin UFactory Interface
 };
diff --git a/Engine/Source/Editor/UnrealEd/Classes/UserDefinedStructure/UserDefinedStructEditorData.h b/Engine/Source/Editor/UnrealEd/Classes/UserDefinedStructure/UserDefinedStructEditorData.h
--- a/Engine/Source/Editor/UnrealEd/Classes/UserDefinedStructure/UserDefinedStructEditorData.h
+++ b/Engine/Source/Editor/UnrealEd/Classes/UserDefinedStructure/UserDefinedStructEditorData.h
@@ -79,6 +79,9 @@ public:
 	UPROPERTY()
 	FString ToolTip;
 
+	UPROPERTY()
+	UScriptStruct* NativeBase;
+
 	// UObject interface.
 	virtual void PostEditUndo() override;
 	virtual void PostLoadSubobjects(struct FObjectInstancingGraph* OuterInstanceGraph) override;
diff --git a/Engine/Source/Editor/UnrealEd/Private/Factories/CSVImportFactory.cpp b/Engine/Source/Editor/UnrealEd/Private/Factories/CSVImportFactory.cpp
--- a/Engine/Source/Editor/UnrealEd/Private/Factories/CSVImportFactory.cpp
+++ b/Engine/Source/Editor/UnrealEd/Private/Factories/CSVImportFactory.cpp
@@ -104,7 +104,9 @@ public:
 			{
 				UScriptStruct* Struct = *It;
 				// If a child of the table row struct base, but not itself
-				if (Struct->IsChildOf(TableRowStruct) && Struct != TableRowStruct)
+				if (Struct->IsChildOf(TableRowStruct) 
+					&& (Struct != TableRowStruct) 
+					&& (Struct->GetOutermost() != GetTransientPackage()))
 				{
 					RowStructs.Add(Struct);
 				}
diff --git a/Engine/Source/Editor/UnrealEd/Private/Factories/EditorFactories.cpp b/Engine/Source/Editor/UnrealEd/Private/Factories/EditorFactories.cpp
--- a/Engine/Source/Editor/UnrealEd/Private/Factories/EditorFactories.cpp
+++ b/Engine/Source/Editor/UnrealEd/Private/Factories/EditorFactories.cpp
@@ -7034,10 +7034,21 @@ UStructureFactory::UStructureFactory(const class FPostConstructInitializePropert
 	bEditAfterNew = true;
 }
 
+bool UStructureFactory::ConfigureProperties()
+{
+	static FBoolConfigValueHelper CanBeBasedOnTableRow(TEXT("UserDefinedStructure"), TEXT("bCanBeBasedOnTableRow"));
+	if (CanBeBasedOnTableRow)
+	{
+		const auto ReturnCode = FMessageDialog::Open(EAppMsgType::YesNo, LOCTEXT("BasedOnTableRowBase", "Is this structure based on TableRowBase?"));
+		Base = (EAppReturnType::Yes == ReturnCode) ? FTableRowBase::StaticStruct() : NULL;
+	}
+	return true;
+}
+
 UObject* UStructureFactory::FactoryCreateNew(UClass* Class, UObject* InParent, FName Name, EObjectFlags Flags, UObject* Context, FFeedbackContext* Warn)
 {
 	ensure(UUserDefinedStruct::StaticClass() == Class);
-	return FStructureEditorUtils::CreateUserDefinedStruct(InParent, Name, Flags);
+	return FStructureEditorUtils::CreateUserDefinedStruct(InParent, Name, Flags, Base);
 }
 
 /*-----------------------------------------------------------------------------
diff --git a/Engine/Source/Editor/UnrealEd/Private/Kismet2/StructureEditorUtils.cpp b/Engine/Source/Editor/UnrealEd/Private/Kismet2/StructureEditorUtils.cpp
--- a/Engine/Source/Editor/UnrealEd/Private/Kismet2/StructureEditorUtils.cpp
+++ b/Engine/Source/Editor/UnrealEd/Private/Kismet2/StructureEditorUtils.cpp
@@ -26,7 +26,7 @@ FStructureEditorUtils::FStructEditorManager& FStructureEditorUtils::FStructEdito
 
 //////////////////////////////////////////////////////////////////////////
 // FStructureEditorUtils
-UUserDefinedStruct* FStructureEditorUtils::CreateUserDefinedStruct(UObject* InParent, FName Name, EObjectFlags Flags)
+UUserDefinedStruct* FStructureEditorUtils::CreateUserDefinedStruct(UObject* InParent, FName Name, EObjectFlags Flags, UScriptStruct* BaseStruct)
 {
 	UUserDefinedStruct* Struct = NULL;
 	
@@ -36,6 +36,11 @@ UUserDefinedStruct* FStructureEditorUtils::CreateUserDefinedStruct(UObject* InPa
 		check(Struct);
 		Struct->EditorData = NewNamedObject<UUserDefinedStructEditorData>(Struct, NAME_None, RF_Transactional);
 		check(Struct->EditorData);
+
+		check(IsPropertBaseStruct(BaseStruct));
+		CastChecked<UUserDefinedStructEditorData>(Struct->EditorData)->NativeBase = BaseStruct;
+		Struct->SetSuperStruct(BaseStruct);
+
 		Struct->SetMetaData(TEXT("BlueprintType"), TEXT("true"));
 		Struct->Bind();
 		Struct->StaticLink(true);
@@ -687,4 +692,21 @@ FGuid FStructureEditorUtils::GetGuidFromPropertyName(const FName Name)
 	return FMemberVariableNameHelper::GetGuidFromName(Name);
 }
 
+bool FStructureEditorUtils::IsPropertBaseStruct(UScriptStruct* Base)
+{
+	return !Base || (FTableRowBase::StaticStruct() == Base);
+}
+
+bool FStructureEditorUtils::SetNativeBase(UUserDefinedStruct* Struct, UScriptStruct* NativeBase)
+{
+	bool bResult = false;
+	if (Struct && IsPropertBaseStruct(NativeBase))
+	{
+		CastChecked<UUserDefinedStructEditorData>(Struct->EditorData)->NativeBase = NativeBase;
+		OnStructureChanged(Struct);
+		bResult = true;
+	}
+	return bResult;
+}
+
 #undef LOCTEXT_NAMESPACE
diff --git a/Engine/Source/Editor/UnrealEd/Public/Kismet2/StructureEditorUtils.h b/Engine/Source/Editor/UnrealEd/Public/Kismet2/StructureEditorUtils.h
--- a/Engine/Source/Editor/UnrealEd/Public/Kismet2/StructureEditorUtils.h
+++ b/Engine/Source/Editor/UnrealEd/Public/Kismet2/StructureEditorUtils.h
@@ -44,7 +44,7 @@ public:
 	};
 
 	//STRUCTURE
-	static UUserDefinedStruct* CreateUserDefinedStruct(UObject* InParent, FName Name, EObjectFlags Flags);
+	static UUserDefinedStruct* CreateUserDefinedStruct(UObject* InParent, FName Name, EObjectFlags Flags, UScriptStruct* BaseStruct = NULL);
 
 	static void CompileStructure(UUserDefinedStruct* Struct);
 
@@ -52,6 +52,10 @@ public:
 
 	static bool ChangeTooltip(UUserDefinedStruct* Struct, const FString& InTooltip);
 
+	static bool SetNativeBase(UUserDefinedStruct* Struct, UScriptStruct* NativeBase);
+
+	static bool IsPropertBaseStruct(UScriptStruct* Base);
+
 	//VARIABLE
 	static bool AddVariable(UUserDefinedStruct* Struct, const FEdGraphPinType& VarType);
 
diff --git a/Engine/Source/Runtime/Engine/Private/DataTable.cpp b/Engine/Source/Runtime/Engine/Private/DataTable.cpp
--- a/Engine/Source/Runtime/Engine/Private/DataTable.cpp
+++ b/Engine/Source/Runtime/Engine/Private/DataTable.cpp
@@ -347,6 +347,8 @@ UProperty* UDataTable::FindTableProperty(const FName& PropertyName) const
 /** Get array of UProperties that corresponds to columns in the table */
 TArray<UProperty*> UDataTable::GetTablePropertyArray(const FString& FirstRowString, UStruct* InRowStruct, TArray<FString>& OutProblems)
 {
+	static const FName DisplayNameKey(TEXT("DisplayName"));
+
 	TArray<UProperty*> ColumnProps;
 
 	// Get list of all expected properties from the struct
@@ -373,6 +375,13 @@ TArray<UProperty*> UDataTable::GetTablePropertyArray(const FString& FirstRowStri
 			else
 			{
 				UProperty* ColumnProp = FindField<UProperty>(InRowStruct, PropName);
+
+				for (TFieldIterator<UProperty> It(InRowStruct); It && !ColumnProp; ++It)
+				{
+					const FString DisplayName = *It ? (*It)->GetMetaData(DisplayNameKey) : FString();
+					ColumnProp = (!DisplayName.IsEmpty() && (DisplayName == ColumnNameStrings[ColIdx])) ? *It : NULL;
+				}
+
 				// Didn't find a property with this name, problem..
 				if(ColumnProp == NULL)
 				{
@@ -398,7 +407,7 @@ TArray<UProperty*> UDataTable::GetTablePropertyArray(const FString& FirstRowStri
 					}
 
 					// Track that we found this one
-					ExpectedPropNames.Remove(PropName);
+					ExpectedPropNames.Remove(ColumnProp->GetFName());
 				}
 			}
 		}
@@ -407,7 +416,11 @@ TArray<UProperty*> UDataTable::GetTablePropertyArray(const FString& FirstRowStri
 	// Generate warning for any properties in struct we are not filling in
 	for(int32 PropIdx=0; PropIdx < ExpectedPropNames.Num(); PropIdx++)
 	{
-		OutProblems.Add(FString::Printf(TEXT("Expected column '%s' not found in input."), *ExpectedPropNames[PropIdx].ToString()));
+		const UProperty* const ColumnProp = FindField<UProperty>(InRowStruct, ExpectedPropNames[PropIdx]);
+		const FString DisplayName = (ColumnProp && ColumnProp->HasMetaData(DisplayNameKey)) 
+			? ColumnProp->GetMetaData(DisplayNameKey) 
+			: ExpectedPropNames[PropIdx].ToString();
+		OutProblems.Add(FString::Printf(TEXT("Expected column '%s' not found in input."), *DisplayName));
 	}
 
 	return ColumnProps;
@@ -526,13 +539,17 @@ TArray<FString> UDataTable::CreateTableFromJSONString(const FString& InString)
 
 TArray<FString> UDataTable::GetColumnTitles() const
 {
+	static const FName DisplayNameKey(TEXT("DisplayName"));
 	TArray<FString> Result;
 	Result.Add(TEXT("Name"));
 	for (TFieldIterator<UProperty> It(RowStruct); It; ++It)
 	{
 		UProperty* Prop = *It;
 		check(Prop != NULL);
-		Result.Add(Prop->GetName());
+		const FString DisplayName = Prop->HasMetaData(DisplayNameKey)
+			? Prop->GetMetaData(DisplayNameKey)
+			: Prop->GetName();
+		Result.Add(DisplayName);
 	}
 	return Result;
 }
