File missing in prev submission #UE4

diff --git a/Engine/Source/Runtime/Engine/Classes/AI/Navigation/NavigationTypes.h b/Engine/Source/Runtime/Engine/Classes/AI/Navigation/NavigationTypes.h
--- a/Engine/Source/Runtime/Engine/Classes/AI/Navigation/NavigationTypes.h
+++ b/Engine/Source/Runtime/Engine/Classes/AI/Navigation/NavigationTypes.h
@@ -250,6 +250,15 @@ namespace EPathObservationResult
 	};
 }
 
+namespace ENavAreaEvent
+{
+	enum Type
+	{
+		Registered,
+		Unregistered
+	};
+}
+
 typedef TSharedRef<struct FNavigationPath, ESPMode::ThreadSafe> FNavPathSharedRef;
 typedef TSharedPtr<struct FNavigationPath, ESPMode::ThreadSafe> FNavPathSharedPtr;
 typedef TWeakPtr<struct FNavigationPath, ESPMode::ThreadSafe> FNavPathWeakPtr;
