CIS fix

diff --git a/Engine/Source/Runtime/Engine/Private/DataTable.cpp b/Engine/Source/Runtime/Engine/Private/DataTable.cpp
--- a/Engine/Source/Runtime/Engine/Private/DataTable.cpp
+++ b/Engine/Source/Runtime/Engine/Private/DataTable.cpp
@@ -378,7 +378,11 @@ TArray<UProperty*> UDataTable::GetTablePropertyArray(const FString& FirstRowStri
 
 				for (TFieldIterator<UProperty> It(InRowStruct); It && !ColumnProp; ++It)
 				{
+#if WITH_EDITOR || HACK_HEADER_GENERATOR
 					const FString DisplayName = *It ? (*It)->GetMetaData(DisplayNameKey) : FString();
+#else
+					const FString DisplayName = TEXT("");
+#endif
 					ColumnProp = (!DisplayName.IsEmpty() && (DisplayName == ColumnNameStrings[ColIdx])) ? *It : NULL;
 				}
 
@@ -417,9 +421,13 @@ TArray<UProperty*> UDataTable::GetTablePropertyArray(const FString& FirstRowStri
 	for(int32 PropIdx=0; PropIdx < ExpectedPropNames.Num(); PropIdx++)
 	{
 		const UProperty* const ColumnProp = FindField<UProperty>(InRowStruct, ExpectedPropNames[PropIdx]);
-		const FString DisplayName = (ColumnProp && ColumnProp->HasMetaData(DisplayNameKey)) 
+#if WITH_EDITOR || HACK_HEADER_GENERATOR
+		const FString DisplayName = (ColumnProp && ColumnProp->HasMetaData(DisplayNameKey))
 			? ColumnProp->GetMetaData(DisplayNameKey) 
 			: ExpectedPropNames[PropIdx].ToString();
+#else
+		const FString DisplayName = ExpectedPropNames[PropIdx].ToString();
+#endif
 		OutProblems.Add(FString::Printf(TEXT("Expected column '%s' not found in input."), *DisplayName));
 	}
 
@@ -546,9 +554,13 @@ TArray<FString> UDataTable::GetColumnTitles() const
 	{
 		UProperty* Prop = *It;
 		check(Prop != NULL);
+#if WITH_EDITOR || HACK_HEADER_GENERATOR
 		const FString DisplayName = Prop->HasMetaData(DisplayNameKey)
 			? Prop->GetMetaData(DisplayNameKey)
 			: Prop->GetName();
+#else
+		const FString DisplayName = Prop->GetName();
+#endif
 		Result.Add(DisplayName);
 	}
 	return Result;
