Remove SoundNode GetUniqueString

diff --git a/Engine/Plugins/Runtime/SoundMod/Source/SoundMod/Private/SoundNodeModPlayer.cpp b/Engine/Plugins/Runtime/SoundMod/Source/SoundMod/Private/SoundNodeModPlayer.cpp
--- a/Engine/Plugins/Runtime/SoundMod/Source/SoundMod/Private/SoundNodeModPlayer.cpp
+++ b/Engine/Plugins/Runtime/SoundMod/Source/SoundMod/Private/SoundNodeModPlayer.cpp
@@ -83,11 +83,6 @@ FString USoundNodeModPlayer::GetTitle() const
 }
 #endif
 
-FString USoundNodeModPlayer::GetUniqueString() const
-{
-	return TEXT("");
-}
-
 // A Mod Player is the end of the chain and has no children
 int32 USoundNodeModPlayer::GetMaxChildNodes() const
 {
diff --git a/Engine/Plugins/Runtime/SoundMod/Source/SoundMod/Public/SoundNodeModPlayer.h b/Engine/Plugins/Runtime/SoundMod/Source/SoundMod/Public/SoundNodeModPlayer.h
--- a/Engine/Plugins/Runtime/SoundMod/Source/SoundMod/Public/SoundNodeModPlayer.h
+++ b/Engine/Plugins/Runtime/SoundMod/Source/SoundMod/Public/SoundNodeModPlayer.h
@@ -24,7 +24,6 @@ public:
 	virtual int32 GetMaxChildNodes() const override;
 	virtual float GetDuration() override;
 	virtual void ParseNodes(FAudioDevice* AudioDevice, const UPTRINT NodeWaveInstanceHash, FActiveSound& ActiveSound, const FSoundParseParameters& ParseParams, TArray<FWaveInstance*>& WaveInstances) override;
-	virtual FString GetUniqueString() const override;
 #if WITH_EDITOR
 	virtual FString GetTitle() const override;
 #endif
diff --git a/Engine/Source/Runtime/Engine/Classes/Sound/SoundNode.h b/Engine/Source/Runtime/Engine/Classes/Sound/SoundNode.h
--- a/Engine/Source/Runtime/Engine/Classes/Sound/SoundNode.h
+++ b/Engine/Source/Runtime/Engine/Classes/Sound/SoundNode.h
@@ -162,7 +162,6 @@ public:
 	/** 
 	 * Used to create a unique string to identify unique nodes
 	 */
-	virtual FString GetUniqueString() const;
 	static UPTRINT GetNodeWaveInstanceHash(const UPTRINT ParentWaveInstanceHash, const USoundNode* ChildNode, const uint32 ChildIndex);
 	static UPTRINT GetNodeWaveInstanceHash(const UPTRINT ParentWaveInstanceHash, const UPTRINT ChildNodeHash, const uint32 ChildIndex);
 };
diff --git a/Engine/Source/Runtime/Engine/Classes/Sound/SoundNodeAttenuation.h b/Engine/Source/Runtime/Engine/Classes/Sound/SoundNodeAttenuation.h
--- a/Engine/Source/Runtime/Engine/Classes/Sound/SoundNodeAttenuation.h
+++ b/Engine/Source/Runtime/Engine/Classes/Sound/SoundNodeAttenuation.h
@@ -61,7 +61,6 @@ public:
 	// Begin USoundNode interface. 
 	virtual void ParseNodes( class FAudioDevice* AudioDevice, const UPTRINT NodeWaveInstanceHash, FActiveSound& ActiveSound, const FSoundParseParameters& ParseParams, TArray<FWaveInstance*>& WaveInstances ) override;
 	virtual float MaxAudibleDistance( float CurrentMaxDistance ) override;
-	virtual FString GetUniqueString() const override;
 	// End USoundNode interface. 
 
 	ENGINE_API FAttenuationSettings* GetAttenuationSettingsToApply();
diff --git a/Engine/Source/Runtime/Engine/Classes/Sound/SoundNodeConcatenator.h b/Engine/Source/Runtime/Engine/Classes/Sound/SoundNodeConcatenator.h
--- a/Engine/Source/Runtime/Engine/Classes/Sound/SoundNodeConcatenator.h
+++ b/Engine/Source/Runtime/Engine/Classes/Sound/SoundNodeConcatenator.h
@@ -38,7 +38,6 @@ public:
 	/** Ensure amount of inputs matches new amount of children */
 	virtual void SetChildNodes(TArray<USoundNode*>& InChildNodes) override;
 #endif
-	virtual FString GetUniqueString() const override;
 	// End USoundNode interface. 
 };
 
diff --git a/Engine/Source/Runtime/Engine/Classes/Sound/SoundNodeDelay.h b/Engine/Source/Runtime/Engine/Classes/Sound/SoundNodeDelay.h
--- a/Engine/Source/Runtime/Engine/Classes/Sound/SoundNodeDelay.h
+++ b/Engine/Source/Runtime/Engine/Classes/Sound/SoundNodeDelay.h
@@ -24,7 +24,6 @@ public:
 	// Begin USoundNode interface. 
 	virtual void ParseNodes( FAudioDevice* AudioDevice, const UPTRINT NodeWaveInstanceHash, FActiveSound& ActiveSound, const FSoundParseParameters& ParseParams, TArray<FWaveInstance*>& WaveInstances ) override;
 	virtual float GetDuration( void ) override;
-	virtual FString GetUniqueString() const override;
 	// End USoundNode interface. 
 };
 
diff --git a/Engine/Source/Runtime/Engine/Classes/Sound/SoundNodeDialoguePlayer.h b/Engine/Source/Runtime/Engine/Classes/Sound/SoundNodeDialoguePlayer.h
--- a/Engine/Source/Runtime/Engine/Classes/Sound/SoundNodeDialoguePlayer.h
+++ b/Engine/Source/Runtime/Engine/Classes/Sound/SoundNodeDialoguePlayer.h
@@ -25,7 +25,6 @@ public:
 	virtual int32 GetMaxChildNodes() const override;
 	virtual float GetDuration() override;
 	virtual void ParseNodes( FAudioDevice* AudioDevice, const UPTRINT NodeWaveInstanceHash, FActiveSound& ActiveSound, const FSoundParseParameters& ParseParams, TArray<FWaveInstance*>& WaveInstances ) override;
-	virtual FString GetUniqueString() const override;
 	// End USoundNode Interface
 
 };
diff --git a/Engine/Source/Runtime/Engine/Classes/Sound/SoundNodeDistanceCrossFade.h b/Engine/Source/Runtime/Engine/Classes/Sound/SoundNodeDistanceCrossFade.h
--- a/Engine/Source/Runtime/Engine/Classes/Sound/SoundNodeDistanceCrossFade.h
+++ b/Engine/Source/Runtime/Engine/Classes/Sound/SoundNodeDistanceCrossFade.h
@@ -98,7 +98,6 @@ public:
 	virtual void SetChildNodes(TArray<USoundNode*>& InChildNodes) override;
 #endif //WITH_EDITOR
 	virtual float MaxAudibleDistance( float CurrentMaxDistance ) override;
-	virtual FString GetUniqueString() const override;
 	// End USoundNode interface. 
 
 	virtual float GetCurrentDistance(FAudioDevice* AudioDevice, FActiveSound& ActiveSound, const FSoundParseParameters& ParseParams) const;
diff --git a/Engine/Source/Runtime/Engine/Classes/Sound/SoundNodeDoppler.h b/Engine/Source/Runtime/Engine/Classes/Sound/SoundNodeDoppler.h
--- a/Engine/Source/Runtime/Engine/Classes/Sound/SoundNodeDoppler.h
+++ b/Engine/Source/Runtime/Engine/Classes/Sound/SoundNodeDoppler.h
@@ -22,7 +22,6 @@ class USoundNodeDoppler : public USoundNode
 public:
 	// Begin USoundNode interface. 
 	virtual void ParseNodes( FAudioDevice* AudioDevice, const UPTRINT NodeWaveInstanceHash, FActiveSound& ActiveSound, const FSoundParseParameters& ParseParams, TArray<FWaveInstance*>& WaveInstances ) override;
-	virtual FString GetUniqueString() const override;
 	// End USoundNode interface. 
 
 protected:
diff --git a/Engine/Source/Runtime/Engine/Classes/Sound/SoundNodeEnveloper.h b/Engine/Source/Runtime/Engine/Classes/Sound/SoundNodeEnveloper.h
--- a/Engine/Source/Runtime/Engine/Classes/Sound/SoundNodeEnveloper.h
+++ b/Engine/Source/Runtime/Engine/Classes/Sound/SoundNodeEnveloper.h
@@ -79,7 +79,6 @@ public:
 
 	// Begin USoundNode interface. 
 	virtual void ParseNodes( FAudioDevice* AudioDevice, const UPTRINT NodeWaveInstanceHash, FActiveSound& ActiveSound, const FSoundParseParameters& ParseParams, TArray<FWaveInstance*>& WaveInstances ) override;
-	virtual FString GetUniqueString() const override;
 	virtual float GetDuration( void ) override;
 	// End USoundNode interface. 
 
diff --git a/Engine/Source/Runtime/Engine/Classes/Sound/SoundNodeGroupControl.h b/Engine/Source/Runtime/Engine/Classes/Sound/SoundNodeGroupControl.h
--- a/Engine/Source/Runtime/Engine/Classes/Sound/SoundNodeGroupControl.h
+++ b/Engine/Source/Runtime/Engine/Classes/Sound/SoundNodeGroupControl.h
@@ -32,7 +32,6 @@ public:
 	virtual void SetChildNodes(TArray<USoundNode*>& InChildNodes) override;
 #endif //WITH_EDITOR
 	virtual void CreateStartingConnectors() override;
-	virtual FString GetUniqueString() const override;
 	// End USoundNode interface.
 
 private:
diff --git a/Engine/Source/Runtime/Engine/Classes/Sound/SoundNodeLooping.h b/Engine/Source/Runtime/Engine/Classes/Sound/SoundNodeLooping.h
--- a/Engine/Source/Runtime/Engine/Classes/Sound/SoundNodeLooping.h
+++ b/Engine/Source/Runtime/Engine/Classes/Sound/SoundNodeLooping.h
@@ -22,7 +22,6 @@ public:
 	}
 	virtual float GetDuration( void ) override;
 	virtual void ParseNodes( FAudioDevice* AudioDevice, const UPTRINT NodeWaveInstanceHash, FActiveSound& ActiveSound, const FSoundParseParameters& ParseParams, TArray<FWaveInstance*>& WaveInstances ) override;
-	virtual FString GetUniqueString() const override;
 	// End USoundNode interface. 
 };
 
diff --git a/Engine/Source/Runtime/Engine/Classes/Sound/SoundNodeMature.h b/Engine/Source/Runtime/Engine/Classes/Sound/SoundNodeMature.h
--- a/Engine/Source/Runtime/Engine/Classes/Sound/SoundNodeMature.h
+++ b/Engine/Source/Runtime/Engine/Classes/Sound/SoundNodeMature.h
@@ -24,7 +24,6 @@ public:
 	virtual void ParseNodes( FAudioDevice* AudioDevice, const UPTRINT NodeWaveInstanceHash, FActiveSound& ActiveSound, const FSoundParseParameters& ParseParams, TArray<FWaveInstance*>& WaveInstances ) override;
 	virtual void CreateStartingConnectors( void ) override;
 	virtual int32 GetMaxChildNodes() const override;
-	virtual FString GetUniqueString() const override;
 	// End USoundNode interface.
 };
 
diff --git a/Engine/Source/Runtime/Engine/Classes/Sound/SoundNodeMixer.h b/Engine/Source/Runtime/Engine/Classes/Sound/SoundNodeMixer.h
--- a/Engine/Source/Runtime/Engine/Classes/Sound/SoundNodeMixer.h
+++ b/Engine/Source/Runtime/Engine/Classes/Sound/SoundNodeMixer.h
@@ -32,7 +32,6 @@ public:
 	/** Ensure amount of inputs matches new amount of children */
 	virtual void SetChildNodes(TArray<USoundNode*>& InChildNodes) override;
 #endif //WITH_EDITOR
-	virtual FString GetUniqueString() const override;
 	// End USoundNode interface.
 };
 
diff --git a/Engine/Source/Runtime/Engine/Classes/Sound/SoundNodeModulator.h b/Engine/Source/Runtime/Engine/Classes/Sound/SoundNodeModulator.h
--- a/Engine/Source/Runtime/Engine/Classes/Sound/SoundNodeModulator.h
+++ b/Engine/Source/Runtime/Engine/Classes/Sound/SoundNodeModulator.h
@@ -32,7 +32,6 @@ class USoundNodeModulator : public USoundNode
 public:
 	// Begin USoundNode interface. 
 	virtual void ParseNodes( FAudioDevice* AudioDevice, const UPTRINT NodeWaveInstanceHash, FActiveSound& ActiveSound, const FSoundParseParameters& ParseParams, TArray<FWaveInstance*>& WaveInstances ) override;
-	virtual FString GetUniqueString() const override;
 	// End USoundNode interface. 
 };
 
diff --git a/Engine/Source/Runtime/Engine/Classes/Sound/SoundNodeModulatorContinuous.h b/Engine/Source/Runtime/Engine/Classes/Sound/SoundNodeModulatorContinuous.h
--- a/Engine/Source/Runtime/Engine/Classes/Sound/SoundNodeModulatorContinuous.h
+++ b/Engine/Source/Runtime/Engine/Classes/Sound/SoundNodeModulatorContinuous.h
@@ -90,7 +90,6 @@ public:
 
 	// Begin USoundNode interface.
 	virtual void ParseNodes( FAudioDevice* AudioDevice, const UPTRINT NodeWaveInstanceHash, FActiveSound& ActiveSound, const FSoundParseParameters& ParseParams, TArray<FWaveInstance*>& WaveInstances ) override;
-	virtual FString GetUniqueString() const override;
 	// End USoundNode interface.
 };
 
diff --git a/Engine/Source/Runtime/Engine/Classes/Sound/SoundNodeOscillator.h b/Engine/Source/Runtime/Engine/Classes/Sound/SoundNodeOscillator.h
--- a/Engine/Source/Runtime/Engine/Classes/Sound/SoundNodeOscillator.h
+++ b/Engine/Source/Runtime/Engine/Classes/Sound/SoundNodeOscillator.h
@@ -56,7 +56,6 @@ class USoundNodeOscillator : public USoundNode
 public:	
 	// Begin USoundNode interface. 
 	virtual void ParseNodes( FAudioDevice* AudioDevice, const UPTRINT NodeWaveInstanceHash, FActiveSound& ActiveSound, const FSoundParseParameters& ParseParams, TArray<FWaveInstance*>& WaveInstances ) override;
-	virtual FString GetUniqueString() const override;
 	// End USoundNode interface. 
 };
 
diff --git a/Engine/Source/Runtime/Engine/Classes/Sound/SoundNodeParamCrossFade.h b/Engine/Source/Runtime/Engine/Classes/Sound/SoundNodeParamCrossFade.h
--- a/Engine/Source/Runtime/Engine/Classes/Sound/SoundNodeParamCrossFade.h
+++ b/Engine/Source/Runtime/Engine/Classes/Sound/SoundNodeParamCrossFade.h
@@ -18,7 +18,6 @@ public:
 	UPROPERTY(EditAnywhere, Category=CrossFade )
 	FName ParamName;
 
-	virtual FString GetUniqueString() const override;
 	virtual float GetCurrentDistance(FAudioDevice* AudioDevice, FActiveSound& ActiveSound, const FSoundParseParameters& ParseParams) const override;
 	virtual bool AllowCrossfading(FActiveSound& ActiveSound) const override;
 };
diff --git a/Engine/Source/Runtime/Engine/Classes/Sound/SoundNodeRandom.h b/Engine/Source/Runtime/Engine/Classes/Sound/SoundNodeRandom.h
--- a/Engine/Source/Runtime/Engine/Classes/Sound/SoundNodeRandom.h
+++ b/Engine/Source/Runtime/Engine/Classes/Sound/SoundNodeRandom.h
@@ -73,7 +73,6 @@ public:
 	virtual void OnBeginPIE(const bool bIsSimulating) override;
 #endif //WITH_EDITOR
 	virtual void CreateStartingConnectors( void ) override;
-	virtual FString GetUniqueString() const override;
 	// End USoundNode interface.
 
 	// @todo document
diff --git a/Engine/Source/Runtime/Engine/Classes/Sound/SoundNodeSoundClass.h b/Engine/Source/Runtime/Engine/Classes/Sound/SoundNodeSoundClass.h
--- a/Engine/Source/Runtime/Engine/Classes/Sound/SoundNodeSoundClass.h
+++ b/Engine/Source/Runtime/Engine/Classes/Sound/SoundNodeSoundClass.h
@@ -19,6 +19,5 @@ class USoundNodeSoundClass : public USoundNode
 public:
 	// Begin USoundNode interface. 
 	virtual void ParseNodes( class FAudioDevice* AudioDevice, const UPTRINT NodeWaveInstanceHash, FActiveSound& ActiveSound, const FSoundParseParameters& ParseParams, TArray<FWaveInstance*>& WaveInstances ) override;
-	virtual FString GetUniqueString() const override;
 	// End USoundNode interface. 
 };
diff --git a/Engine/Source/Runtime/Engine/Classes/Sound/SoundNodeWaveParam.h b/Engine/Source/Runtime/Engine/Classes/Sound/SoundNodeWaveParam.h
--- a/Engine/Source/Runtime/Engine/Classes/Sound/SoundNodeWaveParam.h
+++ b/Engine/Source/Runtime/Engine/Classes/Sound/SoundNodeWaveParam.h
@@ -21,7 +21,6 @@ public:
 	// Begin USoundNode Interface
 	virtual float GetDuration( void ) override;
 	virtual void ParseNodes( FAudioDevice* AudioDevice, const UPTRINT NodeWaveInstanceHash, FActiveSound& ActiveSound, const FSoundParseParameters& ParseParams, TArray<FWaveInstance*>& WaveInstances ) override;
-	virtual FString GetUniqueString() const override;
 	// End USoundNode Interface
 };
 
diff --git a/Engine/Source/Runtime/Engine/Classes/Sound/SoundNodeWavePlayer.h b/Engine/Source/Runtime/Engine/Classes/Sound/SoundNodeWavePlayer.h
--- a/Engine/Source/Runtime/Engine/Classes/Sound/SoundNodeWavePlayer.h
+++ b/Engine/Source/Runtime/Engine/Classes/Sound/SoundNodeWavePlayer.h
@@ -24,7 +24,6 @@ public:
 	virtual int32 GetMaxChildNodes() const override;
 	virtual float GetDuration() override;
 	virtual void ParseNodes( FAudioDevice* AudioDevice, const UPTRINT NodeWaveInstanceHash, FActiveSound& ActiveSound, const FSoundParseParameters& ParseParams, TArray<FWaveInstance*>& WaveInstances ) override;
-	virtual FString GetUniqueString() const override;
 #if WITH_EDITOR
 	virtual FString GetTitle() const override;
 #endif
diff --git a/Engine/Source/Runtime/Engine/Private/AudioDevice.cpp b/Engine/Source/Runtime/Engine/Private/AudioDevice.cpp
--- a/Engine/Source/Runtime/Engine/Private/AudioDevice.cpp
+++ b/Engine/Source/Runtime/Engine/Private/AudioDevice.cpp
@@ -416,70 +416,6 @@ void FAudioDevice::ShowSoundClassHierarchy( FOutputDevice& Ar, USoundClass* InSo
 	}
 }
 
-struct FSimpleIntCompare
-{
-	FORCEINLINE bool operator()( const int32 A, const int32 B ) const { return B < A; }
-};
-
-bool FAudioDevice::HandleSoundTemplateInfoCommand( const TCHAR* Cmd, FOutputDevice& Ar )
-{
-	int32 NumCues = 0;
-	TMap<FString, int32> UniqueCues;
-
-	for( TObjectIterator<USoundCue> It; It; ++It )
-	{
-		TArray<USoundNode*> SoundNodes;
-
-		USoundCue* Cue = *It;
-		if( Cue )
-		{
-			if( Cue->FirstNode )
-			{
-				Cue->FirstNode->GetAllNodes( SoundNodes );
-
-				FString Unique = TEXT( "" );
-				for( int32 SoundNodeIndex = 0; SoundNodeIndex < SoundNodes.Num(); SoundNodeIndex++ )
-				{
-					USoundNode* SoundNode = SoundNodes[ SoundNodeIndex ];
-					Unique += SoundNode->GetUniqueString();
-				}
-
-				if( !FCString::Stristr( *Unique, TEXT( "Complex" ) ) )
-				{
-					int32 Count = 1;
-					if( UniqueCues.Find( Unique ) )
-					{
-						Count = UniqueCues.FindRef( Unique ) + 1;
-					}
-					UniqueCues.Add( Unique, Count );
-				}
-
-				Ar.Logf( TEXT( "Cue: %s : %s" ), *Cue->GetFullName(), *Unique );
-				NumCues++;
-			}
-			else
-			{
-				Ar.Logf( TEXT( "No FirstNode : %s" ), *Cue->GetFullName() );
-			}
-		}
-	}
-
-	Ar.Logf( TEXT( "Potential Templates -" ) );
-
-	UniqueCues.ValueSort( FSimpleIntCompare() );
-
-	for( TMap<FString, int32>::TIterator It( UniqueCues ); It; ++It )
-	{
-		FString Template = It.Key();
-		int32 TemplateCount = It.Value();
-		Ar.Logf( TEXT( "%05d : %s" ), TemplateCount, *Template );
-	}
-
-	Ar.Logf( TEXT( "SoundCues processed: %d" ), NumCues );
-	Ar.Logf( TEXT( "Unique SoundCues   : %d" ), UniqueCues.Num() );
-	return true;
-}
-
 int32 PrecachedRealtime = 0;
 int32 PrecachedNative = 0;
 int32 TotalNativeSize = 0;
@@ -823,10 +759,6 @@ bool FAudioDevice::Exec( UWorld* InWorld, const TCHAR* Cmd, FOutputDevice& Ar )
 	{
 		return HandleListSoundDurationsCommand( Cmd, Ar );
 	}
-	else if( FParse::Command( &Cmd, TEXT( "SoundTemplateInfo" ) ) )
-	{
-		return HandleSoundTemplateInfoCommand( Cmd, Ar );
-	}
 	else if( FParse::Command( &Cmd, TEXT( "PlaySoundCue" ) ) )
 	{
 		return HandlePlaySoundCueCommand( Cmd, Ar );
diff --git a/Engine/Source/Runtime/Engine/Private/SoundNode.cpp b/Engine/Source/Runtime/Engine/Private/SoundNode.cpp
--- a/Engine/Source/Runtime/Engine/Private/SoundNode.cpp
+++ b/Engine/Source/Runtime/Engine/Private/SoundNode.cpp
@@ -136,11 +136,6 @@ void USoundNode::PostEditChangeProperty(FPropertyChangedEvent& PropertyChangedEv
 }
 #endif // WITH_EDITOR
 
-FString USoundNode::GetUniqueString() const
-{
-	return TEXT("ERROR" );
-}
-
 #if WITH_EDITOR
 void USoundNode::PlaceNode( int32 NodeColumn, int32 NodeRow, int32 RowCount )
 {
diff --git a/Engine/Source/Runtime/Engine/Private/SoundNodeAttenuation.cpp b/Engine/Source/Runtime/Engine/Private/SoundNodeAttenuation.cpp
--- a/Engine/Source/Runtime/Engine/Private/SoundNodeAttenuation.cpp
+++ b/Engine/Source/Runtime/Engine/Private/SoundNodeAttenuation.cpp
@@ -118,62 +118,4 @@ void USoundNodeAttenuation::ParseNodes( FAudioDevice* AudioDevice, const UPTRINT
 	}
 
 	Super::ParseNodes( AudioDevice, NodeWaveInstanceHash, ActiveSound, UpdatedParseParams, WaveInstances );
-}
-
-
-FString USoundNodeAttenuation::GetUniqueString() const
-{
-	const FAttenuationSettings* Settings = NULL;
-
-	if (bOverrideAttenuation)
-	{
-		Settings = &AttenuationOverrides;
-	}
-	else if (AttenuationSettings)
-	{
-		Settings = &AttenuationSettings->Attenuation;
-	}
-
-	FString Unique = TEXT( "Attenuation" );
-
-	if (Settings)
-	{
-		if( Settings->bAttenuate )
-		{
-			switch(Settings->AttenuationShape)
-			{
-			case EAttenuationShape::Sphere:
-				Unique += FString::Printf( TEXT( " Vol Sphere Radius %g" ), Settings->AttenuationShapeExtents.X );
-				break;
-
-			case EAttenuationShape::Box:
-				Unique += FString::Printf( TEXT( " Vol Box (%g,%g,%g)" ), Settings->AttenuationShapeExtents.X, Settings->AttenuationShapeExtents.Y, Settings->AttenuationShapeExtents.Z);
-				break;
-
-			case EAttenuationShape::Capsule:
-				Unique += FString::Printf( TEXT( " Vol Capsule (%g,%g)" ), Settings->AttenuationShapeExtents.X, Settings->AttenuationShapeExtents.Y);
-				break;
-
-			case EAttenuationShape::Cone:
-				Unique += FString::Printf( TEXT( " Vol Cone Offset %g Radius %g Angle %g Falloff Angle %g" ), Settings->ConeOffset, Settings->AttenuationShapeExtents.X, Settings->AttenuationShapeExtents.Y, Settings->AttenuationShapeExtents.Z);
-				break;
-
-			default:
-				check(false);
-			}
-			Unique += FString::Printf( TEXT(" Falloff %g"), Settings->FalloffDistance);
-		}
-		if( Settings->bAttenuateWithLPF )
-		{
-			Unique += FString::Printf( TEXT( " Lpf %g %g" ), Settings->LPFRadiusMin, Settings->LPFRadiusMax );
-		}
-		if (Settings->bSpatialize)
-		{
-			Unique += TEXT( " Spat" );
-		}
-		Unique += FString::Printf( TEXT( " %d %g" ), (int32)Settings->DistanceAlgorithm, Settings->dBAttenuationAtMax );
-	}
-
-	Unique += TEXT( "/" );
-	return Unique;
-}
+}
\ No newline at end of file
diff --git a/Engine/Source/Runtime/Engine/Private/SoundNodeConcatenator.cpp b/Engine/Source/Runtime/Engine/Private/SoundNodeConcatenator.cpp
--- a/Engine/Source/Runtime/Engine/Private/SoundNodeConcatenator.cpp
+++ b/Engine/Source/Runtime/Engine/Private/SoundNodeConcatenator.cpp
@@ -118,9 +118,4 @@ void USoundNodeConcatenator::SetChildNodes(TArray<USoundNode*>& InChildNodes)
 		InputVolume.RemoveAt(InputVolume.Num() - NumToRemove, NumToRemove);
 	}
 }
-#endif //WITH_EDITOR
-
-FString USoundNodeConcatenator::GetUniqueString() const
-{
-	return TEXT( "Concatenator/" );
-}
\ No newline at end of file
+#endif //WITH_EDITOR
\ No newline at end of file
diff --git a/Engine/Source/Runtime/Engine/Private/SoundNodeDelay.cpp b/Engine/Source/Runtime/Engine/Private/SoundNodeDelay.cpp
--- a/Engine/Source/Runtime/Engine/Private/SoundNodeDelay.cpp
+++ b/Engine/Source/Runtime/Engine/Private/SoundNodeDelay.cpp
@@ -66,14 +66,4 @@ float USoundNodeDelay::GetDuration()
 
 	// And return the two together.
 	return( ChildDuration + DelayMax );
-}
-
-FString USoundNodeDelay::GetUniqueString() const
-{
-	FString Unique = TEXT( "Delay" );
-
-	Unique += FString::Printf( TEXT( " %g %g" ), DelayMin, DelayMax );
-
-	Unique += TEXT( "/" );
-	return( Unique );
-}
+}
\ No newline at end of file
diff --git a/Engine/Source/Runtime/Engine/Private/SoundNodeDialoguePlayer.cpp b/Engine/Source/Runtime/Engine/Private/SoundNodeDialoguePlayer.cpp
--- a/Engine/Source/Runtime/Engine/Private/SoundNodeDialoguePlayer.cpp
+++ b/Engine/Source/Runtime/Engine/Private/SoundNodeDialoguePlayer.cpp
@@ -45,11 +45,6 @@ float USoundNodeDialoguePlayer::GetDuration()
 	return Duration;
 }
 
-FString USoundNodeDialoguePlayer::GetUniqueString() const
-{
-	return TEXT("");
-}
-
 // A Wave Player is the end of the chain and has no children
 int32 USoundNodeDialoguePlayer::GetMaxChildNodes() const
 {
diff --git a/Engine/Source/Runtime/Engine/Private/SoundNodeDistanceCrossFade.cpp b/Engine/Source/Runtime/Engine/Private/SoundNodeDistanceCrossFade.cpp
--- a/Engine/Source/Runtime/Engine/Private/SoundNodeDistanceCrossFade.cpp
+++ b/Engine/Source/Runtime/Engine/Private/SoundNodeDistanceCrossFade.cpp
@@ -148,11 +148,6 @@ float USoundNodeDistanceCrossFade::GetCurrentDistance(FAudioDevice* AudioDevice,
 	return ActiveSound.bLocationDefined ? FVector::Dist( ParseParams.Transform.GetTranslation(), AudioDevice->Listeners[0].Transform.GetTranslation() ) : 0.f;
 }
 
-FString USoundNodeDistanceCrossFade::GetUniqueString() const
-{
-	return TEXT( "DistanceCrossFadeComplex/" );
-}
-
 bool USoundNodeDistanceCrossFade::AllowCrossfading(FActiveSound& ActiveSound) const
 {
 	// Ignore distance calculations for preview components as they are undefined
diff --git a/Engine/Source/Runtime/Engine/Private/SoundNodeDoppler.cpp b/Engine/Source/Runtime/Engine/Private/SoundNodeDoppler.cpp
--- a/Engine/Source/Runtime/Engine/Private/SoundNodeDoppler.cpp
+++ b/Engine/Source/Runtime/Engine/Private/SoundNodeDoppler.cpp
@@ -42,9 +42,4 @@ float USoundNodeDoppler::GetDopplerPitchMultiplier(FListener const& InListener,
 	//UE_LOG(LogAudio, Log, TEXT("Applying doppler pitchscale %f, raw scale %f, deltaspeed was %f"), FinalPitchScale, PitchScale, ListenerVelMagAwayFromSource - SourceVelMagTorwardListener);
 
 	return FinalPitchScale;
-}
-
-FString USoundNodeDoppler::GetUniqueString() const
-{
-	return TEXT( "Doppler/" );
-}
+}
\ No newline at end of file
diff --git a/Engine/Source/Runtime/Engine/Private/SoundNodeEnveloper.cpp b/Engine/Source/Runtime/Engine/Private/SoundNodeEnveloper.cpp
--- a/Engine/Source/Runtime/Engine/Private/SoundNodeEnveloper.cpp
+++ b/Engine/Source/Runtime/Engine/Private/SoundNodeEnveloper.cpp
@@ -67,11 +67,6 @@ void USoundNodeEnveloper::Serialize(FArchive& Ar)
 	}
 }
 
-FString USoundNodeEnveloper::GetUniqueString() const
-{
-	return TEXT("Enveloper");
-}
-
 void USoundNodeEnveloper::ParseNodes( FAudioDevice* AudioDevice, const UPTRINT NodeWaveInstanceHash, FActiveSound& ActiveSound, const FSoundParseParameters& ParseParams, TArray<FWaveInstance*>& WaveInstances )
 {
 	RETRIEVE_SOUNDNODE_PAYLOAD( sizeof( float ) + sizeof(float) + sizeof(float) + sizeof(int32) );
diff --git a/Engine/Source/Runtime/Engine/Private/SoundNodeGroupControl.cpp b/Engine/Source/Runtime/Engine/Private/SoundNodeGroupControl.cpp
--- a/Engine/Source/Runtime/Engine/Private/SoundNodeGroupControl.cpp
+++ b/Engine/Source/Runtime/Engine/Private/SoundNodeGroupControl.cpp
@@ -226,8 +226,3 @@ void USoundNodeGroupControl::SetChildNodes(TArray<USoundNode*>& InChildNodes)
 	}
 }
 #endif //WITH_EDITOR
-
-FString USoundNodeGroupControl::GetUniqueString() const
-{
-	return TEXT( "GroupControl/" );
-}
diff --git a/Engine/Source/Runtime/Engine/Private/SoundNodeLooping.cpp b/Engine/Source/Runtime/Engine/Private/SoundNodeLooping.cpp
--- a/Engine/Source/Runtime/Engine/Private/SoundNodeLooping.cpp
+++ b/Engine/Source/Runtime/Engine/Private/SoundNodeLooping.cpp
@@ -97,7 +97,3 @@ float USoundNodeLooping::GetDuration()
 	return INDEFINITELY_LOOPING_DURATION;
 }
 
-FString USoundNodeLooping::GetUniqueString() const
-{
-	return TEXT( "Looping Forever/" );
-}
diff --git a/Engine/Source/Runtime/Engine/Private/SoundNodeMature.cpp b/Engine/Source/Runtime/Engine/Private/SoundNodeMature.cpp
--- a/Engine/Source/Runtime/Engine/Private/SoundNodeMature.cpp
+++ b/Engine/Source/Runtime/Engine/Private/SoundNodeMature.cpp
@@ -148,11 +148,6 @@ void USoundNodeMature::CreateStartingConnectors()
 	InsertChildNode( ChildNodes.Num() );
 }
 
-FString USoundNodeMature::GetUniqueString() const
-{
-	return TEXT( "Mature/" );
-}
-
 void USoundNodeMature::PostLoad()
 {
 	Super::PostLoad();
diff --git a/Engine/Source/Runtime/Engine/Private/SoundNodeMixer.cpp b/Engine/Source/Runtime/Engine/Private/SoundNodeMixer.cpp
--- a/Engine/Source/Runtime/Engine/Private/SoundNodeMixer.cpp
+++ b/Engine/Source/Runtime/Engine/Private/SoundNodeMixer.cpp
@@ -67,17 +67,4 @@ void USoundNodeMixer::SetChildNodes(TArray<USoundNode*>& InChildNodes)
 		InputVolume.RemoveAt(InputVolume.Num() - NumToRemove, NumToRemove);
 	}
 }
-#endif //WITH_EDITOR
-
-FString USoundNodeMixer::GetUniqueString() const
-{
-	FString Unique = TEXT( "Mixer" );
-
-	for( int32 Index = 0; Index < InputVolume.Num(); Index++ )
-	{
-		Unique += FString::Printf( TEXT( " %g" ), InputVolume[ Index ] );
-	}
-
-	Unique += TEXT( "/" );
-	return( Unique );
-}
+#endif //WITH_EDITOR
\ No newline at end of file
diff --git a/Engine/Source/Runtime/Engine/Private/SoundNodeModulator.cpp b/Engine/Source/Runtime/Engine/Private/SoundNodeModulator.cpp
--- a/Engine/Source/Runtime/Engine/Private/SoundNodeModulator.cpp
+++ b/Engine/Source/Runtime/Engine/Private/SoundNodeModulator.cpp
@@ -37,16 +37,3 @@ void USoundNodeModulator::ParseNodes( FAudioDevice* AudioDevice, const UPTRINT N
 
 	Super::ParseNodes( AudioDevice, NodeWaveInstanceHash, ActiveSound, UpdatedParams, WaveInstances );
 }
-
-/**
- * Used to create a unique string to identify unique nodes
- */
-FString USoundNodeModulator::GetUniqueString() const
-{
-	FString Unique = TEXT( "Modulator" );
-
-	Unique += FString::Printf( TEXT( " %g %g %g %g" ), VolumeMin, VolumeMax, PitchMin, PitchMax );
-
-	Unique += TEXT( "/" );
-	return( Unique );
-}
diff --git a/Engine/Source/Runtime/Engine/Private/SoundNodeModulatorContinuous.cpp b/Engine/Source/Runtime/Engine/Private/SoundNodeModulatorContinuous.cpp
--- a/Engine/Source/Runtime/Engine/Private/SoundNodeModulatorContinuous.cpp
+++ b/Engine/Source/Runtime/Engine/Private/SoundNodeModulatorContinuous.cpp
@@ -96,7 +96,3 @@ void USoundNodeModulatorContinuous::ParseNodes( FAudioDevice* AudioDevice, const
 	Super::ParseNodes( AudioDevice, NodeWaveInstanceHash, ActiveSound, UpdatedParams, WaveInstances );
 }
 
-FString USoundNodeModulatorContinuous::GetUniqueString() const
-{
-	return( TEXT( "ModulatorContinuousComplex/" ) );
-}
diff --git a/Engine/Source/Runtime/Engine/Private/SoundNodeOscillator.cpp b/Engine/Source/Runtime/Engine/Private/SoundNodeOscillator.cpp
--- a/Engine/Source/Runtime/Engine/Private/SoundNodeOscillator.cpp
+++ b/Engine/Source/Runtime/Engine/Private/SoundNodeOscillator.cpp
@@ -57,8 +57,3 @@ void USoundNodeOscillator::ParseNodes( FAudioDevice* AudioDevice, const UPTRINT
 	Super::ParseNodes( AudioDevice, NodeWaveInstanceHash, ActiveSound, UpdatedParams, WaveInstances );
 }
 
-
-FString USoundNodeOscillator::GetUniqueString() const
-{
-	return( TEXT( "OscillatorComplex/" ) );
-}
diff --git a/Engine/Source/Runtime/Engine/Private/SoundNodeParamCrossFade.cpp b/Engine/Source/Runtime/Engine/Private/SoundNodeParamCrossFade.cpp
--- a/Engine/Source/Runtime/Engine/Private/SoundNodeParamCrossFade.cpp
+++ b/Engine/Source/Runtime/Engine/Private/SoundNodeParamCrossFade.cpp
@@ -20,11 +20,6 @@ float USoundNodeParamCrossFade::GetCurrentDistance(FAudioDevice* AudioDevice, FA
 	return ParamValue;
 }
 
-FString USoundNodeParamCrossFade::GetUniqueString() const
-{
-	return TEXT( "ParamCrossFadeComplex/" );
-}
-
 bool USoundNodeParamCrossFade::AllowCrossfading(FActiveSound& ActiveSound) const
 {
 	// Always allow parameter to control crossfading, even on 2D/preview sounds
diff --git a/Engine/Source/Runtime/Engine/Private/SoundNodeRandom.cpp b/Engine/Source/Runtime/Engine/Private/SoundNodeRandom.cpp
--- a/Engine/Source/Runtime/Engine/Private/SoundNodeRandom.cpp
+++ b/Engine/Source/Runtime/Engine/Private/SoundNodeRandom.cpp
@@ -250,16 +250,6 @@ void USoundNodeRandom::OnBeginPIE(const bool bIsSimulating)
 }
 #endif //WITH_EDITOR
 
-FString USoundNodeRandom::GetUniqueString() const
-{
-	FString Unique = TEXT( "Random" );
-
-	Unique += bRandomizeWithoutReplacement ? TEXT( " True" ) : TEXT( " False" );
-
-	Unique += TEXT( "/" );
-	return( Unique );
-}
-
 #if WITH_EDITOR
 void USoundNodeRandom::UpdatePIEHiddenNodes()
 {
diff --git a/Engine/Source/Runtime/Engine/Private/SoundNodeSoundClass.cpp b/Engine/Source/Runtime/Engine/Private/SoundNodeSoundClass.cpp
--- a/Engine/Source/Runtime/Engine/Private/SoundNodeSoundClass.cpp
+++ b/Engine/Source/Runtime/Engine/Private/SoundNodeSoundClass.cpp
@@ -23,17 +23,4 @@ void USoundNodeSoundClass::ParseNodes( class FAudioDevice* AudioDevice, const UP
 	}
 
 	Super::ParseNodes( AudioDevice, NodeWaveInstanceHash, ActiveSound, UpdatedParseParams, WaveInstances );
-}
-
-FString USoundNodeSoundClass::GetUniqueString() const
-{
-	FString Unique = TEXT( "SoundClass" );
-
-	if (SoundClassOverride)
-	{
-		Unique += SoundClassOverride->GetFName().ToString();
-	}
-
-	Unique += TEXT( "/" );
-	return( Unique );
-}
+}
\ No newline at end of file
diff --git a/Engine/Source/Runtime/Engine/Private/SoundNodeWaveParam.cpp b/Engine/Source/Runtime/Engine/Private/SoundNodeWaveParam.cpp
--- a/Engine/Source/Runtime/Engine/Private/SoundNodeWaveParam.cpp
+++ b/Engine/Source/Runtime/Engine/Private/SoundNodeWaveParam.cpp
@@ -35,7 +35,3 @@ void USoundNodeWaveParam::ParseNodes( FAudioDevice* AudioDevice, const UPTRINT N
 	}
 }
 
-FString USoundNodeWaveParam::GetUniqueString() const
-{
-	return TEXT("WaveParam/");
-}
diff --git a/Engine/Source/Runtime/Engine/Private/SoundNodeWavePlayer.cpp b/Engine/Source/Runtime/Engine/Private/SoundNodeWavePlayer.cpp
--- a/Engine/Source/Runtime/Engine/Private/SoundNodeWavePlayer.cpp
+++ b/Engine/Source/Runtime/Engine/Private/SoundNodeWavePlayer.cpp
@@ -83,11 +83,6 @@ FString USoundNodeWavePlayer::GetTitle() const
 }
 #endif
 
-FString USoundNodeWavePlayer::GetUniqueString() const
-{
-	return TEXT("");
-}
-
 // A Wave Player is the end of the chain and has no children
 int32 USoundNodeWavePlayer::GetMaxChildNodes() const
 {
