Fixed race condition when shifting precomputed light volume

diff --git a/Engine/Source/Runtime/Engine/Private/Level.cpp b/Engine/Source/Runtime/Engine/Private/Level.cpp
--- a/Engine/Source/Runtime/Engine/Private/Level.cpp
+++ b/Engine/Source/Runtime/Engine/Private/Level.cpp
@@ -1843,13 +1843,24 @@ void ULevel::ApplyWorldOffset(const FVector& InWorldOffset, bool bWorldShift)
 	// Move precomputed light samples
 	if (PrecomputedLightVolume)
 	{
-		PrecomputedLightVolume->ApplyWorldOffset(InWorldOffset, bWorldShift);
+		if (!PrecomputedLightVolume->IsAddedToScene())
+		{
+			PrecomputedLightVolume->ApplyWorldOffset(InWorldOffset);
+		}
+		// At world origin rebasing all registered volumes will be moved during FScene shifting
+		// Otherwise we need to send a command to move just this volume
+		else if (!bWorldShift) 
+		{
+			ENQUEUE_UNIQUE_RENDER_COMMAND_TWOPARAMETER(
+ 				ApplyWorldOffset_PLV,
+ 				FPrecomputedLightVolume*, InPrecomputedLightVolume, PrecomputedLightVolume,
+ 				FVector, InWorldOffset, InWorldOffset,
+ 			{
+				InPrecomputedLightVolume->ApplyWorldOffset(InWorldOffset);
+ 			});
+		}
 	}
 
-	// Move precomputed visibility volume
-	// TODO: should probably move it to RT
-	PrecomputedVisibilityHandler.ApplyWorldOffset(InWorldOffset);
-
 	// Iterate over all actors in the level and move them
 	for (int32 ActorIndex = 0; ActorIndex < Actors.Num(); ActorIndex++)
 	{
diff --git a/Engine/Source/Runtime/Engine/Private/PrecomputedLightVolume.cpp b/Engine/Source/Runtime/Engine/Private/PrecomputedLightVolume.cpp
--- a/Engine/Source/Runtime/Engine/Private/PrecomputedLightVolume.cpp
+++ b/Engine/Source/Runtime/Engine/Private/PrecomputedLightVolume.cpp
@@ -403,7 +403,7 @@ SIZE_T FPrecomputedLightVolume::GetAllocatedBytes() const
 	return NodeBytes;
 }
 
-void FPrecomputedLightVolume::ApplyWorldOffset(const FVector& InOffset, bool bWorldShift)
+void FPrecomputedLightVolume::ApplyWorldOffset(const FVector& InOffset)
 {
 	Bounds.Min+= InOffset;
 	Bounds.Max+= InOffset;
diff --git a/Engine/Source/Runtime/Engine/Public/PrecomputedLightVolume.h b/Engine/Source/Runtime/Engine/Public/PrecomputedLightVolume.h
--- a/Engine/Source/Runtime/Engine/Public/PrecomputedLightVolume.h
+++ b/Engine/Source/Runtime/Engine/Public/PrecomputedLightVolume.h
@@ -152,7 +152,7 @@ public:
 		return OctreeForRendering->GetNodeLevelExtent(Level);
 	}
 
-	void ApplyWorldOffset(const FVector& InOffset, bool bWorldShift);
+	ENGINE_API void ApplyWorldOffset(const FVector& InOffset);
 
 	// temporary, for ES2 preview verification.
 	FORCEINLINE bool IsUsingHighQualityLightMap() const
diff --git a/Engine/Source/Runtime/Renderer/Private/RendererScene.cpp b/Engine/Source/Runtime/Renderer/Private/RendererScene.cpp
--- a/Engine/Source/Runtime/Renderer/Private/RendererScene.cpp
+++ b/Engine/Source/Runtime/Renderer/Private/RendererScene.cpp
@@ -1888,6 +1888,18 @@ void FScene::ApplyWorldOffset_RenderThread(FVector InOffset)
 		(*It)->ApplyWorldOffset(InOffset);
 	}
 
+	// Precomputed light volumes
+	for (const FPrecomputedLightVolume* It : PrecomputedLightVolumes)
+	{
+		const_cast<FPrecomputedLightVolume*>(It)->ApplyWorldOffset(InOffset);
+	}
+
+	// Precomputed visibility
+	if (PrecomputedVisibilityHandler)
+	{
+		const_cast<FPrecomputedVisibilityHandler*>(PrecomputedVisibilityHandler)->ApplyWorldOffset(InOffset);
+	}
+	
 	// Invalidate indirect lighting cache
 	IndirectLightingCache.SetLightingCacheDirty();
 
