Engine.h/EnginePrivate.h cleanup - delegate bindings.

diff --git a/Engine/Source/Developer/AssetTools/Private/AssetTypeActions/AssetTypeActions_VectorField.h b/Engine/Source/Developer/AssetTools/Private/AssetTypeActions/AssetTypeActions_VectorField.h
--- a/Engine/Source/Developer/AssetTools/Private/AssetTypeActions/AssetTypeActions_VectorField.h
+++ b/Engine/Source/Developer/AssetTools/Private/AssetTypeActions/AssetTypeActions_VectorField.h
@@ -1,6 +1,8 @@
 // Copyright 1998-2014 Epic Games, Inc. All Rights Reserved.
 
 #pragma once
+#include "AssetTypeActions_Base.h"
+#include "VectorField/VectorField.h"
 
 class FAssetTypeActions_VectorField : public FAssetTypeActions_Base
 {
diff --git a/Engine/Source/Developer/AssetTools/Private/AssetTypeActions/AssetTypeActions_VectorFieldAnimated.h b/Engine/Source/Developer/AssetTools/Private/AssetTypeActions/AssetTypeActions_VectorFieldAnimated.h
--- a/Engine/Source/Developer/AssetTools/Private/AssetTypeActions/AssetTypeActions_VectorFieldAnimated.h
+++ b/Engine/Source/Developer/AssetTools/Private/AssetTypeActions/AssetTypeActions_VectorFieldAnimated.h
@@ -1,6 +1,8 @@
 // Copyright 1998-2014 Epic Games, Inc. All Rights Reserved.
 
 #pragma once
+#include "AssetTypeActions_VectorField.h"
+#include "VectorField/VectorFieldAnimated.h"
 
 class FAssetTypeActions_VectorFieldAnimated : public FAssetTypeActions_VectorField
 {
diff --git a/Engine/Source/Developer/AssetTools/Private/AssetTypeActions/AssetTypeActions_VectorFieldStatic.h b/Engine/Source/Developer/AssetTools/Private/AssetTypeActions/AssetTypeActions_VectorFieldStatic.h
--- a/Engine/Source/Developer/AssetTools/Private/AssetTypeActions/AssetTypeActions_VectorFieldStatic.h
+++ b/Engine/Source/Developer/AssetTools/Private/AssetTypeActions/AssetTypeActions_VectorFieldStatic.h
@@ -1,6 +1,8 @@
 // Copyright 1998-2014 Epic Games, Inc. All Rights Reserved.
 
 #pragma once
+#include "AssetTypeActions/AssetTypeActions_VectorField.h"
+#include "VectorField/VectorFieldStatic.h"
 
 class FAssetTypeActions_VectorFieldStatic : public FAssetTypeActions_VectorField
 {
diff --git a/Engine/Source/Editor/BlueprintGraph/Private/K2Node_ComponentBoundEvent.cpp b/Engine/Source/Editor/BlueprintGraph/Private/K2Node_ComponentBoundEvent.cpp
--- a/Engine/Source/Editor/BlueprintGraph/Private/K2Node_ComponentBoundEvent.cpp
+++ b/Engine/Source/Editor/BlueprintGraph/Private/K2Node_ComponentBoundEvent.cpp
@@ -1,6 +1,7 @@
 // Copyright 1998-2014 Epic Games, Inc. All Rights Reserved.
 
 #include "BlueprintGraphPrivatePCH.h"
+#include "Engine/ComponentDelegateBinding.h"
 
 #define LOCTEXT_NAMESPACE "K2Node"
 
diff --git a/Engine/Source/Editor/BlueprintGraph/Private/K2Node_GetInputAxisKeyValue.cpp b/Engine/Source/Editor/BlueprintGraph/Private/K2Node_GetInputAxisKeyValue.cpp
--- a/Engine/Source/Editor/BlueprintGraph/Private/K2Node_GetInputAxisKeyValue.cpp
+++ b/Engine/Source/Editor/BlueprintGraph/Private/K2Node_GetInputAxisKeyValue.cpp
@@ -5,6 +5,7 @@
 #include "CompilerResultsLog.h"
 #include "BlueprintNodeSpawner.h"
 #include "EditorCategoryUtils.h"
+#include "Engine/InputAxisKeyDelegateBinding.h"
 
 #define LOCTEXT_NAMESPACE "K2Node_GetInputAxisKeyValue"
 
diff --git a/Engine/Source/Editor/BlueprintGraph/Private/K2Node_GetInputAxisValue.cpp b/Engine/Source/Editor/BlueprintGraph/Private/K2Node_GetInputAxisValue.cpp
--- a/Engine/Source/Editor/BlueprintGraph/Private/K2Node_GetInputAxisValue.cpp
+++ b/Engine/Source/Editor/BlueprintGraph/Private/K2Node_GetInputAxisValue.cpp
@@ -5,6 +5,7 @@
 #include "CompilerResultsLog.h"
 #include "BlueprintNodeSpawner.h"
 #include "EditorCategoryUtils.h"
+#include "Engine/InputAxisDelegateBinding.h"
 
 #define LOCTEXT_NAMESPACE "K2Node_GetInputAxisValue"
 
diff --git a/Engine/Source/Editor/BlueprintGraph/Private/K2Node_GetInputVectorAxisValue.cpp b/Engine/Source/Editor/BlueprintGraph/Private/K2Node_GetInputVectorAxisValue.cpp
--- a/Engine/Source/Editor/BlueprintGraph/Private/K2Node_GetInputVectorAxisValue.cpp
+++ b/Engine/Source/Editor/BlueprintGraph/Private/K2Node_GetInputVectorAxisValue.cpp
@@ -4,6 +4,7 @@
 #include "K2Node_GetInputVectorAxisValue.h"
 #include "CompilerResultsLog.h"
 #include "BlueprintNodeSpawner.h"
+#include "Engine/InputVectorAxisDelegateBinding.h"
 
 UK2Node_GetInputVectorAxisValue::UK2Node_GetInputVectorAxisValue(const class FPostConstructInitializeProperties& PCIP)
 	: Super(PCIP)
diff --git a/Engine/Source/Editor/BlueprintGraph/Private/K2Node_InputActionEvent.cpp b/Engine/Source/Editor/BlueprintGraph/Private/K2Node_InputActionEvent.cpp
--- a/Engine/Source/Editor/BlueprintGraph/Private/K2Node_InputActionEvent.cpp
+++ b/Engine/Source/Editor/BlueprintGraph/Private/K2Node_InputActionEvent.cpp
@@ -2,6 +2,7 @@
 
 #include "BlueprintGraphPrivatePCH.h"
 #include "K2Node_InputActionEvent.h"
+#include "Engine/InputActionDelegateBinding.h"
 
 UK2Node_InputActionEvent::UK2Node_InputActionEvent(const class FPostConstructInitializeProperties& PCIP)
 	: Super(PCIP)
diff --git a/Engine/Source/Editor/BlueprintGraph/Private/K2Node_InputAxisEvent.cpp b/Engine/Source/Editor/BlueprintGraph/Private/K2Node_InputAxisEvent.cpp
--- a/Engine/Source/Editor/BlueprintGraph/Private/K2Node_InputAxisEvent.cpp
+++ b/Engine/Source/Editor/BlueprintGraph/Private/K2Node_InputAxisEvent.cpp
@@ -4,6 +4,7 @@
 #include "CompilerResultsLog.h"
 #include "BlueprintNodeSpawner.h"
 #include "EditorCategoryUtils.h"
+#include "Engine/InputAxisDelegateBinding.h"
 
 #define LOCTEXT_NAMESPACE "K2Node_InputAxisEvent"
 
diff --git a/Engine/Source/Editor/BlueprintGraph/Private/K2Node_InputAxisKeyEvent.cpp b/Engine/Source/Editor/BlueprintGraph/Private/K2Node_InputAxisKeyEvent.cpp
--- a/Engine/Source/Editor/BlueprintGraph/Private/K2Node_InputAxisKeyEvent.cpp
+++ b/Engine/Source/Editor/BlueprintGraph/Private/K2Node_InputAxisKeyEvent.cpp
@@ -5,6 +5,7 @@
 #include "CompilerResultsLog.h"
 #include "BlueprintNodeSpawner.h"
 #include "EditorCategoryUtils.h"
+#include "Engine/InputAxisKeyDelegateBinding.h"
 
 #define LOCTEXT_NAMESPACE "UK2Node_InputAxisKeyEvent"
 
diff --git a/Engine/Source/Editor/BlueprintGraph/Private/K2Node_InputKeyEvent.cpp b/Engine/Source/Editor/BlueprintGraph/Private/K2Node_InputKeyEvent.cpp
--- a/Engine/Source/Editor/BlueprintGraph/Private/K2Node_InputKeyEvent.cpp
+++ b/Engine/Source/Editor/BlueprintGraph/Private/K2Node_InputKeyEvent.cpp
@@ -2,6 +2,7 @@
 
 #include "BlueprintGraphPrivatePCH.h"
 #include "K2Node_InputKeyEvent.h"
+#include "Engine/InputKeyDelegateBinding.h"
 
 UK2Node_InputKeyEvent::UK2Node_InputKeyEvent(const class FPostConstructInitializeProperties& PCIP)
 	: Super(PCIP)
diff --git a/Engine/Source/Editor/BlueprintGraph/Private/K2Node_InputTouchEvent.cpp b/Engine/Source/Editor/BlueprintGraph/Private/K2Node_InputTouchEvent.cpp
--- a/Engine/Source/Editor/BlueprintGraph/Private/K2Node_InputTouchEvent.cpp
+++ b/Engine/Source/Editor/BlueprintGraph/Private/K2Node_InputTouchEvent.cpp
@@ -3,6 +3,7 @@
 #include "BlueprintGraphPrivatePCH.h"
 #include "K2Node_InputTouchEvent.h"
 #include "CompilerResultsLog.h"
+#include "Engine/InputTouchDelegateBinding.h"
 
 UK2Node_InputTouchEvent::UK2Node_InputTouchEvent(const class FPostConstructInitializeProperties& PCIP)
 	: Super(PCIP)
diff --git a/Engine/Source/Editor/BlueprintGraph/Private/K2Node_InputVectorAxisEvent.cpp b/Engine/Source/Editor/BlueprintGraph/Private/K2Node_InputVectorAxisEvent.cpp
--- a/Engine/Source/Editor/BlueprintGraph/Private/K2Node_InputVectorAxisEvent.cpp
+++ b/Engine/Source/Editor/BlueprintGraph/Private/K2Node_InputVectorAxisEvent.cpp
@@ -4,6 +4,7 @@
 #include "K2Node_InputVectorAxisEvent.h"
 #include "CompilerResultsLog.h"
 #include "BlueprintNodeSpawner.h"
+#include "Engine/InputVectorAxisDelegateBinding.h"
 
 UK2Node_InputVectorAxisEvent::UK2Node_InputVectorAxisEvent(const class FPostConstructInitializeProperties& PCIP)
 : Super(PCIP)
diff --git a/Engine/Source/Editor/KismetCompiler/Private/KismetCompiler.cpp b/Engine/Source/Editor/KismetCompiler/Private/KismetCompiler.cpp
--- a/Engine/Source/Editor/KismetCompiler/Private/KismetCompiler.cpp
+++ b/Engine/Source/Editor/KismetCompiler/Private/KismetCompiler.cpp
@@ -19,6 +19,7 @@
 #include "EditorCategoryUtils.h"
 #include "K2Node_EnumLiteral.h"
 #include "EdGraph/EdGraphNode_Documentation.h"
+#include "Engine/DynamicBlueprintBinding.h"
 
 static bool bDebugPropertyPropagation = false;
 
diff --git a/Engine/Source/Editor/UnrealEd/Private/Factories/ActorFactory.cpp b/Engine/Source/Editor/UnrealEd/Private/Factories/ActorFactory.cpp
--- a/Engine/Source/Editor/UnrealEd/Private/Factories/ActorFactory.cpp
+++ b/Engine/Source/Editor/UnrealEd/Private/Factories/ActorFactory.cpp
@@ -26,6 +26,8 @@ ActorFactory.cpp:
 #include "Particles/ParticleSystem.h"
 #include "Particles/ParticleSystemComponent.h"
 
+#include "VectorField/VectorField.h"
+
 DEFINE_LOG_CATEGORY_STATIC(LogActorFactory, Log, All);
 
 #define LOCTEXT_NAMESPACE "ActorFactory"
diff --git a/Engine/Source/Editor/UnrealEd/Private/Factories/VectorFieldFactory.cpp b/Engine/Source/Editor/UnrealEd/Private/Factories/VectorFieldFactory.cpp
--- a/Engine/Source/Editor/UnrealEd/Private/Factories/VectorFieldFactory.cpp
+++ b/Engine/Source/Editor/UnrealEd/Private/Factories/VectorFieldFactory.cpp
@@ -8,6 +8,7 @@
 
 #include "Particles/ParticleSystemComponent.h"
 #include "ComponentReregisterContext.h"
+#include "VectorField/VectorFieldStatic.h"
 
 DEFINE_LOG_CATEGORY_STATIC(LogVectorFieldFactory, Log, All);
 
diff --git a/Engine/Source/Runtime/Engine/Classes/Engine/ComponentDelegateBinding.h b/Engine/Source/Runtime/Engine/Classes/Engine/ComponentDelegateBinding.h
--- a/Engine/Source/Runtime/Engine/Classes/Engine/ComponentDelegateBinding.h
+++ b/Engine/Source/Runtime/Engine/Classes/Engine/ComponentDelegateBinding.h
@@ -1,6 +1,7 @@
 // Copyright 1998-2014 Epic Games, Inc. All Rights Reserved.
 
 #pragma once
+#include "Engine/DynamicBlueprintBinding.h"
 #include "ComponentDelegateBinding.generated.h"
 
 /** Entry for a delegate to assign after a blueprint has been instanced */
diff --git a/Engine/Source/Runtime/Engine/Classes/Engine/InputActionDelegateBinding.h b/Engine/Source/Runtime/Engine/Classes/Engine/InputActionDelegateBinding.h
--- a/Engine/Source/Runtime/Engine/Classes/Engine/InputActionDelegateBinding.h
+++ b/Engine/Source/Runtime/Engine/Classes/Engine/InputActionDelegateBinding.h
@@ -1,6 +1,7 @@
 // Copyright 1998-2014 Epic Games, Inc. All Rights Reserved.
 
 #pragma once
+#include "Engine/InputDelegateBinding.h"
 #include "InputActionDelegateBinding.generated.h"
 
 USTRUCT()
diff --git a/Engine/Source/Runtime/Engine/Classes/Engine/InputAxisDelegateBinding.h b/Engine/Source/Runtime/Engine/Classes/Engine/InputAxisDelegateBinding.h
--- a/Engine/Source/Runtime/Engine/Classes/Engine/InputAxisDelegateBinding.h
+++ b/Engine/Source/Runtime/Engine/Classes/Engine/InputAxisDelegateBinding.h
@@ -1,6 +1,7 @@
 // Copyright 1998-2014 Epic Games, Inc. All Rights Reserved.
 
 #pragma once
+#include "Engine/InputDelegateBinding.h"
 #include "InputAxisDelegateBinding.generated.h"
 
 USTRUCT()
diff --git a/Engine/Source/Runtime/Engine/Classes/Engine/InputAxisKeyDelegateBinding.h b/Engine/Source/Runtime/Engine/Classes/Engine/InputAxisKeyDelegateBinding.h
--- a/Engine/Source/Runtime/Engine/Classes/Engine/InputAxisKeyDelegateBinding.h
+++ b/Engine/Source/Runtime/Engine/Classes/Engine/InputAxisKeyDelegateBinding.h
@@ -1,6 +1,7 @@
 // Copyright 1998-2014 Epic Games, Inc. All Rights Reserved.
 
 #pragma once
+#include "Engine/InputDelegateBinding.h"
 #include "InputAxisKeyDelegateBinding.generated.h"
 
 USTRUCT()
diff --git a/Engine/Source/Runtime/Engine/Classes/Engine/InputDelegateBinding.h b/Engine/Source/Runtime/Engine/Classes/Engine/InputDelegateBinding.h
--- a/Engine/Source/Runtime/Engine/Classes/Engine/InputDelegateBinding.h
+++ b/Engine/Source/Runtime/Engine/Classes/Engine/InputDelegateBinding.h
@@ -1,6 +1,7 @@
 // Copyright 1998-2014 Epic Games, Inc. All Rights Reserved.
 
 #pragma once
+#include "Engine/DynamicBlueprintBinding.h"
 #include "InputDelegateBinding.generated.h"
 
 USTRUCT()
diff --git a/Engine/Source/Runtime/Engine/Classes/Engine/InputTouchDelegateBinding.h b/Engine/Source/Runtime/Engine/Classes/Engine/InputTouchDelegateBinding.h
--- a/Engine/Source/Runtime/Engine/Classes/Engine/InputTouchDelegateBinding.h
+++ b/Engine/Source/Runtime/Engine/Classes/Engine/InputTouchDelegateBinding.h
@@ -1,6 +1,7 @@
 // Copyright 1998-2014 Epic Games, Inc. All Rights Reserved.
 
 #pragma once
+#include "Engine/InputDelegateBinding.h"
 #include "InputTouchDelegateBinding.generated.h"
 
 USTRUCT()
diff --git a/Engine/Source/Runtime/Engine/Classes/Engine/InputVectorAxisDelegateBinding.h b/Engine/Source/Runtime/Engine/Classes/Engine/InputVectorAxisDelegateBinding.h
--- a/Engine/Source/Runtime/Engine/Classes/Engine/InputVectorAxisDelegateBinding.h
+++ b/Engine/Source/Runtime/Engine/Classes/Engine/InputVectorAxisDelegateBinding.h
@@ -1,6 +1,7 @@
 // Copyright 1998-2014 Epic Games, Inc. All Rights Reserved.
 
 #pragma once
+#include "Engine/InputAxisKeyDelegateBinding.h"
 #include "InputVectorAxisDelegateBinding.generated.h"
 
 UCLASS()
diff --git a/Engine/Source/Runtime/Engine/Classes/VectorField/VectorFieldStatic.h b/Engine/Source/Runtime/Engine/Classes/VectorField/VectorFieldStatic.h
--- a/Engine/Source/Runtime/Engine/Classes/VectorField/VectorFieldStatic.h
+++ b/Engine/Source/Runtime/Engine/Classes/VectorField/VectorFieldStatic.h
@@ -5,7 +5,7 @@
 =============================================================================*/
 
 #pragma once
-
+#include "VectorField/VectorField.h"
 #include "VectorFieldStatic.generated.h"
 
 UCLASS(hidecategories=VectorFieldBounds, MinimalAPI)
diff --git a/Engine/Source/Runtime/Engine/Private/Actor.cpp b/Engine/Source/Runtime/Engine/Private/Actor.cpp
--- a/Engine/Source/Runtime/Engine/Private/Actor.cpp
+++ b/Engine/Source/Runtime/Engine/Private/Actor.cpp
@@ -2,6 +2,7 @@
 
 
 #include "EnginePrivate.h"
+#include "Engine/InputDelegateBinding.h"
 #include "PhysicsPublic.h"
 #include "ParticleDefinitions.h"
 #include "Sound/SoundCue.h"
diff --git a/Engine/Source/Runtime/Engine/Private/BlueprintGeneratedClass.cpp b/Engine/Source/Runtime/Engine/Private/BlueprintGeneratedClass.cpp
--- a/Engine/Source/Runtime/Engine/Private/BlueprintGeneratedClass.cpp
+++ b/Engine/Source/Runtime/Engine/Private/BlueprintGeneratedClass.cpp
@@ -2,6 +2,7 @@
 
 #include "EnginePrivate.h"
 #include "BlueprintUtilities.h"
+#include "Engine/InputDelegateBinding.h"
 
 #if WITH_EDITOR
 #include "BlueprintEditorUtils.h"
diff --git a/Engine/Source/Runtime/Engine/Private/ComponentDelegateBinding.cpp b/Engine/Source/Runtime/Engine/Private/ComponentDelegateBinding.cpp
--- a/Engine/Source/Runtime/Engine/Private/ComponentDelegateBinding.cpp
+++ b/Engine/Source/Runtime/Engine/Private/ComponentDelegateBinding.cpp
@@ -2,7 +2,7 @@
 
 #include "EnginePrivate.h"
 #include "BlueprintUtilities.h"
-
+#include "Engine/ComponentDelegateBinding.h"
 
 UComponentDelegateBinding::UComponentDelegateBinding(const class FPostConstructInitializeProperties& PCIP)
 	: Super(PCIP)
diff --git a/Engine/Source/Runtime/Engine/Private/CurveEdPresetCurve.cpp b/Engine/Source/Runtime/Engine/Private/CurveEdPresetCurve.cpp
--- a/Engine/Source/Runtime/Engine/Private/CurveEdPresetCurve.cpp
+++ b/Engine/Source/Runtime/Engine/Private/CurveEdPresetCurve.cpp
@@ -5,7 +5,7 @@
 =============================================================================*/
 
 #include "EnginePrivate.h"
-
+#include "Curves/CurveEdPresetCurve.h"
 /**
  *	UCurveEdPresetCurve
  */
diff --git a/Engine/Source/Runtime/Engine/Private/DynamicBlueprintBinding.cpp b/Engine/Source/Runtime/Engine/Private/DynamicBlueprintBinding.cpp
--- a/Engine/Source/Runtime/Engine/Private/DynamicBlueprintBinding.cpp
+++ b/Engine/Source/Runtime/Engine/Private/DynamicBlueprintBinding.cpp
@@ -2,6 +2,7 @@
 
 #include "EnginePrivate.h"
 #include "BlueprintUtilities.h"
+#include "Engine/InputDelegateBinding.h"
 
 
 UDynamicBlueprintBinding::UDynamicBlueprintBinding(const class FPostConstructInitializeProperties& PCIP)
diff --git a/Engine/Source/Runtime/Engine/Private/EnginePrivate.h b/Engine/Source/Runtime/Engine/Private/EnginePrivate.h
--- a/Engine/Source/Runtime/Engine/Private/EnginePrivate.h
+++ b/Engine/Source/Runtime/Engine/Private/EnginePrivate.h
@@ -161,7 +161,6 @@
 #include "Engine/CollisionProfile.h"
 #include "Curves/CurveLinearColor.h"
 #include "Curves/CurveVector.h"
-#include "Curves/CurveEdPresetCurve.h"
 #include "Engine/CurveTable.h"
 #include "GameFramework/DamageType.h"
 #include "Engine/DataTable.h"
@@ -183,15 +182,6 @@
 #include "Distributions/DistributionVectorConstantCurve.h"
 #include "Distributions/DistributionVectorUniform.h"
 #include "Distributions/DistributionVectorUniformCurve.h"
-#include "Engine/DynamicBlueprintBinding.h"
-#include "Engine/ComponentDelegateBinding.h"
-#include "Engine/InputDelegateBinding.h"
-#include "Engine/InputActionDelegateBinding.h"
-#include "Engine/InputAxisDelegateBinding.h"
-#include "Engine/InputAxisKeyDelegateBinding.h"
-#include "Engine/InputVectorAxisDelegateBinding.h"
-#include "Engine/InputKeyDelegateBinding.h"
-#include "Engine/InputTouchDelegateBinding.h"
 #include "Engine/Engine.h"
 #include "Engine/GameEngine.h"
 #include "Exporters/Exporter.h"
diff --git a/Engine/Source/Runtime/Engine/Private/InputActionDelegateBinding.cpp b/Engine/Source/Runtime/Engine/Private/InputActionDelegateBinding.cpp
--- a/Engine/Source/Runtime/Engine/Private/InputActionDelegateBinding.cpp
+++ b/Engine/Source/Runtime/Engine/Private/InputActionDelegateBinding.cpp
@@ -2,7 +2,7 @@
 
 #include "EnginePrivate.h"
 #include "BlueprintUtilities.h"
-
+#include "Engine/InputActionDelegateBinding.h"
 
 UInputActionDelegateBinding::UInputActionDelegateBinding(const class FPostConstructInitializeProperties& PCIP)
 	: Super(PCIP)
diff --git a/Engine/Source/Runtime/Engine/Private/InputAxisDelegateBinding.cpp b/Engine/Source/Runtime/Engine/Private/InputAxisDelegateBinding.cpp
--- a/Engine/Source/Runtime/Engine/Private/InputAxisDelegateBinding.cpp
+++ b/Engine/Source/Runtime/Engine/Private/InputAxisDelegateBinding.cpp
@@ -2,7 +2,7 @@
 
 #include "EnginePrivate.h"
 #include "BlueprintUtilities.h"
-
+#include "Engine/InputAxisDelegateBinding.h"
 
 UInputAxisDelegateBinding::UInputAxisDelegateBinding(const class FPostConstructInitializeProperties& PCIP)
 	: Super(PCIP)
diff --git a/Engine/Source/Runtime/Engine/Private/InputAxisKeyDelegateBinding.cpp b/Engine/Source/Runtime/Engine/Private/InputAxisKeyDelegateBinding.cpp
--- a/Engine/Source/Runtime/Engine/Private/InputAxisKeyDelegateBinding.cpp
+++ b/Engine/Source/Runtime/Engine/Private/InputAxisKeyDelegateBinding.cpp
@@ -2,6 +2,7 @@
 
 #include "EnginePrivate.h"
 #include "BlueprintUtilities.h"
+#include "Engine/InputAxisKeyDelegateBinding.h"
 
 
 UInputAxisKeyDelegateBinding::UInputAxisKeyDelegateBinding(const class FPostConstructInitializeProperties& PCIP)
diff --git a/Engine/Source/Runtime/Engine/Private/InputDelegateBinding.cpp b/Engine/Source/Runtime/Engine/Private/InputDelegateBinding.cpp
--- a/Engine/Source/Runtime/Engine/Private/InputDelegateBinding.cpp
+++ b/Engine/Source/Runtime/Engine/Private/InputDelegateBinding.cpp
@@ -2,7 +2,13 @@
 
 #include "EnginePrivate.h"
 #include "BlueprintUtilities.h"
-
+#include "Engine/InputDelegateBinding.h"
+#include "Engine/InputActionDelegateBinding.h"
+#include "Engine/InputAxisDelegateBinding.h"
+#include "Engine/InputKeyDelegateBinding.h"
+#include "Engine/InputTouchDelegateBinding.h"
+#include "Engine/InputAxisKeyDelegateBinding.h"
+#include "Engine/InputVectorAxisDelegateBinding.h"
 
 UInputDelegateBinding::UInputDelegateBinding(const class FPostConstructInitializeProperties& PCIP)
 	: Super(PCIP)
diff --git a/Engine/Source/Runtime/Engine/Private/InputKeyDelegateBinding.cpp b/Engine/Source/Runtime/Engine/Private/InputKeyDelegateBinding.cpp
--- a/Engine/Source/Runtime/Engine/Private/InputKeyDelegateBinding.cpp
+++ b/Engine/Source/Runtime/Engine/Private/InputKeyDelegateBinding.cpp
@@ -2,7 +2,7 @@
 
 #include "EnginePrivate.h"
 #include "BlueprintUtilities.h"
-
+#include "Engine/InputKeyDelegateBinding.h"
 
 UInputKeyDelegateBinding::UInputKeyDelegateBinding(const class FPostConstructInitializeProperties& PCIP)
 	: Super(PCIP)
diff --git a/Engine/Source/Runtime/Engine/Private/InputTouchDelegateBinding.cpp b/Engine/Source/Runtime/Engine/Private/InputTouchDelegateBinding.cpp
--- a/Engine/Source/Runtime/Engine/Private/InputTouchDelegateBinding.cpp
+++ b/Engine/Source/Runtime/Engine/Private/InputTouchDelegateBinding.cpp
@@ -2,7 +2,7 @@
 
 #include "EnginePrivate.h"
 #include "BlueprintUtilities.h"
-
+#include "Engine/InputTouchDelegateBinding.h"
 
 UInputTouchDelegateBinding::UInputTouchDelegateBinding(const class FPostConstructInitializeProperties& PCIP)
 	: Super(PCIP)
diff --git a/Engine/Source/Runtime/Engine/Private/InputVectorAxisDelegateBinding.cpp b/Engine/Source/Runtime/Engine/Private/InputVectorAxisDelegateBinding.cpp
--- a/Engine/Source/Runtime/Engine/Private/InputVectorAxisDelegateBinding.cpp
+++ b/Engine/Source/Runtime/Engine/Private/InputVectorAxisDelegateBinding.cpp
@@ -2,7 +2,7 @@
 
 #include "EnginePrivate.h"
 #include "BlueprintUtilities.h"
-
+#include "Engine/InputVectorAxisDelegateBinding.h"
 
 UInputVectorAxisDelegateBinding::UInputVectorAxisDelegateBinding(const class FPostConstructInitializeProperties& PCIP)
 	: Super(PCIP)
diff --git a/Engine/Source/Runtime/Engine/Private/LevelScriptActor.cpp b/Engine/Source/Runtime/Engine/Private/LevelScriptActor.cpp
--- a/Engine/Source/Runtime/Engine/Private/LevelScriptActor.cpp
+++ b/Engine/Source/Runtime/Engine/Private/LevelScriptActor.cpp
@@ -2,6 +2,7 @@
 
 
 #include "EnginePrivate.h"
+#include "Engine/InputDelegateBinding.h"
 #include "Engine/LevelScriptBlueprint.h"
 #include "Engine/LevelScriptActor.h"
 #include "BlueprintUtilities.h"
diff --git a/Engine/Source/Runtime/Engine/Private/Pawn.cpp b/Engine/Source/Runtime/Engine/Private/Pawn.cpp
--- a/Engine/Source/Runtime/Engine/Private/Pawn.cpp
+++ b/Engine/Source/Runtime/Engine/Private/Pawn.cpp
@@ -7,6 +7,7 @@
 =============================================================================*/
 
 #include "EnginePrivate.h"
+#include "Engine/InputDelegateBinding.h"
 #include "GameFramework/Character.h"
 #include "GameFramework/PawnMovementComponent.h"
 #include "Net/UnrealNetwork.h"
diff --git a/Engine/Source/Runtime/Engine/Private/PlayerController.cpp b/Engine/Source/Runtime/Engine/Private/PlayerController.cpp
--- a/Engine/Source/Runtime/Engine/Private/PlayerController.cpp
+++ b/Engine/Source/Runtime/Engine/Private/PlayerController.cpp
@@ -26,6 +26,7 @@
 #include "GameFramework/PawnMovementComponent.h"
 #include "Camera/CameraActor.h"
 #include "GenericPlatform/IForceFeedbackSystem.h"
+#include "Engine/InputDelegateBinding.h"
 
 DEFINE_LOG_CATEGORY(LogPlayerController);
 
diff --git a/Engine/Source/Runtime/Engine/Public/Engine.h b/Engine/Source/Runtime/Engine/Public/Engine.h
--- a/Engine/Source/Runtime/Engine/Public/Engine.h
+++ b/Engine/Source/Runtime/Engine/Public/Engine.h
@@ -219,7 +219,6 @@
 #include "GameFramework/CheatManager.h"
 #include "Curves/CurveLinearColor.h"
 #include "Curves/CurveVector.h"
-#include "Curves/CurveEdPresetCurve.h"
 #include "Engine/CurveTable.h"
 #include "GameFramework/DamageType.h"
 #include "Vehicles/TireType.h"
@@ -242,15 +241,6 @@
 #include "Distributions/DistributionVectorConstantCurve.h"
 #include "Distributions/DistributionVectorUniform.h"
 #include "Distributions/DistributionVectorUniformCurve.h"
-#include "Engine/DynamicBlueprintBinding.h"
-#include "Engine/ComponentDelegateBinding.h"
-#include "Engine/InputDelegateBinding.h"
-#include "Engine/InputActionDelegateBinding.h"
-#include "Engine/InputAxisDelegateBinding.h"
-#include "Engine/InputAxisKeyDelegateBinding.h"
-#include "Engine/InputVectorAxisDelegateBinding.h"
-#include "Engine/InputKeyDelegateBinding.h"
-#include "Engine/InputTouchDelegateBinding.h"
 #include "Engine/Engine.h"
 #include "Engine/GameEngine.h"
 #include "Exporters/Exporter.h"
@@ -337,9 +327,6 @@
 #include "GameFramework/TouchInterface.h"
 #include "Engine/UserDefinedEnum.h"
 #include "Engine/UserDefinedStruct.h"
-#include "VectorField/VectorField.h"
-#include "VectorField/VectorFieldAnimated.h"
-#include "VectorField/VectorFieldStatic.h"
 #include "Animation/VertexAnim/MorphTarget.h"
 
 #include "SystemSettings.h"					// Scalability options.
