Fixed regression with transforming level using 3d world origin

diff --git a/Engine/Source/Runtime/Engine/Classes/Engine/WorldComposition.h b/Engine/Source/Runtime/Engine/Classes/Engine/WorldComposition.h
--- a/Engine/Source/Runtime/Engine/Classes/Engine/WorldComposition.h
+++ b/Engine/Source/Runtime/Engine/Classes/Engine/WorldComposition.h
@@ -101,8 +101,8 @@ class ENGINE_API UWorldComposition : public UObject
 	/** Handles level is being removed from the world */
 	void OnLevelRemovedFromWorld(ULevel* InLevel);
 
-	/** @returns Level offset from zero origin, with respect to parent levels */
-	FIntPoint GetLevelOffset(ULevel* InLevel) const;
+	/** @returns Level offset from current origin, with respect to parent levels */
+	FIntVector GetLevelOffset(ULevel* InLevel) const;
 
 	/** @returns Level bounding box in current shifted space */
 	FBox GetLevelBounds(ULevel* InLevel) const;
diff --git a/Engine/Source/Runtime/Engine/Private/World.cpp b/Engine/Source/Runtime/Engine/Private/World.cpp
--- a/Engine/Source/Runtime/Engine/Private/World.cpp
+++ b/Engine/Source/Runtime/Engine/Private/World.cpp
@@ -2172,8 +2172,8 @@ UWorld* UWorld::DuplicateWorldForPIE(const FString& PackageName, UWorld* OwningW
 		// Correct placement will occur during UWorld::AddToWorld
 		if (EditorLevel->OwningWorld->WorldComposition && EditorLevel->bIsVisible)
 		{
-			FIntPoint LevelOffset = FIntPoint::ZeroValue - EditorLevel->OwningWorld->WorldComposition->GetLevelOffset(EditorLevel);
-			PIELevel->ApplyWorldOffset(FVector(LevelOffset, 0), false);
+			FIntVector LevelOffset = FIntVector::ZeroValue - EditorLevel->OwningWorld->WorldComposition->GetLevelOffset(EditorLevel);
+			PIELevel->ApplyWorldOffset(FVector(LevelOffset), false);
 		}
 	}
 
diff --git a/Engine/Source/Runtime/Engine/Private/WorldComposition.cpp b/Engine/Source/Runtime/Engine/Private/WorldComposition.cpp
--- a/Engine/Source/Runtime/Engine/Private/WorldComposition.cpp
+++ b/Engine/Source/Runtime/Engine/Private/WorldComposition.cpp
@@ -606,9 +606,9 @@ void UWorldComposition::OnLevelAddedToWorld(ULevel* InLevel)
 	}
 #endif
 		
-	// Move level according to current global origin offset
-	FIntPoint LevelOffset = GetLevelOffset(InLevel);
-	InLevel->ApplyWorldOffset(FVector(LevelOffset, 0), false);
+	// Move level according to current global origin
+	FIntVector LevelOffset = GetLevelOffset(InLevel);
+	InLevel->ApplyWorldOffset(FVector(LevelOffset), false);
 }
 
 void UWorldComposition::OnLevelRemovedFromWorld(ULevel* InLevel)
@@ -621,8 +621,8 @@ void UWorldComposition::OnLevelRemovedFromWorld(ULevel* InLevel)
 #endif
 	
 	// Move level to his local origin
-	FIntPoint LevelOffset = GetLevelOffset(InLevel);
-	InLevel->ApplyWorldOffset(-FVector(LevelOffset, 0), false);
+	FIntVector LevelOffset = GetLevelOffset(InLevel);
+	InLevel->ApplyWorldOffset(-FVector(LevelOffset), false);
 }
 
 void UWorldComposition::OnLevelPostLoad(ULevel* InLevel)
@@ -671,18 +671,19 @@ void UWorldComposition::OnLevelPostSave(ULevel* InLevel)
 	}
 }
 
-FIntPoint UWorldComposition::GetLevelOffset(ULevel* InLevel) const
+FIntVector UWorldComposition::GetLevelOffset(ULevel* InLevel) const
 {
 	UWorld* OwningWorld = GetWorld();
 	UPackage* LevelPackage = Cast<UPackage>(InLevel->GetOutermost());
 	
-	FIntPoint LevelPosition = FIntPoint::ZeroValue;
+	FIntVector LevelPosition = FIntVector::ZeroValue;
 	if (LevelPackage->WorldTileInfo)
 	{
-		LevelPosition = LevelPackage->WorldTileInfo->AbsolutePosition;
+		FIntPoint AbsolutePosition = LevelPackage->WorldTileInfo->AbsolutePosition;
+		LevelPosition = FIntVector(AbsolutePosition.X, AbsolutePosition.Y, 0);
 	}
 	
-	return LevelPosition - FIntPoint(OwningWorld->OriginLocation.X, OwningWorld->OriginLocation.Y);
+	return LevelPosition - OwningWorld->OriginLocation;
 }
 
 FBox UWorldComposition::GetLevelBounds(ULevel* InLevel) const
@@ -693,7 +694,7 @@ FBox UWorldComposition::GetLevelBounds(ULevel* InLevel) const
 	FBox LevelBBox(0);
 	if (LevelPackage->WorldTileInfo)
 	{
-		LevelBBox = LevelPackage->WorldTileInfo->Bounds.ShiftBy(FVector(GetLevelOffset(InLevel), 0));
+		LevelBBox = LevelPackage->WorldTileInfo->Bounds.ShiftBy(FVector(GetLevelOffset(InLevel)));
 	}
 	
 	return LevelBBox;
