Header cleanup: Engine module, public headers.

diff --git a/Engine/Source/Editor/UnrealEd/Private/Commandlets/DerivedDataCacheCommandlet.cpp b/Engine/Source/Editor/UnrealEd/Private/Commandlets/DerivedDataCacheCommandlet.cpp
--- a/Engine/Source/Editor/UnrealEd/Private/Commandlets/DerivedDataCacheCommandlet.cpp
+++ b/Engine/Source/Editor/UnrealEd/Private/Commandlets/DerivedDataCacheCommandlet.cpp
@@ -10,6 +10,7 @@
 #include "GlobalShader.h"
 #include "TargetPlatform.h"
 #include "IConsoleManager.h"
+#include "ShaderCompiler.h"
 
 DEFINE_LOG_CATEGORY_STATIC(LogDerivedDataCacheCommandlet, Log, All);
 
diff --git a/Engine/Source/Editor/UnrealEd/Private/EditorServer.cpp b/Engine/Source/Editor/UnrealEd/Private/EditorServer.cpp
--- a/Engine/Source/Editor/UnrealEd/Private/EditorServer.cpp
+++ b/Engine/Source/Editor/UnrealEd/Private/EditorServer.cpp
@@ -53,6 +53,7 @@
 
 #include "ComponentReregisterContext.h"
 #include "Engine/DocumentationActor.h"
+#include "ShaderCompiler.h"
 
 DEFINE_LOG_CATEGORY_STATIC(LogEditorServer, Log, All);
 
diff --git a/Engine/Source/Editor/UnrealEd/Private/Lightmass/Lightmass.cpp b/Engine/Source/Editor/UnrealEd/Private/Lightmass/Lightmass.cpp
--- a/Engine/Source/Editor/UnrealEd/Private/Lightmass/Lightmass.cpp
+++ b/Engine/Source/Editor/UnrealEd/Private/Lightmass/Lightmass.cpp
@@ -18,6 +18,7 @@
 #include "Lightmass/PrecomputedVisibilityVolume.h"
 #include "Lightmass/PrecomputedVisibilityOverrideVolume.h"
 #include "ComponentReregisterContext.h"
+#include "ShaderCompiler.h"
 
 extern FSwarmDebugOptions GSwarmDebugOptions;
 
diff --git a/Engine/Source/Editor/UnrealEd/Private/PerformanceMonitor.cpp b/Engine/Source/Editor/UnrealEd/Private/PerformanceMonitor.cpp
--- a/Engine/Source/Editor/UnrealEd/Private/PerformanceMonitor.cpp
+++ b/Engine/Source/Editor/UnrealEd/Private/PerformanceMonitor.cpp
@@ -4,6 +4,7 @@
 #include "UnrealEd.h"
 #include "SScalabilitySettings.h"
 #include "PerformanceMonitor.h"
+#include "ShaderCompiler.h"
 
 #define LOCTEXT_NAMESPACE "PerformanceMonitor"
 
diff --git a/Engine/Source/Editor/UnrealEd/Private/ShaderCompilingNotification.cpp b/Engine/Source/Editor/UnrealEd/Private/ShaderCompilingNotification.cpp
--- a/Engine/Source/Editor/UnrealEd/Private/ShaderCompilingNotification.cpp
+++ b/Engine/Source/Editor/UnrealEd/Private/ShaderCompilingNotification.cpp
@@ -3,6 +3,7 @@
 
 #include "UnrealEd.h"
 #include "MainFrame.h"
+#include "ShaderCompiler.h"
 
 /** Notification class for asynchronous shader compiling. */
 class FShaderCompilingNotificationImpl
diff --git a/Engine/Source/Editor/UnrealEd/Private/UnrealEdMisc.cpp b/Engine/Source/Editor/UnrealEd/Private/UnrealEdMisc.cpp
--- a/Engine/Source/Editor/UnrealEd/Private/UnrealEdMisc.cpp
+++ b/Engine/Source/Editor/UnrealEd/Private/UnrealEdMisc.cpp
@@ -31,6 +31,7 @@
 #include "EngineBuildSettings.h"
 #include "Slate.h"
 #include "DesktopPlatformModule.h"
+#include "ShaderCompiler.h"
 
 #define LOCTEXT_NAMESPACE "UnrealEd"
 
diff --git a/Engine/Source/Editor/WorldBrowser/Private/LevelCollectionModel.cpp b/Engine/Source/Editor/WorldBrowser/Private/LevelCollectionModel.cpp
--- a/Engine/Source/Editor/WorldBrowser/Private/LevelCollectionModel.cpp
+++ b/Engine/Source/Editor/WorldBrowser/Private/LevelCollectionModel.cpp
@@ -7,6 +7,7 @@
 #include "Matinee/MatineeActor.h"
 
 #include "LevelCollectionModel.h"
+#include "ShaderCompiler.h"
 
 #define LOCTEXT_NAMESPACE "WorldBrowser"
 
diff --git a/Engine/Source/Runtime/Engine/Classes/Engine/TextureRenderTarget2D.h b/Engine/Source/Runtime/Engine/Classes/Engine/TextureRenderTarget2D.h
--- a/Engine/Source/Runtime/Engine/Classes/Engine/TextureRenderTarget2D.h
+++ b/Engine/Source/Runtime/Engine/Classes/Engine/TextureRenderTarget2D.h
@@ -1,6 +1,7 @@
 // Copyright 1998-2014 Epic Games, Inc. All Rights Reserved.
 
 #pragma once
+#include "RenderUtils.h"
 #include "TextureRenderTarget2D.generated.h"
 
 /**
diff --git a/Engine/Source/Runtime/Engine/Private/EngineService.cpp b/Engine/Source/Runtime/Engine/Private/EngineService.cpp
--- a/Engine/Source/Runtime/Engine/Private/EngineService.cpp
+++ b/Engine/Source/Runtime/Engine/Private/EngineService.cpp
@@ -5,7 +5,9 @@
 =============================================================================*/
 
 #include "EnginePrivate.h"
-
+#include "MessageEndpoint.h"
+#include "MessageEndpointBuilder.h"
+#include "EngineServiceMessages.h"
 
 DEFINE_LOG_CATEGORY_STATIC(LogEngineService, Log, All)
 
diff --git a/Engine/Source/Runtime/Engine/Private/Materials/Material.cpp b/Engine/Source/Runtime/Engine/Private/Materials/Material.cpp
--- a/Engine/Source/Runtime/Engine/Private/Materials/Material.cpp
+++ b/Engine/Source/Runtime/Engine/Private/Materials/Material.cpp
@@ -31,6 +31,7 @@
 #include "MaterialCompiler.h"
 #include "TargetPlatform.h"
 #include "ComponentReregisterContext.h"
+#include "ShaderCompiler.h"
 
 #if WITH_EDITOR
 #include "UnrealEd.h"
diff --git a/Engine/Source/Runtime/Engine/Private/Materials/MaterialShared.cpp b/Engine/Source/Runtime/Engine/Private/Materials/MaterialShared.cpp
--- a/Engine/Source/Runtime/Engine/Private/Materials/MaterialShared.cpp
+++ b/Engine/Source/Runtime/Engine/Private/Materials/MaterialShared.cpp
@@ -17,8 +17,11 @@
 #include "Developer/TargetPlatform/Public/TargetPlatform.h"
 #include "ComponentReregisterContext.h"
 #include "EngineModule.h"
+
 #include "LocalVertexFactory.h"
 
+#include "VertexFactory.h"
+#include "RendererInterface.h"
 DEFINE_LOG_CATEGORY(LogMaterial);
 
 FName MaterialQualityLevelNames[] = 
@@ -2342,3 +2345,29 @@ bool FMaterialInstanceBasePropertyOverrides::Update(const FMaterialInstanceBaseP
 
 	return bRet;
 }
+
+bool FMaterialShaderMapId::ContainsShaderType(const FShaderType* ShaderType) const
+{
+	for (int32 TypeIndex = 0; TypeIndex < ShaderTypeDependencies.Num(); TypeIndex++)
+	{
+		if (ShaderTypeDependencies[TypeIndex].ShaderType == ShaderType)
+		{
+			return true;
+		}
+	}
+
+	return false;
+}
+
+bool FMaterialShaderMapId::ContainsVertexFactoryType(const FVertexFactoryType* VFType) const
+{
+	for (int32 TypeIndex = 0; TypeIndex < VertexFactoryTypeDependencies.Num(); TypeIndex++)
+	{
+		if (VertexFactoryTypeDependencies[TypeIndex].VertexFactoryType == VFType)
+		{
+			return true;
+		}
+	}
+
+	return false;
+}
diff --git a/Engine/Source/Runtime/Engine/Private/ReflectionCaptureComponent.cpp b/Engine/Source/Runtime/Engine/Private/ReflectionCaptureComponent.cpp
--- a/Engine/Source/Runtime/Engine/Private/ReflectionCaptureComponent.cpp
+++ b/Engine/Source/Runtime/Engine/Private/ReflectionCaptureComponent.cpp
@@ -8,6 +8,8 @@
 #include "DerivedDataCacheInterface.h"
 #include "TargetPlatform.h"
 #include "EngineModule.h"
+#include "RendererInterface.h"
+#include "ShaderCompiler.h"
 
 /** 
  * Size of all reflection captures.
diff --git a/Engine/Source/Runtime/Engine/Private/ShaderCompiler/ShaderCompiler.cpp b/Engine/Source/Runtime/Engine/Private/ShaderCompiler/ShaderCompiler.cpp
--- a/Engine/Source/Runtime/Engine/Private/ShaderCompiler/ShaderCompiler.cpp
+++ b/Engine/Source/Runtime/Engine/Private/ShaderCompiler/ShaderCompiler.cpp
@@ -11,6 +11,8 @@
 #include "TargetPlatform.h"
 #include "DerivedDataCacheInterface.h"
 #include "EngineModule.h"
+#include "ShaderCompiler.h"
+#include "RendererInterface.h"
 
 DEFINE_LOG_CATEGORY(LogShaderCompilers);
 
diff --git a/Engine/Source/Runtime/Engine/Private/SkyLightComponent.cpp b/Engine/Source/Runtime/Engine/Private/SkyLightComponent.cpp
--- a/Engine/Source/Runtime/Engine/Private/SkyLightComponent.cpp
+++ b/Engine/Source/Runtime/Engine/Private/SkyLightComponent.cpp
@@ -14,6 +14,7 @@
 #include "Net/UnrealNetwork.h"
 #include "MapErrors.h"
 #include "ComponentInstanceDataCache.h"
+#include "ShaderCompiler.h"
 
 #define LOCTEXT_NAMESPACE "SkyLightComponent"
 
diff --git a/Engine/Source/Runtime/Engine/Private/World.cpp b/Engine/Source/Runtime/Engine/Private/World.cpp
--- a/Engine/Source/Runtime/Engine/Private/World.cpp
+++ b/Engine/Source/Runtime/Engine/Private/World.cpp
@@ -50,6 +50,9 @@
 #include "EngineModule.h"
 #include "ContentStreaming.h"
 #include "Runtime/Engine/Classes/Landscape/LandscapeComponent.h"
+#include "RendererInterface.h"
+#include "DataChannel.h"
+#include "ShaderCompiler.h"
 
 DEFINE_LOG_CATEGORY_STATIC(LogWorld, Log, All);
 DEFINE_LOG_CATEGORY(LogSpawn);
diff --git a/Engine/Source/Runtime/Engine/Public/MaterialCompiler.h b/Engine/Source/Runtime/Engine/Public/MaterialCompiler.h
--- a/Engine/Source/Runtime/Engine/Public/MaterialCompiler.h
+++ b/Engine/Source/Runtime/Engine/Public/MaterialCompiler.h
@@ -6,10 +6,10 @@
 
 #pragma once
 
-#include "Materials/MaterialExpressionWorldPosition.h"
-#include "Materials/MaterialExpressionTextureSample.h"
-#include "Materials/MaterialExpressionSpeedTree.h"
 #include "MaterialShared.h"
+#include "Materials/MaterialExpressionSpeedTree.h"
+#include "Materials/MaterialExpressionTextureSample.h"
+#include "Materials/MaterialExpressionWorldPosition.h"
 
 /** 
  * The interface used to translate material expressions into executable code. 
diff --git a/Engine/Source/Runtime/Engine/Public/MaterialShaderType.h b/Engine/Source/Runtime/Engine/Public/MaterialShaderType.h
--- a/Engine/Source/Runtime/Engine/Public/MaterialShaderType.h
+++ b/Engine/Source/Runtime/Engine/Public/MaterialShaderType.h
@@ -7,8 +7,6 @@
 #pragma once
 
 #include "Shader.h"
-#include "ShaderParameters.h"
-#include "MaterialShared.h"
 #include "GlobalShader.h"
 
 /** A macro to implement material shaders. */
diff --git a/Engine/Source/Runtime/Engine/Public/MaterialShared.h b/Engine/Source/Runtime/Engine/Public/MaterialShared.h
--- a/Engine/Source/Runtime/Engine/Public/MaterialShared.h
+++ b/Engine/Source/Runtime/Engine/Public/MaterialShared.h
@@ -5,28 +5,33 @@
 =============================================================================*/
 
 #pragma once
-
+#include "Misc/SecureHash.h"
 #include "RefCounting.h"
-#include "ShaderCore.h"
-#include "Shader.h"
-#include "VertexFactory.h"
 #include "RenderResource.h"
-#include "ShaderCompiler.h"
 #include "UniformBuffer.h"
 #include "SceneTypes.h"
 #include "StaticParameterSet.h"
+#include "Shader.h"
+#include "VertexFactory.h"
 
-class FMeshMaterialShaderMap;
 class FMaterialShaderMap;
 class FMaterialShaderType;
 class FMaterial;
 class FMaterialRenderProxy;
+class FMeshMaterialShaderMap;
+class FMeshMaterialShaderType;
+class FShaderCompileJob;
+class FShaderType;
+class FShaderTypeDependency;
+class FVertexFactoryType;
+class FVertexFactoryTypeDependency;
 class UMaterial;
 class UMaterialInstance;
 class UMaterialExpression;
 class UMaterialInterface;
 class UTexture;
 struct FExpressionInput;
+struct FShaderCompilerEnvironment;
 
 #define ME_CAPTION_HEIGHT		18
 #define ME_STD_VPADDING			16
@@ -362,14 +367,17 @@ public:
 	/** A hash of the base property overrides for this material instance. */
 	FSHAHash BasePropertyOverridesHash;
 	
+
 	FMaterialShaderMapId()
 		: BaseMaterialId(0, 0, 0, 0)
 		, QualityLevel(EMaterialQualityLevel::High)
 		, FeatureLevel(ERHIFeatureLevel::SM4)
 		, Usage(EMaterialShaderMapUsage::Default)
-	{
-	}
-	
+	{ }
+
+	~FMaterialShaderMapId()
+	{ }
+
 	void SetShaderDependencies(const TArray<FShaderType*>& ShaderTypes, const TArray<FVertexFactoryType*>& VFTypes);
 
 	void Serialize(FArchive& Ar);
@@ -408,32 +416,10 @@ public:
 	void AppendKeyString(FString& KeyString) const;
 
 	/** Returns true if the requested shader type is a dependency of this shader map Id. */
-	bool ContainsShaderType(const FShaderType* ShaderType) const
-	{
-		for (int32 TypeIndex = 0; TypeIndex < ShaderTypeDependencies.Num(); TypeIndex++)
-		{
-			if (ShaderTypeDependencies[TypeIndex].ShaderType == ShaderType)
-			{
-				return true;
-			}
-		}
-
-		return false;
-	}
+	bool ContainsShaderType(const FShaderType* ShaderType) const;
 
 	/** Returns true if the requested vertex factory type is a dependency of this shader map Id. */
-	bool ContainsVertexFactoryType(const FVertexFactoryType* VFType) const
-	{
-		for (int32 TypeIndex = 0; TypeIndex < VertexFactoryTypeDependencies.Num(); TypeIndex++)
-		{
-			if (VertexFactoryTypeDependencies[TypeIndex].VertexFactoryType == VFType)
-			{
-				return true;
-			}
-		}
-
-		return false;
-	}
+	bool ContainsVertexFactoryType(const FVertexFactoryType* VFType) const;
 };
 
 /**
diff --git a/Engine/Source/Runtime/Engine/Public/MeshBatch.h b/Engine/Source/Runtime/Engine/Public/MeshBatch.h
--- a/Engine/Source/Runtime/Engine/Public/MeshBatch.h
+++ b/Engine/Source/Runtime/Engine/Public/MeshBatch.h
@@ -1,9 +1,9 @@
 // Copyright 1998-2014 Epic Games, Inc. All Rights Reserved.
 
 #pragma once
+#include "MaterialShared.h"
 
-#include "PrimitiveUniformShaderParameters.h"
-
+class FLightCacheInterface;
 /**
  * A batch mesh element definition.
  */
diff --git a/Engine/Source/Runtime/Engine/Public/MeshMaterialShaderType.h b/Engine/Source/Runtime/Engine/Public/MeshMaterialShaderType.h
--- a/Engine/Source/Runtime/Engine/Public/MeshMaterialShaderType.h
+++ b/Engine/Source/Runtime/Engine/Public/MeshMaterialShaderType.h
@@ -7,7 +7,7 @@
 #pragma once
 
 #include "Shader.h"
-#include "MaterialShader.h"
+#include "MaterialShaderType.h"
 
 /**
  * A shader meta type for material-linked shaders which use a vertex factory.
diff --git a/Engine/Source/Runtime/Engine/Public/Model.h b/Engine/Source/Runtime/Engine/Public/Model.h
--- a/Engine/Source/Runtime/Engine/Public/Model.h
+++ b/Engine/Source/Runtime/Engine/Public/Model.h
@@ -6,10 +6,9 @@
 
 #pragma once
 
-#include "RawIndexBuffer.h"
 #include "LocalVertexFactory.h"
-#include "RenderUtils.h"
 #include "SceneTypes.h"
+#include "RawIndexBuffer.h"
 
 struct FLightmassPrimitiveSettings;
 struct FStaticLightingVertex;
@@ -17,6 +16,7 @@ class ABrush;
 class ULevel;
 class ULightComponent;
 class UMaterialInterface;
+class FRawIndexBuffer16or32;
 class UModel;
 class UPolys;
 
