Added a togglable flag to APawn for enabling it to affect navigation generation #UE4

- also, added a deducated NavModifierComponent to be used to empower actors to affect navigation generation in a streamed-lined fashion

WIP

diff --git a/Engine/Source/Runtime/Engine/Classes/AI/Navigation/NavModifierComponent.h b/Engine/Source/Runtime/Engine/Classes/AI/Navigation/NavModifierComponent.h
new file mode 100644
--- /dev/null
+++ b/Engine/Source/Runtime/Engine/Classes/AI/Navigation/NavModifierComponent.h
@@ -0,0 +1,26 @@
+// Copyright 1998-2014 Epic Games, Inc. All Rights Reserved.
+
+#pragma once
+
+#include "AI/Navigation/NavRelevantComponent.h"
+#include "NavModifierComponent.generated.h"
+
+class UNavArea;
+struct FCompositeNavModifier;
+
+UCLASS(ClassGroup = (Navigation), meta = (BlueprintSpawnableComponent), hidecategories = (Activation))
+class UNavModifierComponent : public UNavRelevantComponent
+{
+	GENERATED_UCLASS_BODY()
+
+	UPROPERTY(EditAnywhere, Category = Navigation)
+	TSubclassOf<UNavArea> AreaClass;
+
+	virtual void OnRegister() override;
+	virtual void OnApplyModifiers(FCompositeNavModifier& Modifiers) override;
+	virtual void OnOwnerRegistered() override;
+	virtual void OnOwnerUnregistered() override;
+
+protected:
+	FVector ObstacleExtent;
+};
diff --git a/Engine/Source/Runtime/Engine/Classes/AI/Navigation/NavRelevantComponent.h b/Engine/Source/Runtime/Engine/Classes/AI/Navigation/NavRelevantComponent.h
--- a/Engine/Source/Runtime/Engine/Classes/AI/Navigation/NavRelevantComponent.h
+++ b/Engine/Source/Runtime/Engine/Classes/AI/Navigation/NavRelevantComponent.h
@@ -17,7 +17,9 @@ class ENGINE_API UNavRelevantComponent : public UActorComponent
 
 	virtual void OnApplyModifiers(struct FCompositeNavModifier& Modifiers);
 
+	UFUNCTION(BlueprintCallable, Category = "Navigation")
 	void SetNavigationRelevancy(bool bRelevant);
+
 	bool IsNavigationRelevant() const { return bNavigationRelevant; }
 
 	/** force refresh in navigation octree */
diff --git a/Engine/Source/Runtime/Engine/Classes/GameFramework/Pawn.h b/Engine/Source/Runtime/Engine/Classes/GameFramework/Pawn.h
--- a/Engine/Source/Runtime/Engine/Classes/GameFramework/Pawn.h
+++ b/Engine/Source/Runtime/Engine/Classes/GameFramework/Pawn.h
@@ -50,6 +50,14 @@ public:
 	UPROPERTY(EditAnywhere, BlueprintReadWrite, Category=Pawn)
 	uint32 bUseControllerRotationRoll:1;
 
+	/** If set to false (default) given pawn instance will never affect navigation generation. 
+	 *	Setting it to true will result in using regular AActor's navigation relevancy 
+	 *	calculation to check if this pawn instance should affect navigation generation
+	 *	Use SetCanAffectNavigationGeneration to change this value at runtime.
+	 *	Note that modifying this value at runtime will result in any navigation change only if runtime navigation generation is enabled. */
+	UPROPERTY(EditAnywhere, BlueprintReadOnly, Category = Pawn)
+	uint32 bCanAffectNavigationGeneration : 1;
+
 private:
 	/* Whether this Pawn's input handling is enabled.  Pawn must still be possessed to get input even if this is true */
 	uint32 bInputEnabled:1;
@@ -200,7 +208,7 @@ public:
 	virtual void PostRegisterAllComponents() override;
 	virtual float TakeDamage(float Damage, struct FDamageEvent const& DamageEvent, AController* EventInstigator, AActor* DamageCauser) override;
 	virtual void BecomeViewTarget(APlayerController* PC) override;
-	virtual bool UpdateNavigationRelevancy() override { SetNavigationRelevancy(false); return false; }
+	virtual bool UpdateNavigationRelevancy() override; 
 	virtual void EnableInput(APlayerController* PlayerController) override;
 	virtual void DisableInput(APlayerController* PlayerController) override;
 
@@ -208,6 +216,11 @@ public:
 	virtual bool CanBeBaseForCharacter(APawn* APawn) const override;
 	// End AActor Interface
 
+	/** Use SetCanAffectNavigationGeneration to change this value at runtime.
+	 *	Note that calling this function at runtime will result in any navigation change only if runtime navigation generation is enabled. */
+	UFUNCTION(BlueprintCallable, Category="Navigation")
+	void SetCanAffectNavigationGeneration(bool bNewValue);
+
 	// Begin INavAgentInterface Interface
 	virtual const struct FNavAgentProperties* GetNavAgentProperties() const override;
 	/** Basically retrieved pawn's location on navmesh */
diff --git a/Engine/Source/Runtime/Engine/Private/AI/Navigation/NavModifierComponent.cpp b/Engine/Source/Runtime/Engine/Private/AI/Navigation/NavModifierComponent.cpp
new file mode 100644
--- /dev/null
+++ b/Engine/Source/Runtime/Engine/Private/AI/Navigation/NavModifierComponent.cpp
@@ -0,0 +1,55 @@
+// Copyright 1998-2014 Epic Games, Inc. All Rights Reserved.
+
+#include "EnginePrivate.h"
+#include "AI/NavigationModifier.h"
+#include "AI/Navigation/NavModifierComponent.h"
+
+UNavModifierComponent::UNavModifierComponent(const FPostConstructInitializeProperties& PCIP)
+: Super(PCIP)
+{
+	AreaClass = UNavArea_Null::StaticClass();
+	Bounds = FNavigationSystem::InvalidBoundingBox;
+}
+
+void UNavModifierComponent::OnOwnerRegistered()
+{
+	if (AActor* MyOwner = GetOwner())
+	{
+		MyOwner->UpdateNavigationRelevancy();
+	}
+}
+
+void UNavModifierComponent::OnOwnerUnregistered()
+{
+	if (AActor* MyOwner = GetOwner())
+	{
+		MyOwner->UpdateNavigationRelevancy();
+	}
+}
+
+void UNavModifierComponent::OnRegister()
+{
+	Super::OnRegister();
+
+	const AActor* MyOwner = GetOwner();
+
+	if (MyOwner)
+	{
+		const float Radius = MyOwner->GetSimpleCollisionRadius();
+		const float HalfHeght = MyOwner->GetSimpleCollisionHalfHeight();
+		const FVector Loc = MyOwner->GetActorLocation();
+		ObstacleExtent = FVector(Radius, Radius, HalfHeght);
+		Bounds = FBox(Loc + ObstacleExtent, Loc - ObstacleExtent);
+	}
+}
+
+void UNavModifierComponent::OnApplyModifiers(FCompositeNavModifier& Modifiers)
+{
+	const AActor* MyOwner = GetOwner();
+	if (MyOwner)
+	{
+		const FVector Loc = MyOwner->GetActorLocation();
+		Bounds = FBox(Loc + ObstacleExtent, Loc - ObstacleExtent);
+		Modifiers.Add(FAreaNavModifier(Bounds, FTransform::Identity, AreaClass));
+	}
+}
diff --git a/Engine/Source/Runtime/Engine/Private/Pawn.cpp b/Engine/Source/Runtime/Engine/Private/Pawn.cpp
--- a/Engine/Source/Runtime/Engine/Private/Pawn.cpp
+++ b/Engine/Source/Runtime/Engine/Private/Pawn.cpp
@@ -7,7 +7,6 @@
 =============================================================================*/
 
 #include "EnginePrivate.h"
-#include "Engine/InputDelegateBinding.h"
 #include "GameFramework/Character.h"
 #include "GameFramework/PawnMovementComponent.h"
 #include "Net/UnrealNetwork.h"
@@ -17,6 +16,7 @@
 #include "NetworkingDistanceConstants.h"
 #include "VisualLog.h"
 #include "AIController.h"
+#include "Engine/InputDelegateBinding.h"
 
 DEFINE_LOG_CATEGORY(LogDamage);
 DEFINE_LOG_CATEGORY_STATIC(LogPawn, Warning, All);
@@ -130,6 +130,28 @@ void APawn::UpdateNavAgent()
 	}
 }
 
+bool APawn::UpdateNavigationRelevancy()
+{
+	if (bCanAffectNavigationGeneration == false)
+	{
+		SetNavigationRelevancy(false);
+		return false;
+	}
+	else
+	{
+		return Super::UpdateNavigationRelevancy();
+	}
+}
+
+void APawn::SetCanAffectNavigationGeneration(bool bNewValue)
+{
+	if (bCanAffectNavigationGeneration != bNewValue)
+	{
+		bCanAffectNavigationGeneration = bNewValue;
+		UpdateNavigationRelevancy();
+	}
+}
+
 void APawn::PawnStartFire(uint8 FireModeNum) {}
 
 AActor* APawn::GetMovementBaseActor(const APawn* Pawn)
