EngineAutomationTests.cpp
--Fixes TTP# 342453
--Change the map being used to be the games default map.

diff --git a/Engine/Source/Runtime/Engine/Private/Tests/EngineAutomationTests.cpp b/Engine/Source/Runtime/Engine/Private/Tests/EngineAutomationTests.cpp
--- a/Engine/Source/Runtime/Engine/Private/Tests/EngineAutomationTests.cpp
+++ b/Engine/Source/Runtime/Engine/Private/Tests/EngineAutomationTests.cpp
@@ -74,15 +74,19 @@ IMPLEMENT_SIMPLE_AUTOMATION_TEST( FSetResTest, "Windows.Set Resolution", EAutoma
  */
 bool FSetResTest::RunTest(const FString& Parameters)
 {
-	UAutomationTestSettings const* AutomationTestSettings = GetDefault<UAutomationTestSettings>();
-	check(AutomationTestSettings);
-	FString MapName = AutomationTestSettings->AutomationTestmap.FilePath;
+	//Gets the default map that the game uses.
+	const UGameMapsSettings* GameMapsSettings = GetDefault<UGameMapsSettings>();
+	const FString& MapName = GameMapsSettings->GetGameDefaultMap();
+
+	//Opens the actual default map in game.
 	GEngine->Exec(GetSimpleEngineAutomationTestGameWorld(GetTestFlags()), *FString::Printf(TEXT("Open %s"), *MapName));
 
+	//Gets the current resolution.
 	int32 ResX = GSystemResolution.ResX;
 	int32 ResY = GSystemResolution.ResY;
 	FString RestoreResolutionString = FString::Printf(TEXT("setres %dx%d"), ResX, ResY);
 
+	//Change the resolution and then restore it.
 	ADD_LATENT_AUTOMATION_COMMAND(FEngineWaitLatentCommand(2.0f));
 	ADD_LATENT_AUTOMATION_COMMAND(FExecStringLatentCommand(TEXT("setres 640x480")));
 	ADD_LATENT_AUTOMATION_COMMAND(FEngineWaitLatentCommand(2.0f));
