#ue4 - non unity cis fix

diff --git a/Engine/Source/Runtime/Engine/Private/AI/Navigation/NavModifierComponent.cpp b/Engine/Source/Runtime/Engine/Private/AI/Navigation/NavModifierComponent.cpp
--- a/Engine/Source/Runtime/Engine/Private/AI/Navigation/NavModifierComponent.cpp
+++ b/Engine/Source/Runtime/Engine/Private/AI/Navigation/NavModifierComponent.cpp
@@ -3,6 +3,7 @@
 #include "EnginePrivate.h"
 #include "AI/NavigationModifier.h"
 #include "AI/Navigation/NavModifierComponent.h"
+#include "AI/Navigation/NavAreas/NavArea_Null.h"
 
 UNavModifierComponent::UNavModifierComponent(const FPostConstructInitializeProperties& PCIP)
 : Super(PCIP)
