#UE4 - added an inactive player state lifespan config value to InactivePlayerState code feature
- allows games to configure amount of time a playerstate remains on player logout

diff --git a/Engine/Config/BaseGame.ini b/Engine/Config/BaseGame.ini
--- a/Engine/Config/BaseGame.ini
+++ b/Engine/Config/BaseGame.ini
@@ -2,6 +2,7 @@
 +LocalizationPaths=%GAMEDIR%Content/Localization/Game
 
 [/Script/Engine.GameMode]
+InactivePlayerStateLifeSpan=300
 
 [DefaultPlayer]
 Name=Player
diff --git a/Engine/Source/Runtime/Engine/Classes/GameFramework/GameMode.h b/Engine/Source/Runtime/Engine/Classes/GameFramework/GameMode.h
--- a/Engine/Source/Runtime/Engine/Classes/GameFramework/GameMode.h
+++ b/Engine/Source/Runtime/Engine/Classes/GameFramework/GameMode.h
@@ -243,6 +243,10 @@ protected:
 	UPROPERTY(config)
 	TArray<struct FGameClassShortName> GameModeClassAliases;
 
+	/** Time a playerstate will stick around in an inactive state after a player logout */
+	UPROPERTY(config)
+	float InactivePlayerStateLifeSpan;
+
 public:
 
 	/** Alters the synthetic bandwidth limit for a running game. */
diff --git a/Engine/Source/Runtime/Engine/Private/GameMode.cpp b/Engine/Source/Runtime/Engine/Private/GameMode.cpp
--- a/Engine/Source/Runtime/Engine/Private/GameMode.cpp
+++ b/Engine/Source/Runtime/Engine/Private/GameMode.cpp
@@ -1531,8 +1531,8 @@ void AGameMode::AddInactivePlayer(APlayerState* PlayerState, APlayerController*
 		// make PlayerState inactive
 		NewPlayerState->SetReplicates(false);
 
-		// delete after 5 minutes
-		NewPlayerState->SetLifeSpan( 300 );
+		// delete after some time
+		NewPlayerState->SetLifeSpan(InactivePlayerStateLifeSpan);
 
 		// On console, we have to check the unique net id as network address isn't valid
 		bool bIsConsole = GEngine->IsConsoleBuild();
diff --git a/Engine/Source/Runtime/Engine/Private/UnrealEngine.cpp b/Engine/Source/Runtime/Engine/Private/UnrealEngine.cpp
--- a/Engine/Source/Runtime/Engine/Private/UnrealEngine.cpp
+++ b/Engine/Source/Runtime/Engine/Private/UnrealEngine.cpp
@@ -7160,13 +7160,13 @@ void UEngine::ShutdownAllNetDrivers()
 	{
 		TArray<FNamedNetDriver> & ActiveNetDrivers = It->ActiveNetDrivers;
 
-	for (int32 Index = 0; Index < ActiveNetDrivers.Num(); Index++)
-	{
-		FNamedNetDriver& NamedNetDriver = ActiveNetDrivers[Index];
-		UNetDriver* NetDriver = NamedNetDriver.NetDriver;
-		if (NetDriver)
+		for (int32 Index = 0; Index < ActiveNetDrivers.Num(); Index++)
 		{
-			UE_LOG(LogNet, Log, TEXT("World NetDriver shutdown %s [%s]"), *NetDriver->GetName(), *NetDriver->NetDriverName.ToString());
+			FNamedNetDriver& NamedNetDriver = ActiveNetDrivers[Index];
+			UNetDriver* NetDriver = NamedNetDriver.NetDriver;
+			if (NetDriver)
+			{
+				UE_LOG(LogNet, Log, TEXT("World NetDriver shutdown %s [%s]"), *NetDriver->GetName(), *NetDriver->NetDriverName.ToString());
 				UWorld * World = NetDriver->GetWorld();
 				if (World)
 				{
@@ -7174,11 +7174,11 @@ void UEngine::ShutdownAllNetDrivers()
 				}
 				NetDriver->SetWorld(NULL);
 				DestroyNamedNetDriver(It->World(), NetDriver->NetDriverName);
+			}
 		}
-	}
 
-	ActiveNetDrivers.Empty();
-}
+		ActiveNetDrivers.Empty();
+	}
 }
 
 UNetDriver* FindNamedNetDriver_Local(const TArray<FNamedNetDriver>& ActiveNetDrivers, FName NetDriverName)
