Convert between vector and vector2d in blueprints

diff --git a/Engine/Source/Runtime/Engine/Classes/Kismet/KismetMathLibrary.h b/Engine/Source/Runtime/Engine/Classes/Kismet/KismetMathLibrary.h
--- a/Engine/Source/Runtime/Engine/Classes/Kismet/KismetMathLibrary.h
+++ b/Engine/Source/Runtime/Engine/Classes/Kismet/KismetMathLibrary.h
@@ -663,6 +663,14 @@ class UKismetMathLibrary : public UBlueprintFunctionLibrary
 	/** Convert a vector to a transform. Uses vector as location */
 	UFUNCTION(BlueprintPure, meta=(FriendlyName = "ToTransform (vector)", CompactNodeTitle = "->", Keywords="cast convert"), Category="Math|Conversions")
 	static FTransform Conv_VectorToTransform(FVector InLocation);
+	
+	/** Convert a Vector to a Vector2D */
+	UFUNCTION(BlueprintPure, meta=(FriendlyName = "ToVector2D (Vector)", CompactNodeTitle = "->", Keywords="cast convert"), Category="Math|Conversions")
+	static FVector2D Conv_VectorToVector2D(FVector InVector);
+
+	/** Convert a Vector2D to a Vector */
+	UFUNCTION(BlueprintPure, meta=(FriendlyName = "ToVector (Vector2D)", CompactNodeTitle = "->", Keywords="cast convert"), Category="Math|Conversions")
+	static FVector Conv_Vector2DToVector(FVector2D InVector2D, float Z = 0);
 
 	/** Convert a float into a vector, where each element is that float */
 	UFUNCTION(BlueprintPure, meta=(FriendlyName = "ToVector (float)", CompactNodeTitle = "->", Keywords="cast convert"), Category="Math|Conversions")
diff --git a/Engine/Source/Runtime/Engine/Private/KismetMathLibrary.cpp b/Engine/Source/Runtime/Engine/Private/KismetMathLibrary.cpp
--- a/Engine/Source/Runtime/Engine/Private/KismetMathLibrary.cpp
+++ b/Engine/Source/Runtime/Engine/Private/KismetMathLibrary.cpp
@@ -1053,6 +1053,16 @@ FLinearColor UKismetMathLibrary::Conv_VectorToLinearColor(FVector InVec)
 	return FLinearColor(InVec);	
 }
 
+FVector2D UKismetMathLibrary::Conv_VectorToVector2D(FVector InVec)
+{
+	return FVector2D(InVec);
+}
+
+FVector UKismetMathLibrary::Conv_Vector2DToVector(FVector2D InVec2D, float Z)
+{
+	return FVector(InVec2D, Z);
+}
+
 FVector UKismetMathLibrary::Conv_LinearColorToVector(FLinearColor InLinearColor)
 {
 	return FVector(InLinearColor);
