Fixed a crash in navmesh generation when reducing size of navmesh while navmesh is being built #UE4

- also fixed an edge case where navmesh width or height in tiles could end up being 0

diff --git a/Engine/Source/Runtime/Engine/Private/AI/Navigation/RecastNavMeshGenerator.cpp b/Engine/Source/Runtime/Engine/Private/AI/Navigation/RecastNavMeshGenerator.cpp
--- a/Engine/Source/Runtime/Engine/Private/AI/Navigation/RecastNavMeshGenerator.cpp
+++ b/Engine/Source/Runtime/Engine/Private/AI/Navigation/RecastNavMeshGenerator.cpp
@@ -3154,10 +3154,9 @@ void FRecastNavMeshGenerator::Init()
 	const float* BMin = Config.bmin;
 	const float* BMax = Config.bmax;
 	rcCalcGridSize(BMin, BMax, CellSize, &GridWidth, &GridHeight);
-	const int32 ConstTileSize = Config.tileSize;
-	const float tcs = Config.tileSize * Config.cs;
-	int32 NewTilesWidth = (GridWidth + ConstTileSize-1) / ConstTileSize;
-	int32 NewTilesHeight = (GridHeight + ConstTileSize-1) / ConstTileSize;
+	const float TileSizeInWorldUnits = Config.tileSize * Config.cs;
+	int32 NewTilesWidth = (GridWidth + Config.tileSize - 1) / Config.tileSize;
+	int32 NewTilesHeight = (GridHeight + Config.tileSize - 1) / Config.tileSize;
 
 	// limit max amount of tiles: config values
 	if (NavGenParams)
@@ -3214,6 +3213,12 @@ void FRecastNavMeshGenerator::Init()
 	if (bInitialized == false || NewTilesHeight != TilesHeight || NewTilesWidth != TilesWidth) 
 	{
 		FScopeLock Lock(&TileGenerationLock);
+
+		{
+			FScopeLock Lock(&NavMeshDirtyLock);
+			DirtyAreas.Reset();
+			DirtyGenerators.Empty(MaxActiveGenerators);
+		}
 		
 		float TileBmin[3];
 		float TileBmax[3];
@@ -3232,13 +3237,13 @@ void FRecastNavMeshGenerator::Init()
 			const int32 X = TileData->X;
 			const int32 Y = TileData->Y;
 
-			TileBmin[0] = BMin[0] + X*tcs;
+			TileBmin[0] = BMin[0] + X*TileSizeInWorldUnits;
 			TileBmin[1] = BMin[1];
-			TileBmin[2] = BMin[2] + Y*tcs;
+			TileBmin[2] = BMin[2] + Y*TileSizeInWorldUnits;
 
-			TileBmax[0] = BMin[0] + (X+1)*tcs;
+			TileBmax[0] = BMin[0] + (X+1)*TileSizeInWorldUnits;
 			TileBmax[1] = BMax[1];
-			TileBmax[2] = BMin[2] + (Y+1)*tcs;
+			TileBmax[2] = BMin[2] + (Y+1)*TileSizeInWorldUnits;
 
 			// @todo check if this tile overlaps current navmesh generation boundaries 
 			// if not remove it
@@ -3418,26 +3423,42 @@ void FRecastNavMeshGenerator::SetUpGeneration(float CellSize, float CellHeight,
 	// expand bounds a bit to support later inclusion tests
 	NavBounds.ExpandBy(CellSize);
 
+	bool bAdjust = false;
 	bool bClampBounds = false;
 	const float ExtentLimit = float(MAX_int32);
 	FVector BoundsExtent = NavBounds.GetExtent();
 	if (BoundsExtent.X > ExtentLimit)
 	{
 		BoundsExtent.X = ExtentLimit;
-		bClampBounds = true;
+		bClampBounds = bAdjust = true;
+	}
+	else if (BoundsExtent.X < Config.cs)
+	{
+		// minor adjustment to have at least 1 voxel of size here
+		BoundsExtent.X = Config.cs;
+		bAdjust = true;
 	}
 	if (BoundsExtent.Y > ExtentLimit)
 	{
 		BoundsExtent.Y = ExtentLimit;
-		bClampBounds = true;
+		bClampBounds = bAdjust = true;
+	}
+	else if (BoundsExtent.Y < Config.cs)
+	{
+		// minor adjustment to have at least 1 voxel of size here
+		BoundsExtent.Y = Config.cs;
+		bAdjust = true;
 	}
 
-	if (bClampBounds)
+	if (bAdjust)
 	{
 		const FVector BoundsCenter = NavBounds.GetCenter();
 		NavBounds = FBox(BoundsCenter - BoundsExtent, BoundsCenter + BoundsExtent);
 
-		UE_LOG(LogNavigation, Warning, TEXT("Navigation bounds are too large. Cutting down every dimention down to %.f"), ExtentLimit);
+		if (bClampBounds)
+		{
+			UE_LOG(LogNavigation, Warning, TEXT("Navigation bounds are too large. Cutting down every dimention down to %.f"), ExtentLimit);
+		}
 	}
 
 	UnrealNavBounds = NavBounds;
@@ -4085,7 +4106,7 @@ void FRecastNavMeshGenerator::StartDirtyGenerators()
 		SCOPE_SECONDS_COUNTER(ThisTime);
 		SCOPE_CYCLE_COUNTER(STAT_Navigation_BuildTime)
 
-			const float TileCellSize = (Config.tileSize * Config.cs);
+		const float TileCellSize = (Config.tileSize * Config.cs);
 		const float* BMin = Config.bmin;
 		const float* BMax = Config.bmax;
 
