Header cleanup: Engine module, Public/Net headers.

diff --git a/Engine/Source/Runtime/Engine/Private/LocalPlayer.cpp b/Engine/Source/Runtime/Engine/Private/LocalPlayer.cpp
--- a/Engine/Source/Runtime/Engine/Private/LocalPlayer.cpp
+++ b/Engine/Source/Runtime/Engine/Private/LocalPlayer.cpp
@@ -17,7 +17,7 @@
 
 #include "IHeadMountedDisplay.h"
 #include "SceneViewExtension.h"
-
+#include "DataChannel.h"
 
 DEFINE_LOG_CATEGORY(LogPlayerManagement);
 DEFINE_LOG_CATEGORY_STATIC(LogEngine, Log, All);
diff --git a/Engine/Source/Runtime/Engine/Private/NetConnection.cpp b/Engine/Source/Runtime/Engine/Private/NetConnection.cpp
--- a/Engine/Source/Runtime/Engine/Private/NetConnection.cpp
+++ b/Engine/Source/Runtime/Engine/Private/NetConnection.cpp
@@ -9,6 +9,7 @@
 #include "Net/NetworkProfiler.h"
 #include "Net/DataReplication.h"
 #include "Engine/ActorChannel.h"
+#include "DataChannel.h"
 
 /*-----------------------------------------------------------------------------
 	UNetConnection implementation.
diff --git a/Engine/Source/Runtime/Engine/Private/NetworkDriver.cpp b/Engine/Source/Runtime/Engine/Private/NetworkDriver.cpp
--- a/Engine/Source/Runtime/Engine/Private/NetworkDriver.cpp
+++ b/Engine/Source/Runtime/Engine/Private/NetworkDriver.cpp
@@ -8,11 +8,13 @@
 #include "Net/DataReplication.h"
 #include "Net/UnrealNetwork.h"
 #include "Net/NetworkProfiler.h"
+#include "Net/RepLayout.h"
 #include "Engine/ActorChannel.h"
 #include "Engine/VoiceChannel.h"
 #include "GameFramework/GameNetworkManager.h"
 #include "OnlineSubsystemUtils.h"
 #include "NetworkingDistanceConstants.h"
+#include "DataChannel.h"
 
 // Default net driver stats
 DEFINE_STAT(STAT_Ping);
diff --git a/Engine/Source/Runtime/Engine/Private/PackageMapClient.cpp b/Engine/Source/Runtime/Engine/Private/PackageMapClient.cpp
--- a/Engine/Source/Runtime/Engine/Private/PackageMapClient.cpp
+++ b/Engine/Source/Runtime/Engine/Private/PackageMapClient.cpp
@@ -6,6 +6,7 @@
 #include "Net/DataChannel.h"
 #include "Net/DataReplication.h"
 #include "Engine/ActorChannel.h"
+#include "RepLayout.h"
 
 // ( OutPacketId == GUID_PACKET_NOT_ACKED ) == NAK'd		(this GUID is not acked, and is not pending either, so sort of waiting)
 // ( OutPacketId == GUID_PACKET_ACKED )		== FULLY ACK'd	(this GUID is fully acked, and we no longer need to send full path)
diff --git a/Engine/Source/Runtime/Engine/Private/PendingNetGame.cpp b/Engine/Source/Runtime/Engine/Private/PendingNetGame.cpp
--- a/Engine/Source/Runtime/Engine/Private/PendingNetGame.cpp
+++ b/Engine/Source/Runtime/Engine/Private/PendingNetGame.cpp
@@ -7,6 +7,7 @@
 #include "EnginePrivate.h"
 #include "Net/UnrealNetwork.h"
 #include "Net/NetworkProfiler.h"
+#include "Net/DataChannel.h"
 
 UPendingNetGame::UPendingNetGame( const class FPostConstructInitializeProperties& PCIP, const FURL& InURL )
 	: Super(PCIP)
diff --git a/Engine/Source/Runtime/Engine/Public/Net/DataChannel.h b/Engine/Source/Runtime/Engine/Public/Net/DataChannel.h
--- a/Engine/Source/Runtime/Engine/Public/Net/DataChannel.h
+++ b/Engine/Source/Runtime/Engine/Public/Net/DataChannel.h
@@ -5,9 +5,6 @@
 =============================================================================*/
 
 #pragma once
-#include "Engine/Channel.h"
-#include "Engine/NetConnection.h"
-#include "Engine/ChildConnection.h"
 
 /*-----------------------------------------------------------------------------
 	UControlChannel base class.
diff --git a/Engine/Source/Runtime/Engine/Public/Net/DataReplication.h b/Engine/Source/Runtime/Engine/Public/Net/DataReplication.h
--- a/Engine/Source/Runtime/Engine/Public/Net/DataReplication.h
+++ b/Engine/Source/Runtime/Engine/Public/Net/DataReplication.h
@@ -6,7 +6,7 @@
 =============================================================================*/
 #pragma once
 
-#include "RepLayout.h"
+class FRepState;
 
 /** struct containing property and offset for replicated actor properties */
 struct FReplicatedActorProperty
diff --git a/Engine/Source/Runtime/Engine/Public/Net/UnrealNetwork.h b/Engine/Source/Runtime/Engine/Public/Net/UnrealNetwork.h
--- a/Engine/Source/Runtime/Engine/Public/Net/UnrealNetwork.h
+++ b/Engine/Source/Runtime/Engine/Public/Net/UnrealNetwork.h
@@ -53,13 +53,6 @@ inline bool IsNetworkCompatible(bool bRequireEngineVersionMatch, int32 RemoteVer
 	return bEngineVerMatch && bMinEngineVerMatch;
 }
 
-/*-----------------------------------------------------------------------------
-	Includes.
------------------------------------------------------------------------------*/
-
-#include "DataBunch.h"		// Bunch class.
-#include "DataChannel.h"	// Channel class.
-
 /*-----------------------------------------------------------------------------
 	Replication.
 -----------------------------------------------------------------------------*/
diff --git a/Engine/Source/Runtime/Online/OnlineSubsystemUtils/Private/OnlineBeaconHost.cpp b/Engine/Source/Runtime/Online/OnlineSubsystemUtils/Private/OnlineBeaconHost.cpp
--- a/Engine/Source/Runtime/Online/OnlineSubsystemUtils/Private/OnlineBeaconHost.cpp
+++ b/Engine/Source/Runtime/Online/OnlineSubsystemUtils/Private/OnlineBeaconHost.cpp
@@ -2,6 +2,7 @@
 
 #include "OnlineSubsystemUtilsPrivatePCH.h"
 #include "Net/UnrealNetwork.h"
+#include "Net/DataChannel.h"
 
 AOnlineBeaconHost::AOnlineBeaconHost(const class FPostConstructInitializeProperties& PCIP) :
 	Super(PCIP)
