Header cleanup: Engine module, public headers.

diff --git a/Engine/Source/Runtime/Engine/Private/AI/Navigation/NavigationModifier.cpp b/Engine/Source/Runtime/Engine/Private/AI/Navigation/NavigationModifier.cpp
--- a/Engine/Source/Runtime/Engine/Private/AI/Navigation/NavigationModifier.cpp
+++ b/Engine/Source/Runtime/Engine/Private/AI/Navigation/NavigationModifier.cpp
@@ -11,6 +11,7 @@
 #include "AI/Navigation/NavLinkDefinition.h"
 #include "AI/Navigation/NavLinkTrivial.h"
 #include "AI/Navigation/NavAreas/NavAreaMeta.h"
+#include "AI/Navigation/NavigationProxy.h"
 
 // if square distance between two points is less than this the those points
 // will be considered identical when calculating convex hull
diff --git a/Engine/Source/Runtime/Engine/Private/AI/Navigation/NavigationOctree.cpp b/Engine/Source/Runtime/Engine/Private/AI/Navigation/NavigationOctree.cpp
--- a/Engine/Source/Runtime/Engine/Private/AI/Navigation/NavigationOctree.cpp
+++ b/Engine/Source/Runtime/Engine/Private/AI/Navigation/NavigationOctree.cpp
@@ -4,6 +4,7 @@
 #include "NavigationOctree.h"
 #include "RecastHelpers.h"
 #include "AI/Navigation/NavRelevantComponent.h"
+#include "AI/Navigation/NavigationProxy.h"
 
 #if NAVOCTREE_CONTAINS_COLLISION_DATA
 #include "RecastNavMeshGenerator.h"
@@ -127,3 +128,35 @@ bool FNavigationRelevantData::IsMatchingFilter(const FNavigationOctreeFilter& Fi
 		(Filter.bIncludeAreas && Modifiers.HasAreas()) ||
 		(Filter.bIncludeMetaAreas && Modifiers.HasMetaAreas());
 }
+
+#if NAVSYS_DEBUG
+FORCENOINLINE
+#endif // NAVSYS_DEBUG
+void FNavigationOctreeSemantics::SetElementId(const FNavigationOctreeElement& Element, FOctreeElementId Id)
+{
+	UWorld* World = NULL;
+	UObject* ElementOwner = Element.Owner.Get();
+
+	if (AActor* Actor = Cast<AActor>(ElementOwner))
+	{
+		World = Actor->GetWorld();
+	}
+	else if (UNavigationProxy* Proxy = Cast<UNavigationProxy>(ElementOwner))
+	{
+		World = Proxy->MyOwner != NULL ? Proxy->MyOwner->GetWorld() : NULL;
+	}
+	else if (UActorComponent* AC = Cast<UActorComponent>(ElementOwner))
+	{
+		World = AC->GetWorld();
+	}
+	else if (ULevel* Level = Cast<ULevel>(ElementOwner))
+	{
+		World = Level->OwningWorld;
+	}
+
+	UNavigationSystem* NavSys = UNavigationSystem::GetCurrent(World);
+	if (NavSys)
+	{
+		NavSys->SetObjectsNavOctreeId(ElementOwner, Id);
+	}
+}
\ No newline at end of file
diff --git a/Engine/Source/Runtime/Engine/Private/AI/Navigation/NavigationSystem.cpp b/Engine/Source/Runtime/Engine/Private/AI/Navigation/NavigationSystem.cpp
--- a/Engine/Source/Runtime/Engine/Private/AI/Navigation/NavigationSystem.cpp
+++ b/Engine/Source/Runtime/Engine/Private/AI/Navigation/NavigationSystem.cpp
@@ -25,6 +25,7 @@
 #include "AI/Navigation/NavigationSystem.h"
 #include "AI/Navigation/NavRelevantComponent.h"
 #include "AI/Navigation/NavigationPath.h"
+#include "AI/Navigation/NavigationProxy.h"
 
 static const uint32 INITIAL_ASYNC_QUERIES_SIZE = 32;
 static const uint32 REGISTRATION_QUEUE_SIZE = 16;	// and we'll not reallocate
diff --git a/Engine/Source/Runtime/Engine/Public/AI/NavigationOctree.h b/Engine/Source/Runtime/Engine/Public/AI/NavigationOctree.h
--- a/Engine/Source/Runtime/Engine/Public/AI/NavigationOctree.h
+++ b/Engine/Source/Runtime/Engine/Public/AI/NavigationOctree.h
@@ -3,7 +3,6 @@
 #pragma once
 #include "AI/NavigationModifier.h"
 #include "GenericOctree.h"
-#include "AI/Navigation/NavigationProxy.h"
 
 class AActor;
 class UActorComponent;
@@ -131,34 +130,7 @@ struct FNavigationOctreeSemantics
 #if NAVSYS_DEBUG
 	FORCENOINLINE 
 #endif // NAVSYS_DEBUG
-	static void SetElementId(const FNavigationOctreeElement& Element, FOctreeElementId Id)
-	{
-		UWorld* World = NULL;
-		UObject* ElementOwner = Element.Owner.Get();
-
-		if (AActor* Actor = Cast<AActor>(ElementOwner))
-		{
-			World = Actor->GetWorld();
-		}
-		else if (UNavigationProxy* Proxy = Cast<UNavigationProxy>(ElementOwner))
-		{
-			World = Proxy->MyOwner != NULL ? Proxy->MyOwner->GetWorld() : NULL;
-		}
-		else if (UActorComponent* AC = Cast<UActorComponent>(ElementOwner))
-		{
-			World = AC->GetWorld();
-		}
-		else if (ULevel* Level = Cast<ULevel>(ElementOwner))
-		{
-			World = Level->OwningWorld;
-		}
-
-		UNavigationSystem* NavSys = UNavigationSystem::GetCurrent(World);
-		if (NavSys)
-		{
-			NavSys->SetObjectsNavOctreeId(ElementOwner, Id);
-		}
-	}
+	static void SetElementId(const FNavigationOctreeElement& Element, FOctreeElementId Id);
 };
 
 class FNavigationOctree : public TOctree<FNavigationOctreeElement, FNavigationOctreeSemantics>
