Header cleanup: Engine module, public headers.

diff --git a/Engine/Source/Runtime/Engine/Public/Landscape/LandscapeRender.h b/Engine/Source/Runtime/Engine/Public/Landscape/LandscapeRender.h
--- a/Engine/Source/Runtime/Engine/Public/Landscape/LandscapeRender.h
+++ b/Engine/Source/Runtime/Engine/Public/Landscape/LandscapeRender.h
@@ -11,6 +11,7 @@ LandscapeRender.h: New terrain rendering
 #include "Landscape/LandscapeComponent.h"
 #include "Landscape/LandscapeProxy.h"
 #include "LightMap.h"
+#include "MeshBatch.h"
 #include "ShadowMap.h"
 
 // This defines the number of border blocks to surround terrain by when generating lightmaps
diff --git a/Engine/Source/Runtime/Engine/Public/PrimitiveSceneProxy.h b/Engine/Source/Runtime/Engine/Public/PrimitiveSceneProxy.h
--- a/Engine/Source/Runtime/Engine/Public/PrimitiveSceneProxy.h
+++ b/Engine/Source/Runtime/Engine/Public/PrimitiveSceneProxy.h
@@ -5,18 +5,18 @@
 =============================================================================*/
 
 #pragma once
+#include "SceneView.h"
+#include "PrimitiveUniformShaderParameters.h"
+#include "PrimitiveViewRelevance.h"
 
+// Forward declarations.
 class FSimpleLightEntry;
 class HHitProxy;
 class FStaticPrimitiveDrawInterface;
 class FPrimitiveSceneInfo;
 class FLightSceneProxy;
 class FLightSceneInfo;
-
-
-#include "SceneView.h"
-#include "PrimitiveUniformShaderParameters.h"
-#include "PrimitiveViewRelevance.h"
+class FPrimitiveDrawInterface;
 
 /** Data for a simple dynamic light. */
 class FSimpleLightEntry
diff --git a/Engine/Source/Runtime/Engine/Public/SceneView.h b/Engine/Source/Runtime/Engine/Public/SceneView.h
--- a/Engine/Source/Runtime/Engine/Public/SceneView.h
+++ b/Engine/Source/Runtime/Engine/Public/SceneView.h
@@ -8,6 +8,7 @@
 #include "ConvexVolume.h"
 #include "FinalPostProcessSettings.h"
 #include "SceneInterface.h"
+#include "SceneTypes.h"
 
 class FSceneViewStateInterface;
 class FViewUniformShaderParameters;
diff --git a/Engine/Source/Runtime/Engine/Public/SkeletalMeshTypes.h b/Engine/Source/Runtime/Engine/Public/SkeletalMeshTypes.h
--- a/Engine/Source/Runtime/Engine/Public/SkeletalMeshTypes.h
+++ b/Engine/Source/Runtime/Engine/Public/SkeletalMeshTypes.h
@@ -12,7 +12,6 @@
 
 #include "Components.h"
 #include "GPUSkinPublicDefs.h"
-#include "Components/PrimitiveComponent.h"
 #include "PrimitiveSceneProxy.h"
 #include "Materials/MaterialInterface.h"
 
@@ -25,6 +24,7 @@ typedef uint16 FBoneIndexType;
 
 
 class FRawStaticIndexBuffer16or32Interface;
+class FPrimitiveDrawInterface;
 
 /** 
 * A pair of bone indices
