ttp#343622 BP:Support Class type variables for ExposeOnSpawn

diff --git a/Engine/Source/Editor/BlueprintGraph/Private/EdGraphSchema_K2.cpp b/Engine/Source/Editor/BlueprintGraph/Private/EdGraphSchema_K2.cpp
--- a/Engine/Source/Editor/BlueprintGraph/Private/EdGraphSchema_K2.cpp
+++ b/Engine/Source/Editor/BlueprintGraph/Private/EdGraphSchema_K2.cpp
@@ -3281,6 +3281,11 @@ UFunction* UEdGraphSchema_K2::FindSetVariableByNameFunction(const FEdGraphPinTyp
 		static FName SetObjectName(GET_FUNCTION_NAME_CHECKED(UKismetSystemLibrary, SetObjectPropertyByName));
 		SetFunctionName = SetObjectName;
 	}
+	else if (PinType.PinCategory == K2Schema->PC_Class)
+	{
+		static FName SetObjectName(GET_FUNCTION_NAME_CHECKED(UKismetSystemLibrary, SetClassPropertyByName));
+		SetFunctionName = SetObjectName;
+	}
 	else if(PinType.PinCategory == K2Schema->PC_Interface)
 	{
 		static FName SetObjectName(GET_FUNCTION_NAME_CHECKED(UKismetSystemLibrary, SetObjectPropertyByName));
diff --git a/Engine/Source/Runtime/Engine/Classes/Kismet/KismetSystemLibrary.h b/Engine/Source/Runtime/Engine/Classes/Kismet/KismetSystemLibrary.h
--- a/Engine/Source/Runtime/Engine/Classes/Kismet/KismetSystemLibrary.h
+++ b/Engine/Source/Runtime/Engine/Classes/Kismet/KismetSystemLibrary.h
@@ -346,6 +346,10 @@ class UKismetSystemLibrary : public UBlueprintFunctionLibrary
 	UFUNCTION(BlueprintCallable, meta=(BlueprintInternalUseOnly = "true"))
 	static void SetObjectPropertyByName(UObject* Object, FName PropertyName, UObject* Value);
 
+	/** Set an OBJECT property by name */
+	UFUNCTION(BlueprintCallable, meta = (BlueprintInternalUseOnly = "true"))
+	static void SetClassPropertyByName(UObject* Object, FName PropertyName, TSubclassOf<UObject> Value);
+
 	/** Set a NAME property by name */
 	UFUNCTION(BlueprintCallable, meta=(BlueprintInternalUseOnly = "true", AutoCreateRefTerm = "Value" ))
 	static void SetNamePropertyByName(UObject* Object, FName PropertyName, const FName& Value);
diff --git a/Engine/Source/Runtime/Engine/Private/KismetSystemLibrary.cpp b/Engine/Source/Runtime/Engine/Private/KismetSystemLibrary.cpp
--- a/Engine/Source/Runtime/Engine/Private/KismetSystemLibrary.cpp
+++ b/Engine/Source/Runtime/Engine/Private/KismetSystemLibrary.cpp
@@ -481,6 +481,18 @@ void UKismetSystemLibrary::SetObjectPropertyByName(UObject* Object, FName Proper
 	}
 }
 
+void UKismetSystemLibrary::SetClassPropertyByName(UObject* Object, FName PropertyName, TSubclassOf<UObject> Value)
+{
+	if (Object && *Value)
+	{
+		auto ClassProp = FindField<UClassProperty>(Object->GetClass(), PropertyName);
+		if (ClassProp != NULL && Value->IsChildOf(ClassProp->MetaClass)) // check it's the right type
+		{
+			ClassProp->SetObjectPropertyValue_InContainer(Object, *Value);
+		}
+	}
+}
+
 void UKismetSystemLibrary::SetStringPropertyByName(UObject* Object, FName PropertyName, const FString& Value)
 {
 	if(Object != NULL)
