Fixed NIL matrix issue when invert

diff --git a/Engine/Source/Runtime/Engine/Public/SceneView.h b/Engine/Source/Runtime/Engine/Public/SceneView.h
--- a/Engine/Source/Runtime/Engine/Public/SceneView.h
+++ b/Engine/Source/Runtime/Engine/Public/SceneView.h
@@ -177,13 +177,13 @@ struct FViewMatrices
 
 	FMatrix GetInvProjMatrix() const
 	{
-		return ProjMatrix.Inverse();
+		return ProjMatrix.InverseSafe();
 	}
 
 	FMatrix GetInvViewMatrix() const
 	{
 		// can be optimized: it's not a perspective matrix so transpose would be enough
-		return ViewMatrix.Inverse();
+		return ViewMatrix.InverseSafe();
 	}
 
 	FMatrix GetInvViewProjMatrix() const
