#UE4

- Adding in templated getters for GameState and Player Controller

diff --git a/Engine/Source/Runtime/Engine/Classes/Engine/LocalPlayer.h b/Engine/Source/Runtime/Engine/Classes/Engine/LocalPlayer.h
--- a/Engine/Source/Runtime/Engine/Classes/Engine/LocalPlayer.h
+++ b/Engine/Source/Runtime/Engine/Classes/Engine/LocalPlayer.h
@@ -11,9 +11,9 @@
 struct ENGINE_API FLocalPlayerContext
 {
 	FLocalPlayerContext();
-	FLocalPlayerContext( const class ULocalPlayer* InLocalPlayer );
-	FLocalPlayerContext( const class APlayerController* InPlayerController );
-	FLocalPlayerContext( const FLocalPlayerContext& InPlayerContext );
+	FLocalPlayerContext(const class ULocalPlayer* InLocalPlayer);
+	FLocalPlayerContext(const class APlayerController* InPlayerController);
+	FLocalPlayerContext(const FLocalPlayerContext& InPlayerContext);
 
 	/* Returns the world context. */
 	UWorld* GetWorld() const;
@@ -24,12 +24,26 @@ struct ENGINE_API FLocalPlayerContext
 	/* Returns the player controller. */
 	class APlayerController* GetPlayerController() const;
 
+	/** Templated version of GetPlayerController() */
+	template<class T>
+	FORCEINLINE T* GetPlayerController() const
+	{
+		return CastChecked<T>(GetPlayerController(), ECastCheckedType::NullAllowed);
+	}
+
 	/** Is this context initialized and still valid? */
 	bool IsValid() const;
 
 	/** Is this context initialized */
 	bool IsInitialized() const;
 
+	/** Templated Getter for the Game State */
+	template<class T>
+	FORCEINLINE T* GetGameState() const
+	{
+		return CastChecked<T>(GetWorld()->GameState, ECastCheckedType::NullAllowed);
+	}
+
 private:	
 
 	/* Set the local player. */
