Fixed destructible crash when streaming in/out

diff --git a/Engine/Source/Runtime/Engine/Private/SkinnedMeshComponent.cpp b/Engine/Source/Runtime/Engine/Private/SkinnedMeshComponent.cpp
--- a/Engine/Source/Runtime/Engine/Private/SkinnedMeshComponent.cpp
+++ b/Engine/Source/Runtime/Engine/Private/SkinnedMeshComponent.cpp
@@ -1552,7 +1552,7 @@ void USkinnedMeshComponent::UnHideBone( int32 BoneIndex )
 
 bool USkinnedMeshComponent::IsBoneHidden( int32 BoneIndex ) const
 {
-	if (ShouldUpdateBoneVisibility())
+	if (ShouldUpdateBoneVisibility() && BoneIndex < BoneVisibilityStates.Num())
 	{
 		if ( BoneIndex != INDEX_NONE )
 		{
