#ue4 - Pawn::AddController*Input() methods don't need to pass along the value to the controller if it's zero.

diff --git a/Engine/Source/Runtime/Engine/Private/Pawn.cpp b/Engine/Source/Runtime/Engine/Private/Pawn.cpp
--- a/Engine/Source/Runtime/Engine/Private/Pawn.cpp
+++ b/Engine/Source/Runtime/Engine/Private/Pawn.cpp
@@ -624,7 +624,7 @@ FVector APawn::Internal_ConsumeMovementInputVector()
 
 void APawn::AddControllerPitchInput(float Val)
 {
-	if (Controller && Controller->IsLocalPlayerController())
+	if (Val != 0.f && Controller && Controller->IsLocalPlayerController())
 	{
 		APlayerController* const PC = CastChecked<APlayerController>(Controller);
 		PC->AddPitchInput(Val);
@@ -633,7 +633,7 @@ void APawn::AddControllerPitchInput(float Val)
 
 void APawn::AddControllerYawInput(float Val)
 {
-	if (Controller && Controller->IsLocalPlayerController())
+	if (Val != 0.f && Controller && Controller->IsLocalPlayerController())
 	{
 		APlayerController* const PC = CastChecked<APlayerController>(Controller);
 		PC->AddYawInput(Val);
@@ -642,7 +642,7 @@ void APawn::AddControllerYawInput(float Val)
 
 void APawn::AddControllerRollInput(float Val)
 {
-	if (Controller && Controller->IsLocalPlayerController())
+	if (Val != 0.f && Controller && Controller->IsLocalPlayerController())
 	{
 		APlayerController* const PC = CastChecked<APlayerController>(Controller);
 		PC->AddRollInput(Val);
