#UE4 - added a null check against uniqueid in Login debug output

diff --git a/Engine/Source/Runtime/Engine/Private/LevelActor.cpp b/Engine/Source/Runtime/Engine/Private/LevelActor.cpp
--- a/Engine/Source/Runtime/Engine/Private/LevelActor.cpp
+++ b/Engine/Source/Runtime/Engine/Private/LevelActor.cpp
@@ -602,7 +602,7 @@ APlayerController* UWorld::SpawnPlayActor(UPlayer* NewPlayer, ENetRole RemoteRol
 		return NULL;
 	}
 
-	UE_LOG(LogSpawn, Log, TEXT("%s got player %s [%s]"), *NewPlayerController->GetName(), *NewPlayer->GetName(), *UniqueId->ToString());
+	UE_LOG(LogSpawn, Log, TEXT("%s got player %s [%s]"), *NewPlayerController->GetName(), *NewPlayer->GetName(), UniqueId.IsValid() ? *UniqueId->ToString() : TEXT("Invalid"));
 
 	// Possess the newly-spawned player.
 	NewPlayerController->NetPlayerIndex = InNetPlayerIndex;
