Move swapping of GPU particle textures to PreInitViews. This breaks a dependency that is preventing our refactor of DrawDynamicElements.

diff --git a/Engine/Source/Runtime/Engine/Private/Particles/FXSystem.cpp b/Engine/Source/Runtime/Engine/Private/Particles/FXSystem.cpp
--- a/Engine/Source/Runtime/Engine/Private/Particles/FXSystem.cpp
+++ b/Engine/Source/Runtime/Engine/Private/Particles/FXSystem.cpp
@@ -310,7 +310,7 @@ void FFXSystem::PreInitViews()
 {
 	if (RHISupportsGPUParticles(FeatureLevel))
 	{
-		ResetSortedGPUParticles();
+		AdvanceGPUParticleFrame();
 	}
 }
 
diff --git a/Engine/Source/Runtime/Engine/Private/Particles/FXSystemPrivate.h b/Engine/Source/Runtime/Engine/Private/Particles/FXSystemPrivate.h
--- a/Engine/Source/Runtime/Engine/Private/Particles/FXSystemPrivate.h
+++ b/Engine/Source/Runtime/Engine/Private/Particles/FXSystemPrivate.h
@@ -132,10 +132,9 @@ private:
 	void ReleaseGPUResources();
 
 	/**
-	 * Clears the current set of sorted GPU particles. Sorted results are no
-	 * longer valid after this call.
+	 * Prepares GPU particles for simulation and rendering in the next frame.
 	 */
-	void ResetSortedGPUParticles();
+	void AdvanceGPUParticleFrame();
 
 	/**
 	 * Sorts all GPU particles that have called AddSortedGPUSimulation since the
diff --git a/Engine/Source/Runtime/Engine/Private/Particles/ParticleGpuSimulation.cpp b/Engine/Source/Runtime/Engine/Private/Particles/ParticleGpuSimulation.cpp
--- a/Engine/Source/Runtime/Engine/Private/Particles/ParticleGpuSimulation.cpp
+++ b/Engine/Source/Runtime/Engine/Private/Particles/ParticleGpuSimulation.cpp
@@ -3951,8 +3951,13 @@ int32 FFXSystem::AddSortedGPUSimulation(FParticleSimulationGPU* Simulation, cons
 	return BufferOffset;
 }
 
-void FFXSystem::ResetSortedGPUParticles()
+void FFXSystem::AdvanceGPUParticleFrame()
 {
+	// We double buffer, so swap the current and previous textures.
+	ParticleSimulationResources->FrameIndex ^= 0x1;
+
+	// Reset the list of sorted simulations. As PreRenderView is called on GPU simulations we'll
+	// allocate space for them in the sorted index buffer.
 	ParticleSimulationResources->SimulationsToSort.Reset();
 	ParticleSimulationResources->SortedParticleCount = 0;
 }
@@ -4020,12 +4025,6 @@ void FFXSystem::SimulateGPUParticles(
 
 	FMemMark Mark(FMemStack::Get());
 
-	if (Phase == EParticleSimulatePhase::First)
-	{
-		// Advance to the next frame.
-		ParticleSimulationResources->FrameIndex ^= 0x1;
-	}
-
 	// Grab resources.
 	FParticleStateTextures& CurrentStateTextures = ParticleSimulationResources->GetCurrentStateTextures();
 	FParticleStateTextures& PrevStateTextures = ParticleSimulationResources->GetPreviousStateTextures();
