Remove dead code: PositionZWTexture

diff --git a/Engine/Source/Runtime/Engine/Private/Particles/ParticleGpuSimulation.cpp b/Engine/Source/Runtime/Engine/Private/Particles/ParticleGpuSimulation.cpp
--- a/Engine/Source/Runtime/Engine/Private/Particles/ParticleGpuSimulation.cpp
+++ b/Engine/Source/Runtime/Engine/Private/Particles/ParticleGpuSimulation.cpp
@@ -161,8 +161,6 @@ public:
 	/** Contains the positions of all simulating particles. */
 	FTexture2DRHIRef PositionTextureTargetRHI;
 	FTexture2DRHIRef PositionTextureRHI;
-	FTexture2DRHIRef PositionZWTextureTargetRHI;
-	FTexture2DRHIRef PositionZWTextureRHI;
 	/** Contains the velocity of all simulating particles. */
 	FTexture2DRHIRef VelocityTextureTargetRHI;
 	FTexture2DRHIRef VelocityTextureRHI;
@@ -223,8 +221,6 @@ public:
 		// Release textures.
 		PositionTextureTargetRHI.SafeRelease();
 		PositionTextureRHI.SafeRelease();
-		PositionZWTextureTargetRHI.SafeRelease();
-		PositionZWTextureRHI.SafeRelease();
 		VelocityTextureTargetRHI.SafeRelease();
 		VelocityTextureRHI.SafeRelease();
 	}
@@ -472,8 +468,6 @@ public:
 		ParticleIndicesOffset.Bind(ParameterMap, TEXT("ParticleIndicesOffset"));
 		PositionTexture.Bind(ParameterMap, TEXT("PositionTexture"));
 		PositionTextureSampler.Bind(ParameterMap, TEXT("PositionTextureSampler"));
-		PositionZWTexture.Bind(ParameterMap, TEXT("PositionZWTexture"));
-		PositionZWTextureSampler.Bind(ParameterMap, TEXT("PositionZWTextureSampler"));
 		VelocityTexture.Bind(ParameterMap, TEXT("VelocityTexture"));
 		VelocityTextureSampler.Bind(ParameterMap, TEXT("VelocityTextureSampler"));
 		AttributesTexture.Bind(ParameterMap, TEXT("AttributesTexture"));
@@ -488,8 +482,6 @@ public:
 		Ar << ParticleIndicesOffset;
 		Ar << PositionTexture;
 		Ar << PositionTextureSampler;
-		Ar << PositionZWTexture;
-		Ar << PositionZWTextureSampler;
 		Ar << VelocityTexture;
 		Ar << VelocityTextureSampler;
 		Ar << AttributesTexture;
@@ -511,8 +503,6 @@ private:
 	/** Texture containing positions for all particles. */
 	FShaderResourceParameter PositionTexture;
 	FShaderResourceParameter PositionTextureSampler;
-	FShaderResourceParameter PositionZWTexture;
-	FShaderResourceParameter PositionZWTextureSampler;
 	/** Texture containing velocities for all particles. */
 	FShaderResourceParameter VelocityTexture;
 	FShaderResourceParameter VelocityTextureSampler;
@@ -578,7 +568,6 @@ public:
 	uint32 ParticleIndicesOffset;
 	/** Texture containing positions for all particles. */
 	FTexture2DRHIParamRef PositionTextureRHI;
-	FTexture2DRHIParamRef PositionZWTextureRHI;
 	/** Texture containing velocities for all particles. */
 	FTexture2DRHIParamRef VelocityTextureRHI;
 	/** Texture containint attributes for all particles. */
@@ -659,7 +648,6 @@ void FGPUSpriteVertexFactoryShaderParameters::SetMesh(FRHICommandList& RHICmdLis
 	}
 	SetShaderValue(RHICmdList, VertexShader, ParticleIndicesOffset, GPUVF->ParticleIndicesOffset);
 	SetTextureParameter(RHICmdList, VertexShader, PositionTexture, PositionTextureSampler, SamplerStatePoint, GPUVF->PositionTextureRHI );
-	SetTextureParameter(RHICmdList, VertexShader, PositionZWTexture, PositionZWTextureSampler, SamplerStatePoint, GPUVF->PositionZWTextureRHI );
 	SetTextureParameter(RHICmdList, VertexShader, VelocityTexture, VelocityTextureSampler, SamplerStatePoint, GPUVF->VelocityTextureRHI );
 	SetTextureParameter(RHICmdList, VertexShader, AttributesTexture, AttributesTextureSampler, SamplerStatePoint, GPUVF->AttributesTextureRHI );
 	SetTextureParameter(RHICmdList, VertexShader, CurveTexture, CurveTextureSampler, SamplerStateLinear, GParticleCurveTexture.GetCurveTexture() );
@@ -876,8 +864,6 @@ public:
 	{
 		PositionTexture.Bind(Initializer.ParameterMap, TEXT("PositionTexture"));
 		PositionTextureSampler.Bind(Initializer.ParameterMap, TEXT("PositionTextureSampler"));
-		PositionZWTexture.Bind(Initializer.ParameterMap, TEXT("PositionZWTexture"));
-		PositionZWTextureSampler.Bind(Initializer.ParameterMap, TEXT("PositionZWTextureSampler"));
 		VelocityTexture.Bind(Initializer.ParameterMap, TEXT("VelocityTexture"));
 		VelocityTextureSampler.Bind(Initializer.ParameterMap, TEXT("VelocityTextureSampler"));
 		AttributesTexture.Bind(Initializer.ParameterMap, TEXT("AttributesTexture"));
@@ -902,8 +888,6 @@ public:
 		bool bShaderHasOutdatedParameters = FGlobalShader::Serialize( Ar );
 		Ar << PositionTexture;
 		Ar << PositionTextureSampler;
-		Ar << PositionZWTexture;
-		Ar << PositionZWTextureSampler;
 		Ar << VelocityTexture;
 		Ar << VelocityTextureSampler;
 		Ar << AttributesTexture;
@@ -940,7 +924,6 @@ public:
 		FSamplerStateRHIParamRef SamplerStatePoint = TStaticSamplerState<SF_Point>::GetRHI();
 		FSamplerStateRHIParamRef SamplerStateLinear = TStaticSamplerState<SF_Bilinear,AM_Clamp,AM_Clamp,AM_Clamp>::GetRHI();
 		SetTextureParameter(RHICmdList, PixelShaderRHI, PositionTexture, PositionTextureSampler, SamplerStatePoint, TextureResources.PositionTextureRHI);
-		SetTextureParameter(RHICmdList, PixelShaderRHI, PositionZWTexture, PositionZWTextureSampler, SamplerStatePoint, TextureResources.PositionZWTextureRHI);
 		SetTextureParameter(RHICmdList, PixelShaderRHI, VelocityTexture, VelocityTextureSampler, SamplerStatePoint, TextureResources.VelocityTextureRHI);
 		SetTextureParameter(RHICmdList, PixelShaderRHI, AttributesTexture, AttributesTextureSampler, SamplerStatePoint, InAttributesTexture.TextureRHI);
 		SetTextureParameter(RHICmdList, PixelShaderRHI, CurveTexture, CurveTextureSampler, SamplerStateLinear, GParticleCurveTexture.GetCurveTexture());
@@ -1023,8 +1006,6 @@ private:
 	/** The position texture parameter. */
 	FShaderResourceParameter PositionTexture;
 	FShaderResourceParameter PositionTextureSampler;
-	FShaderResourceParameter PositionZWTexture;
-	FShaderResourceParameter PositionZWTextureSampler;
 	/** The velocity texture parameter. */
 	FShaderResourceParameter VelocityTexture;
 	FShaderResourceParameter VelocityTextureSampler;
@@ -1677,8 +1658,6 @@ public:
 		VisualizationMode.Bind( Initializer.ParameterMap, TEXT("VisualizationMode") );
 		PositionTexture.Bind( Initializer.ParameterMap, TEXT("PositionTexture") );
 		PositionTextureSampler.Bind( Initializer.ParameterMap, TEXT("PositionTextureSampler") );
-		PositionZWTexture.Bind( Initializer.ParameterMap, TEXT("PositionZWTexture") );
-		PositionZWTextureSampler.Bind( Initializer.ParameterMap, TEXT("PositionZWTextureSampler") );
 		CurveTexture.Bind( Initializer.ParameterMap, TEXT("CurveTexture") );
 		CurveTextureSampler.Bind( Initializer.ParameterMap, TEXT("CurveTextureSampler") );
 	}
@@ -1690,8 +1669,6 @@ public:
 		Ar << VisualizationMode;
 		Ar << PositionTexture;
 		Ar << PositionTextureSampler;
-		Ar << PositionZWTexture;
-		Ar << PositionZWTextureSampler;
 		Ar << CurveTexture;
 		Ar << CurveTextureSampler;
 		return bShaderHasOutdatedParameters;
@@ -1700,13 +1677,12 @@ public:
 	/**
 	 * Set parameters for this shader.
 	 */
-	void SetParameters(FRHICommandList& RHICmdList, int32 InVisualizationMode, FTexture2DRHIParamRef PositionTextureRHI, FTexture2DRHIParamRef PositionZWTextureRHI, FTexture2DRHIParamRef CurveTextureRHI )
+	void SetParameters(FRHICommandList& RHICmdList, int32 InVisualizationMode, FTexture2DRHIParamRef PositionTextureRHI, FTexture2DRHIParamRef CurveTextureRHI )
 	{
 		FPixelShaderRHIParamRef PixelShader = GetPixelShader();
 		SetShaderValue(RHICmdList, PixelShader, VisualizationMode, InVisualizationMode );
 		FSamplerStateRHIParamRef SamplerStatePoint = TStaticSamplerState<SF_Point>::GetRHI();
 		SetTextureParameter(RHICmdList, PixelShader, PositionTexture, PositionTextureSampler, SamplerStatePoint, PositionTextureRHI );
-		SetTextureParameter(RHICmdList, PixelShader, PositionZWTexture, PositionZWTextureSampler, SamplerStatePoint, PositionZWTextureRHI );
 		SetTextureParameter(RHICmdList, PixelShader, CurveTexture, CurveTextureSampler, SamplerStatePoint, CurveTextureRHI );
 	}
 
@@ -1715,8 +1691,6 @@ private:
 	FShaderParameter VisualizationMode;
 	FShaderResourceParameter PositionTexture;
 	FShaderResourceParameter PositionTextureSampler;
-	FShaderResourceParameter PositionZWTexture;
-	FShaderResourceParameter PositionZWTextureSampler;
 	FShaderResourceParameter CurveTexture;
 	FShaderResourceParameter CurveTextureSampler;
 };
@@ -1803,7 +1777,7 @@ static void VisualizeGPUSimulation(
 	Parameters.ScaleBias.W = 2.0f * DisplayOffsetY / (float)TargetSize.Y - 1.0f;
 	FParticleSimVisualizeBufferRef UniformBuffer = FParticleSimVisualizeBufferRef::CreateUniformBufferImmediate( Parameters, UniformBuffer_SingleDraw );
 	VertexShader->SetParameters(RHICmdList, UniformBuffer);
-	PixelShader->SetParameters(RHICmdList, VisualizationMode, StateTextures.PositionTextureRHI, StateTextures.PositionZWTextureRHI, CurveTextureRHI);
+	PixelShader->SetParameters(RHICmdList, VisualizationMode, StateTextures.PositionTextureRHI, CurveTextureRHI);
 
 	const int32 VertexStride = sizeof(FVector2D);
 	
@@ -1929,8 +1903,6 @@ public:
 		InParticleIndices.Bind( Initializer.ParameterMap, TEXT("InParticleIndices") );
 		PositionTexture.Bind( Initializer.ParameterMap, TEXT("PositionTexture") );
 		PositionTextureSampler.Bind( Initializer.ParameterMap, TEXT("PositionTextureSampler") );
-		PositionZWTexture.Bind( Initializer.ParameterMap, TEXT("PositionZWTexture") );
-		PositionZWTextureSampler.Bind( Initializer.ParameterMap, TEXT("PositionZWTextureSampler") );
 		OutBounds.Bind( Initializer.ParameterMap, TEXT("OutBounds") );
 	}
 
@@ -1941,8 +1913,6 @@ public:
 		Ar << InParticleIndices;
 		Ar << PositionTexture;
 		Ar << PositionTextureSampler;
-		Ar << PositionZWTexture;
-		Ar << PositionZWTextureSampler;
 		Ar << OutBounds;
 		return bShaderHasOutdatedParameters;
 	}
@@ -1966,8 +1936,7 @@ public:
 		FRHICommandList& RHICmdList,
 		FParticleBoundsUniformBufferRef& UniformBuffer,
 		FShaderResourceViewRHIParamRef InIndicesSRV,
-		FTexture2DRHIParamRef PositionTextureRHI,
-		FTexture2DRHIParamRef PositionZWTextureRHI
+		FTexture2DRHIParamRef PositionTextureRHI
 		)
 	{
 		FComputeShaderRHIParamRef ComputeShaderRHI = GetComputeShader();
@@ -1980,10 +1949,6 @@ public:
 		{
 			RHICmdList.SetShaderTexture(ComputeShaderRHI, PositionTexture.GetBaseIndex(), PositionTextureRHI);
 		}
-		if ( PositionZWTexture.IsBound() )
-		{
-			RHICmdList.SetShaderTexture(ComputeShaderRHI, PositionZWTexture.GetBaseIndex(), PositionZWTextureRHI);
-		}
 	}
 
 	/**
@@ -2009,8 +1974,6 @@ private:
 	/** Texture containing particle positions. */
 	FShaderResourceParameter PositionTexture;
 	FShaderResourceParameter PositionTextureSampler;
-	FShaderResourceParameter PositionZWTexture;
-	FShaderResourceParameter PositionZWTextureSampler;
 	/** Output key buffer. */
 	FShaderResourceParameter OutBounds;
 };
@@ -2035,7 +1998,6 @@ static FBox ComputeParticleBounds(
 	FRHICommandListImmediate& RHICmdList,
 	FShaderResourceViewRHIParamRef VertexBufferSRV,
 	FTexture2DRHIParamRef PositionTextureRHI,
-	FTexture2DRHIParamRef PositionZWTextureRHI,
 	int32 ParticleCount )
 {
 	FBox BoundingBox;
@@ -2072,8 +2034,8 @@ static FBox ComputeParticleBounds(
 		RHICmdList.SetComputeShader(ParticleBoundsCS->GetComputeShader());
 
 		// Dispatch shader to compute bounds.
-		ParticleBoundsCS->SetOutput(RHICmdList, BoundsVertexBufferUAV );
-		ParticleBoundsCS->SetParameters(RHICmdList, UniformBuffer, VertexBufferSRV, PositionTextureRHI, PositionZWTextureRHI );
+		ParticleBoundsCS->SetOutput(RHICmdList, BoundsVertexBufferUAV);
+		ParticleBoundsCS->SetParameters(RHICmdList, UniformBuffer, VertexBufferSRV, PositionTextureRHI);
 		DispatchComputeShader(
 			RHICmdList, 
 			*ParticleBoundsCS, 
@@ -2698,7 +2660,6 @@ struct FGPUSpriteDynamicEmitterData : FDynamicEmitterDataBase
 					VertexFactory.EmitterUniformBuffer = Resources->UniformBuffer;
 					VertexFactory.EmitterDynamicUniformBuffer = DynamicUniformBuffer;
 					VertexFactory.PositionTextureRHI = StateTextures.PositionTextureRHI;
-					VertexFactory.PositionZWTextureRHI = StateTextures.PositionZWTextureRHI;
 					VertexFactory.VelocityTextureRHI = StateTextures.VelocityTextureRHI;
 					VertexFactory.AttributesTextureRHI = ParticleSimulationResources->RenderAttributesTexture.TextureRHI;
 
@@ -3308,7 +3269,6 @@ public:
 				RHICmdList,
 				EmitterInstance->Simulation->VertexBuffer.VertexBufferSRV, 
 				EmitterInstance->FXSystem->GetParticleSimulationResources()->GetCurrentStateTextures().PositionTextureRHI,
-				EmitterInstance->FXSystem->GetParticleSimulationResources()->GetCurrentStateTextures().PositionZWTextureRHI,
 				EmitterInstance->Simulation->VertexBuffer.ParticleCount
 				);
 		});
@@ -3970,7 +3930,6 @@ void FFXSystem::SortGPUParticles(FRHICommandListImmediate& RHICmdList)
 			RHICmdList,
 			GParticleSortBuffers,
 			ParticleSimulationResources->GetCurrentStateTextures().PositionTextureRHI,
-			ParticleSimulationResources->GetCurrentStateTextures().PositionZWTextureRHI,
 			ParticleSimulationResources->SimulationsToSort
 			);
 		ParticleSimulationResources->SortedVertexBuffer.VertexBufferRHI =
@@ -4052,11 +4011,10 @@ void FFXSystem::SimulateGPUParticles(
 	}
 
 	// Setup render states.
-	FTextureRHIParamRef RenderTargets[3] =
+	FTextureRHIParamRef RenderTargets[2] =
 	{
 		CurrentStateTextures.PositionTextureTargetRHI,
 		CurrentStateTextures.VelocityTextureTargetRHI,
-		CurrentStateTextures.PositionZWTextureTargetRHI
 	};
 	SetRenderTargets(RHICmdList, 2, RenderTargets, FTextureRHIParamRef(), 0, NULL);
 	RHICmdList.SetViewport(0, 0, 0.0f, GParticleSimulationTextureSizeX, GParticleSimulationTextureSizeY, 1.0f);
@@ -4213,13 +4171,12 @@ void FFXSystem::SimulateGPUParticles(
 	{
 		SCOPED_DRAW_EVENT(ParticleInjection, DEC_PARTICLE);
 		// Set render targets.
-		FTextureRHIParamRef InjectRenderTargets[5] =
+		FTextureRHIParamRef InjectRenderTargets[4] =
 		{
 			CurrentStateTextures.PositionTextureTargetRHI,
 			CurrentStateTextures.VelocityTextureTargetRHI,
 			ParticleSimulationResources->RenderAttributesTexture.TextureTargetRHI,
-			ParticleSimulationResources->SimulationAttributesTexture.TextureTargetRHI,
-			CurrentStateTextures.PositionZWTextureTargetRHI
+			ParticleSimulationResources->SimulationAttributesTexture.TextureTargetRHI
 		};
 		SetRenderTargets(RHICmdList, 4, InjectRenderTargets, FTextureRHIParamRef(), 0, NULL);
 		RHICmdList.SetViewport(0, 0, 0.0f, GParticleSimulationTextureSizeX, GParticleSimulationTextureSizeY, 1.0f);
diff --git a/Engine/Source/Runtime/Engine/Private/Particles/ParticleSortingGPU.cpp b/Engine/Source/Runtime/Engine/Private/Particles/ParticleSortingGPU.cpp
--- a/Engine/Source/Runtime/Engine/Private/Particles/ParticleSortingGPU.cpp
+++ b/Engine/Source/Runtime/Engine/Private/Particles/ParticleSortingGPU.cpp
@@ -71,8 +71,6 @@ public:
 		InParticleIndices.Bind( Initializer.ParameterMap, TEXT("InParticleIndices") );
 		PositionTexture.Bind( Initializer.ParameterMap, TEXT("PositionTexture") );
 		PositionTextureSampler.Bind( Initializer.ParameterMap, TEXT("PositionTextureSampler") );
-		PositionZWTexture.Bind( Initializer.ParameterMap, TEXT("PositionZWTexture") );
-		PositionZWTextureSampler.Bind( Initializer.ParameterMap, TEXT("PositionZWTextureSampler") );
 		OutKeys.Bind( Initializer.ParameterMap, TEXT("OutKeys") );
 		OutParticleIndices.Bind( Initializer.ParameterMap, TEXT("OutParticleIndices") );
 	}
@@ -84,8 +82,6 @@ public:
 		Ar << InParticleIndices;
 		Ar << PositionTexture;
 		Ar << PositionTextureSampler;
-		Ar << PositionZWTexture;
-		Ar << PositionZWTextureSampler;
 		Ar << OutKeys;
 		Ar << OutParticleIndices;
 		return bShaderHasOutdatedParameters;
@@ -127,17 +123,13 @@ public:
 	/**
 	 * Set the texture from which particle positions can be read.
 	 */
-	void SetPositionTextures(FRHICommandList& RHICmdList, FTexture2DRHIParamRef PositionTextureRHI, FTexture2DRHIParamRef PositionZWTextureRHI)
+	void SetPositionTextures(FRHICommandList& RHICmdList, FTexture2DRHIParamRef PositionTextureRHI)
 	{
 		FComputeShaderRHIParamRef ComputeShaderRHI = GetComputeShader();
 		if (PositionTexture.IsBound())
 		{
 			RHICmdList.SetShaderTexture(ComputeShaderRHI, PositionTexture.GetBaseIndex(), PositionTextureRHI);
 		}
-		if (PositionZWTexture.IsBound())
-		{
-			RHICmdList.SetShaderTexture(ComputeShaderRHI, PositionZWTexture.GetBaseIndex(), PositionZWTextureRHI);
-		}
 	}
 
 	/**
@@ -167,8 +159,6 @@ private:
 	/** Texture containing particle positions. */
 	FShaderResourceParameter PositionTexture;
 	FShaderResourceParameter PositionTextureSampler;
-	FShaderResourceParameter PositionZWTexture;
-	FShaderResourceParameter PositionZWTextureSampler;
 	/** Output key buffer. */
 	FShaderResourceParameter OutKeys;
 	/** Output indices buffer. */
@@ -189,7 +179,6 @@ static int32 GenerateParticleSortKeys(
 	FUnorderedAccessViewRHIParamRef KeyBufferUAV,
 	FUnorderedAccessViewRHIParamRef SortedVertexBufferUAV,
 	FTexture2DRHIParamRef PositionTextureRHI,
-	FTexture2DRHIParamRef PositionZWTextureRHI,
 	const TArray<FParticleSimulationSortInfo>& SimulationsToSort
 	)
 {
@@ -206,7 +195,7 @@ static int32 GenerateParticleSortKeys(
 	TShaderMapRef<FParticleSortKeyGenCS> KeyGenCS(GetGlobalShaderMap());
 	RHICmdList.SetComputeShader(KeyGenCS->GetComputeShader());
 	KeyGenCS->SetOutput(RHICmdList, KeyBufferUAV, SortedVertexBufferUAV);
-	KeyGenCS->SetPositionTextures(RHICmdList, PositionTextureRHI, PositionZWTextureRHI);
+	KeyGenCS->SetPositionTextures(RHICmdList, PositionTextureRHI);
 
 	// For each simulation, generate keys and store them in the sorting buffers.
 	const int32 SimulationCount = SimulationsToSort.Num();
@@ -321,7 +310,6 @@ int32 SortParticlesGPU(
 	FRHICommandListImmediate& RHICmdList,
 	FParticleSortBuffers& ParticleSortBuffers,
 	FTexture2DRHIParamRef PositionTextureRHI,
-	FTexture2DRHIParamRef PositionZWTextureRHI,
 	const TArray<FParticleSimulationSortInfo>& SimulationsToSort
 	)
 {
@@ -346,7 +334,6 @@ int32 SortParticlesGPU(
 		ParticleSortBuffers.GetKeyBufferUAV(),
 		ParticleSortBuffers.GetVertexBufferUAV(),
 		PositionTextureRHI,
-		PositionZWTextureRHI,
 		SimulationsToSort
 		);
 
diff --git a/Engine/Source/Runtime/Engine/Private/Particles/ParticleSortingGPU.h b/Engine/Source/Runtime/Engine/Private/Particles/ParticleSortingGPU.h
--- a/Engine/Source/Runtime/Engine/Private/Particles/ParticleSortingGPU.h
+++ b/Engine/Source/Runtime/Engine/Private/Particles/ParticleSortingGPU.h
@@ -125,6 +125,5 @@ int32 SortParticlesGPU(
 	FRHICommandListImmediate& RHICmdList,
 	FParticleSortBuffers& ParticleSortBuffers,
 	FTexture2DRHIParamRef PositionTextureRHI,
-	FTexture2DRHIParamRef PositionZWTextureRHI,
 	const TArray<FParticleSimulationSortInfo>& SimulationsToSort
 	);
