FPrimitiveSceneProxy::PreRenderView / DrawDynamicElements refactored into GetDynamicMeshElements

diff --git a/Engine/Source/Runtime/Engine/Private/InstancedStaticMesh.cpp b/Engine/Source/Runtime/Engine/Private/InstancedStaticMesh.cpp
--- a/Engine/Source/Runtime/Engine/Private/InstancedStaticMesh.cpp
+++ b/Engine/Source/Runtime/Engine/Private/InstancedStaticMesh.cpp
@@ -923,20 +923,21 @@ public:
 		return Result;
 	}
 
-	/** Common path for the Get*MeshElement functions */
-	inline void SetupInstancedMeshBatch(int32 LODIndex, TArray<FMeshBatch, TInlineAllocator<1>>& OutMeshBatches) const;
+	virtual void GetDynamicMeshElements(const TArray<const FSceneView*>& Views, const FSceneViewFamily& ViewFamily, uint32 VisibilityMap, FMeshElementCollector& Collector) const override;
 
 	/** Draw the scene proxy as a dynamic element */
 	virtual void DrawDynamicElements(FPrimitiveDrawInterface* PDI,const FSceneView* View) override;
-		
+
+	virtual int32 GetNumMeshBatches() const override;
+
 	/** Sets up a shadow FMeshBatch for a specific LOD. */
-	virtual bool GetShadowMeshElements(int32 LODIndex, uint8 InDepthPriorityGroup, TArray<FMeshBatch, TInlineAllocator<1>>& OutMeshBatches) const override;
+	virtual bool GetShadowMeshElement(int32 LODIndex, int32 BatchIndex, uint8 InDepthPriorityGroup, FMeshBatch& OutMeshBatch) const override;
 
 	/** Sets up a FMeshBatch for a specific LOD and element. */
-	virtual bool GetMeshElements(int32 LODIndex,int32 ElementIndex,uint8 InDepthPriorityGroup, TArray<FMeshBatch, TInlineAllocator<1>>& OutMeshBatches, const bool bUseSelectedMaterial, const bool bUseHoveredMaterial) const override;
+	virtual bool GetMeshElement(int32 LODIndex, int32 BatchIndex, int32 ElementIndex, uint8 InDepthPriorityGroup, const bool bUseSelectedMaterial, const bool bUseHoveredMaterial, FMeshBatch& OutMeshBatch) const override;
 
 	/** Sets up a wireframe FMeshBatch for a specific LOD. */
-	virtual bool GetWireframeMeshElements(int32 LODIndex, const FMaterialRenderProxy* WireframeRenderProxy, uint8 InDepthPriorityGroup, TArray<FMeshBatch, TInlineAllocator<1>>& OutMeshBatches) const override;
+	virtual bool GetWireframeMeshElement(int32 LODIndex, int32 BatchIndex, const FMaterialRenderProxy* WireframeRenderProxy, uint8 InDepthPriorityGroup, FMeshBatch& OutMeshBatch) const override;
 
 	/**
 	 * Creates the hit proxies are used when DrawDynamicElements is called.
@@ -980,8 +981,76 @@ private:
 	FInstancingUserData UserData_AllInstances;
 	FInstancingUserData UserData_SelectedInstances;
 	FInstancingUserData UserData_DeselectedInstances;
+
+	/** Common path for the Get*MeshElement functions */
+	void SetupInstancedMeshBatch2(int32 LODIndex, TArray<FMeshBatch, TInlineAllocator<1>>& OutMeshBatches) const;
+
+	/** Common path for the Get*MeshElement functions */
+	void SetupInstancedMeshBatch(int32 LODIndex, int32 BatchIndex, FMeshBatch& OutMeshBatch) const;
 };
 
+void FInstancedStaticMeshSceneProxy::GetDynamicMeshElements(const TArray<const FSceneView*>& Views, const FSceneViewFamily& ViewFamily, uint32 VisibilityMap, FMeshElementCollector& Collector) const
+{
+	QUICK_SCOPE_CYCLE_COUNTER(STAT_InstancedStaticMeshSceneProxy_GetMeshElements);
+
+#if !(UE_BUILD_SHIPPING || UE_BUILD_TEST)
+	const bool bSelectionRenderEnabled = GIsEditor && ViewFamily.EngineShowFlags.Selection;
+
+	// If the first pass rendered selected instances only, we need to render the deselected instances in a second pass
+	const int32 NumSelectionGroups = (bSelectionRenderEnabled && bHasSelectedInstances) ? 2 : 1;
+
+	const FInstancingUserData* PassUserData[2] =
+	{
+		bHasSelectedInstances && bSelectionRenderEnabled ? &UserData_SelectedInstances : &UserData_AllInstances,
+		&UserData_DeselectedInstances
+	};
+
+	bool BatchRenderSelection[2] = 
+	{
+		bSelectionRenderEnabled && IsSelected(),
+		false
+	};
+
+	const bool bIsWireframe = ViewFamily.EngineShowFlags.Wireframe;
+
+	for (int32 ViewIndex = 0; ViewIndex < Views.Num(); ViewIndex++)
+	{
+		if (VisibilityMap & (1 << ViewIndex))
+		{
+			const FSceneView* View = Views[ViewIndex];
+
+			for (int32 SelectionGroupIndex = 0; SelectionGroupIndex < NumSelectionGroups; SelectionGroupIndex++)
+			{
+				const int32 LODIndex = GetLOD(View);
+				const FStaticMeshLODResources& LODModel = StaticMesh->RenderData->LODResources[LODIndex];
+
+				for (int32 SectionIndex = 0; SectionIndex < LODModel.Sections.Num(); SectionIndex++)
+				{
+					const int32 NumBatches = GetNumMeshBatches();
+
+					for (int32 BatchIndex = 0; BatchIndex < NumBatches; BatchIndex++)
+					{
+						FMeshBatch& MeshElement = Collector.AllocateMesh();
+
+						if (GetMeshElement(LODIndex, BatchIndex, SectionIndex, GetDepthPriorityGroup(View), BatchRenderSelection[SelectionGroupIndex], IsHovered(), MeshElement))
+						{
+							//@todo-rco this is only supporting selection on the first element
+							MeshElement.Elements[0].UserData = PassUserData[SelectionGroupIndex];
+							MeshElement.bCanApplyViewModeOverrides = true;
+							MeshElement.bUseSelectionOutline = BatchRenderSelection[SelectionGroupIndex];
+							MeshElement.bUseWireframeSelectionColoring = BatchRenderSelection[SelectionGroupIndex];
+
+							Collector.AddMesh(ViewIndex, MeshElement);
+							INC_DWORD_STAT_BY(STAT_StaticMeshTriangles, MeshElement.GetNumPrimitives());
+						}
+					}
+				}
+			}
+		}
+	}
+#endif
+}
+
 void FInstancedStaticMeshSceneProxy::DrawDynamicElements(FPrimitiveDrawInterface* PDI,const FSceneView* View)
 {
 	QUICK_SCOPE_CYCLE_COUNTER( STAT_InstancedStaticMeshSceneProxy_DrawDynamicElements );
@@ -1019,17 +1088,19 @@ void FInstancedStaticMeshSceneProxy::DrawDynamicElements(FPrimitiveDrawInterface
 
 			for(int32 SectionIndex = 0;SectionIndex < LODModel.Sections.Num();SectionIndex++)
 			{
-				TArray<FMeshBatch, TInlineAllocator<1>> MeshBatches;
-				new (MeshBatches) FMeshBatch();
-				if (GetMeshElements(LODIndex,SectionIndex,GetDepthPriorityGroup(View), MeshBatches, PassRenderSelection[Pass], IsHovered()))
+				const int32 NumBatches = GetNumMeshBatches();
+
+				for (int32 BatchIndex = 0; BatchIndex < NumBatches; BatchIndex++)
 				{
-					for (int32 Index = 0; Index < MeshBatches.Num(); ++Index)
+					FMeshBatch MeshBatch;
+
+					if (GetMeshElement(LODIndex, BatchIndex, SectionIndex, GetDepthPriorityGroup(View), PassRenderSelection[Pass], IsHovered(), MeshBatch))
 					{
-						MeshBatches[Index].Elements[0].UserData = PassUserData[Pass];
+						MeshBatch.Elements[0].UserData = PassUserData[Pass];
 
 						const int32 NumCalls = DrawRichMesh(
 							PDI,
-							MeshBatches[Index],
+							MeshBatch,
 							WireframeColor,
 							UtilColor,
 							PropertyColor,
@@ -1037,7 +1108,7 @@ void FInstancedStaticMeshSceneProxy::DrawDynamicElements(FPrimitiveDrawInterface
 							PassRenderSelection[Pass],
 							bIsWireframe
 							);
-						INC_DWORD_STAT_BY(STAT_StaticMeshTriangles,MeshBatches[Index].GetNumPrimitives() * NumCalls);
+						INC_DWORD_STAT_BY(STAT_StaticMeshTriangles,MeshBatch.GetNumPrimitives() * NumCalls);
 					}
 				}
 			}
@@ -1046,94 +1117,96 @@ void FInstancedStaticMeshSceneProxy::DrawDynamicElements(FPrimitiveDrawInterface
 #endif
 }
 
-inline void FInstancedStaticMeshSceneProxy::SetupInstancedMeshBatch(int32 LODIndex, TArray<FMeshBatch, TInlineAllocator<1>>& OutMeshBatches) const
+int32 FInstancedStaticMeshSceneProxy::GetNumMeshBatches() const
 {
 	const bool bInstanced = RHISupportsInstancing(GRHIShaderPlatform);
-	OutMeshBatches[0].VertexFactory = &InstancedRenderData.VertexFactories[LODIndex];
-	uint32 NumInstances = InstancedRenderData.InstanceBuffer.GetNumInstances();
-	auto* OutBatchElement0 = OutMeshBatches[0].Elements.GetTypedData();
-	OutBatchElement0->UserData = (void*)&UserData_AllInstances;
-	OutBatchElement0->UserIndex = 0;
+
 	if (bInstanced)
 	{
-		OutBatchElement0->NumInstances = NumInstances;
+		return 1;
 	}
 	else
 	{
-		// Create a BatchElement per Instance, and if they overflow the size per mask, create additional MeshBatches
-
-		const uint32 TotalInstances = NumInstances;
+		const uint32 NumInstances = InstancedRenderData.InstanceBuffer.GetNumInstances();
 		const uint32 MaxInstancesPerBatch = FInstancedStaticMeshVertexFactory::NumBitsForVisibilityMask();
-		uint32 NumBatches = (TotalInstances + MaxInstancesPerBatch - 1) / MaxInstancesPerBatch;
+		const uint32 NumBatches = FMath::DivideAndRoundUp(NumInstances, MaxInstancesPerBatch);
+		return NumBatches;
+	}
+}
 
-		// Make a copy of the Original MeshBatch, as it only has one Element at this point (cheaper to memcpy)
-		FMeshBatch OriginalBatch = OutMeshBatches[0];
+void FInstancedStaticMeshSceneProxy::SetupInstancedMeshBatch(int32 LODIndex, int32 BatchIndex, FMeshBatch& OutMeshBatch) const
+{
+	const bool bInstanced = RHISupportsInstancing(GRHIShaderPlatform);
+	OutMeshBatch.VertexFactory = &InstancedRenderData.VertexFactories[LODIndex];
+	const uint32 NumInstances = InstancedRenderData.InstanceBuffer.GetNumInstances();
+	FMeshBatchElement& BatchElement0 = OutMeshBatch.Elements[0];
+	BatchElement0.UserData = (void*)&UserData_AllInstances;
+	BatchElement0.UserIndex = 0;
 
-		// Add more Batches if we can't fit all Instances within the visibility mask
-		OutMeshBatches.Reserve(NumBatches);
+	if (bInstanced)
+	{
+		BatchElement0.NumInstances = NumInstances;
+	}
+	else
+	{
+		const uint32 MaxInstancesPerBatch = FInstancedStaticMeshVertexFactory::NumBitsForVisibilityMask();
+		const uint32 NumBatches = FMath::DivideAndRoundUp(NumInstances, MaxInstancesPerBatch);
+		const uint32 NumInstancesThisBatch = BatchIndex == NumBatches - 1 ? NumInstances % MaxInstancesPerBatch : MaxInstancesPerBatch;
 
-		uint32 RemainingInstances = TotalInstances;
-		int32 InstanceIndex = 0;
-		for (uint32 Batch = 0; Batch < NumBatches; ++Batch)
+		OutMeshBatch.Elements.Reserve(NumInstances);
+
+		for (uint32 Instance = 0; Instance < NumInstancesThisBatch; ++Instance)
 		{
-			NumInstances = FMath::Min(RemainingInstances, MaxInstancesPerBatch);
-			if (Batch > 0)
+			FMeshBatchElement* NewBatchElement; 
+
+			if (Instance == 0)
 			{
-				auto* NewBatch = new (OutMeshBatches) FMeshBatch();
-				*NewBatch = OriginalBatch;
+				NewBatchElement = &BatchElement0;
 			}
-
-			OutMeshBatches[Batch].Elements.Reserve(NumInstances);
-			for (uint32 Instance = 0; Instance < NumInstances; ++Instance)
+			else
 			{
-				auto* NewBatchElement = (Instance == 0)
-					? OutMeshBatches[Batch].Elements.GetTypedData()
-					: new(OutMeshBatches[Batch].Elements) FMeshBatchElement();
-				*NewBatchElement = *OutBatchElement0;
-				NewBatchElement->UserIndex = InstanceIndex++;
+				NewBatchElement = new(OutMeshBatch.Elements) FMeshBatchElement();
+				*NewBatchElement = BatchElement0;
 			}
-
-			RemainingInstances -= NumInstances;
+			
+			const int32 InstanceIndex = BatchIndex * MaxInstancesPerBatch + Instance;
+			NewBatchElement->UserIndex = InstanceIndex;
 		}
-
-		check(RemainingInstances == 0);
 	}
 }
 
-
-bool FInstancedStaticMeshSceneProxy::GetShadowMeshElements(int32 LODIndex, uint8 InDepthPriorityGroup, TArray<FMeshBatch, TInlineAllocator<1>>& OutMeshBatches) const
+bool FInstancedStaticMeshSceneProxy::GetShadowMeshElement(int32 LODIndex, int32 BatchIndex, uint8 InDepthPriorityGroup, FMeshBatch& OutMeshBatch) const
 {
-	if (LODIndex < InstancedRenderData.VertexFactories.Num() && FStaticMeshSceneProxy::GetShadowMeshElements(LODIndex, InDepthPriorityGroup, OutMeshBatches))
+	if (LODIndex < InstancedRenderData.VertexFactories.Num() && FStaticMeshSceneProxy::GetShadowMeshElement(LODIndex, BatchIndex, InDepthPriorityGroup, OutMeshBatch))
 	{
-		SetupInstancedMeshBatch(LODIndex, OutMeshBatches);
+		SetupInstancedMeshBatch(LODIndex, BatchIndex, OutMeshBatch);
 		return true;
 	}
 	return false;
 }
 
 /** Sets up a FMeshBatch for a specific LOD and element. */
-bool FInstancedStaticMeshSceneProxy::GetMeshElements(int32 LODIndex,int32 ElementIndex,uint8 InDepthPriorityGroup, TArray<FMeshBatch, TInlineAllocator<1>>& OutMeshBatches, const bool bUseSelectedMaterial, const bool bUseHoveredMaterial) const
+bool FInstancedStaticMeshSceneProxy::GetMeshElement(int32 LODIndex, int32 BatchIndex, int32 ElementIndex, uint8 InDepthPriorityGroup, const bool bUseSelectedMaterial, const bool bUseHoveredMaterial, FMeshBatch& OutMeshBatch) const
 {
-	if (LODIndex < InstancedRenderData.VertexFactories.Num() && FStaticMeshSceneProxy::GetMeshElements(LODIndex, ElementIndex, InDepthPriorityGroup, OutMeshBatches, bUseSelectedMaterial, bUseHoveredMaterial))
+	if (LODIndex < InstancedRenderData.VertexFactories.Num() && FStaticMeshSceneProxy::GetMeshElement(LODIndex, BatchIndex, ElementIndex, InDepthPriorityGroup, bUseSelectedMaterial, bUseHoveredMaterial, OutMeshBatch))
 	{
-		SetupInstancedMeshBatch(LODIndex, OutMeshBatches);
+		SetupInstancedMeshBatch(LODIndex, BatchIndex, OutMeshBatch);
 		return true;
 	}
 	return false;
 };
 
 /** Sets up a wireframe FMeshBatch for a specific LOD. */
-bool FInstancedStaticMeshSceneProxy::GetWireframeMeshElements(int32 LODIndex, const FMaterialRenderProxy* WireframeRenderProxy, uint8 InDepthPriorityGroup, TArray<FMeshBatch, TInlineAllocator<1>>& OutMeshBatches) const
+bool FInstancedStaticMeshSceneProxy::GetWireframeMeshElement(int32 LODIndex, int32 BatchIndex, const FMaterialRenderProxy* WireframeRenderProxy, uint8 InDepthPriorityGroup, FMeshBatch& OutMeshBatch) const
 {
-	if (LODIndex < InstancedRenderData.VertexFactories.Num() && FStaticMeshSceneProxy::GetWireframeMeshElements(LODIndex, WireframeRenderProxy, InDepthPriorityGroup, OutMeshBatches))
+	if (LODIndex < InstancedRenderData.VertexFactories.Num() && FStaticMeshSceneProxy::GetWireframeMeshElement(LODIndex, BatchIndex, WireframeRenderProxy, InDepthPriorityGroup, OutMeshBatch))
 	{
-		SetupInstancedMeshBatch(LODIndex, OutMeshBatches);
+		SetupInstancedMeshBatch(LODIndex, BatchIndex, OutMeshBatch);
 		return true;
 	}
 	return false;
 }
 
-
 /*-----------------------------------------------------------------------------
 	UInstancedStaticMeshComponent
 -----------------------------------------------------------------------------*/
@@ -1781,7 +1854,7 @@ void FInstancedStaticMeshVertexFactoryShaderParameters::SetMesh( FRHICommandList
 	{
 		FVector4 InstancingFadeOutParams(0.f,0.f,1.f,1.f);
 
-		const FInstancingUserData* InstancingUserData = (FInstancingUserData*)BatchElement.UserData;
+		const FInstancingUserData* InstancingUserData = (const FInstancingUserData*)BatchElement.UserData;
 		if (InstancingUserData)
 		{
 			InstancingFadeOutParams.X = InstancingUserData->StartCullDistance;
@@ -1812,8 +1885,8 @@ void FInstancedStaticMeshVertexFactoryShaderParameters::SetMesh( FRHICommandList
 	{
 		if (CPUInstanceShadowMapBias.IsBound())
 		{
-			const auto* InstancingData = (FInstancingUserData*)BatchElement.UserData;
-			auto* InstanceStream = ((FInstancingUserData::FInstanceStream*)InstancingData->RenderData->InstanceBuffer.GetRawData()) + BatchElement.UserIndex;
+			auto* InstancingData = (const FInstancingUserData*)BatchElement.UserData;
+			auto* InstanceStream = ((const FInstancingUserData::FInstanceStream*)InstancingData->RenderData->InstanceBuffer.GetRawData()) + BatchElement.UserIndex;
 			SetShaderValue(RHICmdList, VS, CPUInstanceShadowMapBias, InstanceStream->InstanceShadowmapUVBias);
 			SetShaderValueArray(RHICmdList, VS, CPUInstanceTransform, InstanceStream->InstanceTransform, 3);
 			SetShaderValueArray(RHICmdList, VS, CPUInstanceInverseTransform, InstanceStream->InstanceInverseTransform, 3);
diff --git a/Engine/Source/Runtime/Engine/Private/Landscape/LandscapeGizmoActor.cpp b/Engine/Source/Runtime/Engine/Private/Landscape/LandscapeGizmoActor.cpp
--- a/Engine/Source/Runtime/Engine/Private/Landscape/LandscapeGizmoActor.cpp
+++ b/Engine/Source/Runtime/Engine/Private/Landscape/LandscapeGizmoActor.cpp
@@ -206,6 +206,102 @@ public:
 		HeightmapRenderProxy = NULL;
 	}
 
+	virtual void GetDynamicMeshElements(const TArray<const FSceneView*>& Views, const FSceneViewFamily& ViewFamily, uint32 VisibilityMap, FMeshElementCollector& Collector) const override
+	{
+		//FMemMark Mark(FMemStack::Get());
+#if WITH_EDITOR
+		if( GizmoRenderProxy &&  HeightmapRenderProxy )
+		{
+			for (int32 ViewIndex = 0; ViewIndex < Views.Num(); ViewIndex++)
+			{
+				if (VisibilityMap & (1 << ViewIndex))
+				{
+					const FSceneView* View = Views[ViewIndex];
+					FPrimitiveDrawInterface* PDI = Collector.GetPDI(ViewIndex);
+
+					// Axis
+					PDI->DrawLine( Origin, XAxis, FLinearColor(1, 0, 0), SDPG_World );
+					PDI->DrawLine( Origin, YAxis, FLinearColor(0, 1, 0), SDPG_World );
+
+					{
+						FDynamicMeshBuilder MeshBuilder;
+
+						MeshBuilder.AddVertex(FrustumVerts[0], FVector2D(0, 0), FVector(1,0,0), FVector(0,1,0), FVector(0,0,1), FColor(255,255,255));
+						MeshBuilder.AddVertex(FrustumVerts[1], FVector2D(1, 0), FVector(1,0,0), FVector(0,1,0), FVector(0,0,1), FColor(255,255,255));
+						MeshBuilder.AddVertex(FrustumVerts[2], FVector2D(1, 1), FVector(1,0,0), FVector(0,1,0), FVector(0,0,1), FColor(255,255,255));
+						MeshBuilder.AddVertex(FrustumVerts[3], FVector2D(0, 1), FVector(1,0,0), FVector(0,1,0), FVector(0,0,1), FColor(255,255,255));
+
+						MeshBuilder.AddVertex(FrustumVerts[4], FVector2D(0, 0), FVector(1,0,0), FVector(0,1,0), FVector(0,0,1), FColor(255,255,255));
+						MeshBuilder.AddVertex(FrustumVerts[5], FVector2D(1, 0), FVector(1,0,0), FVector(0,1,0), FVector(0,0,1), FColor(255,255,255));
+						MeshBuilder.AddVertex(FrustumVerts[6], FVector2D(1, 1), FVector(1,0,0), FVector(0,1,0), FVector(0,0,1), FColor(255,255,255));
+						MeshBuilder.AddVertex(FrustumVerts[7], FVector2D(0, 1), FVector(1,0,0), FVector(0,1,0), FVector(0,0,1), FColor(255,255,255));
+
+						MeshBuilder.AddVertex(FrustumVerts[1], FVector2D(0, 0), FVector(1,0,0), FVector(0,1,0), FVector(0,0,1), FColor(255,255,255));
+						MeshBuilder.AddVertex(FrustumVerts[0], FVector2D(1, 0), FVector(1,0,0), FVector(0,1,0), FVector(0,0,1), FColor(255,255,255));
+						MeshBuilder.AddVertex(FrustumVerts[4], FVector2D(1, 1), FVector(1,0,0), FVector(0,1,0), FVector(0,0,1), FColor(255,255,255));
+						MeshBuilder.AddVertex(FrustumVerts[5], FVector2D(0, 1), FVector(1,0,0), FVector(0,1,0), FVector(0,0,1), FColor(255,255,255));
+
+						MeshBuilder.AddVertex(FrustumVerts[3], FVector2D(0, 0), FVector(1,0,0), FVector(0,1,0), FVector(0,0,1), FColor(255,255,255));
+						MeshBuilder.AddVertex(FrustumVerts[2], FVector2D(1, 0), FVector(1,0,0), FVector(0,1,0), FVector(0,0,1), FColor(255,255,255));
+						MeshBuilder.AddVertex(FrustumVerts[6], FVector2D(1, 1), FVector(1,0,0), FVector(0,1,0), FVector(0,0,1), FColor(255,255,255));
+						MeshBuilder.AddVertex(FrustumVerts[7], FVector2D(0, 1), FVector(1,0,0), FVector(0,1,0), FVector(0,0,1), FColor(255,255,255));
+
+						MeshBuilder.AddVertex(FrustumVerts[2], FVector2D(0, 0), FVector(1,0,0), FVector(0,1,0), FVector(0,0,1), FColor(255,255,255));
+						MeshBuilder.AddVertex(FrustumVerts[1], FVector2D(1, 0), FVector(1,0,0), FVector(0,1,0), FVector(0,0,1), FColor(255,255,255));
+						MeshBuilder.AddVertex(FrustumVerts[5], FVector2D(1, 1), FVector(1,0,0), FVector(0,1,0), FVector(0,0,1), FColor(255,255,255));
+						MeshBuilder.AddVertex(FrustumVerts[6], FVector2D(0, 1), FVector(1,0,0), FVector(0,1,0), FVector(0,0,1), FColor(255,255,255));
+
+						MeshBuilder.AddVertex(FrustumVerts[0], FVector2D(0, 0), FVector(1,0,0), FVector(0,1,0), FVector(0,0,1), FColor(255,255,255));
+						MeshBuilder.AddVertex(FrustumVerts[3], FVector2D(1, 0), FVector(1,0,0), FVector(0,1,0), FVector(0,0,1), FColor(255,255,255));
+						MeshBuilder.AddVertex(FrustumVerts[7], FVector2D(1, 1), FVector(1,0,0), FVector(0,1,0), FVector(0,0,1), FColor(255,255,255));
+						MeshBuilder.AddVertex(FrustumVerts[4], FVector2D(0, 1), FVector(1,0,0), FVector(0,1,0), FVector(0,0,1), FColor(255,255,255));
+
+						for (int32 i = 0; i < 6; ++i)
+						{
+							int32 Idx = i*4;
+							MeshBuilder.AddTriangle( Idx, Idx+2, Idx+1 );
+							MeshBuilder.AddTriangle( Idx, Idx+3, Idx+2 );
+						}
+
+						MeshBuilder.GetMesh(FMatrix::Identity, GizmoRenderProxy , SDPG_World, true, false, ViewIndex, Collector);
+					}
+
+					if (bHeightmapRendering)
+					{		  		
+						FDynamicMeshBuilder MeshBuilder;
+
+						for (int32 Y = 0; Y < SampleSizeY; ++Y)
+						{
+							for (int32 X = 0; X < SampleSizeX; ++X)
+							{
+								FVector SampledNormal = SampledNormals[X + Y * SampleSizeX];
+								FVector TangentX(SampledNormal.Z, 0, -SampledNormal.X);
+								TangentX = TangentX.SafeNormal();
+
+								MeshBuilder.AddVertex(SampledPositions[X + Y * SampleSizeX], FVector2D((float)X / (SampleSizeX), (float)Y / (SampleSizeY)), TangentX, SampledNormal^TangentX, SampledNormal, FColor(255, 255, 255) );
+							}
+						}
+
+						for (int32 Y = 0; Y < SampleSizeY; ++Y)
+						{
+							for (int32 X = 0; X < SampleSizeX; ++X)
+							{
+								if (X < SampleSizeX - 1 && Y < SampleSizeY - 1)
+								{
+									MeshBuilder.AddTriangle( (X+0) + (Y+0) * SampleSizeX, (X+1) + (Y+1) * SampleSizeX, (X+1) + (Y+0) * SampleSizeX );
+									MeshBuilder.AddTriangle( (X+0) + (Y+0) * SampleSizeX, (X+0) + (Y+1) * SampleSizeX, (X+1) + (Y+1) * SampleSizeX );
+								}
+							}
+						}
+
+						MeshBuilder.GetMesh(MeshRT, HeightmapRenderProxy , SDPG_World, false, false, ViewIndex, Collector);
+					}
+				}
+			}
+		}
+#endif
+	};
+
 	virtual void DrawDynamicElements(FPrimitiveDrawInterface* PDI,const FSceneView* View)
 	{
 		//FMemMark Mark(FMemStack::Get());
diff --git a/Engine/Source/Runtime/Engine/Private/Landscape/LandscapeRender.cpp b/Engine/Source/Runtime/Engine/Private/Landscape/LandscapeRender.cpp
--- a/Engine/Source/Runtime/Engine/Private/Landscape/LandscapeRender.cpp
+++ b/Engine/Source/Runtime/Engine/Private/Landscape/LandscapeRender.cpp
@@ -548,11 +548,12 @@ FLandscapeComponentSceneProxy::FLandscapeComponentSceneProxy(ULandscapeComponent
 ,	EditToolRenderData(InEditToolRenderData)
 ,	ComponentLightInfo(NULL)
 ,	LandscapeComponent(InComponent)
-,	LevelColor(FLinearColor::White)
 ,	ForcedLOD(InComponent->ForcedLOD)
 ,	LODBias(InComponent->LODBias)
 ,	LODFalloff(InComponent->GetLandscapeProxy()->LODFalloff)
 {
+	LevelColor = FLinearColor(1.f, 1.f, 1.f);
+
 	const auto FeatureLevel = GetScene()->GetFeatureLevel();
 	if (FeatureLevel == ERHIFeatureLevel::ES2)
 	{
@@ -848,7 +849,7 @@ FPrimitiveViewRelevance FLandscapeComponentSceneProxy::GetViewRelevance(const FS
 #endif
 
 	// Use the dynamic path for rendering landscape components pass only for Rich Views or if the static path is disabled for debug.
-	if(	IsRichView(View) || 
+	if(	IsRichView(*View->Family) || 
 		GLandscapeDebugOptions.bDisableStatic || 
 		View->Family->EngineShowFlags.Wireframe || 
 #if WITH_EDITOR
@@ -1346,6 +1347,318 @@ void FLandscapeComponentSceneProxy::CalcLODParamsForSubsection(const class FScen
 	}
 }
 
+void FLandscapeComponentSceneProxy::GetDynamicMeshElements(const TArray<const FSceneView*>& Views, const FSceneViewFamily& ViewFamily, uint32 VisibilityMap, FMeshElementCollector& Collector) const
+{
+	QUICK_SCOPE_CYCLE_COUNTER(STAT_FLandscapeComponentSceneProxy_GetMeshElements);
+
+	int32 NumPasses=0;
+	int32 NumTriangles=0;
+	int32 NumDrawCalls=0;
+	const bool bIsWireframe = ViewFamily.EngineShowFlags.Wireframe;
+
+	for (int32 ViewIndex = 0; ViewIndex < Views.Num(); ViewIndex++)
+	{
+		if (VisibilityMap & (1 << ViewIndex))
+		{
+			const FSceneView* View = Views[ViewIndex];
+			FVector CameraLocalPos3D = WorldToLocal.TransformPosition(View->ViewMatrices.ViewOrigin); 
+			FVector2D CameraLocalPos(CameraLocalPos3D.X, CameraLocalPos3D.Y);
+
+			FLandscapeElementParamArray& ParameterArray = Collector.AllocateOneFrameResource<FLandscapeElementParamArray>();
+			ParameterArray.ElementParams.Empty(NumSubsections * NumSubsections);
+			ParameterArray.ElementParams.AddUninitialized(NumSubsections * NumSubsections);
+
+			FMeshBatch& Mesh = Collector.AllocateMesh();
+
+			// Could be different from bRequiresAdjacencyInformation during shader compilation
+			FMaterialRenderProxy* RenderProxy = MaterialInterface->GetRenderProxy(false);
+			bool bCurrentRequiresAdjacencyInformation = RequiresAdjacencyInformation( RenderProxy, View->GetFeatureLevel() );
+			Mesh.Type = bCurrentRequiresAdjacencyInformation ? PT_12_ControlPointPatchList : PT_TriangleList;
+			Mesh.LCI = ComponentLightInfo; 
+			Mesh.CastShadow = true;
+			Mesh.VertexFactory = VertexFactory;
+			Mesh.ReverseCulling = IsLocalToWorldDeterminantNegative();
+
+#if WITH_EDITOR
+			FMeshBatch& MeshTools = Collector.AllocateMesh();
+			MeshTools.LCI = ComponentLightInfo; 
+			MeshTools.Type = PT_TriangleList;
+			MeshTools.CastShadow = false;
+			MeshTools.VertexFactory = VertexFactory;
+			MeshTools.ReverseCulling = IsLocalToWorldDeterminantNegative();
+#endif
+
+			// Setup the LOD parameters
+			for( int32 SubY=0;SubY<NumSubsections;SubY++ )
+			{
+				for( int32 SubX=0;SubX<NumSubsections;SubX++ )
+				{
+					int32 SubSectionIdx = SubX + SubY*NumSubsections;
+					int32 CurrentLOD = CalcLODForSubsection(*View, SubX, SubY, CameraLocalPos);
+
+					FMeshBatchElement& BatchElement = (SubX==0 && SubY==0) ? *Mesh.Elements.GetTypedData() : *(new(Mesh.Elements) FMeshBatchElement);
+					BatchElement.PrimitiveUniformBufferResource = &GetUniformBuffer();
+					FLandscapeBatchElementParams& BatchElementParams = ParameterArray.ElementParams[SubSectionIdx];
+					BatchElementParams.LocalToWorldNoScalingPtr = &LocalToWorldNoScaling;
+					BatchElement.UserData = &BatchElementParams;
+
+					BatchElementParams.LandscapeUniformShaderParametersResource = &LandscapeUniformShaderParameters;
+					BatchElementParams.SceneProxy = this;
+					BatchElementParams.SubX = SubX;
+					BatchElementParams.SubY = SubY;
+					BatchElementParams.CurrentLOD = CurrentLOD;
+
+					if (bCurrentRequiresAdjacencyInformation)
+					{
+						check(SharedBuffers->AdjacencyIndexBuffers);
+						BatchElement.IndexBuffer = SharedBuffers->AdjacencyIndexBuffers->IndexBuffers[CurrentLOD];
+					}
+					else
+					{
+						BatchElement.IndexBuffer = SharedBuffers->IndexBuffers[CurrentLOD];
+					}
+
+					int32 LodSubsectionSizeVerts = (SubsectionSizeVerts >> CurrentLOD);
+
+					int32 NumPrimitives = FMath::Square((LodSubsectionSizeVerts - 1)) * 2;
+					BatchElement.NumPrimitives = NumPrimitives;
+					BatchElement.FirstIndex = (SubX + SubY * NumSubsections) * NumPrimitives * 3;
+					BatchElement.MinVertexIndex = SharedBuffers->IndexRanges[CurrentLOD].MinIndex[SubX][SubY];
+					BatchElement.MaxVertexIndex = SharedBuffers->IndexRanges[CurrentLOD].MaxIndex[SubX][SubY];
+					
+#if WITH_EDITOR
+					FMeshBatchElement& BatchElementTools = (SubX==0 && SubY==0) ? *MeshTools.Elements.GetTypedData() : *(new(MeshTools.Elements) FMeshBatchElement);
+					BatchElementTools.PrimitiveUniformBufferResource = &GetUniformBuffer();
+					BatchElementTools.UserData = &BatchElementParams;
+
+					// Tools never use tessellation
+					BatchElementTools.IndexBuffer = SharedBuffers->IndexBuffers[CurrentLOD];
+					BatchElementTools.NumPrimitives = NumPrimitives;
+					BatchElementTools.FirstIndex = BatchElement.FirstIndex;
+					BatchElementTools.MinVertexIndex = SharedBuffers->IndexRanges[CurrentLOD].MinIndex[SubX][SubY];
+					BatchElementTools.MaxVertexIndex = SharedBuffers->IndexRanges[CurrentLOD].MaxIndex[SubX][SubY];
+#endif
+				}
+			}
+			// Render the landscape component
+#if WITH_EDITOR
+			
+			switch(GLandscapeViewMode)
+			{
+			case ELandscapeViewMode::DebugLayer:
+				if( GLayerDebugColorMaterial && EditToolRenderData )
+				{
+					auto DebugColorMaterialInstance = new FLandscapeDebugMaterialRenderProxy(GLayerDebugColorMaterial->GetRenderProxy(false), 
+						(EditToolRenderData->DebugChannelR >= 0 ? WeightmapTextures[EditToolRenderData->DebugChannelR/4] : NULL),
+						(EditToolRenderData->DebugChannelG >= 0 ? WeightmapTextures[EditToolRenderData->DebugChannelG/4] : NULL),
+						(EditToolRenderData->DebugChannelB >= 0 ? WeightmapTextures[EditToolRenderData->DebugChannelB/4] : NULL),	
+						(EditToolRenderData->DebugChannelR >= 0 ? DebugColorMask::Masks[EditToolRenderData->DebugChannelR%4] : DebugColorMask::Masks[4]),
+						(EditToolRenderData->DebugChannelG >= 0 ? DebugColorMask::Masks[EditToolRenderData->DebugChannelG%4] : DebugColorMask::Masks[4]),
+						(EditToolRenderData->DebugChannelB >= 0 ? DebugColorMask::Masks[EditToolRenderData->DebugChannelB%4] : DebugColorMask::Masks[4])
+						);
+
+					MeshTools.MaterialRenderProxy = DebugColorMaterialInstance;
+					Collector.RegisterOneFrameMaterialProxy(DebugColorMaterialInstance);
+
+					MeshTools.bCanApplyViewModeOverrides = true;
+					MeshTools.bUseWireframeSelectionColoring = IsSelected();
+
+					Collector.AddMesh(ViewIndex, MeshTools);
+
+					NumPasses++;
+					NumTriangles += MeshTools.GetNumPrimitives();
+					NumDrawCalls += MeshTools.Elements.Num();
+				}
+				break;
+
+			case ELandscapeViewMode::LayerDensity:
+				if( EditToolRenderData )
+				{
+					int32 ColorIndex = FMath::Min<int32>(NumWeightmapLayerAllocations, GEngine->ShaderComplexityColors.Num());
+					auto LayerDensityMaterialInstance = new FColoredMaterialRenderProxy(GEngine->LevelColorationUnlitMaterial->GetRenderProxy(false), ColorIndex ? GEngine->ShaderComplexityColors[ColorIndex-1] : FLinearColor::Black);
+
+					MeshTools.MaterialRenderProxy = LayerDensityMaterialInstance;
+					Collector.RegisterOneFrameMaterialProxy(LayerDensityMaterialInstance);
+
+					MeshTools.bCanApplyViewModeOverrides = true;
+					MeshTools.bUseWireframeSelectionColoring = IsSelected();
+
+					Collector.AddMesh(ViewIndex, MeshTools);
+
+					NumPasses++;
+					NumTriangles += MeshTools.GetNumPrimitives();
+					NumDrawCalls += MeshTools.Elements.Num();
+				}
+				break;
+
+			case ELandscapeViewMode::LOD:
+				{
+					FLinearColor WireColors[7];
+					WireColors[0] = FLinearColor(1,1,1,1);
+					WireColors[1] = FLinearColor(1,0,0,1);
+					WireColors[2] = FLinearColor(0,1,0,1);
+					WireColors[3] = FLinearColor(0,0,1,1);
+					WireColors[4] = FLinearColor(1,1,0,1);
+					WireColors[5] = FLinearColor(1,0,1,1);
+					WireColors[6] = FLinearColor(0,1,1,1);
+
+					for( int32 i=0;i<MeshTools.Elements.Num();i++ )
+					{
+						FMeshBatch& LODMesh = Collector.AllocateMesh();
+						LODMesh = MeshTools;
+						LODMesh.Elements.Empty(1);
+						LODMesh.Elements.Add(MeshTools.Elements[i]);
+						int32 ColorIndex = ((FLandscapeBatchElementParams*)MeshTools.Elements[i].UserData)->CurrentLOD;
+						FLinearColor Color = ForcedLOD >= 0 ? WireColors[ColorIndex] : WireColors[ColorIndex]*0.2f;
+						auto LODMaterialInstance = new FColoredMaterialRenderProxy(GEngine->LevelColorationUnlitMaterial->GetRenderProxy(false), Color);
+						LODMesh.MaterialRenderProxy = LODMaterialInstance;
+						Collector.RegisterOneFrameMaterialProxy(LODMaterialInstance);
+
+						LODMesh.bCanApplyViewModeOverrides = true;
+						LODMesh.bUseWireframeSelectionColoring = IsSelected();
+
+						Collector.AddMesh(ViewIndex, LODMesh);
+
+						NumPasses++;
+						NumTriangles += MeshTools.Elements[i].NumPrimitives;
+						NumDrawCalls++;
+					}
+				}
+				break;
+
+			default:
+			
+#else
+			{
+#endif // WITH_EDITOR
+				// Regular Landscape rendering. Only use the dynamic path if we're rendering a rich view or we've disabled the static path for debugging.
+				if( IsRichView(ViewFamily) || 
+					GLandscapeDebugOptions.bDisableStatic || 
+					bIsWireframe ||
+#if WITH_EDITOR
+					(IsSelected() && !GLandscapeEditModeActive)
+#else
+					IsSelected()
+#endif
+					)
+				{
+					Mesh.MaterialRenderProxy = MaterialInterface->GetRenderProxy(false);
+
+					Mesh.bCanApplyViewModeOverrides = true;
+					Mesh.bUseWireframeSelectionColoring = IsSelected();
+
+					Collector.AddMesh(ViewIndex, Mesh);
+
+					NumPasses++;
+					NumTriangles += Mesh.GetNumPrimitives();
+					NumDrawCalls += Mesh.Elements.Num();
+				}
+			}
+#if WITH_EDITOR
+			
+			// Extra render passes for landscape tools
+			if( GLandscapeEditModeActive )
+			{
+				// Region selection
+				if ( EditToolRenderData && EditToolRenderData->SelectedType )
+				{
+					if ((GLandscapeEditRenderMode & ELandscapeEditRenderMode::SelectRegion) && (EditToolRenderData->SelectedType & FLandscapeEditToolRenderData::ST_REGION)
+						&& !(GLandscapeEditRenderMode & ELandscapeEditRenderMode::Mask))
+					{
+						FMeshBatch& SelectMesh = Collector.AllocateMesh();
+						SelectMesh = MeshTools;
+						auto SelectMaterialInstance = new FLandscapeSelectMaterialRenderProxy(GSelectionRegionMaterial->GetRenderProxy(false), EditToolRenderData->DataTexture ? EditToolRenderData->DataTexture : GLandscapeBlackTexture);
+						SelectMesh.MaterialRenderProxy = SelectMaterialInstance;
+						Collector.RegisterOneFrameMaterialProxy(SelectMaterialInstance);
+						Collector.AddMesh(ViewIndex, SelectMesh);
+						NumPasses++;
+						NumTriangles += SelectMesh.GetNumPrimitives();
+						NumDrawCalls += SelectMesh.Elements.Num();
+					}
+
+					if ((GLandscapeEditRenderMode & ELandscapeEditRenderMode::SelectComponent) && (EditToolRenderData->SelectedType & FLandscapeEditToolRenderData::ST_COMPONENT))
+					{
+						FMeshBatch& SelectMesh = Collector.AllocateMesh();
+						SelectMesh = MeshTools;
+						SelectMesh.MaterialRenderProxy = GSelectionColorMaterial->GetRenderProxy(0);
+						Collector.AddMesh(ViewIndex, SelectMesh);
+						NumPasses++;
+						NumTriangles += SelectMesh.GetNumPrimitives();
+						NumDrawCalls += SelectMesh.Elements.Num();
+					}
+				}
+
+				// Mask
+				if ((GLandscapeEditRenderMode & ELandscapeEditRenderMode::SelectRegion) && (GLandscapeEditRenderMode & ELandscapeEditRenderMode::Mask))
+				{
+					if (EditToolRenderData && (EditToolRenderData->SelectedType & FLandscapeEditToolRenderData::ST_REGION) )
+					{
+						FMeshBatch& MaskMesh = Collector.AllocateMesh();
+						MaskMesh = MeshTools;
+						auto MaskMaterialInstance = new FLandscapeMaskMaterialRenderProxy(GMaskRegionMaterial->GetRenderProxy(false), EditToolRenderData->DataTexture ? EditToolRenderData->DataTexture : GLandscapeBlackTexture, !!(GLandscapeEditRenderMode & ELandscapeEditRenderMode::InvertedMask) );
+						MaskMesh.MaterialRenderProxy = MaskMaterialInstance;
+						Collector.RegisterOneFrameMaterialProxy(MaskMaterialInstance);
+						Collector.AddMesh(ViewIndex, MaskMesh);
+						NumPasses++;
+						NumTriangles += MaskMesh.GetNumPrimitives();
+						NumDrawCalls += MaskMesh.Elements.Num();
+					}
+					else if (!(GLandscapeEditRenderMode & ELandscapeEditRenderMode::InvertedMask))
+					{
+						FMeshBatch& MaskMesh = Collector.AllocateMesh();
+						MaskMesh = MeshTools;
+						auto MaskMaterialInstance = new FLandscapeMaskMaterialRenderProxy(GMaskRegionMaterial->GetRenderProxy(false), GLandscapeBlackTexture, false );
+						MaskMesh.MaterialRenderProxy = MaskMaterialInstance;
+						Collector.RegisterOneFrameMaterialProxy(MaskMaterialInstance);
+						Collector.AddMesh(ViewIndex, MaskMesh);
+						NumPasses++;
+						NumTriangles += MaskMesh.GetNumPrimitives();
+						NumDrawCalls += MaskMesh.Elements.Num();
+					}
+				}
+
+				// Edit mode tools
+				if( EditToolRenderData )
+				{
+					if (EditToolRenderData->ToolMaterial)
+					{
+						FMeshBatch& EditMesh = Collector.AllocateMesh();
+						EditMesh = MeshTools;
+						EditMesh.MaterialRenderProxy = EditToolRenderData->ToolMaterial->GetRenderProxy(0);		
+						Collector.AddMesh(ViewIndex, EditMesh);
+						NumPasses++;
+						NumTriangles += EditMesh.GetNumPrimitives();
+						NumDrawCalls += EditMesh.Elements.Num();
+					}
+
+					if (EditToolRenderData->GizmoMaterial && GLandscapeEditRenderMode & ELandscapeEditRenderMode::Gizmo)
+					{
+						FMeshBatch& EditMesh = Collector.AllocateMesh();
+						EditMesh = MeshTools;
+						EditMesh.MaterialRenderProxy = EditToolRenderData->GizmoMaterial->GetRenderProxy(0);
+						Collector.AddMesh(ViewIndex, EditMesh);
+						NumPasses++;
+						NumTriangles += EditMesh.GetNumPrimitives();
+						NumDrawCalls += EditMesh.Elements.Num();
+					}
+				}
+			}
+#endif // WITH_EDITOR
+
+			if ( GLandscapeDebugOptions.bShowPatches )
+			{
+				DrawWireBox(Collector.GetPDI(ViewIndex), GetBounds().GetBox(), FColor(255, 255, 0), SDPG_World);
+			}
+
+			RenderBounds(Collector.GetPDI(ViewIndex), ViewFamily.EngineShowFlags, GetBounds(), IsSelected());
+		}
+	}
+
+	INC_DWORD_STAT_BY(STAT_LandscapeComponents, NumPasses);
+	INC_DWORD_STAT_BY(STAT_LandscapeDrawCalls, NumDrawCalls);
+	INC_DWORD_STAT_BY(STAT_LandscapeTriangles, NumTriangles * NumPasses);
+}
+
 void FLandscapeComponentSceneProxy::DrawDynamicElements(FPrimitiveDrawInterface* PDI,const FSceneView* View)
 {
 	SCOPE_CYCLE_COUNTER(STAT_LandscapeDynamicDrawTime);
@@ -1474,7 +1787,7 @@ void FLandscapeComponentSceneProxy::DrawDynamicElements(FPrimitiveDrawInterface*
 	{
 #endif // WITH_EDITOR
 		// Regular Landscape rendering. Only use the dynamic path if we're rendering a rich view or we've disabled the static path for debugging.
-		if( IsRichView(View) || 
+		if( IsRichView(*View->Family) || 
 			GLandscapeDebugOptions.bDisableStatic || 
 			bIsWireframe ||
 #if WITH_EDITOR
@@ -1996,7 +2309,7 @@ public:
 	{
 		SCOPE_CYCLE_COUNTER(STAT_LandscapeVFDrawTime);
 
-		FLandscapeBatchElementParams* BatchElementParams = (FLandscapeBatchElementParams*)BatchElement.UserData;
+		const FLandscapeBatchElementParams* BatchElementParams = (const FLandscapeBatchElementParams*)BatchElement.UserData;
 		check(BatchElementParams);
 
 		const FLandscapeComponentSceneProxy* SceneProxy = BatchElementParams->SceneProxy;
@@ -2101,7 +2414,7 @@ public:
 	{
 		SCOPE_CYCLE_COUNTER(STAT_LandscapeVFDrawTime);
 
-		FLandscapeBatchElementParams* BatchElementParams = (FLandscapeBatchElementParams*)BatchElement.UserData;
+		const FLandscapeBatchElementParams* BatchElementParams = (const FLandscapeBatchElementParams*)BatchElement.UserData;
 		check(BatchElementParams);
 
 		const FLandscapeComponentSceneProxy* SceneProxy = BatchElementParams->SceneProxy;
@@ -2223,7 +2536,7 @@ void FLandscapeVertexFactoryPixelShaderParameters::SetMesh(FRHICommandList& RHIC
 {
 	SCOPE_CYCLE_COUNTER(STAT_LandscapeVFDrawTime);
 
-	FLandscapeBatchElementParams* BatchElementParams = (FLandscapeBatchElementParams*)BatchElement.UserData;
+	const FLandscapeBatchElementParams* BatchElementParams = (const FLandscapeBatchElementParams*)BatchElement.UserData;
 
 	if( LocalToWorldNoScalingParameter.IsBound() )
 	{
diff --git a/Engine/Source/Runtime/Engine/Private/Landscape/LandscapeRenderMobile.cpp b/Engine/Source/Runtime/Engine/Private/Landscape/LandscapeRenderMobile.cpp
--- a/Engine/Source/Runtime/Engine/Private/Landscape/LandscapeRenderMobile.cpp
+++ b/Engine/Source/Runtime/Engine/Private/Landscape/LandscapeRenderMobile.cpp
@@ -65,7 +65,7 @@ public:
 	{
 		SCOPE_CYCLE_COUNTER(STAT_LandscapeVFDrawTime);
 
-		FLandscapeBatchElementParams* BatchElementParams = (FLandscapeBatchElementParams*)BatchElement.UserData;
+		const FLandscapeBatchElementParams* BatchElementParams = (const FLandscapeBatchElementParams*)BatchElement.UserData;
 		check(BatchElementParams);
 
 		const FLandscapeComponentSceneProxy* SceneProxy = BatchElementParams->SceneProxy;
diff --git a/Engine/Source/Runtime/Engine/Private/Landscape/LandscapeSplines.cpp b/Engine/Source/Runtime/Engine/Private/Landscape/LandscapeSplines.cpp
--- a/Engine/Source/Runtime/Engine/Private/Landscape/LandscapeSplines.cpp
+++ b/Engine/Source/Runtime/Engine/Private/Landscape/LandscapeSplines.cpp
@@ -106,6 +106,166 @@ public:
 		return NULL;
 	}
 
+	virtual void GetDynamicMeshElements(const TArray<const FSceneView*>& Views, const FSceneViewFamily& ViewFamily, uint32 VisibilityMap, FMeshElementCollector& Collector) const override
+	{
+		// Slight Depth Bias so that the splines show up when they exactly match the target surface
+		// e.g. someone playing with splines on a newly-created perfectly-flat landscape
+		static const float DepthBias = -0.0001;
+
+		const FMatrix& LocalToWorld = GetLocalToWorld();
+
+		const FLinearColor SelectedSplineColor = GEngine->GetSelectedMaterialColor();
+		const FLinearColor SelectedControlPointSpriteColor = FLinearColor::White + (GEngine->GetSelectedMaterialColor() * GEngine->SelectionHighlightIntensityBillboards * 10); // copied from FSpriteSceneProxy::DrawDynamicElements()
+
+		for (int32 ViewIndex = 0; ViewIndex < Views.Num(); ViewIndex++)
+		{
+			if (VisibilityMap & (1 << ViewIndex))
+			{
+				const FSceneView* View = Views[ViewIndex];
+				FPrimitiveDrawInterface* PDI = Collector.GetPDI(ViewIndex);
+
+				for (const FSegmentProxy& Segment : Segments)
+				{
+					const FLinearColor SegmentColor = Segment.bSelected ? SelectedSplineColor : SplineColor;
+
+					FLandscapeSplineInterpPoint OldPoint = Segment.Points[0];
+					OldPoint.Center       = LocalToWorld.TransformPosition(OldPoint.Center);
+					OldPoint.Left         = LocalToWorld.TransformPosition(OldPoint.Left);
+					OldPoint.Right        = LocalToWorld.TransformPosition(OldPoint.Right);
+					OldPoint.FalloffLeft  = LocalToWorld.TransformPosition(OldPoint.FalloffLeft);
+					OldPoint.FalloffRight = LocalToWorld.TransformPosition(OldPoint.FalloffRight);
+					for (int32 i = 1; i < Segment.Points.Num(); i++)
+					{
+						FLandscapeSplineInterpPoint NewPoint = Segment.Points[i];
+						NewPoint.Center       = LocalToWorld.TransformPosition(NewPoint.Center);
+						NewPoint.Left         = LocalToWorld.TransformPosition(NewPoint.Left);
+						NewPoint.Right        = LocalToWorld.TransformPosition(NewPoint.Right);
+						NewPoint.FalloffLeft  = LocalToWorld.TransformPosition(NewPoint.FalloffLeft);
+						NewPoint.FalloffRight = LocalToWorld.TransformPosition(NewPoint.FalloffRight);
+
+						// Draw lines from the last keypoint.
+						PDI->SetHitProxy(Segment.HitProxy);
+
+						// center line
+						PDI->DrawLine(OldPoint.Center, NewPoint.Center, SegmentColor, GetDepthPriorityGroup(View), 0.0f, DepthBias);
+
+						// draw sides
+						PDI->DrawLine(OldPoint.Left, NewPoint.Left, SegmentColor, GetDepthPriorityGroup(View), 0.0f, DepthBias);
+						PDI->DrawLine(OldPoint.Right, NewPoint.Right, SegmentColor, GetDepthPriorityGroup(View), 0.0f, DepthBias);
+
+						PDI->SetHitProxy(NULL);
+
+						// draw falloff sides
+						if (bDrawFalloff)
+						{
+							DrawDashedLine(PDI, OldPoint.FalloffLeft, NewPoint.FalloffLeft, SegmentColor, 100, GetDepthPriorityGroup(View), DepthBias);
+							DrawDashedLine(PDI, OldPoint.FalloffRight, NewPoint.FalloffRight, SegmentColor, 100, GetDepthPriorityGroup(View), DepthBias);
+						}
+
+						OldPoint = NewPoint;
+					}
+				}
+
+				for (const FControlPointProxy& ControlPoint : ControlPoints)
+				{
+					const float ControlPointSpriteScale = LocalToWorld.GetScaleVector().X * ControlPoint.SpriteScale;
+					const FVector ControlPointLocation = LocalToWorld.TransformPosition(ControlPoint.Location) + FVector(0, 0, ControlPointSpriteScale * 0.75f);
+
+					// Draw Sprite
+
+					const FLinearColor ControlPointSpriteColor = ControlPoint.bSelected ? SelectedControlPointSpriteColor : FLinearColor::White;
+
+					PDI->SetHitProxy(ControlPoint.HitProxy);
+
+					PDI->DrawSprite(
+						ControlPointLocation,
+						ControlPointSpriteScale,
+						ControlPointSpriteScale,
+						ControlPointSprite->Resource,
+						ControlPointSpriteColor,
+						GetDepthPriorityGroup(View),
+						0, ControlPointSprite->Resource->GetSizeX(),
+						0, ControlPointSprite->Resource->GetSizeY(),
+						SE_BLEND_Masked);
+
+
+					// Draw Lines
+					const FLinearColor ControlPointColor = ControlPoint.bSelected ? SelectedSplineColor : SplineColor;
+
+					if (ControlPoint.Points.Num() == 1)
+					{
+						FLandscapeSplineInterpPoint NewPoint = ControlPoint.Points[0];
+						NewPoint.Center = LocalToWorld.TransformPosition(NewPoint.Center);
+						NewPoint.Left   = LocalToWorld.TransformPosition(NewPoint.Left);
+						NewPoint.Right  = LocalToWorld.TransformPosition(NewPoint.Right);
+						NewPoint.FalloffLeft  = LocalToWorld.TransformPosition(NewPoint.FalloffLeft);
+						NewPoint.FalloffRight = LocalToWorld.TransformPosition(NewPoint.FalloffRight);
+
+						// draw end for spline connection
+						PDI->DrawPoint(NewPoint.Center, ControlPointColor, 6.0f, GetDepthPriorityGroup(View));
+						PDI->DrawLine(NewPoint.Left, NewPoint.Center, ControlPointColor, GetDepthPriorityGroup(View), 0.0f, DepthBias);
+						PDI->DrawLine(NewPoint.Right, NewPoint.Center, ControlPointColor, GetDepthPriorityGroup(View), 0.0f, DepthBias);
+						if (bDrawFalloff)
+						{
+							DrawDashedLine(PDI, NewPoint.FalloffLeft, NewPoint.Left, ControlPointColor, 100, GetDepthPriorityGroup(View), DepthBias);
+							DrawDashedLine(PDI, NewPoint.FalloffRight, NewPoint.Right, ControlPointColor, 100, GetDepthPriorityGroup(View), DepthBias);
+						}
+					}
+					else if (ControlPoint.Points.Num() >= 2)
+					{
+						FLandscapeSplineInterpPoint OldPoint = ControlPoint.Points.Last();
+						//OldPoint.Left   = LocalToWorld.TransformPosition(OldPoint.Left);
+						OldPoint.Right  = LocalToWorld.TransformPosition(OldPoint.Right);
+						//OldPoint.FalloffLeft  = LocalToWorld.TransformPosition(OldPoint.FalloffLeft);
+						OldPoint.FalloffRight = LocalToWorld.TransformPosition(OldPoint.FalloffRight);
+
+						for (const FLandscapeSplineInterpPoint& Point : ControlPoint.Points)
+						{
+							FLandscapeSplineInterpPoint NewPoint = Point;
+							NewPoint.Center = LocalToWorld.TransformPosition(NewPoint.Center);
+							NewPoint.Left   = LocalToWorld.TransformPosition(NewPoint.Left);
+							NewPoint.Right  = LocalToWorld.TransformPosition(NewPoint.Right);
+							NewPoint.FalloffLeft  = LocalToWorld.TransformPosition(NewPoint.FalloffLeft);
+							NewPoint.FalloffRight = LocalToWorld.TransformPosition(NewPoint.FalloffRight);
+
+							PDI->SetHitProxy(ControlPoint.HitProxy);
+
+							// center line
+							PDI->DrawLine(ControlPointLocation, NewPoint.Center, ControlPointColor, GetDepthPriorityGroup(View), 0.0f, DepthBias);
+
+							// draw sides
+							PDI->DrawLine(OldPoint.Right, NewPoint.Left, ControlPointColor, GetDepthPriorityGroup(View), 0.0f, DepthBias);
+
+							PDI->SetHitProxy(NULL);
+
+							// draw falloff sides
+							if (bDrawFalloff)
+							{
+								DrawDashedLine(PDI, OldPoint.FalloffRight, NewPoint.FalloffLeft, ControlPointColor, 100, GetDepthPriorityGroup(View), DepthBias);
+							}
+
+							// draw end for spline connection
+							PDI->DrawPoint(NewPoint.Center, ControlPointColor, 6.0f, GetDepthPriorityGroup(View));
+							PDI->DrawLine(NewPoint.Left, NewPoint.Center, ControlPointColor, GetDepthPriorityGroup(View), 0.0f, DepthBias);
+							PDI->DrawLine(NewPoint.Right, NewPoint.Center, ControlPointColor, GetDepthPriorityGroup(View), 0.0f, DepthBias);
+							if (bDrawFalloff)
+							{
+								DrawDashedLine(PDI, NewPoint.FalloffLeft, NewPoint.Left, ControlPointColor, 100, GetDepthPriorityGroup(View), DepthBias);
+								DrawDashedLine(PDI, NewPoint.FalloffRight, NewPoint.Right, ControlPointColor, 100, GetDepthPriorityGroup(View), DepthBias);
+							}
+
+							//OldPoint = NewPoint;
+							OldPoint.Right = NewPoint.Right;
+							OldPoint.FalloffRight = NewPoint.FalloffRight;
+						}
+					}
+				}
+
+				PDI->SetHitProxy(NULL);
+			}
+		}
+	}
+
 	virtual void DrawDynamicElements(FPrimitiveDrawInterface* PDI,const FSceneView* View)
 	{
 		// Slight Depth Bias so that the splines show up when they exactly match the target surface
diff --git a/Engine/Source/Runtime/Engine/Private/Level.cpp b/Engine/Source/Runtime/Engine/Private/Level.cpp
--- a/Engine/Source/Runtime/Engine/Private/Level.cpp
+++ b/Engine/Source/Runtime/Engine/Private/Level.cpp
@@ -1982,76 +1982,126 @@ class FLineBatcherSceneProxy : public FPrimitiveSceneProxy
 {
 public:
 	FLineBatcherSceneProxy(const ULineBatchComponent* InComponent) :
-	  FPrimitiveSceneProxy(InComponent), Lines(InComponent->BatchedLines), 
-		  Points(InComponent->BatchedPoints), Meshes(InComponent->BatchedMeshes)
-	  {
-		  bWillEverBeLit = false;
-		  ViewRelevance.bDrawRelevance = true;
-		  ViewRelevance.bDynamicRelevance = true;
-		  ViewRelevance.bNormalTranslucencyRelevance = true;
-	  }
-
-	  /** 
-	  * Draw the scene proxy as a dynamic element
-	  *
-	  * @param	PDI - draw interface to render to
-	  * @param	View - current view
-	  */
-	  virtual void DrawDynamicElements(FPrimitiveDrawInterface* PDI,const FSceneView* View) override
-	  {
-		  QUICK_SCOPE_CYCLE_COUNTER( STAT_LineBatcherSceneProxy_DrawDynamicElements );
-
-		  for (int32 i = 0; i < Lines.Num(); i++)
-		  {
-			  PDI->DrawLine(Lines[i].Start, Lines[i].End, Lines[i].Color, Lines[i].DepthPriority, Lines[i].Thickness);
-		  }
-
-		  for (int32 i = 0; i < Points.Num(); i++)
-		  {
-			  PDI->DrawPoint(Points[i].Position, Points[i].Color, Points[i].PointSize, Points[i].DepthPriority);
-		  }
-
-		  for (int32 i = 0; i < Meshes.Num(); i++)
-		  {
-			  static FVector const PosX(1.f,0,0);
-			  static FVector const PosY(0,1.f,0);
-			  static FVector const PosZ(0,0,1.f);
-
-			  FBatchedMesh const& M = Meshes[i];
-
-			  // this seems far from optimal in terms of perf, but it's for debugging
-			  FDynamicMeshBuilder MeshBuilder;
-
-			  // set up geometry
-			  for (int32 VertIdx=0; VertIdx < M.MeshVerts.Num(); ++VertIdx)
-			  {
-				  MeshBuilder.AddVertex( M.MeshVerts[VertIdx], FVector2D::ZeroVector, PosX, PosY, PosZ, FColor::White );
-			  }
-			  //MeshBuilder.AddTriangles(M.MeshIndices);
-			  for (int32 Idx=0; Idx < M.MeshIndices.Num(); Idx+=3)
-			  {
-				  MeshBuilder.AddTriangle( M.MeshIndices[Idx], M.MeshIndices[Idx+1], M.MeshIndices[Idx+2] );
-			  }
-
-			  FMaterialRenderProxy* const MaterialRenderProxy = new(FMemStack::Get()) FColoredMaterialRenderProxy(GEngine->DebugMeshMaterial->GetRenderProxy(false), M.Color);
-			  MeshBuilder.Draw(PDI, FMatrix::Identity, MaterialRenderProxy, M.DepthPriority);
-		  }
-	  }
-
-	  /**
-	  *  Returns a struct that describes to the renderer when to draw this proxy.
-	  *	@param		Scene view to use to determine our relevence.
-	  *  @return		View relevance struct
-	  */
-	  virtual FPrimitiveViewRelevance GetViewRelevance(const FSceneView* View)
-	  {
-		  return ViewRelevance;
-	  }
-	  virtual uint32 GetMemoryFootprint( void ) const { return( sizeof( *this ) + GetAllocatedSize() ); }
-	  uint32 GetAllocatedSize( void ) const 
-	  { 
-		  return( FPrimitiveSceneProxy::GetAllocatedSize() + Lines.GetAllocatedSize() + Points.GetAllocatedSize() + Meshes.GetAllocatedSize() ); 
-	  }
+		FPrimitiveSceneProxy(InComponent), Lines(InComponent->BatchedLines), 
+		Points(InComponent->BatchedPoints), Meshes(InComponent->BatchedMeshes)
+	{
+		bWillEverBeLit = false;
+		ViewRelevance.bDrawRelevance = true;
+		ViewRelevance.bDynamicRelevance = true;
+		ViewRelevance.bNormalTranslucencyRelevance = true;
+	}
+
+	virtual void GetDynamicMeshElements(const TArray<const FSceneView*>& Views, const FSceneViewFamily& ViewFamily, uint32 VisibilityMap, FMeshElementCollector& Collector) const override
+	{
+		QUICK_SCOPE_CYCLE_COUNTER( STAT_LineBatcherSceneProxy_GetDynamicMeshElements );
+
+		for (int32 ViewIndex = 0; ViewIndex < Views.Num(); ViewIndex++)
+		{
+			if (VisibilityMap & (1 << ViewIndex))
+			{
+				const FSceneView* View = Views[ViewIndex];
+				FPrimitiveDrawInterface* PDI = Collector.GetPDI(ViewIndex);
+
+				for (int32 i = 0; i < Lines.Num(); i++)
+				{
+					PDI->DrawLine(Lines[i].Start, Lines[i].End, Lines[i].Color, Lines[i].DepthPriority, Lines[i].Thickness);
+				}
+
+				for (int32 i = 0; i < Points.Num(); i++)
+				{
+					PDI->DrawPoint(Points[i].Position, Points[i].Color, Points[i].PointSize, Points[i].DepthPriority);
+				}
+
+				for (int32 i = 0; i < Meshes.Num(); i++)
+				{
+					static FVector const PosX(1.f,0,0);
+					static FVector const PosY(0,1.f,0);
+					static FVector const PosZ(0,0,1.f);
+
+					FBatchedMesh const& M = Meshes[i];
+
+					// this seems far from optimal in terms of perf, but it's for debugging
+					FDynamicMeshBuilder MeshBuilder;
+
+					// set up geometry
+					for (int32 VertIdx=0; VertIdx < M.MeshVerts.Num(); ++VertIdx)
+					{
+						MeshBuilder.AddVertex( M.MeshVerts[VertIdx], FVector2D::ZeroVector, PosX, PosY, PosZ, FColor::White );
+					}
+					//MeshBuilder.AddTriangles(M.MeshIndices);
+					for (int32 Idx=0; Idx < M.MeshIndices.Num(); Idx+=3)
+					{
+						MeshBuilder.AddTriangle( M.MeshIndices[Idx], M.MeshIndices[Idx+1], M.MeshIndices[Idx+2] );
+					}
+
+					FMaterialRenderProxy* const MaterialRenderProxy = new(FMemStack::Get()) FColoredMaterialRenderProxy(GEngine->DebugMeshMaterial->GetRenderProxy(false), M.Color);
+					MeshBuilder.GetMesh(FMatrix::Identity, MaterialRenderProxy, M.DepthPriority, false, false, ViewIndex, Collector);
+				}
+			}
+		}
+	}
+
+	/** 
+	* Draw the scene proxy as a dynamic element
+	*
+	* @param	PDI - draw interface to render to
+	* @param	View - current view
+	*/
+	virtual void DrawDynamicElements(FPrimitiveDrawInterface* PDI,const FSceneView* View) override
+	{
+		QUICK_SCOPE_CYCLE_COUNTER( STAT_LineBatcherSceneProxy_DrawDynamicElements );
+
+		for (int32 i = 0; i < Lines.Num(); i++)
+		{
+			PDI->DrawLine(Lines[i].Start, Lines[i].End, Lines[i].Color, Lines[i].DepthPriority, Lines[i].Thickness);
+		}
+
+		for (int32 i = 0; i < Points.Num(); i++)
+		{
+			PDI->DrawPoint(Points[i].Position, Points[i].Color, Points[i].PointSize, Points[i].DepthPriority);
+		}
+
+		for (int32 i = 0; i < Meshes.Num(); i++)
+		{
+			static FVector const PosX(1.f,0,0);
+			static FVector const PosY(0,1.f,0);
+			static FVector const PosZ(0,0,1.f);
+
+			FBatchedMesh const& M = Meshes[i];
+
+			// this seems far from optimal in terms of perf, but it's for debugging
+			FDynamicMeshBuilder MeshBuilder;
+
+			// set up geometry
+			for (int32 VertIdx=0; VertIdx < M.MeshVerts.Num(); ++VertIdx)
+			{
+				MeshBuilder.AddVertex( M.MeshVerts[VertIdx], FVector2D::ZeroVector, PosX, PosY, PosZ, FColor::White );
+			}
+			//MeshBuilder.AddTriangles(M.MeshIndices);
+			for (int32 Idx=0; Idx < M.MeshIndices.Num(); Idx+=3)
+			{
+				MeshBuilder.AddTriangle( M.MeshIndices[Idx], M.MeshIndices[Idx+1], M.MeshIndices[Idx+2] );
+			}
+
+			FMaterialRenderProxy* const MaterialRenderProxy = new(FMemStack::Get()) FColoredMaterialRenderProxy(GEngine->DebugMeshMaterial->GetRenderProxy(false), M.Color);
+			MeshBuilder.Draw(PDI, FMatrix::Identity, MaterialRenderProxy, M.DepthPriority);
+		}
+	}
+
+	/**
+	*  Returns a struct that describes to the renderer when to draw this proxy.
+	*	@param		Scene view to use to determine our relevence.
+	*  @return		View relevance struct
+	*/
+	virtual FPrimitiveViewRelevance GetViewRelevance(const FSceneView* View)
+	{
+		return ViewRelevance;
+	}
+	virtual uint32 GetMemoryFootprint( void ) const { return( sizeof( *this ) + GetAllocatedSize() ); }
+	uint32 GetAllocatedSize( void ) const 
+	{ 
+		return( FPrimitiveSceneProxy::GetAllocatedSize() + Lines.GetAllocatedSize() + Points.GetAllocatedSize() + Meshes.GetAllocatedSize() ); 
+	}
 
 private:
 	TArray<FBatchedLine> Lines;
diff --git a/Engine/Source/Runtime/Engine/Private/MaterialBillboardComponent.cpp b/Engine/Source/Runtime/Engine/Private/MaterialBillboardComponent.cpp
--- a/Engine/Source/Runtime/Engine/Private/MaterialBillboardComponent.cpp
+++ b/Engine/Source/Runtime/Engine/Private/MaterialBillboardComponent.cpp
@@ -73,6 +73,12 @@ public:
 };
 static TGlobalResource<FMaterialSpriteVertexFactory> GMaterialSpriteVertexFactory;
 
+class FMaterialSpriteVertexArray : public FOneFrameResource
+{
+public:
+	TArray<FMaterialSpriteVertex, TInlineAllocator<4> > Vertices;
+};
+
 /** Represents a sprite to the scene manager. */
 class FMaterialSpriteSceneProxy : public FPrimitiveSceneProxy
 {
@@ -83,8 +89,6 @@ public:
 	: FPrimitiveSceneProxy(InComponent)
 	, Elements(InComponent->Elements)
 	, BaseColor(FLinearColor::White)
-	, LevelColor(FLinearColor::White)
-	, PropertyColor(255,255,255)
 	{
 		AActor* Owner = InComponent->GetOwner();
 		if (Owner)
@@ -108,7 +112,121 @@ public:
 			}
 		}
 
-		GEngine->GetPropertyColorationColor( (UObject*)InComponent, PropertyColor );
+		FColor NewPropertyColor;
+		GEngine->GetPropertyColorationColor( (UObject*)InComponent, NewPropertyColor );
+		PropertyColor = NewPropertyColor;
+	}
+
+	virtual void GetDynamicMeshElements(const TArray<const FSceneView*>& Views, const FSceneViewFamily& ViewFamily, uint32 VisibilityMap, FMeshElementCollector& Collector) const override
+	{
+		QUICK_SCOPE_CYCLE_COUNTER( STAT_MaterialSpriteSceneProxy_GetDynamicMeshElements );
+
+		for (int32 ViewIndex = 0; ViewIndex < Views.Num(); ViewIndex++)
+		{
+			if (VisibilityMap & (1 << ViewIndex))
+			{
+				const FSceneView* View = Views[ViewIndex];
+
+				const bool bIsWireframe = View->Family->EngineShowFlags.Wireframe;
+				// Determine the position of the source
+				const FVector SourcePosition = GetLocalToWorld().GetOrigin();
+				const FVector CameraToSource = View->ViewMatrices.ViewOrigin - SourcePosition;
+				const float DistanceToSource = CameraToSource.Size();
+
+				const FVector CameraUp      = -View->InvViewProjectionMatrix.TransformVector(FVector(1.0f,0.0f,0.0f));
+				const FVector CameraRight   = -View->InvViewProjectionMatrix.TransformVector(FVector(0.0f,1.0f,0.0f));
+				const FVector CameraForward = -View->InvViewProjectionMatrix.TransformVector(FVector(0.0f,0.0f,1.0f));
+				const FMatrix WorldToLocal = GetLocalToWorld().Inverse();
+				const FVector LocalCameraUp = WorldToLocal.TransformVector(CameraUp);
+				const FVector LocalCameraRight = WorldToLocal.TransformVector(CameraRight);
+				const FVector LocalCameraForward = WorldToLocal.TransformVector(CameraForward);
+
+				// Draw the elements ordered so the last is on top of the first.
+				for(int32 ElementIndex = 0;ElementIndex < Elements.Num();++ElementIndex)
+				{
+					const FMaterialSpriteElement& Element = Elements[ElementIndex];
+
+					if(Element.Material)
+					{
+						// Evaluate the size of the sprite.
+						float SizeX = Element.BaseSizeX;
+						float SizeY = Element.BaseSizeY;
+						if(Element.DistanceToSizeCurve)
+						{
+							const float SizeFactor = Element.DistanceToSizeCurve->GetFloatValue(DistanceToSource);
+							SizeX *= SizeFactor;
+							SizeY *= SizeFactor;
+						}
+						
+						// Convert the size into world-space.
+						const float W = View->ViewProjectionMatrix.TransformPosition(SourcePosition).W;
+						const float AspectRatio = CameraRight.Size() / CameraUp.Size();
+						const float WorldSizeX = Element.bSizeIsInScreenSpace ? (SizeX               * W) : (SizeX / CameraRight.Size());
+						const float WorldSizeY = Element.bSizeIsInScreenSpace ? (SizeY * AspectRatio * W) : (SizeY / CameraUp.Size());
+			
+						// Evaluate the color/opacity of the sprite.
+						FLinearColor Color = BaseColor;
+						if(Element.DistanceToOpacityCurve)
+						{
+							Color.A *= Element.DistanceToOpacityCurve->GetFloatValue(DistanceToSource);
+						}
+
+						// Set up the sprite vertex attributes that are constant across the sprite.
+						FMaterialSpriteVertexArray& VertexArray = Collector.AllocateOneFrameResource<FMaterialSpriteVertexArray>();
+						VertexArray.Vertices.Empty(4);
+						VertexArray.Vertices.AddUninitialized(4);
+
+						for(uint32 VertexIndex = 0;VertexIndex < 4;++VertexIndex)
+						{
+							VertexArray.Vertices[VertexIndex].Color = Color;
+							VertexArray.Vertices[VertexIndex].TangentX = FPackedNormal(LocalCameraRight.SafeNormal());
+							VertexArray.Vertices[VertexIndex].TangentZ = FPackedNormal(-LocalCameraForward.SafeNormal());
+						}
+
+						// Set up the sprite vertex positions and texture coordinates.
+						VertexArray.Vertices[0].Position  = -WorldSizeX * LocalCameraRight + +WorldSizeY * LocalCameraUp;
+						VertexArray.Vertices[0].TexCoords = FVector2D(0,0);
+						VertexArray.Vertices[1].Position  = +WorldSizeX * LocalCameraRight + +WorldSizeY * LocalCameraUp;
+						VertexArray.Vertices[1].TexCoords = FVector2D(1,0);
+						VertexArray.Vertices[2].Position  = -WorldSizeX * LocalCameraRight + -WorldSizeY * LocalCameraUp;
+						VertexArray.Vertices[2].TexCoords = FVector2D(0,1);
+						VertexArray.Vertices[3].Position  = +WorldSizeX * LocalCameraRight + -WorldSizeY * LocalCameraUp;
+						VertexArray.Vertices[3].TexCoords = FVector2D(1,1);
+			
+						// Set up the FMeshElement.
+						FMeshBatch& Mesh = Collector.AllocateMesh();
+						Mesh.UseDynamicData      = true;
+						Mesh.DynamicVertexData   = VertexArray.Vertices.GetData();
+						Mesh.DynamicVertexStride = sizeof(FMaterialSpriteVertex);
+
+						Mesh.VertexFactory           = &GMaterialSpriteVertexFactory;
+						Mesh.MaterialRenderProxy     = Element.Material->GetRenderProxy((View->Family->EngineShowFlags.Selection) && IsSelected(),IsHovered());
+						Mesh.LCI                     = NULL;
+						Mesh.ReverseCulling          = IsLocalToWorldDeterminantNegative() ? true : false;
+						Mesh.CastShadow              = false;
+						Mesh.DepthPriorityGroup      = (ESceneDepthPriorityGroup)GetDepthPriorityGroup(View);
+						Mesh.Type                    = PT_TriangleStrip;
+						Mesh.bDisableBackfaceCulling = true;
+
+						// Set up the FMeshBatchElement.
+						FMeshBatchElement& BatchElement = Mesh.Elements[0];
+						BatchElement.IndexBuffer            = NULL;
+						BatchElement.DynamicIndexData       = NULL;
+						BatchElement.DynamicIndexStride     = 0;
+						BatchElement.FirstIndex             = 0;
+						BatchElement.MinVertexIndex         = 0;
+						BatchElement.MaxVertexIndex         = 3;
+						BatchElement.PrimitiveUniformBuffer = GetUniformBuffer();
+						BatchElement.NumPrimitives          = 2;
+
+						Mesh.bCanApplyViewModeOverrides = true;
+						Mesh.bUseWireframeSelectionColoring = IsSelected();
+
+						Collector.AddMesh(ViewIndex, Mesh);
+					}
+				}
+			}
+		}
 	}
 
 	// FPrimitiveSceneProxy interface.
@@ -218,6 +336,7 @@ public:
 			}
 		}
 	}
+
 	virtual FPrimitiveViewRelevance GetViewRelevance(const FSceneView* View)
 	{
 		bool bVisible = View->Family->EngineShowFlags.BillboardSprites;
@@ -236,8 +355,6 @@ private:
 	TArray<FMaterialSpriteElement> Elements;
 	FMaterialRelevance MaterialRelevance;
 	FColor BaseColor;
-	FLinearColor LevelColor;
-	FColor PropertyColor;
 };
 
 UMaterialBillboardComponent::UMaterialBillboardComponent(const class FPostConstructInitializeProperties& PCIP)
diff --git a/Engine/Source/Runtime/Engine/Private/MeshParticleVertexFactory.cpp b/Engine/Source/Runtime/Engine/Private/MeshParticleVertexFactory.cpp
--- a/Engine/Source/Runtime/Engine/Private/MeshParticleVertexFactory.cpp
+++ b/Engine/Source/Runtime/Engine/Private/MeshParticleVertexFactory.cpp
@@ -45,7 +45,7 @@ public:
 
 		if (!bInstanced)
 		{
-			FMeshParticleVertexFactory::FBatchParametersCPU* BatchParameters = (FMeshParticleVertexFactory::FBatchParametersCPU*)BatchElement.UserData;
+			const FMeshParticleVertexFactory::FBatchParametersCPU* BatchParameters = (const FMeshParticleVertexFactory::FBatchParametersCPU*)BatchElement.UserData;
 			const FMeshParticleInstanceVertex* Vertex = BatchParameters->InstanceBuffer + BatchElement.UserIndex;
 			const FMeshParticleInstanceVertexDynamicParameter* DynamicVertex = BatchParameters->DynamicParameterBuffer + BatchElement.UserIndex;
 
diff --git a/Engine/Source/Runtime/Engine/Private/ModelRender.cpp b/Engine/Source/Runtime/Engine/Private/ModelRender.cpp
--- a/Engine/Source/Runtime/Engine/Private/ModelRender.cpp
+++ b/Engine/Source/Runtime/Engine/Private/ModelRender.cpp
@@ -150,8 +150,6 @@ public:
 	FModelSceneProxy(UModelComponent* InComponent):
 		FPrimitiveSceneProxy(InComponent),
 		Component(InComponent),
-		LevelColor(FLinearColor::White),
-		PropertyColor(255,255,255),
 		CollisionResponse(InComponent->GetCollisionResponseToChannels())
 #if WITH_EDITOR
 		, CollisionMaterialInstance(GEngine->ShadedLevelColorationUnlitMaterial->GetRenderProxy(false, false),FColor(157,149,223,255)	)
@@ -181,7 +179,9 @@ public:
 		}
 
 		// Get a color for property coloration.
-		GEngine->GetPropertyColorationColor( (UObject*)InComponent, PropertyColor );
+		FColor NewPropertyColor;
+		GEngine->GetPropertyColorationColor( (UObject*)InComponent, NewPropertyColor );
+		PropertyColor = NewPropertyColor;
 	}
 
 	bool IsCollisionView(const FSceneView* View, bool & bDrawCollision) const
@@ -208,6 +208,200 @@ public:
 		return ModelHitProxy;
 	}
 
+	virtual void GetDynamicMeshElements(const TArray<const FSceneView*>& Views, const FSceneViewFamily& ViewFamily, uint32 VisibilityMap, FMeshElementCollector& Collector) const override
+	{
+		QUICK_SCOPE_CYCLE_COUNTER(STAT_FModelSceneProxy_GetMeshElements);
+		bool bAnySelectedSurfs = false;
+
+		for (int32 ViewIndex = 0; ViewIndex < Views.Num(); ViewIndex++)
+		{
+			if (VisibilityMap & (1 << ViewIndex))
+			{
+				const FSceneView* View = Views[ViewIndex];
+
+				bool bShowSelection = GIsEditor && !View->bIsGameView && ViewFamily.EngineShowFlags.Selection;
+				bool bDynamicBSPTriangles = bShowSelection || IsRichView(ViewFamily);
+				bool bShowBSPTriangles = ViewFamily.EngineShowFlags.BSPTriangles;
+				bool bShowBSP = ViewFamily.EngineShowFlags.BSP;
+
+#if WITH_EDITOR
+				bool bDrawCollision = false;
+				const bool bInCollisionView = IsCollisionView(View, bDrawCollision);
+				// draw bsp as dynamic when in collision view mode
+				if(bInCollisionView)
+				{
+					bDynamicBSPTriangles = true;
+				}
+#endif
+
+				// If in a collision view, only draw if we have collision
+				if (bDynamicBSPTriangles && bShowBSPTriangles && bShowBSP 
+#if WITH_EDITOR
+					&& (!bInCollisionView || bDrawCollision)
+#endif
+					)
+				{
+					ESceneDepthPriorityGroup DepthPriorityGroup = (ESceneDepthPriorityGroup)GetDepthPriorityGroup(View);
+
+					const FMaterialRenderProxy* MatProxyOverride = NULL;
+
+#if WITH_EDITOR
+					if(bInCollisionView && AllowDebugViewmodes())
+					{
+						MatProxyOverride = &CollisionMaterialInstance;
+					}
+#endif
+
+					// If selection is being shown, batch triangles based on whether they are selected or not.
+					if (bShowSelection)
+					{
+						uint32 TotalIndices = 0;
+
+						for (int32 ElementIndex = 0; ElementIndex < Elements.Num(); ElementIndex++)
+						{
+							const FModelElement& ModelElement = Component->GetElements()[ElementIndex];
+							TotalIndices += ModelElement.NumTriangles * 3;
+						}
+
+						if (TotalIndices > 0)
+						{
+							FGlobalDynamicIndexBuffer::FAllocation IndexAllocation = FGlobalDynamicIndexBuffer::Get().Allocate<uint32>(TotalIndices);
+
+							if (IndexAllocation.IsValid())
+							{
+								uint32* Indices = (uint32*)IndexAllocation.Buffer;
+								uint32 FirstIndex = IndexAllocation.FirstIndex;
+
+								for (int32 ElementIndex = 0; ElementIndex < Elements.Num(); ElementIndex++)
+								{
+									const FModelElement& ModelElement = Component->GetElements()[ElementIndex];
+
+									if (ModelElement.NumTriangles > 0)
+									{
+										const FElementInfo& ProxyElementInfo = Elements[ElementIndex];
+										bool bHasSelectedSurfs = false;
+										bool bHasHoveredSurfs = false;
+
+										for (uint32 BatchIndex = 0; BatchIndex < 3; BatchIndex++)
+										{
+											// Three batches total:
+											//		Batch 0: Only surfaces that are neither selected, nor hovered
+											//		Batch 1: Only selected surfaces
+											//		Batch 2: Only hovered surfaces
+											const bool bOnlySelectedSurfaces = ( BatchIndex == 1 );
+											const bool bOnlyHoveredSurfaces = ( BatchIndex == 2 );
+
+											if (bOnlySelectedSurfaces && !bHasSelectedSurfs)
+											{
+												continue;
+											}
+
+											if (bOnlyHoveredSurfaces && !bHasHoveredSurfs)
+											{
+												continue;
+											}
+
+											uint32 MinVertexIndex = MAX_uint32;
+											uint32 MaxVertexIndex = 0;
+											uint32 NumIndices = 0;
+
+											for(int32 NodeIndex = 0;NodeIndex < ModelElement.Nodes.Num();NodeIndex++)
+											{
+												FBspNode& Node = Component->GetModel()->Nodes[ModelElement.Nodes[NodeIndex]];
+												FBspSurf& Surf = Component->GetModel()->Surfs[Node.iSurf];
+
+												if (!ShouldDrawSurface(Surf))
+												{
+													continue;
+												}
+
+												const bool bSurfaceSelected = (Surf.PolyFlags & PF_Selected) == PF_Selected;
+												const bool bSurfaceHovered = !bSurfaceSelected && ((Surf.PolyFlags & PF_Hovered) == PF_Hovered);
+												bHasSelectedSurfs |= bSurfaceSelected;
+												bHasHoveredSurfs |= bSurfaceHovered;
+
+												if (bSurfaceSelected == bOnlySelectedSurfaces && bSurfaceHovered == bOnlyHoveredSurfaces)
+												{
+													for (uint32 BackFace = 0; BackFace < (uint32)((Surf.PolyFlags & PF_TwoSided) ? 2 : 1); BackFace++)
+													{
+														for (int32 VertexIndex = 2; VertexIndex < Node.NumVertices; VertexIndex++)
+														{
+															*Indices++ = Node.iVertexIndex + Node.NumVertices * BackFace;
+															*Indices++ = Node.iVertexIndex + Node.NumVertices * BackFace + VertexIndex;
+															*Indices++ = Node.iVertexIndex + Node.NumVertices * BackFace + VertexIndex - 1;
+															NumIndices += 3;
+														}
+														MinVertexIndex = FMath::Min(Node.iVertexIndex + Node.NumVertices * BackFace,MinVertexIndex);
+														MaxVertexIndex = FMath::Max(Node.iVertexIndex + Node.NumVertices * BackFace + Node.NumVertices - 1,MaxVertexIndex);
+													}
+												}
+											}
+
+											if (NumIndices > 0)
+											{
+												FMeshBatch& MeshElement = Collector.AllocateMesh();
+												FMeshBatchElement& BatchElement = MeshElement.Elements[0];
+												BatchElement.IndexBuffer = IndexAllocation.IndexBuffer;
+												MeshElement.VertexFactory = &Component->GetModel()->VertexFactory;
+												MeshElement.MaterialRenderProxy = (MatProxyOverride != NULL) ? MatProxyOverride : ProxyElementInfo.GetMaterial()->GetRenderProxy(bOnlySelectedSurfaces, bOnlyHoveredSurfaces);
+												MeshElement.LCI = &ProxyElementInfo;
+												BatchElement.PrimitiveUniformBufferResource = &GetUniformBuffer();
+												BatchElement.FirstIndex = FirstIndex;
+												BatchElement.NumPrimitives = NumIndices / 3;
+												BatchElement.MinVertexIndex = MinVertexIndex;
+												BatchElement.MaxVertexIndex = MaxVertexIndex;
+												MeshElement.Type = PT_TriangleList;
+												MeshElement.DepthPriorityGroup = DepthPriorityGroup;
+												MeshElement.bCanApplyViewModeOverrides = true;
+												MeshElement.bUseWireframeSelectionColoring = false;
+												MeshElement.bUseSelectionOutline = bOnlySelectedSurfaces;
+												Collector.AddMesh(ViewIndex, MeshElement);
+												FirstIndex += NumIndices;
+											}
+										}
+
+										bAnySelectedSurfs |= bHasSelectedSurfs;
+									}
+								}
+							}
+						}
+					}
+					else
+					{
+						for (int32 ElementIndex = 0; ElementIndex < Elements.Num(); ElementIndex++)
+						{
+							const FModelElement& ModelElement = Component->GetElements()[ElementIndex];
+
+							if (ModelElement.NumTriangles > 0)
+							{
+								FMeshBatch& MeshElement = Collector.AllocateMesh();
+								FMeshBatchElement& BatchElement = MeshElement.Elements[0];
+								BatchElement.IndexBuffer = ModelElement.IndexBuffer;
+								MeshElement.VertexFactory = &Component->GetModel()->VertexFactory;
+								MeshElement.MaterialRenderProxy = (MatProxyOverride != NULL) ? MatProxyOverride : Elements[ElementIndex].GetMaterial()->GetRenderProxy(false);
+								MeshElement.LCI = &Elements[ElementIndex];
+								BatchElement.PrimitiveUniformBufferResource = &GetUniformBuffer();
+								BatchElement.FirstIndex = ModelElement.FirstIndex;
+								BatchElement.NumPrimitives = ModelElement.NumTriangles;
+								BatchElement.MinVertexIndex = ModelElement.MinVertexIndex;
+								BatchElement.MaxVertexIndex = ModelElement.MaxVertexIndex;
+								MeshElement.Type = PT_TriangleList;
+								MeshElement.DepthPriorityGroup = DepthPriorityGroup;
+								MeshElement.bCanApplyViewModeOverrides = true;
+								MeshElement.bUseWireframeSelectionColoring = false;
+								Collector.AddMesh(ViewIndex, MeshElement);
+							}
+						}
+					}
+				}
+			}
+		}
+
+		//@todo parallelrendering - remove this legacy state modification
+		// Poly selected state is modified in many places, so it's hard to push the selection state to the proxy
+		const_cast<FModelSceneProxy*>(this)->SetSelection_RenderThread(bAnySelectedSurfs);
+	}
+
 	virtual void PreRenderView(const FSceneViewFamily* ViewFamily, const uint32 VisibilityMap, int32 FrameNumber) override
 	{
 		// Reset any batches leftover from last frame.
@@ -225,7 +419,7 @@ public:
 			const FSceneView* View = ViewFamily->Views[ViewIndex];
 			bool bVisibleInView = (VisibilityMap & ViewBit) == ViewBit;
 			bool bShowSelection = GIsEditor && !View->bIsGameView && View->Family->EngineShowFlags.Selection;
-			bool bDynamicBSPTriangles = bShowSelection || IsRichView(View) ;
+			bool bDynamicBSPTriangles = bShowSelection || IsRichView(*View->Family) ;
 			bool bShowBSPTriangles = View->Family->EngineShowFlags.BSPTriangles;
 			bool bShowBSP = View->Family->EngineShowFlags.BSP;
 
@@ -418,15 +612,6 @@ public:
 				FDynamicModelMeshBatch& Batch = DynamicMeshBatches[BatchIndex];
 				if (!bRenderingSelectionOutline || Batch.bIsSelectedBatch)
 				{
-					if (Batch.ModelElementIndex < ElementLightMapResolutions.Num())
-					{
-						SetLightMapResolutionScale(ElementLightMapResolutions[Batch.ModelElementIndex]);
-					}
-					else
-					{
-						SetLightMapResolutionScale(FVector2D(0.0f, 0.0f));
-					}
-
 					DrawRichMesh(PDI,Batch,FLinearColor::White,UtilColor,PropertyColor,this,false, bIsWireframe);
 				}
 			}
@@ -472,7 +657,7 @@ public:
 		Result.bDrawRelevance = IsShown(View) && View->Family->EngineShowFlags.BSPTriangles && View->Family->EngineShowFlags.BSP;
 		bool bShowSelectedTriangles = GIsEditor && !View->bIsGameView && View->Family->EngineShowFlags.Selection;
 		bool bCollisionView = View->Family->EngineShowFlags.CollisionPawn || View->Family->EngineShowFlags.CollisionVisibility;
-		if (IsRichView(View) || HasViewDependentDPG() || bCollisionView
+		if (IsRichView(*View->Family) || HasViewDependentDPG() || bCollisionView
 			|| (bShowSelectedTriangles && HasSelectedSurfaces()))
 		{
 			Result.bDynamicRelevance = true;
@@ -661,10 +846,6 @@ private:
 	};
 
 	TArray<FElementInfo> Elements;
-	TArray<FVector2D> ElementLightMapResolutions;
-
-	FLinearColor LevelColor;
-	FColor PropertyColor;
 
 	FMaterialRelevance MaterialRelevance;
 
@@ -729,36 +910,12 @@ public:
 		return NULL;
 	}
 
-	/**
-	 *	Clear the element LightMap resolutions array.
-	 */
-	void ClearElementLightMapResolutions()
-	{
-		ElementLightMapResolutions.Empty();
-	}
-
-	void AddElementLightMapResolution(int32 InElementIdx, int32 InSizeX, int32 InSizeY)
-	{
-		if (ElementLightMapResolutions.Num() <= InElementIdx)
-		{
-			ElementLightMapResolutions.AddZeroed(InElementIdx - ElementLightMapResolutions.Num() + 1);
-		}
-
-		ElementLightMapResolutions[InElementIdx] = FVector2D((float)InSizeX, (float)InSizeY);
-	}
-
 	friend class UModelComponent;
 };
 
 FPrimitiveSceneProxy* UModelComponent::CreateSceneProxy()
 {
 	FPrimitiveSceneProxy* Proxy = ::new FModelSceneProxy(this);
-#if WITH_EDITOR
-	if (GIsEditor && Proxy)
-	{
-		SetupLightmapResolutionViewInfo(*Proxy);
-	}
-#endif
 	return Proxy;
 }
 
@@ -787,57 +944,3 @@ FBoxSphereBounds UModelComponent::CalcBounds(const FTransform & LocalToWorld) co
 		return FBoxSphereBounds(LocalToWorld.GetLocation(), FVector::ZeroVector, 0.f);
 	}
 }
-
-#if WITH_EDITOR
-bool UModelComponent::SetupLightmapResolutionViewInfo(FPrimitiveSceneProxy& Proxy) const
-{
-	FModelSceneProxy* ModelProxy = (FModelSceneProxy*)(&Proxy);
-
-	// Alway texture based...
-	ModelProxy->SetLightMapType(LMIT_Texture);
-	ModelProxy->SetIsLightMapResolutionPadded(true);
-
-	// Fill in the resolutions array
-	ModelProxy->ClearElementLightMapResolutions();
-
-	if (Model->NodeGroups.Num() > 0)
-	{
-		for (int32 ElementIdx = 0; ElementIdx < ModelProxy->GetElementCount(); ElementIdx++)
-		{
-			const FModelSceneProxy::FElementInfo* Element = ModelProxy->GetElement(ElementIdx);
-			if (Element)
-			{
-				const FModelElement* ModelElement = Element->GetModelElement();
-				if (ModelElement && (ModelElement->Nodes.Num() > 0))
-				{
-					// All the nodes in this element SHOULD be in the same node group...
-					int32 NodeIdx = ModelElement->Nodes[0];
-
-					// Find the node group it belong to
-					FNodeGroup* FoundNodeGroup = NULL;
-					// find the NodeGroup that this node went into, and get all of its node
-					for (TMap<int32, FNodeGroup*>::TIterator It(Model->NodeGroups); It && (FoundNodeGroup == NULL); ++It)
-					{
-						FNodeGroup* NodeGroup = It.Value();
-						check(NodeGroup != NULL);
-						for (int32 NodeIndex = 0; NodeIndex < NodeGroup->Nodes.Num(); NodeIndex++)
-						{
-							if (NodeGroup->Nodes[NodeIndex] == NodeIdx)
-							{
-								FoundNodeGroup = NodeGroup;
-								break;
-							}
-						}
-					}
-
-					ModelProxy->AddElementLightMapResolution(ElementIdx, 
-						FoundNodeGroup ? FoundNodeGroup->SizeX : 0, 
-						FoundNodeGroup ? FoundNodeGroup->SizeY : 0);
-				}
-			}
-		}
-	}
-
-	return true;
-}
-#endif // WITH_EDITOR
diff --git a/Engine/Source/Runtime/Engine/Private/NavMeshRenderingHelpers.h b/Engine/Source/Runtime/Engine/Private/NavMeshRenderingHelpers.h
--- a/Engine/Source/Runtime/Engine/Private/NavMeshRenderingHelpers.h
+++ b/Engine/Source/Runtime/Engine/Private/NavMeshRenderingHelpers.h
@@ -293,7 +293,7 @@ public:
 			}
 		}
 
-		const bool bNeedToSwitchVerticalAxis = IsES2Platform(GRHIShaderPlatform) && !IsPCPlatform(GRHIShaderPlatform);
+		const bool bNeedToSwitchVerticalAxis = RHINeedsToSwitchVerticalAxis(GRHIShaderPlatform);
 		const FTexture2DRHIRef DepthTexture;
 			
 		FRHICommandListImmediate& RHICmdList = FRHICommandListExecutor::GetImmediateCommandList();
diff --git a/Engine/Source/Runtime/Engine/Private/NiagaraComponent.cpp b/Engine/Source/Runtime/Engine/Private/NiagaraComponent.cpp
--- a/Engine/Source/Runtime/Engine/Private/NiagaraComponent.cpp
+++ b/Engine/Source/Runtime/Engine/Private/NiagaraComponent.cpp
@@ -26,6 +26,18 @@ struct FNiagaraDynamicData
 	TArray<FParticleSpriteVertex> VertexData;
 };
 
+class FNiagaraMeshCollectorResources : public FOneFrameResource
+{
+public:
+	FParticleSpriteVertexFactory VertexFactory;
+	FParticleSpriteUniformBufferRef UniformBuffer;
+
+	virtual ~FNiagaraMeshCollectorResources()
+	{
+		VertexFactory.ReleaseResource();
+	}
+};
+
 /**
  * Scene proxy for drawing niagara particle simulations.
  */
@@ -102,6 +114,100 @@ private:
 		WorldSpacePrimitiveUniformBuffer.ReleaseResource();
 	}
 
+	virtual void GetDynamicMeshElements(const TArray<const FSceneView*>& Views, const FSceneViewFamily& ViewFamily, uint32 VisibilityMap, FMeshElementCollector& Collector) const override
+	{
+		SCOPE_CYCLE_COUNTER(STAT_NiagaraRender);
+
+		check(DynamicData && DynamicData->VertexData.Num());
+
+		const bool bIsWireframe = ViewFamily.EngineShowFlags.Wireframe;
+		FMaterialRenderProxy* MaterialRenderProxy = Material->GetRenderProxy(IsSelected(),IsHovered());
+
+		int32 SizeInBytes = DynamicData->VertexData.GetTypeSize() * DynamicData->VertexData.Num();
+		FGlobalDynamicVertexBuffer::FAllocation LocalDynamicVertexAllocation = FGlobalDynamicVertexBuffer::Get().Allocate(SizeInBytes);
+
+		if (LocalDynamicVertexAllocation.IsValid())
+		{
+			// Update the primitive uniform buffer if needed.
+			if (!WorldSpacePrimitiveUniformBuffer.IsInitialized())
+			{
+				FPrimitiveUniformShaderParameters PrimitiveUniformShaderParameters = GetPrimitiveUniformShaderParameters(
+					FMatrix::Identity,
+					GetActorPosition(),
+					GetBounds(),
+					GetLocalBounds(),
+					ReceivesDecals(),
+					false,
+					UseEditorDepthTest()
+					);
+				WorldSpacePrimitiveUniformBuffer.SetContents(PrimitiveUniformShaderParameters);
+				WorldSpacePrimitiveUniformBuffer.InitResource();
+			}
+
+			// Copy the vertex data over.
+			FMemory::Memcpy(LocalDynamicVertexAllocation.Buffer, DynamicData->VertexData.GetData(), SizeInBytes);
+
+			// Compute the per-view uniform buffers.
+			for (int32 ViewIndex = 0; ViewIndex < Views.Num(); ViewIndex++)
+			{
+				if (VisibilityMap & (1 << ViewIndex))
+				{
+					const FSceneView* View = Views[ViewIndex];
+
+					FNiagaraMeshCollectorResources& CollectorResources = Collector.AllocateOneFrameResource<FNiagaraMeshCollectorResources>();
+					FParticleSpriteUniformParameters PerViewUniformParameters = UniformParameters;
+
+					// Collector.AllocateOneFrameResource uses default ctor, initialize the vertex factory
+					CollectorResources.VertexFactory.SetFeatureLevel(ViewFamily.GetFeatureLevel());
+					CollectorResources.VertexFactory.SetParticleFactoryType(PVFT_Sprite);
+
+					PerViewUniformParameters.MacroUVParameters = FVector4(0.0f,0.0f,1.0f,1.0f);
+					CollectorResources.UniformBuffer = FParticleSpriteUniformBufferRef::CreateUniformBufferImmediate(PerViewUniformParameters, UniformBuffer_SingleFrame);
+
+					CollectorResources.VertexFactory.InitResource();
+					CollectorResources.VertexFactory.SetSpriteUniformBuffer(CollectorResources.UniformBuffer);
+					CollectorResources.VertexFactory.SetInstanceBuffer(
+						LocalDynamicVertexAllocation.VertexBuffer,
+						LocalDynamicVertexAllocation.VertexOffset,
+						sizeof(FParticleSpriteVertex),
+						true
+						);
+					CollectorResources.VertexFactory.SetDynamicParameterBuffer(NULL,0,0,true);
+
+					FMeshBatch& MeshBatch = Collector.AllocateMesh();
+					MeshBatch.VertexFactory = &CollectorResources.VertexFactory;
+					MeshBatch.CastShadow = CastsDynamicShadow();
+					MeshBatch.bUseAsOccluder = false;
+					MeshBatch.ReverseCulling = IsLocalToWorldDeterminantNegative();
+					MeshBatch.Type = PT_TriangleList;
+					MeshBatch.DepthPriorityGroup = GetDepthPriorityGroup(View);
+					MeshBatch.bCanApplyViewModeOverrides = true;
+					MeshBatch.bUseWireframeSelectionColoring = IsSelected();
+
+					if (bIsWireframe)
+					{
+						MeshBatch.MaterialRenderProxy = UMaterial::GetDefaultMaterial(MD_Surface)->GetRenderProxy(IsSelected(),IsHovered());
+					}
+					else
+					{
+						MeshBatch.MaterialRenderProxy = MaterialRenderProxy;
+					}
+
+					FMeshBatchElement& MeshElement = MeshBatch.Elements[0];
+					MeshElement.IndexBuffer = &GParticleIndexBuffer;
+					MeshElement.FirstIndex = 0;
+					MeshElement.NumPrimitives = 2;
+					MeshElement.NumInstances = DynamicData->VertexData.Num();
+					MeshElement.MinVertexIndex = 0;
+					MeshElement.MaxVertexIndex = MeshElement.NumInstances * 4 - 1;
+					MeshElement.PrimitiveUniformBufferResource = &WorldSpacePrimitiveUniformBuffer;
+
+					Collector.AddMesh(ViewIndex, MeshBatch);
+				}
+			}
+		}
+	}
+
 	virtual void PreRenderView(const FSceneViewFamily* ViewFamily, const uint32 VisibilityMap, int32 FrameNumber) override
 	{
 		SCOPE_CYCLE_COUNTER(STAT_NiagaraPreRenderView);
@@ -144,7 +250,8 @@ private:
 					GetBounds(),
 					GetLocalBounds(),
 					ReceivesDecals(),
-					false
+					false,
+					UseEditorDepthTest()
 					);
 				WorldSpacePrimitiveUniformBuffer.SetContents(PrimitiveUniformShaderParameters);
 				WorldSpacePrimitiveUniformBuffer.InitResource();
@@ -258,7 +365,7 @@ private:
 private:
 	UMaterialInterface* Material;
 	FNiagaraDynamicData* DynamicData;
-	TUniformBuffer<FPrimitiveUniformShaderParameters> WorldSpacePrimitiveUniformBuffer;
+	mutable TUniformBuffer<FPrimitiveUniformShaderParameters> WorldSpacePrimitiveUniformBuffer;
 	FParticleSpriteVertexFactory VertexFactory;
 	FParticleSpriteUniformParameters UniformParameters;
 	TArray<FParticleSpriteUniformBufferRef, TInlineAllocator<2> > PerViewUniformBuffers;
diff --git a/Engine/Source/Runtime/Engine/Private/Particles/ParticleBeam2EmitterInstance.cpp b/Engine/Source/Runtime/Engine/Private/Particles/ParticleBeam2EmitterInstance.cpp
--- a/Engine/Source/Runtime/Engine/Private/Particles/ParticleBeam2EmitterInstance.cpp
+++ b/Engine/Source/Runtime/Engine/Private/Particles/ParticleBeam2EmitterInstance.cpp
@@ -1479,6 +1479,7 @@ bool FParticleBeam2EmitterInstance::FillReplayData( FDynamicEmitterReplayDataBas
 		NewReplayData->bSmoothNoise_Enabled = false;
 	}
 	NewReplayData->Sheets = (BeamTypeData->Sheets > 0) ? BeamTypeData->Sheets : 1;
+	NewReplayData->Sheets = FMath::Max(NewReplayData->Sheets, 1);
 
 	NewReplayData->TextureTile = BeamTypeData->TextureTile;
 	NewReplayData->TextureTileDistance = BeamTypeData->TextureTileDistance;
@@ -1498,6 +1499,7 @@ bool FParticleBeam2EmitterInstance::FillReplayData( FDynamicEmitterReplayDataBas
 	else
 	{
 		NewReplayData->Frequency			= (BeamModule_Noise->Frequency > 0) ? BeamModule_Noise->Frequency : 1;
+		NewReplayData->Frequency = FMath::Max(NewReplayData->Frequency, 1);
 		NewReplayData->NoiseTessellation	= (BeamModule_Noise->NoiseTessellation > 0) ? BeamModule_Noise->NoiseTessellation : 1;
 		NewReplayData->NoiseTangentStrength= BeamModule_Noise->NoiseTangentStrength.GetValue(EmitterTime);
 		if (BeamModule_Noise->bNRScaleEmitterTime)
diff --git a/Engine/Source/Runtime/Engine/Private/Particles/ParticleGpuSimulation.cpp b/Engine/Source/Runtime/Engine/Private/Particles/ParticleGpuSimulation.cpp
--- a/Engine/Source/Runtime/Engine/Private/Particles/ParticleGpuSimulation.cpp
+++ b/Engine/Source/Runtime/Engine/Private/Particles/ParticleGpuSimulation.cpp
@@ -2445,6 +2445,17 @@ public:
 	}
 };
 
+class FGPUSpriteCollectorResources : public FOneFrameResource
+{
+public:
+	FGPUSpriteVertexFactory VertexFactory;
+
+	~FGPUSpriteCollectorResources()
+	{
+		VertexFactory.ReleaseResource();
+	}
+};
+
 /**
  * Dynamic emitter data for Cascade.
  */
@@ -2511,7 +2522,7 @@ struct FGPUSpriteDynamicEmitterData : FDynamicEmitterDataBase
 	/**
 	 * Called to create render thread resources.
 	 */
-	virtual void CreateRenderThreadResources(const FParticleSystemSceneProxy* InOwnerProxy)
+	virtual void UpdateRenderThreadResourcesEmitter(const FParticleSystemSceneProxy* InOwnerProxy)
 	{
 		check(Simulation);
 
@@ -2635,10 +2646,124 @@ struct FGPUSpriteDynamicEmitterData : FDynamicEmitterDataBase
 		}
 	}
 
+	virtual void GetDynamicMeshElementsEmitter(const FParticleSystemSceneProxy* Proxy, const FSceneView* View, const FSceneViewFamily& ViewFamily, int32 ViewIndex, FMeshElementCollector& Collector) const override
+	{
+		auto FeatureLevel = ViewFamily.GetFeatureLevel();
+
+		if (RHISupportsGPUParticles(FeatureLevel))
+		{
+			SCOPE_CYCLE_COUNTER(STAT_GPUSpritePreRenderTime);
+
+			check(Simulation);
+
+			// Do not render orphaned emitters. This can happen if the emitter
+			// instance has been destroyed but we are rendering before the
+			// scene proxy has received the update to clear dynamic data.
+			if (Simulation->SimulationIndex != INDEX_NONE
+				&& Simulation->VertexBuffer.ParticleCount > 0)
+			{
+				FGPUSpriteEmitterDynamicUniformBufferRef LocalDynamicUniformBuffer;
+				// Create view agnostic render data.  Do here rather than in CreateRenderThreadResources because in some cases Render can be called before CreateRenderThreadResources
+				{
+					// Create per-emitter uniform buffer for dynamic parameters
+					LocalDynamicUniformBuffer = FGPUSpriteEmitterDynamicUniformBufferRef::CreateUniformBufferImmediate(EmitterDynamicParameters, UniformBuffer_SingleFrame);
+				}
+
+				if (bUseLocalSpace == false)
+				{
+					Proxy->UpdateWorldSpacePrimitiveUniformBuffer();
+				}
+
+				const bool bTranslucent = RendersWithTranslucentMaterial();
+				const bool bAllowSorting = FXConsoleVariables::bAllowGPUSorting
+					&& FeatureLevel == ERHIFeatureLevel::SM5
+					&& bTranslucent;
+
+				// Iterate over views and assign parameters for each.
+				FParticleSimulationResources* SimulationResources = FXSystem->GetParticleSimulationResources();
+				FGPUSpriteCollectorResources& CollectorResources = Collector.AllocateOneFrameResource<FGPUSpriteCollectorResources>();
+				CollectorResources.VertexFactory.SetFeatureLevel(FeatureLevel);
+				CollectorResources.VertexFactory.InitResource();
+				FGPUSpriteVertexFactory& VertexFactory = CollectorResources.VertexFactory;
+
+				if (bAllowSorting && SortMode == PSORTMODE_DistanceToView)
+				{
+					// Extensibility TODO: This call to AddSortedGPUSimulation is very awkward. When rendering a frame we need to
+					// accumulate all GPU particle emitters that need to be sorted. That is so they can be sorted in one big radix
+					// sort for efficiency. Ideally that state is per-scene renderer but the renderer doesn't know anything about particles.
+					const int32 SortedBufferOffset = FXSystem->AddSortedGPUSimulation(Simulation, View->ViewMatrices.ViewOrigin);
+					check(SimulationResources->SortedVertexBuffer.IsInitialized());
+					VertexFactory.SetVertexBuffer(&SimulationResources->SortedVertexBuffer, SortedBufferOffset);
+				}
+				else
+				{
+					check(Simulation->VertexBuffer.IsInitialized());
+					VertexFactory.SetVertexBuffer(&Simulation->VertexBuffer, 0);
+				}
+
+				const int32 ParticleCount = Simulation->VertexBuffer.ParticleCount;
+				const bool bIsWireframe = ViewFamily.EngineShowFlags.Wireframe;
+
+				{
+					SCOPE_CYCLE_COUNTER(STAT_GPUSpriteRenderingTime);
+
+					FParticleSimulationResources* ParticleSimulationResources = FXSystem->GetParticleSimulationResources();
+					FParticleStateTextures& StateTextures = ParticleSimulationResources->GetCurrentStateTextures();
+							
+					VertexFactory.EmitterUniformBuffer = Resources->UniformBuffer;
+					VertexFactory.EmitterDynamicUniformBuffer = LocalDynamicUniformBuffer;
+					VertexFactory.PositionTextureRHI = StateTextures.PositionTextureRHI;
+					VertexFactory.PositionZWTextureRHI = StateTextures.PositionZWTextureRHI;
+					VertexFactory.VelocityTextureRHI = StateTextures.VelocityTextureRHI;
+					VertexFactory.AttributesTextureRHI = ParticleSimulationResources->RenderAttributesTexture.TextureRHI;
+
+					FMeshBatch& Mesh = Collector.AllocateMesh();
+					FMeshBatchElement& BatchElement = Mesh.Elements[0];
+					BatchElement.IndexBuffer = &GParticleIndexBuffer;
+					BatchElement.NumPrimitives = MAX_PARTICLES_PER_INSTANCE * 2;
+					BatchElement.NumInstances = ParticleCount / MAX_PARTICLES_PER_INSTANCE;
+					BatchElement.FirstIndex = 0;
+					Mesh.VertexFactory = &VertexFactory;
+					Mesh.LCI = NULL;
+					if ( bUseLocalSpace )
+					{
+						BatchElement.PrimitiveUniformBufferResource = &Proxy->GetUniformBuffer();
+					}
+					else
+					{
+						BatchElement.PrimitiveUniformBufferResource = &Proxy->GetWorldSpacePrimitiveUniformBuffer();
+					}
+					BatchElement.MinVertexIndex = 0;
+					BatchElement.MaxVertexIndex = 3;
+					Mesh.ReverseCulling = Proxy->IsLocalToWorldDeterminantNegative();
+					Mesh.CastShadow = Proxy->GetCastShadow();
+					Mesh.DepthPriorityGroup = (ESceneDepthPriorityGroup)Proxy->GetDepthPriorityGroup(View);
+					const bool bUseSelectedMaterial = GIsEditor && (ViewFamily.EngineShowFlags.Selection) ? bSelected : 0;
+					Mesh.MaterialRenderProxy = Material->GetRenderProxy(bUseSelectedMaterial);
+					Mesh.Type = PT_TriangleList;
+					Mesh.bCanApplyViewModeOverrides = true;
+					Mesh.bUseWireframeSelectionColoring = Proxy->IsSelected();
+
+					Collector.AddMesh(ViewIndex, Mesh);
+				}
+
+				const bool bHaveLocalVectorField = Simulation && Simulation->LocalVectorField.Resource;
+				if (bHaveLocalVectorField && ViewFamily.EngineShowFlags.VectorFields)
+				{
+					// Create a vertex factory for visualization if needed.
+					Simulation->CreateVectorFieldVisualizationVertexFactory();
+					check(Simulation->VectorFieldVisualizationVertexFactory);
+					DrawVectorFieldBounds(Collector.GetPDI(ViewIndex), View, &Simulation->LocalVectorField);
+					GetVectorFieldMesh(Simulation->VectorFieldVisualizationVertexFactory, &Simulation->LocalVectorField, ViewIndex, Collector);
+				}
+			}
+		}
+	}
+
 	/**
 	 * Called to render the sprites for this emitter.
 	 */
-	virtual int32 Render(FParticleSystemSceneProxy* Proxy, FPrimitiveDrawInterface* PDI,const FSceneView* View)
+	virtual int32 RenderEmitter(FParticleSystemSceneProxy* Proxy, FPrimitiveDrawInterface* PDI,const FSceneView* View)
 	{
 		if (RHISupportsGPUParticles(View->GetFeatureLevel()))
 		{
