FPrimitiveSceneProxy::PreRenderView / DrawDynamicElements refactored into GetDynamicMeshElements

diff --git a/Engine/Source/Runtime/Engine/Private/Particles/ParticleSystemRender.cpp b/Engine/Source/Runtime/Engine/Private/Particles/ParticleSystemRender.cpp
--- a/Engine/Source/Runtime/Engine/Private/Particles/ParticleSystemRender.cpp
+++ b/Engine/Source/Runtime/Engine/Private/Particles/ParticleSystemRender.cpp
@@ -254,16 +254,34 @@ void ENGINE_API FinishOneFrameParticleStats()
 }
 #endif
 
-#define LOG_DETAILED_PARTICLE_RENDER_STATS 0
-
-#if LOG_DETAILED_PARTICLE_RENDER_STATS 
-/** Global detailed update stats. */
-static FDetailedTickStats GDetailedParticleRenderStats( 20, 10, 1, 4, TEXT("rendering") );
-#define TRACK_DETAILED_PARTICLE_RENDER_STATS(Object) FScopedDetailTickStats DetailedTickStats(GDetailedParticleRenderStats,Object);
-#else
-#define TRACK_DETAILED_PARTICLE_RENDER_STATS(Object)
-#endif
-
+/** 
+ *	Create a vertex factory for the given type.
+ *
+ *	@param	InType						The type of vertex factory to create.
+ *
+ *	@return	FParticleVertexFactoryBase*	The created VF; NULL if invalid InType
+ */
+FParticleVertexFactoryBase* CreateParticleVertexFactory(EParticleVertexFactoryType InType, ERHIFeatureLevel::Type InFeatureLevel)
+{
+	FParticleVertexFactoryBase* NewVertexFactory = NULL;
+	switch (InType)
+	{
+	case PVFT_Sprite:
+		NewVertexFactory = new FParticleSpriteVertexFactory(PVFT_Sprite, InFeatureLevel);
+		break;
+	case PVFT_BeamTrail:
+		NewVertexFactory = new FParticleBeamTrailVertexFactory(PVFT_BeamTrail, InFeatureLevel);
+		break;
+	case PVFT_Mesh:
+		NewVertexFactory = new FMeshParticleVertexFactory(PVFT_Mesh, InFeatureLevel);
+		break;
+	default:
+		break;
+	}
+	check(NewVertexFactory);
+	NewVertexFactory->InitResource();
+	return NewVertexFactory;
+}
 
 ///////////////////////////////////////////////////////////////////////////////
 FParticleOrderPool GParticleOrderPool;
@@ -420,35 +438,6 @@ void FParticleVertexFactoryPool::DumpInfo(FOutputDevice& Ar)
 }
 #endif
 
-/** 
- *	Create a vertex factory for the given type.
- *
- *	@param	InType						The type of vertex factory to create.
- *
- *	@return	FParticleVertexFactoryBase*	The created VF; NULL if invalid InType
- */
-FParticleVertexFactoryBase* FParticleVertexFactoryPool::CreateParticleVertexFactory(EParticleVertexFactoryType InType, ERHIFeatureLevel::Type InFeatureLevel)
-{
-	FParticleVertexFactoryBase* NewVertexFactory = NULL;
-	switch (InType)
-	{
-	case PVFT_Sprite:
-		NewVertexFactory = new FParticleSpriteVertexFactory(PVFT_Sprite, InFeatureLevel);
-		break;
-	case PVFT_BeamTrail:
-		NewVertexFactory = new FParticleBeamTrailVertexFactory(PVFT_BeamTrail, InFeatureLevel);
-		break;
-	case PVFT_Mesh:
-		NewVertexFactory = new FMeshParticleVertexFactory(PVFT_Mesh, InFeatureLevel);
-		break;
-	default:
-		break;
-	}
-	check(NewVertexFactory);
-	NewVertexFactory->InitResource();
-	return NewVertexFactory;
-}
-
 void ParticleVertexFactoryPool_FreePool_RenderingThread()
 {
 	GParticleVertexFactoryPool.FreePool();
@@ -569,18 +558,9 @@ FORCEINLINE FVector GetCameraOffsetFromPayload(
 	}
 }
 
-/**  
- * Simple function to pass the async buffer fill task off to the owning emitter
-*/
-void FAsyncParticleFill::DoTask(ENamedThreads::Type CurrentThread, const FGraphEventRef& MyCompletionGraphEvent)
-{
-	// Test integrity of async particle fill: FPlatformProcess::Sleep(.01);
-	Parent->DoBufferFill();
-}
-
 void FDynamicSpriteEmitterDataBase::SortSpriteParticles(int32 SortMode, bool bLocalSpace, 
 	int32 ParticleCount, const TArray<uint8>& ParticleData, int32 ParticleStride, const TArray<uint16>& ParticleIndices,
-	const FSceneView* View, const FMatrix& LocalToWorld, FParticleOrder* ParticleOrder)
+	const FSceneView* View, const FMatrix& LocalToWorld, FParticleOrder* ParticleOrder) const
 {
 	SCOPE_CYCLE_COUNTER(STAT_SortingTime);
 
@@ -656,7 +636,7 @@ void FDynamicSpriteEmitterDataBase::SortSpriteParticles(int32 SortMode, bool bLo
 	}
 }
 
-void FDynamicSpriteEmitterDataBase::RenderDebug(FParticleSystemSceneProxy* Proxy, FPrimitiveDrawInterface* PDI, const FSceneView* View, bool bCrosses)
+void FDynamicSpriteEmitterDataBase::RenderDebug(const FParticleSystemSceneProxy* Proxy, FPrimitiveDrawInterface* PDI, const FSceneView* View, bool bCrosses) const
 {
 	check(Proxy);
 
@@ -689,120 +669,43 @@ void FDynamicSpriteEmitterDataBase::RenderDebug(FParticleSystemSceneProxy* Proxy
 	}
 }
 
-/**
- *	Set up an buffer for async filling
- *
- *	@param	Proxy					The primitive scene proxy for the emitter.
- *	@param	InBufferIndex			Index of this buffer
- *	@param	InView					View for this buffer
- *	@param	InVertexCount			Count of verts for this buffer
- *	@param	InVertexSize			Stride of these verts, only used for verification
- *	@param	InDynamicParameterVertexStride	Stride of the dynamic parameter
- */
-void FDynamicSpriteEmitterDataBase::BuildViewFillData(FParticleSystemSceneProxy* Proxy, int32 InBufferIndex,const FSceneView *InView,int32 InVertexCount,int32 InVertexSize,int32 InDynamicParameterVertexStride)
+void FDynamicSpriteEmitterDataBase::BuildViewFillData(
+	const FParticleSystemSceneProxy* Proxy, 
+	const FSceneView* InView, 
+	int32 InVertexCount, 
+	int32 InVertexSize, 
+	int32 InDynamicParameterVertexStride, 
+	FGlobalDynamicVertexBuffer::FAllocation& DynamicVertexAllocation,
+	FGlobalDynamicIndexBuffer::FAllocation& DynamicIndexAllocation,
+	FGlobalDynamicVertexBuffer::FAllocation* DynamicParameterAllocation,
+	FAsyncBufferFillData& Data) const
 {
-	bool bSuccess = true;
-
-	FAsyncBufferFillData Data;
-
 	Data.LocalToWorld = Proxy->GetLocalToWorld();
 	Data.WorldToLocal = Proxy->GetWorldToLocal();
 	Data.View = InView;
 	check(Data.VertexSize == 0 || Data.VertexSize == InVertexSize);
 
-	check(InBufferIndex == InstanceDataAllocations.Num())
-	FGlobalDynamicVertexBuffer::FAllocation* DynamicVertexAllocation = new(InstanceDataAllocations) FGlobalDynamicVertexBuffer::FAllocation();
-	*DynamicVertexAllocation = FGlobalDynamicVertexBuffer::Get().Allocate( InVertexCount * InVertexSize );
-	Data.VertexData = DynamicVertexAllocation->Buffer;
+	DynamicVertexAllocation = FGlobalDynamicVertexBuffer::Get().Allocate( InVertexCount * InVertexSize );
+	Data.VertexData = DynamicVertexAllocation.Buffer;
 	Data.VertexCount = InVertexCount;
 	Data.VertexSize = InVertexSize;
 
-	bSuccess &= DynamicVertexAllocation->IsValid();
-
 	int32 NumIndices, IndexStride;
 	GetIndexAllocInfo(NumIndices, IndexStride);
 	check((uint32)NumIndices <= 65535);
 	check(IndexStride > 0);
 
-	FGlobalDynamicIndexBuffer::FAllocation* DynamicIndexAllocation = new(IndexDataAllocations) FGlobalDynamicIndexBuffer::FAllocation();
-	*DynamicIndexAllocation = FGlobalDynamicIndexBuffer::Get().Allocate( NumIndices, IndexStride );
-	Data.IndexData = DynamicIndexAllocation->Buffer;
+	DynamicIndexAllocation = FGlobalDynamicIndexBuffer::Get().Allocate( NumIndices, IndexStride );
+	Data.IndexData = DynamicIndexAllocation.Buffer;
 	Data.IndexCount = NumIndices;
 
-	bSuccess &= DynamicIndexAllocation->IsValid();
-
 	Data.DynamicParameterData = NULL;
+
 	if( bUsesDynamicParameter )
 	{
 		check( InDynamicParameterVertexStride > 0 );
-		check(InBufferIndex == DynamicParameterDataAllocations.Num())
-		FGlobalDynamicVertexBuffer::FAllocation* DynamicParameterAllocation = new(DynamicParameterDataAllocations) FGlobalDynamicVertexBuffer::FAllocation();
 		*DynamicParameterAllocation = FGlobalDynamicVertexBuffer::Get().Allocate( InVertexCount * InDynamicParameterVertexStride );
 		Data.DynamicParameterData = DynamicParameterAllocation->Buffer;
-
-		bSuccess &= DynamicParameterAllocation->IsValid();
-	}
-
-	//Create the data but only fill it if all the allocations succeeded.
-	if (InBufferIndex >= AsyncBufferFillTasks.Num())
-	{
-		new (AsyncBufferFillTasks) FAsyncBufferFillData();
-	}
-	check(InBufferIndex < AsyncBufferFillTasks.Num()); // please add the views in order
-
-	if( bSuccess )
-	{
-		AsyncBufferFillTasks[InBufferIndex] = Data;
-	}
-}
-
-/**
- *	Set up all buffers for async filling
- *
- *	@param	Proxy							The primitive scene proxy for the emitter.
- *	@param	ViewFamily						View family to process
- *	@param	VisibilityMap					Visibility map for the sub-views
- *	@param	bOnlyOneView					If true, then we don't need per-view buffers
- *	@param	InVertexCount					Count of verts for this buffer
- *	@param	InVertexSize					Stride of these verts, only used for verification
- *  @param  InDynamicParameterVertexSize	Stride of the dynamic parameter
- */
-void FDynamicSpriteEmitterDataBase::BuildViewFillDataAndSubmit(FParticleSystemSceneProxy* Proxy, const FSceneViewFamily* ViewFamily,const uint32 VisibilityMap,bool bOnlyOneView,int32 InVertexCount,int32 InVertexSize,int32 InDynamicParameterVertexSize)
-{
-	EnsureAsyncTaskComplete();
-	int32 NumUsedViews = 0;
-	InstanceDataAllocations.Empty();
-	IndexDataAllocations.Empty();
-	DynamicParameterDataAllocations.Empty();
-	for (int32 ViewIndex = 0; ViewIndex < ViewFamily->Views.Num(); ViewIndex++)
-	{
-		if (VisibilityMap & (1<<ViewIndex))
-		{
-			const FSceneView* View = ViewFamily->Views[ViewIndex];
-			BuildViewFillData(Proxy,NumUsedViews++,View,InVertexCount,InVertexSize, InDynamicParameterVertexSize);
-			if (bOnlyOneView)
-			{
-				break;
-			}
-		}
-	}
-	if (AsyncBufferFillTasks.Num() > NumUsedViews)
-	{
-		AsyncBufferFillTasks.RemoveAt(NumUsedViews,AsyncBufferFillTasks.Num() - NumUsedViews);
-	}
-	if (NumUsedViews)
-	{
-#if 0 // async particle fill is broken in UE4
-		if (!GIsEditor)
-		{
-			check(!AsyncTask.GetReference()); // this should not be still outstanding
-			AsyncTask = TGraphTask<FAsyncParticleFill>::CreateTask(NULL, ENamedThreads::RenderThread).ConstructAndDispatchWhenReady(this);
-		}
-		else
-#endif
-		{
-			DoBufferFill();
-		}
 	}
 }
 
@@ -873,7 +776,7 @@ void ApplyOrbitToPosition(
 	}
 }
 
-bool FDynamicSpriteEmitterData::GetVertexAndIndexData(void* VertexData, void* DynamicParameterVertexData, void* FillIndexData, FParticleOrder* ParticleOrder, const FVector& InCameraPosition, const FMatrix& InLocalToWorld)
+bool FDynamicSpriteEmitterData::GetVertexAndIndexData(void* VertexData, void* DynamicParameterVertexData, void* FillIndexData, FParticleOrder* ParticleOrder, const FVector& InCameraPosition, const FMatrix& InLocalToWorld) const
 {
 	SCOPE_CYCLE_COUNTER(STAT_ParticlePackingTime);
 	int32 ParticleCount = Source.ActiveParticleCount;
@@ -911,18 +814,18 @@ bool FDynamicSpriteEmitterData::GetVertexAndIndexData(void* VertexData, void* Dy
 	FVector ParticleOldPosition;
 	float SubImageIndex = 0.0f;
 
-	uint8* ParticleData = Source.ParticleData.GetData();
-	uint16* ParticleIndices = Source.ParticleIndices.GetData();
+	const uint8* ParticleData = Source.ParticleData.GetData();
+	const uint16* ParticleIndices = Source.ParticleIndices.GetData();
 	const FParticleOrder* OrderedIndices = ParticleOrder;
 
 	for (int32 i = 0; i < ParticleCount; i++)
 	{
 		ParticleIndex = OrderedIndices ? OrderedIndices[i].ParticleIndex : i;
-		DECLARE_PARTICLE(Particle, ParticleData + Source.ParticleStride * ParticleIndices[ParticleIndex]);
+		DECLARE_PARTICLE_CONST(Particle, ParticleData + Source.ParticleStride * ParticleIndices[ParticleIndex]);
 		if (i + 1 < ParticleCount)
 		{
 			int32 NextIndex = OrderedIndices ? OrderedIndices[i+1].ParticleIndex : (i + 1);
-			DECLARE_PARTICLE(NextParticle, ParticleData + Source.ParticleStride * ParticleIndices[NextIndex]);
+			DECLARE_PARTICLE_CONST(NextParticle, ParticleData + Source.ParticleStride * ParticleIndices[NextIndex]);
 			FPlatformMisc::Prefetch(&NextParticle);
 		}
 
@@ -984,7 +887,7 @@ bool FDynamicSpriteEmitterData::GetVertexAndIndexData(void* VertexData, void* Dy
 	return true;
 }
 
-bool FDynamicSpriteEmitterData::GetVertexAndIndexDataNonInstanced(void* VertexData, void* DynamicParameterVertexData, void* FillIndexData, FParticleOrder* ParticleOrder, const FVector& InCameraPosition, const FMatrix& InLocalToWorld)
+bool FDynamicSpriteEmitterData::GetVertexAndIndexDataNonInstanced(void* VertexData, void* DynamicParameterVertexData, void* FillIndexData, FParticleOrder* ParticleOrder, const FVector& InCameraPosition, const FMatrix& InLocalToWorld) const
 {
 	SCOPE_CYCLE_COUNTER(STAT_ParticlePackingTime);
 
@@ -1024,18 +927,18 @@ bool FDynamicSpriteEmitterData::GetVertexAndIndexDataNonInstanced(void* VertexDa
 	FVector ParticleOldPosition;
 	float SubImageIndex = 0.0f;
 
-	uint8* ParticleData = Source.ParticleData.GetData();
-	uint16* ParticleIndices = Source.ParticleIndices.GetData();
+	const uint8* ParticleData = Source.ParticleData.GetData();
+	const uint16* ParticleIndices = Source.ParticleIndices.GetData();
 	const FParticleOrder* OrderedIndices = ParticleOrder;
 
 	for (int32 i = 0; i < ParticleCount; i++)
 	{
 		ParticleIndex = OrderedIndices ? OrderedIndices[i].ParticleIndex : i;
-		DECLARE_PARTICLE(Particle, ParticleData + Source.ParticleStride * ParticleIndices[ParticleIndex]);
+		DECLARE_PARTICLE_CONST(Particle, ParticleData + Source.ParticleStride * ParticleIndices[ParticleIndex]);
 		if (i + 1 < ParticleCount)
 		{
 			int32 NextIndex = OrderedIndices ? OrderedIndices[i+1].ParticleIndex : (i + 1);
-			DECLARE_PARTICLE(NextParticle, ParticleData + Source.ParticleStride * ParticleIndices[NextIndex]);
+			DECLARE_PARTICLE_CONST(NextParticle, ParticleData + Source.ParticleStride * ParticleIndices[NextIndex]);
 			FPlatformMisc::Prefetch(&NextParticle);
 		}
 
@@ -1320,6 +1223,176 @@ void FDynamicSpriteEmitterData::GatherSimpleLights(const FParticleSystemScenePro
 	GatherParticleLightData(Source, Proxy->GetLocalToWorld(), ViewFamily, OutParticleLights);
 }
 
+class FDynamicSpriteCollectorResources : public FOneFrameResource
+{
+public:
+	FParticleSpriteVertexFactory VertexFactory;
+	FParticleSpriteUniformBufferRef UniformBuffer;
+
+	virtual ~FDynamicSpriteCollectorResources()
+	{
+		VertexFactory.ReleaseResource();
+	}
+};
+
+void FDynamicSpriteEmitterData::GetDynamicMeshElementsEmitter(const FParticleSystemSceneProxy* Proxy, const FSceneView* View, const FSceneViewFamily& ViewFamily, int32 ViewIndex, FMeshElementCollector& Collector) const
+{
+	SCOPE_CYCLE_COUNTER(STAT_SpriteRenderingTime);
+
+	const auto FeatureLevel = View->GetFeatureLevel();
+	const bool bInstanced = FeatureLevel >= ERHIFeatureLevel::SM3;
+
+	// Sort and generate particles for this view.
+	const FDynamicSpriteEmitterReplayDataBase* SourceData = GetSourceData();
+
+	if (bValid && SourceData)
+	{
+		if (SourceData->EmitterRenderMode == ERM_Normal)
+		{
+			// Determine how many vertices and indices are needed to render.
+			const int32 ParticleCount = SourceData->ActiveParticleCount;
+			const int32 VertexSize = GetDynamicVertexStride(FeatureLevel);
+			const int32 DynamicParameterVertexSize = sizeof(FParticleVertexDynamicParameter);
+			const int32 VertexPerParticle = bInstanced ? 1 : 4;
+
+			FDynamicSpriteCollectorResources& CollectorResources = Collector.AllocateOneFrameResource<FDynamicSpriteCollectorResources>();
+			CollectorResources.VertexFactory.SetParticleFactoryType(PVFT_Sprite);
+			CollectorResources.VertexFactory.SetFeatureLevel(FeatureLevel);
+			CollectorResources.VertexFactory.InitResource();
+
+			if (SourceData->bUseLocalSpace == false)
+			{
+				Proxy->UpdateWorldSpacePrimitiveUniformBuffer();
+			}
+
+			FGlobalDynamicVertexBuffer::FAllocation* Allocation = new FGlobalDynamicVertexBuffer::FAllocation();
+			FGlobalDynamicVertexBuffer::FAllocation* DynamicParameterAllocation = new FGlobalDynamicVertexBuffer::FAllocation();
+
+			// Allocate memory for render data.
+			*Allocation = FGlobalDynamicVertexBuffer::Get().Allocate( ParticleCount * VertexSize * VertexPerParticle );
+
+			if (bUsesDynamicParameter)
+			{
+				*DynamicParameterAllocation = FGlobalDynamicVertexBuffer::Get().Allocate( ParticleCount * DynamicParameterVertexSize * VertexPerParticle );
+			}
+
+			if ( Allocation->IsValid() && (!bUsesDynamicParameter || DynamicParameterAllocation->IsValid()) )
+			{
+				// Sort the particles if needed.
+				FParticleOrder* ParticleOrder = NULL;
+				if (SourceData->SortMode != PSORTMODE_None)
+				{
+					// If material is using unlit translucency and the blend mode is translucent then we need to sort (back to front)
+					const FMaterial* Material = MaterialResource[bSelected]->GetMaterial(FeatureLevel);
+
+					if (Material && 
+						(Material->GetBlendMode() == BLEND_Translucent ||
+						((SourceData->SortMode == PSORTMODE_Age_OldestFirst) || (SourceData->SortMode == PSORTMODE_Age_NewestFirst)))
+						)
+					{
+						ParticleOrder = GParticleOrderPool.GetParticleOrderData(ParticleCount);
+						SortSpriteParticles(SourceData->SortMode, SourceData->bUseLocalSpace, SourceData->ActiveParticleCount, 
+							SourceData->ParticleData, SourceData->ParticleStride, SourceData->ParticleIndices,
+							View, Proxy->GetLocalToWorld(), ParticleOrder);
+					}
+				}
+
+				// Fill vertex buffers.
+				if(bInstanced)
+				{
+					GetVertexAndIndexData(Allocation->Buffer, DynamicParameterAllocation->Buffer, NULL, ParticleOrder, View->ViewMatrices.ViewOrigin, Proxy->GetLocalToWorld());
+				}
+				else
+				{
+					GetVertexAndIndexDataNonInstanced(Allocation->Buffer, DynamicParameterAllocation->Buffer, NULL, ParticleOrder, View->ViewMatrices.ViewOrigin, Proxy->GetLocalToWorld());
+				}
+
+				// Create per-view uniform buffer.
+				FParticleSpriteUniformParameters PerViewUniformParameters = UniformParameters;
+				FVector2D ObjectNDCPosition;
+				FVector2D ObjectMacroUVScales;
+				Proxy->GetObjectPositionAndScale(*View,ObjectNDCPosition, ObjectMacroUVScales);
+				PerViewUniformParameters.MacroUVParameters = FVector4(ObjectNDCPosition.X, ObjectNDCPosition.Y, ObjectMacroUVScales.X, ObjectMacroUVScales.Y);
+				CollectorResources.UniformBuffer = FParticleSpriteUniformBufferRef::CreateUniformBufferImmediate(PerViewUniformParameters, UniformBuffer_SingleFrame);
+
+				const FDynamicSpriteEmitterReplayDataBase* SourceData = GetSourceData();
+				check(SourceData);
+
+				FParticleSpriteVertexFactory* SpriteVertexFactory = &CollectorResources.VertexFactory;
+
+				// Don't render if the material will be ignored
+				const bool bIsWireframe = ViewFamily.EngineShowFlags.Wireframe;
+
+				// Calculate the number of particles that must be drawn.
+				int32 ParticleCount = Source.ActiveParticleCount;
+				if ((Source.MaxDrawCount >= 0) && (ParticleCount > Source.MaxDrawCount))
+				{
+					ParticleCount = Source.MaxDrawCount;
+				}
+
+				// Set the sprite uniform buffer for this view.
+				SpriteVertexFactory->SetSpriteUniformBuffer( CollectorResources.UniformBuffer );
+				SpriteVertexFactory->SetInstanceBuffer( Allocation->VertexBuffer, Allocation->VertexOffset, GetDynamicVertexStride(FeatureLevel), bInstanced );
+				SpriteVertexFactory->SetDynamicParameterBuffer( DynamicParameterAllocation ? DynamicParameterAllocation->VertexBuffer : NULL, DynamicParameterAllocation->VertexOffset, GetDynamicParameterVertexStride(), bInstanced );
+
+				// Construct the mesh element to render.
+				FMeshBatch& Mesh = Collector.AllocateMesh();
+				FMeshBatchElement& BatchElement = Mesh.Elements[0];
+				BatchElement.IndexBuffer = &GParticleIndexBuffer;
+				if(bInstanced)
+				{
+					BatchElement.NumPrimitives = 2;
+					BatchElement.NumInstances = ParticleCount;
+				}
+				else
+				{
+					BatchElement.NumPrimitives = 2 * ParticleCount;
+					BatchElement.NumInstances = 1;
+				}
+				BatchElement.FirstIndex = 0;
+				Mesh.VertexFactory = SpriteVertexFactory;
+				// if the particle rendering data is presupplied, use it directly
+				Mesh.LCI = NULL;
+				if (SourceData->bUseLocalSpace == true)
+				{
+					BatchElement.PrimitiveUniformBufferResource = &Proxy->GetUniformBuffer();
+				}
+				else
+				{
+					BatchElement.PrimitiveUniformBufferResource = &Proxy->GetWorldSpacePrimitiveUniformBuffer();
+				}
+				BatchElement.MinVertexIndex = 0;
+				BatchElement.MaxVertexIndex = (ParticleCount * 4) - 1;
+				Mesh.CastShadow = Proxy->GetCastShadow();
+				Mesh.DepthPriorityGroup = (ESceneDepthPriorityGroup)Proxy->GetDepthPriorityGroup(View);
+
+				if ( bIsWireframe )
+				{
+					Mesh.MaterialRenderProxy = UMaterial::GetDefaultMaterial( MD_Surface )->GetRenderProxy( ViewFamily.EngineShowFlags.Selection ? bSelected : false );
+				}
+				else
+				{
+					Mesh.MaterialRenderProxy = MaterialResource[GIsEditor && (ViewFamily.EngineShowFlags.Selection) ? bSelected : 0];
+				}
+				Mesh.Type = PT_TriangleList;
+
+				Mesh.bCanApplyViewModeOverrides = true;
+				Mesh.bUseWireframeSelectionColoring = Proxy->IsSelected();
+
+				Collector.AddMesh(ViewIndex, Mesh);
+			}
+		}
+		else if (SourceData->EmitterRenderMode == ERM_Point)
+		{
+			RenderDebug(Proxy, Collector.GetPDI(ViewIndex), View, false);
+		}
+		else if (SourceData->EmitterRenderMode == ERM_Cross)
+		{
+			RenderDebug(Proxy, Collector.GetPDI(ViewIndex), View, true);
+		}
+	}
+}
+
 /**
  *	Render thread only draw call
  *
@@ -1327,7 +1400,7 @@ void FDynamicSpriteEmitterData::GatherSimpleLights(const FParticleSystemScenePro
  *	@param	PDI			The primitive draw interface to render with
  *	@param	View		The scene view being rendered
  */
-int32 FDynamicSpriteEmitterData::Render(FParticleSystemSceneProxy* Proxy, FPrimitiveDrawInterface* PDI,const FSceneView* View)
+int32 FDynamicSpriteEmitterData::RenderEmitter(FParticleSystemSceneProxy* Proxy, FPrimitiveDrawInterface* PDI,const FSceneView* View)
 {
 	SCOPE_CYCLE_COUNTER(STAT_SpriteRenderingTime);
 
@@ -1458,7 +1531,7 @@ int32 FDynamicSpriteEmitterData::Render(FParticleSystemSceneProxy* Proxy, FPrimi
  *
  *	@return	bool			true if successful, false if failed
  */
-void FDynamicSpriteEmitterData::CreateRenderThreadResources(const FParticleSystemSceneProxy* InOwnerProxy)
+void FDynamicSpriteEmitterData::UpdateRenderThreadResourcesEmitter(const FParticleSystemSceneProxy* InOwnerProxy)
 {
 	// Create the vertex factory...
 	if (VertexFactory == NULL)
@@ -1680,7 +1753,7 @@ void FDynamicMeshEmitterData::Init( bool bInSelected,
  *
  *	@return	bool			true if successful, false if failed
  */
-void FDynamicMeshEmitterData::CreateRenderThreadResources(const FParticleSystemSceneProxy* InOwnerProxy)
+void FDynamicMeshEmitterData::UpdateRenderThreadResourcesEmitter(const FParticleSystemSceneProxy* InOwnerProxy)
 {
 	// Create the vertex factory
 	if (VertexFactory == NULL)
@@ -1690,38 +1763,227 @@ void FDynamicMeshEmitterData::CreateRenderThreadResources(const FParticleSystemS
 		SetupVertexFactory((FMeshParticleVertexFactory*)VertexFactory, StaticMesh->RenderData->LODResources[0]);
 	}
 
-	// Generate the uniform buffer.
-	const FDynamicSpriteEmitterReplayDataBase* SourceData = GetSourceData();
-	if( SourceData )
-	{
-		FMeshParticleUniformParameters UniformParameters;
-		UniformParameters.SubImageSize = FVector4(
-		1.0f / SourceData->SubImages_Horizontal,
-		1.0f / SourceData->SubImages_Vertical,
-		0,0);
+	// Generate the uniform buffer.
+	const FDynamicSpriteEmitterReplayDataBase* SourceData = GetSourceData();
+	if( SourceData )
+	{
+		FMeshParticleUniformParameters UniformParameters;
+		UniformParameters.SubImageSize = FVector4(
+		1.0f / SourceData->SubImages_Horizontal,
+		1.0f / SourceData->SubImages_Vertical,
+		0,0);
+
+		// A weight is used to determine whether the mesh texture coordinates or SubUVs are passed from the vertex shader to the pixel shader.
+		const uint32 TexCoordWeight = (SourceData->SubUVDataOffset > 0) ? 1 : 0;
+		UniformParameters.TexCoordWeightA = TexCoordWeight;
+		UniformParameters.TexCoordWeightB = 1 - TexCoordWeight;
+		UniformBuffer = FMeshParticleUniformBufferRef::CreateUniformBufferImmediate( UniformParameters, UniformBuffer_SingleFrame );
+	}
+}
+
+/**
+ *	Release the render thread resources for this emitter data
+ *
+ *	@param	InOwnerProxy	The proxy that owns this dynamic emitter data
+ *
+ *	@return	bool			true if successful, false if failed
+ */
+void FDynamicMeshEmitterData::ReleaseRenderThreadResources(const FParticleSystemSceneProxy* InOwnerProxy)
+{
+	return FDynamicSpriteEmitterDataBase::ReleaseRenderThreadResources( InOwnerProxy );
+	UniformBuffer.SafeRelease();
+}
+
+class FDynamicMeshEmitterCollectorResources : public FOneFrameResource
+{
+public:
+	FMeshParticleVertexFactory VertexFactory;
+	FParticleSpriteUniformBufferRef UniformBuffer;
+
+	virtual ~FDynamicMeshEmitterCollectorResources()
+	{
+		VertexFactory.ReleaseResource();
+	}
+};
+
+class FMeshParticleInstanceVertices : public FOneFrameResource
+{
+public:
+	TArray<FMeshParticleInstanceVertex, SceneRenderingAllocator> InstanceDataAllocationsCPU;
+	TArray<FMeshParticleInstanceVertexDynamicParameter, SceneRenderingAllocator> DynamicParameterDataAllocationsCPU;
+};
+
+void FDynamicMeshEmitterData::GetDynamicMeshElementsEmitter(const FParticleSystemSceneProxy* Proxy, const FSceneView* View, const FSceneViewFamily& ViewFamily, int32 ViewIndex, FMeshElementCollector& Collector) const
+{
+	SCOPE_CYCLE_COUNTER(STAT_MeshRenderingTime);
+
+	const bool bInstanced = GRHIFeatureLevel >= ERHIFeatureLevel::SM3;
+
+	if (bValid)
+	{
+		if (Source.EmitterRenderMode == ERM_Normal)
+		{
+			const auto FeatureLevel = ViewFamily.GetFeatureLevel();
+			const auto ShaderPlatform = GShaderPlatformForFeatureLevel[FeatureLevel];
+
+			int32 ParticleCount = Source.ActiveParticleCount;
+			if ((Source.MaxDrawCount >= 0) && (ParticleCount > Source.MaxDrawCount))
+			{
+				ParticleCount = Source.MaxDrawCount;
+			}
+
+			const int32 InstanceVertexStride = GetDynamicVertexStride(FeatureLevel);
+			const int32 DynamicParameterVertexStride  = GetDynamicParameterVertexStride();
+
+			FDynamicMeshEmitterCollectorResources& CollectorResources = Collector.AllocateOneFrameResource<FDynamicMeshEmitterCollectorResources>();
+			CollectorResources.VertexFactory.SetParticleFactoryType(PVFT_Mesh);
+			CollectorResources.VertexFactory.SetFeatureLevel(FeatureLevel);
+			CollectorResources.VertexFactory.InitResource();
+			SetupVertexFactory(&CollectorResources.VertexFactory, StaticMesh->RenderData->LODResources[0]);
+
+			// Setup the vertex factory.
+			FMeshParticleVertexFactory* MeshVertexFactory = &CollectorResources.VertexFactory;
+			MeshVertexFactory->SetUniformBuffer( UniformBuffer );
+
+			FMeshParticleInstanceVertices* InstanceVerticesCPU = NULL;
+
+			if(bInstanced)
+			{
+				FGlobalDynamicVertexBuffer::FAllocation Allocation = FGlobalDynamicVertexBuffer::Get().Allocate( ParticleCount * InstanceVertexStride );
+				FGlobalDynamicVertexBuffer::FAllocation DynamicParameterAllocation;
+
+				if (bUsesDynamicParameter)
+				{
+					DynamicParameterAllocation = FGlobalDynamicVertexBuffer::Get().Allocate( ParticleCount * DynamicParameterVertexStride );
+				}
+
+				if(Allocation.IsValid() && (!bUsesDynamicParameter || DynamicParameterAllocation.IsValid()))
+				{
+					// Fill instance buffer.
+					GetInstanceData(Allocation.Buffer, DynamicParameterAllocation.Buffer, Proxy, View);
+				}
+
+				MeshVertexFactory->SetInstanceBuffer(Allocation.VertexBuffer, Allocation.VertexOffset, GetDynamicVertexStride(FeatureLevel));
+				MeshVertexFactory->SetDynamicParameterBuffer(DynamicParameterAllocation.VertexBuffer, DynamicParameterAllocation.VertexOffset , GetDynamicParameterVertexStride());
+			}
+			else
+			{
+				InstanceVerticesCPU = &Collector.AllocateOneFrameResource<FMeshParticleInstanceVertices>();
+				InstanceVerticesCPU->InstanceDataAllocationsCPU.Reset(ParticleCount);
+				InstanceVerticesCPU->InstanceDataAllocationsCPU.AddUninitialized(ParticleCount);
+
+				if (bUsesDynamicParameter)
+				{
+					InstanceVerticesCPU->DynamicParameterDataAllocationsCPU.Reset(ParticleCount);
+					InstanceVerticesCPU->DynamicParameterDataAllocationsCPU.AddUninitialized(ParticleCount);
+				}
+
+				// Fill instance buffer.
+				GetInstanceData((void*)InstanceVerticesCPU->InstanceDataAllocationsCPU.GetTypedData(), (void*)InstanceVerticesCPU->DynamicParameterDataAllocationsCPU.GetTypedData(), Proxy, View);
+			}
+
+			Proxy->UpdateWorldSpacePrimitiveUniformBuffer();
+
+			const FStaticMeshLODResources& LODModel = StaticMesh->RenderData->LODResources[0];
+			const bool bIsWireframe = AllowDebugViewmodes() && View->Family->EngineShowFlags.Wireframe;
+
+			//@todo. Handle LODs.
+			for (int32 LODIndex = 0; LODIndex < 1; LODIndex++)
+			{
+				for (int32 SectionIndex = 0; SectionIndex < LODModel.Sections.Num(); SectionIndex++)
+				{
+					FMaterialRenderProxy* MaterialProxy = MeshMaterials[SectionIndex]->GetRenderProxy(bSelected);
+					const FStaticMeshSection& Section = LODModel.Sections[SectionIndex];
+
+					if ((Section.NumTriangles == 0) || (MaterialProxy == NULL))
+					{
+						//@todo. This should never occur, but it does occasionally.
+						continue;
+					}
+
+					FMeshBatch& Mesh = Collector.AllocateMesh();
+					Mesh.VertexFactory = &CollectorResources.VertexFactory;
+					Mesh.DynamicVertexData = NULL;
+					Mesh.LCI = NULL;
+					Mesh.UseDynamicData = false;
+					Mesh.ReverseCulling = Proxy->IsLocalToWorldDeterminantNegative();
+					Mesh.CastShadow = Proxy->GetCastShadow();
+					Mesh.DepthPriorityGroup = (ESceneDepthPriorityGroup)Proxy->GetDepthPriorityGroup(View);
+
+					FMeshBatchElement& BatchElement = Mesh.Elements[0];
+					BatchElement.PrimitiveUniformBufferResource = &Proxy->GetWorldSpacePrimitiveUniformBuffer();
+					BatchElement.FirstIndex = Section.FirstIndex;
+					BatchElement.MinVertexIndex = Section.MinVertexIndex;
+					BatchElement.MaxVertexIndex = Section.MaxVertexIndex;
+					BatchElement.NumInstances = bInstanced ? ParticleCount : 1;
+
+					if (bIsWireframe)
+					{
+						if (LODModel.WireframeIndexBuffer.IsInitialized()
+							&& !(RHISupportsTessellation(ShaderPlatform) && Mesh.VertexFactory->GetType()->SupportsTessellationShaders()))
+						{
+							Mesh.Type = PT_LineList;
+							Mesh.MaterialRenderProxy = Proxy->GetDeselectedWireframeMatInst();
+							BatchElement.FirstIndex = 0;
+							BatchElement.IndexBuffer = &LODModel.WireframeIndexBuffer;
+							BatchElement.NumPrimitives = LODModel.WireframeIndexBuffer.GetNumIndices() / 2;
+
+						}
+						else
+						{
+							Mesh.Type = PT_TriangleList;
+							Mesh.MaterialRenderProxy = MeshMaterials[SectionIndex]->GetRenderProxy(bSelected);
+							Mesh.bWireframe = true;	
+							BatchElement.FirstIndex = 0;
+							BatchElement.IndexBuffer = &LODModel.IndexBuffer;
+							BatchElement.NumPrimitives = LODModel.IndexBuffer.GetNumIndices() / 3;
+						}
+					}
+					else
+					{
+						Mesh.Type = PT_TriangleList;
+						Mesh.MaterialRenderProxy = MeshMaterials[SectionIndex]->GetRenderProxy(bSelected);
+						BatchElement.IndexBuffer = &LODModel.IndexBuffer;
+						BatchElement.FirstIndex = Section.FirstIndex;
+						BatchElement.NumPrimitives = Section.NumTriangles;
+					}
+
+					if (!bInstanced)
+					{
+						FMeshParticleVertexFactory::FBatchParametersCPU& BatchParameters = Collector.AllocateOneFrameResource<FMeshParticleVertexFactory::FBatchParametersCPU>();
+						BatchParameters.InstanceBuffer = InstanceVerticesCPU->InstanceDataAllocationsCPU.GetTypedData();
+						BatchParameters.DynamicParameterBuffer = InstanceVerticesCPU->DynamicParameterDataAllocationsCPU.GetTypedData();
+						BatchElement.UserData = &BatchParameters;
+						BatchElement.UserIndex = 0;
+
+						Mesh.Elements.Reserve(ParticleCount);
+						for(int32 ParticleIndex = 1; ParticleIndex < ParticleCount; ++ParticleIndex)
+						{
+							FMeshBatchElement* NextElement = new(Mesh.Elements) FMeshBatchElement();
+							*NextElement = Mesh.Elements[0];
+							NextElement->UserIndex = ParticleIndex;
+						}
+					}
 
-		// A weight is used to determine whether the mesh texture coordinates or SubUVs are passed from the vertex shader to the pixel shader.
-		const uint32 TexCoordWeight = (SourceData->SubUVDataOffset > 0) ? 1 : 0;
-		UniformParameters.TexCoordWeightA = TexCoordWeight;
-		UniformParameters.TexCoordWeightB = 1 - TexCoordWeight;
-		UniformBuffer = FMeshParticleUniformBufferRef::CreateUniformBufferImmediate( UniformParameters, UniformBuffer_SingleFrame );
-	}
-}
+					Mesh.bCanApplyViewModeOverrides = true;
+					Mesh.bUseWireframeSelectionColoring = Proxy->IsSelected();
 
-/**
- *	Release the render thread resources for this emitter data
- *
- *	@param	InOwnerProxy	The proxy that owns this dynamic emitter data
- *
- *	@return	bool			true if successful, false if failed
- */
-void FDynamicMeshEmitterData::ReleaseRenderThreadResources(const FParticleSystemSceneProxy* InOwnerProxy)
-{
-	return FDynamicSpriteEmitterDataBase::ReleaseRenderThreadResources( InOwnerProxy );
-	UniformBuffer.SafeRelease();
+					Collector.AddMesh(ViewIndex, Mesh);
+				}
+			}
+		}
+		else if (Source.EmitterRenderMode == ERM_Point)
+		{
+			RenderDebug(Proxy, Collector.GetPDI(ViewIndex), View, false);
+		}
+		else if (Source.EmitterRenderMode == ERM_Cross)
+		{
+			RenderDebug(Proxy, Collector.GetPDI(ViewIndex), View, true);
+		}
+	}
 }
 
-int32 FDynamicMeshEmitterData::Render(FParticleSystemSceneProxy* Proxy, FPrimitiveDrawInterface* PDI,const FSceneView* View)
+int32 FDynamicMeshEmitterData::RenderEmitter(FParticleSystemSceneProxy* Proxy, FPrimitiveDrawInterface* PDI,const FSceneView* View)
 {
 	SCOPE_CYCLE_COUNTER(STAT_MeshRenderingTime);
 
@@ -1985,11 +2247,11 @@ void FDynamicMeshEmitterData::GatherSimpleLights(const FParticleSystemSceneProxy
 	GatherParticleLightData(Source, Proxy->GetLocalToWorld(), ViewFamily, OutParticleLights);
 }
 
-void FDynamicMeshEmitterData::GetParticleTransform(FBaseParticle& InParticle, const FVector& CameraPosition, const FVector& CameraFacingOpVector, const FQuat& PointToLockedAxis, 
-	FParticleSystemSceneProxy* Proxy, const FSceneView* View, FMatrix& OutTransformMat)
+void FDynamicMeshEmitterData::GetParticleTransform(const FBaseParticle& InParticle, const FVector& CameraPosition, const FVector& CameraFacingOpVector, const FQuat& PointToLockedAxis, 
+	const FParticleSystemSceneProxy* Proxy, const FSceneView* View, FMatrix& OutTransformMat) const
 	
 {
-	const uint8* ParticleBase = (uint8*)&InParticle;
+	const uint8* ParticleBase = (const uint8*)&InParticle;
 	OutTransformMat = FMatrix::Identity;
 	
 	FMatrix LocalToWorld;
@@ -2236,7 +2498,7 @@ void FDynamicMeshEmitterData::GetParticleTransform(FBaseParticle& InParticle, co
 		FVector kRotVec = FVector(fRot, fRot, fRot);
 		kRotator = FRotator::MakeFromEuler(kRotVec);
 
-		FMeshRotationPayloadData* PayloadData = (FMeshRotationPayloadData*)((uint8*)&InParticle + Source.MeshRotationOffset);
+		const FMeshRotationPayloadData* PayloadData = (const FMeshRotationPayloadData*)((const uint8*)&InParticle + Source.MeshRotationOffset);
 		kRotator += FRotator::MakeFromEuler(PayloadData->Rotation);
 	}
 
@@ -2288,7 +2550,7 @@ void FDynamicMeshEmitterData::GetParticleTransform(FBaseParticle& InParticle, co
 }
 
 
-void FDynamicMeshEmitterData::GetInstanceData(void* InstanceData, void* DynamicParameterData, FParticleSystemSceneProxy* Proxy, const FSceneView* View)
+void FDynamicMeshEmitterData::GetInstanceData(void* InstanceData, void* DynamicParameterData, const FParticleSystemSceneProxy* Proxy, const FSceneView* View) const
 {
 	SCOPE_CYCLE_COUNTER(STAT_ParticlePackingTime);
 
@@ -2363,7 +2625,7 @@ void FDynamicMeshEmitterData::GetInstanceData(void* InstanceData, void* DynamicP
 	{
 		const int32	CurrentIndex	= Source.ParticleIndices[i];
 		const uint8* ParticleBase	= Source.ParticleData.GetData() + CurrentIndex * Source.ParticleStride;
-		FBaseParticle& Particle		= *((FBaseParticle*) ParticleBase);
+		const FBaseParticle& Particle		= *((const FBaseParticle*) ParticleBase);
 		FMeshParticleInstanceVertex* CurrentInstanceVertex = (FMeshParticleInstanceVertex*)TempVert;
 		
 		// Populate instance buffer;
@@ -2420,7 +2682,7 @@ void FDynamicMeshEmitterData::GetInstanceData(void* InstanceData, void* DynamicP
 		// SubUVs 
 		if (Source.SubUVInterpMethod != PSUVIM_None && Source.SubUVDataOffset > 0)
 		{
-			FFullSubUVPayload* SubUVPayload = (FFullSubUVPayload*)(((uint8*)&Particle) + Source.SubUVDataOffset);
+			const FFullSubUVPayload* SubUVPayload = (const FFullSubUVPayload*)(((const uint8*)&Particle) + Source.SubUVDataOffset);
 
 			float SubImageIndex = SubUVPayload->ImageIndex;
 			float SubImageLerp = FMath::Fractional(SubImageIndex);
@@ -2445,7 +2707,7 @@ void FDynamicMeshEmitterData::GetInstanceData(void* InstanceData, void* DynamicP
 		TempVert += InstanceVertexStride; 
 	}
 }
-void FDynamicMeshEmitterData::SetupVertexFactory( FMeshParticleVertexFactory* VertexFactory, FStaticMeshLODResources& LODResources)
+void FDynamicMeshEmitterData::SetupVertexFactory( FMeshParticleVertexFactory* VertexFactory, FStaticMeshLODResources& LODResources) const
 {
 		FMeshParticleVertexFactory::DataType Data;
 
@@ -2592,7 +2854,7 @@ void FDynamicBeam2EmitterData::Init( bool bInSelected )
  *
  *	@return	bool			true if successful, false if failed
  */
-void FDynamicBeam2EmitterData::CreateRenderThreadResources(const FParticleSystemSceneProxy* InOwnerProxy)
+void FDynamicBeam2EmitterData::UpdateRenderThreadResourcesEmitter(const FParticleSystemSceneProxy* InOwnerProxy)
 {
 	// Create the vertex factory...
 	//@todo. Cache these??
@@ -2606,7 +2868,7 @@ void FDynamicBeam2EmitterData::CreateRenderThreadResources(const FParticleSystem
 /** Perform the actual work of filling the buffer, often called from another thread 
 * @param Me Fill data structure
 */
-void FDynamicBeam2EmitterData::DoBufferFill(FAsyncBufferFillData& Me)
+void FDynamicBeam2EmitterData::DoBufferFill(FAsyncBufferFillData& Me) const
 {
 	if( Me.VertexCount <= 0 || Me.IndexCount <= 0 || Me.VertexData == NULL || Me.IndexData == NULL )
 	{
@@ -2614,6 +2876,7 @@ void FDynamicBeam2EmitterData::DoBufferFill(FAsyncBufferFillData& Me)
 	}
 
 	FillIndexData(Me);
+
 	if (Source.bLowFreqNoise_Enabled)
 	{
 		FillData_Noise(Me);
@@ -2625,8 +2888,8 @@ void FDynamicBeam2EmitterData::DoBufferFill(FAsyncBufferFillData& Me)
 }
 
 FParticleBeamTrailUniformBufferRef CreateBeamTrailUniformBuffer(
-	FParticleSystemSceneProxy* Proxy,
-	FDynamicSpriteEmitterReplayDataBase* SourceData,
+	const FParticleSystemSceneProxy* Proxy,
+	const FDynamicSpriteEmitterReplayDataBase* SourceData,
 	const FSceneView* View )
 {
 	FParticleBeamTrailUniformParameters UniformParameters;
@@ -2654,7 +2917,122 @@ FParticleBeamTrailUniformBufferRef CreateBeamTrailUniformBuffer(
 	return FParticleBeamTrailUniformBufferRef::CreateUniformBufferImmediate( UniformParameters, UniformBuffer_SingleFrame );
 }
 
-int32 FDynamicBeam2EmitterData::Render(FParticleSystemSceneProxy* Proxy, FPrimitiveDrawInterface* PDI,const FSceneView* View)
+class FDynamicBeamTrailCollectorResources : public FOneFrameResource
+{
+public:
+	FParticleBeamTrailVertexFactory VertexFactory;
+
+	virtual ~FDynamicBeamTrailCollectorResources()
+	{
+		VertexFactory.ReleaseResource();
+	}
+};
+
+void FDynamicBeam2EmitterData::GetDynamicMeshElementsEmitter(const FParticleSystemSceneProxy* Proxy, const FSceneView* View, const FSceneViewFamily& ViewFamily, int32 ViewIndex, FMeshElementCollector& Collector) const
+{
+	SCOPE_CYCLE_COUNTER(STAT_BeamRenderingTime);
+	INC_DWORD_STAT(STAT_BeamParticlesRenderCalls);
+
+	if (bValid == false)
+	{
+		return;
+	}
+
+	if ((Source.VertexCount == 0) && (Source.IndexCount == 0))
+	{
+		return;
+	}
+
+	FGlobalDynamicVertexBuffer::FAllocation DynamicVertexAllocation;
+	FGlobalDynamicIndexBuffer::FAllocation DynamicIndexAllocation;
+	FGlobalDynamicVertexBuffer::FAllocation DynamicParameterAllocation;
+	FAsyncBufferFillData Data;
+
+	BuildViewFillData(Proxy, View, Source.VertexCount, sizeof(FParticleBeamTrailVertex), 0, DynamicVertexAllocation, DynamicIndexAllocation, &DynamicParameterAllocation, Data);
+	DoBufferFill(Data);
+
+	if (Source.bUseLocalSpace == false)
+	{
+		Proxy->UpdateWorldSpacePrimitiveUniformBuffer();
+	}
+
+	const bool bIsWireframe = ViewFamily.EngineShowFlags.Wireframe;
+
+	if (Data.OutTriangleCount > 0)
+	{
+		FDynamicBeamTrailCollectorResources& CollectorResources = Collector.AllocateOneFrameResource<FDynamicBeamTrailCollectorResources>();
+		CollectorResources.VertexFactory.SetParticleFactoryType(PVFT_BeamTrail);
+		CollectorResources.VertexFactory.SetFeatureLevel(ViewFamily.GetFeatureLevel());
+		CollectorResources.VertexFactory.InitResource();
+
+		// Create and set the uniform buffer for this emitter.
+		FParticleBeamTrailVertexFactory* BeamTrailVertexFactory = &CollectorResources.VertexFactory;
+		BeamTrailVertexFactory->SetBeamTrailUniformBuffer( CreateBeamTrailUniformBuffer( Proxy, &Source, View ) );			
+		BeamTrailVertexFactory->SetVertexBuffer( DynamicVertexAllocation.VertexBuffer, DynamicVertexAllocation.VertexOffset, GetDynamicVertexStride(View->GetFeatureLevel()) );
+		BeamTrailVertexFactory->SetDynamicParameterBuffer( NULL, 0, 0 );
+
+		FMeshBatch& Mesh = Collector.AllocateMesh();
+		FMeshBatchElement& BatchElement = Mesh.Elements[0];
+		BatchElement.IndexBuffer	= DynamicIndexAllocation.IndexBuffer;
+		BatchElement.FirstIndex		= DynamicIndexAllocation.FirstIndex;
+		Mesh.VertexFactory			= BeamTrailVertexFactory;
+		Mesh.DynamicVertexData		= NULL;
+		Mesh.DynamicVertexStride	= 0;
+		BatchElement.DynamicIndexData		= NULL;
+		BatchElement.DynamicIndexStride		= 0;
+		Mesh.LCI					= NULL;
+		if (Source.bUseLocalSpace == true)
+		{
+			BatchElement.PrimitiveUniformBufferResource = &Proxy->GetUniformBuffer();
+		}
+		else
+		{
+			BatchElement.PrimitiveUniformBufferResource = &Proxy->GetWorldSpacePrimitiveUniformBuffer();
+		}
+		int32 TrianglesToRender = Data.OutTriangleCount;
+		if ((TrianglesToRender % 2) != 0)
+		{
+			TrianglesToRender--;
+		}
+		BatchElement.NumPrimitives			= TrianglesToRender;
+		BatchElement.MinVertexIndex			= 0;
+		BatchElement.MaxVertexIndex			= Source.VertexCount - 1;
+		Mesh.UseDynamicData			= false;
+		Mesh.ReverseCulling			= Proxy->IsLocalToWorldDeterminantNegative();
+		Mesh.CastShadow				= Proxy->GetCastShadow();
+		Mesh.DepthPriorityGroup		= (ESceneDepthPriorityGroup)Proxy->GetDepthPriorityGroup(View);
+
+		if (AllowDebugViewmodes() && bIsWireframe && !ViewFamily.EngineShowFlags.Materials)
+		{
+			Mesh.MaterialRenderProxy	= Proxy->GetDeselectedWireframeMatInst();
+		}
+		else
+		{
+			Mesh.MaterialRenderProxy	= MaterialResource[GIsEditor && (ViewFamily.EngineShowFlags.Selection) ? bSelected : 0];
+		}
+		Mesh.Type = PT_TriangleStrip;
+
+		Mesh.bCanApplyViewModeOverrides = true;
+		Mesh.bUseWireframeSelectionColoring = Proxy->IsSelected();
+
+		Collector.AddMesh(ViewIndex, Mesh);
+
+		INC_DWORD_STAT_BY(STAT_BeamParticlesTrianglesRendered, Mesh.GetNumPrimitives());
+
+		if (Source.bRenderDirectLine == true)
+		{
+			RenderDirectLine(Proxy, Collector.GetPDI(ViewIndex), View);
+		}
+
+		if ((Source.bRenderLines == true) ||
+			(Source.bRenderTessellation == true))
+		{
+			RenderLines(Proxy, Collector.GetPDI(ViewIndex), View);
+		}
+	}
+}
+
+int32 FDynamicBeam2EmitterData::RenderEmitter(FParticleSystemSceneProxy* Proxy, FPrimitiveDrawInterface* PDI,const FSceneView* View)
 {
 	SCOPE_CYCLE_COUNTER(STAT_BeamRenderingTime);
 	INC_DWORD_STAT(STAT_BeamParticlesRenderCalls);
@@ -2795,9 +3173,25 @@ void FDynamicBeam2EmitterData::PreRenderView(FParticleSystemSceneProxy* Proxy, c
 	// Only need to do this once per-view
 	if (LastFramePreRendered < FrameNumber)
 	{
-		bool bOnlyOneView = !GIsEditor && ((GEngine && GEngine->GameViewport && (GEngine->GameViewport->GetCurrentSplitscreenConfiguration() == ESplitScreenType::None) && !GEngine->IsStereoscopic3D()) ? true : false);
+		InstanceDataAllocations.Empty();
+		IndexDataAllocations.Empty();
+		DynamicParameterDataAllocations.Empty();
+
+		for (int32 ViewIndex = 0; ViewIndex < ViewFamily->Views.Num(); ViewIndex++)
+		{
+			if (VisibilityMap & (1 << ViewIndex))
+			{
+				const FSceneView* View = ViewFamily->Views[ViewIndex];
 
-		BuildViewFillDataAndSubmit(Proxy,ViewFamily,VisibilityMap,bOnlyOneView,Source.VertexCount,sizeof(FParticleBeamTrailVertex),0);
+				FGlobalDynamicVertexBuffer::FAllocation* DynamicVertexAllocation = new(InstanceDataAllocations) FGlobalDynamicVertexBuffer::FAllocation();
+				FGlobalDynamicIndexBuffer::FAllocation* DynamicIndexAllocation = new(IndexDataAllocations) FGlobalDynamicIndexBuffer::FAllocation();
+				FGlobalDynamicVertexBuffer::FAllocation* DynamicParameterAllocation = bUsesDynamicParameter ? new(DynamicParameterDataAllocations) FGlobalDynamicVertexBuffer::FAllocation() : NULL;
+				FAsyncBufferFillData* Data = new (AsyncBufferFillTasks) FAsyncBufferFillData();
+
+				BuildViewFillData(Proxy, View, Source.VertexCount, sizeof(FParticleBeamTrailVertex), 0, *DynamicVertexAllocation, *DynamicIndexAllocation, DynamicParameterAllocation, *Data);
+				DoBufferFill(*Data);
+			}
+		}
 
 		// Set the frame tracker
 		LastFramePreRendered = FrameNumber;
@@ -2809,7 +3203,7 @@ void FDynamicBeam2EmitterData::PreRenderView(FParticleSystemSceneProxy* Proxy, c
 	}
 }
 
-void FDynamicBeam2EmitterData::RenderDirectLine(FParticleSystemSceneProxy* Proxy, FPrimitiveDrawInterface* PDI,const FSceneView* View)
+void FDynamicBeam2EmitterData::RenderDirectLine(const FParticleSystemSceneProxy* Proxy, FPrimitiveDrawInterface* PDI,const FSceneView* View) const
 {
 	for (int32 Beam = 0; Beam < Source.ActiveParticleCount; Beam++)
 	{
@@ -2835,7 +3229,7 @@ void FDynamicBeam2EmitterData::RenderDirectLine(FParticleSystemSceneProxy* Proxy
 	}
 }
 
-void FDynamicBeam2EmitterData::RenderLines(FParticleSystemSceneProxy* Proxy, FPrimitiveDrawInterface* PDI,const FSceneView* View)
+void FDynamicBeam2EmitterData::RenderLines(const FParticleSystemSceneProxy* Proxy, FPrimitiveDrawInterface* PDI,const FSceneView* View) const
 {
 	if (Source.bLowFreqNoise_Enabled)
 	{
@@ -2846,11 +3240,6 @@ void FDynamicBeam2EmitterData::RenderLines(FParticleSystemSceneProxy* Proxy, FPr
 		FMatrix CameraToWorld = View->ViewMatrices.ViewMatrix.Inverse();
 		FVector	ViewOrigin = CameraToWorld.GetOrigin();
 
-		Source.Sheets = (Source.Sheets > 0) ? Source.Sheets : 1;
-
-		// Frequency is the number of noise points to generate, evenly distributed along the line.
-		Source.Frequency = (Source.Frequency > 0) ? Source.Frequency : 1;
-
 		// NoiseTessellation is the amount of tessellation that should occur between noise points.
 		int32	TessFactor	= Source.NoiseTessellation ? Source.NoiseTessellation : 1;
 		float	InvTessFactor	= 1.0f / TessFactor;
@@ -3261,7 +3650,7 @@ void FDynamicBeam2EmitterData::RenderLines(FParticleSystemSceneProxy* Proxy, FPr
 	}
 }
 
-void FDynamicBeam2EmitterData::RenderDebug(FParticleSystemSceneProxy* Proxy, FPrimitiveDrawInterface* PDI, const FSceneView* View, bool bCrosses)
+void FDynamicBeam2EmitterData::RenderDebug(const FParticleSystemSceneProxy* Proxy, FPrimitiveDrawInterface* PDI, const FSceneView* View, bool bCrosses) const
 {
 }
 
@@ -3359,7 +3748,7 @@ static int32 CreateDynamicBeam2EmitterIndices(TIndexType* OutIndex, const FDynam
 	return TrianglesToRender;
 }
 
-int32 FDynamicBeam2EmitterData::FillIndexData(struct FAsyncBufferFillData& Data)
+int32 FDynamicBeam2EmitterData::FillIndexData(struct FAsyncBufferFillData& Data) const
 {
 	SCOPE_CYCLE_COUNTER(STAT_BeamFillIndexTime);
 
@@ -3375,10 +3764,6 @@ int32 FDynamicBeam2EmitterData::FillIndexData(struct FAsyncBufferFillData& Data)
 	// NOTE: This is primed for moving to tri-strips...
 	//
 	int32 TessFactor = Source.InterpolationPoints ? Source.InterpolationPoints : 1;
-	if (Source.Sheets <= 0)
-	{
-		Source.Sheets = 1;
-	}
 
 	check(Data.IndexCount > 0 && Data.IndexData != NULL);
 
@@ -3400,7 +3785,7 @@ int32 FDynamicBeam2EmitterData::FillIndexData(struct FAsyncBufferFillData& Data)
 	return TrianglesToRender;
 }
 
-int32 FDynamicBeam2EmitterData::FillVertexData_NoNoise(FAsyncBufferFillData& Me)
+int32 FDynamicBeam2EmitterData::FillVertexData_NoNoise(FAsyncBufferFillData& Me) const
 {
 	SCOPE_CYCLE_COUNTER(STAT_BeamFillVertexTime);
 
@@ -3411,11 +3796,6 @@ int32 FDynamicBeam2EmitterData::FillVertexData_NoNoise(FAsyncBufferFillData& Me)
 	FVector	ViewOrigin = CameraToWorld.GetOrigin();
 	int32 TessFactor = Source.InterpolationPoints ? Source.InterpolationPoints : 1;
 
-	if (Source.Sheets <= 0)
-	{
-		Source.Sheets = 1;
-	}
-
 	FVector	Offset(0.0f), LastOffset(0.0f);
 
 	int32 PackedCount = 0;
@@ -3816,7 +4196,7 @@ int32 FDynamicBeam2EmitterData::FillVertexData_NoNoise(FAsyncBufferFillData& Me)
 	return TrianglesToRender;
 }
 
-int32 FDynamicBeam2EmitterData::FillData_Noise(FAsyncBufferFillData& Me)
+int32 FDynamicBeam2EmitterData::FillData_Noise(FAsyncBufferFillData& Me) const
 {
 	SCOPE_CYCLE_COUNTER(STAT_BeamFillVertexTime);
 
@@ -3830,19 +4210,8 @@ int32 FDynamicBeam2EmitterData::FillData_Noise(FAsyncBufferFillData& Me)
 	FParticleBeamTrailVertex* Vertex = (FParticleBeamTrailVertex*)Me.VertexData;
 	FMatrix CameraToWorld = Me.View->ViewMatrices.ViewMatrix.InverseSafe();
 
-	if (Source.Sheets <= 0)
-	{
-		Source.Sheets = 1;
-	}
-
 	FVector	ViewOrigin	= CameraToWorld.GetOrigin();
 
-	// Frequency is the number of noise points to generate, evenly distributed along the line.
-	if (Source.Frequency <= 0)
-	{
-		Source.Frequency = 1;
-	}
-
 	// NoiseTessellation is the amount of tessellation that should occur between noise points.
 	int32	TessFactor	= Source.NoiseTessellation ? Source.NoiseTessellation : 1;
 	
@@ -4792,7 +5161,7 @@ int32 FDynamicBeam2EmitterData::FillData_Noise(FAsyncBufferFillData& Me)
 	return TrianglesToRender;
 }
 
-int32 FDynamicBeam2EmitterData::FillData_InterpolatedNoise(FAsyncBufferFillData& Me)
+int32 FDynamicBeam2EmitterData::FillData_InterpolatedNoise(FAsyncBufferFillData& Me) const
 {
 	int32	TrianglesToRender = 0;
 
@@ -4802,19 +5171,8 @@ int32 FDynamicBeam2EmitterData::FillData_InterpolatedNoise(FAsyncBufferFillData&
 	FParticleBeamTrailVertex* Vertex = (FParticleBeamTrailVertex*)Me.VertexData;
 	FMatrix CameraToWorld = Me.View->ViewMatrices.ViewMatrix.Inverse();
 	
-	if (Source.Sheets <= 0)
-	{
-		Source.Sheets = 1;
-	}
-
 	FVector	ViewOrigin	= CameraToWorld.GetOrigin();
 
-	// Frequency is the number of noise points to generate, evenly distributed along the line.
-	if (Source.Frequency <= 0)
-	{
-		Source.Frequency = 1;
-	}
-
 	// NoiseTessellation is the amount of tessellation that should occur between noise points.
 	int32	TessFactor	= Source.NoiseTessellation ? Source.NoiseTessellation : 1;
 	
@@ -5452,7 +5810,7 @@ void FDynamicTrailsEmitterData::Init(bool bInSelected)
  *
  *	@return	bool			true if successful, false if failed
  */
-void FDynamicTrailsEmitterData::CreateRenderThreadResources(const FParticleSystemSceneProxy* InOwnerProxy)
+void FDynamicTrailsEmitterData::UpdateRenderThreadResourcesEmitter(const FParticleSystemSceneProxy* InOwnerProxy)
 {
 	// Create the vertex factory...
 	if (VertexFactory == NULL)
@@ -5462,8 +5820,120 @@ void FDynamicTrailsEmitterData::CreateRenderThreadResources(const FParticleSyste
 	}
 }
 
+void FDynamicTrailsEmitterData::GetDynamicMeshElementsEmitter(const FParticleSystemSceneProxy* Proxy, const FSceneView* View, const FSceneViewFamily& ViewFamily, int32 ViewIndex, FMeshElementCollector& Collector) const
+{
+	SCOPE_CYCLE_COUNTER(STAT_TrailRenderingTime);
+	INC_DWORD_STAT(STAT_TrailParticlesRenderCalls);
+
+	if (bValid == false)
+	{
+		return;
+	}
+
+	if ((SourcePointer->VertexCount <= 0) || (SourcePointer->ActiveParticleCount <= 0) || (SourcePointer->IndexCount < 3))
+	{
+		return;
+	}
+
+	FGlobalDynamicVertexBuffer::FAllocation DynamicVertexAllocation;
+	FGlobalDynamicIndexBuffer::FAllocation DynamicIndexAllocation;
+	FGlobalDynamicVertexBuffer::FAllocation DynamicParameterAllocation;
+	FAsyncBufferFillData Data;
+
+	const int32 VertexStride = GetDynamicVertexStride(ViewFamily.GetFeatureLevel());
+	const int32 DynamicParameterVertexStride = bUsesDynamicParameter ? GetDynamicParameterVertexStride() : 0;
+
+	BuildViewFillData(Proxy, View, SourcePointer->VertexCount, VertexStride, DynamicParameterVertexStride, DynamicVertexAllocation, DynamicIndexAllocation, &DynamicParameterAllocation, Data);
+	DoBufferFill(Data);
+
+	if (SourcePointer->bUseLocalSpace == false)
+	{
+		Proxy->UpdateWorldSpacePrimitiveUniformBuffer();
+	}
+
+	const bool bIsWireframe = ViewFamily.EngineShowFlags.Wireframe;
+	int32 RenderedPrimitiveCount = 0;
+
+	if (Data.OutTriangleCount > 0 && bRenderGeometry)
+	{
+		FDynamicBeamTrailCollectorResources& CollectorResources = Collector.AllocateOneFrameResource<FDynamicBeamTrailCollectorResources>();
+		CollectorResources.VertexFactory.SetParticleFactoryType(PVFT_BeamTrail);
+		CollectorResources.VertexFactory.SetFeatureLevel(ViewFamily.GetFeatureLevel());
+		CollectorResources.VertexFactory.InitResource();
+
+		// Create and set the uniform buffer for this emitter.
+		FParticleBeamTrailVertexFactory* BeamTrailVertexFactory = &CollectorResources.VertexFactory;
+		BeamTrailVertexFactory->SetBeamTrailUniformBuffer( CreateBeamTrailUniformBuffer( Proxy, SourcePointer, View ) );
+		BeamTrailVertexFactory->SetVertexBuffer( DynamicVertexAllocation.VertexBuffer, DynamicVertexAllocation.VertexOffset, GetDynamicVertexStride(View->GetFeatureLevel()) );
+		BeamTrailVertexFactory->SetDynamicParameterBuffer( DynamicParameterAllocation.IsValid() ? DynamicParameterAllocation.VertexBuffer : NULL, DynamicParameterAllocation.IsValid() ? DynamicParameterAllocation.VertexOffset : 0, GetDynamicParameterVertexStride() );
+
+		FMeshBatch& Mesh = Collector.AllocateMesh();
+		FMeshBatchElement& BatchElement = Mesh.Elements[0];
+		BatchElement.IndexBuffer	= DynamicIndexAllocation.IndexBuffer;
+		BatchElement.FirstIndex		= DynamicIndexAllocation.FirstIndex;
+		Mesh.VertexFactory			= BeamTrailVertexFactory;
+		Mesh.LCI					= NULL;
+
+		BatchElement.PrimitiveUniformBufferResource = &Proxy->GetWorldSpacePrimitiveUniformBuffer();
+		BatchElement.NumPrimitives			= Data.OutTriangleCount;
+		BatchElement.MinVertexIndex			= 0;
+		BatchElement.MaxVertexIndex			= SourcePointer->VertexCount - 1;
+		Mesh.UseDynamicData			= false;
+		Mesh.ReverseCulling			= Proxy->IsLocalToWorldDeterminantNegative();
+		Mesh.CastShadow				= Proxy->GetCastShadow();
+		Mesh.DepthPriorityGroup		= (ESceneDepthPriorityGroup)Proxy->GetDepthPriorityGroup(View);
+
+		if (AllowDebugViewmodes() && bIsWireframe && !ViewFamily.EngineShowFlags.Materials)
+		{
+			Mesh.MaterialRenderProxy = Proxy->GetDeselectedWireframeMatInst();
+		}
+		else
+		{
+#if !(UE_BUILD_SHIPPING || UE_BUILD_TEST)
+			if (Data.OutTriangleCount != SourcePointer->PrimitiveCount)
+			{
+				UE_LOG(LogParticles, Log, TEXT("Data.OutTriangleCount = %4d vs. SourcePrimCount = %4d"), Data.OutTriangleCount, SourcePointer->PrimitiveCount);
+
+				int32 CheckTrailCount = 0;
+				int32 CheckTriangleCount = 0;
+				for (int32 ParticleIdx = 0; ParticleIdx < SourcePointer->ActiveParticleCount; ParticleIdx++)
+				{
+					int32 CurrentIndex = SourcePointer->ParticleIndices[ParticleIdx];
+					DECLARE_PARTICLE_PTR(CheckParticle, SourcePointer->ParticleData.GetData() + SourcePointer->ParticleStride * CurrentIndex);
+					FTrailsBaseTypeDataPayload* TrailPayload = (FTrailsBaseTypeDataPayload*)((uint8*)CheckParticle + SourcePointer->TrailDataOffset);
+					if (TRAIL_EMITTER_IS_HEAD(TrailPayload->Flags) == false)
+					{
+						continue;
+					}
+
+					UE_LOG(LogParticles, Log, TEXT("Trail %2d has %5d triangles"), TrailPayload->TrailIndex, TrailPayload->TriangleCount);
+					CheckTriangleCount += TrailPayload->TriangleCount;
+					CheckTrailCount++;
+				}
+				UE_LOG(LogParticles, Log, TEXT("Total 'live' trail count = %d"), CheckTrailCount);
+				UE_LOG(LogParticles, Log, TEXT("\t%5d triangles total (not counting degens)"), CheckTriangleCount);
+			}
+#endif
+			checkf(Data.OutTriangleCount <= SourcePointer->PrimitiveCount, TEXT("Data.OutTriangleCount = %4d vs. SourcePrimCount = %4d"), Data.OutTriangleCount, SourcePointer->PrimitiveCount);
+			Mesh.MaterialRenderProxy = MaterialResource[GIsEditor && (ViewFamily.EngineShowFlags.Selection) ? bSelected : 0];
+		}
+		Mesh.Type = PT_TriangleStrip;
+
+		Mesh.bCanApplyViewModeOverrides = true;
+		Mesh.bUseWireframeSelectionColoring = Proxy->IsSelected();
+
+		Collector.AddMesh(ViewIndex, Mesh);
+
+		RenderedPrimitiveCount = Mesh.GetNumPrimitives();
+	}
+
+	RenderDebug(Proxy, Collector.GetPDI(ViewIndex), View, false);
+
+	INC_DWORD_STAT_BY(STAT_TrailParticlesTrianglesRendered, RenderedPrimitiveCount);
+}
+
 // Render thread only draw call
-int32 FDynamicTrailsEmitterData::Render(FParticleSystemSceneProxy* Proxy, FPrimitiveDrawInterface* PDI, const FSceneView* View)
+int32 FDynamicTrailsEmitterData::RenderEmitter(FParticleSystemSceneProxy* Proxy, FPrimitiveDrawInterface* PDI, const FSceneView* View)
 {
 	SCOPE_CYCLE_COUNTER(STAT_TrailRenderingTime);
 	INC_DWORD_STAT(STAT_TrailParticlesRenderCalls);
@@ -5626,15 +6096,23 @@ void FDynamicTrailsEmitterData::PreRenderView(FParticleSystemSceneProxy* Proxy,
 			DynamicParameterVertexStride = GetDynamicParameterVertexStride();
 		}
 
-		bool bOnlyOneView = ShouldUsePrerenderView() || ((GEngine && GEngine->GameViewport && (GEngine->GameViewport->GetCurrentSplitscreenConfiguration() == ESplitScreenType::None)) ? true : false);
+		bool bOnlyOneView = false;
 
-		// Render both views when in stereo
-		if (GEngine->HMDDevice.IsValid() && GEngine->IsStereoscopic3D())
+		for (int32 ViewIndex = 0; ViewIndex < ViewFamily->Views.Num(); ViewIndex++)
 		{
-			bOnlyOneView = false;
-		}
+			if (VisibilityMap & (1 << ViewIndex))
+			{
+				const FSceneView* View = ViewFamily->Views[ViewIndex];
+
+				FGlobalDynamicVertexBuffer::FAllocation* DynamicVertexAllocation = new(InstanceDataAllocations) FGlobalDynamicVertexBuffer::FAllocation();
+				FGlobalDynamicIndexBuffer::FAllocation* DynamicIndexAllocation = new(IndexDataAllocations) FGlobalDynamicIndexBuffer::FAllocation();
+				FGlobalDynamicVertexBuffer::FAllocation* DynamicParameterAllocation = bUsesDynamicParameter ? new(DynamicParameterDataAllocations) FGlobalDynamicVertexBuffer::FAllocation() : NULL;
+				FAsyncBufferFillData* Data = new (AsyncBufferFillTasks) FAsyncBufferFillData();
 
-		BuildViewFillDataAndSubmit(Proxy,ViewFamily,VisibilityMap,bOnlyOneView,SourcePointer->VertexCount,VertexStride, DynamicParameterVertexStride);
+				BuildViewFillData(Proxy, View, SourcePointer->VertexCount, VertexStride, DynamicParameterVertexStride, *DynamicVertexAllocation, *DynamicIndexAllocation, DynamicParameterAllocation, *Data);
+				DoBufferFill(*Data);
+			}
+		}
 
 		// Set the frame tracker
 		LastFramePreRendered = FrameNumber;
@@ -5646,7 +6124,7 @@ void FDynamicTrailsEmitterData::PreRenderView(FParticleSystemSceneProxy* Proxy,
 	}
 }
 
-void FDynamicTrailsEmitterData::RenderDebug(FParticleSystemSceneProxy* Proxy, FPrimitiveDrawInterface* PDI, const FSceneView* View, bool bCrosses)
+void FDynamicTrailsEmitterData::RenderDebug(const FParticleSystemSceneProxy* Proxy, FPrimitiveDrawInterface* PDI, const FSceneView* View, bool bCrosses) const
 {
 	// Can't do anything in here...
 }
@@ -5658,7 +6136,7 @@ void FDynamicTrailsEmitterData::GetIndexAllocInfo( int32& OutNumIndices, int32&
 }
 
 // Data fill functions
-int32 FDynamicTrailsEmitterData::FillIndexData(struct FAsyncBufferFillData& Data)
+int32 FDynamicTrailsEmitterData::FillIndexData(struct FAsyncBufferFillData& Data) const
 {
 	SCOPE_CYCLE_COUNTER(STAT_TrailFillIndexTime);
 
@@ -5738,7 +6216,7 @@ int32 FDynamicTrailsEmitterData::FillIndexData(struct FAsyncBufferFillData& Data
 	return TrianglesToRender;
 }
 
-int32 FDynamicTrailsEmitterData::FillVertexData(struct FAsyncBufferFillData& Data)
+int32 FDynamicTrailsEmitterData::FillVertexData(struct FAsyncBufferFillData& Data) const
 {
 	check(!TEXT("FillVertexData: Base implementation should NOT be called!"));
 	return 0;
@@ -5761,7 +6239,7 @@ bool FDynamicRibbonEmitterData::ShouldUsePrerenderView()
 	return (RenderAxisOption != Trails_CameraUp);
 }
 
-void FDynamicRibbonEmitterData::RenderDebug(FParticleSystemSceneProxy* Proxy, FPrimitiveDrawInterface* PDI, const FSceneView* View, bool bCrosses)
+void FDynamicRibbonEmitterData::RenderDebug(const FParticleSystemSceneProxy* Proxy, FPrimitiveDrawInterface* PDI, const FSceneView* View, bool bCrosses) const
 {
 	if ((bRenderParticles == true) || (bRenderTangents == true))
 	{
@@ -5773,13 +6251,14 @@ void FDynamicRibbonEmitterData::RenderDebug(FParticleSystemSceneProxy* Proxy, FP
 		FColor PrevDrawColor;
 		FVector DrawTangentEnd;
 
-		uint8* Address = Source.ParticleData.GetData();
-		FRibbonTypeDataPayload* StartTrailPayload;
-		FRibbonTypeDataPayload* EndTrailPayload = NULL;
-		FBaseParticle* DebugParticle;
-		FRibbonTypeDataPayload* TrailPayload;
-		FBaseParticle* PrevParticle = NULL;
-		FRibbonTypeDataPayload* PrevTrailPayload;
+		const uint8* Address = Source.ParticleData.GetData();
+		const FRibbonTypeDataPayload* StartTrailPayload;
+		const FRibbonTypeDataPayload* EndTrailPayload = NULL;
+		const FBaseParticle* DebugParticle;
+		const FRibbonTypeDataPayload* TrailPayload;
+		const FBaseParticle* PrevParticle = NULL;
+		const FRibbonTypeDataPayload* PrevTrailPayload;
+
 		for (int32 ParticleIdx = 0; ParticleIdx < Source.ActiveParticleCount; ParticleIdx++)
 		{
 			DECLARE_PARTICLE_PTR(Particle, Address + Source.ParticleStride * Source.ParticleIndices[ParticleIdx]);
@@ -5796,7 +6275,7 @@ void FDynamicRibbonEmitterData::RenderDebug(FParticleSystemSceneProxy* Proxy, FP
 			DebugParticle = Particle;
 			// Find the end particle in this chain...
 			TrailPayload = StartTrailPayload;
-			FBaseParticle* IteratorParticle = DebugParticle;
+			const FBaseParticle* IteratorParticle = DebugParticle;
 			while (TrailPayload)
 			{
 				int32	Next = TRAIL_EMITTER_GET_NEXT(TrailPayload->Flags);
@@ -5815,8 +6294,8 @@ void FDynamicRibbonEmitterData::RenderDebug(FParticleSystemSceneProxy* Proxy, FP
 			}
 			if (EndTrailPayload != StartTrailPayload)
 			{
-				FBaseParticle* CurrSpawnedParticle = NULL;
-				FBaseParticle* NextSpawnedParticle = NULL;
+				const FBaseParticle* CurrSpawnedParticle = NULL;
+				const FBaseParticle* NextSpawnedParticle = NULL;
 				// We have more than one particle in the trail...
 				TrailPayload = EndTrailPayload;
 
@@ -5918,14 +6397,14 @@ void FDynamicRibbonEmitterData::RenderDebug(FParticleSystemSceneProxy* Proxy, FP
 						{
 							PDI->DrawLine(DrawPosition, DrawTangentEnd, FLinearColor(0.0f, 1.0f, 0.0f), Proxy->GetDepthPriorityGroup(View));
 						}
- 						else if (TrailPayload == EndTrailPayload)
- 						{
- 							PDI->DrawLine(DrawPosition, DrawTangentEnd, FLinearColor(1.0f, 0.0f, 0.0f), Proxy->GetDepthPriorityGroup(View));
- 						}
- 						else
- 						{
- 							PDI->DrawLine(DrawPosition, DrawTangentEnd, FLinearColor(1.0f, 1.0f, 0.0f), Proxy->GetDepthPriorityGroup(View));
- 						}
+						else if (TrailPayload == EndTrailPayload)
+						{
+							PDI->DrawLine(DrawPosition, DrawTangentEnd, FLinearColor(1.0f, 0.0f, 0.0f), Proxy->GetDepthPriorityGroup(View));
+						}
+						else
+						{
+							PDI->DrawLine(DrawPosition, DrawTangentEnd, FLinearColor(1.0f, 1.0f, 0.0f), Proxy->GetDepthPriorityGroup(View));
+						}
 					}
 
 					// The end will have Next set to the NULL flag...
@@ -5945,7 +6424,7 @@ void FDynamicRibbonEmitterData::RenderDebug(FParticleSystemSceneProxy* Proxy, FP
 }
 
 // Data fill functions
-int32 FDynamicRibbonEmitterData::FillVertexData(struct FAsyncBufferFillData& Data)
+int32 FDynamicRibbonEmitterData::FillVertexData(struct FAsyncBufferFillData& Data) const
 {
 	SCOPE_CYCLE_COUNTER(STAT_TrailFillVertexTime);
 
@@ -5968,7 +6447,7 @@ int32 FDynamicRibbonEmitterData::FillVertexData(struct FAsyncBufferFillData& Dat
 	float CurrDistance = 0.0f;
 
 	FBaseParticle* PackingParticle;
-	uint8* ParticleData = Source.ParticleData.GetData();
+	const uint8* ParticleData = Source.ParticleData.GetData();
 	for (int32 ParticleIdx = 0; ParticleIdx < Source.ActiveParticleCount; ParticleIdx++)
 	{
 		DECLARE_PARTICLE_PTR(Particle, ParticleData + Source.ParticleStride * Source.ParticleIndices[ParticleIdx]);
@@ -6275,19 +6754,19 @@ Helper class for keeping track of all the particles being used for vertex genera
 */
 struct FAnimTrailParticleRenderData
 {
-	FDynamicTrailsEmitterReplayData& Source;
-	uint8* ParticleDataAddress;
-
-	FBaseParticle* PrevPrevParticle;
-	FAnimTrailTypeDataPayload* PrevPrevPayload;
-	FBaseParticle* PrevParticle;
-	FAnimTrailTypeDataPayload* PrevPayload;
-	FBaseParticle* Particle;
-	FAnimTrailTypeDataPayload* Payload;
-	FBaseParticle* NextParticle;
-	FAnimTrailTypeDataPayload* NextPayload;
+	const FDynamicTrailsEmitterReplayData& Source;
+	const uint8* ParticleDataAddress;
+
+	const FBaseParticle* PrevPrevParticle;
+	const FAnimTrailTypeDataPayload* PrevPrevPayload;
+	const FBaseParticle* PrevParticle;
+	const FAnimTrailTypeDataPayload* PrevPayload;
+	const FBaseParticle* Particle;
+	const FAnimTrailTypeDataPayload* Payload;
+	const FBaseParticle* NextParticle;
+	const FAnimTrailTypeDataPayload* NextPayload;
 	
-	FAnimTrailParticleRenderData( FDynamicTrailsEmitterReplayData& InSource, FBaseParticle* InParticle, FAnimTrailTypeDataPayload* InPayload )
+	FAnimTrailParticleRenderData( const FDynamicTrailsEmitterReplayData& InSource, const FBaseParticle* InParticle, const FAnimTrailTypeDataPayload* InPayload )
 		:	Source(InSource),
 			ParticleDataAddress(Source.ParticleData.GetData()),
 			PrevPrevParticle(NULL),
@@ -6409,7 +6888,7 @@ struct FAnimTrailParticleRenderData
 			float PrevPrevTiledU;
 			float PrevPrevSize;
 			FLinearColor PrevPrevColor;
-			FBaseParticle* PrevPrevDynParamParticle;
+			const FBaseParticle* PrevPrevDynParamParticle;
 			if( PrevPrevParticle )
 			{
 				PrevPrevLocation = PrevPrevParticle->Location;
@@ -6437,7 +6916,7 @@ struct FAnimTrailParticleRenderData
 			float NextTiledU;
 			float NextSize;
 			FLinearColor NextColor;
-			FBaseParticle* NextDynParamParticle;
+			const FBaseParticle* NextDynParamParticle;
 			if( NextParticle )
 			{
 				NextLocation = NextParticle->Location;
@@ -6501,7 +6980,7 @@ struct FAnimTrailParticleRenderData
 	}
 };
 
-void FDynamicAnimTrailEmitterData::RenderDebug(FParticleSystemSceneProxy* Proxy, FPrimitiveDrawInterface* PDI, const FSceneView* View, bool bCrosses)
+void FDynamicAnimTrailEmitterData::RenderDebug(const FParticleSystemSceneProxy* Proxy, FPrimitiveDrawInterface* PDI, const FSceneView* View, bool bCrosses) const
 {
 	if ((bRenderParticles == true) || (bRenderTangents == true))
 	{
@@ -6517,8 +6996,9 @@ void FDynamicAnimTrailEmitterData::RenderDebug(FParticleSystemSceneProxy* Proxy,
 		float TiledU;
 		FLinearColor DummyColor;
 
-		uint8* Address = Source.ParticleData.GetData();
-		FAnimTrailTypeDataPayload* StartTrailPayload;
+		const uint8* Address = Source.ParticleData.GetData();
+		const FAnimTrailTypeDataPayload* StartTrailPayload;
+
 		for (int32 ParticleIdx = 0; ParticleIdx < Source.ActiveParticleCount; ParticleIdx++)
 		{
 			DECLARE_PARTICLE_PTR(Particle, Address + Source.ParticleStride * Source.ParticleIndices[ParticleIdx]);
@@ -6644,7 +7124,7 @@ void FDynamicAnimTrailEmitterData::RenderDebug(FParticleSystemSceneProxy* Proxy,
 }
 
 // Data fill functions
-int32 FDynamicAnimTrailEmitterData::FillVertexData(struct FAsyncBufferFillData& Data)
+int32 FDynamicAnimTrailEmitterData::FillVertexData(struct FAsyncBufferFillData& Data) const
 {
 	SCOPE_CYCLE_COUNTER(STAT_TrailFillVertexTime);
 
@@ -6664,7 +7144,7 @@ int32 FDynamicAnimTrailEmitterData::FillVertexData(struct FAsyncBufferFillData&
 	// The distance tracking for tiling the 2nd UV set
 	float CurrDistance = 0.0f;
 
-	uint8* ParticleData = Source.ParticleData.GetData();
+	const uint8* ParticleData = Source.ParticleData.GetData();
 	for (int32 ParticleIdx = 0; ParticleIdx < Source.ActiveParticleCount; ParticleIdx++)
 	{
 		DECLARE_PARTICLE_PTR(Particle, ParticleData + Source.ParticleStride * Source.ParticleIndices[ParticleIdx]);
@@ -6865,10 +7345,6 @@ FParticleSystemSceneProxy::FParticleSystemSceneProxy(const UParticleSystemCompon
 		)
 	, DynamicData(InDynamicData)
 	, LastDynamicData(NULL)
-	, SelectedWireframeMaterialInstance(
-		GEngine->WireframeMaterial ? GEngine->WireframeMaterial->GetRenderProxy(false) : NULL,
-		GetSelectionColor(FLinearColor(1.0f, 0.0f, 0.0f, 1.0f),true,false)
-		)
 	, DeselectedWireframeMaterialInstance(
 		GEngine->WireframeMaterial ? GEngine->WireframeMaterial->GetRenderProxy(false) : NULL,
 		GetSelectionColor(FLinearColor(1.0f, 0.0f, 0.0f, 1.0f),false,false)
@@ -6877,9 +7353,12 @@ FParticleSystemSceneProxy::FParticleSystemSceneProxy(const UParticleSystemCompon
 	, LastFramePreRendered(-1)
 	, FirstFreeMeshBatch(0)
 {
+	WireframeColor = FLinearColor(1.0f, 0.0f, 0.0f);
+	LevelColor = FLinearColor(1.0f, 1.0f, 0.0f);
+	PropertyColor = FLinearColor(1.0f, 1.0f, 1.0f);
+
 #if STATS
 	LastStatCaptureTime = FApp::GetCurrentTime();
-	bCountedThisFrame = false;
 #endif
 	LODMethod = Component->LODMethod;
 }
@@ -6908,6 +7387,74 @@ FMeshBatch* FParticleSystemSceneProxy::GetPooledMeshBatch()
 
 // FPrimitiveSceneProxy interface.
 
+void FParticleSystemSceneProxy::GetDynamicMeshElements(const TArray<const FSceneView*>& Views, const FSceneViewFamily& ViewFamily, uint32 VisibilityMap, FMeshElementCollector& Collector) const
+{
+	QUICK_SCOPE_CYCLE_COUNTER(STAT_FParticleSystemSceneProxy_GetMeshElements);
+
+	if ((GIsEditor == true) || (GbEnableGameThreadLODCalculation == false))
+	{
+		for (int32 ViewIndex = 0; ViewIndex < Views.Num(); ViewIndex++)
+		{
+			if (VisibilityMap & (1 << ViewIndex))
+			{
+				const FSceneView* View = Views[ViewIndex];
+				//@todo parallelrendering - get rid of this legacy feedback to the game thread!  
+				const_cast<FParticleSystemSceneProxy*>(this)->DetermineLODDistance(View, GFrameNumberRenderThread);
+			}
+		}
+	}
+
+	if (ViewFamily.EngineShowFlags.Particles)
+	{
+		SCOPE_CYCLE_COUNTER(STAT_ParticleRenderingTime);
+		FScopeCycleCounter Context(GetStatId());
+
+		const double StartTime = GTrackParticleRenderingStats || GTrackParticleRenderingStatsForOneFrame ? FPlatformTime::Seconds() : 0;
+		int32 NumDraws = 0;
+
+		if (DynamicData != NULL)
+		{
+			for (int32 Index = 0; Index < DynamicData->DynamicEmitterDataArray.Num(); Index++)
+			{
+				FDynamicEmitterDataBase* Data =	DynamicData->DynamicEmitterDataArray[Index];
+				if ((Data == NULL) || (Data->bValid != true))
+				{
+					continue;
+				}
+				//hold on to the emitter index in case we need to access any of its properties
+				DynamicData->EmitterIndex = Index;
+
+				for (int32 ViewIndex = 0; ViewIndex < Views.Num(); ViewIndex++)
+				{
+					if (VisibilityMap & (1 << ViewIndex))
+					{
+						const FSceneView* View = Views[ViewIndex];
+						Data->GetDynamicMeshElementsEmitter(this, View, ViewFamily, ViewIndex, Collector);
+						NumDraws++;
+					}
+				}
+			}
+		}
+
+		INC_DWORD_STAT_BY(STAT_ParticleDrawCalls, NumDraws);
+
+		if (ViewFamily.EngineShowFlags.Particles)
+		{
+			for (int32 ViewIndex = 0; ViewIndex < Views.Num(); ViewIndex++)
+			{
+				if (VisibilityMap & (1 << ViewIndex))
+				{
+					RenderBounds(Collector.GetPDI(ViewIndex), ViewFamily.EngineShowFlags, GetBounds(), IsSelected());
+					if (HasCustomOcclusionBounds())
+					{
+						RenderBounds(Collector.GetPDI(ViewIndex), ViewFamily.EngineShowFlags, GetCustomOcclusionBounds(), IsSelected());
+					}
+				}
+			}
+		}
+	}
+}
+
 /** 
 * Draw the scene proxy as a dynamic element
 *
@@ -6920,23 +7467,7 @@ void FParticleSystemSceneProxy::DrawDynamicElements(FPrimitiveDrawInterface* PDI
 
 	if (View->Family->EngineShowFlags.Particles)
 	{
-#if LOG_DETAILED_PARTICLE_RENDER_STATS
-		static uint64 LastFrameCounter = 0;
-		if( LastFrameCounter != GFrameCounter )
-		{
-			GDetailedParticleRenderStats.DumpStats();
-			GDetailedParticleRenderStats.Reset();
-			LastFrameCounter = GFrameCounter;
-		}
-
-		UParticleSystemComponent* ParticleSystemComponent = CastChecked<UParticleSystemComponent>(PrimitiveSceneInfo->Component);
-		if (ParticleSystemComponent->Template == NULL)
-		{
-			return;
-		}
-#endif
 		SCOPE_CYCLE_COUNTER(STAT_ParticleRenderingTime);
-		TRACK_DETAILED_PARTICLE_RENDER_STATS(ParticleSystemComponent->Template);
 		FScopeCycleCounter Context(GetStatId());
 
 		const double StartTime = GTrackParticleRenderingStats || GTrackParticleRenderingStatsForOneFrame ? FPlatformTime::Seconds() : 0;
@@ -6956,7 +7487,7 @@ void FParticleSystemSceneProxy::DrawDynamicElements(FPrimitiveDrawInterface* PDI
 					//hold on to the emitter index in case we need to access any of its properties
 					DynamicData->EmitterIndex = Index;
 
-					const int32 DrawCalls = Data->Render(this, PDI, View);
+					const int32 DrawCalls = Data->RenderEmitter(this, PDI, View);
 					NumDraws += DrawCalls;
 					if (DrawCalls > 0)
 					{
@@ -7107,7 +7638,7 @@ void FParticleSystemSceneProxy::CreateRenderThreadResourcesForEmitterData()
 			FDynamicEmitterDataBase* Data =	DynamicData->DynamicEmitterDataArray[Index];
 			if (Data != NULL)
 			{
-				Data->CreateRenderThreadResources(this);
+				Data->UpdateRenderThreadResourcesEmitter(this);
 			}
 		}
 	}
@@ -7260,7 +7791,7 @@ void FParticleSystemSceneProxy::OnTransformChanged()
 	WorldSpacePrimitiveUniformBuffer.ReleaseResource();
 }
 
-void FParticleSystemSceneProxy::UpdateWorldSpacePrimitiveUniformBuffer()
+void FParticleSystemSceneProxy::UpdateWorldSpacePrimitiveUniformBuffer() const
 {
 	check(IsInRenderingThread());
 	if (!WorldSpacePrimitiveUniformBuffer.IsInitialized())
@@ -7272,6 +7803,7 @@ void FParticleSystemSceneProxy::UpdateWorldSpacePrimitiveUniformBuffer()
 			GetLocalBounds(),
 			ReceivesDecals(),
 			false,
+			UseEditorDepthTest(),
 			1.0f			// LPV bias
 			);
 		WorldSpacePrimitiveUniformBuffer.SetContents(PrimitiveUniformShaderParameters);
diff --git a/Engine/Source/Runtime/Engine/Private/Particles/ParticleTrail2EmitterInstance.cpp b/Engine/Source/Runtime/Engine/Private/Particles/ParticleTrail2EmitterInstance.cpp
--- a/Engine/Source/Runtime/Engine/Private/Particles/ParticleTrail2EmitterInstance.cpp
+++ b/Engine/Source/Runtime/Engine/Private/Particles/ParticleTrail2EmitterInstance.cpp
@@ -2797,6 +2797,7 @@ bool FParticleRibbonEmitterInstance::FillReplayData(FDynamicEmitterReplayDataBas
 	NewReplayData->MaxActiveParticleCount = MaxActiveParticles;
 	NewReplayData->MaxTessellationBetweenParticles = TrailTypeData->MaxTessellationBetweenParticles ? TrailTypeData->MaxTessellationBetweenParticles : 1;
 	NewReplayData->Sheets = TrailTypeData->SheetsPerTrail ? TrailTypeData->SheetsPerTrail : 1;
+	NewReplayData->Sheets = FMath::Max(NewReplayData->Sheets, 1);
 
 	NewReplayData->VertexCount = VertexCount;
 	NewReplayData->IndexCount = TriangleCount + 2;
@@ -3692,7 +3693,6 @@ void FParticleAnimTrailEmitterInstance::DetermineVertexAndTriangleCount()
 	HeadOnlyParticles = 0;
 
 	int32 CheckParticleCount = 0;
-
 	int32 TempVertexCount;
 
 	bool bApplyDistanceTessellation = !FMath::IsNearlyZero(TrailTypeData->DistanceTessellationStepSize);
diff --git a/Engine/Source/Runtime/Engine/Private/PhysicsEngine/PhysDrawing.cpp b/Engine/Source/Runtime/Engine/Private/PhysicsEngine/PhysDrawing.cpp
--- a/Engine/Source/Runtime/Engine/Private/PhysicsEngine/PhysDrawing.cpp
+++ b/Engine/Source/Runtime/Engine/Private/PhysicsEngine/PhysDrawing.cpp
@@ -46,6 +46,11 @@ void FKSphereElem::DrawElemSolid(class FPrimitiveDrawInterface* PDI, const FTran
 	DrawSphere(PDI, ElemTM.GetLocation(), FVector( this->Radius * Scale ), DrawCollisionSides, DrawCollisionSides/2, MaterialRenderProxy, SDPG_World );
 }
 
+void FKSphereElem::GetElemSolid(const FTransform& ElemTM, float Scale, const FMaterialRenderProxy* MaterialRenderProxy, int32 ViewIndex, FMeshElementCollector& Collector) const
+{
+	GetSphereMesh(ElemTM.GetLocation(), FVector( this->Radius * Scale ), DrawCollisionSides, DrawCollisionSides/2, MaterialRenderProxy, SDPG_World, false, ViewIndex, Collector );
+}
+
 
 /////////////////////////////////////////////////////////////////////////////////////
 // FKBoxElem
@@ -91,6 +96,11 @@ void FKBoxElem::DrawElemSolid(class FPrimitiveDrawInterface* PDI, const FTransfo
 	DrawBox(PDI, ElemTM.ToMatrixWithScale(), 0.5f * FVector(X, Y, Z), MaterialRenderProxy, SDPG_World );
 }
 
+void FKBoxElem::GetElemSolid(const FTransform& ElemTM, float Scale, const FMaterialRenderProxy* MaterialRenderProxy, int32 ViewIndex, FMeshElementCollector& Collector) const
+{
+	GetBoxMesh(ElemTM.ToMatrixWithScale(), 0.5f * FVector(X, Y, Z), MaterialRenderProxy, SDPG_World, ViewIndex, Collector);
+}
+
 /////////////////////////////////////////////////////////////////////////////////////
 // FKSphylElem
 /////////////////////////////////////////////////////////////////////////////////////
@@ -139,6 +149,106 @@ void FKSphylElem::DrawElemWire(FPrimitiveDrawInterface* PDI, const FTransform& E
 	PDI->DrawLine(TopEnd - Scale*Radius*YAxis, BottomEnd - Scale*Radius*YAxis, Color, SDPG_World);
 }
 
+
+void FKSphylElem::GetElemSolid(const FTransform& ElemTM, float Scale, const FMaterialRenderProxy* MaterialRenderProxy, int32 ViewIndex, FMeshElementCollector& Collector) const
+{
+	const int32 NumSides = DrawCollisionSides;
+	const int32 NumRings = (DrawCollisionSides/2) + 1;
+
+	// The first/last arc are on top of each other.
+	const int32 NumVerts = (NumSides+1) * (NumRings+1);
+	FDynamicMeshVertex* Verts = (FDynamicMeshVertex*)FMemory::Malloc( NumVerts * sizeof(FDynamicMeshVertex) );
+
+	// Calculate verts for one arc
+	FDynamicMeshVertex* ArcVerts = (FDynamicMeshVertex*)FMemory::Malloc( (NumRings+1) * sizeof(FDynamicMeshVertex) );
+
+	for(int32 RingIdx=0; RingIdx<NumRings+1; RingIdx++)
+	{
+		FDynamicMeshVertex* ArcVert = &ArcVerts[RingIdx];
+
+		float Angle;
+		float ZOffset;
+		if( RingIdx <= DrawCollisionSides/4 )
+		{
+			Angle = ((float)RingIdx/(NumRings-1)) * PI;
+			ZOffset = 0.5 * Scale * Length;
+		}
+		else
+		{
+			Angle = ((float)(RingIdx-1)/(NumRings-1)) * PI;
+			ZOffset = -0.5 * Scale * Length;
+		}
+
+		// Note- unit sphere, so position always has mag of one. We can just use it for normal!		
+		FVector SpherePos;
+		SpherePos.X = 0.0f;
+		SpherePos.Y = Scale * Radius * FMath::Sin(Angle);
+		SpherePos.Z = Scale * Radius * FMath::Cos(Angle);
+
+		ArcVert->Position = SpherePos + FVector(0,0,ZOffset);
+
+		ArcVert->SetTangents(
+			FVector(1,0,0),
+			FVector(0.0f, -SpherePos.Z, SpherePos.Y),
+			SpherePos
+			);
+
+		ArcVert->TextureCoordinate.X = 0.0f;
+		ArcVert->TextureCoordinate.Y = ((float)RingIdx/NumRings);
+	}
+
+	// Then rotate this arc NumSides+1 times.
+	for(int32 SideIdx=0; SideIdx<NumSides+1; SideIdx++)
+	{
+		const FRotator ArcRotator(0, 360.f * ((float)SideIdx/NumSides), 0);
+		const FRotationMatrix ArcRot( ArcRotator );
+		const float XTexCoord = ((float)SideIdx/NumSides);
+
+		for(int32 VertIdx=0; VertIdx<NumRings+1; VertIdx++)
+		{
+			int32 VIx = (NumRings+1)*SideIdx + VertIdx;
+
+			Verts[VIx].Position = ArcRot.TransformPosition( ArcVerts[VertIdx].Position );
+
+			Verts[VIx].SetTangents(
+				ArcRot.TransformVector( ArcVerts[VertIdx].TangentX ),
+				ArcRot.TransformVector( ArcVerts[VertIdx].GetTangentY() ),
+				ArcRot.TransformVector( ArcVerts[VertIdx].TangentZ )
+				);
+
+			Verts[VIx].TextureCoordinate.X = XTexCoord;
+			Verts[VIx].TextureCoordinate.Y = ArcVerts[VertIdx].TextureCoordinate.Y;
+		}
+	}
+
+	FDynamicMeshBuilder MeshBuilder;
+	{
+		// Add all of the vertices to the mesh.
+		for(int32 VertIdx=0; VertIdx<NumVerts; VertIdx++)
+		{
+			MeshBuilder.AddVertex(Verts[VertIdx]);
+		}
+
+		// Add all of the triangles to the mesh.
+		for(int32 SideIdx=0; SideIdx<NumSides; SideIdx++)
+		{
+			const int32 a0start = (SideIdx+0) * (NumRings+1);
+			const int32 a1start = (SideIdx+1) * (NumRings+1);
+
+			for(int32 RingIdx=0; RingIdx<NumRings; RingIdx++)
+			{
+				MeshBuilder.AddTriangle(a0start + RingIdx + 0, a1start + RingIdx + 0, a0start + RingIdx + 1);
+				MeshBuilder.AddTriangle(a1start + RingIdx + 0, a1start + RingIdx + 1, a0start + RingIdx + 1);
+			}
+		}
+
+	}
+	MeshBuilder.GetMesh(ElemTM.ToMatrixWithScale(), MaterialRenderProxy, SDPG_World, false, false, ViewIndex, Collector);
+
+	FMemory::Free(Verts);
+	FMemory::Free(ArcVerts);
+}
+
 void FKSphylElem::DrawElemSolid(FPrimitiveDrawInterface* PDI, const FTransform& ElemTM, float Scale, const FMaterialRenderProxy* MaterialRenderProxy) const
 {
 	const int32 NumSides = DrawCollisionSides;
@@ -296,12 +406,11 @@ void FKConvexElem::DrawElemWire(FPrimitiveDrawInterface* PDI, const FTransform&
 #endif // WITH_PHYSX
 }
 
-void FKConvexElem::AddCachedSolidConvexGeom(TArray<FDynamicMeshVertex>& VertexBuffer, TArray<int32>& IndexBuffer, const FColor VertexColor)
+void FKConvexElem::AddCachedSolidConvexGeom(TArray<FDynamicMeshVertex>& VertexBuffer, TArray<int32>& IndexBuffer, const FColor VertexColor) const
 {
 #if WITH_PHYSX
 	if(ConvexMesh)
 	{
-
 		int32 StartVertOffset = VertexBuffer.Num();
 
 		// get PhysX data
@@ -418,7 +527,7 @@ void FConvexCollisionVertexFactory::InitConvexVertexFactory(const FConvexCollisi
 	}
 }
 
-void FKAggregateGeom::DrawAggGeom(FPrimitiveDrawInterface* PDI, const FTransform& Transform, const FColor Color, const FMaterialRenderProxy* MatInst, bool bPerHullColor, bool bDrawSolid)
+void FKAggregateGeom::DrawAggGeom(FPrimitiveDrawInterface* PDI, const FTransform& Transform, const FColor Color, const FMaterialRenderProxy* MatInst, bool bPerHullColor, bool bDrawSolid, bool bUseEditorDepthTest)
 {
 	const FVector Scale3D = Transform.GetScale3D();
 	FTransform ParentTM = Transform;
@@ -507,7 +616,7 @@ void FKAggregateGeom::DrawAggGeom(FPrimitiveDrawInterface* PDI, const FTransform
 				FBoxSphereBounds WorldBounds, LocalBounds;
 				CalcBoxSphereBounds(WorldBounds, LocalToWorld);
 				CalcBoxSphereBounds(LocalBounds, FTransform::Identity);
-				BatchElement.PrimitiveUniformBuffer = CreatePrimitiveUniformBufferImmediate(LocalToWorld.ToMatrixWithScale(), WorldBounds, LocalBounds, true);
+				BatchElement.PrimitiveUniformBuffer = CreatePrimitiveUniformBufferImmediate(LocalToWorld.ToMatrixWithScale(), WorldBounds, LocalBounds, true, bUseEditorDepthTest);
 				// previous l2w not used so treat as static
 				BatchElement.FirstIndex = 0;
 				BatchElement.NumPrimitives = RenderInfo->IndexBuffer->Indices.Num() / 3;
@@ -532,6 +641,124 @@ void FKAggregateGeom::DrawAggGeom(FPrimitiveDrawInterface* PDI, const FTransform
 	}
 }
 
+
+void FKAggregateGeom::GetAggGeom(const FTransform& Transform, const FColor Color, const FMaterialRenderProxy* MatInst, bool bPerHullColor, bool bDrawSolid, bool bUseEditorDepthTest, int32 ViewIndex, FMeshElementCollector& Collector) const
+{
+	const FVector Scale3D = Transform.GetScale3D();
+	FTransform ParentTM = Transform;
+	ParentTM.RemoveScaling();
+
+	if( Scale3D.GetAbs().IsUniform() )
+	{
+		for(int32 i=0; i<SphereElems.Num(); i++)
+		{
+			FTransform ElemTM = SphereElems[i].GetTransform();
+			ElemTM.ScaleTranslation(Scale3D);
+			ElemTM *= ParentTM;
+
+			if(bDrawSolid)
+				SphereElems[i].GetElemSolid(ElemTM, Scale3D.X, MatInst, ViewIndex, Collector);
+			else
+				SphereElems[i].DrawElemWire(Collector.GetPDI(ViewIndex), ElemTM, Scale3D.X, Color);
+		}
+
+		for(int32 i=0; i<BoxElems.Num(); i++)
+		{
+			FTransform ElemTM = BoxElems[i].GetTransform();
+			ElemTM.ScaleTranslation(Scale3D);
+			ElemTM *= ParentTM;
+
+			if(bDrawSolid)
+				BoxElems[i].GetElemSolid(ElemTM, Scale3D.X, MatInst, ViewIndex, Collector);
+			else
+				BoxElems[i].DrawElemWire(Collector.GetPDI(ViewIndex), ElemTM, Scale3D.X, Color);
+		}
+
+		for(int32 i=0; i<SphylElems.Num(); i++)
+		{
+			FTransform ElemTM = SphylElems[i].GetTransform();
+			ElemTM.ScaleTranslation(Scale3D);
+			ElemTM *= ParentTM;
+
+			if(bDrawSolid)
+				SphylElems[i].GetElemSolid(ElemTM, Scale3D.X, MatInst, ViewIndex, Collector);
+			else
+				SphylElems[i].DrawElemWire(Collector.GetPDI(ViewIndex), ElemTM, Scale3D.X, Color);
+		}
+	}
+
+	if(ConvexElems.Num() > 0)
+	{
+		if(bDrawSolid)
+		{
+			// Cache collision vertex/index buffer
+			if(!RenderInfo)
+			{
+				//@todo - parallelrendering, remove const cast
+				FKAggregateGeom& ThisGeom = const_cast<FKAggregateGeom&>(*this);
+				ThisGeom.RenderInfo = new FKConvexGeomRenderInfo();
+				ThisGeom.RenderInfo->VertexBuffer = new FConvexCollisionVertexBuffer();
+				ThisGeom.RenderInfo->IndexBuffer = new FConvexCollisionIndexBuffer();
+
+				for(int32 i=0; i<ConvexElems.Num(); i++)
+				{
+					// Get vertices/triangles from this hull.
+					ConvexElems[i].AddCachedSolidConvexGeom(ThisGeom.RenderInfo->VertexBuffer->Vertices, ThisGeom.RenderInfo->IndexBuffer->Indices, FColor(255,255,255));
+				}
+
+				// Only continue if we actually got some valid geometry
+				// Will crash if we try to init buffers with no data
+				if(ThisGeom.RenderInfo->HasValidGeometry())
+				{
+					ThisGeom.RenderInfo->VertexBuffer->InitResource();
+					ThisGeom.RenderInfo->IndexBuffer->InitResource();
+
+					ThisGeom.RenderInfo->CollisionVertexFactory = new FConvexCollisionVertexFactory(RenderInfo->VertexBuffer);
+					ThisGeom.RenderInfo->CollisionVertexFactory->InitResource();
+				}
+			}
+
+			// If we have geometry to draw, do so
+			if(RenderInfo->HasValidGeometry())
+			{
+				// Calculate transform
+				FTransform LocalToWorld = FTransform( FQuat::Identity, FVector::ZeroVector, Scale3D ) * ParentTM;
+
+				// Draw the mesh.
+				FMeshBatch& Mesh = Collector.AllocateMesh();
+				FMeshBatchElement& BatchElement = Mesh.Elements[0];
+				BatchElement.IndexBuffer = RenderInfo->IndexBuffer;
+				Mesh.VertexFactory = RenderInfo->CollisionVertexFactory;
+				Mesh.MaterialRenderProxy = MatInst;
+				FBoxSphereBounds WorldBounds, LocalBounds;
+				CalcBoxSphereBounds(WorldBounds, LocalToWorld);
+				CalcBoxSphereBounds(LocalBounds, FTransform::Identity);
+				BatchElement.PrimitiveUniformBuffer = CreatePrimitiveUniformBufferImmediate(LocalToWorld.ToMatrixWithScale(), WorldBounds, LocalBounds, true, bUseEditorDepthTest);
+				// previous l2w not used so treat as static
+				BatchElement.FirstIndex = 0;
+				BatchElement.NumPrimitives = RenderInfo->IndexBuffer->Indices.Num() / 3;
+				BatchElement.MinVertexIndex = 0;
+				BatchElement.MaxVertexIndex = RenderInfo->VertexBuffer->Vertices.Num() - 1;
+				Mesh.ReverseCulling = LocalToWorld.GetDeterminant() < 0.0f ? true : false;
+				Mesh.Type = PT_TriangleList;
+				Mesh.DepthPriorityGroup = SDPG_World;
+				Mesh.bCanApplyViewModeOverrides = false;
+				Collector.AddMesh(ViewIndex, Mesh);
+			}
+		}
+		else
+		{
+			for(int32 i=0; i<ConvexElems.Num(); i++)
+			{
+				FColor ConvexColor = bPerHullColor ? DebugUtilColor[i%NUM_DEBUG_UTIL_COLORS] : Color;
+				FTransform ElemTM = ConvexElems[i].GetTransform();
+				ElemTM *= Transform;
+				ConvexElems[i].DrawElemWire(Collector.GetPDI(ViewIndex), ElemTM, ConvexColor);
+			}
+		}
+	}
+}
+
 /** Release the RenderInfo (if its there) and safely clean up any resources. Call on the game thread. */
 void FKAggregateGeom::FreeRenderInfo()
 {
@@ -598,7 +825,7 @@ void UPhysicsAsset::DrawCollision(FPrimitiveDrawInterface* PDI, const USkeletalM
 		FTransform BoneTransform = GetSkelBoneTransform(BoneIndex, SpaceBases, LocalToWorld);
 		BoneTransform.SetScale3D(FVector(Scale));
 		BodySetup[i]->CreatePhysicsMeshes();
-		BodySetup[i]->AggGeom.DrawAggGeom( PDI, BoneTransform, *BoneColor, NULL, false, false );
+		BodySetup[i]->AggGeom.DrawAggGeom( PDI, BoneTransform, *BoneColor, NULL, false, false, false );
 	}
 }
 
diff --git a/Engine/Source/Runtime/Engine/Private/PhysicsEngine2D/AggregateGeometry2D.cpp b/Engine/Source/Runtime/Engine/Private/PhysicsEngine2D/AggregateGeometry2D.cpp
--- a/Engine/Source/Runtime/Engine/Private/PhysicsEngine2D/AggregateGeometry2D.cpp
+++ b/Engine/Source/Runtime/Engine/Private/PhysicsEngine2D/AggregateGeometry2D.cpp
@@ -118,7 +118,7 @@ void FAggregateGeometry2D::EmptyElements()
 	ConvexElements.Empty();
 }
 
-void FAggregateGeometry2D::DrawAggGeom(FPrimitiveDrawInterface* PDI, const FTransform& Transform, const FColor Color, const FMaterialRenderProxy* MatInst, bool bPerHullColor, bool bDrawSolid) const
+void FAggregateGeometry2D::DrawAggGeom(FPrimitiveDrawInterface* PDI, const FTransform& Transform, const FColor Color, const FMaterialRenderProxy* MatInst, bool bPerHullColor, bool bDrawSolid, bool bUseEditorDepthTest) const
 {
 	//@TODO: PAPER2D: Implement this method
 }
diff --git a/Engine/Source/Runtime/Engine/Private/PrimitiveComponent.cpp b/Engine/Source/Runtime/Engine/Private/PrimitiveComponent.cpp
--- a/Engine/Source/Runtime/Engine/Private/PrimitiveComponent.cpp
+++ b/Engine/Source/Runtime/Engine/Private/PrimitiveComponent.cpp
@@ -909,24 +909,6 @@ void UPrimitiveComponent::SetViewOwnerDepthPriorityGroup(
 	MarkRenderStateDirty();
 }
 
-
-
-#if WITH_EDITOR
-bool UPrimitiveComponent::SetupLightmapResolutionViewInfo(FPrimitiveSceneProxy& Proxy) const
-{
-	if (GIsEditor)
-	{
-		int32 Width, Height;
-		Proxy.SetIsLightMapResolutionPadded(GetLightMapResolution(Width, Height));
-		FVector2D TempVector(Width, Height);
-		Proxy.SetLightMapResolutionScale(TempVector);
-		Proxy.SetLightMapType(GetStaticLightingType());
-		return true;
-	}
-	return false;
-}
-#endif
-
 bool UPrimitiveComponent::IsWorldGeometry() const
 {
 	// if modify flag doesn't change, and 
