FPrimitiveSceneProxy::PreRenderView / DrawDynamicElements refactored into GetDynamicMeshElements

diff --git a/Engine/Source/Runtime/Engine/Private/PrimitiveDrawingUtils.cpp b/Engine/Source/Runtime/Engine/Private/PrimitiveDrawingUtils.cpp
--- a/Engine/Source/Runtime/Engine/Private/PrimitiveDrawingUtils.cpp
+++ b/Engine/Source/Runtime/Engine/Private/PrimitiveDrawingUtils.cpp
@@ -76,6 +76,60 @@ void DrawPlane10x10(class FPrimitiveDrawInterface* PDI,const FMatrix& ObjectToWo
 	MeshBuilder.Draw(PDI, FScaleMatrix(Radii) * ObjectToWorld, MaterialRenderProxy, DepthPriorityGroup, 0.f);
 }
 
+
+void GetBoxMesh(const FMatrix& BoxToWorld,const FVector& Radii,const FMaterialRenderProxy* MaterialRenderProxy,uint8 DepthPriorityGroup,int32 ViewIndex,FMeshElementCollector& Collector)
+{
+	// Calculate verts for a face pointing down Z
+	FVector Positions[4] =
+	{
+		FVector(-1, -1, +1),
+		FVector(-1, +1, +1),
+		FVector(+1, +1, +1),
+		FVector(+1, -1, +1)
+	};
+	FVector2D UVs[4] =
+	{
+		FVector2D(0,0),
+		FVector2D(0,1),
+		FVector2D(1,1),
+		FVector2D(1,0),
+	};
+
+	// Then rotate this face 6 times
+	FRotator FaceRotations[6];
+	FaceRotations[0] = FRotator(0,		0,	0);
+	FaceRotations[1] = FRotator(90.f,	0,	0);
+	FaceRotations[2] = FRotator(-90.f,	0,  0);
+	FaceRotations[3] = FRotator(0,		0,	90.f);
+	FaceRotations[4] = FRotator(0,		0,	-90.f);
+	FaceRotations[5] = FRotator(180.f,	0,	0);
+
+	FDynamicMeshBuilder MeshBuilder;
+
+	for(int32 f=0; f<6; f++)
+	{
+		FMatrix FaceTransform = FRotationMatrix(FaceRotations[f]);
+
+		int32 VertexIndices[4];
+		for(int32 VertexIndex = 0;VertexIndex < 4;VertexIndex++)
+		{
+			VertexIndices[VertexIndex] = MeshBuilder.AddVertex(
+				FaceTransform.TransformPosition( Positions[VertexIndex] ),
+				UVs[VertexIndex],
+				FaceTransform.TransformVector(FVector(1,0,0)),
+				FaceTransform.TransformVector(FVector(0,1,0)),
+				FaceTransform.TransformVector(FVector(0,0,1)),
+				FColor::White
+				);
+		}
+
+		MeshBuilder.AddTriangle(VertexIndices[0],VertexIndices[1],VertexIndices[2]);
+		MeshBuilder.AddTriangle(VertexIndices[0],VertexIndices[2],VertexIndices[3]);
+	}
+
+	MeshBuilder.GetMesh(FScaleMatrix(Radii) * BoxToWorld,MaterialRenderProxy,DepthPriorityGroup,false,false,ViewIndex,Collector);
+}
+
 void DrawBox(FPrimitiveDrawInterface* PDI,const FMatrix& BoxToWorld,const FVector& Radii,const FMaterialRenderProxy* MaterialRenderProxy,uint8 DepthPriorityGroup)
 {
 	// Calculate verts for a face pointing down Z
@@ -129,6 +183,88 @@ void DrawBox(FPrimitiveDrawInterface* PDI,const FMatrix& BoxToWorld,const FVecto
 	MeshBuilder.Draw(PDI,FScaleMatrix(Radii) * BoxToWorld,MaterialRenderProxy,DepthPriorityGroup,0.f);
 }
 
+void GetSphereMesh(const FVector& Center,const FVector& Radii,int32 NumSides,int32 NumRings,const FMaterialRenderProxy* MaterialRenderProxy,uint8 DepthPriority,bool bDisableBackfaceCulling,int32 ViewIndex,FMeshElementCollector& Collector)
+{
+	// Use a mesh builder to draw the sphere.
+	FDynamicMeshBuilder MeshBuilder;
+	{
+		// The first/last arc are on top of each other.
+		int32 NumVerts = (NumSides+1) * (NumRings+1);
+		FDynamicMeshVertex* Verts = (FDynamicMeshVertex*)FMemory::Malloc( NumVerts * sizeof(FDynamicMeshVertex) );
+
+		// Calculate verts for one arc
+		FDynamicMeshVertex* ArcVerts = (FDynamicMeshVertex*)FMemory::Malloc( (NumRings+1) * sizeof(FDynamicMeshVertex) );
+
+		for(int32 i=0; i<NumRings+1; i++)
+		{
+			FDynamicMeshVertex* ArcVert = &ArcVerts[i];
+
+			float angle = ((float)i/NumRings) * PI;
+
+			// Note- unit sphere, so position always has mag of one. We can just use it for normal!			
+			ArcVert->Position.X = 0.0f;
+			ArcVert->Position.Y = FMath::Sin(angle);
+			ArcVert->Position.Z = FMath::Cos(angle);
+
+			ArcVert->SetTangents(
+				FVector(1,0,0),
+				FVector(0.0f,-ArcVert->Position.Z,ArcVert->Position.Y),
+				ArcVert->Position
+				);
+
+			ArcVert->TextureCoordinate.X = 0.0f;
+			ArcVert->TextureCoordinate.Y = ((float)i/NumRings);
+		}
+
+		// Then rotate this arc NumSides+1 times.
+		for(int32 s=0; s<NumSides+1; s++)
+		{
+			FRotator ArcRotator(0, 360.f * (float)s/NumSides, 0);
+			FRotationMatrix ArcRot( ArcRotator );
+			float XTexCoord = ((float)s/NumSides);
+
+			for(int32 v=0; v<NumRings+1; v++)
+			{
+				int32 VIx = (NumRings+1)*s + v;
+
+				Verts[VIx].Position = ArcRot.TransformPosition( ArcVerts[v].Position );
+
+				Verts[VIx].SetTangents(
+					ArcRot.TransformVector( ArcVerts[v].TangentX ),
+					ArcRot.TransformVector( ArcVerts[v].GetTangentY() ),
+					ArcRot.TransformVector( ArcVerts[v].TangentZ )
+					);
+
+				Verts[VIx].TextureCoordinate.X = XTexCoord;
+				Verts[VIx].TextureCoordinate.Y = ArcVerts[v].TextureCoordinate.Y;
+			}
+		}
+
+		// Add all of the vertices we generated to the mesh builder.
+		for(int32 VertIdx=0; VertIdx < NumVerts; VertIdx++)
+		{
+			MeshBuilder.AddVertex(Verts[VertIdx]);
+		}
+
+		// Add all of the triangles we generated to the mesh builder.
+		for(int32 s=0; s<NumSides; s++)
+		{
+			int32 a0start = (s+0) * (NumRings+1);
+			int32 a1start = (s+1) * (NumRings+1);
+
+			for(int32 r=0; r<NumRings; r++)
+			{
+				MeshBuilder.AddTriangle(a0start + r + 0, a1start + r + 0, a0start + r + 1);
+				MeshBuilder.AddTriangle(a1start + r + 0, a1start + r + 1, a0start + r + 1);
+			}
+		}
+
+		// Free our local copy of verts and arc verts
+		FMemory::Free(Verts);
+		FMemory::Free(ArcVerts);
+	}
+	MeshBuilder.GetMesh(FScaleMatrix( Radii ) * FTranslationMatrix( Center ), MaterialRenderProxy, DepthPriority, bDisableBackfaceCulling, false, ViewIndex, Collector);
+}
 
 
 void DrawSphere(FPrimitiveDrawInterface* PDI,const FVector& Center,const FVector& Radii,int32 NumSides,int32 NumRings,const FMaterialRenderProxy* MaterialRenderProxy,uint8 DepthPriority,bool bDisableBackfaceCulling)
@@ -441,6 +577,27 @@ void BuildCylinderVerts(const FVector& Base, const FVector& XAxis, const FVector
 
 }
 
+
+void GetCylinderMesh(const FVector& Base, const FVector& XAxis, const FVector& YAxis, const FVector& ZAxis,
+				  float Radius, float HalfHeight, int32 Sides, const FMaterialRenderProxy* MaterialRenderProxy, uint8 DepthPriority, int32 ViewIndex, FMeshElementCollector& Collector)
+{
+	GetCylinderMesh( FMatrix::Identity, Base, XAxis, YAxis, ZAxis, Radius, HalfHeight, Sides, MaterialRenderProxy, DepthPriority, ViewIndex, Collector );
+}
+
+void GetCylinderMesh(const FMatrix& CylToWorld, const FVector& Base, const FVector& XAxis, const FVector& YAxis, const FVector& ZAxis, float Radius, float HalfHeight, int32 Sides, const FMaterialRenderProxy* MaterialRenderProxy, uint8 DepthPriority, int32 ViewIndex, FMeshElementCollector& Collector)
+{
+	TArray<FDynamicMeshVertex> MeshVerts;
+	TArray<int32> MeshIndices;
+	BuildCylinderVerts(Base, XAxis, YAxis, ZAxis, Radius, HalfHeight, Sides, MeshVerts, MeshIndices);
+
+
+	FDynamicMeshBuilder MeshBuilder;
+	MeshBuilder.AddVertices(MeshVerts);
+	MeshBuilder.AddTriangles(MeshIndices);
+
+	MeshBuilder.GetMesh(CylToWorld, MaterialRenderProxy, DepthPriority, false, false, ViewIndex, Collector);
+}
+
 void DrawCylinder(FPrimitiveDrawInterface* PDI,const FVector& Base, const FVector& XAxis, const FVector& YAxis, const FVector& ZAxis,
 	float Radius, float HalfHeight, int32 Sides, const FMaterialRenderProxy* MaterialRenderProxy, uint8 DepthPriority)
 {
@@ -1218,28 +1375,28 @@ FLinearColor GetSelectionColor(const FLinearColor& BaseColor,bool bSelected,bool
 	return ret;
 }
 
-bool IsRichView(const FSceneView* View)
+bool IsRichView(const FSceneViewFamily& ViewFamily)
 {
 	// Flags which make the view rich when absent.
-	if( !View->Family->EngineShowFlags.LOD ||
-		!View->Family->EngineShowFlags.IndirectLightingCache ||
-		!View->Family->EngineShowFlags.Lighting ||
-		!View->Family->EngineShowFlags.Materials)
+	if( !ViewFamily.EngineShowFlags.LOD ||
+		!ViewFamily.EngineShowFlags.IndirectLightingCache ||
+		!ViewFamily.EngineShowFlags.Lighting ||
+		!ViewFamily.EngineShowFlags.Materials)
 	{
 		return true;
 	}
 
 	// Flags which make the view rich when present.
-	if( View->Family->EngineShowFlags.LightComplexity ||
-		View->Family->EngineShowFlags.ShaderComplexity ||
-		View->Family->EngineShowFlags.StationaryLightOverlap ||
-		View->Family->EngineShowFlags.BSPSplit ||
-		View->Family->EngineShowFlags.LightMapDensity ||
-		View->Family->EngineShowFlags.PropertyColoration ||
-		View->Family->EngineShowFlags.MeshEdges ||
-		View->Family->EngineShowFlags.LightInfluences ||
-		View->Family->EngineShowFlags.Wireframe ||
-		View->Family->EngineShowFlags.LevelColoration)
+	if( ViewFamily.EngineShowFlags.LightComplexity ||
+		ViewFamily.EngineShowFlags.ShaderComplexity ||
+		ViewFamily.EngineShowFlags.StationaryLightOverlap ||
+		ViewFamily.EngineShowFlags.BSPSplit ||
+		ViewFamily.EngineShowFlags.LightMapDensity ||
+		ViewFamily.EngineShowFlags.PropertyColoration ||
+		ViewFamily.EngineShowFlags.MeshEdges ||
+		ViewFamily.EngineShowFlags.LightInfluences ||
+		ViewFamily.EngineShowFlags.Wireframe ||
+		ViewFamily.EngineShowFlags.LevelColoration)
 	{
 		return true;
 	}
@@ -1247,6 +1404,226 @@ bool IsRichView(const FSceneView* View)
 	return false;
 }
 
+void ApplyViewModeOverrides(
+	int32 ViewIndex,
+	const FEngineShowFlags& EngineShowFlags,
+	ERHIFeatureLevel::Type FeatureLevel,
+	const FPrimitiveSceneProxy* PrimitiveSceneProxy,
+	bool bSelected,
+	FMeshBatch& Mesh,
+	FMeshElementCollector& Collector
+	)
+{
+#if !(UE_BUILD_SHIPPING || UE_BUILD_TEST)
+
+	// If debug viewmodes are not allowed, skip all of the debug viewmode handling
+	if (!AllowDebugViewmodes())
+	{
+		return;
+	}
+
+	if (EngineShowFlags.Wireframe)
+	{
+		// In wireframe mode, draw the edges of the mesh with the specified wireframe color, or
+		// with the level or property color if level or property coloration is enabled.
+		FLinearColor BaseColor( PrimitiveSceneProxy->GetWireframeColor() );
+		if (EngineShowFlags.PropertyColoration)
+		{
+			BaseColor = PrimitiveSceneProxy->GetPropertyColor();
+		}
+		else if (EngineShowFlags.LevelColoration)
+		{
+			BaseColor = PrimitiveSceneProxy->GetLevelColor();
+		}
+
+		if (Mesh.MaterialRenderProxy->GetMaterial(FeatureLevel)->MaterialModifiesMeshPosition_RenderThread())
+		{
+			// If the material is mesh-modifying, we cannot rely on substitution
+			auto WireframeMaterialInstance = new FOverrideSelectionColorMaterialRenderProxy(
+				Mesh.MaterialRenderProxy,
+				GetSelectionColor( BaseColor, bSelected, Mesh.MaterialRenderProxy->IsHovered(), /*bUseOverlayIntensity=*/false)
+				);
+
+			Mesh.bWireframe = true;
+			Mesh.MaterialRenderProxy = WireframeMaterialInstance;
+			Collector.RegisterOneFrameMaterialProxy(WireframeMaterialInstance);
+		}
+		else
+		{
+			auto WireframeMaterialInstance = new FColoredMaterialRenderProxy(
+				GEngine->WireframeMaterial->GetRenderProxy(Mesh.MaterialRenderProxy->IsSelected(), Mesh.MaterialRenderProxy->IsHovered()),
+				GetSelectionColor( BaseColor, bSelected, Mesh.MaterialRenderProxy->IsHovered(), /*bUseOverlayIntensity=*/false)
+				);
+
+			Mesh.bWireframe = true;
+			Mesh.MaterialRenderProxy = WireframeMaterialInstance;
+			Collector.RegisterOneFrameMaterialProxy(WireframeMaterialInstance);
+		}
+	}
+	else if (EngineShowFlags.LightComplexity)
+	{
+		// Don't render unlit translucency when in 'light complexity' viewmode.
+		if (!Mesh.IsTranslucent(FeatureLevel) || Mesh.MaterialRenderProxy->GetMaterial(FeatureLevel)->GetShadingModel() != MSM_Unlit)
+		{
+			// Count the number of lights interacting with this primitive.
+			int32 NumDynamicLights = GetRendererModule().GetNumDynamicLightsAffectingPrimitive(PrimitiveSceneProxy->GetPrimitiveSceneInfo(),Mesh.LCI);
+
+			// Get a colored material to represent the number of lights.
+			// Some component types (BSP) have multiple FLightCacheInterface's per component, so make sure the whole component represents the number of dominant lights affecting
+			NumDynamicLights = FMath::Min( NumDynamicLights, GEngine->LightComplexityColors.Num() - 1 );
+			const FColor Color = GEngine->LightComplexityColors[NumDynamicLights];
+
+			auto LightComplexityMaterialInstance = new FColoredMaterialRenderProxy(
+				GEngine->LevelColorationUnlitMaterial->GetRenderProxy(Mesh.MaterialRenderProxy->IsSelected(), Mesh.MaterialRenderProxy->IsHovered()),
+				Color
+				);
+
+			// Draw the mesh colored by light complexity.
+			Mesh.MaterialRenderProxy = LightComplexityMaterialInstance;
+			Collector.RegisterOneFrameMaterialProxy(LightComplexityMaterialInstance);
+		}
+	}
+	else if (!EngineShowFlags.Materials)
+	{
+		// Don't render unlit translucency when in 'lighting only' viewmode.
+		if (Mesh.MaterialRenderProxy->GetMaterial(FeatureLevel)->GetShadingModel() != MSM_Unlit
+			// Don't render translucency in 'lighting only', since the viewmode works by overriding with an opaque material
+			// This would cause a mismatch of the material's blend mode with the primitive's view relevance,
+			// And make faint particles block the view
+			&& !IsTranslucentBlendMode(Mesh.MaterialRenderProxy->GetMaterial(FeatureLevel)->GetBlendMode()))
+		{
+			// When materials aren't shown, apply the same basic material to all meshes.
+			bool bTextureMapped = false;
+			FVector2D LMResolution;
+
+			if (EngineShowFlags.LightMapDensity &&
+				Mesh.LCI &&
+				(Mesh.LCI->GetLightMapInteraction().GetType() == LMIT_Texture) &&
+				Mesh.LCI->GetLightMapInteraction().GetTexture())
+			{
+				LMResolution.X = Mesh.LCI->GetLightMapInteraction().GetTexture()->GetSizeX();
+				LMResolution.Y = Mesh.LCI->GetLightMapInteraction().GetTexture()->GetSizeY();
+				bTextureMapped = true;
+			}
+
+			if (bTextureMapped == false)
+			{
+				FMaterialRenderProxy* RenderProxy = GEngine->LevelColorationLitMaterial->GetRenderProxy(Mesh.MaterialRenderProxy->IsSelected(),Mesh.MaterialRenderProxy->IsHovered());
+				auto LightingOnlyMaterialInstance = new FColoredMaterialRenderProxy(
+					RenderProxy,
+					GEngine->LightingOnlyBrightness
+					);
+
+				Mesh.MaterialRenderProxy = LightingOnlyMaterialInstance;
+				Collector.RegisterOneFrameMaterialProxy(LightingOnlyMaterialInstance);
+			}
+			else
+			{
+				FMaterialRenderProxy* RenderProxy = GEngine->LightingTexelDensityMaterial->GetRenderProxy(Mesh.MaterialRenderProxy->IsSelected(),Mesh.MaterialRenderProxy->IsHovered());
+				auto LightingDensityMaterialInstance = new FLightingDensityMaterialRenderProxy(
+					RenderProxy,
+					GEngine->LightingOnlyBrightness,
+					LMResolution
+					);
+
+				Mesh.MaterialRenderProxy = LightingDensityMaterialInstance;
+				Collector.RegisterOneFrameMaterialProxy(LightingDensityMaterialInstance);
+			}
+		}
+	}
+	else
+	{	
+		if (EngineShowFlags.PropertyColoration)
+		{
+			// In property coloration mode, override the mesh's material with a color that was chosen based on property value.
+			const UMaterial* PropertyColorationMaterial = EngineShowFlags.Lighting ? GEngine->LevelColorationLitMaterial : GEngine->LevelColorationUnlitMaterial;
+
+			auto PropertyColorationMaterialInstance = new FColoredMaterialRenderProxy(
+				PropertyColorationMaterial->GetRenderProxy(Mesh.MaterialRenderProxy->IsSelected(), Mesh.MaterialRenderProxy->IsHovered()),
+				GetSelectionColor(PrimitiveSceneProxy->GetPropertyColor(),bSelected,Mesh.MaterialRenderProxy->IsHovered())
+				);
+
+			Mesh.MaterialRenderProxy = PropertyColorationMaterialInstance;
+			Collector.RegisterOneFrameMaterialProxy(PropertyColorationMaterialInstance);
+		}
+		else if (EngineShowFlags.LevelColoration)
+		{
+			const UMaterial* LevelColorationMaterial = EngineShowFlags.Lighting ? GEngine->LevelColorationLitMaterial : GEngine->LevelColorationUnlitMaterial;
+			// Draw the mesh with level coloration.
+			auto LevelColorationMaterialInstance = new FColoredMaterialRenderProxy(
+				LevelColorationMaterial->GetRenderProxy(Mesh.MaterialRenderProxy->IsSelected(), Mesh.MaterialRenderProxy->IsHovered()),
+				GetSelectionColor(PrimitiveSceneProxy->GetLevelColor(),bSelected,Mesh.MaterialRenderProxy->IsHovered())
+				);
+			Mesh.MaterialRenderProxy = LevelColorationMaterialInstance;
+			Collector.RegisterOneFrameMaterialProxy(LevelColorationMaterialInstance);
+		}
+		else if (EngineShowFlags.BSPSplit
+			&& PrimitiveSceneProxy->ShowInBSPSplitViewmode())
+		{
+			// Determine unique color for model component.
+			FLinearColor BSPSplitColor;
+			FRandomStream RandomStream(GetTypeHash(PrimitiveSceneProxy->GetPrimitiveComponentId().Value));
+			BSPSplitColor.R = RandomStream.GetFraction();
+			BSPSplitColor.G = RandomStream.GetFraction();
+			BSPSplitColor.B = RandomStream.GetFraction();
+			BSPSplitColor.A = 1.0f;
+
+			// Piggy back on the level coloration material.
+			const UMaterial* BSPSplitMaterial = EngineShowFlags.Lighting ? GEngine->LevelColorationLitMaterial : GEngine->LevelColorationUnlitMaterial;
+			
+			// Draw BSP mesh with unique color for each model component.
+			auto BSPSplitMaterialInstance = new FColoredMaterialRenderProxy(
+				BSPSplitMaterial->GetRenderProxy(Mesh.MaterialRenderProxy->IsSelected(), Mesh.MaterialRenderProxy->IsHovered()),
+				GetSelectionColor(BSPSplitColor,bSelected,Mesh.MaterialRenderProxy->IsHovered())
+				);
+			Mesh.MaterialRenderProxy = BSPSplitMaterialInstance;
+			Collector.RegisterOneFrameMaterialProxy(BSPSplitMaterialInstance);
+		}
+		else if (PrimitiveSceneProxy->HasStaticLighting() && !PrimitiveSceneProxy->HasValidSettingsForStaticLighting())
+		{
+			auto InvalidSettingsMaterialInstance = new FColoredMaterialRenderProxy(
+				GEngine->InvalidLightmapSettingsMaterial->GetRenderProxy(bSelected),
+				GetSelectionColor(PrimitiveSceneProxy->GetLevelColor(),bSelected,Mesh.MaterialRenderProxy->IsHovered())
+				);
+			Mesh.MaterialRenderProxy = InvalidSettingsMaterialInstance;
+			Collector.RegisterOneFrameMaterialProxy(InvalidSettingsMaterialInstance);
+		}
+
+		//Draw a wireframe overlay last, if requested
+		if (EngineShowFlags.MeshEdges)
+		{
+			FMeshBatch& MeshEdgeElement = Collector.AllocateMesh();
+
+			// Draw the mesh's edges in blue, on top of the base geometry.
+			if (MeshEdgeElement.MaterialRenderProxy->GetMaterial(FeatureLevel)->MaterialModifiesMeshPosition_RenderThread())
+			{
+				// If the material is mesh-modifying, we cannot rely on substitution
+				auto WireframeMaterialInstance = new FOverrideSelectionColorMaterialRenderProxy(
+					MeshEdgeElement.MaterialRenderProxy,
+					PrimitiveSceneProxy->GetWireframeColor()
+					);
+
+				MeshEdgeElement.bWireframe = true;
+				MeshEdgeElement.MaterialRenderProxy = WireframeMaterialInstance;
+				Collector.RegisterOneFrameMaterialProxy(WireframeMaterialInstance);
+			}
+			else
+			{
+				auto WireframeMaterialInstance = new FColoredMaterialRenderProxy(
+					GEngine->WireframeMaterial->GetRenderProxy(MeshEdgeElement.MaterialRenderProxy->IsSelected(), MeshEdgeElement.MaterialRenderProxy->IsHovered()),
+					PrimitiveSceneProxy->GetWireframeColor()
+					);
+
+				MeshEdgeElement.bWireframe = true;
+				MeshEdgeElement.MaterialRenderProxy = WireframeMaterialInstance;
+				Collector.RegisterOneFrameMaterialProxy(WireframeMaterialInstance);
+			}
+
+			Collector.AddMesh(ViewIndex, MeshEdgeElement);
+		}
+	}
+#endif
+}
 
 /**
  * Draws a mesh, modifying the material which is used depending on the view's show flags.
diff --git a/Engine/Source/Runtime/Engine/Private/PrimitiveSceneProxy.cpp b/Engine/Source/Runtime/Engine/Private/PrimitiveSceneProxy.cpp
--- a/Engine/Source/Runtime/Engine/Private/PrimitiveSceneProxy.cpp
+++ b/Engine/Source/Runtime/Engine/Private/PrimitiveSceneProxy.cpp
@@ -8,7 +8,10 @@
 #include "PrimitiveSceneProxy.h"
 
 FPrimitiveSceneProxy::FPrimitiveSceneProxy(const UPrimitiveComponent* InComponent, FName InResourceName)
-:	DrawInGame(InComponent->bVisible && !InComponent->bHiddenInGame)
+:	WireframeColor(FLinearColor::White)
+,	LevelColor(FLinearColor::White)
+,	PropertyColor(FLinearColor::White)
+,	DrawInGame(InComponent->bVisible && !InComponent->bHiddenInGame)
 ,	DrawInEditor(InComponent->bVisible)
 ,	bReceivesDecals(InComponent->bReceivesDecals)
 ,	bOnlyOwnerSee(InComponent->bOnlyOwnerSee)
@@ -42,6 +45,7 @@ FPrimitiveSceneProxy::FPrimitiveSceneProxy(const UPrimitiveComponent* InComponen
 ,	bLightAttachmentsAsGroup(InComponent->bLightAttachmentsAsGroup)
 ,	bStaticElementsAlwaysUseProxyPrimitiveUniformBuffer(false)
 ,	bAlwaysHasVelocity(false)
+,	bUseEditorDepthTest(true)
 ,	bUseAsOccluder(InComponent->bUseAsOccluder)
 ,	bAllowApproximateOcclusion(InComponent->Mobility != EComponentMobility::Movable)
 ,	bSelectable(InComponent->bSelectable)
@@ -114,15 +118,6 @@ FPrimitiveSceneProxy::FPrimitiveSceneProxy(const UPrimitiveComponent* InComponen
 		HiddenEditorViews = InComponent->GetOwner()->HiddenEditorViews;
 #endif
 	}
-
-#if WITH_EDITOR
-	if (GIsEditor && InComponent)
-	{
-		LightMapResolutionScale = FVector2D(0.0f, 0.0f);
-		LightMapType = LMIT_None;
-		bLightMapResolutionPadded = false;
-	}
-#endif
 	
 	// 
 	// Flag components to render only after level will be fully added to the world
@@ -185,7 +180,15 @@ void FPrimitiveSceneProxy::UpdateActorPosition(FVector InActorPosition)
 			PrimitiveSceneProxy->ActorPosition = InActorPosition;
 			// Update the uniform shader parameters.
 			const FPrimitiveUniformShaderParameters PrimitiveUniformShaderParameters =
-			GetPrimitiveUniformShaderParameters(PrimitiveSceneProxy->LocalToWorld, InActorPosition, PrimitiveSceneProxy->Bounds,PrimitiveSceneProxy->LocalBounds, PrimitiveSceneProxy->bReceivesDecals, PrimitiveSceneProxy->HasDistanceFieldRepresentation(), PrimitiveSceneProxy->GetLpvBiasMultiplier() );
+			GetPrimitiveUniformShaderParameters(
+			PrimitiveSceneProxy->LocalToWorld, 
+			InActorPosition, 
+			PrimitiveSceneProxy->Bounds,
+			PrimitiveSceneProxy->LocalBounds, 
+			PrimitiveSceneProxy->bReceivesDecals, 
+			PrimitiveSceneProxy->HasDistanceFieldRepresentation(), 
+			PrimitiveSceneProxy->UseEditorDepthTest(),
+			PrimitiveSceneProxy->GetLpvBiasMultiplier() );
 
 			PrimitiveSceneProxy->UniformBuffer.SetContents(PrimitiveUniformShaderParameters);
 			PrimitiveSceneProxy->OnActorPositionChanged();
@@ -207,7 +210,7 @@ void FPrimitiveSceneProxy::SetTransform(const FMatrix& InLocalToWorld, const FBo
 	ActorPosition = InActorPosition;
 	
 	// Update the uniform shader parameters.
-	const FPrimitiveUniformShaderParameters PrimitiveUniformShaderParameters = GetPrimitiveUniformShaderParameters(LocalToWorld, ActorPosition, Bounds, LocalBounds, bReceivesDecals, HasDistanceFieldRepresentation(), LpvBiasMultiplier );
+	const FPrimitiveUniformShaderParameters PrimitiveUniformShaderParameters = GetPrimitiveUniformShaderParameters(LocalToWorld, ActorPosition, Bounds, LocalBounds, bReceivesDecals, HasDistanceFieldRepresentation(), UseEditorDepthTest(), LpvBiasMultiplier );
 	UniformBuffer.SetContents(PrimitiveUniformShaderParameters);
 	
 	// Notify the proxy's implementation of the change.
diff --git a/Engine/Source/Runtime/Engine/Private/SceneManagement.cpp b/Engine/Source/Runtime/Engine/Private/SceneManagement.cpp
--- a/Engine/Source/Runtime/Engine/Private/SceneManagement.cpp
+++ b/Engine/Source/Runtime/Engine/Private/SceneManagement.cpp
@@ -4,22 +4,157 @@
 #include "StaticMeshResources.h"
 #include "../../Renderer/Private/ScenePrivate.h"
 
-/*
-TUniformBufferRef<FPrimitiveUniformShaderParameters> CreatePrimitiveUniformBufferImmediate(
-	const FMatrix& LocalToWorld,
-	FVector BoundsOrigin,
-	FVector Bounds,
-	float BoundsRadius,
-	bool bReceivesDecals
+FSimpleElementCollector::~FSimpleElementCollector()
+{
+	// Cleanup the dynamic resources.
+	for(int32 ResourceIndex = 0;ResourceIndex < DynamicResources.Num();ResourceIndex++)
+	{
+		//release the resources before deleting, they will delete themselves
+		DynamicResources[ResourceIndex]->ReleasePrimitiveResource();
+	}
+}
+
+void FSimpleElementCollector::SetHitProxy(HHitProxy* HitProxy)
+{
+	if (HitProxy)
+	{
+		HitProxyId = HitProxy->Id;
+	}
+	else
+	{
+		HitProxyId = FHitProxyId();
+	}
+}
+
+void FSimpleElementCollector::DrawSprite(
+	const FVector& Position,
+	float SizeX,
+	float SizeY,
+	const FTexture* Sprite,
+	const FLinearColor& Color,
+	uint8 DepthPriorityGroup,
+	float U,
+	float UL,
+	float V,
+	float VL,
+	uint8 BlendMode
+	)
+{
+	BatchedElements.AddSprite(
+		Position,
+		SizeX,
+		SizeY,
+		Sprite,
+		Color,
+		HitProxyId,
+		U,
+		UL,
+		V,
+		VL,
+		BlendMode
+		);
+}
+
+void FSimpleElementCollector::DrawLine(
+	const FVector& Start,
+	const FVector& End,
+	const FLinearColor& Color,
+	uint8 DepthPriorityGroup,
+	float Thickness/* = 0.0f*/,
+	float DepthBias/* = 0.0f*/,
+	bool bScreenSpace/* = false*/
+	)
+{
+	BatchedElements.AddLine(
+		Start,
+		End,
+		Color,
+		HitProxyId,
+		Thickness,
+		DepthBias,
+		bScreenSpace
+		);
+}
+
+void FSimpleElementCollector::DrawPoint(
+	const FVector& Position,
+	const FLinearColor& Color,
+	float PointSize,
+	uint8 DepthPriorityGroup
 	)
 {
-	check(IsInRenderingThread());
-	return TUniformBufferRef<FPrimitiveUniformShaderParameters>::CreateUniformBufferImmediate(
-		GetPrimitiveUniformShaderParameters(LocalToWorld, BoundsOrigin, BoundsOrigin, Bounds, BoundsRadius, bReceivesDecals),
-		UniformBuffer_MultiFrame
+	BatchedElements.AddPoint(
+		Position,
+		PointSize,
+		Color,
+		HitProxyId
+		);
+}
+
+void FSimpleElementCollector::RegisterDynamicResource(FDynamicPrimitiveResource* DynamicResource)
+{
+	// Add the dynamic resource to the list of resources to cleanup on destruction.
+	DynamicResources.Add(DynamicResource);
+
+	// Initialize the dynamic resource immediately.
+	DynamicResource->InitPrimitiveResource();
+}
+
+void FSimpleElementCollector::DrawBatchedElements(FRHICommandList& RHICmdList, const FSceneView& View, FTexture2DRHIRef DepthTexture, EBlendModeFilter::Type Filter) const
+{
+	static auto* MobileHDRCvar = IConsoleManager::Get().FindTConsoleVariableDataInt(TEXT("r.MobileHDR"));
+	const bool bIsMobileHDR = MobileHDRCvar->GetValueOnRenderThread() == 1;
+
+	// Mobile HDR does not execute post process, so does not need to render flipped
+	const bool bNeedToSwitchVerticalAxis = RHINeedsToSwitchVerticalAxis(GRHIShaderPlatform) && !bIsMobileHDR;
+
+	// Draw the batched elements.
+	BatchedElements.Draw(
+		RHICmdList,
+		bNeedToSwitchVerticalAxis,
+		View.ViewProjectionMatrix,
+		View.ViewRect.Width(),
+		View.ViewRect.Height(),
+		View.Family->EngineShowFlags.HitProxies,
+		1.0f,
+		&View,
+		DepthTexture,
+		Filter
 		);
 }
-*/
+
+FMeshBatchAndRelevance::FMeshBatchAndRelevance(const FMeshBatch& InMesh, const FPrimitiveSceneProxy* InPrimitiveSceneProxy, ERHIFeatureLevel::Type FeatureLevel) :
+	Mesh(&InMesh),
+	PrimitiveSceneProxy(InPrimitiveSceneProxy)
+{
+	EBlendMode BlendMode = InMesh.MaterialRenderProxy->GetMaterial(FeatureLevel)->GetBlendMode();
+	bHasOpaqueOrMaskedMaterial = !IsTranslucentBlendMode(BlendMode);
+	bRenderInMainPass = PrimitiveSceneProxy->ShouldRenderInMainPass();
+}
+
+void FMeshElementCollector::AddMesh(int32 ViewIndex, FMeshBatch& MeshBatch)
+{
+	checkSlow(MeshBatch.GetNumPrimitives() > 0);
+	checkSlow(MeshBatch.VertexFactory && MeshBatch.MaterialRenderProxy);
+	checkSlow(PrimitiveSceneProxy);
+
+	if (MeshBatch.bCanApplyViewModeOverrides)
+	{
+		FSceneView* View = Views[ViewIndex];
+
+		ApplyViewModeOverrides(
+			ViewIndex,
+			View->Family->EngineShowFlags,
+			View->GetFeatureLevel(),
+			PrimitiveSceneProxy,
+			MeshBatch.bUseWireframeSelectionColoring,
+			MeshBatch,
+			*this);
+	}
+
+	TArray<FMeshBatchAndRelevance,SceneRenderingAllocator>& ViewMeshBatches = *MeshBatches[ViewIndex];
+	new (ViewMeshBatches) FMeshBatchAndRelevance(MeshBatch, PrimitiveSceneProxy, FeatureLevel);	
+}
 
 FLightMapInteraction FLightMapInteraction::Texture(
 	const class ULightMapTexture2D* const* InTextures,
diff --git a/Engine/Source/Runtime/Engine/Private/SceneView.cpp b/Engine/Source/Runtime/Engine/Private/SceneView.cpp
--- a/Engine/Source/Runtime/Engine/Private/SceneView.cpp
+++ b/Engine/Source/Runtime/Engine/Private/SceneView.cpp
@@ -1132,14 +1132,7 @@ void FSceneView::ConfigureBufferVisualizationSettings()
 
 ERHIFeatureLevel::Type FSceneView::GetFeatureLevel() const
 { 
-	if (Family->Scene)
-	{
-		return Family->Scene->GetFeatureLevel();
-	}
-	else
-	{
-		return GRHIFeatureLevel;
-	}
+	return Family->GetFeatureLevel();
 }
 
 FSceneViewFamily::FSceneViewFamily( const ConstructionValues& CVS )
@@ -1240,6 +1233,18 @@ void FSceneViewFamily::ComputeFamilySize()
 	check(bInitializedExtents);
 }
 
+ERHIFeatureLevel::Type FSceneViewFamily::GetFeatureLevel() const
+{ 
+	if (Scene)
+	{
+		return Scene->GetFeatureLevel();
+	}
+	else
+	{
+		return GRHIFeatureLevel;
+	}
+}
+
 FSceneViewFamilyContext::~FSceneViewFamilyContext()
 {
 	// Cleanup the views allocated for this view family.
diff --git a/Engine/Source/Runtime/Engine/Private/SkeletalMesh.cpp b/Engine/Source/Runtime/Engine/Private/SkeletalMesh.cpp
--- a/Engine/Source/Runtime/Engine/Private/SkeletalMesh.cpp
+++ b/Engine/Source/Runtime/Engine/Private/SkeletalMesh.cpp
@@ -3779,13 +3779,10 @@ FSkeletalMeshSceneProxy::FSkeletalMeshSceneProxy(const USkinnedMeshComponent* Co
 		,	SkelMeshResource(InSkelMeshResource)
 		,	SkeletalMeshForDebug(Component->SkeletalMesh)
 		,	PhysicsAssetForDebug(Component->GetPhysicsAsset())
-		,	LevelColor(FLinearColor::White)
-		,	PropertyColor(255,255,255)
 		,	bForceWireframe(Component->bForceWireframe)
 		,	bCanHighlightSelectedSections(Component->bCanHighlightSelectedSections)
 		,	MaterialRelevance(Component->GetMaterialRelevance())
 		,	bMaterialsNeedMorphUsage_GameThread(false)
-		,   WireframeOverlayColor(255, 255, 255, 255)
 {
 	check(MeshObject);
 	check(SkelMeshResource);
@@ -3898,7 +3895,9 @@ FSkeletalMeshSceneProxy::FSkeletalMeshSceneProxy(const USkinnedMeshComponent* Co
 	}
 
 	// Get a color for property coloration
-	GEngine->GetPropertyColorationColor( (UObject*)Component, PropertyColor );
+	FColor NewPropertyColor;
+	GEngine->GetPropertyColorationColor( (UObject*)Component, NewPropertyColor );
+	PropertyColor = NewPropertyColor;
 }
 
 // FPrimitiveSceneProxy interface.
@@ -4049,6 +4048,87 @@ HHitProxy* FSkeletalMeshSceneProxy::CreateHitProxies(UPrimitiveComponent* Compon
 }
 #endif
 
+void FSkeletalMeshSceneProxy::GetDynamicMeshElements(const TArray<const FSceneView*>& Views, const FSceneViewFamily& ViewFamily, uint32 VisibilityMap, FMeshElementCollector& Collector) const
+{
+	QUICK_SCOPE_CYCLE_COUNTER(STAT_FSkeletalMeshSceneProxy_GetMeshElements);
+	GetMeshElementsConditionallySelectable(Views, ViewFamily, true, VisibilityMap, Collector);
+}
+
+void FSkeletalMeshSceneProxy::GetMeshElementsConditionallySelectable(const TArray<const FSceneView*>& Views, const FSceneViewFamily& ViewFamily, bool bSelectable, uint32 VisibilityMap, FMeshElementCollector& Collector) const
+{
+	if( !MeshObject )
+	{
+		return;
+	}	
+
+	for (int32 ViewIndex = 0; ViewIndex < Views.Num(); ViewIndex++)
+	{
+		if (VisibilityMap & (1 << ViewIndex))
+		{
+			const FSceneView* View = Views[ViewIndex];
+			MeshObject->UpdateMinDesiredLODLevel(View, GetBounds(), GFrameNumberRenderThread);
+		}
+	}
+
+	const int32 LODIndex = MeshObject->GetLOD();
+	check(LODIndex < SkelMeshResource->LODModels.Num());
+	const FStaticLODModel& LODModel = SkelMeshResource->LODModels[LODIndex];
+
+	if( LODSections.Num() > 0 )
+	{
+		const FLODSectionElements& LODSection = LODSections[LODIndex];
+
+		check(LODSection.SectionElements.Num() == LODModel.Sections.Num());
+
+		for (FSkeletalMeshSectionIter Iter(LODIndex, *MeshObject, LODModel, LODSection); Iter; ++Iter)
+		{
+			FSkelMeshSection Section = Iter.GetSection();
+			const FSkelMeshChunk& Chunk = Iter.GetChunk();
+			const FSectionElementInfo& SectionElementInfo = Iter.GetSectionElementInfo();
+			const FTwoVectors& CustomLeftRightVectors = Iter.GetCustomLeftRightVectors();
+
+#if WITH_EDITORONLY_DATA
+			// TODO: This is not threadsafe! A render command should be used to propagate SelectedEditorSection to the scene proxy.
+			Section.bSelected = (SkeletalMeshForDebug->SelectedEditorSection == SectionElementInfo.UseMaterialIndex);
+#endif
+			// If hidden skip the draw
+			if (MeshObject->IsMaterialHidden(LODIndex, SectionElementInfo.UseMaterialIndex))
+			{
+				continue;
+			}
+
+			// If disabled, then skip the draw
+			if(Section.bDisabled)
+			{
+				continue;
+			}
+
+			GetDynamicElementsSection(Views, ViewFamily, VisibilityMap, LODModel, LODIndex, Section, Chunk, SectionElementInfo, CustomLeftRightVectors, bSelectable, Collector);
+		}
+	}
+
+#if !(UE_BUILD_SHIPPING || UE_BUILD_TEST)
+	for (int32 ViewIndex = 0; ViewIndex < Views.Num(); ViewIndex++)
+	{
+		if (VisibilityMap & (1 << ViewIndex))
+		{
+			// debug drawing
+			TArray<FTransform>* BoneSpaceBases = MeshObject->GetSpaceBases();
+
+			if( PhysicsAssetForDebug )
+			{
+				DebugDrawPhysicsAsset(Collector.GetPDI(ViewIndex),ViewFamily.EngineShowFlags);
+			}
+
+			if (ViewFamily.EngineShowFlags.SkeletalMeshes)
+			{
+				RenderBounds(Collector.GetPDI(ViewIndex), ViewFamily.EngineShowFlags, GetBounds(), IsSelected());
+			}
+		}
+	}
+#endif
+}
+
 /** 
 * Draw the scene proxy as a dynamic element
 *
@@ -4114,7 +4194,7 @@ void FSkeletalMeshSceneProxy::DrawDynamicElements(FPrimitiveDrawInterface* PDI,
 
 	if( PhysicsAssetForDebug )
 	{
-		DebugDrawPhysicsAsset(PDI,View);
+		DebugDrawPhysicsAsset(PDI,View->Family->EngineShowFlags);
 	}
 
 	if (View->Family->EngineShowFlags.SkeletalMeshes)
@@ -4124,46 +4204,141 @@ void FSkeletalMeshSceneProxy::DrawDynamicElements(FPrimitiveDrawInterface* PDI,
 #endif
 }
 
-/** 
-* Draw only one section of the scene proxy as a dynamic element
-*
-* @param	PDI - draw interface to render to
-* @param	View - current view
-* @param 	ForceLOD - Force this LOD. If -1, use current LOD of mesh. 
-* @param	InMaterial - which material section to draw
-*/
-void FSkeletalMeshSceneProxy::DrawDynamicElementsByMaterial(FPrimitiveDrawInterface* PDI,const FSceneView* View, int32 ForceLOD, int32 InMaterial)
+void FSkeletalMeshSceneProxy::GetDynamicElementsSection(const TArray<const FSceneView*>& Views, const FSceneViewFamily& ViewFamily, uint32 VisibilityMap, 
+	const FStaticLODModel& LODModel, const int32 LODIndex, const FSkelMeshSection& Section, const FSkelMeshChunk& Chunk, 
+	const FSectionElementInfo& SectionElementInfo, const FTwoVectors& CustomLeftRightVectors, bool bSelectable, FMeshElementCollector& Collector ) const
 {
-	QUICK_SCOPE_CYCLE_COUNTER( STAT_SkeletalMeshSceneProxy_DrawDynamicElementsByMaterial );
-
-	// Check if this has valid DynamicData
-	if( !MeshObject || !MeshObject->HaveValidDynamicData())
+	// If hidden skip the draw
+	if (Section.bDisabled || MeshObject->IsMaterialHidden(LODIndex,SectionElementInfo.UseMaterialIndex))
 	{
 		return;
-	}		  
-
-	const int32 LODIndex = ( ForceLOD<0 )? MeshObject->GetLOD() : ForceLOD;
-	check(LODIndex < SkelMeshResource->LODModels.Num());
-	const FStaticLODModel& LODModel = SkelMeshResource->LODModels[LODIndex];
-	const FLODSectionElements& LODSection = LODSections[LODIndex];
+	}
 
-	check(LODSection.SectionElements.Num()==LODModel.Sections.Num());
+#if !WITH_EDITOR
+	const bool bIsSelected = false;
+#else // #if !WITH_EDITOR
+	bool bIsSelected = IsSelected();
 
-	for (FSkeletalMeshSectionIter Iter(LODIndex,*MeshObject,LODModel,LODSection); Iter; ++Iter)
+	// if the mesh isn't selected but the mesh section is selected in the AnimSetViewer, find the mesh component and make sure that it can be highlighted (ie. are we rendering for the AnimSetViewer or not?)
+	if( !bIsSelected && Section.bSelected && bCanHighlightSelectedSections )
 	{
-		const FSkelMeshSection& Section = Iter.GetSection();
-		const FSkelMeshChunk& Chunk = Iter.GetChunk();
-		const FSectionElementInfo& SectionElementInfo = Iter.GetSectionElementInfo();
-		const FTwoVectors& CustomLeftRightVectors = Iter.GetCustomLeftRightVectors();
+		bIsSelected = true;
+	}
+#endif // #if WITH_EDITOR
+
+	const bool bIsWireframe = ViewFamily.EngineShowFlags.Wireframe;
+	const ERHIFeatureLevel::Type FeatureLevel = ViewFamily.GetFeatureLevel();
 
-		if (SectionElementInfo.UseMaterialIndex != InMaterial)
+	for (int32 ViewIndex = 0; ViewIndex < Views.Num(); ViewIndex++)
+	{
+		if (VisibilityMap & (1 << ViewIndex))
 		{
-			continue;
-		}
+			const FSceneView* View = Views[ViewIndex];
+
+			FMeshBatch& Mesh = Collector.AllocateMesh();
+			FMeshBatchElement& BatchElement = Mesh.Elements[0];
+			Mesh.DynamicVertexData = NULL;
+			Mesh.UseDynamicData = false;
+			Mesh.LCI = NULL;
+			Mesh.bWireframe |= bForceWireframe;
+			Mesh.Type = PT_TriangleList;
+			Mesh.VertexFactory = MeshObject->GetVertexFactory(LODIndex,Section.ChunkIndex);
+			Mesh.bSelectable = bSelectable;
+			BatchElement.FirstIndex = Section.BaseIndex;
+
+			BatchElement.IndexBuffer = LODModel.MultiSizeIndexContainer.GetIndexBuffer();
+			BatchElement.MaxVertexIndex = LODModel.NumVertices - 1;
 
-		DrawDynamicElementsSection(PDI, View, LODModel, LODIndex, Section, Chunk, SectionElementInfo, CustomLeftRightVectors );
+			BatchElement.UserIndex = MeshObject->GPUSkinCacheKeys[Section.ChunkIndex];
+
+			const bool bRequiresAdjacencyInformation = RequiresAdjacencyInformation( SectionElementInfo.Material, Mesh.VertexFactory->GetType(), ViewFamily.GetFeatureLevel() );
+			if ( bRequiresAdjacencyInformation )
+			{
+				check( LODModel.AdjacencyMultiSizeIndexContainer.IsIndexBufferValid() );
+				BatchElement.IndexBuffer = LODModel.AdjacencyMultiSizeIndexContainer.GetIndexBuffer();
+				Mesh.Type = PT_12_ControlPointPatchList;
+				BatchElement.FirstIndex *= 4;
+			}
+
+		#if WITH_EDITOR
+
+			Mesh.BatchHitProxyId = SectionElementInfo.HitProxy ? SectionElementInfo.HitProxy->Id : FHitProxyId();
+
+			if (Section.bSelected)
+			{
+				auto SelectionOverrideProxy = new FOverrideSelectionColorMaterialRenderProxy(
+					SectionElementInfo.Material->GetRenderProxy(bIsSelected, IsHovered()),
+					GetSelectionColor(GEngine->GetSelectedMaterialColor(), bIsSelected, IsHovered())
+					);
+
+				Collector.RegisterOneFrameMaterialProxy(SelectionOverrideProxy);
+
+				Mesh.MaterialRenderProxy = SelectionOverrideProxy;
+			}
+			else
+			{
+				Mesh.MaterialRenderProxy = SectionElementInfo.Material->GetRenderProxy(bIsSelected, IsHovered());
+			}
+		#else
+			Mesh.MaterialRenderProxy = SectionElementInfo.Material->GetRenderProxy(bIsSelected, IsHovered());
+		#endif
+
+			BatchElement.PrimitiveUniformBufferResource = &GetUniformBuffer();
+
+			// Select which indices to use if TRISORT_CustomLeftRight
+			if( Section.TriangleSorting == TRISORT_CustomLeftRight )
+			{
+				switch( MeshObject->CustomSortAlternateIndexMode )
+				{
+				case CSAIM_Left:
+					// Left view - use second set of indices.
+					BatchElement.FirstIndex += Section.NumTriangles * 3;
+					break;
+				case  CSAIM_Right:
+					// Right view - use first set of indices.
+					break;
+				default:
+					// Calculate the viewing direction
+					FVector SortWorldOrigin = GetLocalToWorld().TransformPosition(CustomLeftRightVectors.v1);
+					FVector SortWorldDirection = GetLocalToWorld().TransformVector(CustomLeftRightVectors.v2);
+
+					if( (SortWorldDirection | (SortWorldOrigin - View->ViewMatrices.ViewOrigin)) < 0.f )
+					{
+						BatchElement.FirstIndex += Section.NumTriangles * 3;
+					}
+					break;
+				}
+			}
+
+			BatchElement.NumPrimitives = Section.NumTriangles;
+			if( GIsEditor && MeshObject->ProgressiveDrawingFraction != 1.f )
+			{
+				if (Mesh.MaterialRenderProxy->GetMaterial(GRHIFeatureLevel)->GetBlendMode() == BLEND_Translucent)
+				{
+					BatchElement.NumPrimitives = FMath::RoundToInt(((float)Section.NumTriangles)*FMath::Clamp<float>(MeshObject->ProgressiveDrawingFraction,0.f,1.f));
+					if( BatchElement.NumPrimitives == 0 )
+					{
+						continue;
+					}
+				}
+			}
+			BatchElement.MinVertexIndex = Chunk.BaseVertexIndex;
+			Mesh.ReverseCulling = IsLocalToWorldDeterminantNegative();
+			Mesh.CastShadow = SectionElementInfo.bEnableShadowCasting;
+
+			Mesh.bCanApplyViewModeOverrides = true;
+			Mesh.bUseWireframeSelectionColoring = bIsSelected;
+
+			Collector.AddMesh(ViewIndex, Mesh);
+
+			const int32 NumVertices = Chunk.NumRigidVertices + Chunk.NumSoftVertices;
+			INC_DWORD_STAT_BY(STAT_GPUSkinVertices,(uint32)(bIsCPUSkinned ? 0 : NumVertices));
+			INC_DWORD_STAT_BY(STAT_SkelMeshTriangles,Mesh.GetNumPrimitives());
+			INC_DWORD_STAT(STAT_SkelMeshDrawCalls);
+		}
 	}
 }
+
 void FSkeletalMeshSceneProxy::DrawDynamicElementsSection(FPrimitiveDrawInterface* PDI,const FSceneView* View,
 														const FStaticLODModel& LODModel, const int32 LODIndex, const FSkelMeshSection& Section, 
 														const FSkelMeshChunk& Chunk, const FSectionElementInfo& SectionElementInfo, const FTwoVectors& CustomLeftRightVectors )
@@ -4181,7 +4356,6 @@ void FSkeletalMeshSceneProxy::DrawDynamicElementsSection(FPrimitiveDrawInterface
 #endif // #if WITH_EDITOR
 
 	const bool bIsWireframe = View->Family->EngineShowFlags.Wireframe;
-	FLinearColor WireframeLinearColor(WireframeOverlayColor);
 
 	if(Section.bDisabled)
 		return;
@@ -4292,7 +4466,7 @@ void FSkeletalMeshSceneProxy::DrawDynamicElementsSection(FPrimitiveDrawInterface
 	const int32 NumPasses = DrawRichMesh(
 		PDI,
 		Mesh,
-		WireframeLinearColor,
+		WireframeColor,
 		LevelColor,
 		PropertyColor,
 		this,
@@ -4308,11 +4482,10 @@ void FSkeletalMeshSceneProxy::DrawDynamicElementsSection(FPrimitiveDrawInterface
 
 /**
  * Returns the world transform to use for drawing.
- * @param View - Current view
  * @param OutLocalToWorld - Will contain the local-to-world transform when the function returns.
  * @param OutWorldToLocal - Will contain the world-to-local transform when the function returns.
  */
-void FSkeletalMeshSceneProxy::GetWorldMatrices( const FSceneView* View, FMatrix& OutLocalToWorld, FMatrix& OutWorldToLocal )
+void FSkeletalMeshSceneProxy::GetWorldMatrices( FMatrix& OutLocalToWorld, FMatrix& OutWorldToLocal ) const
 {
 	OutLocalToWorld = GetLocalToWorld();
 	OutWorldToLocal = GetLocalToWorld().Inverse();
@@ -4356,10 +4529,10 @@ int32 FSkeletalMeshSceneProxy::GetCurrentLODIndex()
 /** 
  * Render physics asset for debug display
  */
-void FSkeletalMeshSceneProxy::DebugDrawPhysicsAsset(FPrimitiveDrawInterface* PDI,const FSceneView* View)
+void FSkeletalMeshSceneProxy::DebugDrawPhysicsAsset(FPrimitiveDrawInterface* PDI, const FEngineShowFlags& EngineShowFlags) const
 {
 	FMatrix LocalToWorld, WorldToLocal;
-	GetWorldMatrices( View, LocalToWorld, WorldToLocal );
+	GetWorldMatrices( LocalToWorld, WorldToLocal );
 
 	FMatrix ScalingMatrix = LocalToWorld;
 	FVector TotalScale = ScalingMatrix.ExtractScaling();
@@ -4373,11 +4546,11 @@ void FSkeletalMeshSceneProxy::DebugDrawPhysicsAsset(FPrimitiveDrawInterface* PDI
 		if(BoneSpaceBases)
 		{
 			check(PhysicsAssetForDebug);
-			if( View->Family->EngineShowFlags.Collision && IsCollisionEnabled() )
+			if( EngineShowFlags.Collision && IsCollisionEnabled() )
 			{
 				PhysicsAssetForDebug->DrawCollision(PDI, SkeletalMeshForDebug, *BoneSpaceBases, LocalToWorldTransform, TotalScale.X);
 			}
-			if( View->Family->EngineShowFlags.Constraints )
+			if( EngineShowFlags.Constraints )
 			{
 				PhysicsAssetForDebug->DrawConstraints(PDI, SkeletalMeshForDebug, *BoneSpaceBases, LocalToWorldTransform, TotalScale.X);
 			}
diff --git a/Engine/Source/Runtime/Engine/Private/SphereComponent.cpp b/Engine/Source/Runtime/Engine/Private/SphereComponent.cpp
--- a/Engine/Source/Runtime/Engine/Private/SphereComponent.cpp
+++ b/Engine/Source/Runtime/Engine/Private/SphereComponent.cpp
@@ -90,7 +90,45 @@ FPrimitiveSceneProxy* USphereComponent::CreateSceneProxy()
 		}
 
 		  // FPrimitiveSceneProxy interface.
-		  
+		
+		virtual void GetDynamicMeshElements(const TArray<const FSceneView*>& Views, const FSceneViewFamily& ViewFamily, uint32 VisibilityMap, FMeshElementCollector& Collector) const override
+		{
+			QUICK_SCOPE_CYCLE_COUNTER( STAT_SphereSceneProxy_GetDynamicMeshElements );
+
+			for (int32 ViewIndex = 0; ViewIndex < Views.Num(); ViewIndex++)
+			{
+				if (VisibilityMap & (1 << ViewIndex))
+				{
+					const FSceneView* View = Views[ViewIndex];
+					FPrimitiveDrawInterface* PDI = Collector.GetPDI(ViewIndex);
+
+					const FMatrix& LocalToWorld = GetLocalToWorld();
+					int32 SphereSides =  FMath::Clamp<int32>(SphereRadius/4.f, 16, 64);
+					if(ShapeMaterial && !View->Family->EngineShowFlags.Wireframe)
+					{
+						GetSphereMesh(LocalToWorld.GetOrigin(), FVector(SphereRadius), SphereSides, SphereSides/2, ShapeMaterial->GetRenderProxy(false), SDPG_World, false, ViewIndex, Collector);
+					}
+					else
+					{
+						const FLinearColor DrawSphereColor = GetSelectionColor(SphereColor, IsSelected(), IsHovered(), /*bUseOverlayIntensity=*/false);
+
+						float AbsScaleX = LocalToWorld.GetScaledAxis(EAxis::X).Size();
+						float AbsScaleY = LocalToWorld.GetScaledAxis(EAxis::Y).Size();
+						float AbsScaleZ = LocalToWorld.GetScaledAxis(EAxis::Z).Size();
+						float MinAbsScale = FMath::Min3(AbsScaleX, AbsScaleY, AbsScaleZ);
+
+						FVector ScaledX = LocalToWorld.GetUnitAxis(EAxis::X) * MinAbsScale;
+						FVector ScaledY = LocalToWorld.GetUnitAxis(EAxis::Y) * MinAbsScale;
+						FVector ScaledZ = LocalToWorld.GetUnitAxis(EAxis::Z) * MinAbsScale;
+
+						DrawCircle(PDI, LocalToWorld.GetOrigin(), ScaledX, ScaledY, DrawSphereColor, SphereRadius, SphereSides, SDPG_World);
+						DrawCircle(PDI, LocalToWorld.GetOrigin(), ScaledX, ScaledZ, DrawSphereColor, SphereRadius, SphereSides, SDPG_World);
+						DrawCircle(PDI, LocalToWorld.GetOrigin(), ScaledY, ScaledZ, DrawSphereColor, SphereRadius, SphereSides, SDPG_World);
+					}
+				}
+			}
+		}
+
 		/** 
 		* Draw the scene proxy as a dynamic element
 		*
diff --git a/Engine/Source/Runtime/Engine/Private/SplineMeshSceneProxy.cpp b/Engine/Source/Runtime/Engine/Private/SplineMeshSceneProxy.cpp
--- a/Engine/Source/Runtime/Engine/Private/SplineMeshSceneProxy.cpp
+++ b/Engine/Source/Runtime/Engine/Private/SplineMeshSceneProxy.cpp
@@ -50,40 +50,40 @@ FSplineMeshSceneProxy::~FSplineMeshSceneProxy()
 	LODResources.Empty();
 }
 
-bool FSplineMeshSceneProxy::GetShadowMeshElements(int32 LODIndex, uint8 InDepthPriorityGroup, TArray<FMeshBatch, TInlineAllocator<1>>& OutMeshBatches) const
+bool FSplineMeshSceneProxy::GetShadowMeshElement(int32 LODIndex, int32 BatchIndex, uint8 InDepthPriorityGroup, FMeshBatch& OutMeshBatch) const
 {
 	//checkf(LODIndex == 0, TEXT("Getting spline static mesh element with invalid LOD [%d]"), LODIndex);
 
-	if (FStaticMeshSceneProxy::GetShadowMeshElements(LODIndex, InDepthPriorityGroup, OutMeshBatches))
+	if (FStaticMeshSceneProxy::GetShadowMeshElement(LODIndex, BatchIndex, InDepthPriorityGroup, OutMeshBatch))
 	{
-		OutMeshBatches[0].VertexFactory = LODResources[LODIndex].VertexFactory;
-		OutMeshBatches[0].ReverseCulling ^= (SplineParams.StartScale.X < 0) ^ (SplineParams.StartScale.Y < 0);
+		OutMeshBatch.VertexFactory = LODResources[LODIndex].VertexFactory;
+		OutMeshBatch.ReverseCulling ^= (SplineParams.StartScale.X < 0) ^ (SplineParams.StartScale.Y < 0);
 		return true;
 	}
 	return false;
 }
 
-bool FSplineMeshSceneProxy::GetMeshElements(int32 LODIndex, int32 SectionIndex, uint8 InDepthPriorityGroup, TArray<FMeshBatch, TInlineAllocator<1>>& OutMeshBatches, const bool bUseSelectedMaterial, const bool bUseHoveredMaterial) const
+bool FSplineMeshSceneProxy::GetMeshElement(int32 LODIndex, int32 BatchIndex, int32 SectionIndex, uint8 InDepthPriorityGroup, const bool bUseSelectedMaterial, const bool bUseHoveredMaterial, FMeshBatch& OutMeshBatch) const
 {
 	//checkf(LODIndex == 0 /*&& SectionIndex == 0*/, TEXT("Getting spline static mesh element with invalid params [%d, %d]"), LODIndex, SectionIndex);
 
-	if (FStaticMeshSceneProxy::GetMeshElements(LODIndex, SectionIndex, InDepthPriorityGroup, OutMeshBatches, bUseSelectedMaterial, bUseHoveredMaterial))
+	if (FStaticMeshSceneProxy::GetMeshElement(LODIndex, BatchIndex, SectionIndex, InDepthPriorityGroup, bUseSelectedMaterial, bUseHoveredMaterial, OutMeshBatch))
 	{
-		OutMeshBatches[0].VertexFactory = LODResources[LODIndex].VertexFactory;
-		OutMeshBatches[0].ReverseCulling ^= (SplineParams.StartScale.X < 0) ^ (SplineParams.StartScale.Y < 0);
+		OutMeshBatch.VertexFactory = LODResources[LODIndex].VertexFactory;
+		OutMeshBatch.ReverseCulling ^= (SplineParams.StartScale.X < 0) ^ (SplineParams.StartScale.Y < 0);
 		return true;
 	}
 	return false;
 }
 
-bool FSplineMeshSceneProxy::GetWireframeMeshElements(int32 LODIndex, const FMaterialRenderProxy* WireframeRenderProxy, uint8 InDepthPriorityGroup, TArray<FMeshBatch, TInlineAllocator<1>>& OutMeshBatches) const
+bool FSplineMeshSceneProxy::GetWireframeMeshElement(int32 LODIndex, int32 BatchIndex, const FMaterialRenderProxy* WireframeRenderProxy, uint8 InDepthPriorityGroup, FMeshBatch& OutMeshBatch) const
 {
 	//checkf(LODIndex == 0, TEXT("Getting spline static mesh element with invalid LOD [%d]"), LODIndex);
 
-	if (FStaticMeshSceneProxy::GetWireframeMeshElements(LODIndex, WireframeRenderProxy, InDepthPriorityGroup, OutMeshBatches))
+	if (FStaticMeshSceneProxy::GetWireframeMeshElement(LODIndex, BatchIndex, WireframeRenderProxy, InDepthPriorityGroup, OutMeshBatch))
 	{
-		OutMeshBatches[0].VertexFactory = LODResources[LODIndex].VertexFactory;
-		OutMeshBatches[0].ReverseCulling ^= (SplineParams.StartScale.X < 0) ^ (SplineParams.StartScale.Y < 0);
+		OutMeshBatch.VertexFactory = LODResources[LODIndex].VertexFactory;
+		OutMeshBatch.ReverseCulling ^= (SplineParams.StartScale.X < 0) ^ (SplineParams.StartScale.Y < 0);
 		return true;
 	}
 	return false;
diff --git a/Engine/Source/Runtime/Engine/Private/SplineMeshSceneProxy.h b/Engine/Source/Runtime/Engine/Private/SplineMeshSceneProxy.h
--- a/Engine/Source/Runtime/Engine/Private/SplineMeshSceneProxy.h
+++ b/Engine/Source/Runtime/Engine/Private/SplineMeshSceneProxy.h
@@ -155,13 +155,13 @@ public:
 	void ReleaseResources();
 
 	/** Sets up a shadow FMeshBatch for a specific LOD. */
-	virtual bool GetShadowMeshElements(int32 LODIndex, uint8 InDepthPriorityGroup, TArray<FMeshBatch, TInlineAllocator<1>>& OutMeshBatches) const override;
+	virtual bool GetShadowMeshElement(int32 LODIndex, int32 BatchIndex, uint8 InDepthPriorityGroup, FMeshBatch& OutMeshBatch) const override;
 
 	/** Sets up a FMeshBatch for a specific LOD and element. */
-	virtual bool GetMeshElements(int32 LODIndex, int32 SectionIndex, uint8 InDepthPriorityGroup, TArray<FMeshBatch, TInlineAllocator<1>>& OutMeshBatches, const bool bUseSelectedMaterial, const bool bUseHoveredMaterial) const override;
+	virtual bool GetMeshElement(int32 LODIndex, int32 BatchIndex, int32 ElementIndex, uint8 InDepthPriorityGroup, const bool bUseSelectedMaterial, const bool bUseHoveredMaterial, FMeshBatch& OutMeshBatch) const override;
 
 	/** Sets up a wireframe FMeshBatch for a specific LOD. */
-	virtual bool GetWireframeMeshElements(int32 LODIndex, const FMaterialRenderProxy* WireframeRenderProxy, uint8 InDepthPriorityGroup, TArray<FMeshBatch, TInlineAllocator<1>>& OutMeshBatches) const override;
+	virtual bool GetWireframeMeshElement(int32 LODIndex, int32 BatchIndex, const FMaterialRenderProxy* WireframeRenderProxy, uint8 InDepthPriorityGroup, FMeshBatch& OutMeshBatch) const override;
 
 	virtual FPrimitiveViewRelevance GetViewRelevance(const FSceneView* View) override
 	{
diff --git a/Engine/Source/Runtime/Engine/Private/StaticMeshRender.cpp b/Engine/Source/Runtime/Engine/Private/StaticMeshRender.cpp
--- a/Engine/Source/Runtime/Engine/Private/StaticMeshRender.cpp
+++ b/Engine/Source/Runtime/Engine/Private/StaticMeshRender.cpp
@@ -52,16 +52,17 @@ FStaticMeshSceneProxy::FStaticMeshSceneProxy(UStaticMeshComponent* InComponent):
 	BodySetup(InComponent->GetBodySetup()),
 	RenderData(InComponent->StaticMesh->RenderData),
 	ForcedLodModel(InComponent->ForcedLodModel),
-	LevelColor(1,1,1),
-	PropertyColor(1,1,1),
 	bCastShadow(InComponent->CastShadow),
 	CollisionTraceFlag(ECollisionTraceFlag::CTF_UseDefault),
 	MaterialRelevance(InComponent->GetMaterialRelevance()),
-	WireframeColor(InComponent->GetWireframeColor()), 
 	CollisionResponse(InComponent->GetCollisionResponseToChannels())
 {
 	check(RenderData);
 
+	WireframeColor = InComponent->GetWireframeColor();
+	LevelColor = FLinearColor(1,1,1);
+	PropertyColor = FLinearColor(1,1,1);
+
 	if (GForceDefaultMaterial)
 	{
 		MaterialRelevance |= UMaterial::GetDefaultMaterial(MD_Surface)->GetRelevance();
@@ -149,26 +150,26 @@ void UStaticMeshComponent::SetLODDataCount( const uint32 MinSize, const uint32 M
 	}
 }
 
-bool FStaticMeshSceneProxy::GetShadowMeshElements(int32 LODIndex, uint8 InDepthPriorityGroup, TArray<FMeshBatch, TInlineAllocator<1>>& OutMeshBatches) const
+bool FStaticMeshSceneProxy::GetShadowMeshElement(int32 LODIndex, int32 BatchIndex, uint8 InDepthPriorityGroup, FMeshBatch& OutMeshBatch) const
 {
 	const FStaticMeshLODResources& LOD = RenderData->LODResources[LODIndex];
 	const FLODInfo& ProxyLODInfo = LODs[LODIndex];
 
-	FMeshBatchElement& OutBatchElement = OutMeshBatches[0].Elements[0];
-	OutMeshBatches[0].MaterialRenderProxy = UMaterial::GetDefaultMaterial(MD_Surface)->GetRenderProxy(false, false);
-	OutMeshBatches[0].VertexFactory = &LOD.VertexFactory;
+	FMeshBatchElement& OutBatchElement = OutMeshBatch.Elements[0];
+	OutMeshBatch.MaterialRenderProxy = UMaterial::GetDefaultMaterial(MD_Surface)->GetRenderProxy(false, false);
+	OutMeshBatch.VertexFactory = &LOD.VertexFactory;
 	OutBatchElement.IndexBuffer = &LOD.DepthOnlyIndexBuffer;
-	OutMeshBatches[0].Type = PT_TriangleList;
+	OutMeshBatch.Type = PT_TriangleList;
 	OutBatchElement.FirstIndex = 0;
 	OutBatchElement.NumPrimitives = LOD.DepthOnlyIndexBuffer.GetNumIndices() / 3;
 	OutBatchElement.PrimitiveUniformBufferResource = &GetUniformBuffer();
 	OutBatchElement.MinVertexIndex = 0;
 	OutBatchElement.MaxVertexIndex = LOD.PositionVertexBuffer.GetNumVertices() - 1;
-	OutMeshBatches[0].DepthPriorityGroup = InDepthPriorityGroup;
-	OutMeshBatches[0].ReverseCulling = IsLocalToWorldDeterminantNegative();
-	OutMeshBatches[0].CastShadow = true;
-	OutMeshBatches[0].LODIndex = LODIndex;
-	OutMeshBatches[0].LCI = &ProxyLODInfo;
+	OutMeshBatch.DepthPriorityGroup = InDepthPriorityGroup;
+	OutMeshBatch.ReverseCulling = IsLocalToWorldDeterminantNegative();
+	OutMeshBatch.CastShadow = true;
+	OutMeshBatch.LODIndex = LODIndex;
+	OutMeshBatch.LCI = &ProxyLODInfo;
 	if (ForcedLodModel > 0) 
 	{
 		OutBatchElement.MaxScreenSize = 0.0f;
@@ -188,15 +189,15 @@ bool FStaticMeshSceneProxy::GetShadowMeshElements(int32 LODIndex, uint8 InDepthP
 }
 
 /** Sets up a FMeshBatch for a specific LOD and element. */
-bool FStaticMeshSceneProxy::GetMeshElements(int32 LODIndex,int32 SectionIndex,uint8 InDepthPriorityGroup, TArray<FMeshBatch, TInlineAllocator<1>>& OutMeshBatches, const bool bUseSelectedMaterial, const bool bUseHoveredMaterial) const
+bool FStaticMeshSceneProxy::GetMeshElement(int32 LODIndex, int32 BatchIndex, int32 SectionIndex, uint8 InDepthPriorityGroup, const bool bUseSelectedMaterial, const bool bUseHoveredMaterial, FMeshBatch& OutMeshBatch) const
 {
 	const FStaticMeshLODResources& LOD = RenderData->LODResources[LODIndex];
 	const FStaticMeshSection& Section = LOD.Sections[SectionIndex];
 
 	const FLODInfo& ProxyLODInfo = LODs[LODIndex];
 	UMaterialInterface* Material = ProxyLODInfo.Sections[SectionIndex].Material;
-	OutMeshBatches[0].MaterialRenderProxy = Material->GetRenderProxy(bUseSelectedMaterial,bUseHoveredMaterial);
-	OutMeshBatches[0].VertexFactory = &LOD.VertexFactory;
+	OutMeshBatch.MaterialRenderProxy = Material->GetRenderProxy(bUseSelectedMaterial,bUseHoveredMaterial);
+	OutMeshBatch.VertexFactory = &LOD.VertexFactory;
 
 	// Has the mesh component overridden the vertex color stream for this mesh LOD?
 	if( ProxyLODInfo.OverrideColorVertexBuffer != NULL )
@@ -208,29 +209,29 @@ bool FStaticMeshSceneProxy::GetMeshElements(int32 LODIndex,int32 SectionIndex,ui
 		{
 			// Switch out the stock mesh vertex factory with own own vertex factory that points to
 			// our overridden color data
-			OutMeshBatches[0].VertexFactory = ProxyLODInfo.OverrideColorVertexFactory.GetOwnedPointer();
+			OutMeshBatch.VertexFactory = ProxyLODInfo.OverrideColorVertexFactory.GetOwnedPointer();
 		}
 	}
 
 	const bool bWireframe = false;
-	const bool bRequiresAdjacencyInformation = RequiresAdjacencyInformation( Material, OutMeshBatches[0].VertexFactory->GetType(), GetScene()->GetFeatureLevel() );
+	const bool bRequiresAdjacencyInformation = RequiresAdjacencyInformation( Material, OutMeshBatch.VertexFactory->GetType(), GetScene()->GetFeatureLevel() );
 
-	SetIndexSource(LODIndex, SectionIndex, OutMeshBatches[0], bWireframe, bRequiresAdjacencyInformation );
+	SetIndexSource(LODIndex, SectionIndex, OutMeshBatch, bWireframe, bRequiresAdjacencyInformation );
 
-	FMeshBatchElement& OutBatchElement = OutMeshBatches[0].Elements[0];
+	FMeshBatchElement& OutBatchElement = OutMeshBatch.Elements[0];
 
 	if(OutBatchElement.NumPrimitives > 0)
 	{
-		OutMeshBatches[0].DynamicVertexData = NULL;
-		OutMeshBatches[0].LCI = &ProxyLODInfo;
+		OutMeshBatch.DynamicVertexData = NULL;
+		OutMeshBatch.LCI = &ProxyLODInfo;
 		OutBatchElement.PrimitiveUniformBufferResource = &GetUniformBuffer();
 		OutBatchElement.MinVertexIndex = Section.MinVertexIndex;
 		OutBatchElement.MaxVertexIndex = Section.MaxVertexIndex;
-		OutMeshBatches[0].LODIndex = LODIndex;
-		OutMeshBatches[0].UseDynamicData = false;
-		OutMeshBatches[0].ReverseCulling = IsLocalToWorldDeterminantNegative();
-		OutMeshBatches[0].CastShadow = bCastShadow && Section.bCastShadow;
-		OutMeshBatches[0].DepthPriorityGroup = (ESceneDepthPriorityGroup)InDepthPriorityGroup;
+		OutMeshBatch.LODIndex = LODIndex;
+		OutMeshBatch.UseDynamicData = false;
+		OutMeshBatch.ReverseCulling = IsLocalToWorldDeterminantNegative();
+		OutMeshBatch.CastShadow = bCastShadow && Section.bCastShadow;
+		OutMeshBatch.DepthPriorityGroup = (ESceneDepthPriorityGroup)InDepthPriorityGroup;
 		if (ForcedLodModel > 0) 
 		{
 			OutBatchElement.MaxScreenSize = 0.0f;
@@ -255,20 +256,20 @@ bool FStaticMeshSceneProxy::GetMeshElements(int32 LODIndex,int32 SectionIndex,ui
 }
 
 /** Sets up a wireframe FMeshBatch for a specific LOD. */
-bool FStaticMeshSceneProxy::GetWireframeMeshElements(int32 LODIndex, const FMaterialRenderProxy* WireframeRenderProxy, uint8 InDepthPriorityGroup, TArray<FMeshBatch, TInlineAllocator<1>>& OutMeshBatches) const
+bool FStaticMeshSceneProxy::GetWireframeMeshElement(int32 LODIndex, int32 BatchIndex, const FMaterialRenderProxy* WireframeRenderProxy, uint8 InDepthPriorityGroup, FMeshBatch& OutMeshBatch) const
 {
 	const FStaticMeshLODResources& LODModel = RenderData->LODResources[LODIndex];
 	
-	FMeshBatchElement& OutBatchElement = OutMeshBatches[0].Elements[0];
+	FMeshBatchElement& OutBatchElement = OutMeshBatch.Elements[0];
 
-	OutMeshBatches[0].VertexFactory = &LODModel.VertexFactory;
-	OutMeshBatches[0].MaterialRenderProxy = WireframeRenderProxy;
+	OutMeshBatch.VertexFactory = &LODModel.VertexFactory;
+	OutMeshBatch.MaterialRenderProxy = WireframeRenderProxy;
 	OutBatchElement.PrimitiveUniformBufferResource = &GetUniformBuffer();
 	OutBatchElement.MinVertexIndex = 0;
 	OutBatchElement.MaxVertexIndex = LODModel.GetNumVertices() - 1;
-	OutMeshBatches[0].ReverseCulling = IsLocalToWorldDeterminantNegative();
-	OutMeshBatches[0].CastShadow = bCastShadow;
-	OutMeshBatches[0].DepthPriorityGroup = (ESceneDepthPriorityGroup)InDepthPriorityGroup;
+	OutMeshBatch.ReverseCulling = IsLocalToWorldDeterminantNegative();
+	OutMeshBatch.CastShadow = bCastShadow;
+	OutMeshBatch.DepthPriorityGroup = (ESceneDepthPriorityGroup)InDepthPriorityGroup;
 	if (ForcedLodModel > 0) 
 	{
 		OutBatchElement.MaxScreenSize = 0.0f;
@@ -287,7 +288,7 @@ bool FStaticMeshSceneProxy::GetWireframeMeshElements(int32 LODIndex, const FMate
 	const bool bWireframe = true;
 	const bool bRequiresAdjacencyInformation = false;
 
-	SetIndexSource(LODIndex, 0, OutMeshBatches[0], bWireframe, bRequiresAdjacencyInformation);
+	SetIndexSource(LODIndex, 0, OutMeshBatch, bWireframe, bRequiresAdjacencyInformation);
 
 	return OutBatchElement.NumPrimitives > 0;
 }
@@ -418,7 +419,6 @@ void FStaticMeshSceneProxy::DrawStaticElements(FStaticPrimitiveDrawInterface* PD
 		//Never use the dynamic path in this path, because only unselected elements will use DrawStaticElements
 		bool bUseSelectedMaterial = false;
 		const bool bUseHoveredMaterial = false;
-		
 
 		//check if a LOD is being forced
 		if (ForcedLodModel > 0) 
@@ -438,13 +438,15 @@ void FStaticMeshSceneProxy::DrawStaticElements(FStaticPrimitiveDrawInterface* PD
 				}
 #endif // WITH_EDITOR
 
-				TArray<FMeshBatch, TInlineAllocator<1>> MeshBatches;
-				new (MeshBatches) FMeshBatch();
-				if(GetMeshElements(LODIndex, SectionIndex, PrimitiveDPG, MeshBatches, bUseSelectedMaterial, bUseHoveredMaterial))
+				const int32 NumBatches = GetNumMeshBatches();
+
+				for (int32 BatchIndex = 0; BatchIndex < NumBatches; BatchIndex++)
 				{
-					for (int32 Index = 0; Index < MeshBatches.Num(); ++Index)
+					FMeshBatch MeshBatch;
+
+					if (GetMeshElement(LODIndex, BatchIndex, SectionIndex, PrimitiveDPG, bUseSelectedMaterial, bUseHoveredMaterial, MeshBatch))
 					{
-						PDI->DrawMesh(MeshBatches[Index], FLT_MAX);
+						PDI->DrawMesh(MeshBatch, FLT_MAX);
 					}
 				}
 			}
@@ -482,14 +484,16 @@ void FStaticMeshSceneProxy::DrawStaticElements(FStaticPrimitiveDrawInterface* PD
 
 					if (bAnySectionCastsShadow && bSafeToUseShadowOnlyMesh)
 					{
-						TArray<FMeshBatch, TInlineAllocator<1>> MeshBatches;
-						new (MeshBatches) FMeshBatch();
-						if (GetShadowMeshElements(LODIndex, PrimitiveDPG, MeshBatches))
+						const int32 NumBatches = GetNumMeshBatches();
+
+						for (int32 BatchIndex = 0; BatchIndex < NumBatches; BatchIndex++)
 						{
-							bHaveShadowOnlyMesh = true;
-							for (int32 Index = 0; Index < MeshBatches.Num(); ++Index)
+							FMeshBatch MeshBatch;
+
+							if (GetShadowMeshElement(LODIndex, BatchIndex, PrimitiveDPG, MeshBatch))
 							{
-								PDI->DrawMesh(MeshBatches[Index], ScreenSize, /*bShadowOnly=*/true);
+								bHaveShadowOnlyMesh = true;
+								PDI->DrawMesh(MeshBatch, ScreenSize, /*bShadowOnly=*/true);
 							}
 						}
 					}
@@ -508,15 +512,17 @@ void FStaticMeshSceneProxy::DrawStaticElements(FStaticPrimitiveDrawInterface* PD
 					}
 #endif // WITH_EDITOR
 
-					TArray<FMeshBatch, TInlineAllocator<1>> MeshBatches;
-					new (MeshBatches) FMeshBatch();
-					if(GetMeshElements(LODIndex, SectionIndex, PrimitiveDPG, MeshBatches, bUseSelectedMaterial, bUseHoveredMaterial))
+					const int32 NumBatches = GetNumMeshBatches();
+
+					for (int32 BatchIndex = 0; BatchIndex < NumBatches; BatchIndex++)
 					{
-						// If we have submitted an optimized shadow-only mesh, remaining mesh elements must not cast shadows.
-						for (int32 Index = 0; Index < MeshBatches.Num(); ++Index)
+						FMeshBatch MeshBatch;
+
+						if (GetMeshElement(LODIndex, BatchIndex, SectionIndex, PrimitiveDPG, bUseSelectedMaterial, bUseHoveredMaterial, MeshBatch))
 						{
-							MeshBatches[Index].CastShadow = MeshBatches[Index].CastShadow && !bHaveShadowOnlyMesh;
-							PDI->DrawMesh(MeshBatches[Index], ScreenSize);
+							// If we have submitted an optimized shadow-only mesh, remaining mesh elements must not cast shadows.
+							MeshBatch.CastShadow = MeshBatch.CastShadow && !bHaveShadowOnlyMesh;
+							PDI->DrawMesh(MeshBatch, ScreenSize);
 						}
 					}
 				}
@@ -525,27 +531,280 @@ void FStaticMeshSceneProxy::DrawStaticElements(FStaticPrimitiveDrawInterface* PD
 	}
 }
 
-bool FStaticMeshSceneProxy::IsCollisionView(const FSceneView* View, bool & bDrawSimpleCollision, bool & bDrawComplexCollision) const
+bool FStaticMeshSceneProxy::IsCollisionView(const FEngineShowFlags& EngineShowFlags, bool& bDrawSimpleCollision, bool& bDrawComplexCollision) const
 {
 	bDrawSimpleCollision = bDrawComplexCollision = false;
 
-	const bool bInCollisionView = View->Family->EngineShowFlags.CollisionVisibility || View->Family->EngineShowFlags.CollisionPawn;
+	const bool bInCollisionView = EngineShowFlags.CollisionVisibility || EngineShowFlags.CollisionPawn;
 	// If in a 'collision view' and collision is enabled
 	if (bInCollisionView && IsCollisionEnabled())
 	{
 		// See if we have a response to the interested channel
-		bool bHasResponse = View->Family->EngineShowFlags.CollisionPawn && CollisionResponse.GetResponse(ECC_Pawn) != ECR_Ignore;
-		bHasResponse |= View->Family->EngineShowFlags.CollisionVisibility && CollisionResponse.GetResponse(ECC_Visibility) != ECR_Ignore;
+		bool bHasResponse = EngineShowFlags.CollisionPawn && CollisionResponse.GetResponse(ECC_Pawn) != ECR_Ignore;
+		bHasResponse |= EngineShowFlags.CollisionVisibility && CollisionResponse.GetResponse(ECC_Visibility) != ECR_Ignore;
 
 		if(bHasResponse)
 		{
-			bDrawComplexCollision = View->Family->EngineShowFlags.CollisionVisibility;
-			bDrawSimpleCollision = View->Family->EngineShowFlags.CollisionPawn;
+			bDrawComplexCollision = EngineShowFlags.CollisionVisibility;
+			bDrawSimpleCollision = EngineShowFlags.CollisionPawn;
 		}
 	}
 
 	return bInCollisionView;
 }
+
+void FStaticMeshSceneProxy::GetDynamicMeshElements(const TArray<const FSceneView*>& Views, const FSceneViewFamily& ViewFamily, uint32 VisibilityMap, FMeshElementCollector& Collector) const
+{
+	QUICK_SCOPE_CYCLE_COUNTER(STAT_StaticMeshSceneProxy_GetMeshElements);
+	checkSlow(IsInRenderingThread());
+
+	const bool bIsLightmapSettingError = HasStaticLighting() && !HasValidSettingsForStaticLighting();
+	const bool bProxyIsSelected = IsSelected();
+	const FEngineShowFlags& EngineShowFlags = ViewFamily.EngineShowFlags;
+
+	bool bDrawSimpleCollision = false, bDrawComplexCollision = false;
+	const bool bInCollisionView = IsCollisionView(EngineShowFlags, bDrawSimpleCollision, bDrawComplexCollision);
+	// Draw simple collision as wireframe if 'show collision', collision is enabled, and we are not using the complex as the simple
+	const bool bDrawWireframeCollision = (EngineShowFlags.Collision && IsCollisionEnabled() && CollisionTraceFlag != ECollisionTraceFlag::CTF_UseComplexAsSimple);
+	
+	const bool bDrawMesh = (bInCollisionView) ? (bDrawComplexCollision) : 
+		(	IsRichView(ViewFamily) || HasViewDependentDPG()
+			|| EngineShowFlags.Collision
+			|| EngineShowFlags.Bounds
+			|| bProxyIsSelected 
+			|| IsHovered()
+			|| bIsLightmapSettingError ) ;
+
+	// Draw polygon mesh if we are either not in a collision view, or are drawing it as collision.
+	if (EngineShowFlags.StaticMeshes && bDrawMesh)
+	{
+		// how we should draw the collision for this mesh.
+		const bool bIsWireframeView = EngineShowFlags.Wireframe;
+		const bool bLevelColorationEnabled = EngineShowFlags.LevelColoration;
+		const bool bPropertyColorationEnabled = EngineShowFlags.PropertyColoration;
+		const ERHIFeatureLevel::Type FeatureLevel = ViewFamily.GetFeatureLevel();
+
+		for (int32 ViewIndex = 0; ViewIndex < Views.Num(); ViewIndex++)
+		{
+			if (VisibilityMap & (1 << ViewIndex))
+			{
+				const FSceneView* View = Views[ViewIndex];
+				const int32 LODIndex = GetLOD(View);
+
+				if (RenderData->LODResources.IsValidIndex(LODIndex))
+				{
+					const FStaticMeshLODResources& LODModel = RenderData->LODResources[LODIndex];
+					const FLODInfo& ProxyLODInfo = LODs[LODIndex];
+
+					if (AllowDebugViewmodes() && bIsWireframeView && !EngineShowFlags.Materials
+						// If any of the materials are mesh-modifying, we can't use the single merged mesh element of GetWireframeMeshElement()
+						&& !ProxyLODInfo.UsesMeshModifyingMaterials())
+					{
+						FLinearColor ViewWireframeColor( bLevelColorationEnabled ? LevelColor : WireframeColor );
+						if ( bPropertyColorationEnabled )
+						{
+							ViewWireframeColor = PropertyColor;
+						}
+
+						auto WireframeMaterialInstance = new FColoredMaterialRenderProxy(
+							GEngine->WireframeMaterial->GetRenderProxy(false),
+							GetSelectionColor(ViewWireframeColor,!(GIsEditor && EngineShowFlags.Selection) || bProxyIsSelected, IsHovered(), false)
+							);
+
+						Collector.RegisterOneFrameMaterialProxy(WireframeMaterialInstance);
+
+						const int32 NumBatches = GetNumMeshBatches();
+
+						for (int32 BatchIndex = 0; BatchIndex < NumBatches; BatchIndex++)
+						{
+							FMeshBatch& Mesh = Collector.AllocateMesh();
+
+							if (GetWireframeMeshElement(LODIndex, BatchIndex, WireframeMaterialInstance, SDPG_World, Mesh))
+							{
+								// We implemented our own wireframe
+								Mesh.bCanApplyViewModeOverrides = false;
+								Collector.AddMesh(ViewIndex, Mesh);
+								INC_DWORD_STAT_BY(STAT_StaticMeshTriangles,Mesh.GetNumPrimitives());
+							}
+						}
+					}
+					else
+					{
+						const FLinearColor UtilColor( LevelColor );
+
+						// Draw the static mesh sections.
+						for (int32 SectionIndex = 0; SectionIndex < LODModel.Sections.Num(); SectionIndex++)
+						{
+							const int32 NumBatches = GetNumMeshBatches();
+
+							for (int32 BatchIndex = 0; BatchIndex < NumBatches; BatchIndex++)
+							{
+								bool bSectionIsSelected = false;
+								FMeshBatch& MeshElement = Collector.AllocateMesh();
+
+#if WITH_EDITOR
+								if (GIsEditor)
+								{
+									const FLODInfo::FSectionInfo& Section = LODs[LODIndex].Sections[SectionIndex];
+
+									bSectionIsSelected = Section.bSelected;
+									MeshElement.BatchHitProxyId = Section.HitProxy ? Section.HitProxy->Id : FHitProxyId();
+								}
+#endif // WITH_EDITOR
+								
+								if (GetMeshElement(LODIndex, BatchIndex, SectionIndex, SDPG_World, bSectionIsSelected, IsHovered(), MeshElement))
+								{
+									bool bAddMesh = true;
+
+#if !(UE_BUILD_SHIPPING || UE_BUILD_TEST)
+									if (EngineShowFlags.VertexColors && AllowDebugViewmodes())
+									{
+										// Override the mesh's material with our material that draws the vertex colors
+										UMaterial* VertexColorVisualizationMaterial = NULL;
+										switch( GVertexColorViewMode )
+										{
+										case EVertexColorViewMode::Color:
+											VertexColorVisualizationMaterial = GEngine->VertexColorViewModeMaterial_ColorOnly;
+											break;
+
+										case EVertexColorViewMode::Alpha:
+											VertexColorVisualizationMaterial = GEngine->VertexColorViewModeMaterial_AlphaAsColor;
+											break;
+
+										case EVertexColorViewMode::Red:
+											VertexColorVisualizationMaterial = GEngine->VertexColorViewModeMaterial_RedOnly;
+											break;
+
+										case EVertexColorViewMode::Green:
+											VertexColorVisualizationMaterial = GEngine->VertexColorViewModeMaterial_GreenOnly;
+											break;
+
+										case EVertexColorViewMode::Blue:
+											VertexColorVisualizationMaterial = GEngine->VertexColorViewModeMaterial_BlueOnly;
+											break;
+										}
+										check( VertexColorVisualizationMaterial != NULL );
+
+										auto VertexColorVisualizationMaterialInstance = new FColoredMaterialRenderProxy(
+											VertexColorVisualizationMaterial->GetRenderProxy( MeshElement.MaterialRenderProxy->IsSelected(),MeshElement.MaterialRenderProxy->IsHovered() ),
+											GetSelectionColor( FLinearColor::White, bSectionIsSelected, IsHovered() )
+											);
+
+										Collector.RegisterOneFrameMaterialProxy(VertexColorVisualizationMaterialInstance);
+										MeshElement.MaterialRenderProxy = VertexColorVisualizationMaterialInstance;
+									}
+									else if( bInCollisionView && bDrawComplexCollision && AllowDebugViewmodes() )
+									{
+										if (LODModel.Sections[SectionIndex].bEnableCollision)
+										{
+											// Override the mesh's material with our material that draws the collision color
+											auto CollisionMaterialInstance = new FColoredMaterialRenderProxy(
+												GEngine->ShadedLevelColorationUnlitMaterial->GetRenderProxy(bProxyIsSelected, IsHovered()),
+												WireframeColor
+												);
+
+											Collector.RegisterOneFrameMaterialProxy(CollisionMaterialInstance);
+											MeshElement.MaterialRenderProxy = CollisionMaterialInstance;
+										}
+										else
+										{
+											bAddMesh = false;
+										}
+									}
+									else
+#endif
+#if WITH_EDITOR
+									if (bSectionIsSelected)
+									{
+										// Override the mesh's material with our material that draws the collision color
+										auto SelectedMaterialInstance = new FOverrideSelectionColorMaterialRenderProxy(
+											GEngine->ShadedLevelColorationUnlitMaterial->GetRenderProxy(bSectionIsSelected, IsHovered()),
+											GetSelectionColor(GEngine->GetSelectedMaterialColor(), bSectionIsSelected, IsHovered())
+											);
+
+										MeshElement.MaterialRenderProxy = SelectedMaterialInstance;
+									}
+#endif
+								
+									MeshElement.bCanApplyViewModeOverrides = true;
+									MeshElement.bUseWireframeSelectionColoring = bSectionIsSelected;
+
+									if (bAddMesh)
+									{
+										Collector.AddMesh(ViewIndex, MeshElement);
+										INC_DWORD_STAT_BY(STAT_StaticMeshTriangles,MeshElement.GetNumPrimitives());
+									}
+								}
+							}
+						}
+					}
+				}
+			}
+		}
+	}
+	
+	for (int32 ViewIndex = 0; ViewIndex < Views.Num(); ViewIndex++)
+	{
+		if (VisibilityMap & (1 << ViewIndex))
+		{
+			if((bDrawSimpleCollision || bDrawWireframeCollision) && AllowDebugViewmodes())
+			{
+				if(BodySetup)
+				{
+					if(FMath::Abs(GetLocalToWorld().Determinant()) < SMALL_NUMBER)
+					{
+						// Catch this here or otherwise GeomTransform below will assert
+						// This spams so commented out
+						//UE_LOG(LogStaticMesh, Log, TEXT("Zero scaling not supported (%s)"), *StaticMesh->GetPathName());
+					}
+					else
+					{
+						const bool bDrawSolid = !bDrawWireframeCollision;
+
+						if(bDrawSolid)
+						{
+							// Make a material for drawing solid collision stuff
+							auto SolidMaterialInstance = new FColoredMaterialRenderProxy(
+								GEngine->ShadedLevelColorationUnlitMaterial->GetRenderProxy(IsSelected(), IsHovered()),
+								WireframeColor
+								);
+
+							Collector.RegisterOneFrameMaterialProxy(SolidMaterialInstance);
+
+							FTransform GeomTransform(GetLocalToWorld());
+							BodySetup->AggGeom.GetAggGeom(GeomTransform, WireframeColor, SolidMaterialInstance, false, true, UseEditorDepthTest(), ViewIndex, Collector);
+						}
+						// wireframe
+						else
+						{
+							FColor CollisionColor = FColor(157,149,223,255);
+							FTransform GeomTransform(GetLocalToWorld());
+							BodySetup->AggGeom.GetAggGeom(GeomTransform, GetSelectionColor(CollisionColor, bProxyIsSelected, IsHovered()), NULL, (Owner == NULL), false, UseEditorDepthTest(), ViewIndex, Collector);
+						}
+
+
+						// The simple nav geometry is only used by dynamic obstacles for now
+						if (StaticMesh->NavCollision && StaticMesh->NavCollision->bIsDynamicObstacle)
+						{
+							// Draw the static mesh's body setup (simple collision)
+							FTransform GeomTransform(GetLocalToWorld());
+							FColor NavCollisionColor = FColor(118,84,255,255);
+							StaticMesh->NavCollision->DrawSimpleGeom(Collector.GetPDI(ViewIndex), GeomTransform, GetSelectionColor(NavCollisionColor, bProxyIsSelected, IsHovered()));
+						}
+					}
+				}
+			}
+	
+#if !(UE_BUILD_SHIPPING || UE_BUILD_TEST)
+			if (EngineShowFlags.StaticMeshes)
+			{
+				RenderBounds(Collector.GetPDI(ViewIndex), EngineShowFlags, GetBounds(), !Owner || IsSelected());
+			}
+#endif
+		}
+	}
+}
+
 /** 
 * Draw the scene proxy as a dynamic element
 *
@@ -563,12 +822,12 @@ void FStaticMeshSceneProxy::DrawDynamicElements(FPrimitiveDrawInterface* PDI,con
 	bool bProxyIsSelected = IsSelected();
 
 	bool bDrawSimpleCollision = false, bDrawComplexCollision = false;
-	const bool bInCollisionView = IsCollisionView(View, bDrawSimpleCollision, bDrawComplexCollision);
+	const bool bInCollisionView = IsCollisionView(View->Family->EngineShowFlags, bDrawSimpleCollision, bDrawComplexCollision);
 	// Draw simple collision as wireframe if 'show collision', collision is enabled, and we are not using the complex as the simple
 	const bool bDrawWireframeCollision = (View->Family->EngineShowFlags.Collision && IsCollisionEnabled() && CollisionTraceFlag != ECollisionTraceFlag::CTF_UseComplexAsSimple);
 
 	const bool bDrawMesh = (bInCollisionView) ? (bDrawComplexCollision) : 
-		(	IsRichView(View) || HasViewDependentDPG()
+		(	IsRichView(*View->Family) || HasViewDependentDPG()
 			|| View->Family->EngineShowFlags.Collision
 			|| View->Family->EngineShowFlags.Bounds
 			|| bProxyIsSelected 
@@ -615,14 +874,16 @@ void FStaticMeshSceneProxy::DrawDynamicElements(FPrimitiveDrawInterface* PDI,con
 					GetSelectionColor(ViewWireframeColor,!(GIsEditor && View->Family->EngineShowFlags.Selection) || bProxyIsSelected, IsHovered(), /*bUseOverlayIntensity=*/false)
 					);
 
-				TArray<FMeshBatch, TInlineAllocator<1>> MeshBatches;
-				new (MeshBatches) FMeshBatch();
-				if (GetWireframeMeshElements(LODIndex, &WireframeMaterialInstance, GetDepthPriorityGroup(View), MeshBatches))
+				const int32 NumBatches = GetNumMeshBatches();
+
+				for (int32 BatchIndex = 0; BatchIndex < NumBatches; BatchIndex++)
 				{
-					for (int32 Index = 0; Index < MeshBatches.Num(); ++Index)
+					FMeshBatch MeshBatch;
+
+					if (GetWireframeMeshElement(LODIndex, BatchIndex, &WireframeMaterialInstance, GetDepthPriorityGroup(View), MeshBatch))
 					{
-						const int32 NumPasses = PDI->DrawMesh(MeshBatches[Index]);
-						INC_DWORD_STAT_BY(STAT_StaticMeshTriangles,MeshBatches[Index].GetNumPrimitives() * NumPasses);
+						const int32 NumPasses = PDI->DrawMesh(MeshBatch);
+						INC_DWORD_STAT_BY(STAT_StaticMeshTriangles,MeshBatch.GetNumPrimitives() * NumPasses);
 					}
 				}
 			}
@@ -645,46 +906,48 @@ void FStaticMeshSceneProxy::DrawDynamicElements(FPrimitiveDrawInterface* PDI,con
 					}
 #endif // WITH_EDITOR
 
-					TArray<FMeshBatch, TInlineAllocator<1>> MeshBatches;
-					new (MeshBatches) FMeshBatch();
-					if(GetMeshElements(LODIndex,SectionIndex,GetDepthPriorityGroup(View), MeshBatches,bSectionIsSelected,IsHovered()))
+					const int32 NumBatches = GetNumMeshBatches();
+
+					for (int32 BatchIndex = 0; BatchIndex < NumBatches; BatchIndex++)
 					{
+						FMeshBatch MeshBatch;
+
+						if (GetMeshElement(LODIndex, BatchIndex, SectionIndex, GetDepthPriorityGroup(View), bSectionIsSelected, IsHovered(), MeshBatch))
+						{			
 #if !(UE_BUILD_SHIPPING || UE_BUILD_TEST)
-						if( View->Family->EngineShowFlags.VertexColors && AllowDebugViewmodes() )
-						{
-							// Override the mesh's material with our material that draws the vertex colors
-							UMaterial* VertexColorVisualizationMaterial = NULL;
-							switch( GVertexColorViewMode )
+							if( View->Family->EngineShowFlags.VertexColors && AllowDebugViewmodes() )
 							{
-							case EVertexColorViewMode::Color:
-								VertexColorVisualizationMaterial = GEngine->VertexColorViewModeMaterial_ColorOnly;
-								break;
+								// Override the mesh's material with our material that draws the vertex colors
+								UMaterial* VertexColorVisualizationMaterial = NULL;
+								switch( GVertexColorViewMode )
+								{
+								case EVertexColorViewMode::Color:
+									VertexColorVisualizationMaterial = GEngine->VertexColorViewModeMaterial_ColorOnly;
+									break;
 
-							case EVertexColorViewMode::Alpha:
-								VertexColorVisualizationMaterial = GEngine->VertexColorViewModeMaterial_AlphaAsColor;
-								break;
+								case EVertexColorViewMode::Alpha:
+									VertexColorVisualizationMaterial = GEngine->VertexColorViewModeMaterial_AlphaAsColor;
+									break;
 
-							case EVertexColorViewMode::Red:
-								VertexColorVisualizationMaterial = GEngine->VertexColorViewModeMaterial_RedOnly;
-								break;
+								case EVertexColorViewMode::Red:
+									VertexColorVisualizationMaterial = GEngine->VertexColorViewModeMaterial_RedOnly;
+									break;
 
-							case EVertexColorViewMode::Green:
-								VertexColorVisualizationMaterial = GEngine->VertexColorViewModeMaterial_GreenOnly;
-								break;
+								case EVertexColorViewMode::Green:
+									VertexColorVisualizationMaterial = GEngine->VertexColorViewModeMaterial_GreenOnly;
+									break;
 
-							case EVertexColorViewMode::Blue:
-								VertexColorVisualizationMaterial = GEngine->VertexColorViewModeMaterial_BlueOnly;
-								break;
-							}
-							check( VertexColorVisualizationMaterial != NULL );
+								case EVertexColorViewMode::Blue:
+									VertexColorVisualizationMaterial = GEngine->VertexColorViewModeMaterial_BlueOnly;
+									break;
+								}
+								check( VertexColorVisualizationMaterial != NULL );
 
-							for (int32 Index = 0; Index < MeshBatches.Num(); ++Index)
-							{
 								const FColoredMaterialRenderProxy VertexColorVisualizationMaterialInstance(
-									VertexColorVisualizationMaterial->GetRenderProxy( MeshBatches[Index].MaterialRenderProxy->IsSelected(),MeshBatches[Index].MaterialRenderProxy->IsHovered() ),
+									VertexColorVisualizationMaterial->GetRenderProxy( MeshBatch.MaterialRenderProxy->IsSelected(),MeshBatch.MaterialRenderProxy->IsHovered() ),
 									GetSelectionColor( FLinearColor::White, bSectionIsSelected, IsHovered() )
 									);
-								FMeshBatch ModifiedMeshElement = MeshBatches[Index];
+								FMeshBatch ModifiedMeshElement = MeshBatch;
 								ModifiedMeshElement.MaterialRenderProxy = &VertexColorVisualizationMaterialInstance;
 
 								const int32 NumPasses = DrawRichMesh(
@@ -697,22 +960,19 @@ void FStaticMeshSceneProxy::DrawDynamicElements(FPrimitiveDrawInterface* PDI,con
 									bSectionIsSelected,
 									bIsWireframeView
 									);
-								INC_DWORD_STAT_BY(STAT_StaticMeshTriangles,MeshBatches[Index].GetNumPrimitives() * NumPasses);
+								INC_DWORD_STAT_BY(STAT_StaticMeshTriangles,MeshBatch.GetNumPrimitives() * NumPasses);
 							}
-						}
-						else if( bInCollisionView && bDrawComplexCollision && AllowDebugViewmodes() )
-						{
-							if (LODModel.Sections[SectionIndex].bEnableCollision)
+							else if( bInCollisionView && bDrawComplexCollision && AllowDebugViewmodes() )
 							{
-								// Override the mesh's material with our material that draws the collision color
-								const FColoredMaterialRenderProxy CollisionMaterialInstance(
-									GEngine->ShadedLevelColorationUnlitMaterial->GetRenderProxy(bProxyIsSelected, IsHovered()),
-									WireframeColor
-									);
-
-								for (int32 Index = 0; Index < MeshBatches.Num(); ++Index)
+								if (LODModel.Sections[SectionIndex].bEnableCollision)
 								{
-									FMeshBatch ModifiedMeshElement = MeshBatches[Index];
+									// Override the mesh's material with our material that draws the collision color
+									const FColoredMaterialRenderProxy CollisionMaterialInstance(
+										GEngine->ShadedLevelColorationUnlitMaterial->GetRenderProxy(bProxyIsSelected, IsHovered()),
+										WireframeColor
+										);
+
+									FMeshBatch ModifiedMeshElement = MeshBatch;
 									ModifiedMeshElement.MaterialRenderProxy = &CollisionMaterialInstance;
 
 									const int32 NumPasses = DrawRichMesh(
@@ -725,24 +985,21 @@ void FStaticMeshSceneProxy::DrawDynamicElements(FPrimitiveDrawInterface* PDI,con
 										bSectionIsSelected,
 										bIsWireframeView
 										);
-									INC_DWORD_STAT_BY(STAT_StaticMeshTriangles,MeshBatches[Index].GetNumPrimitives() * NumPasses);
+									INC_DWORD_STAT_BY(STAT_StaticMeshTriangles,MeshBatch.GetNumPrimitives() * NumPasses);
 								}
 							}
-						}
-						else
+							else
 #endif
 #if WITH_EDITOR
-						if (bSectionIsSelected)
-						{
-							// Override the mesh's material with our material that draws the collision color
-							const FOverrideSelectionColorMaterialRenderProxy SelectedMaterialInstance(
-								GEngine->ShadedLevelColorationUnlitMaterial->GetRenderProxy(bSectionIsSelected, IsHovered()),
-								GetSelectionColor(GEngine->GetSelectedMaterialColor(), bSectionIsSelected, IsHovered())
-								);
-
-							for (int32 Index = 0; Index < MeshBatches.Num(); ++Index)
+							if (bSectionIsSelected)
 							{
-								FMeshBatch ModifiedMeshElement(MeshBatches[Index]);
+								// Override the mesh's material with our material that draws the collision color
+								const FOverrideSelectionColorMaterialRenderProxy SelectedMaterialInstance(
+									GEngine->ShadedLevelColorationUnlitMaterial->GetRenderProxy(bSectionIsSelected, IsHovered()),
+									GetSelectionColor(GEngine->GetSelectedMaterialColor(), bSectionIsSelected, IsHovered())
+									);
+
+								FMeshBatch ModifiedMeshElement(MeshBatch);
 								ModifiedMeshElement.MaterialRenderProxy = &SelectedMaterialInstance;
 
 								const int32 NumPasses = DrawRichMesh(
@@ -755,17 +1012,14 @@ void FStaticMeshSceneProxy::DrawDynamicElements(FPrimitiveDrawInterface* PDI,con
 									bSectionIsSelected,
 									bIsWireframeView
 									);
-								INC_DWORD_STAT_BY(STAT_StaticMeshTriangles, MeshBatches[Index].GetNumPrimitives() * NumPasses);
+								INC_DWORD_STAT_BY(STAT_StaticMeshTriangles, MeshBatch.GetNumPrimitives() * NumPasses);
 							}
-						}
-						else
+							else
 #endif
-						{
-							for (int32 Index = 0; Index < MeshBatches.Num(); ++Index)
 							{
 								const int32 NumPasses = DrawRichMesh(
 									PDI,
-									MeshBatches[Index],
+									MeshBatch,
 									WireframeColor,
 									UtilColor,
 									PropertyColor,
@@ -773,17 +1027,15 @@ void FStaticMeshSceneProxy::DrawDynamicElements(FPrimitiveDrawInterface* PDI,con
 									bSectionIsSelected,
 									bIsWireframeView
 									);
-								INC_DWORD_STAT_BY(STAT_StaticMeshTriangles,MeshBatches[Index].GetNumPrimitives() * NumPasses);
+								INC_DWORD_STAT_BY(STAT_StaticMeshTriangles,MeshBatch.GetNumPrimitives() * NumPasses);
 							}
 						}
-
 					}
 				}
 			}
 		}
 	}
 
-
 	if((bDrawSimpleCollision || bDrawWireframeCollision) && AllowDebugViewmodes())
 	{
 		if(BodySetup)
@@ -807,14 +1059,14 @@ void FStaticMeshSceneProxy::DrawDynamicElements(FPrimitiveDrawInterface* PDI,con
 						);
 
 					FTransform GeomTransform(GetLocalToWorld());
-					BodySetup->AggGeom.DrawAggGeom(PDI, GeomTransform, WireframeColor, /*Material=*/&SolidMaterialInstance, false, /*bSolid=*/true);
+					BodySetup->AggGeom.DrawAggGeom(PDI, GeomTransform, WireframeColor, /*Material=*/&SolidMaterialInstance, false, /*bSolid=*/true, UseEditorDepthTest());
 				}
 				// wireframe
 				else
 				{
 					FColor CollisionColor = FColor(157,149,223,255);
 					FTransform GeomTransform(GetLocalToWorld());
-					BodySetup->AggGeom.DrawAggGeom(PDI, GeomTransform, GetSelectionColor(CollisionColor, bProxyIsSelected, IsHovered()), /*Material=*/NULL, (Owner == NULL), /*bSolid=*/false);
+					BodySetup->AggGeom.DrawAggGeom(PDI, GeomTransform, GetSelectionColor(CollisionColor, bProxyIsSelected, IsHovered()), /*Material=*/NULL, (Owner == NULL), /*bSolid=*/false, UseEditorDepthTest());
 				}
 
 
@@ -863,12 +1115,12 @@ FPrimitiveViewRelevance FStaticMeshSceneProxy::GetViewRelevance(const FSceneView
 
 #if !(UE_BUILD_SHIPPING || UE_BUILD_TEST) || WITH_EDITOR
 	bool bDrawSimpleCollision = false, bDrawComplexCollision = false;
-	const bool bInCollisionView = IsCollisionView(View, bDrawSimpleCollision, bDrawComplexCollision);
+	const bool bInCollisionView = IsCollisionView(View->Family->EngineShowFlags, bDrawSimpleCollision, bDrawComplexCollision);
 #endif
 
 	if(
 #if !(UE_BUILD_SHIPPING || UE_BUILD_TEST) || WITH_EDITOR
-		IsRichView(View) || 
+		IsRichView(*View->Family) || 
 		View->Family->EngineShowFlags.Collision ||
 		bInCollisionView ||
 		View->Family->EngineShowFlags.Bounds ||
@@ -1166,17 +1418,8 @@ FPrimitiveSceneProxy* UStaticMeshComponent::CreateSceneProxy()
 		return NULL;
 	}
 
-#if WITH_EDITOR
 	FPrimitiveSceneProxy* Proxy = ::new FStaticMeshSceneProxy(this);
-	if (GIsEditor && Proxy)
-	{
-		SetupLightmapResolutionViewInfo(*Proxy);
-	}
 	return Proxy;
-#else
-	//@todo: figure out why i need a ::new (gcc3-specific)
-	return ::new FStaticMeshSceneProxy(this);
-#endif
 }
 
 bool UStaticMeshComponent::ShouldRecreateProxyOnUpdateTransform() const
diff --git a/Engine/Source/Runtime/Engine/Private/TextRenderComponent.cpp b/Engine/Source/Runtime/Engine/Private/TextRenderComponent.cpp
--- a/Engine/Source/Runtime/Engine/Private/TextRenderComponent.cpp
+++ b/Engine/Source/Runtime/Engine/Private/TextRenderComponent.cpp
@@ -356,6 +356,7 @@ public:
 
 	// Begin FPrimitiveSceneProxy interface
 	virtual void DrawDynamicElements(FPrimitiveDrawInterface* PDI,const FSceneView* View) override;
+	virtual void GetDynamicMeshElements(const TArray<const FSceneView*>& Views, const FSceneViewFamily& ViewFamily, uint32 VisibilityMap, FMeshElementCollector& Collector) const override;
 	virtual void DrawStaticElements(FStaticPrimitiveDrawInterface* PDI) override;
 	virtual FPrimitiveViewRelevance GetViewRelevance(const FSceneView* View) override;
 	virtual bool CanBeOccluded() const override;
@@ -400,6 +401,7 @@ FTextRenderSceneProxy::FTextRenderSceneProxy( UTextRenderComponent* Component) :
 	VerticalAlignment(Component->VerticalAlignment),
 	bAlwaysRenderAsText(Component->bAlwaysRenderAsText)
 {
+	WireframeColor = FLinearColor(1, 0, 0);
 	UMaterialInterface* EffectiveMaterial = 0;
 
 	if(Component->TextMaterial)
@@ -447,6 +449,46 @@ void FTextRenderSceneProxy::ReleaseRenderThreadResources()
 	VertexFactory.ReleaseResource();
 }
 
+void FTextRenderSceneProxy::GetDynamicMeshElements(const TArray<const FSceneView*>& Views, const FSceneViewFamily& ViewFamily, uint32 VisibilityMap, FMeshElementCollector& Collector) const 
+{
+	QUICK_SCOPE_CYCLE_COUNTER( STAT_TextRenderSceneProxy_GetDynamicMeshElements );
+
+	// Vertex factory will not been initialized when the text string is empty or font is invalid.
+	if(VertexFactory.IsInitialized())
+	{
+		for (int32 ViewIndex = 0; ViewIndex < Views.Num(); ViewIndex++)
+		{
+			if (VisibilityMap & (1 << ViewIndex))
+			{
+				const FSceneView* View = Views[ViewIndex];
+				// Draw the mesh.
+				FMeshBatch& Mesh = Collector.AllocateMesh();
+				FMeshBatchElement& BatchElement = Mesh.Elements[0];
+				BatchElement.IndexBuffer = &IndexBuffer;
+				Mesh.VertexFactory = &VertexFactory;
+				BatchElement.PrimitiveUniformBufferResource = &GetUniformBuffer();
+				BatchElement.FirstIndex = 0;
+				BatchElement.NumPrimitives = IndexBuffer.Indices.Num() / 3;
+				BatchElement.MinVertexIndex = 0;
+				BatchElement.MaxVertexIndex = VertexBuffer.Vertices.Num() - 1;
+				Mesh.ReverseCulling = IsLocalToWorldDeterminantNegative();
+				Mesh.bDisableBackfaceCulling = false;
+				Mesh.Type = PT_TriangleList;
+				Mesh.DepthPriorityGroup = SDPG_World;
+				const bool bUseSelectedMaterial = GIsEditor && (View->Family->EngineShowFlags.Selection) ? IsSelected() : false;
+				Mesh.MaterialRenderProxy = TextMaterial->GetRenderProxy(bUseSelectedMaterial);
+				Mesh.bCanApplyViewModeOverrides = !bAlwaysRenderAsText;
+
+				Collector.AddMesh(ViewIndex, Mesh);
+
+#if !(UE_BUILD_SHIPPING || UE_BUILD_TEST)
+				RenderBounds(Collector.GetPDI(ViewIndex), View->Family->EngineShowFlags, GetBounds(), IsSelected());
+#endif
+			}
+		}
+	}
+}
+
 void FTextRenderSceneProxy::DrawDynamicElements(FPrimitiveDrawInterface* PDI,const FSceneView* View)
 {
 	QUICK_SCOPE_CYCLE_COUNTER( STAT_TextRenderSceneProxy_DrawDynamicElements );
@@ -535,7 +577,7 @@ FPrimitiveViewRelevance FTextRenderSceneProxy::GetViewRelevance(const FSceneView
 	Result.bRenderCustomDepth = ShouldRenderCustomDepth();
 	Result.bRenderInMainPass = ShouldRenderInMainPass();
 
-	if( IsRichView(View) 
+	if( IsRichView(*View->Family) 
 		|| View->Family->EngineShowFlags.Bounds 
 		|| View->Family->EngineShowFlags.Collision 
 		|| IsSelected() 
diff --git a/Engine/Source/Runtime/Engine/Private/VectorField.cpp b/Engine/Source/Runtime/Engine/Private/VectorField.cpp
--- a/Engine/Source/Runtime/Engine/Private/VectorField.cpp
+++ b/Engine/Source/Runtime/Engine/Private/VectorField.cpp
@@ -323,6 +323,17 @@ void UVectorFieldStatic::PostEditChangeProperty(FPropertyChangedEvent& PropertyC
 }
 #endif // WITH_EDITOR
 
+class FVectorFieldCollectorResources : public FOneFrameResource
+{
+public:
+	FVectorFieldVisualizationVertexFactory VisualizationVertexFactory;
+
+	virtual ~FVectorFieldCollectorResources()
+	{
+		VisualizationVertexFactory.ReleaseResource();
+	}
+};
+
 /*------------------------------------------------------------------------------
 	Scene proxy for visualizing vector fields.
 ------------------------------------------------------------------------------*/
@@ -357,6 +368,30 @@ public:
 		VisualizationVertexFactory.InitResource();
 	}
 
+	virtual void GetDynamicMeshElements(const TArray<const FSceneView*>& Views, const FSceneViewFamily& ViewFamily, uint32 VisibilityMap, FMeshElementCollector& Collector) const override
+	{	
+		QUICK_SCOPE_CYCLE_COUNTER( STAT_VectorFieldSceneProxy_GetDynamicMeshElements );
+
+		for (int32 ViewIndex = 0; ViewIndex < Views.Num(); ViewIndex++)
+		{
+			if (VisibilityMap & (1 << ViewIndex))
+			{
+				const FSceneView* View = Views[ViewIndex];
+				FPrimitiveDrawInterface* PDI = Collector.GetPDI(ViewIndex);
+				DrawVectorFieldBounds(PDI, View, VectorFieldInstance);
+
+				// Draw a visualization of the vectors contained in the field when selected.
+				if (IsSelected() || View->Family->EngineShowFlags.VectorFields)
+				{
+					FVectorFieldCollectorResources& CollectorResources = Collector.AllocateOneFrameResource<FVectorFieldCollectorResources>();
+					CollectorResources.VisualizationVertexFactory.InitResource();
+
+					GetVectorFieldMesh(&CollectorResources.VisualizationVertexFactory, VectorFieldInstance, ViewIndex, Collector);
+				}
+			}
+		}
+	}
+
 	/**
 	 * Draws the bounding box for the vector field.
 	 */
diff --git a/Engine/Source/Runtime/Engine/Private/VectorFieldVisualization.cpp b/Engine/Source/Runtime/Engine/Private/VectorFieldVisualization.cpp
--- a/Engine/Source/Runtime/Engine/Private/VectorFieldVisualization.cpp
+++ b/Engine/Source/Runtime/Engine/Private/VectorFieldVisualization.cpp
@@ -269,3 +269,43 @@ void DrawVectorField(
 		PDI->DrawMesh(MeshBatch);
 	}
 }
+
+void GetVectorFieldMesh(
+	FVectorFieldVisualizationVertexFactory* VertexFactory,
+	FVectorFieldInstance* VectorFieldInstance,
+	int32 ViewIndex,
+	FMeshElementCollector& Collector)
+{
+	FVectorFieldResource* Resource = VectorFieldInstance->Resource;
+
+	if (Resource && IsValidRef(Resource->VolumeTextureRHI))
+	{
+		// Set up parameters.
+		FVectorFieldVisualizationParameters UniformParameters;
+		UniformParameters.VolumeToWorld = VectorFieldInstance->VolumeToWorld;
+		UniformParameters.VolumeToWorldNoScale = VectorFieldInstance->VolumeToWorldNoScale;
+		UniformParameters.VoxelSize = FVector( 1.0f / Resource->SizeX, 1.0f / Resource->SizeY, 1.0f / Resource->SizeZ );
+		UniformParameters.Scale = VectorFieldInstance->Intensity * Resource->Intensity;
+		VertexFactory->SetParameters(UniformParameters, Resource->VolumeTextureRHI);
+
+		// Create a mesh batch for the visualization.
+		FMeshBatch& MeshBatch = Collector.AllocateMesh();
+		MeshBatch.CastShadow = false;
+		MeshBatch.bUseAsOccluder = false;
+		MeshBatch.VertexFactory = VertexFactory;
+		MeshBatch.MaterialRenderProxy = GEngine->LevelColorationUnlitMaterial->GetRenderProxy(false);
+		MeshBatch.Type = PT_LineList;
+
+		// A single mesh element.
+		FMeshBatchElement& MeshElement = MeshBatch.Elements[0];
+		MeshElement.NumPrimitives = 1;
+		MeshElement.NumInstances = Resource->SizeX * Resource->SizeY * Resource->SizeZ;
+		MeshElement.FirstIndex = 0;
+		MeshElement.MinVertexIndex = 0;
+		MeshElement.MaxVertexIndex = 1;
+		MeshElement.PrimitiveUniformBufferResource = &GIdentityPrimitiveUniformBuffer;
+
+		MeshBatch.bCanApplyViewModeOverrides = false;
+		Collector.AddMesh(ViewIndex, MeshBatch);
+	}
+}
diff --git a/Engine/Source/Runtime/Engine/Private/VectorFieldVisualization.h b/Engine/Source/Runtime/Engine/Private/VectorFieldVisualization.h
--- a/Engine/Source/Runtime/Engine/Private/VectorFieldVisualization.h
+++ b/Engine/Source/Runtime/Engine/Private/VectorFieldVisualization.h
@@ -115,3 +115,9 @@ void DrawVectorField(
 	const FSceneView* View,
 	FVectorFieldVisualizationVertexFactory* VertexFactory,
 	FVectorFieldInstance* VectorFieldInstance );
+
+void GetVectorFieldMesh(
+	FVectorFieldVisualizationVertexFactory* VertexFactory,
+	FVectorFieldInstance* VectorFieldInstance,
+	int32 ViewIndex,
+	FMeshElementCollector& Collector);
\ No newline at end of file
diff --git a/Engine/Source/Runtime/Engine/Public/AI/NavLinkRenderingProxy.h b/Engine/Source/Runtime/Engine/Public/AI/NavLinkRenderingProxy.h
--- a/Engine/Source/Runtime/Engine/Public/AI/NavLinkRenderingProxy.h
+++ b/Engine/Source/Runtime/Engine/Public/AI/NavLinkRenderingProxy.h
@@ -36,12 +36,15 @@ public:
 	/** Initialization constructor. */
 	FNavLinkRenderingProxy(const UPrimitiveComponent* InComponent);
 	virtual void DrawDynamicElements(FPrimitiveDrawInterface* PDI,const FSceneView* View) override;
+	virtual void GetDynamicMeshElements(const TArray<const FSceneView*>& Views, const FSceneViewFamily& ViewFamily, uint32 VisibilityMap, FMeshElementCollector& Collector) const override;
 	virtual FPrimitiveViewRelevance GetViewRelevance(const FSceneView* View) override;
 	virtual uint32 GetMemoryFootprint( void ) const override;
 	uint32 GetAllocatedSize( void ) const;
 	void StorePointLinks(const FTransform& LocalToWorld, const TArray<FNavigationLink>& LinksArray);
 	void StoreSegmentLinks(const FTransform& LocalToWorld, const TArray<FNavigationSegmentLink>& LinksArray);
 
+	static void GetLinkMeshes(const TArray<FNavLinkDrawing>& OffMeshPointLinks, const TArray<FNavLinkSegmentDrawing>& OffMeshSegmentLinks, TArray<float>& StepHeights, FMaterialRenderProxy* const MeshColorInstance, int32 ViewIndex, FMeshElementCollector& Collector);
+
 	/** made static to allow consistent navlinks drawing even if something is drawing links without FNavLinkRenderingProxy */
 	static void DrawLinks(FPrimitiveDrawInterface* PDI, TArray<FNavLinkDrawing>& OffMeshPointLinks, TArray<FNavLinkSegmentDrawing>& OffMeshSegmentLinks, TArray<float>& StepHeights, FMaterialRenderProxy* const MeshColorInstance);
 };
diff --git a/Engine/Source/Runtime/Engine/Public/BatchedElements.h b/Engine/Source/Runtime/Engine/Public/BatchedElements.h
--- a/Engine/Source/Runtime/Engine/Public/BatchedElements.h
+++ b/Engine/Source/Runtime/Engine/Public/BatchedElements.h
@@ -12,6 +12,17 @@
 #include "SceneTypes.h"
 
 
+namespace EBlendModeFilter
+{
+	enum Type
+	{
+		None = 0,
+		OpaqueAndMasked = 1,
+		Translucent = 2,
+		All = (OpaqueAndMasked | Translucent)
+	};
+};
+
 /** The type used to store batched line vertices. */
 struct FSimpleElementVertex
 {
@@ -165,7 +176,7 @@ public:
 	 * @param View			Optional FSceneView for shaders that need access to view constants
 	 * @param DepthTexture	DepthTexture for manual depth testing with editor compositing in the pixel shader
 	 */
-	bool Draw(FRHICommandList& RHICmdList, bool bNeedToSwitchVerticalAxis, const FMatrix& Transform, uint32 ViewportSizeX, uint32 ViewportSizeY, bool bHitTesting, float Gamma = 1.0f, const FSceneView* View = NULL, FTexture2DRHIRef DepthTexture = FTexture2DRHIRef()) const;
+	bool Draw(FRHICommandList& RHICmdList, bool bNeedToSwitchVerticalAxis, const FMatrix& Transform, uint32 ViewportSizeX, uint32 ViewportSizeY, bool bHitTesting, float Gamma = 1.0f, const FSceneView* View = NULL, FTexture2DRHIRef DepthTexture = FTexture2DRHIRef(), EBlendModeFilter::Type Filter = EBlendModeFilter::All) const;
 	
 	FORCEINLINE bool HasPrimsToDraw() const
 	{
