FPrimitiveSceneProxy::PreRenderView / DrawDynamicElements refactored into GetDynamicMeshElements

diff --git a/Engine/Source/Runtime/Engine/Public/DebugRenderSceneProxy.h b/Engine/Source/Runtime/Engine/Public/DebugRenderSceneProxy.h
--- a/Engine/Source/Runtime/Engine/Public/DebugRenderSceneProxy.h
+++ b/Engine/Source/Runtime/Engine/Public/DebugRenderSceneProxy.h
@@ -28,6 +28,8 @@ public:
 	 */
 	ENGINE_API virtual void DrawDynamicElements(FPrimitiveDrawInterface* PDI,const FSceneView* View);
 
+	ENGINE_API virtual void GetDynamicMeshElements(const TArray<const FSceneView*>& Views, const FSceneViewFamily& ViewFamily, uint32 VisibilityMap, FMeshElementCollector& Collector) const override;
+
 	/**
 	 * Draws a line with an arrow at the end.
 	 *
@@ -36,7 +38,7 @@ public:
 	 * @param Color		Color of the line.
 	 * @param Mag		Size of the arrow.
 	 */
-	void DrawLineArrow(FPrimitiveDrawInterface* PDI,const FVector &Start,const FVector &End,const FColor &Color,float Mag);
+	void DrawLineArrow(FPrimitiveDrawInterface* PDI,const FVector &Start,const FVector &End,const FColor &Color,float Mag) const;
 
 	ENGINE_API virtual void DrawDebugLabels(UCanvas* Canvas, APlayerController*);
 
diff --git a/Engine/Source/Runtime/Engine/Public/DynamicMeshBuilder.h b/Engine/Source/Runtime/Engine/Public/DynamicMeshBuilder.h
--- a/Engine/Source/Runtime/Engine/Public/DynamicMeshBuilder.h
+++ b/Engine/Source/Runtime/Engine/Public/DynamicMeshBuilder.h
@@ -93,6 +93,9 @@ public:
 	/** Add many indices to the mesh. */
 	ENGINE_API void AddTriangles(const TArray<int32> &InIndices);
 
+	/** Adds a mesh of what's been built so far to the collector. */
+	ENGINE_API void GetMesh(const FMatrix& LocalToWorld,const FMaterialRenderProxy* MaterialRenderProxy,uint8 DepthPriorityGroup,bool bDisableBackfaceCulling, bool bReceivesDecals, int32 ViewIndex, FMeshElementCollector& Collector);
+
 	/**
 	 * Draws the mesh to the given primitive draw interface.
 	 * @param PDI - The primitive draw interface to draw the mesh on.
diff --git a/Engine/Source/Runtime/Engine/Public/HitProxies.h b/Engine/Source/Runtime/Engine/Public/HitProxies.h
--- a/Engine/Source/Runtime/Engine/Public/HitProxies.h
+++ b/Engine/Source/Runtime/Engine/Public/HitProxies.h
@@ -91,6 +91,11 @@ public:
 	 */
 	friend ENGINE_API class HHitProxy* GetHitProxyById(FHitProxyId Id);
 
+	friend bool operator==(FHitProxyId X, FHitProxyId Y)
+	{
+		return X.Index == Y.Index;
+	}
+
 private:
 	
 	/** Initialization constructor. */
diff --git a/Engine/Source/Runtime/Engine/Public/Landscape/LandscapeRender.h b/Engine/Source/Runtime/Engine/Public/Landscape/LandscapeRender.h
--- a/Engine/Source/Runtime/Engine/Public/Landscape/LandscapeRender.h
+++ b/Engine/Source/Runtime/Engine/Public/Landscape/LandscapeRender.h
@@ -83,7 +83,7 @@ END_UNIFORM_BUFFER_STRUCT(FLandscapeUniformShaderParameters)
 struct FLandscapeBatchElementParams
 {
 	const TUniformBuffer<FLandscapeUniformShaderParameters>* LandscapeUniformShaderParametersResource;
-	FMatrix* LocalToWorldNoScalingPtr;
+	const FMatrix* LocalToWorldNoScalingPtr;
 
 	// LOD calculation-related params
 	const class FLandscapeComponentSceneProxy* SceneProxy;
@@ -92,6 +92,11 @@ struct FLandscapeBatchElementParams
 	int32	CurrentLOD;
 };
 
+class FLandscapeElementParamArray : public FOneFrameResource
+{
+public:
+	TArray<FLandscapeBatchElementParams, SceneRenderingAllocator> ElementParams;
+};
 
 /** Pixel shader parameters for use with FLandscapeVertexFactory */
 class FLandscapeVertexFactoryPixelShaderParameters : public FVertexFactoryShaderParameters
@@ -446,8 +451,6 @@ protected:
 
 	const ULandscapeComponent* LandscapeComponent;
 
-	FLinearColor LevelColor;
-
 	FVector2D				NeighborPosition[LANDSCAPE_NEIGHBOR_NUM];
 	int32					ForcedLOD;
 	int32					LODBias;
@@ -477,6 +480,7 @@ public:
 
 	// FPrimitiveSceneProxy interface.
 	virtual void DrawStaticElements(FStaticPrimitiveDrawInterface* PDI) override;
+	virtual void GetDynamicMeshElements(const TArray<const FSceneView*>& Views, const FSceneViewFamily& ViewFamily, uint32 VisibilityMap, FMeshElementCollector& Collector) const override;
 	virtual void DrawDynamicElements(FPrimitiveDrawInterface* PDI,const FSceneView* View) override;
 	virtual uint32 GetMemoryFootprint( void ) const override { return( sizeof( *this ) + GetAllocatedSize() ); }
 	virtual FPrimitiveViewRelevance GetViewRelevance(const FSceneView* View) override;
diff --git a/Engine/Source/Runtime/Engine/Public/MeshBatch.h b/Engine/Source/Runtime/Engine/Public/MeshBatch.h
--- a/Engine/Source/Runtime/Engine/Public/MeshBatch.h
+++ b/Engine/Source/Runtime/Engine/Public/MeshBatch.h
@@ -9,7 +9,10 @@ class FLightCacheInterface;
  */
 struct FMeshBatchElement
 {
+	/** Primitive uniform buffer to use for rendering. */
 	const TUniformBuffer<FPrimitiveUniformShaderParameters>* PrimitiveUniformBufferResource;
+
+	/** Used for lifetime management of a temporary uniform buffer, can be NULL. */
 	TUniformBufferRef<FPrimitiveUniformShaderParameters> PrimitiveUniformBuffer;
 
 	const FIndexBuffer* IndexBuffer;
@@ -19,7 +22,7 @@ struct FMeshBatchElement
 	uint32 MinVertexIndex;
 	uint32 MaxVertexIndex;
 	int32 UserIndex;
-	void* UserData;
+	const void* UserData;
 	float MinScreenSize;
 	float MaxScreenSize;
 
@@ -67,7 +70,26 @@ struct FMeshBatch
 	uint32 DepthPriorityGroup : SDPG_NumBits;
 	uint32 bUseAsOccluder : 1;
 
-	// can be 0
+	/** Whether view mode overrides can be applied to this mesh eg unlit, wireframe. */
+	uint32 bCanApplyViewModeOverrides : 1;
+
+	/** 
+	 * Whether to treat the batch as selected in special viewmodes like wireframe. 
+	 * This is needed instead of just Proxy->IsSelected() because some proxies do weird things with selection highlighting, like FStaticMeshSceneProxy.
+	 */
+	uint32 bUseWireframeSelectionColoring : 1;
+
+	/** 
+	 * Whether the batch should receive the selection outline.  
+	 * This is useful for proxies which support selection on a per-mesh batch basis.
+	 * They submit multiple mesh batches when selected, some of which have bUseSelectionOutline enabled.
+	 */
+	uint32 bUseSelectionOutline : 1;
+
+	/** Whether the mesh batch can be selected through editor selection, aka hit proxies. */
+	uint32 bSelectable : 1;
+
+	// can be NULL
 	const FLightCacheInterface* LCI;
 
 	/** 
@@ -77,22 +99,25 @@ struct FMeshBatch
 	 */
 	const void* DynamicVertexData;
 
-	// can be 0
+	/** Vertex factory for rendering, required. */
 	const FVertexFactory* VertexFactory;
-	// can be 0
+
+	/** Material proxy for rendering, required. */
 	const FMaterialRenderProxy* MaterialRenderProxy;
 
+	/** The current hit proxy ID being rendered. */
+	FHitProxyId BatchHitProxyId;
 
 	FORCEINLINE bool IsTranslucent(ERHIFeatureLevel::Type InFeatureLevel) const
 	{
 		// Note: blend mode does not depend on the feature level we are actually rendering in.
-		return MaterialRenderProxy && IsTranslucentBlendMode(MaterialRenderProxy->GetMaterial(InFeatureLevel)->GetBlendMode());
+		return IsTranslucentBlendMode(MaterialRenderProxy->GetMaterial(InFeatureLevel)->GetBlendMode());
 	}
 
 	FORCEINLINE bool IsMasked(ERHIFeatureLevel::Type InFeatureLevel) const
 	{
 		// Note: blend mode does not depend on the feature level we are actually rendering in.
-		return MaterialRenderProxy && MaterialRenderProxy->GetMaterial(InFeatureLevel)->IsMasked();
+		return MaterialRenderProxy->GetMaterial(InFeatureLevel)->IsMasked();
 	}
 
 	/** Converts from an int32 index into a int8 */
@@ -122,7 +147,7 @@ struct FMeshBatch
 		int32 Count=0;
 		for( int32 ElementIdx=0;ElementIdx<Elements.Num();ElementIdx++ )
 		{
-			Count += Elements[ElementIdx].NumPrimitives;
+			Count += Elements[ElementIdx].NumPrimitives * Elements[ElementIdx].NumInstances;
 		}
 		return Count;
 	}
@@ -149,6 +174,10 @@ struct FMeshBatch
 	,	Type(PT_TriangleList)
 	,	DepthPriorityGroup(SDPG_World)
 	,	bUseAsOccluder(true)
+	,	bCanApplyViewModeOverrides(false)
+	,	bUseWireframeSelectionColoring(false)
+	,	bUseSelectionOutline(true)
+	,	bSelectable(true)
 	,	LCI(NULL)
 	,	DynamicVertexData(NULL)
 	,	VertexFactory(NULL)
diff --git a/Engine/Source/Runtime/Engine/Public/MeshParticleVertexFactory.h b/Engine/Source/Runtime/Engine/Public/MeshParticleVertexFactory.h
--- a/Engine/Source/Runtime/Engine/Public/MeshParticleVertexFactory.h
+++ b/Engine/Source/Runtime/Engine/Public/MeshParticleVertexFactory.h
@@ -7,6 +7,8 @@
 #pragma once
 
 #include "UniformBuffer.h"
+//@todo - parallelrendering - remove once FOneFrameResource no longer needs to be referenced in header
+#include "SceneManagement.h"
 
 /**
  * Uniform buffer for mesh particle vertex factories.
@@ -65,8 +67,9 @@ public:
 		}
 	};
 
-	struct FBatchParametersCPU
+	class FBatchParametersCPU : public FOneFrameResource
 	{
+	public:
 		const struct FMeshParticleInstanceVertex* InstanceBuffer;
 		const struct FMeshParticleInstanceVertexDynamicParameter* DynamicParameterBuffer;
 	};
@@ -74,8 +77,11 @@ public:
 	/** Default constructor. */
 	FMeshParticleVertexFactory(EParticleVertexFactoryType InType, ERHIFeatureLevel::Type InFeatureLevel)
 		: FParticleVertexFactoryBase(InType, InFeatureLevel)
-	{
-	}
+	{}
+
+	FMeshParticleVertexFactory()
+		: FParticleVertexFactoryBase(PVFT_MAX, ERHIFeatureLevel::Num)
+	{}
 
 	/**
 	 * Should we cache the material's shadertype on this platform with this vertex factory? 
diff --git a/Engine/Source/Runtime/Engine/Public/ParticleBeamTrailVertexFactory.h b/Engine/Source/Runtime/Engine/Public/ParticleBeamTrailVertexFactory.h
--- a/Engine/Source/Runtime/Engine/Public/ParticleBeamTrailVertexFactory.h
+++ b/Engine/Source/Runtime/Engine/Public/ParticleBeamTrailVertexFactory.h
@@ -31,8 +31,11 @@ public:
 	/** Default constructor. */
 	FParticleBeamTrailVertexFactory( EParticleVertexFactoryType InType, ERHIFeatureLevel::Type InFeatureLevel )
 		: FParticleVertexFactoryBase(InType, InFeatureLevel)
-	{
-	}
+	{}
+
+	FParticleBeamTrailVertexFactory()
+		: FParticleVertexFactoryBase(PVFT_MAX, ERHIFeatureLevel::Num)
+	{}
 
 	/**
 	 * Should we cache the material's shadertype on this platform with this vertex factory? 
diff --git a/Engine/Source/Runtime/Engine/Public/ParticleHelper.h b/Engine/Source/Runtime/Engine/Public/ParticleHelper.h
--- a/Engine/Source/Runtime/Engine/Public/ParticleHelper.h
+++ b/Engine/Source/Runtime/Engine/Public/ParticleHelper.h
@@ -989,10 +989,10 @@ struct FEmitterDynamicParameterPayload
  *	@param	InParticle					The particle being processed
  *	@param	OutDynamicData				The dynamic data from the particle
  */
-FORCEINLINE void GetDynamicValueFromPayload(int32 InDynamicPayloadOffset, FBaseParticle& InParticle, FVector4& OutDynamicData)
+FORCEINLINE void GetDynamicValueFromPayload(int32 InDynamicPayloadOffset, const FBaseParticle& InParticle, FVector4& OutDynamicData)
 {
 	checkSlow(InDynamicPayloadOffset > 0);
-	FEmitterDynamicParameterPayload* DynPayload = ((FEmitterDynamicParameterPayload*)((uint8*)(&InParticle) + InDynamicPayloadOffset));
+	const FEmitterDynamicParameterPayload* DynPayload = ((const FEmitterDynamicParameterPayload*)((uint8*)(&InParticle) + InDynamicPayloadOffset));
 	OutDynamicData.X = DynPayload->DynamicParameterValue[0];
 	OutDynamicData.Y = DynPayload->DynamicParameterValue[1];
 	OutDynamicData.Z = DynPayload->DynamicParameterValue[2];
@@ -1079,42 +1079,6 @@ struct FAsyncBufferFillData
 	}
 };
 
-/*-----------------------------------------------------------------------------
-	Async Fill Task, simple wrapper to forward the request to a FDynamicSpriteEmitterDataBase
------------------------------------------------------------------------------*/
-
-struct FAsyncParticleFill
-{
-	/** Emitter to forward to   */
-	struct FDynamicSpriteEmitterDataBase* Parent;
-
-	/** Constructor, just sets up the parent pointer  
-	  * @param InParent emitter to forward the eventual async call to
-	*/
-	FAsyncParticleFill(struct FDynamicSpriteEmitterDataBase* InParent)
-		: Parent(InParent)
-	{
-	}
-
-	/** Work function, just forwards the request to the parent  */
-	void DoTask(ENamedThreads::Type CurrentThread, const FGraphEventRef& MyCompletionGraphEvent);
-
-	FORCEINLINE TStatId GetStatId() const
-	{
-		return GET_STATID( STAT_ParticleAsyncTime );
-	}
-
-	static ENamedThreads::Type GetDesiredThread()
-	{
-		return ENamedThreads::AnyThread;
-	}
-
-	static ESubsequentsMode::Type GetSubsequentsMode() 
-	{ 
-		return ESubsequentsMode::TrackSubsequents; 
-	}
-};
-
 // TAsyncBufferFillTasks - handy typedef for an inline array of buffer fill tasks
 typedef TArray<FAsyncBufferFillData, TInlineAllocator<2> > TAsyncBufferFillTasks;
 
@@ -1164,15 +1128,6 @@ public:
 protected:
 	void ClearPoolInternal();
 
-	/** 
-	 *	Create a vertex factory for the given type.
-	 *
-	 *	@param	InType						The type of vertex factory to create.
-	 *
-	 *	@return	FParticleVertexFactoryBase*	The created VF; NULL if invalid InType
-	 */
-	FParticleVertexFactoryBase* CreateParticleVertexFactory(EParticleVertexFactoryType InType, ERHIFeatureLevel::Type InFeatureLevel);
-
 	TArray<FParticleVertexFactoryBase*>	VertexFactoriesAvailable[PVFT_MAX][ERHIFeatureLevel::Num];
 	TArray<FParticleVertexFactoryBase*>	VertexFactories;
 };
@@ -1342,7 +1297,7 @@ struct FDynamicEmitterDataBase
 	 *
 	 *	@return	bool			true if successful, false if failed
 	 */
-	virtual void CreateRenderThreadResources(const FParticleSystemSceneProxy* InOwnerProxy)
+	virtual void UpdateRenderThreadResourcesEmitter(const FParticleSystemSceneProxy* InOwnerProxy)
 	{
 	}
 
@@ -1365,7 +1320,9 @@ struct FDynamicEmitterDataBase
 	 *	@param	PDI			The primitive draw interface to render with
 	 *	@param	View		The scene view being rendered
 	 */
-	virtual int32 Render(FParticleSystemSceneProxy* Proxy, FPrimitiveDrawInterface* PDI,const FSceneView* View) = 0;
+	virtual int32 RenderEmitter(FParticleSystemSceneProxy* Proxy, FPrimitiveDrawInterface* PDI,const FSceneView* View) = 0;
+
+	virtual void GetDynamicMeshElementsEmitter(const FParticleSystemSceneProxy* Proxy, const FSceneView* View, const FSceneViewFamily& ViewFamily, int32 ViewIndex, FMeshElementCollector& Collector) const {}
 
 	/**
 	 *	Retrieve the material render proxy to use for rendering this emitter. PURE VIRTUAL
@@ -1474,7 +1431,6 @@ struct FDynamicSpriteEmitterDataBase : public FDynamicEmitterDataBase
 {
 	FDynamicSpriteEmitterDataBase(const UParticleModuleRequired* RequiredModule) : 
 		FDynamicEmitterDataBase(RequiredModule),
-		AsyncTask(NULL),
 		bUsesDynamicParameter( false )
 	{
 		MaterialResource[0] = NULL;
@@ -1483,7 +1439,6 @@ struct FDynamicSpriteEmitterDataBase : public FDynamicEmitterDataBase
 
 	virtual ~FDynamicSpriteEmitterDataBase()
 	{
-		EnsureAsyncTaskComplete();
 	}
 
 	/**
@@ -1513,7 +1468,7 @@ struct FDynamicSpriteEmitterDataBase : public FDynamicEmitterDataBase
 	 */
 	void SortSpriteParticles(int32 SortMode, bool bLocalSpace, 
 		int32 ParticleCount, const TArray<uint8>& ParticleData, int32 ParticleStride, const TArray<uint16>& ParticleIndices,
-		const FSceneView* View, const FMatrix& LocalToWorld, FParticleOrder* ParticleOrder);
+		const FSceneView* View, const FMatrix& LocalToWorld, FParticleOrder* ParticleOrder) const;
 
 	/**
 	 *	Get the vertex stride for the dynamic rendering data
@@ -1558,64 +1513,33 @@ struct FDynamicSpriteEmitterDataBase : public FDynamicEmitterDataBase
 	 *	@param	View		The scene view being rendered
 	 *	@param	bCrosses	If true, render Crosses at particle position; false, render points
 	 */
-	virtual void RenderDebug(FParticleSystemSceneProxy* Proxy, FPrimitiveDrawInterface* PDI, const FSceneView* View, bool bCrosses);
+	virtual void RenderDebug(const FParticleSystemSceneProxy* Proxy, FPrimitiveDrawInterface* PDI, const FSceneView* View, bool bCrosses) const;
 
-	/**
-	 *	Fill index and vertex buffers. Often called from a different thread
-	 *
-	 */
-	void DoBufferFill()
+	virtual void DoBufferFill(FAsyncBufferFillData& Me) const
 	{
-		for (int32 TaskIndex = 0; TaskIndex < AsyncBufferFillTasks.Num(); TaskIndex++) 
-		{
-			DoBufferFill(AsyncBufferFillTasks[TaskIndex]);
-		}
-	}
-	/**
-	 *	Fill index and vertex buffers. Often called from a different thread
-	 *
-	 *	@param	Me			buffer pair to compute
-	 */
-	virtual void DoBufferFill(FAsyncBufferFillData& Me)
-	{
-		// this must be overridden, but in some cases a destructor call will leave this a no-op
-		// because the vtable has been reset to the base class 
-		// checkf(0, TEXT("DoBufferFill MUST be overridden"));
+		// Must be overridden if called
+		check(0);
 	}
+
 	/**
 	 *	Set up an buffer for async filling
 	 *
 	 *	@param	Proxy					The primitive scene proxy for the emitter.
-	 *	@param	InBufferIndex			Index of this buffer
 	 *	@param	InView					View for this buffer
 	 *	@param	InVertexCount			Count of verts for this buffer
 	 *	@param	InVertexSize			Stride of these verts, only used for verification
 	 *	@param	InDynamicParameterVertexStride	Stride of the dynamic parameter
 	 */
-	void BuildViewFillData(FParticleSystemSceneProxy* Proxy, int32 InBufferIndex,const FSceneView *InView,int32 InVertexCount,int32 InVertexSize,int32 InDynamicParameterVertexSize);
-
-	/**
-	 *	Set up all buffers for async filling
-	 *
-	 *	@param	Proxy							The primitive scene proxy for the emitter.
-	 *	@param	ViewFamily						View family to process
-	 *	@param	VisibilityMap					Visibility map for the sub-views
-	 *	@param	bOnlyOneView					If true, then we don't need per-view buffers
-	 *	@param	InVertexCount					Count of verts for this buffer
-	 *	@param	InVertexSize					Stride of these verts, only used for verification
-	 *	@param	InDynamicParameterVertexStride	Stride of the dynamic parameter
-	 */
-	void BuildViewFillDataAndSubmit(FParticleSystemSceneProxy* Proxy, const FSceneViewFamily* ViewFamily,const uint32 VisibilityMap,bool bOnlyOneView,int32 InVertexCount,int32 InVertexSize, int32 InDynamicParameterVertexSize);
-
-	void EnsureAsyncTaskComplete()
-	{
-		if (AsyncTask.GetReference())
-		{
-			SCOPE_CYCLE_COUNTER(STAT_ParticleAsyncWaitTime);
-			FTaskGraphInterface::Get().WaitUntilTaskCompletes(AsyncTask, ENamedThreads::RenderThread_Local);
-			AsyncTask = NULL;
-		}
-	}
+	void BuildViewFillData(
+		const FParticleSystemSceneProxy* Proxy, 
+		const FSceneView *InView, 
+		int32 InVertexCount, 
+		int32 InVertexSize, 
+		int32 InDynamicParameterVertexSize, 
+		FGlobalDynamicVertexBuffer::FAllocation& DynamicVertexAllocation,
+		FGlobalDynamicIndexBuffer::FAllocation& DynamicIndexAllocation,
+		FGlobalDynamicVertexBuffer::FAllocation* DynamicParameterAllocation,
+		FAsyncBufferFillData& Data) const;
 
 	/**
 	 *	Called to verify that a buffer is ready to use, blocks to wait and can sometimes execute the buffer fill on the current thread
@@ -1623,7 +1547,7 @@ struct FDynamicSpriteEmitterDataBase : public FDynamicEmitterDataBase
 	 *	@param	InView			View to look up in the buffer table
 	 *  @return Completed buffers
 	 */
-	const FAsyncBufferFillData& EnsureFillCompletion(const FSceneView *InView)
+	const FAsyncBufferFillData& EnsureFillCompletion(const FSceneView *InView) const
 	{
 		check(AsyncBufferFillTasks.Num());
 		// - 1 because we often fill only one, for _all_ views, if no match we always take the last one
@@ -1641,9 +1565,6 @@ struct FDynamicSpriteEmitterDataBase : public FDynamicEmitterDataBase
 	/** Async task is queued for execution */
 	bool									bAsyncTaskOutstanding;
 
-	/** Async task that is queued in the hi priority pool */
-	FGraphEventRef							AsyncTask;
-
 	/** Array of buffers for filling by async task */
 	TAsyncBufferFillTasks					AsyncBufferFillTasks;
 
@@ -1734,7 +1655,7 @@ struct FDynamicSpriteEmitterData : public FDynamicSpriteEmitterDataBase
 	 *
 	 *	@return	bool			true if successful, false if failed
 	 */
-	bool GetVertexAndIndexData(void* VertexData, void* DynamicParameterVertexData, void* FillIndexData, FParticleOrder* ParticleOrder, const FVector& InCameraPosition, const FMatrix& InLocalToWorld);
+	bool GetVertexAndIndexData(void* VertexData, void* DynamicParameterVertexData, void* FillIndexData, FParticleOrder* ParticleOrder, const FVector& InCameraPosition, const FMatrix& InLocalToWorld) const;
 
 	/**
 	 *	Retrieve the vertex and (optional) index required to render this emitter.
@@ -1749,7 +1670,7 @@ struct FDynamicSpriteEmitterData : public FDynamicSpriteEmitterDataBase
 	 *
 	 *	@return	bool			true if successful, false if failed
 	 */
-	bool GetVertexAndIndexDataNonInstanced(void* VertexData, void* DynamicParameterVertexData, void* FillIndexData, FParticleOrder* ParticleOrder, const FVector& InCameraPosition, const FMatrix& InLocalToWorld);
+	bool GetVertexAndIndexDataNonInstanced(void* VertexData, void* DynamicParameterVertexData, void* FillIndexData, FParticleOrder* ParticleOrder, const FVector& InCameraPosition, const FMatrix& InLocalToWorld) const;
 
 	/**
 	 *	Called during InitViews for view processing on scene proxies before rendering them
@@ -1765,6 +1686,8 @@ struct FDynamicSpriteEmitterData : public FDynamicSpriteEmitterDataBase
 	/** Gathers simple lights for this emitter. */
 	virtual void GatherSimpleLights(const FParticleSystemSceneProxy* Proxy, const FSceneViewFamily& ViewFamily, FSimpleLightArray& OutParticleLights) const override;
 
+	virtual void GetDynamicMeshElementsEmitter(const FParticleSystemSceneProxy* Proxy, const FSceneView* View, const FSceneViewFamily& ViewFamily, int32 ViewIndex, FMeshElementCollector& Collector) const override;
+
 	/**
 	 *	Render thread only draw call
 	 *
@@ -1772,7 +1695,7 @@ struct FDynamicSpriteEmitterData : public FDynamicSpriteEmitterDataBase
 	 *	@param	PDI			The primitive draw interface to render with
 	 *	@param	View		The scene view being rendered
 	 */
-	virtual int32 Render(FParticleSystemSceneProxy* Proxy, FPrimitiveDrawInterface* PDI,const FSceneView* View) override;
+	virtual int32 RenderEmitter(FParticleSystemSceneProxy* Proxy, FPrimitiveDrawInterface* PDI,const FSceneView* View) override;
 
 	/**
 	 *	Create the render thread resources for this emitter data
@@ -1781,7 +1704,7 @@ struct FDynamicSpriteEmitterData : public FDynamicSpriteEmitterDataBase
 	 *
 	 *	@return	bool			true if successful, false if failed
 	 */
-	virtual void CreateRenderThreadResources(const FParticleSystemSceneProxy* InOwnerProxy);
+	virtual void UpdateRenderThreadResourcesEmitter(const FParticleSystemSceneProxy* InOwnerProxy);
 
 	/**
 	 *	Release the render thread resources for this emitter data
@@ -1877,7 +1800,7 @@ struct FDynamicMeshEmitterData : public FDynamicSpriteEmitterDataBase
 	 *
 	 *	@return	bool			true if successful, false if failed
 	 */
-	virtual void CreateRenderThreadResources(const FParticleSystemSceneProxy* InOwnerProxy);
+	virtual void UpdateRenderThreadResourcesEmitter(const FParticleSystemSceneProxy* InOwnerProxy);
 
 	/**
 	 *	Release the render thread resources for this emitter data
@@ -1888,6 +1811,8 @@ struct FDynamicMeshEmitterData : public FDynamicSpriteEmitterDataBase
 	 */
 	virtual void ReleaseRenderThreadResources(const FParticleSystemSceneProxy* InOwnerProxy);
 
+	virtual void GetDynamicMeshElementsEmitter(const FParticleSystemSceneProxy* Proxy, const FSceneView* View, const FSceneViewFamily& ViewFamily, int32 ViewIndex, FMeshElementCollector& Collector) const override;
+
 	/**
 	 *	Render thread only draw call
 	 *
@@ -1895,7 +1820,7 @@ struct FDynamicMeshEmitterData : public FDynamicSpriteEmitterDataBase
 	 *	@param	PDI			The primitive draw interface to render with
 	 *	@param	View		The scene view being rendered
 	 */
-	virtual int32 Render(FParticleSystemSceneProxy* Proxy, FPrimitiveDrawInterface* PDI,const FSceneView* View);
+	virtual int32 RenderEmitter(FParticleSystemSceneProxy* Proxy, FPrimitiveDrawInterface* PDI,const FSceneView* View);
 	
 	/**
 	 *	Retrieve the instance data required to render this emitter.
@@ -1906,7 +1831,7 @@ struct FDynamicMeshEmitterData : public FDynamicSpriteEmitterDataBase
 	 *	@param	Proxy                   The scene proxy for the particle system that owns this emitter
 	 *	@param	View                    The scene view being rendered
 	 */
-	void GetInstanceData(void* InstanceData, void* DynamicParameterData, FParticleSystemSceneProxy* Proxy, const FSceneView* View);
+	void GetInstanceData(void* InstanceData, void* DynamicParameterData, const FParticleSystemSceneProxy* Proxy, const FSceneView* View) const;
 
 	/**
 	 *	Helper function for retrieving the particle transform.
@@ -1920,8 +1845,8 @@ struct FDynamicMeshEmitterData : public FDynamicSpriteEmitterDataBase
 	 *	@param	View						The scene view being rendered
 	 *	@param	OutTransform				The InstanceToWorld transform matrix for the particle
 	 */
-	void GetParticleTransform(FBaseParticle& InParticle, const FVector& CameraPosition, const FVector& CameraFacingOpVector, 
-		const FQuat& PointToLockedAxis, FParticleSystemSceneProxy* Proxy, const FSceneView* View, FMatrix& OutTransformMat);
+	void GetParticleTransform(const FBaseParticle& InParticle, const FVector& CameraPosition, const FVector& CameraFacingOpVector, 
+		const FQuat& PointToLockedAxis, const FParticleSystemSceneProxy* Proxy, const FSceneView* View, FMatrix& OutTransformMat) const;
 
 	/**
 	 *	Called during InitViews for view processing on scene proxies before rendering them
@@ -1961,7 +1886,7 @@ struct FDynamicMeshEmitterData : public FDynamicSpriteEmitterDataBase
 	/**
 	 *	 Initialize this emitter's vertex factory with the vertex buffers from the mesh's rendering data.
 	 */
-	void SetupVertexFactory( FMeshParticleVertexFactory* VertexFactory, FStaticMeshLODResources& LODResources);
+	void SetupVertexFactory( FMeshParticleVertexFactory* VertexFactory, FStaticMeshLODResources& LODResources) const;
 
 	/** Returns the source data for this particle system */
 	virtual const FDynamicEmitterReplayDataBase& GetSource() const
@@ -2189,7 +2114,9 @@ struct FDynamicBeam2EmitterData : public FDynamicSpriteEmitterDataBase
 	 *
 	 *	@return	bool			true if successful, false if failed
 	 */
-	virtual void CreateRenderThreadResources(const FParticleSystemSceneProxy* InOwnerProxy);
+	virtual void UpdateRenderThreadResourcesEmitter(const FParticleSystemSceneProxy* InOwnerProxy);
+
+	virtual void GetDynamicMeshElementsEmitter(const FParticleSystemSceneProxy* Proxy, const FSceneView* View, const FSceneViewFamily& ViewFamily, int32 ViewIndex, FMeshElementCollector& Collector) const override;
 
 	/**
 	 *	Render thread only draw call
@@ -2198,7 +2125,7 @@ struct FDynamicBeam2EmitterData : public FDynamicSpriteEmitterDataBase
 	 *	@param	PDI			The primitive draw interface to render with
 	 *	@param	View		The scene view being rendered
 	 */
-	virtual int32 Render(FParticleSystemSceneProxy* Proxy, FPrimitiveDrawInterface* PDI,const FSceneView* View);
+	virtual int32 RenderEmitter(FParticleSystemSceneProxy* Proxy, FPrimitiveDrawInterface* PDI,const FSceneView* View);
 
 	/**
 	 *	Called during InitViews for view processing on scene proxies before rendering them
@@ -2212,16 +2139,16 @@ struct FDynamicBeam2EmitterData : public FDynamicSpriteEmitterDataBase
 	virtual void PreRenderView(FParticleSystemSceneProxy* Proxy, const FSceneViewFamily* ViewFamily, const uint32 VisibilityMap, int32 FrameNumber);
 
 	// Debugging functions
-	virtual void RenderDirectLine(FParticleSystemSceneProxy* Proxy, FPrimitiveDrawInterface* PDI,const FSceneView* View);
-	virtual void RenderLines(FParticleSystemSceneProxy* Proxy, FPrimitiveDrawInterface* PDI,const FSceneView* View);
+	virtual void RenderDirectLine(const FParticleSystemSceneProxy* Proxy, FPrimitiveDrawInterface* PDI,const FSceneView* View) const;
+	virtual void RenderLines(const FParticleSystemSceneProxy* Proxy, FPrimitiveDrawInterface* PDI,const FSceneView* View) const;
 
-	virtual void RenderDebug(FParticleSystemSceneProxy* Proxy, FPrimitiveDrawInterface* PDI, const FSceneView* View, bool bCrosses);
+	virtual void RenderDebug(const FParticleSystemSceneProxy* Proxy, FPrimitiveDrawInterface* PDI, const FSceneView* View, bool bCrosses) const;
 
 	// Data fill functions
-	int32 FillIndexData(struct FAsyncBufferFillData& Data);
-	int32 FillVertexData_NoNoise(struct FAsyncBufferFillData& Data);
-	int32 FillData_Noise(struct FAsyncBufferFillData& Data);
-	int32 FillData_InterpolatedNoise(struct FAsyncBufferFillData& Data);
+	int32 FillIndexData(struct FAsyncBufferFillData& Data) const;
+	int32 FillVertexData_NoNoise(struct FAsyncBufferFillData& Data) const;
+	int32 FillData_Noise(struct FAsyncBufferFillData& Data) const;
+	int32 FillData_InterpolatedNoise(struct FAsyncBufferFillData& Data) const;
 
 	/** Returns the source data for this particle system */
 	virtual const FDynamicEmitterReplayDataBase& GetSource() const
@@ -2229,8 +2156,8 @@ struct FDynamicBeam2EmitterData : public FDynamicSpriteEmitterDataBase
 		return Source;
 	}
 
-	/** Perform the actual work of filling the buffer, often called from another thread */
-	virtual void DoBufferFill(FAsyncBufferFillData& Me);
+	/** Perform the actual work of filling the buffer */
+	virtual void DoBufferFill(FAsyncBufferFillData& Me) const;
 
 	/**
 	 *	Get the vertex stride for the dynamic rendering data
@@ -2358,7 +2285,9 @@ struct FDynamicTrailsEmitterData : public FDynamicSpriteEmitterDataBase
 	 *
 	 *	@return	bool			true if successful, false if failed
 	 */
-	virtual void CreateRenderThreadResources(const FParticleSystemSceneProxy* InOwnerProxy);
+	virtual void UpdateRenderThreadResourcesEmitter(const FParticleSystemSceneProxy* InOwnerProxy);
+
+	virtual void GetDynamicMeshElementsEmitter(const FParticleSystemSceneProxy* Proxy, const FSceneView* View, const FSceneViewFamily& ViewFamily, int32 ViewIndex, FMeshElementCollector& Collector) const override;
 
 	/**
 	 *	Render thread only draw call
@@ -2367,7 +2296,7 @@ struct FDynamicTrailsEmitterData : public FDynamicSpriteEmitterDataBase
 	 *	@param	PDI			The primitive draw interface to render with
 	 *	@param	View		The scene view being rendered
 	 */
-	virtual int32 Render(FParticleSystemSceneProxy* Proxy, FPrimitiveDrawInterface* PDI, const FSceneView* View);
+	virtual int32 RenderEmitter(FParticleSystemSceneProxy* Proxy, FPrimitiveDrawInterface* PDI, const FSceneView* View);
 
 	virtual bool ShouldUsePrerenderView()
 	{
@@ -2385,11 +2314,11 @@ struct FDynamicTrailsEmitterData : public FDynamicSpriteEmitterDataBase
 	 */
 	virtual void PreRenderView(FParticleSystemSceneProxy* Proxy, const FSceneViewFamily* ViewFamily, const uint32 VisibilityMap, int32 FrameNumber);
 
-	virtual void RenderDebug(FParticleSystemSceneProxy* Proxy, FPrimitiveDrawInterface* PDI, const FSceneView* View, bool bCrosses);
+	virtual void RenderDebug(const FParticleSystemSceneProxy* Proxy, FPrimitiveDrawInterface* PDI, const FSceneView* View, bool bCrosses) const;
 
 	// Data fill functions
-	virtual int32 FillIndexData(struct FAsyncBufferFillData& Data);
-	virtual int32 FillVertexData(struct FAsyncBufferFillData& Data);
+	virtual int32 FillIndexData(struct FAsyncBufferFillData& Data) const;
+	virtual int32 FillVertexData(struct FAsyncBufferFillData& Data) const;
 
 	/** Returns the source data for this particle system */
 	virtual const FDynamicEmitterReplayDataBase& GetSource() const
@@ -2404,7 +2333,7 @@ struct FDynamicTrailsEmitterData : public FDynamicSpriteEmitterDataBase
 		return SourcePointer;
 	}
 
-	virtual void DoBufferFill(FAsyncBufferFillData& Me)
+	virtual void DoBufferFill(FAsyncBufferFillData& Me) const
 	{
 		if( Me.VertexCount <= 0 || Me.IndexCount <= 0 || Me.VertexData == NULL || Me.IndexData == NULL )
 		{
@@ -2468,10 +2397,10 @@ struct FDynamicRibbonEmitterData : public FDynamicTrailsEmitterData
 	virtual void Init(bool bInSelected);
 
 	virtual bool ShouldUsePrerenderView();
-	virtual void RenderDebug(FParticleSystemSceneProxy* Proxy, FPrimitiveDrawInterface* PDI, const FSceneView* View, bool bCrosses);
+	virtual void RenderDebug(const FParticleSystemSceneProxy* Proxy, FPrimitiveDrawInterface* PDI, const FSceneView* View, bool bCrosses) const;
 
 	// Data fill functions
-	virtual int32 FillVertexData(struct FAsyncBufferFillData& Data);
+	virtual int32 FillVertexData(struct FAsyncBufferFillData& Data) const;
 		
 	/**
 	 *	Get the source replay data for this emitter
@@ -2507,10 +2436,10 @@ struct FDynamicAnimTrailEmitterData : public FDynamicTrailsEmitterData
 	/** Initialize this emitter's dynamic rendering data, called after source data has been filled in */
 	virtual void Init(bool bInSelected);
 
-	virtual void RenderDebug(FParticleSystemSceneProxy* Proxy, FPrimitiveDrawInterface* PDI, const FSceneView* View, bool bCrosses);
+	virtual void RenderDebug(const FParticleSystemSceneProxy* Proxy, FPrimitiveDrawInterface* PDI, const FSceneView* View, bool bCrosses) const;
 
 	// Data fill functions
-	virtual int32 FillVertexData(struct FAsyncBufferFillData& Data);
+	virtual int32 FillVertexData(struct FAsyncBufferFillData& Data) const;
 
 	/** 
 	 *	The frame source data for this particle system.  This is everything needed to represent this
@@ -2581,6 +2510,7 @@ public:
 		return false;
 	}
 	virtual void DrawDynamicElements(FPrimitiveDrawInterface* PDI,const FSceneView* View) override;
+	virtual void GetDynamicMeshElements(const TArray<const FSceneView*>& Views, const FSceneViewFamily& ViewFamily, uint32 VisibilityMap, FMeshElementCollector& Collector) const override;
 	virtual FPrimitiveViewRelevance GetViewRelevance(const FSceneView* View) override;
 	virtual void OnActorPositionChanged() override;
 	virtual void OnTransformChanged() override;
@@ -2650,7 +2580,7 @@ public:
 	 * world space primitive uniform buffer is up-to-date.
 	 * Only called in the rendering thread.
 	 */
-	void UpdateWorldSpacePrimitiveUniformBuffer();
+	void UpdateWorldSpacePrimitiveUniformBuffer() const;
 
 	/** Object position in post projection space. */
 	void GetObjectPositionAndScale(const FSceneView& View, FVector2D& ObjectNDCPosition, FVector2D& ObjectMacroUVScales) const;
@@ -2666,8 +2596,7 @@ public:
 	float GetPendingLODDistance()		{	return PendingLODDistance;		}
 	inline const TUniformBuffer<FPrimitiveUniformShaderParameters>& GetWorldSpacePrimitiveUniformBuffer() const { return WorldSpacePrimitiveUniformBuffer; }
 
-	FColoredMaterialRenderProxy* GetSelectedWireframeMatInst()		{	return &SelectedWireframeMaterialInstance;		}
-	FColoredMaterialRenderProxy* GetDeselectedWireframeMatInst()	{	return &DeselectedWireframeMaterialInstance;	}
+	const FColoredMaterialRenderProxy* GetDeselectedWireframeMatInst() const	{	return &DeselectedWireframeMaterialInstance;	}
 
 	/** Gets a mesh batch from the pool. */
 	FMeshBatch* GetPooledMeshBatch();
@@ -2698,7 +2627,6 @@ protected:
 	FParticleDynamicData* DynamicData;			// RENDER THREAD USAGE ONLY
 	FParticleDynamicData* LastDynamicData;		// RENDER THREAD USAGE ONLY
 
-	FColoredMaterialRenderProxy SelectedWireframeMaterialInstance;
 	FColoredMaterialRenderProxy DeselectedWireframeMaterialInstance;
 
 	int32 LODMethod;
@@ -2706,8 +2634,8 @@ protected:
 
 	int32 LastFramePreRendered;
 
-	/** The primitive's uniform buffer. */
-	TUniformBuffer<FPrimitiveUniformShaderParameters> WorldSpacePrimitiveUniformBuffer;
+	/** The primitive's uniform buffer.  Mutable because it is cached state during GetDynamicMeshElements. */
+	mutable TUniformBuffer<FPrimitiveUniformShaderParameters> WorldSpacePrimitiveUniformBuffer;
 
 	/** Pool for holding FMeshBatches to reduce allocations. */
 	TIndirectArray<FMeshBatch, TInlineAllocator<4> > MeshBatchPool;
diff --git a/Engine/Source/Runtime/Engine/Public/ParticleVertexFactory.h b/Engine/Source/Runtime/Engine/Public/ParticleVertexFactory.h
--- a/Engine/Source/Runtime/Engine/Public/ParticleVertexFactory.h
+++ b/Engine/Source/Runtime/Engine/Public/ParticleVertexFactory.h
@@ -47,6 +47,11 @@ public:
 		return ParticleFactoryType;
 	}
 
+	inline void SetParticleFactoryType(EParticleVertexFactoryType InType)
+	{
+		ParticleFactoryType = InType;
+	}
+
 	/** Specify whether the factory is in use or not. */
 	FORCEINLINE void SetInUse( bool bInInUse )
 	{
@@ -101,8 +106,11 @@ public:
 	/** Default constructor. */
 	FParticleSpriteVertexFactory( EParticleVertexFactoryType InType, ERHIFeatureLevel::Type InFeatureLevel )
 		: FParticleVertexFactoryBase(InType, InFeatureLevel)
-	{
-	}
+	{}
+
+	FParticleSpriteVertexFactory() 
+		: FParticleVertexFactoryBase(PVFT_MAX, ERHIFeatureLevel::Num)
+	{}
 
 	// FRenderResource interface.
 	virtual void InitRHI() override;
diff --git a/Engine/Source/Runtime/Engine/Public/PrimitiveSceneProxy.h b/Engine/Source/Runtime/Engine/Public/PrimitiveSceneProxy.h
--- a/Engine/Source/Runtime/Engine/Public/PrimitiveSceneProxy.h
+++ b/Engine/Source/Runtime/Engine/Public/PrimitiveSceneProxy.h
@@ -155,6 +155,20 @@ public:
 	 */
 	virtual void DrawDynamicElements(FPrimitiveDrawInterface* PDI, const FSceneView* View, uint32 DrawDynamicFlags ) { DrawDynamicElements( PDI, View ); }
 
+	/** 
+	 * Gathers the primitive's dynamic mesh elements.  This will only be called if GetViewRelevance declares dynamic relevance.
+	 * This is called from the rendering thread for each set of views that might be rendered.  
+	 * Game thread state like UObjects must have their properties mirrored on the proxy to avoid race conditions.  The rendering thread must not dereference UObjects.
+	 * The gathered mesh elements will be used multiple times, any memory referenced must last as long as the Collector (eg no stack memory should be referenced).
+	 * This function should not modify the proxy but simply collect a description of things to render.  Updates to the proxy need to be pushed from game thread or external events.
+	 *
+	 * @param Views - the array of views to consider.  These may not exist in the ViewFamily.
+	 * @param ViewFamily - the view family, for convenience
+	 * @param VisibilityMap - a bit representing this proxy's visibility in the Views array
+	 * @param Collector - gathers the mesh elements to be rendered and provides mechanisms for temporary allocations
+	 */
+	virtual void GetDynamicMeshElements(const TArray<const FSceneView*>& Views, const FSceneViewFamily& ViewFamily, uint32 VisibilityMap, class FMeshElementCollector& Collector) const {}
+
 	/**
 	 * Determines the relevance of this primitive's elements to the given view.
 	 * Called in the rendering thread.
@@ -291,12 +305,6 @@ public:
 	inline bool IsLocalToWorldDeterminantNegative() const { return bIsLocalToWorldDeterminantNegative; }
 	inline const FBoxSphereBounds& GetBounds() const { return Bounds; }
 	inline const FBoxSphereBounds& GetLocalBounds() const { return LocalBounds; }
-	inline const FVector2D& GetLightMapResolutionScale() const { return LightMapResolutionScale; }
-	inline bool IsLightMapResolutionPadded() const { return bLightMapResolutionPadded; }
-	inline ELightMapInteractionType GetLightMapType() const { return (ELightMapInteractionType)LightMapType; }
-	inline void SetLightMapResolutionScale(const FVector2D& InLightMapResolutionScale) { LightMapResolutionScale = InLightMapResolutionScale; }
-	inline void SetIsLightMapResolutionPadded(bool bInLightMapResolutionPadded) { bLightMapResolutionPadded = bInLightMapResolutionPadded; }
-	inline void SetLightMapType(ELightMapInteractionType InLightMapType) { LightMapType = InLightMapType; }
 	inline FName GetOwnerName() const { return OwnerName; }
 	inline FName GetResourceName() const { return ResourceName; }
 	inline FName GetLevelName() const { return LevelName; }
@@ -342,11 +350,16 @@ public:
 	inline bool WillEverBeLit() const { return bWillEverBeLit; }
 	inline bool HasValidSettingsForStaticLighting() const { return bHasValidSettingsForStaticLighting; }
 	inline bool AlwaysHasVelocity() const { return bAlwaysHasVelocity; }
+	inline bool UseEditorDepthTest() const { return bUseEditorDepthTest; }
 	inline bool TreatAsBackgroundForOcclusion() const { return bTreatAsBackgroundForOcclusion; }
 #if WITH_EDITOR
 	inline int32 GetNumUncachedStaticLightingInteractions() { return NumUncachedStaticLightingInteractions; }
 #endif
 
+	inline FLinearColor GetWireframeColor() const { return WireframeColor; }
+	inline FLinearColor GetLevelColor() const { return LevelColor; }
+	inline FLinearColor GetPropertyColor() const { return PropertyColor; }
+
 	/**
 	* Returns whether this proxy should be considered a "detail mesh".
 	* Detail meshes are distance culled even if distance culling is normally disabled for the view. (e.g. in editor)
@@ -408,12 +421,13 @@ protected:
 		OwnerName = InOwnerName;
 	}
 
+	FLinearColor WireframeColor;
+	FLinearColor LevelColor;
+	FLinearColor PropertyColor;
+
 private:
 	friend class FScene;
 
-	/** The LightMap method used by the primitive */
-	uint32 LightMapType : LMIT_NumBits;
-
 	uint32 bIsLocalToWorldDeterminantNegative : 1;
 	uint32 DrawInGame : 1;
 	uint32 DrawInEditor : 1;
@@ -427,9 +441,6 @@ private:
 	/** true if the mouse is currently hovered over this primitive in a level viewport */
 	uint32 bHovered : 1;
 
-	/** true if the LightMapResolutionScale value has been padded. */
-	uint32 bLightMapResolutionPadded : 1;
-
 	/** true if ViewOwnerDepthPriorityGroup should be used. */
 	uint32 bUseViewOwnerDepthPriorityGroup : 1;
 
@@ -515,6 +526,9 @@ protected:
 	/** Whether the primitive should always be considered to have velocities, even if it hasn't moved. */
 	uint32 bAlwaysHasVelocity : 1;
 
+	/** Whether editor compositing depth testing should be used for this primitive.  Only matters for primitives with bUseEditorCompositing. */
+	uint32 bUseEditorDepthTest : 1;
+
 private:
 
 	/** If this is True, this primitive will be used to occlusion cull other primitives. */
@@ -576,9 +590,6 @@ private:
 	/** The name of the level the primitive is in. */
 	FName LevelName;
 
-	/** The StaticLighting resolution for this mesh */
-	FVector2D LightMapResolutionScale;
-
 #if WITH_EDITOR
 	/** A copy of the actor's group membership for handling per-view group hiding */
 	uint64 HiddenEditorViews;
diff --git a/Engine/Source/Runtime/Engine/Public/PrimitiveUniformShaderParameters.h b/Engine/Source/Runtime/Engine/Public/PrimitiveUniformShaderParameters.h
--- a/Engine/Source/Runtime/Engine/Public/PrimitiveUniformShaderParameters.h
+++ b/Engine/Source/Runtime/Engine/Public/PrimitiveUniformShaderParameters.h
@@ -17,6 +17,7 @@ BEGIN_UNIFORM_BUFFER_STRUCT(FPrimitiveUniformShaderParameters,ENGINE_API)
 	DECLARE_UNIFORM_BUFFER_STRUCT_MEMBER(FVector,ActorWorldPosition)
 	DECLARE_UNIFORM_BUFFER_STRUCT_MEMBER_EX(float,DecalReceiverMask,EShaderPrecisionModifier::Half)
 	DECLARE_UNIFORM_BUFFER_STRUCT_MEMBER_EX(float,HasDistanceFieldRepresentation,EShaderPrecisionModifier::Half)
+	DECLARE_UNIFORM_BUFFER_STRUCT_MEMBER_EX(float,UseEditorDepthTest,EShaderPrecisionModifier::Half)
 	DECLARE_UNIFORM_BUFFER_STRUCT_MEMBER_EX(FVector4,ObjectOrientation,EShaderPrecisionModifier::Half)
 	DECLARE_UNIFORM_BUFFER_STRUCT_MEMBER_EX(FVector4,NonUniformScale,EShaderPrecisionModifier::Half)
 	DECLARE_UNIFORM_BUFFER_STRUCT_MEMBER_EX(FVector4,InvNonUniformScale,EShaderPrecisionModifier::Half)
@@ -33,6 +34,7 @@ inline FPrimitiveUniformShaderParameters GetPrimitiveUniformShaderParameters(
 	const FBoxSphereBounds& LocalBounds,
 	bool bReceivesDecals,
 	bool bHasDistanceFieldRepresentation,
+	bool bUseEditorDepthTest,
 	float LpvBiasMultiplier = 1.0f
 )
 {
@@ -67,6 +69,7 @@ inline FPrimitiveUniformShaderParameters GetPrimitiveUniformShaderParameters(
 	Result.LocalToWorldDeterminantSign = FMath::FloatSelect(LocalToWorld.RotDeterminant(),1,-1);
 	Result.DecalReceiverMask = bReceivesDecals ? 1 : 0;
 	Result.HasDistanceFieldRepresentation = bHasDistanceFieldRepresentation ? 1 : 0;
+	Result.UseEditorDepthTest = bUseEditorDepthTest ? 1 : 0;
 	return Result;
 }
 
@@ -75,12 +78,13 @@ inline TUniformBufferRef<FPrimitiveUniformShaderParameters> CreatePrimitiveUnifo
 	const FBoxSphereBounds& WorldBounds,
 	const FBoxSphereBounds& LocalBounds,
 	bool bReceivesDecals,
+	bool bUseEditorDepthTest,
 	float LpvBiasMultiplier = 1.0f
 	)
 {
 	check(IsInRenderingThread());
 	return TUniformBufferRef<FPrimitiveUniformShaderParameters>::CreateUniformBufferImmediate(
-		GetPrimitiveUniformShaderParameters(LocalToWorld, WorldBounds.Origin, WorldBounds, LocalBounds, bReceivesDecals, false, LpvBiasMultiplier ),
+		GetPrimitiveUniformShaderParameters(LocalToWorld, WorldBounds.Origin, WorldBounds, LocalBounds, bReceivesDecals, false, bUseEditorDepthTest, LpvBiasMultiplier ),
 		UniformBuffer_MultiFrame
 		);
 }
@@ -107,6 +111,7 @@ public:
 			FBoxSphereBounds(EForceInit::ForceInit),
 			true,
 			false,
+			true,
 			1.0f		// LPV bias
 			));
 	}
diff --git a/Engine/Source/Runtime/Engine/Public/SceneManagement.h b/Engine/Source/Runtime/Engine/Public/SceneManagement.h
--- a/Engine/Source/Runtime/Engine/Public/SceneManagement.h
+++ b/Engine/Source/Runtime/Engine/Public/SceneManagement.h
@@ -23,7 +23,9 @@
 #include "MeshBatch.h"
 #include "PrimitiveViewRelevance.h"
 #include "SceneInterface.h"
-
+#include "RHIDefinitions.h"
+#include "ChunkedArray.h"
+#include "BatchedElements.h"
 
 // Forward declarations.
 class FSceneViewFamily;
@@ -39,12 +41,6 @@ struct FDynamicMeshVertex;
 
 DECLARE_LOG_CATEGORY_EXTERN(LogBufferVisualization, Log, All);
 
-
-
-// -----------------------------------------------------------------------------
-
-
-
 /**
  * The scene manager's persistent view state.
  */
@@ -1104,6 +1100,275 @@ public:
 		) = 0;
 };
 
+/** Primitive draw interface implementation used to store primitives requested to be drawn when gathering dynamic mesh elements. */
+class ENGINE_API FSimpleElementCollector : public FPrimitiveDrawInterface
+{
+public:
+
+	FSimpleElementCollector() : FPrimitiveDrawInterface(NULL)
+	{}
+
+	~FSimpleElementCollector();
+
+	virtual void SetHitProxy(HHitProxy* HitProxy);
+	virtual void AddReserveLines(uint8 DepthPriorityGroup, int32 NumLines, bool bDepthBiased = false) {}
+
+	virtual void DrawSprite(
+		const FVector& Position,
+		float SizeX,
+		float SizeY,
+		const FTexture* Sprite,
+		const FLinearColor& Color,
+		uint8 DepthPriorityGroup,
+		float U,
+		float UL,
+		float V,
+		float VL,
+		uint8 BlendMode = SE_BLEND_Masked
+		) override;
+
+	virtual void DrawLine(
+		const FVector& Start,
+		const FVector& End,
+		const FLinearColor& Color,
+		uint8 DepthPriorityGroup,
+		float Thickness = 0.0f,
+		float DepthBias = 0.0f,
+		bool bScreenSpace = false
+		) override;
+
+	virtual void DrawPoint(
+		const FVector& Position,
+		const FLinearColor& Color,
+		float PointSize,
+		uint8 DepthPriorityGroup
+		) override;
+
+	virtual void RegisterDynamicResource(FDynamicPrimitiveResource* DynamicResource) override;
+
+	// Not supported
+	virtual bool IsHitTesting() 
+	{ 
+		static bool bTriggered = false;
+
+		if (!bTriggered)
+		{
+			bTriggered = true;
+			ensureMsg(false, TEXT("FSimpleElementCollector::DrawMesh called"));
+		}
+
+		return false; 
+	}
+
+	// Not supported
+	virtual int32 DrawMesh(const FMeshBatch& Mesh) 
+	{
+		static bool bTriggered = false;
+
+		if (!bTriggered)
+		{
+			bTriggered = true;
+			ensureMsg(false, TEXT("FSimpleElementCollector::DrawMesh called"));
+		}
+
+		return 0;
+	}
+
+	// Legacy, should not be used
+	virtual bool IsMaterialIgnored(const FMaterialRenderProxy* MaterialRenderProxy, ERHIFeatureLevel::Type InFeatureLevel) const
+	{
+		static bool bTriggered = false;
+
+		if (!bTriggered)
+		{
+			bTriggered = true;
+			ensureMsg(false, TEXT("FSimpleElementCollector::IsMaterialIgnored called"));
+		}
+
+		return false;
+	}
+
+	// Legacy, should not be used
+	virtual bool IsRenderingSelectionOutline() const
+	{
+		static bool bTriggered = false;
+
+		if (!bTriggered)
+		{
+			bTriggered = true;
+			ensureMsg(false, TEXT("FSimpleElementCollector::IsRenderingSelectionOutline called"));
+		}
+
+		return false;
+	}
+
+	void DrawBatchedElements(FRHICommandList& RHICmdList, const FSceneView& View, FTexture2DRHIRef DepthTexture, EBlendModeFilter::Type Filter) const;
+
+	/** The batched simple elements. */
+	FBatchedElements BatchedElements;
+
+private:
+
+	FHitProxyId HitProxyId;
+
+	/** The dynamic resources which have been registered with this drawer. */
+	TArray<FDynamicPrimitiveResource*,SceneRenderingAllocator> DynamicResources;
+
+	friend class FMeshElementCollector;
+};
+
+/** 
+ * Base class for a resource allocated from a FMeshElementCollector with AllocateOneFrameResource, which the collector releases.
+ * This is useful for per-frame structures which are referenced by a mesh batch given to the FMeshElementCollector.
+ */
+class FOneFrameResource
+{
+public:
+
+	virtual ~FOneFrameResource() {}
+};
+
+/** 
+ * A reference to a mesh batch that is added to the collector, together with some cached relevance flags. 
+ */
+struct FMeshBatchAndRelevance
+{
+	const FMeshBatch* Mesh;
+
+	/** The render info for the primitive which created this mesh, required. */
+	const FPrimitiveSceneProxy* PrimitiveSceneProxy;
+
+	/** 
+	 * Cached usage information to speed up traversal in the most costly passes (depth-only, base pass, shadow depth), 
+	 * This is done so the Mesh does not have to be dereferenced to determine pass relevance. 
+	 */
+	uint32 bHasOpaqueOrMaskedMaterial : 1;
+	uint32 bRenderInMainPass : 1;
+
+	FMeshBatchAndRelevance(const FMeshBatch& InMesh, const FPrimitiveSceneProxy* InPrimitiveSceneProxy, ERHIFeatureLevel::Type FeatureLevel);
+};
+
+/** 
+ * Encapsulates the gathering of meshes from the various FPrimitiveSceneProxy classes. 
+ */
+class FMeshElementCollector
+{
+public:
+
+	/** Accesses the PDI for drawing lines, sprites, etc. */
+	inline FPrimitiveDrawInterface* GetPDI(int32 ViewIndex)
+	{
+		return SimpleElementCollectors[ViewIndex];
+	}
+
+	/** 
+	 * Allocates an FMeshBatch that can be safely referenced by the collector (lifetime will be long enough).
+	 * Returns a reference that will not be invalidated due to further AllocateMesh() calls.
+	 */
+	inline FMeshBatch& AllocateMesh()
+	{
+		return *(new (MeshBatchStorage) FMeshBatch());
+	}
+
+	/** 
+	 * Adds a mesh batch to the collector for the specified view so that it can be rendered.
+	 */
+	ENGINE_API void AddMesh(int32 ViewIndex, FMeshBatch& MeshBatch);
+
+	/** Add a material render proxy that will be cleaned up automatically */
+	void RegisterOneFrameMaterialProxy(FMaterialRenderProxy* Proxy)
+	{
+		TemporaryProxies.Add(Proxy);
+	}
+
+	/** Allocates a temporary resource that is safe to be referenced by an FMeshBatch added to the collector. */
+	template<typename T>
+	T& AllocateOneFrameResource()
+	{
+		T* OneFrameResource = new (FMemStack::Get()) T();
+		OneFrameResources.Add(OneFrameResource);
+		return *OneFrameResource;
+	}
+
+private:
+
+	FMeshElementCollector() :
+		PrimitiveSceneProxy(NULL),
+		FeatureLevel(ERHIFeatureLevel::Num)
+	{}
+
+	~FMeshElementCollector()
+	{
+		for (int32 ProxyIndex = 0; ProxyIndex < TemporaryProxies.Num(); ProxyIndex++)
+		{
+			delete TemporaryProxies[ProxyIndex];
+		}
+
+		// SceneRenderingAllocator does not handle destructors
+		for (int32 ResourceIndex = 0; ResourceIndex < OneFrameResources.Num(); ResourceIndex++)
+		{
+			OneFrameResources[ResourceIndex]->~FOneFrameResource();
+		}
+	}
+
+	void SetPrimitive(const FPrimitiveSceneProxy* InPrimitiveSceneProxy, FHitProxyId DefaultHitProxyId)
+	{
+		check(InPrimitiveSceneProxy);
+		PrimitiveSceneProxy = InPrimitiveSceneProxy;
+
+		for (int32 ViewIndex = 0; ViewIndex < SimpleElementCollectors.Num(); ViewIndex++)
+		{
+			SimpleElementCollectors[ViewIndex]->HitProxyId = DefaultHitProxyId;
+		}
+	}
+
+	void ClearViewMeshArrays()
+	{
+		Views.Empty();
+		MeshBatches.Empty();
+		SimpleElementCollectors.Empty();
+	}
+
+	void AddViewMeshArrays(
+		FSceneView* InView, 
+		TArray<FMeshBatchAndRelevance,SceneRenderingAllocator>* ViewMeshes,
+		FSimpleElementCollector* ViewSimpleElementCollector, 
+		ERHIFeatureLevel::Type InFeatureLevel)
+	{
+		Views.Add(InView);
+		MeshBatches.Add(ViewMeshes);
+		SimpleElementCollectors.Add(ViewSimpleElementCollector);
+		FeatureLevel = InFeatureLevel;
+	}
+
+	/** 
+	 * Using TChunkedArray which will never realloc as new elements are added
+	 * @todo - use mem stack
+	 */
+	TChunkedArray<FMeshBatch> MeshBatchStorage;
+
+	/** Meshes to render */
+	TArray<TArray<FMeshBatchAndRelevance, SceneRenderingAllocator>*, TInlineAllocator<2> > MeshBatches;
+
+	/** PDIs */
+	TArray<FSimpleElementCollector*, TInlineAllocator<2> > SimpleElementCollectors;
+
+	/** Views being collected for */
+	TArray<FSceneView*, TInlineAllocator<2> > Views;
+
+	/** Material proxies that will be deleted at the end of the frame. */
+	TArray<FMaterialRenderProxy*, SceneRenderingAllocator> TemporaryProxies;
+
+	/** Resources that will be deleted at the end of the frame. */
+	TArray<FOneFrameResource*, SceneRenderingAllocator> OneFrameResources;
+
+	/** Current primitive being gathered. */
+	const FPrimitiveSceneProxy* PrimitiveSceneProxy;
+
+	ERHIFeatureLevel::Type FeatureLevel;
+
+	friend class FSceneRenderer;
+};
 
 
 /**
@@ -1351,13 +1616,19 @@ extern ENGINE_API void DrawBox(class FPrimitiveDrawInterface* PDI,const FMatrix&
 extern ENGINE_API void DrawSphere(class FPrimitiveDrawInterface* PDI,const FVector& Center,const FVector& Radii,int32 NumSides,int32 NumRings,const FMaterialRenderProxy* MaterialRenderProxy,uint8 DepthPriority,bool bDisableBackfaceCulling=false);
 extern ENGINE_API void DrawCone(class FPrimitiveDrawInterface* PDI,const FMatrix& ConeToWorld, float Angle1, float Angle2, int32 NumSides, bool bDrawSideLines, const FLinearColor& SideLineColor, const FMaterialRenderProxy* MaterialRenderProxy, uint8 DepthPriority);
 
-
 extern ENGINE_API void DrawCylinder(class FPrimitiveDrawInterface* PDI,const FVector& Base, const FVector& XAxis, const FVector& YAxis, const FVector& ZAxis,
 	float Radius, float HalfHeight, int32 Sides, const FMaterialRenderProxy* MaterialInstance, uint8 DepthPriority);
 
 extern ENGINE_API void DrawCylinder(class FPrimitiveDrawInterface* PDI, const FMatrix& CylToWorld, const FVector& Base, const FVector& XAxis, const FVector& YAxis, const FVector& ZAxis,
 	float Radius, float HalfHeight, int32 Sides, const FMaterialRenderProxy* MaterialInstance, uint8 DepthPriority);
 
+extern ENGINE_API void GetBoxMesh(const FMatrix& BoxToWorld,const FVector& Radii,const FMaterialRenderProxy* MaterialRenderProxy,uint8 DepthPriority,int32 ViewIndex,FMeshElementCollector& Collector);
+extern ENGINE_API void GetSphereMesh(const FVector& Center,const FVector& Radii,int32 NumSides,int32 NumRings,const FMaterialRenderProxy* MaterialRenderProxy,uint8 DepthPriority,bool bDisableBackfaceCulling,int32 ViewIndex,FMeshElementCollector& Collector);
+extern ENGINE_API void GetCylinderMesh(const FVector& Base, const FVector& XAxis, const FVector& YAxis, const FVector& ZAxis,
+									float Radius, float HalfHeight, int32 Sides, const FMaterialRenderProxy* MaterialInstance, uint8 DepthPriority, int32 ViewIndex, FMeshElementCollector& Collector);
+extern ENGINE_API void GetCylinderMesh(const FMatrix& CylToWorld, const FVector& Base, const FVector& XAxis, const FVector& YAxis, const FVector& ZAxis,
+									float Radius, float HalfHeight, int32 Sides, const FMaterialRenderProxy* MaterialInstance, uint8 DepthPriority, int32 ViewIndex, FMeshElementCollector& Collector);
+
 
 extern ENGINE_API void DrawDisc(class FPrimitiveDrawInterface* PDI,const FVector& Base,const FVector& XAxis,const FVector& YAxis,FColor Color,float Radius,int32 NumSides, const FMaterialRenderProxy* MaterialRenderProxy, uint8 DepthPriority);
 extern ENGINE_API void DrawFlatArrow(class FPrimitiveDrawInterface* PDI,const FVector& Base,const FVector& XAxis,const FVector& YAxis,FColor Color,float Length,int32 Width, const FMaterialRenderProxy* MaterialRenderProxy, uint8 DepthPriority);
@@ -1447,7 +1718,7 @@ extern ENGINE_API EVertexColorViewMode::Type GVertexColorViewMode;
  * being made.
  * A view is rich if is missing the EngineShowFlags.Materials showflag, or has any of the render mode affecting showflags.
  */
-extern ENGINE_API bool IsRichView(const FSceneView* View);
+extern ENGINE_API bool IsRichView(const FSceneViewFamily& ViewFamily);
 
 /** Emits draw events for a given FMeshBatch and the PrimitiveSceneProxy corresponding to that mesh element. */
 extern ENGINE_API void EmitMeshDrawEvents(const class FPrimitiveSceneProxy* PrimitiveSceneProxy, const struct FMeshBatch& Mesh);
@@ -1477,6 +1748,16 @@ extern ENGINE_API int32 DrawRichMesh(
 	bool bDrawInWireframe = false
 	);
 
+extern ENGINE_API void ApplyViewModeOverrides(
+	int32 ViewIndex,
+	const FEngineShowFlags& EngineShowFlags,
+	ERHIFeatureLevel::Type FeatureLevel,
+	const FPrimitiveSceneProxy* PrimitiveSceneProxy,
+	bool bSelected,
+	struct FMeshBatch& Mesh,
+	FMeshElementCollector& Collector
+	);
+
 /** Draws the UV layout of the supplied asset (either StaticMeshRenderData OR SkeletalMeshRenderData, not both!) */
 extern ENGINE_API void DrawUVs(FViewport* InViewport, FCanvas* InCanvas, int32 InTextYPos, const int32 LODLevel, int32 UVChannel, TArray<FVector2D> SelectedEdgeTexCoords, class FStaticMeshRenderData* StaticMeshRenderData, class FStaticLODModel* SkeletalMeshRenderData );
 
diff --git a/Engine/Source/Runtime/Engine/Public/SceneView.h b/Engine/Source/Runtime/Engine/Public/SceneView.h
--- a/Engine/Source/Runtime/Engine/Public/SceneView.h
+++ b/Engine/Source/Runtime/Engine/Public/SceneView.h
@@ -613,6 +613,8 @@ public:
 
 	/** Computes FamilySizeX and FamilySizeY from the Views array. */
 	ENGINE_API void ComputeFamilySize();
+
+	ENGINE_API ERHIFeatureLevel::Type GetFeatureLevel() const;
 };
 
 /**
diff --git a/Engine/Source/Runtime/Engine/Public/SkeletalMeshTypes.h b/Engine/Source/Runtime/Engine/Public/SkeletalMeshTypes.h
--- a/Engine/Source/Runtime/Engine/Public/SkeletalMeshTypes.h
+++ b/Engine/Source/Runtime/Engine/Public/SkeletalMeshTypes.h
@@ -1787,27 +1787,18 @@ public:
 	virtual HHitProxy* CreateHitProxies(UPrimitiveComponent* Component, TArray<TRefCountPtr<HHitProxy> >& OutHitProxies) override;
 #endif
 	virtual void DrawDynamicElements(FPrimitiveDrawInterface* PDI, const FSceneView* View) override;
+	virtual void GetDynamicMeshElements(const TArray<const FSceneView*>& Views, const FSceneViewFamily& ViewFamily, uint32 VisibilityMap, FMeshElementCollector& Collector) const override;
 	virtual FPrimitiveViewRelevance GetViewRelevance(const FSceneView* View) override;
 	virtual bool CanBeOccluded() const override;
 	virtual void PreRenderView(const FSceneViewFamily* ViewFamily, const uint32 VisibilityMap, int32 FrameNumber) override;
 	
-	/**
-	* Draw only the section of the material ID given of the scene proxy as a dynamic element
-	*
-	* @param	PDI - draw interface to render to
-	* @param	View - current view
-	* @param 	ForceLOD - Force this LOD. If -1, use current LOD of mesh. 
-	* @param	InMaterial - which material section to draw
-	*/
-	virtual void DrawDynamicElementsByMaterial(FPrimitiveDrawInterface* PDI,const FSceneView* View, int32 ForceLOD, int32 InMaterial);
-
 	/**
 	 * Returns the world transform to use for drawing.
 	 * @param View - Current view
 	 * @param OutLocalToWorld - Will contain the local-to-world transform when the function returns.
 	 * @param OutWorldToLocal - Will contain the world-to-local transform when the function returns.
 	 */
-	virtual void GetWorldMatrices( const FSceneView* View, FMatrix& OutLocalToWorld, FMatrix& OutWorldToLocal );
+	void GetWorldMatrices( FMatrix& OutLocalToWorld, FMatrix& OutWorldToLocal ) const;
 
 	/** Util for getting LOD index currently used by this SceneProxy. */
 	int32 GetCurrentLODIndex();
@@ -1821,7 +1812,7 @@ public:
 	/** 
 	 * Render physics asset for debug display
 	 */
-	void DebugDrawPhysicsAsset(FPrimitiveDrawInterface* PDI,const FSceneView* View);
+	void DebugDrawPhysicsAsset(FPrimitiveDrawInterface* PDI, const FEngineShowFlags& EngineShowFlags) const;
 
 	virtual uint32 GetMemoryFootprint( void ) const { return( sizeof( *this ) + GetAllocatedSize() ); }
 	uint32 GetAllocatedSize( void ) const { return( FPrimitiveSceneProxy::GetAllocatedSize() + LODSections.GetAllocatedSize() ); }
@@ -1845,8 +1836,6 @@ protected:
 	class UPhysicsAsset* PhysicsAssetForDebug;
 
 	/** data copied for rendering */
-	FLinearColor LevelColor;
-	FColor PropertyColor;
 	uint32 bForceWireframe : 1;
 	uint32 bIsCPUSkinned : 1;
 	uint32 bCanHighlightSelectedSections : 1;
@@ -1891,12 +1880,8 @@ protected:
 	TSet<UMaterialInterface*> MaterialsInUse_GameThread;
 	bool bMaterialsNeedMorphUsage_GameThread;
 	
-	/** The color used by the wireframe mesh overlay mode */
-	FColor WireframeOverlayColor;
-
 	/**
 	* Draw only the section of the scene proxy as a dynamic element
-	* This is to avoid redundant code of two functions (DrawDynamicElementsByMaterial & DrawDynamicElements)
 	* 
 	* @param	PDI - draw interface to render to
 	* @param	View - current view
@@ -1909,4 +1894,9 @@ protected:
 		const FStaticLODModel& LODModel, const int32 LODIndex, const FSkelMeshSection& Section, 
 		const FSkelMeshChunk& Chunk, const FSectionElementInfo& SectionElementInfo, const FTwoVectors& CustomLeftRightVectors );
 
+	void GetDynamicElementsSection(const TArray<const FSceneView*>& Views, const FSceneViewFamily& ViewFamily, uint32 VisibilityMap, 
+		const FStaticLODModel& LODModel, const int32 LODIndex, const FSkelMeshSection& Section, const FSkelMeshChunk& Chunk, 
+		const FSectionElementInfo& SectionElementInfo, const FTwoVectors& CustomLeftRightVectors, bool bSelectable, FMeshElementCollector& Collector ) const;
+
+	void GetMeshElementsConditionallySelectable(const TArray<const FSceneView*>& Views, const FSceneViewFamily& ViewFamily, bool bSelectable, uint32 VisibilityMap, FMeshElementCollector& Collector) const;
 };
diff --git a/Engine/Source/Runtime/Engine/Public/StaticMeshResources.h b/Engine/Source/Runtime/Engine/Public/StaticMeshResources.h
--- a/Engine/Source/Runtime/Engine/Public/StaticMeshResources.h
+++ b/Engine/Source/Runtime/Engine/Public/StaticMeshResources.h
@@ -857,21 +857,28 @@ public:
 
 	virtual ~FStaticMeshSceneProxy() {}
 
+	/** Gets the number of mesh batches required to represent the proxy, aside from section needs. */
+	virtual int32 GetNumMeshBatches() const
+	{
+		return 1;
+	}
+
 	/** Sets up a shadow FMeshBatch for a specific LOD. */
-	virtual bool GetShadowMeshElements(int32 LODIndex, uint8 InDepthPriorityGroup, TArray<FMeshBatch, TInlineAllocator<1>>& OutMeshBatches) const;
+	virtual bool GetShadowMeshElement(int32 LODIndex, int32 BatchIndex, uint8 InDepthPriorityGroup, FMeshBatch& OutMeshBatch) const;
 
 	/** Sets up a FMeshBatch for a specific LOD and element. */
-	virtual bool GetMeshElements(int32 LODIndex,int32 ElementIndex,uint8 InDepthPriorityGroup, TArray<FMeshBatch, TInlineAllocator<1>>& OutMeshBatches, const bool bUseSelectedMaterial, const bool bUseHoveredMaterial) const;
+	virtual bool GetMeshElement(int32 LODIndex, int32 BatchIndex, int32 ElementIndex, uint8 InDepthPriorityGroup, const bool bUseSelectedMaterial, const bool bUseHoveredMaterial, FMeshBatch& OutMeshBatch) const;
 
 	/** Sets up a wireframe FMeshBatch for a specific LOD. */
-	virtual bool GetWireframeMeshElements(int32 LODIndex, const FMaterialRenderProxy* WireframeRenderProxy, uint8 InDepthPriorityGroup, TArray<FMeshBatch, TInlineAllocator<1>>& OutMeshBatches) const;
+	virtual bool GetWireframeMeshElement(int32 LODIndex, int32 BatchIndex, const FMaterialRenderProxy* WireframeRenderProxy, uint8 InDepthPriorityGroup, FMeshBatch& OutMeshBatch) const;
+
 
 protected:
 	/**
 	 * Sets IndexBuffer, FirstIndex and NumPrimitives of OutMeshElement.
 	 */
 	virtual void SetIndexSource(int32 LODIndex, int32 ElementIndex, FMeshBatch& OutMeshElement, bool bWireframe, bool bRequiresAdjacencyInformation ) const;
-	bool IsCollisionView(const FSceneView* View, bool & bDrawSimpleCollision, bool & bDrawComplexCollision) const;
+	bool IsCollisionView(const FEngineShowFlags& EngineShowFlags, bool& bDrawSimpleCollision, bool& bDrawComplexCollision) const;
 
 public:
 	// FPrimitiveSceneProxy interface.
@@ -891,6 +898,7 @@ public:
 	virtual uint32 GetMemoryFootprint( void ) const override { return( sizeof( *this ) + GetAllocatedSize() ); }
 	uint32 GetAllocatedSize( void ) const { return( FPrimitiveSceneProxy::GetAllocatedSize() + LODs.GetAllocatedSize() ); }
 
+	virtual void GetDynamicMeshElements(const TArray<const FSceneView*>& Views, const FSceneViewFamily& ViewFamily, uint32 VisibilityMap, FMeshElementCollector& Collector) const override;
 protected:
 
 	/** Information used by the proxy about a single LOD of the mesh. */
@@ -983,17 +991,12 @@ protected:
 
 	FVector TotalScale3D;
 
-	FLinearColor LevelColor;
-	FLinearColor PropertyColor;
-
 	uint32 bCastShadow : 1;
 	ECollisionTraceFlag		CollisionTraceFlag;
 
 	/** The view relevance for all the static mesh's materials. */
 	FMaterialRelevance MaterialRelevance;
 
-	const FLinearColor WireframeColor;
-
 	/** Collision Response of this component**/
 	FCollisionResponseContainer CollisionResponse;
 
diff --git a/Engine/Source/Runtime/RHI/Public/RHIDefinitions.h b/Engine/Source/Runtime/RHI/Public/RHIDefinitions.h
--- a/Engine/Source/Runtime/RHI/Public/RHIDefinitions.h
+++ b/Engine/Source/Runtime/RHI/Public/RHIDefinitions.h
@@ -609,6 +609,12 @@ inline bool RHISupportsTessellation(const EShaderPlatform Platform)
 	return false;
 }
 
+inline bool RHINeedsToSwitchVerticalAxis(EShaderPlatform Platform)
+{
+	// ES2 needs to flip when rendering to an RT that will be post processed
+	return IsES2Platform(Platform) && !IsPCPlatform(Platform);
+}
+
 inline bool RHISupportsInstancing(const EShaderPlatform Platform)
 {
 	//@todo-rco: Add Metal support
diff --git a/Engine/Source/Runtime/RenderCore/Private/RenderCore.cpp b/Engine/Source/Runtime/RenderCore/Private/RenderCore.cpp
--- a/Engine/Source/Runtime/RenderCore/Private/RenderCore.cpp
+++ b/Engine/Source/Runtime/RenderCore/Private/RenderCore.cpp
@@ -74,6 +74,7 @@ DEFINE_STAT(STAT_UpdateIndirectLightingCache);
 DEFINE_STAT(STAT_SortStaticDrawLists);
 DEFINE_STAT(STAT_InitDynamicShadowsTime);
 DEFINE_STAT(STAT_PreRenderView);
+DEFINE_STAT(STAT_GetDynamicMeshElements);
 DEFINE_STAT(STAT_UpdateStaticMeshesTime);
 DEFINE_STAT(STAT_StaticRelevance);
 DEFINE_STAT(STAT_ComputeViewRelevance);
diff --git a/Engine/Source/Runtime/RenderCore/Public/RenderCore.h b/Engine/Source/Runtime/RenderCore/Public/RenderCore.h
--- a/Engine/Source/Runtime/RenderCore/Public/RenderCore.h
+++ b/Engine/Source/Runtime/RenderCore/Public/RenderCore.h
@@ -49,6 +49,7 @@ DECLARE_CYCLE_STAT_EXTERN(TEXT("Compute View Relevance"),STAT_ComputeViewRelevan
 DECLARE_CYCLE_STAT_EXTERN(TEXT("Static Mesh Relevance"),STAT_StaticRelevance,STATGROUP_InitViews, RENDERCORE_API);
 DECLARE_CYCLE_STAT_EXTERN(TEXT("UpdateStaticMeshes"),STAT_UpdateStaticMeshesTime,STATGROUP_InitViews, RENDERCORE_API);
 DECLARE_CYCLE_STAT_EXTERN(TEXT("PreRenderView"),STAT_PreRenderView,STATGROUP_InitViews, RENDERCORE_API);
+DECLARE_CYCLE_STAT_EXTERN(TEXT("GetDynamicMeshElements"),STAT_GetDynamicMeshElements,STATGROUP_InitViews, RENDERCORE_API);
 DECLARE_CYCLE_STAT_EXTERN(TEXT("Init dynamic shadows"),STAT_InitDynamicShadowsTime,STATGROUP_InitViews, RENDERCORE_API);
 DECLARE_CYCLE_STAT_EXTERN(TEXT("SortStaticDrawLists"),STAT_SortStaticDrawLists,STATGROUP_InitViews, RENDERCORE_API);
 DECLARE_CYCLE_STAT_EXTERN(TEXT("Update Indirect Lighting Cache"),STAT_UpdateIndirectLightingCache,STATGROUP_InitViews, RENDERCORE_API);
diff --git a/Engine/Source/Runtime/RenderCore/Public/RenderResource.h b/Engine/Source/Runtime/RenderCore/Public/RenderResource.h
--- a/Engine/Source/Runtime/RenderCore/Public/RenderResource.h
+++ b/Engine/Source/Runtime/RenderCore/Public/RenderResource.h
@@ -92,6 +92,12 @@ public:
 	// Accessors.
 	FORCEINLINE bool IsInitialized() const { return bInitialized; }
 
+	// For those situations when the default ctor had to be used
+	inline void SetFeatureLevel(ERHIFeatureLevel::Type InFeatureLevel)
+	{
+		FeatureLevel = InFeatureLevel;
+	}
+
 protected:
 	// This is used during mobile editor preview refactor, this will eventually be replaced with a parameter to InitRHI() etc..
 	ERHIFeatureLevel::Type GetFeatureLevel() const { return FeatureLevel == ERHIFeatureLevel::Num ? GMaxRHIFeatureLevel : FeatureLevel; }
diff --git a/Engine/Source/Runtime/Renderer/Private/CompositionLighting/CompositionLighting.cpp b/Engine/Source/Runtime/Renderer/Private/CompositionLighting/CompositionLighting.cpp
--- a/Engine/Source/Runtime/Renderer/Private/CompositionLighting/CompositionLighting.cpp
+++ b/Engine/Source/Runtime/Renderer/Private/CompositionLighting/CompositionLighting.cpp
@@ -230,7 +230,7 @@ static void AddDeferredDecalsBeforeLighting(FPostprocessContext& Context)
 	Context.FinalOutput = FRenderingCompositeOutputRef(Pass);
 }
 
-void FCompositionLighting::ProcessBeforeBasePass(FRHICommandListImmediate& RHICmdList, const FViewInfo& View)
+void FCompositionLighting::ProcessBeforeBasePass(FRHICommandListImmediate& RHICmdList, FViewInfo& View)
 {
 	check(IsInRenderingThread());
 
@@ -265,7 +265,7 @@ void FCompositionLighting::ProcessBeforeBasePass(FRHICommandListImmediate& RHICm
 	}
 }
 
-void FCompositionLighting::ProcessAfterBasePass(FRHICommandListImmediate& RHICmdList, const FViewInfo& View)
+void FCompositionLighting::ProcessAfterBasePass(FRHICommandListImmediate& RHICmdList, FViewInfo& View)
 {
 	check(IsInRenderingThread());
 	
@@ -326,7 +326,7 @@ void FCompositionLighting::ProcessAfterBasePass(FRHICommandListImmediate& RHICmd
 }
 
 
-void FCompositionLighting::ProcessLighting(FRHICommandListImmediate& RHICmdList, const FViewInfo& View)
+void FCompositionLighting::ProcessLighting(FRHICommandListImmediate& RHICmdList, FViewInfo& View)
 {
 	check(IsInRenderingThread());
 	
diff --git a/Engine/Source/Runtime/Renderer/Private/CompositionLighting/CompositionLighting.h b/Engine/Source/Runtime/Renderer/Private/CompositionLighting/CompositionLighting.h
--- a/Engine/Source/Runtime/Renderer/Private/CompositionLighting/CompositionLighting.h
+++ b/Engine/Source/Runtime/Renderer/Private/CompositionLighting/CompositionLighting.h
@@ -14,11 +14,11 @@
 class FCompositionLighting
 {
 public:
-	void ProcessBeforeBasePass(FRHICommandListImmediate& RHICmdList, const FViewInfo& View);
+	void ProcessBeforeBasePass(FRHICommandListImmediate& RHICmdList, FViewInfo& View);
 
-	void ProcessAfterBasePass(FRHICommandListImmediate& RHICmdList, const FViewInfo& View);
+	void ProcessAfterBasePass(FRHICommandListImmediate& RHICmdList,  FViewInfo& View);
 
-	void ProcessLighting(FRHICommandListImmediate& RHICmdList, const FViewInfo& View);
+	void ProcessLighting(FRHICommandListImmediate& RHICmdList, FViewInfo& View);
 };
 
 /** The global used for deferred lighting. */
diff --git a/Engine/Source/Runtime/Renderer/Private/CustomDepthRendering.cpp b/Engine/Source/Runtime/Renderer/Private/CustomDepthRendering.cpp
--- a/Engine/Source/Runtime/Renderer/Private/CustomDepthRendering.cpp
+++ b/Engine/Source/Runtime/Renderer/Private/CustomDepthRendering.cpp
@@ -16,28 +16,46 @@
 	FCustomDepthPrimSet
 -----------------------------------------------------------------------------*/
 
-bool FCustomDepthPrimSet::DrawPrims(FRHICommandListImmediate& RHICmdList, const FViewInfo* ViewInfo, bool bInitializeOffsets)
+bool FCustomDepthPrimSet::DrawPrims(FRHICommandListImmediate& RHICmdList, const FViewInfo& View)
 {
 	bool bDirty=false;
 
 	if(Prims.Num())
 	{
-		SCOPED_DRAW_EVENT(DistortionAccum, DEC_SCENE_ITEMS);
-		TDynamicPrimitiveDrawer<FDepthDrawingPolicyFactory> Drawer(RHICmdList, ViewInfo, FDepthDrawingPolicyFactory::ContextType(DDM_AllOccluders), true);
+		SCOPED_DRAW_EVENT(CustomDepth, DEC_SCENE_ITEMS);
+
+		const bool bUseGetMeshElements = ShouldUseGetDynamicMeshElements();
+
+		TDynamicPrimitiveDrawer<FDepthDrawingPolicyFactory> Drawer(RHICmdList, &View, FDepthDrawingPolicyFactory::ContextType(DDM_AllOccluders), true);
 
 		for (int32 PrimIdx = 0; PrimIdx < Prims.Num(); PrimIdx++)
 		{
 			FPrimitiveSceneProxy* PrimitiveSceneProxy = Prims[PrimIdx];
 			const FPrimitiveSceneInfo* PrimitiveSceneInfo = PrimitiveSceneProxy->GetPrimitiveSceneInfo();
 
-			if (ViewInfo->PrimitiveVisibilityMap[PrimitiveSceneInfo->GetIndex()])
+			if (View.PrimitiveVisibilityMap[PrimitiveSceneInfo->GetIndex()])
 			{
-				const FPrimitiveViewRelevance& ViewRelevance = ViewInfo->PrimitiveViewRelevanceMap[PrimitiveSceneInfo->GetIndex()];
+				const FPrimitiveViewRelevance& ViewRelevance = View.PrimitiveViewRelevanceMap[PrimitiveSceneInfo->GetIndex()];
+
+				if (bUseGetMeshElements)
+				{
+					FDepthDrawingPolicyFactory::ContextType Context(DDM_AllOccluders);
+
+					for (int32 MeshBatchIndex = 0; MeshBatchIndex < View.DynamicMeshElements.Num(); MeshBatchIndex++)
+					{
+						const FMeshBatchAndRelevance& MeshBatchAndRelevance = View.DynamicMeshElements[MeshBatchIndex];
 
-				if (ViewRelevance.bDynamicRelevance)
+						if (MeshBatchAndRelevance.PrimitiveSceneProxy == PrimitiveSceneProxy)
+						{
+							const FMeshBatch& MeshBatch = *MeshBatchAndRelevance.Mesh;
+							FDepthDrawingPolicyFactory::DrawDynamicMesh(RHICmdList, View, Context, MeshBatch, false, true, MeshBatchAndRelevance.PrimitiveSceneProxy, MeshBatch.BatchHitProxyId);
+						}
+					}
+				}
+				else if (ViewRelevance.bDynamicRelevance)
 				{
 					Drawer.SetPrimitive(PrimitiveSceneProxy);
-					PrimitiveSceneProxy->DrawDynamicElements(&Drawer, ViewInfo);
+					PrimitiveSceneProxy->DrawDynamicElements(&Drawer, &View);
 				}
 
 				if (ViewRelevance.bStaticRelevance)
@@ -46,17 +64,17 @@ bool FCustomDepthPrimSet::DrawPrims(FRHICommandListImmediate& RHICmdList, const
 					{
 						const FStaticMesh& StaticMesh = PrimitiveSceneInfo->StaticMeshes[StaticMeshIdx];
 
-						if (ViewInfo->StaticMeshVisibilityMap[StaticMesh.Id])
+						if (View.StaticMeshVisibilityMap[StaticMesh.Id])
 						{
 							bDirty |= FDepthDrawingPolicyFactory::DrawStaticMesh(
 								RHICmdList, 
-								*ViewInfo,
+								View,
 								FDepthDrawingPolicyFactory::ContextType(DDM_AllOccluders),
 								StaticMesh,
-								StaticMesh.Elements.Num() == 1 ? 1 : ViewInfo->StaticMeshBatchVisibility[StaticMesh.Id],
+								StaticMesh.Elements.Num() == 1 ? 1 : View.StaticMeshBatchVisibility[StaticMesh.Id],
 								true,
 								PrimitiveSceneProxy,
-								StaticMesh.HitProxyId
+								StaticMesh.BatchHitProxyId
 								);
 						}
 					}
diff --git a/Engine/Source/Runtime/Renderer/Private/CustomDepthRendering.h b/Engine/Source/Runtime/Renderer/Private/CustomDepthRendering.h
--- a/Engine/Source/Runtime/Renderer/Private/CustomDepthRendering.h
+++ b/Engine/Source/Runtime/Renderer/Private/CustomDepthRendering.h
@@ -18,14 +18,14 @@ public:
 	* @param ViewInfo - current view used to draw items
 	* @return true if anything was drawn
 	*/
-	bool DrawPrims(FRHICommandListImmediate& RHICmdList, const class FViewInfo* ViewInfo, bool bInitializeOffsets);
+	bool DrawPrims(FRHICommandListImmediate& RHICmdList, const class FViewInfo& View);
 
 	/**
 	* Add a new primitive to the list of prims
 	* @param PrimitiveSceneProxy - primitive info to add.
 	* @param ViewInfo - used to transform bounds to view space
 	*/
-	void AddScenePrimitive(FPrimitiveSceneProxy* PrimitiveSceneProxy,const FViewInfo& ViewInfo)
+	void AddScenePrimitive(FPrimitiveSceneProxy* PrimitiveSceneProxy)
 	{
 		Prims.Add(PrimitiveSceneProxy);
 	}
@@ -38,16 +38,7 @@ public:
 		return Prims.Num();
 	}
 
-	/** 
-	* @return a prim currently set to render
-	*/
-/*
-	const FPrimitiveSceneProxy* GetPrim(int32 i)const
-	{
-		return Prims[i];
-	}*/
-
 private:
 	/** list of prims added from the scene */
-	TArray<FPrimitiveSceneProxy*> Prims;
+	TArray<FPrimitiveSceneProxy*, SceneRenderingAllocator> Prims;
 };
diff --git a/Engine/Source/Runtime/Renderer/Private/DeferredShadingRenderer.cpp b/Engine/Source/Runtime/Renderer/Private/DeferredShadingRenderer.cpp
--- a/Engine/Source/Runtime/Renderer/Private/DeferredShadingRenderer.cpp
+++ b/Engine/Source/Runtime/Renderer/Private/DeferredShadingRenderer.cpp
@@ -57,6 +57,14 @@ static TAutoConsoleVariable<int32> CVarVisualizeTexturePool(
 	ECVF_Cheat | ECVF_RenderThreadSafe);
 #endif
 
+/**  */
+static TAutoConsoleVariable<int32> CVarUseGetMeshElements(
+	TEXT("r.UseGetMeshElements"),
+	0,
+	TEXT("Whether to use the GetMeshElements path or the legacy DrawDynamicElements"),
+	ECVF_RenderThreadSafe
+	);
+
 /*-----------------------------------------------------------------------------
 	FDeferredShadingSceneRenderer
 -----------------------------------------------------------------------------*/
@@ -350,7 +358,27 @@ void FDeferredShadingSceneRenderer::RenderBasePassDynamicData(FRHICommandList& R
 	SCOPE_CYCLE_COUNTER(STAT_DynamicPrimitiveDrawTime);
 	SCOPED_DRAW_EVENT(Dynamic, DEC_SCENE_ITEMS);
 
-	if (View.VisibleDynamicPrimitives.Num() > 0)
+	const bool bUseGetMeshElements = CVarUseGetMeshElements.GetValueOnRenderThread() != 0;
+
+	if (bUseGetMeshElements)
+	{
+		FBasePassOpaqueDrawingPolicyFactory::ContextType Context(false, ESceneRenderTargetsMode::DontSet);
+
+		for (int32 MeshBatchIndex = 0; MeshBatchIndex < View.DynamicMeshElements.Num(); MeshBatchIndex++)
+		{
+			const FMeshBatchAndRelevance& MeshBatchAndRelevance = View.DynamicMeshElements[MeshBatchIndex];
+
+			if ((MeshBatchAndRelevance.bHasOpaqueOrMaskedMaterial || ViewFamily.EngineShowFlags.Wireframe)
+				&& MeshBatchAndRelevance.bRenderInMainPass)
+			{
+				const FMeshBatch& MeshBatch = *MeshBatchAndRelevance.Mesh;
+				FBasePassOpaqueDrawingPolicyFactory::DrawDynamicMesh(RHICmdList, View, Context, MeshBatch, false, true, MeshBatchAndRelevance.PrimitiveSceneProxy, MeshBatch.BatchHitProxyId);
+			}
+		}
+
+		View.SimpleElementCollector.DrawBatchedElements(RHICmdList, View, FTexture2DRHIRef(), EBlendModeFilter::OpaqueAndMasked);
+	}
+	else if (View.VisibleDynamicPrimitives.Num() > 0)
 	{
 		// Draw the dynamic non-occluded primitives using a base pass drawing policy.
 		TDynamicPrimitiveDrawer<FBasePassOpaqueDrawingPolicyFactory> Drawer(
@@ -384,7 +412,7 @@ void FDeferredShadingSceneRenderer::RenderBasePassDynamicData(FRHICommandList& R
 	if( !View.Family->EngineShowFlags.CompositeEditorPrimitives )
 	{
 		
-		const bool bNeedToSwitchVerticalAxis = IsES2Platform(GRHIShaderPlatform) && !IsPCPlatform(GRHIShaderPlatform);
+		const bool bNeedToSwitchVerticalAxis = RHINeedsToSwitchVerticalAxis(GRHIShaderPlatform);
 
 		// Draw the base pass for the view's batched mesh elements.
 		bDirty = DrawViewElements<FBasePassOpaqueDrawingPolicyFactory>(RHICmdList, View, FBasePassOpaqueDrawingPolicyFactory::ContextType(false, ESceneRenderTargetsMode::DontSet), SDPG_World, true) || bDirty;
@@ -505,7 +533,6 @@ void FDeferredShadingSceneRenderer::RenderBasePassViewParallel(FViewInfo& View,
 bool FDeferredShadingSceneRenderer::RenderBasePassView(FRHICommandListImmediate& RHICmdList, FViewInfo& View)
 {
 	bool bDirty = false; 
-
 	SetupBasePassView(RHICmdList, View.ViewRect, ViewFamily.EngineShowFlags.ShaderComplexity);
 	bDirty |= RenderBasePassStaticData(RHICmdList, View);
 	RenderBasePassDynamicData(RHICmdList, View, bDirty);
@@ -938,43 +965,83 @@ void FDeferredShadingSceneRenderer::Render(FRHICommandListImmediate& RHICmdList)
 
 bool FDeferredShadingSceneRenderer::RenderPrePassViewDynamic(FRHICommandList& RHICmdList, const FViewInfo& View)
 {
-	// Draw the dynamic occluder primitives using a depth drawing policy.
-	TDynamicPrimitiveDrawer<FDepthDrawingPolicyFactory> Drawer(RHICmdList, &View, FDepthDrawingPolicyFactory::ContextType(EarlyZPassMode), true);
+	const bool bUseGetMeshElements = CVarUseGetMeshElements.GetValueOnRenderThread() != 0;
+
+	if (bUseGetMeshElements)
 	{
-		SCOPED_DRAW_EVENT(Dynamic, DEC_SCENE_ITEMS);
-		for(int32 PrimitiveIndex = 0;PrimitiveIndex < View.VisibleDynamicPrimitives.Num();PrimitiveIndex++)
-		{
-			const FPrimitiveSceneInfo* PrimitiveSceneInfo = View.VisibleDynamicPrimitives[PrimitiveIndex];
-			int32 PrimitiveId = PrimitiveSceneInfo->GetIndex();
-			const FPrimitiveViewRelevance& PrimitiveViewRelevance = View.PrimitiveViewRelevanceMap[PrimitiveId];
+		FDepthDrawingPolicyFactory::ContextType Context(EarlyZPassMode);
 
-			bool bShouldUseAsOccluder = true;
+		for (int32 MeshBatchIndex = 0; MeshBatchIndex < View.DynamicMeshElements.Num(); MeshBatchIndex++)
+		{
+			const FMeshBatchAndRelevance& MeshBatchAndRelevance = View.DynamicMeshElements[MeshBatchIndex];
 
-			if(EarlyZPassMode != DDM_AllOccluders)
+			if (MeshBatchAndRelevance.bHasOpaqueOrMaskedMaterial && MeshBatchAndRelevance.bRenderInMainPass)
 			{
-				extern float GMinScreenRadiusForDepthPrepass;
-				const float LODFactorDistanceSquared = (PrimitiveSceneInfo->Proxy->GetBounds().Origin - View.ViewMatrices.ViewOrigin).SizeSquared() * FMath::Square(View.LODDistanceFactor);
-
-				// Only render primitives marked as occluders
-				bShouldUseAsOccluder = PrimitiveSceneInfo->Proxy->ShouldUseAsOccluder()
-					// Only render static objects unless movable are requested
-					&& (!PrimitiveSceneInfo->Proxy->IsMovable() || GEarlyZPassMovable)
-					&& (FMath::Square(PrimitiveSceneInfo->Proxy->GetBounds().SphereRadius) > GMinScreenRadiusForDepthPrepass * GMinScreenRadiusForDepthPrepass * LODFactorDistanceSquared);
+				const FMeshBatch& MeshBatch = *MeshBatchAndRelevance.Mesh;
+				const FPrimitiveSceneProxy* PrimitiveSceneProxy = MeshBatchAndRelevance.PrimitiveSceneProxy;
+				bool bShouldUseAsOccluder = true;
+
+				if (EarlyZPassMode != DDM_AllOccluders)
+				{
+					extern float GMinScreenRadiusForDepthPrepass;
+					//@todo - move these proxy properties into FMeshBatchAndRelevance so we don't have to dereference the proxy in order to reject a mesh
+					const float LODFactorDistanceSquared = (PrimitiveSceneProxy->GetBounds().Origin - View.ViewMatrices.ViewOrigin).SizeSquared() * FMath::Square(View.LODDistanceFactor);
+
+					// Only render primitives marked as occluders
+					bShouldUseAsOccluder = PrimitiveSceneProxy->ShouldUseAsOccluder()
+						// Only render static objects unless movable are requested
+						&& (!PrimitiveSceneProxy->IsMovable() || GEarlyZPassMovable)
+						&& (FMath::Square(PrimitiveSceneProxy->GetBounds().SphereRadius) > GMinScreenRadiusForDepthPrepass * GMinScreenRadiusForDepthPrepass * LODFactorDistanceSquared);
+				}
+
+				if (bShouldUseAsOccluder)
+				{
+					FDepthDrawingPolicyFactory::DrawDynamicMesh(RHICmdList, View, Context, MeshBatch, false, true, PrimitiveSceneProxy, MeshBatch.BatchHitProxyId);
+				}
 			}
-
-			// Only render opaque primitives marked as occluders
-			if (bShouldUseAsOccluder && PrimitiveViewRelevance.bOpaqueRelevance && PrimitiveViewRelevance.bRenderInMainPass)
+		}
+	}
+	else if (View.VisibleDynamicPrimitives.Num() > 0)
+	{
+		// Draw the dynamic occluder primitives using a depth drawing policy.
+		TDynamicPrimitiveDrawer<FDepthDrawingPolicyFactory> Drawer(RHICmdList, &View, FDepthDrawingPolicyFactory::ContextType(EarlyZPassMode), true);
+		{
+			SCOPED_DRAW_EVENT(Dynamic, DEC_SCENE_ITEMS);
+			for(int32 PrimitiveIndex = 0;PrimitiveIndex < View.VisibleDynamicPrimitives.Num();PrimitiveIndex++)
 			{
-				FScopeCycleCounter Context(PrimitiveSceneInfo->Proxy->GetStatId());
-				Drawer.SetPrimitive(PrimitiveSceneInfo->Proxy);
-				PrimitiveSceneInfo->Proxy->DrawDynamicElements(
-					&Drawer,
-					&View
-					);
+				const FPrimitiveSceneInfo* PrimitiveSceneInfo = View.VisibleDynamicPrimitives[PrimitiveIndex];
+				int32 PrimitiveId = PrimitiveSceneInfo->GetIndex();
+				const FPrimitiveViewRelevance& PrimitiveViewRelevance = View.PrimitiveViewRelevanceMap[PrimitiveId];
+
+				bool bShouldUseAsOccluder = true;
+
+				if(EarlyZPassMode != DDM_AllOccluders)
+				{
+					extern float GMinScreenRadiusForDepthPrepass;
+					const float LODFactorDistanceSquared = (PrimitiveSceneInfo->Proxy->GetBounds().Origin - View.ViewMatrices.ViewOrigin).SizeSquared() * FMath::Square(View.LODDistanceFactor);
+
+					// Only render primitives marked as occluders
+					bShouldUseAsOccluder = PrimitiveSceneInfo->Proxy->ShouldUseAsOccluder()
+						// Only render static objects unless movable are requested
+						&& (!PrimitiveSceneInfo->Proxy->IsMovable() || GEarlyZPassMovable)
+						&& (FMath::Square(PrimitiveSceneInfo->Proxy->GetBounds().SphereRadius) > GMinScreenRadiusForDepthPrepass * GMinScreenRadiusForDepthPrepass * LODFactorDistanceSquared);
+				}
+
+				// Only render opaque primitives marked as occluders
+				if (bShouldUseAsOccluder && PrimitiveViewRelevance.bOpaqueRelevance && PrimitiveViewRelevance.bRenderInMainPass)
+				{
+					FScopeCycleCounter Context(PrimitiveSceneInfo->Proxy->GetStatId());
+					Drawer.SetPrimitive(PrimitiveSceneInfo->Proxy);
+					PrimitiveSceneInfo->Proxy->DrawDynamicElements(
+						&Drawer,
+						&View
+						);
+				}
 			}
 		}
 	}
-	return Drawer.IsDirty();
+
+	return true;
 }
 
 static void SetupPrePassView(FRHICommandList& RHICmdList, const FIntRect& ViewRect)
diff --git a/Engine/Source/Runtime/Renderer/Private/DeferredShadingRenderer.h b/Engine/Source/Runtime/Renderer/Private/DeferredShadingRenderer.h
--- a/Engine/Source/Runtime/Renderer/Private/DeferredShadingRenderer.h
+++ b/Engine/Source/Runtime/Renderer/Private/DeferredShadingRenderer.h
@@ -254,6 +254,7 @@ private:
 	/** Renders the velocities for a subset of movable objects for the motion blur effect. */
 	friend class FRenderVelocityDynamicThreadTask;
 	void RenderDynamicVelocitiesInner(FRHICommandList& RHICmdList, const FViewInfo& View, int32 FirstIndex, int32 LastIndex);
+	void RenderDynamicVelocitiesMeshElementsInner(FRHICommandList& RHICmdList, const FViewInfo& View);
 
 	/** Renders the velocities of movable objects for the motion blur effect. */
 	void RenderVelocitiesInner(FRHICommandListImmediate& RHICmdList, const FViewInfo& View);
@@ -268,7 +269,7 @@ private:
 	 * Finish rendering a view, writing the contents to ViewFamily.RenderTarget.
 	 * @param View - The view to process.
 	*/
-	void FinishRenderViewTarget(FRHICommandListImmediate& RHICmdList, const FViewInfo* View, bool bLastView);
+	void FinishRenderViewTarget(FRHICommandListImmediate& RHICmdList, FViewInfo* View, bool bLastView);
 
 	/**
 	  * Used by RenderLights to figure out if projected shadows need to be rendered to the attenuation buffer.
diff --git a/Engine/Source/Runtime/Renderer/Private/DistortionRendering.cpp b/Engine/Source/Runtime/Renderer/Private/DistortionRendering.cpp
--- a/Engine/Source/Runtime/Renderer/Private/DistortionRendering.cpp
+++ b/Engine/Source/Runtime/Renderer/Private/DistortionRendering.cpp
@@ -677,14 +677,14 @@ bool TDistortionMeshDrawingPolicyFactory<DistortMeshPolicy>::IsMaterialIgnored(c
 	FDistortionPrimSet
 -----------------------------------------------------------------------------*/
 
-bool FDistortionPrimSet::DrawAccumulatedOffsets(FRHICommandListImmediate& RHICmdList, const FViewInfo* ViewInfo, bool bInitializeOffsets)
+bool FDistortionPrimSet::DrawAccumulatedOffsets(FRHICommandListImmediate& RHICmdList, const FViewInfo& View, bool bInitializeOffsets)
 {
 	bool bDirty=false;
 
 	// Draw the view's elements with the translucent drawing policy.
 	bDirty |= DrawViewElements<TDistortionMeshDrawingPolicyFactory<FDistortMeshAccumulatePolicy> >(
 		RHICmdList,
-		*ViewInfo,
+		View,
 		bInitializeOffsets,
 		0,	// DPG Index?
 		false // Distortion is rendered post fog.
@@ -692,10 +692,12 @@ bool FDistortionPrimSet::DrawAccumulatedOffsets(FRHICommandListImmediate& RHICmd
 
 	if( Prims.Num() )
 	{
+		const bool bUseGetMeshElements = ShouldUseGetDynamicMeshElements();
+
 		// For drawing scene prims with dynamic relevance.
 		TDynamicPrimitiveDrawer<TDistortionMeshDrawingPolicyFactory<FDistortMeshAccumulatePolicy> > Drawer(
 			RHICmdList,
-			ViewInfo,
+			&View,
 			bInitializeOffsets,
 			false // Distortion is rendered post fog.
 			);
@@ -704,41 +706,54 @@ bool FDistortionPrimSet::DrawAccumulatedOffsets(FRHICommandListImmediate& RHICmd
 		for( int32 PrimIdx=0; PrimIdx < Prims.Num(); PrimIdx++ )
 		{
 			FPrimitiveSceneProxy* PrimitiveSceneProxy = Prims[PrimIdx];
-			const FPrimitiveViewRelevance& ViewRelevance = ViewInfo->PrimitiveViewRelevanceMap[PrimitiveSceneProxy->GetPrimitiveSceneInfo()->GetIndex()];
+			const FPrimitiveViewRelevance& ViewRelevance = View.PrimitiveViewRelevanceMap[PrimitiveSceneProxy->GetPrimitiveSceneInfo()->GetIndex()];
 
-			if( ViewRelevance.bDistortionRelevance )
+			if (bUseGetMeshElements)
 			{
-				// Render dynamic scene prim
-				if( ViewRelevance.bDynamicRelevance )
-				{				
-					Drawer.SetPrimitive(PrimitiveSceneProxy);
-					PrimitiveSceneProxy->DrawDynamicElements(
-						&Drawer,
-						ViewInfo
-						);
+				TDistortionMeshDrawingPolicyFactory<FDistortMeshAccumulatePolicy>::ContextType Context(bInitializeOffsets);
+
+				for (int32 MeshBatchIndex = 0; MeshBatchIndex < View.DynamicMeshElements.Num(); MeshBatchIndex++)
+				{
+					const FMeshBatchAndRelevance& MeshBatchAndRelevance = View.DynamicMeshElements[MeshBatchIndex];
+
+					if (MeshBatchAndRelevance.PrimitiveSceneProxy == PrimitiveSceneProxy)
+					{
+						const FMeshBatch& MeshBatch = *MeshBatchAndRelevance.Mesh;
+						TDistortionMeshDrawingPolicyFactory<FDistortMeshAccumulatePolicy>::DrawDynamicMesh(RHICmdList, View, Context, MeshBatch, false, false, MeshBatchAndRelevance.PrimitiveSceneProxy, MeshBatch.BatchHitProxyId);
+					}
 				}
-				// Render static scene prim
-				if( ViewRelevance.bStaticRelevance )
+			}
+			// Render dynamic scene prim
+			else if( ViewRelevance.bDynamicRelevance )
+			{				
+				Drawer.SetPrimitive(PrimitiveSceneProxy);
+				PrimitiveSceneProxy->DrawDynamicElements(
+					&Drawer,
+					&View
+					);
+			}
+
+			// Render static scene prim
+			if( ViewRelevance.bStaticRelevance )
+			{
+				// Render static meshes from static scene prim
+				for( int32 StaticMeshIdx=0; StaticMeshIdx < PrimitiveSceneProxy->GetPrimitiveSceneInfo()->StaticMeshes.Num(); StaticMeshIdx++ )
 				{
-					// Render static meshes from static scene prim
-					for( int32 StaticMeshIdx=0; StaticMeshIdx < PrimitiveSceneProxy->GetPrimitiveSceneInfo()->StaticMeshes.Num(); StaticMeshIdx++ )
+					FStaticMesh& StaticMesh = PrimitiveSceneProxy->GetPrimitiveSceneInfo()->StaticMeshes[StaticMeshIdx];
+					if( View.StaticMeshVisibilityMap[StaticMesh.Id]
+						// Only render static mesh elements using translucent materials
+						&& StaticMesh.IsTranslucent(View.GetFeatureLevel()) )
 					{
-						FStaticMesh& StaticMesh = PrimitiveSceneProxy->GetPrimitiveSceneInfo()->StaticMeshes[StaticMeshIdx];
-						if( ViewInfo->StaticMeshVisibilityMap[StaticMesh.Id]
-							// Only render static mesh elements using translucent materials
-							&& StaticMesh.IsTranslucent(ViewInfo->GetFeatureLevel()) )
-						{
-							bDirty |= TDistortionMeshDrawingPolicyFactory<FDistortMeshAccumulatePolicy>::DrawStaticMesh(
-								RHICmdList, 
-								ViewInfo,
-								bInitializeOffsets,
-								StaticMesh,
-								StaticMesh.Elements.Num() == 1 ? 1 : ViewInfo->StaticMeshBatchVisibility[StaticMesh.Id],
-								false,
-								PrimitiveSceneProxy,
-								StaticMesh.HitProxyId
-								);
-						}
+						bDirty |= TDistortionMeshDrawingPolicyFactory<FDistortMeshAccumulatePolicy>::DrawStaticMesh(
+							RHICmdList, 
+							&View,
+							bInitializeOffsets,
+							StaticMesh,
+							StaticMesh.Elements.Num() == 1 ? 1 : View.StaticMeshBatchVisibility[StaticMesh.Id],
+							false,
+							PrimitiveSceneProxy,
+							StaticMesh.BatchHitProxyId
+							);
 					}
 				}
 			}
@@ -819,7 +834,7 @@ void FSceneRenderer::RenderDistortion(FRHICommandListImmediate& RHICmdList)
 				RHICmdList.SetBlendState(TStaticBlendState<CW_RGBA, BO_Add, BF_One, BF_One, BO_Add, BF_One, BF_One>::GetRHI());
 
 				// draw only distortion meshes to accumulate their offsets
-				bDirty |= View.DistortionPrimSet.DrawAccumulatedOffsets(RHICmdList, &View, false);
+				bDirty |= View.DistortionPrimSet.DrawAccumulatedOffsets(RHICmdList, View, false);
 			}
 
 			if (bDirty)
diff --git a/Engine/Source/Runtime/Renderer/Private/DistortionRendering.h b/Engine/Source/Runtime/Renderer/Private/DistortionRendering.h
--- a/Engine/Source/Runtime/Renderer/Private/DistortionRendering.h
+++ b/Engine/Source/Runtime/Renderer/Private/DistortionRendering.h
@@ -19,7 +19,7 @@ public:
 	* @param DPGIndex - current DPG used to draw items
 	* @return true if anything was drawn
 	*/
-	bool DrawAccumulatedOffsets(FRHICommandListImmediate& RHICmdList, const class FViewInfo* ViewInfo, bool bInitializeOffsets);
+	bool DrawAccumulatedOffsets(FRHICommandListImmediate& RHICmdList, const class FViewInfo& View, bool bInitializeOffsets);
 
 	/**
 	* Add a new primitive to the list of distortion prims
@@ -47,5 +47,5 @@ public:
 
 private:
 	/** list of distortion prims added from the scene */
-	TArray<FPrimitiveSceneProxy*> Prims;
+	TArray<FPrimitiveSceneProxy*, SceneRenderingAllocator> Prims;
 };
diff --git a/Engine/Source/Runtime/Renderer/Private/DynamicPrimitiveDrawing.h b/Engine/Source/Runtime/Renderer/Private/DynamicPrimitiveDrawing.h
--- a/Engine/Source/Runtime/Renderer/Private/DynamicPrimitiveDrawing.h
+++ b/Engine/Source/Runtime/Renderer/Private/DynamicPrimitiveDrawing.h
@@ -30,7 +30,6 @@ public:
 		const typename DrawingPolicyFactoryType::ContextType& InDrawingContext,
 		bool InPreFog,
 		bool bInIsHitTesting = false,
-		bool bInIsVelocityRendering = false,
 		bool bInEditorCompositeDepthTest = false,
 		bool bInIsSelectionOutlineRendering = false
 		):
@@ -42,7 +41,6 @@ public:
 		bPreFog(InPreFog),
 		bDirty(false),
 		bIsHitTesting(bInIsHitTesting),
-		bIsVelocityRendering(bInIsVelocityRendering),
 		bEditorCompositeDepthTest(bInEditorCompositeDepthTest),
 		bIsSelectionOutlineRendering(bInIsSelectionOutlineRendering)
 	{}
@@ -88,11 +86,6 @@ public:
 		uint8 DepthPriorityGroup
 		) override;
 
-	// Accessors.
-	bool IsPreFog() const
-	{
-		return bPreFog;
-	}
 	bool IsDirty() const
 	{
 		return bDirty;
@@ -102,11 +95,6 @@ public:
 		bDirty = false;
 	}
 
-	/**
-	 * @return true if rendering is occurring during velocity pass 
-	 */
-	virtual bool IsRenderingVelocities() const { return bIsVelocityRendering; }
-
 private:
 
 	/** Command list to draw to. */
@@ -139,9 +127,6 @@ private:
 	/** true if hit proxies are being drawn. */
 	uint32 bIsHitTesting : 1;
 
-	/** true if rendering is occuring during velocity pass */
-	uint32 bIsVelocityRendering : 1;
-
 	/** true if rendering is occuring during editor compositing */
 	uint32 bEditorCompositeDepthTest : 1;
 
diff --git a/Engine/Source/Runtime/Renderer/Private/DynamicPrimitiveDrawing.inl b/Engine/Source/Runtime/Renderer/Private/DynamicPrimitiveDrawing.inl
--- a/Engine/Source/Runtime/Renderer/Private/DynamicPrimitiveDrawing.inl
+++ b/Engine/Source/Runtime/Renderer/Private/DynamicPrimitiveDrawing.inl
@@ -12,7 +12,7 @@ TDynamicPrimitiveDrawer<DrawingPolicyFactoryType>::~TDynamicPrimitiveDrawer()
 {
 	if(View)
 	{
-		const bool bNeedToSwitchVerticalAxis = IsES2Platform(GRHIShaderPlatform) && !IsPCPlatform(GRHIShaderPlatform) && !IsMobileHDR();
+		const bool bNeedToSwitchVerticalAxis = RHINeedsToSwitchVerticalAxis(GRHIShaderPlatform) && !IsMobileHDR();
 
 		// Determine whether or not to depth test simple batched elements manually (msaa only)
 		// We need to read from the scene depth in the pixel shader to do manual depth testing
@@ -224,12 +224,12 @@ void DrawViewElementsInner(
 	)
 {
 	// Get the correct element list based on dpg index
-	const TIndirectArray<FHitProxyMeshPair>& ViewMeshElementList = ( DPGIndex == SDPG_Foreground ? View.TopViewMeshElements : View.ViewMeshElements );
+	const TIndirectArray<FMeshBatch>& ViewMeshElementList = ( DPGIndex == SDPG_Foreground ? View.TopViewMeshElements : View.ViewMeshElements );
 	// Draw the view's mesh elements.
 	check(LastIndex < ViewMeshElementList.Num());
 	for (int32 MeshIndex = FirstIndex; MeshIndex <= LastIndex; MeshIndex++)
 	{
-		const FHitProxyMeshPair& Mesh = ViewMeshElementList[MeshIndex];
+		const FMeshBatch& Mesh = ViewMeshElementList[MeshIndex];
 		const auto FeatureLevel = View.GetFeatureLevel();
 		check(Mesh.MaterialRenderProxy);
 		check(Mesh.MaterialRenderProxy->GetMaterial(FeatureLevel));
@@ -245,7 +245,7 @@ void DrawViewElementsInner(
 				!!bBackFace,
 				bPreFog,
 				NULL,
-				Mesh.HitProxyId
+				Mesh.BatchHitProxyId
 				);
 			--bBackFace;
 		} while( bBackFace >= 0 );
@@ -315,7 +315,7 @@ void DrawViewElementsParallel(
 	)
 {
 	// Get the correct element list based on dpg index
-	const TIndirectArray<FHitProxyMeshPair>& ViewMeshElementList = (DPGIndex == SDPG_Foreground ? View.TopViewMeshElements : View.ViewMeshElements);
+	const TIndirectArray<FMeshBatch>& ViewMeshElementList = (DPGIndex == SDPG_Foreground ? View.TopViewMeshElements : View.ViewMeshElements);
 
 	{
 		int32 NumPrims = ViewMeshElementList.Num();
@@ -358,7 +358,7 @@ bool DrawViewElements(
 	)
 {
 	// Get the correct element list based on dpg index
-	const TIndirectArray<FHitProxyMeshPair>& ViewMeshElementList = (DPGIndex == SDPG_Foreground ? View.TopViewMeshElements : View.ViewMeshElements);
+	const TIndirectArray<FMeshBatch>& ViewMeshElementList = (DPGIndex == SDPG_Foreground ? View.TopViewMeshElements : View.ViewMeshElements);
 	if (ViewMeshElementList.Num() != 0)
 	{
 		DrawViewElementsInner<DrawingPolicyFactoryType>(RHICmdList, View, DrawingContext, DPGIndex, bPreFog, 0, ViewMeshElementList.Num() - 1);
@@ -539,12 +539,10 @@ inline int32 FViewElementPDI::DrawMesh(const FMeshBatch& Mesh)
 		uint8 DPGIndex = Mesh.DepthPriorityGroup;
 		// Get the correct element list based on dpg index
 		// Translucent view mesh elements in the foreground dpg are not supported yet
-		TIndirectArray<FHitProxyMeshPair>& ViewMeshElementList = ( ( DPGIndex == SDPG_Foreground  ) ? ViewInfo->TopViewMeshElements : ViewInfo->ViewMeshElements );
+		TIndirectArray<FMeshBatch>& ViewMeshElementList = ( ( DPGIndex == SDPG_Foreground  ) ? ViewInfo->TopViewMeshElements : ViewInfo->ViewMeshElements );
 
-		new(ViewMeshElementList) FHitProxyMeshPair(
-			Mesh,
-			CurrentHitProxy ? CurrentHitProxy->Id : FHitProxyId()
-			);
+		FMeshBatch* NewMesh = new(ViewMeshElementList) FMeshBatch(Mesh);
+		NewMesh->BatchHitProxyId = CurrentHitProxy ? CurrentHitProxy->Id : FHitProxyId();
 
 		return 1;
 	}
