Removed PositionZWTextureRHI reference

diff --git a/Engine/Source/Runtime/Engine/Private/Particles/ParticleGpuSimulation.cpp b/Engine/Source/Runtime/Engine/Private/Particles/ParticleGpuSimulation.cpp
--- a/Engine/Source/Runtime/Engine/Private/Particles/ParticleGpuSimulation.cpp
+++ b/Engine/Source/Runtime/Engine/Private/Particles/ParticleGpuSimulation.cpp
@@ -2713,7 +2713,6 @@ struct FGPUSpriteDynamicEmitterData : FDynamicEmitterDataBase
 					VertexFactory.EmitterUniformBuffer = Resources->UniformBuffer;
 					VertexFactory.EmitterDynamicUniformBuffer = LocalDynamicUniformBuffer;
 					VertexFactory.PositionTextureRHI = StateTextures.PositionTextureRHI;
-					VertexFactory.PositionZWTextureRHI = StateTextures.PositionZWTextureRHI;
 					VertexFactory.VelocityTextureRHI = StateTextures.VelocityTextureRHI;
 					VertexFactory.AttributesTextureRHI = ParticleSimulationResources->RenderAttributesTexture.TextureRHI;
 
