Fixed GPU particles in GetDynamicMeshElements path

diff --git a/Engine/Source/Runtime/Engine/Private/Particles/ParticleGpuSimulation.cpp b/Engine/Source/Runtime/Engine/Private/Particles/ParticleGpuSimulation.cpp
--- a/Engine/Source/Runtime/Engine/Private/Particles/ParticleGpuSimulation.cpp
+++ b/Engine/Source/Runtime/Engine/Private/Particles/ParticleGpuSimulation.cpp
@@ -561,7 +561,7 @@ public:
 	/** Emitter uniform buffer. */
 	FUniformBufferRHIParamRef EmitterUniformBuffer;
 	/** Emitter uniform buffer for dynamic parameters. */
-	FUniformBufferRHIParamRef EmitterDynamicUniformBuffer;
+	FUniformBufferRHIRef EmitterDynamicUniformBuffer;
 	/** Buffer containing particle indices. */
 	FParticleIndicesVertexBuffer* ParticleIndicesBuffer;
 	/** Offset in to the particle indices buffer. */
