Fixed issue with recording animation when thumbnail is visiible in between.

TTP 343760	saving new animation after rotating bone causes crash

diff --git a/Engine/Source/Editor/UnrealEd/Private/ThumbnailHelpers.cpp b/Engine/Source/Editor/UnrealEd/Private/ThumbnailHelpers.cpp
--- a/Engine/Source/Editor/UnrealEd/Private/ThumbnailHelpers.cpp
+++ b/Engine/Source/Editor/UnrealEd/Private/ThumbnailHelpers.cpp
@@ -522,7 +522,7 @@ bool FAnimationSequenceThumbnailScene::SetAnimation(UAnimSequenceBase* InAnimati
 
 	PreviewAnimation = InAnimation;
 
-	if (InAnimation)
+	if (InAnimation && InAnimation->IsValidToPlay())
 	{
 		if (USkeleton* Skeleton = InAnimation->GetSkeleton())
 		{
diff --git a/Engine/Source/Runtime/Engine/Classes/Animation/AnimSequence.h b/Engine/Source/Runtime/Engine/Classes/Animation/AnimSequence.h
--- a/Engine/Source/Runtime/Engine/Classes/Animation/AnimSequence.h
+++ b/Engine/Source/Runtime/Engine/Classes/Animation/AnimSequence.h
@@ -697,6 +697,11 @@ public:
 	 * This increases framecount + time, so that it requires recompression
 	 */
 	ENGINE_API bool AddLoopingInterpolation();
+
+	/** 
+	 * Add validation check to see if it's being ready to play or not
+	 */
+	virtual bool IsValidToPlay() const override;
 #endif
 
 private:
diff --git a/Engine/Source/Runtime/Engine/Classes/Animation/AnimSequenceBase.h b/Engine/Source/Runtime/Engine/Classes/Animation/AnimSequenceBase.h
--- a/Engine/Source/Runtime/Engine/Classes/Animation/AnimSequenceBase.h
+++ b/Engine/Source/Runtime/Engine/Classes/Animation/AnimSequenceBase.h
@@ -429,5 +429,6 @@ public:
 	/** Registers a delegate to be called after notification has changed*/
 	ENGINE_API void RegisterOnNotifyChanged(const FOnNotifyChanged& Delegate);
 	ENGINE_API void UnregisterOnNotifyChanged(void * Unregister);
+	ENGINE_API virtual bool IsValidToPlay() const { return true; }
 #endif
 };
diff --git a/Engine/Source/Runtime/Engine/Private/Animation/AnimSequence.cpp b/Engine/Source/Runtime/Engine/Private/Animation/AnimSequence.cpp
--- a/Engine/Source/Runtime/Engine/Private/Animation/AnimSequence.cpp
+++ b/Engine/Source/Runtime/Engine/Private/Animation/AnimSequence.cpp
@@ -265,6 +265,12 @@ int32 UAnimSequence::GetNumberOfTracks() const
 	return TrackToSkeletonMapTable.Num();
 }
 
+bool UAnimSequence::IsValidToPlay() const
+{
+	// make sure sequence length is valid and raw animation data exists, and compressed
+	return ( SequenceLength > 0.f && RawAnimationData.Num() > 0 && CompressedTrackOffsets.Num() > 0 );
+}
+
 void UAnimSequence::PostLoad()
 {
 #if WITH_EDITOR
