Small fix to Slate memory trackng

diff --git a/Engine/Source/Runtime/Engine/Private/Slate/SlateTextures.cpp b/Engine/Source/Runtime/Engine/Private/Slate/SlateTextures.cpp
--- a/Engine/Source/Runtime/Engine/Private/Slate/SlateTextures.cpp
+++ b/Engine/Source/Runtime/Engine/Private/Slate/SlateTextures.cpp
@@ -67,9 +67,14 @@ void FSlateTexture2DRHIRef::InitDynamicRHI()
 void FSlateTexture2DRHIRef::ReleaseDynamicRHI()
 {
 	check( IsInRenderingThread() );
+
+	if( IsValidRef(ShaderResource) )
+	{
+		DEC_MEMORY_STAT_BY(STAT_SlateTextureGPUMemory, Width*Height*GPixelFormats[PixelFormat].BlockBytes);
+	}
+
 	ShaderResource.SafeRelease();
 
-	DEC_MEMORY_STAT_BY(STAT_SlateTextureGPUMemory, Width*Height*GPixelFormats[PixelFormat].BlockBytes);
 }
 
 void FSlateTexture2DRHIRef::Resize( uint32 InWidth, uint32 InHeight )
