#UE4
Fix build.

diff --git a/Engine/Source/Runtime/Engine/Classes/Animation/AnimSequence.h b/Engine/Source/Runtime/Engine/Classes/Animation/AnimSequence.h
--- a/Engine/Source/Runtime/Engine/Classes/Animation/AnimSequence.h
+++ b/Engine/Source/Runtime/Engine/Classes/Animation/AnimSequence.h
@@ -697,12 +697,12 @@ public:
 	 * This increases framecount + time, so that it requires recompression
 	 */
 	ENGINE_API bool AddLoopingInterpolation();
+#endif
 
 	/** 
 	 * Add validation check to see if it's being ready to play or not
 	 */
-	virtual bool IsValidToPlay() const override;
-#endif
+	virtual bool IsValidToPlay() const;
 
 private:
 	/** 
