World composition: added checks for detecting loops in parent-child hierarchy

diff --git a/Engine/Source/Runtime/Engine/Private/WorldComposition.cpp b/Engine/Source/Runtime/Engine/Private/WorldComposition.cpp
--- a/Engine/Source/Runtime/Engine/Private/WorldComposition.cpp
+++ b/Engine/Source/Runtime/Engine/Private/WorldComposition.cpp
@@ -7,6 +7,8 @@
 #include "Engine/WorldComposition.h"
 #include "LevelUtils.h"
 
+DEFINE_LOG_CATEGORY_STATIC(LogWorldComposition, Log, All);
+
 #if WITH_EDITOR
 UWorldComposition::FWorldCompositionEvent UWorldComposition::OnWorldCompositionCreated;
 UWorldComposition::FWorldCompositionEvent UWorldComposition::OnWorldCompositionDestroyed;
@@ -280,13 +282,28 @@ void UWorldComposition::CaclulateTilesAbsolutePositions()
 {
 	for (FWorldCompositionTile& Tile : Tiles)
 	{
+		TSet<FName> VisitedParents;
+		
 		Tile.Info.AbsolutePosition = FIntPoint::ZeroValue;
+		FWorldCompositionTile* ParentTile = &Tile;
 		
-		auto* ParentTile = &Tile;
 		do 
 		{
+			// Sum relative offsets
 			Tile.Info.AbsolutePosition+= ParentTile->Info.Position;
-			ParentTile = FindTileByName(FName(*ParentTile->Info.ParentTilePackageName));
+			VisitedParents.Add(ParentTile->PackageName);
+
+			FName NextParentTileName = FName(*ParentTile->Info.ParentTilePackageName);
+			// Detect loops in parent->child hierarchy
+			FWorldCompositionTile* NextParentTile = FindTileByName(NextParentTileName);
+			if (NextParentTile && VisitedParents.Find(NextParentTileName) != nullptr)
+			{
+				UE_LOG(LogWorldComposition, Warning, TEXT("World composition tile (%s) has a cycled parent (%s)"), *Tile.PackageName.ToString(), *NextParentTileName.ToString());
+				NextParentTile = nullptr;
+				ParentTile->Info.ParentTilePackageName = TEXT("");
+			}
+			
+			ParentTile = NextParentTile;
 		} 
 		while (ParentTile);
 	}
