fixed navmesh bounds update after changing brush builder properties in details panel
ttp# 343279

diff --git a/Engine/Source/Runtime/Engine/Classes/AI/Navigation/NavMeshBoundsVolume.h b/Engine/Source/Runtime/Engine/Classes/AI/Navigation/NavMeshBoundsVolume.h
--- a/Engine/Source/Runtime/Engine/Classes/AI/Navigation/NavMeshBoundsVolume.h
+++ b/Engine/Source/Runtime/Engine/Classes/AI/Navigation/NavMeshBoundsVolume.h
@@ -16,13 +16,11 @@ class ANavMeshBoundsVolume : public AVolume
 
 	// Begin AActor Interface
 	virtual void PostInitializeComponents() override;
-#if WITH_EDITOR
-	virtual void ReregisterAllComponents() override;
 	// End AActor Interface
+#if WITH_EDITOR
 	// Begin UObject Interface
 	virtual void PostEditChangeProperty( struct FPropertyChangedEvent& PropertyChangedEvent) override;
 	// End UObject Interface
-
 #endif // WITH_EDITOR
 };
 
diff --git a/Engine/Source/Runtime/Engine/Private/NavMeshBoundsVolume.cpp b/Engine/Source/Runtime/Engine/Private/NavMeshBoundsVolume.cpp
--- a/Engine/Source/Runtime/Engine/Private/NavMeshBoundsVolume.cpp
+++ b/Engine/Source/Runtime/Engine/Private/NavMeshBoundsVolume.cpp
@@ -18,33 +18,19 @@ ANavMeshBoundsVolume::ANavMeshBoundsVolume(const class FPostConstructInitializeP
 
 #if WITH_EDITOR
 
-void ANavMeshBoundsVolume::ReregisterAllComponents()
-{
-	Super::ReregisterAllComponents();
-	/*
-	// uncommenting this would make navigation rebuild during NavMeshBoundsVolume's
-	// changing with Geometry Tool. This gets called too often though, so we'll 
-	// stick to updating navigation upon leaving geometry tool
-	if (GIsEditor)
-	{
-		UWorld* World = GetWorld();
-		if (World != NULL && World->GetNavigationSystem() != NULL)
-		{
-		GetWorld()->GetNavigationSystem()->OnNavigationBoundsUpdated(this);
-		}
-	}
-	*/
-}
-
 void ANavMeshBoundsVolume::PostEditChangeProperty( struct FPropertyChangedEvent& PropertyChangedEvent)
 {
 	Super::PostEditChangeProperty(PropertyChangedEvent);
 
 #if WITH_NAVIGATION_GENERATOR
-	if (PropertyChangedEvent.Property == NULL && GIsEditor == true && GetWorld() != NULL)
+	UNavigationSystem* NavSys = UNavigationSystem::GetCurrent(GetWorld());
+	if (GIsEditor && NavSys)
 	{
-		check(GetWorld()->GetNavigationSystem() != NULL);
-		GetWorld()->GetNavigationSystem()->OnNavigationBoundsUpdated(this);
+		if (PropertyChangedEvent.Property == NULL ||
+			PropertyChangedEvent.Property->GetFName() == GET_MEMBER_NAME_CHECKED(ABrush, BrushBuilder))
+		{
+			NavSys->OnNavigationBoundsUpdated(this);
+		}
 	}
 #endif // WITH_NAVIGATION_GENERATOR
 }
@@ -56,9 +42,10 @@ void ANavMeshBoundsVolume::PostInitializeComponents()
 	Super::PostInitializeComponents();
 
 #if WITH_NAVIGATION_GENERATOR
-	if (GetWorld()->GetNavigationSystem() != NULL && Role == ROLE_Authority)
+	UNavigationSystem* NavSys = UNavigationSystem::GetCurrent(GetWorld());
+	if (NavSys && Role == ROLE_Authority)
 	{
-		GetWorld()->GetNavigationSystem()->OnNavigationBoundsUpdated(this);
+		NavSys->OnNavigationBoundsUpdated(this);
 	}
 #endif
 }
