Prevent highres screen shot from attempting to access module manager from render thread.

diff --git a/Engine/Source/Runtime/Engine/Private/HighResScreenshot.cpp b/Engine/Source/Runtime/Engine/Private/HighResScreenshot.cpp
--- a/Engine/Source/Runtime/Engine/Private/HighResScreenshot.cpp
+++ b/Engine/Source/Runtime/Engine/Private/HighResScreenshot.cpp
@@ -20,13 +20,19 @@ FHighResScreenshotConfig::FHighResScreenshotConfig()
 	, bDumpBufferVisualizationTargets(false)
 {
 	ChangeViewport(TWeakPtr<FSceneViewport>());
-	IImageWrapperModule& ImageWrapperModule = FModuleManager::LoadModuleChecked<IImageWrapperModule>(FName("ImageWrapper"));
-
-	ImageCompressorLDR = ImageWrapperModule.CreateImageWrapper(EImageFormat::PNG);
-	ImageCompressorHDR = ImageWrapperModule.CreateImageWrapper(EImageFormat::EXR);
 	SetHDRCapture(false);
 }
 
+void FHighResScreenshotConfig::Init()
+{
+	IImageWrapperModule* ImageWrapperModule = FModuleManager::LoadModulePtr<IImageWrapperModule>(FName("ImageWrapper"));
+	if (ImageWrapperModule != nullptr)
+	{
+		ImageCompressorLDR = ImageWrapperModule->CreateImageWrapper(EImageFormat::PNG);
+		ImageCompressorHDR = ImageWrapperModule->CreateImageWrapper(EImageFormat::EXR);
+	}
+}
+
 void FHighResScreenshotConfig::ChangeViewport(TWeakPtr<FSceneViewport> InViewport)
 {
 	if (FSceneViewport* Viewport = TargetViewport.Pin().Get())
diff --git a/Engine/Source/Runtime/Engine/Public/HighResScreenshot.h b/Engine/Source/Runtime/Engine/Public/HighResScreenshot.h
--- a/Engine/Source/Runtime/Engine/Public/HighResScreenshot.h
+++ b/Engine/Source/Runtime/Engine/Public/HighResScreenshot.h
@@ -21,6 +21,9 @@ struct ENGINE_API FHighResScreenshotConfig
 
 	FHighResScreenshotConfig();
 
+	/** Initialize the image wrapper modules (required for SaveImage) **/
+	void Init();
+
 	/** Point the screenshot UI at a different viewport **/
 	void ChangeViewport(TWeakPtr<FSceneViewport> InViewport);
 
diff --git a/Engine/Source/Runtime/Launch/Private/LaunchEngineLoop.cpp b/Engine/Source/Runtime/Launch/Private/LaunchEngineLoop.cpp
--- a/Engine/Source/Runtime/Launch/Private/LaunchEngineLoop.cpp
+++ b/Engine/Source/Runtime/Launch/Private/LaunchEngineLoop.cpp
@@ -46,6 +46,7 @@
 	#include "EngineService.h"
 	#include "ContentStreaming.h"
 	#include "HeadMountedDisplay.h"
+	#include "HighResScreenshot.h"
 
 #if !UE_SERVER
 	#include "ISlateRHIRendererModule.h"
@@ -2482,6 +2483,9 @@ bool FEngineLoop::AppInit( )
 	// For platform services that need D3D hooks like Steam
 	FModuleManager::Get().LoadModule(TEXT("OnlineSubsystem"));
 	FModuleManager::Get().LoadModule(TEXT("OnlineSubsystemUtils"));
+
+	// Init HighRes screenshot system.
+	GetHighResScreenshotConfig().Init();
 #endif
 
 	// Checks.
