DataTable supports UserDefinedStruct (WIP), and it uses internally MetaData. Functions that use MD are editor-only.

diff --git a/Engine/Source/Runtime/Engine/Classes/Engine/DataTable.h b/Engine/Source/Runtime/Engine/Classes/Engine/DataTable.h
--- a/Engine/Source/Runtime/Engine/Classes/Engine/DataTable.h
+++ b/Engine/Source/Runtime/Engine/Classes/Engine/DataTable.h
@@ -116,6 +116,11 @@ class UDataTable : public UObject
 		return RowData;
 	}
 
+	/** Empty the table info (will not clear RowStruct) */
+	ENGINE_API void EmptyTable();
+
+#if WITH_EDITOR || HACK_HEADER_GENERATOR
+
 	/** Output entire contents of table as a string */
 	ENGINE_API FString GetTableAsString();
 
@@ -139,14 +144,14 @@ class UDataTable : public UObject
 	*/
 	ENGINE_API TArray<FString> CreateTableFromJSONString(const FString& InString);
 
-	/** Empty the table info (will not clear RowStruct) */
-	ENGINE_API void EmptyTable();
-
 	/** Get array of all the column titles */
 	ENGINE_API TArray<FString> GetColumnTitles() const;
 
+	TArray<UProperty*> GetTablePropertyArray(const FString& FirstRowString, UStruct* RowStruct, TArray<FString>& OutProblems);
+
 	/** Get array for each row in the table. The first row is the titles*/
 	ENGINE_API TArray< TArray<FString> > GetTableData() const;
+#endif //WITH_EDITOR || HACK_HEADER_GENERATOR
 
 	ENGINE_API static FString AssignStringToProperty(const FString& InString, const UProperty* InProp, uint8* InData);
 
@@ -158,8 +163,6 @@ class UDataTable : public UObject
 
 	ENGINE_API static bool IsSupportedTableProperty(const UProperty* InProp);
 
-	TArray<UProperty*> GetTablePropertyArray(const FString& FirstRowString, UStruct* RowStruct, TArray<FString>& OutProblems);
-
 	// End UDataTable interface
 };
 
diff --git a/Engine/Source/Runtime/Engine/Private/DataTable.cpp b/Engine/Source/Runtime/Engine/Private/DataTable.cpp
--- a/Engine/Source/Runtime/Engine/Private/DataTable.cpp
+++ b/Engine/Source/Runtime/Engine/Private/DataTable.cpp
@@ -181,6 +181,50 @@
 	}
 }
 
+void UDataTable::EmptyTable()
+{
+	UScriptStruct* LoadUsingStruct = RowStruct;
+	if (LoadUsingStruct == NULL)
+	{
+		UE_LOG(LogDataTable, Error, TEXT("Missing RowStruct while emptying DataTable '%s'!"), *GetPathName());
+		LoadUsingStruct = FTableRowBase::StaticStruct();
+	}
+
+	// Iterate over all rows in table and free mem
+	for (auto RowIt = RowMap.CreateIterator(); RowIt; ++RowIt)
+	{
+		uint8* RowData = RowIt.Value();
+		LoadUsingStruct->DestroyScriptStruct(RowData);
+		FMemory::Free(RowData);
+	}
+
+	// Finally empty the map
+	RowMap.Empty();
+}
+
+/** Returns the column property where PropertyName matches the name of the column property. Returns NULL if no match is found or the match is not a supported table property */
+UProperty* UDataTable::FindTableProperty(const FName& PropertyName) const
+{
+	UProperty* Property = NULL;
+	for (TFieldIterator<UProperty> It(RowStruct); It; ++It)
+	{
+		Property = *It;
+		check(Property != NULL);
+		if (PropertyName == Property->GetFName())
+		{
+			break;
+		}
+	}
+	if (!IsSupportedTableProperty(Property))
+	{
+		Property = NULL;
+	}
+
+	return Property;
+}
+
+#if WITH_EDITOR || HACK_HEADER_GENERATOR
+
 FString UDataTable::GetTableAsString()
 {
 	FString Result;
@@ -302,48 +346,6 @@
 	return true;
 }
 
-void UDataTable::EmptyTable()
-{
-	UScriptStruct* LoadUsingStruct = RowStruct;
-	if(LoadUsingStruct == NULL)
-	{
-		UE_LOG(LogDataTable, Error, TEXT("Missing RowStruct while emptying DataTable '%s'!"), *GetPathName());
-		LoadUsingStruct = FTableRowBase::StaticStruct();
-	}
-
-	// Iterate over all rows in table and free mem
-	for ( auto RowIt = RowMap.CreateIterator(); RowIt; ++RowIt )
-	{
-		uint8* RowData = RowIt.Value();
-		LoadUsingStruct->DestroyScriptStruct(RowData);
-		FMemory::Free(RowData);
-	}
-
-	// Finally empty the map
-	RowMap.Empty();
-}
-
-/** Returns the column property where PropertyName matches the name of the column property. Returns NULL if no match is found or the match is not a supported table property */
-UProperty* UDataTable::FindTableProperty(const FName& PropertyName) const
-{
-	UProperty* Property = NULL;
-	for (TFieldIterator<UProperty> It(RowStruct); It; ++It)
-	{
-		Property = *It;
-		check(Property != NULL);
-		if (PropertyName == Property->GetFName())
-		{
-			break;
-		}
-	}
-	if (!IsSupportedTableProperty(Property))
-	{
-		Property = NULL;
-	}
-
-	return Property;
-}
-
 /** Get array of UProperties that corresponds to columns in the table */
 TArray<UProperty*> UDataTable::GetTablePropertyArray(const FString& FirstRowString, UStruct* InRowStruct, TArray<FString>& OutProblems)
 {
@@ -378,11 +380,7 @@
 
 				for (TFieldIterator<UProperty> It(InRowStruct); It && !ColumnProp; ++It)
 				{
-#if WITH_EDITOR || HACK_HEADER_GENERATOR
 					const FString DisplayName = *It ? (*It)->GetMetaData(DisplayNameKey) : FString();
-#else
-					const FString DisplayName = TEXT("");
-#endif
 					ColumnProp = (!DisplayName.IsEmpty() && (DisplayName == ColumnNameStrings[ColIdx])) ? *It : NULL;
 				}
 
@@ -421,13 +419,9 @@
 	for(int32 PropIdx=0; PropIdx < ExpectedPropNames.Num(); PropIdx++)
 	{
 		const UProperty* const ColumnProp = FindField<UProperty>(InRowStruct, ExpectedPropNames[PropIdx]);
-#if WITH_EDITOR || HACK_HEADER_GENERATOR
 		const FString DisplayName = (ColumnProp && ColumnProp->HasMetaData(DisplayNameKey))
 			? ColumnProp->GetMetaData(DisplayNameKey) 
 			: ExpectedPropNames[PropIdx].ToString();
-#else
-		const FString DisplayName = ExpectedPropNames[PropIdx].ToString();
-#endif
 		OutProblems.Add(FString::Printf(TEXT("Expected column '%s' not found in input."), *DisplayName));
 	}
 
@@ -554,13 +548,9 @@
 	{
 		UProperty* Prop = *It;
 		check(Prop != NULL);
-#if WITH_EDITOR || HACK_HEADER_GENERATOR
 		const FString DisplayName = Prop->HasMetaData(DisplayNameKey)
 			? Prop->GetMetaData(DisplayNameKey)
 			: Prop->GetName();
-#else
-		const FString DisplayName = Prop->GetName();
-#endif
 		Result.Add(DisplayName);
 	}
 	return Result;
@@ -599,3 +589,4 @@
 
 }
 
+#endif //WITH_EDITOR || HACK_HEADER_GENERATOR
diff --git a/Engine/Source/Runtime/GameplayAbilities/Private/GameplayEffectTests.cpp b/Engine/Source/Runtime/GameplayAbilities/Private/GameplayEffectTests.cpp
--- a/Engine/Source/Runtime/GameplayAbilities/Private/GameplayEffectTests.cpp
+++ b/Engine/Source/Runtime/GameplayAbilities/Private/GameplayEffectTests.cpp
@@ -16,6 +16,8 @@ IMPLEMENT_SIMPLE_AUTOMATION_TEST(FGameplayEffectsTest, "AbilitySystem.GameplayEf
 
 #define SKILL_TEST_TEXT( Format, ... ) FString::Printf(TEXT("%s - %d: %s"), TEXT(__FILE__) , __LINE__ , *FString::Printf(TEXT(Format), ##__VA_ARGS__) )
 
+#if WITH_EDITOR
+
 void GameplayTest_TickWorld(UWorld *World, float Time)
 {
 	const float step = 0.1f;
@@ -6040,10 +6042,12 @@ bool GameplayEffectsTest_ModifyChanceToExecuteOnGE(UWorld *World, FAutomationTes
 	return true;
 }
 
-
+#endif //WITH_EDITOR
 
 bool FGameplayEffectsTest::RunTest( const FString& Parameters )
 {
+#if WITH_EDITOR
+
 	UCurveTable *CurveTable = IGameplayAbilitiesModule::Get().GetAbilitySystemGlobals()->GetGlobalCurveTable();
 	UDataTable *DataTable = IGameplayAbilitiesModule::Get().GetAbilitySystemGlobals()->GetGlobalAttributeDataTable();
 
@@ -6179,6 +6183,8 @@ bool FGameplayEffectsTest::RunTest( const FString& Parameters )
 
 	IGameplayAbilitiesModule::Get().GetAbilitySystemGlobals()->AutomationTestOnly_SetGlobalCurveTable(CurveTable);
 	IGameplayAbilitiesModule::Get().GetAbilitySystemGlobals()->AutomationTestOnly_SetGlobalAttributeDataTable(DataTable);
+
+#endif //WITH_EDITOR
 	return true;
 }
 
diff --git a/Engine/Source/Runtime/GameplayTags/Private/Tests/GameplayTagTests.cpp b/Engine/Source/Runtime/GameplayTags/Private/Tests/GameplayTagTests.cpp
--- a/Engine/Source/Runtime/GameplayTags/Private/Tests/GameplayTagTests.cpp
+++ b/Engine/Source/Runtime/GameplayTags/Private/Tests/GameplayTagTests.cpp
@@ -4,6 +4,7 @@
 
 IMPLEMENT_SIMPLE_AUTOMATION_TEST(FGameplayTagTest, "GameplayTags.GameplayTag", EAutomationTestFlags::ATF_Editor)
 
+#if WITH_EDITOR
 static UDataTable* CreateGameplayDataTable()
 {
 	FString CSV(TEXT(",Tag,CategoryText,\r\n0,GameplayTagTest.Test1\r\n1,GameplayTagTest.Test2"));
@@ -19,6 +20,7 @@ static UDataTable* CreateGameplayDataTable()
 	}
 	return DataTable;
 }
+#endif //WITH_EDITOR 
 
 bool GameplayTagTest_SimpleTest()
 {
@@ -29,6 +31,7 @@ bool GameplayTagTest_SimpleTest()
 
 bool FGameplayTagTest::RunTest(const FString& Parameters)
 {
+#if WITH_EDITOR
 	// Create Test Data 
 	UDataTable* DataTable = CreateGameplayDataTable();
 
@@ -40,4 +43,7 @@ bool FGameplayTagTest::RunTest(const FString& Parameters)
 	// Add more tests here... 
 
 	return bSuccess;
+#else
+	return true;
+#endif //WITH_EDITOR
 }
