Fix lighting of landscapes with a negative x or y scale

diff --git a/Engine/Source/Programs/UnrealLightmass/Private/Lighting/Landscape.cpp b/Engine/Source/Programs/UnrealLightmass/Private/Lighting/Landscape.cpp
--- a/Engine/Source/Programs/UnrealLightmass/Private/Lighting/Landscape.cpp
+++ b/Engine/Source/Programs/UnrealLightmass/Private/Lighting/Landscape.cpp
@@ -42,6 +42,11 @@ namespace Lightmass
 			OutI2 = (QuadX + 1) + (QuadY + 1) * NumVertices;
 			break;
 		}
+
+		if (bReverseWinding)
+		{
+			Swap(OutI1, OutI2);
+		}
 	}
 
 	// from FStaticLightMesh....
@@ -116,6 +121,7 @@ namespace Lightmass
 		NumVertices = ComponentSizeQuads + 2*ExpandQuadsX + 1;
 		NumQuads = NumVertices - 1;
 		UVFactor = LightMapRatio / NumVertices;
+		bReverseWinding = (LocalToWorld.RotDeterminant() < 0.0f);
 	}
 
 	void FLandscapeStaticLightingTextureMapping::Import( class FLightmassImporter& Importer )
diff --git a/Engine/Source/Programs/UnrealLightmass/Private/Lighting/Landscape.h b/Engine/Source/Programs/UnrealLightmass/Private/Lighting/Landscape.h
--- a/Engine/Source/Programs/UnrealLightmass/Private/Lighting/Landscape.h
+++ b/Engine/Source/Programs/UnrealLightmass/Private/Lighting/Landscape.h
@@ -29,6 +29,7 @@ namespace Lightmass
 		int32 NumVertices;
 		int32 NumQuads;
 		float UVFactor;
+		bool bReverseWinding;
 	};
 
 	/** Represents a landscape primitive with texture mapped static lighting. */
diff --git a/Engine/Source/Runtime/Engine/Private/Landscape/LandscapeLight.cpp b/Engine/Source/Runtime/Engine/Private/Landscape/LandscapeLight.cpp
--- a/Engine/Source/Runtime/Engine/Private/Landscape/LandscapeLight.cpp
+++ b/Engine/Source/Runtime/Engine/Private/Landscape/LandscapeLight.cpp
@@ -99,23 +99,23 @@ void FLandscapeStaticLightingTextureMapping::Apply(FQuantizedLightmapData* Quant
 
 /** Initialization constructor. */
 FLandscapeStaticLightingMesh::FLandscapeStaticLightingMesh(ULandscapeComponent* InComponent, const TArray<ULightComponent*>& InRelevantLights, int32 InExpandQuadsX, int32 InExpandQuadsY, float InLightMapRatio, int32 InLOD)
-:	FStaticLightingMesh(
-					FMath::Square(((InComponent->ComponentSizeQuads + 1) >> InLOD) - 1 + 2*InExpandQuadsX) * 2,
-					FMath::Square(((InComponent->ComponentSizeQuads + 1) >> InLOD) - 1 + 2*InExpandQuadsX) * 2,
-					FMath::Square(((InComponent->ComponentSizeQuads + 1) >> InLOD) + 2*InExpandQuadsX),
-					FMath::Square(((InComponent->ComponentSizeQuads + 1) >> InLOD) + 2*InExpandQuadsX),
-					0,
-					!!(InComponent->CastShadow | InComponent->bCastHiddenShadow),
-					false,
-					InRelevantLights,
-					InComponent,
-					InComponent->Bounds.GetBox(),
-					InComponent->GetLightingGuid()
-					)
-,	LandscapeComponent(InComponent)
-,	LightMapRatio(InLightMapRatio)
-,	ExpandQuadsX(InExpandQuadsX)
-,	ExpandQuadsY(InExpandQuadsY)
+	: FStaticLightingMesh(
+		FMath::Square(((InComponent->ComponentSizeQuads + 1) >> InLOD) - 1 + 2 * InExpandQuadsX) * 2,
+		FMath::Square(((InComponent->ComponentSizeQuads + 1) >> InLOD) - 1 + 2 * InExpandQuadsX) * 2,
+		FMath::Square(((InComponent->ComponentSizeQuads + 1) >> InLOD) + 2 * InExpandQuadsX),
+		FMath::Square(((InComponent->ComponentSizeQuads + 1) >> InLOD) + 2 * InExpandQuadsX),
+		0,
+		!!(InComponent->CastShadow | InComponent->bCastHiddenShadow),
+		false,
+		InRelevantLights,
+		InComponent,
+		InComponent->Bounds.GetBox(),
+		InComponent->GetLightingGuid()
+	)
+	, LandscapeComponent(InComponent)
+	, LightMapRatio(InLightMapRatio)
+	, ExpandQuadsX(InExpandQuadsX)
+	, ExpandQuadsY(InExpandQuadsY)
 {
 	const float LODScale = (float)InComponent->ComponentSizeQuads / (((InComponent->ComponentSizeQuads + 1) >> InLOD) - 1);
 	LocalToWorld = FTransform(FQuat::Identity, FVector::ZeroVector, FVector(LODScale, LODScale, 1)) * InComponent->ComponentToWorld;
@@ -123,6 +123,7 @@ FLandscapeStaticLightingMesh::FLandscapeStaticLightingMesh(ULandscapeComponent*
 	NumVertices = ComponentSizeQuads + 2*InExpandQuadsX + 1;
 	NumQuads = NumVertices - 1;
 	UVFactor = LightMapRatio / NumVertices;
+	bReverseWinding = (LocalToWorld.GetDeterminant() < 0.0f);
 
 	int32 GeometricLOD = FMath::Clamp<int32>(InComponent->ForcedLOD >= 0 ? InComponent->ForcedLOD : InComponent->LODBias, 0, FMath::CeilLogTwo(InComponent->SubsectionSizeQuads+1) - 1);
 	GetHeightmapData(InLOD, FMath::Max(GeometricLOD, InLOD));
@@ -546,6 +547,11 @@ void FLandscapeStaticLightingMesh::GetTriangleIndices(int32 TriangleIndex,int32&
 		OutI2 = (QuadX + 1) + (QuadY + 1) * NumVertices;
 		break;
 	}
+
+	if (bReverseWinding)
+	{
+		Swap(OutI1, OutI2);
+	}
 }
 
 
diff --git a/Engine/Source/Runtime/Engine/Public/Landscape/LandscapeLight.h b/Engine/Source/Runtime/Engine/Public/Landscape/LandscapeLight.h
--- a/Engine/Source/Runtime/Engine/Public/Landscape/LandscapeLight.h
+++ b/Engine/Source/Runtime/Engine/Public/Landscape/LandscapeLight.h
@@ -82,6 +82,7 @@ protected:
 	int32 NumVertices;
 	int32 NumQuads;
 	float UVFactor;
+	bool bReverseWinding;
 
 	friend class FLightmassExporter;
 
