Made the "code" property for the custom material node use a multi-line text box

TTP# 341961 - Use multi-line editable text for Custom node in the material editor

diff --git a/Engine/Source/Runtime/Engine/Classes/Materials/MaterialExpressionCustom.h b/Engine/Source/Runtime/Engine/Classes/Materials/MaterialExpressionCustom.h
--- a/Engine/Source/Runtime/Engine/Classes/Materials/MaterialExpressionCustom.h
+++ b/Engine/Source/Runtime/Engine/Classes/Materials/MaterialExpressionCustom.h
@@ -33,7 +33,7 @@ class UMaterialExpressionCustom : public UMaterialExpression
 {
 	GENERATED_UCLASS_BODY()
 
-	UPROPERTY(EditAnywhere, Category=MaterialExpressionCustom)
+	UPROPERTY(EditAnywhere, Category=MaterialExpressionCustom, meta=(MultiLine=true))
 	FString Code;
 
 	UPROPERTY(EditAnywhere, Category=MaterialExpressionCustom)
