Bone visibility : Expose HideBone/UnhideBone to BP accessible

diff --git a/Engine/Source/Runtime/Engine/Classes/Components/SkinnedMeshComponent.h b/Engine/Source/Runtime/Engine/Classes/Components/SkinnedMeshComponent.h
--- a/Engine/Source/Runtime/Engine/Classes/Components/SkinnedMeshComponent.h
+++ b/Engine/Source/Runtime/Engine/Classes/Components/SkinnedMeshComponent.h
@@ -808,6 +808,7 @@ public:
 	 *	@param  BoneName            Name of bone to hide
 	 *	@param	PhysBodyOption		Option for physics bodies that attach to the bones to be hidden
 	 */
+	UFUNCTION(BlueprintCallable, Category="Components|SkinnedMesh")
 	void HideBoneByName( FName BoneName, EPhysBodyOp PhysBodyOption );
 
 	/**
@@ -815,8 +816,19 @@ public:
 	 *	Compoared to HideBone By Index - This keeps track of list of bones and update when LOD changes
 	 *	@param  BoneName            Name of bone to unhide
 	 */
+	UFUNCTION(BlueprintCallable, Category="Components|SkinnedMesh")
 	void UnHideBoneByName( FName BoneName );
 
+	/** 
+	 *	Determines if the specified bone is hidden. 
+	 *
+	 *	@param  BoneName            Name of bone to check
+	 *
+	 *	@return true if hidden
+	 */
+	UFUNCTION(BlueprintCallable, Category="Components|SkinnedMesh")
+	bool IsBoneHiddenByName( FName BoneName );
+
 	/**
 	 *  Show/Hide Material - technical correct name for this is Section, but seems Material is mostly used
 	 *  This disable rendering of certain Material ID (Section)
diff --git a/Engine/Source/Runtime/Engine/Private/SkinnedMeshComponent.cpp b/Engine/Source/Runtime/Engine/Private/SkinnedMeshComponent.cpp
--- a/Engine/Source/Runtime/Engine/Private/SkinnedMeshComponent.cpp
+++ b/Engine/Source/Runtime/Engine/Private/SkinnedMeshComponent.cpp
@@ -1568,6 +1568,18 @@ bool USkinnedMeshComponent::IsBoneHidden( int32 BoneIndex ) const
 }
 
 
+bool USkinnedMeshComponent::IsBoneHiddenByName( FName BoneName )
+{
+	// Find appropriate BoneIndex
+	int32 BoneIndex = GetBoneIndex(BoneName);
+	if(BoneIndex != INDEX_NONE)
+	{
+		return IsBoneHidden(BoneIndex);
+	}
+
+	return false;
+}
+
 void USkinnedMeshComponent::HideBoneByName( FName BoneName, EPhysBodyOp PhysBodyOption )
 {
 	// Find appropriate BoneIndex
