Fix split wildcard pins being unsplit after reconstruction

diff --git a/Engine/Source/Runtime/Engine/Private/EdGraph/EdGraphPin.cpp b/Engine/Source/Runtime/Engine/Private/EdGraph/EdGraphPin.cpp
--- a/Engine/Source/Runtime/Engine/Private/EdGraph/EdGraphPin.cpp
+++ b/Engine/Source/Runtime/Engine/Private/EdGraph/EdGraphPin.cpp
@@ -157,6 +157,15 @@ void UEdGraphPin::CopyPersistentDataFromOldPin(const UEdGraphPin& SourcePin)
 		DefaultTextValue = SourcePin.DefaultTextValue;
 	}
 
+	// In K2 schemas the wildcard pins need to have their type copied before we get to pin splitting
+	// TODO: Better less hacky way of this?
+	if (PinType.PinCategory == TEXT("wildcard"))
+	{
+		PinType.PinCategory = SourcePin.PinType.PinCategory;
+		PinType.PinSubCategory = SourcePin.PinType.PinSubCategory;
+		PinType.PinSubCategoryObject = SourcePin.PinType.PinSubCategoryObject;
+	}
+
 	// Copy the links
 	for (int32 LinkIndex = 0; LinkIndex < SourcePin.LinkedTo.Num(); ++LinkIndex)
 	{
