Fix crash reloading Slate textures

diff --git a/Engine/Source/Runtime/Engine/Private/Slate/SlateTextures.cpp b/Engine/Source/Runtime/Engine/Private/Slate/SlateTextures.cpp
--- a/Engine/Source/Runtime/Engine/Private/Slate/SlateTextures.cpp
+++ b/Engine/Source/Runtime/Engine/Private/Slate/SlateTextures.cpp
@@ -11,6 +11,7 @@ FSlateTexture2DRHIRef::FSlateTexture2DRHIRef( FTexture2DRHIRef InRef, uint32 InW
 	: TSlateTexture( InRef )
 	, Width( InWidth )
 	, Height( InHeight )
+	, PixelFormat( PF_Unknown )
 	, bCreateEmptyTexture( false )
 {
 
