LERP : If Shortest Path is on, use Slerp to interpolate.

diff --git a/Engine/Source/Runtime/Engine/Private/KismetMathLibrary.cpp b/Engine/Source/Runtime/Engine/Private/KismetMathLibrary.cpp
--- a/Engine/Source/Runtime/Engine/Private/KismetMathLibrary.cpp
+++ b/Engine/Source/Runtime/Engine/Private/KismetMathLibrary.cpp
@@ -829,9 +829,16 @@ FRotator UKismetMathLibrary::RLerp(FRotator A, FRotator B, float Alpha, bool bSh
 {
 	FRotator DeltaAngle = B - A;
 
+	// if shortest path, we use Quaternion to interpolate instead of using FRotator
 	if( bShortestPath )
 	{
-		DeltaAngle = DeltaAngle.GetNormalized();
+		FQuat AQuat(A);
+		FQuat BQuat(B);
+
+		FQuat Result = FQuat::Slerp(AQuat, BQuat, Alpha);
+		Result.Normalize();
+
+		return Result.Rotator();
 	}
 
 	return A + Alpha*DeltaAngle;
