Fixing ParticleDirection for mesh emitters with orbit modules

diff --git a/Engine/Source/Runtime/Engine/Private/Particles/ParticleSystemRender.cpp b/Engine/Source/Runtime/Engine/Private/Particles/ParticleSystemRender.cpp
--- a/Engine/Source/Runtime/Engine/Private/Particles/ParticleSystemRender.cpp
+++ b/Engine/Source/Runtime/Engine/Private/Particles/ParticleSystemRender.cpp
@@ -2645,7 +2645,15 @@ void FDynamicMeshEmitterData::GetInstanceData(void* InstanceData, void* DynamicP
 		// Particle velocity. Calculate on CPU to avoid computing in vertex shader.
 		// Note: It would be preferred if we could check whether the material makes use of the 'Particle Direction' node to avoid this work.
 		FVector DeltaPosition = Particle.Location - Particle.OldLocation;
-		if(!DeltaPosition.IsZero())
+
+		int32 CurrentOffset = Source.OrbitModuleOffset;
+		if (CurrentOffset != 0)
+		{
+			FOrbitChainModuleInstancePayload& OrbitPayload = *((FOrbitChainModuleInstancePayload*)((uint8*)&Particle + CurrentOffset));																\
+			DeltaPosition = (Particle.Location + OrbitPayload.Offset) - (Particle.OldLocation + OrbitPayload.PreviousOffset);
+		}
+
+		if (!DeltaPosition.IsZero())
 		{
 			if (Source.bUseLocalSpace)
 			{
