#UE4

- More getters off of Local Player Context

diff --git a/Engine/Source/Runtime/Engine/Classes/Engine/LocalPlayer.h b/Engine/Source/Runtime/Engine/Classes/Engine/LocalPlayer.h
--- a/Engine/Source/Runtime/Engine/Classes/Engine/LocalPlayer.h
+++ b/Engine/Source/Runtime/Engine/Classes/Engine/LocalPlayer.h
@@ -15,6 +15,12 @@ struct ENGINE_API FLocalPlayerContext
 	FLocalPlayerContext(const class APlayerController* InPlayerController);
 	FLocalPlayerContext(const FLocalPlayerContext& InPlayerContext);
 
+	/** Is this context initialized and still valid? */
+	bool IsValid() const;
+
+	/** Is this context initialized */
+	bool IsInitialized() const;
+
 	/* Returns the world context. */
 	UWorld* GetWorld() const;
 
@@ -31,17 +37,24 @@ struct ENGINE_API FLocalPlayerContext
 		return CastChecked<T>(GetPlayerController(), ECastCheckedType::NullAllowed);
 	}
 
-	/** Is this context initialized and still valid? */
-	bool IsValid() const;
-
-	/** Is this context initialized */
-	bool IsInitialized() const;
+	/** Getter for the Game State */
+	class AGameState* GetGameState() const;
 
 	/** Templated Getter for the Game State */
 	template<class T>
 	FORCEINLINE T* GetGameState() const
 	{
-		return CastChecked<T>(GetWorld()->GameState, ECastCheckedType::NullAllowed);
+		return CastChecked<T>(GetGameState(), ECastCheckedType::NullAllowed);
+	}
+
+	/** Getter for the Player State */
+	class APlayerState* GetPlayerState() const;
+
+	/** Templated Getter for the Player State */
+	template<class T>
+	FORCEINLINE T* GetPlayerState() const
+	{
+		return CastChecked<T>(GetPlayerState(), ECastCheckedType::NullAllowed);
 	}
 
 private:	
diff --git a/Engine/Source/Runtime/Engine/Private/LocalPlayer.cpp b/Engine/Source/Runtime/Engine/Private/LocalPlayer.cpp
--- a/Engine/Source/Runtime/Engine/Private/LocalPlayer.cpp
+++ b/Engine/Source/Runtime/Engine/Private/LocalPlayer.cpp
@@ -68,6 +68,16 @@ FLocalPlayerContext::FLocalPlayerContext( const FLocalPlayerContext& InPlayerCon
 	SetLocalPlayer(InPlayerContext.GetLocalPlayer());
 }
 
+bool FLocalPlayerContext::IsValid() const
+{
+	return LocalPlayer.IsValid() && GetWorld() && GetPlayerController() && GetLocalPlayer();
+}
+
+bool FLocalPlayerContext::IsInitialized() const
+{
+	return LocalPlayer.IsValid();
+}
+
 UWorld* FLocalPlayerContext::GetWorld() const
 {
 	check( LocalPlayer.IsValid() );
@@ -86,14 +96,18 @@ APlayerController* FLocalPlayerContext::GetPlayerController() const
 	return LocalPlayer->PlayerController;
 }
 
-bool FLocalPlayerContext::IsValid() const
+AGameState* FLocalPlayerContext::GetGameState() const
 {
-	return LocalPlayer.IsValid() && GetWorld() && GetPlayerController() && GetLocalPlayer();
+	check(LocalPlayer.IsValid());
+	UWorld* World = LocalPlayer->GetWorld();
+	return World ? World->GameState : NULL;
 }
 
-bool FLocalPlayerContext::IsInitialized() const
+APlayerState* FLocalPlayerContext::GetPlayerState() const
 {
-	return LocalPlayer.IsValid();
+	check(LocalPlayer.IsValid());
+	APlayerController* PC = LocalPlayer->PlayerController;
+	return PC ? PC->PlayerState : NULL;
 }
 
 void FLocalPlayerContext::SetLocalPlayer( const ULocalPlayer* InLocalPlayer )
