#ue4 - Multiple fixes to Character falling movement.

- Improve velocity calculations for Characters when falling and deflecting off surfaces. Velocity is in immediate deflected direction, not averaged over a few frames to get there.
- Fix "TestWalk/AirControl" code not actually doing the test because the SweepSingle() result was ignored.
- Fix second impact during falling ignoring result of HandleSlopeBoosting().
- ApplyVelocityBraking() now only clamps velocity if there is some BrakingDeceleration being applied. Fixes small lateral movements being clamped to zero in falling mode.

diff --git a/Engine/Source/Runtime/Engine/Private/CharacterMovementComponent.cpp b/Engine/Source/Runtime/Engine/Private/CharacterMovementComponent.cpp
--- a/Engine/Source/Runtime/Engine/Private/CharacterMovementComponent.cpp
+++ b/Engine/Source/Runtime/Engine/Private/CharacterMovementComponent.cpp
@@ -2428,6 +2428,13 @@ void UCharacterMovementComponent::ApplyVelocityBraking(float DeltaTime, float Fr
 	}
 
 	Friction = FMath::Max(0.f, Friction);
+	BrakingDeceleration = FMath::Max(0.f, Friction);
+
+	if (Friction == 0.f && BrakingDeceleration == 0.f)
+	{
+		return;
+	}
+
 	const FVector OldVel = Velocity;
 
 	// subdivide braking to get reasonably consistent results at lower frame rates
@@ -2436,7 +2443,7 @@ void UCharacterMovementComponent::ApplyVelocityBraking(float DeltaTime, float Fr
 	const float TimeStep = (1.0f / 33.0f);
 
 	// Decelerate to brake to a stop
-	const FVector RevAccel = (-1.f * FMath::Max(0.f,BrakingDeceleration)) * SafeNormalPrecise(Velocity);
+	const FVector RevAccel = (-1.f * BrakingDeceleration) * SafeNormalPrecise(Velocity);
 	while( RemainingTime >= MIN_TICK_TIME )
 	{
 		const float dt = ((RemainingTime > TimeStep) ? TimeStep : RemainingTime);
@@ -2453,8 +2460,9 @@ void UCharacterMovementComponent::ApplyVelocityBraking(float DeltaTime, float Fr
 		}
 	}
 
-	// Clamp to zero if below min threshold.
-	if ((Velocity.SizeSquared() <= FMath::Square(BRAKE_TO_STOP_VELOCITY)) )
+	// Clamp to zero if nearly zero, or if below min threshold and braking.
+	const float VSizeSq = Velocity.SizeSquared();
+	if (VSizeSq <= KINDA_SMALL_NUMBER || (BrakingDeceleration > 0.f && VSizeSq <= FMath::Square(BRAKE_TO_STOP_VELOCITY)))
 	{
 		Velocity = FVector::ZeroVector;
 	}
@@ -2717,7 +2725,7 @@ void UCharacterMovementComponent::PhysFalling(float deltaTime, int32 Iterations)
 
 	// Bound final 2d portion of velocity
 	const float Speed2d = Velocity.Size2D();
-	const float BoundSpeed = FMath::Max(Speed2d, GetMaxSpeed() * AnalogInputModifier);
+	const float BoundSpeed2d = FMath::Max(Speed2d, GetMaxSpeed() * AnalogInputModifier);
 
 	//bound acceleration, falling object has minimal ability to impact acceleration
 	FVector FallAcceleration = Acceleration;
@@ -2747,11 +2755,11 @@ void UCharacterMovementComponent::PhysFalling(float deltaTime, int32 Iterations)
 				if(bUseFlatBaseForFloorChecks)
 				{
 					const FCollisionShape BoxShape = FCollisionShape::MakeBox(FVector(CapsuleShape.GetCapsuleRadius(), CapsuleShape.GetCapsuleRadius(), CapsuleShape.GetCapsuleHalfHeight()));
-					GetWorld()->SweepSingle( Result, PawnLocation, PawnLocation + TestWalk, FQuat::Identity, CollisionChannel, BoxShape, CapsuleQuery, ResponseParam );
+					bHit = GetWorld()->SweepSingle( Result, PawnLocation, PawnLocation + TestWalk, FQuat::Identity, CollisionChannel, BoxShape, CapsuleQuery, ResponseParam );
 				}
 				else
 				{
-					GetWorld()->SweepSingle( Result, PawnLocation, PawnLocation + TestWalk, FQuat::Identity, CollisionChannel, CapsuleShape, CapsuleQuery, ResponseParam );
+					bHit = GetWorld()->SweepSingle( Result, PawnLocation, PawnLocation + TestWalk, FQuat::Identity, CollisionChannel, CapsuleShape, CapsuleQuery, ResponseParam );
 				}
 
 				if (bHit)
@@ -2800,6 +2808,9 @@ void UCharacterMovementComponent::PhysFalling(float deltaTime, int32 Iterations)
 			Velocity.Z = 0.f;
 			CalcVelocity(timeTick, FallingLateralFriction, false, BrakingDecelerationFalling);
 			Velocity.Z = SavedVelZ;
+
+			// make sure we're not exceeding acceptable speed
+			Velocity = Velocity.ClampMaxSize2D(BoundSpeed2d);
 		}
 
 		// Apply gravity
@@ -2812,12 +2823,7 @@ void UCharacterMovementComponent::PhysFalling(float deltaTime, int32 Iterations)
 			NotifyJumpApex();
 		}
 
-		if( !HasRootMotion() )
-		{
-			// make sure not exceeding acceptable speed
-			Velocity = Velocity.ClampMaxSize2D(BoundSpeed);
-		}
-
+		// Move
 		FVector Adjusted = 0.5f*(OldVelocity + Velocity) * timeTick;  
 		SafeMoveUpdatedComponent( Adjusted, PawnRotation, true, Hit);
 		
@@ -2826,9 +2832,10 @@ void UCharacterMovementComponent::PhysFalling(float deltaTime, int32 Iterations)
 			return;
 		}
 		
+		float subTimeTickRemaining = timeTick * (1.f - Hit.Time);
 		if ( IsSwimming() ) //just entered water
 		{
-			remainingTime += timeTick * (1.f - Hit.Time);
+			remainingTime += subTimeTickRemaining;
 			StartSwimming(OldLocation, OldVelocity, timeTick, remainingTime, Iterations);
 			return;
 		}
@@ -2836,12 +2843,16 @@ void UCharacterMovementComponent::PhysFalling(float deltaTime, int32 Iterations)
 		{
 			if (IsValidLandingSpot(UpdatedComponent->GetComponentLocation(), Hit))
 			{
-				remainingTime += timeTick * (1.f - Hit.Time);
+				remainingTime += subTimeTickRemaining;
 				ProcessLanded(Hit, remainingTime, Iterations);
 				return;
 			}
 			else
 			{
+				// Compute impact deflection based on final velocity, not integration step.
+				// This allows us to compute a new velocity from the deflected vector, and ensures the full gravity effect is included in the slide result.
+				Adjusted = Velocity * timeTick;
+
 				// See if we can convert a normally invalid landing spot (based on the hit result) to a usable one.
 				if (!Hit.bStartPenetrating && ShouldCheckForValidLandingSpot(timeTick, Adjusted, Hit))
 				{
@@ -2850,7 +2861,7 @@ void UCharacterMovementComponent::PhysFalling(float deltaTime, int32 Iterations)
 					FindFloor(PawnLocation, FloorResult, false);
 					if (FloorResult.IsWalkableFloor() && IsValidLandingSpot(PawnLocation, FloorResult.HitResult))
 					{
-						remainingTime += timeTick * (1.f - Hit.Time);
+						remainingTime += subTimeTickRemaining;
 						ProcessLanded(FloorResult.HitResult, remainingTime, Iterations);
 						return;
 					}
@@ -2864,24 +2875,35 @@ void UCharacterMovementComponent::PhysFalling(float deltaTime, int32 Iterations)
 					return;
 				}
 
-				const float FirstHitPercent = Hit.Time;
 				const FVector OldHitNormal = Hit.Normal;
-				const FVector OldHitImpactNormal = Hit.ImpactNormal;
-				FVector Delta = ComputeSlideVector(Adjusted, 1.f - FirstHitPercent, OldHitNormal, Hit);
+				const FVector OldHitImpactNormal = Hit.ImpactNormal;				
+				FVector Delta = ComputeSlideVector(Adjusted, 1.f - Hit.Time, OldHitNormal, Hit);
+
+				// Compute velocity after deflection (only gravity component for RootMotion)
+				if (subTimeTickRemaining > KINDA_SMALL_NUMBER)
+				{
+					const FVector NewVelocity = (Delta / subTimeTickRemaining);
+					Velocity = HasRootMotion() ? FVector(Velocity.X, Velocity.Y, NewVelocity.Z) : NewVelocity;
+				}
 
 				if ((Delta | Adjusted) > 0.f)
 				{
+					// Move in deflected direction.
 					SafeMoveUpdatedComponent( Delta, PawnRotation, true, Hit);
-					if (Hit.Time < 1.f) //hit second wall
+					
+					if (Hit.Time < 1.f)
 					{
+						// hit second wall
+						subTimeTickRemaining = subTimeTickRemaining * (1.f - Hit.Time);
+
 						if (IsValidLandingSpot(UpdatedComponent->GetComponentLocation(), Hit))
 						{
-							remainingTime = 0.f;
+							remainingTime += subTimeTickRemaining;
 							ProcessLanded(Hit, remainingTime, Iterations);
 							return;
 						}
 
-						HandleImpact(Hit, timeTick * (1.f - FirstHitPercent), Delta);
+						HandleImpact(Hit, subTimeTickRemaining, Delta);
 
 						// If we've changed physics mode, abort.
 						if (!HasValidData() || !IsFalling())
@@ -2893,7 +2915,14 @@ void UCharacterMovementComponent::PhysFalling(float deltaTime, int32 Iterations)
 						TwoWallAdjust(Delta, Hit, OldHitNormal);
 						if (Delta.Z > 0.f)
 						{
-							HandleSlopeBoosting(Delta, PreTwoWallDelta, 1.f - Hit.Time, Hit.Normal, Hit);
+							Delta = HandleSlopeBoosting(Delta, PreTwoWallDelta, 1.f - Hit.Time, Hit.Normal, Hit);
+						}
+
+						// Compute velocity after deflection (only gravity component for RootMotion)
+						if (subTimeTickRemaining > KINDA_SMALL_NUMBER)
+						{
+							const FVector NewVelocity = (Delta / subTimeTickRemaining);
+							Velocity = HasRootMotion() ? FVector(Velocity.X, Velocity.Y, NewVelocity.Z) : NewVelocity;
 						}
 
 						// bDitch=true means that pawn is straddling two slopes, neither of which he can stand on
@@ -2933,33 +2962,13 @@ void UCharacterMovementComponent::PhysFalling(float deltaTime, int32 Iterations)
 						}
 					}
 				}
-
-				// Calculate average velocity based on actual movement after considering collisions
-				if (!bJustTeleported)
-				{
-					// Use average velocity for XY movement (no acceleration except for air control in those axes), but want actual velocity in Z axis
-					const float OldVelZ = OldVelocity.Z;
-					OldVelocity = (CharacterOwner->GetActorLocation() - OldLocation)/timeTick;
-					OldVelocity.Z = OldVelZ;
-				}
 			}
 		}
 
-		if (!HasRootMotion() && !bJustTeleported && MovementMode != MOVE_None )
+		if (Velocity.SizeSquared2D() <= KINDA_SMALL_NUMBER * 10.f)
 		{
-			// refine the velocity by figuring out the average actual velocity over the tick, and then the final velocity.
-			// This particularly corrects for situations where level geometry affected the fall.
-			Velocity = (CharacterOwner->GetActorLocation() - OldLocation)/timeTick; //actual average velocity
-			if( (Velocity.Z < OldVelocity.Z) || (OldVelocity.Z >= 0.f) )
-			{
-				Velocity = 2.f*Velocity - OldVelocity; //end velocity has 2* accel of avg
-			}
-
-			if (Velocity.SizeSquared2D() <= KINDA_SMALL_NUMBER * 10.f)
-			{
-				Velocity.X = 0.f;
-				Velocity.Y = 0.f;
-			}
+			Velocity.X = 0.f;
+			Velocity.Y = 0.f;
 		}
 	}
 }
