Fixed typo in UWorld::DuplicateWorldForPIE

diff --git a/Engine/Source/Runtime/Engine/Private/World.cpp b/Engine/Source/Runtime/Engine/Private/World.cpp
--- a/Engine/Source/Runtime/Engine/Private/World.cpp
+++ b/Engine/Source/Runtime/Engine/Private/World.cpp
@@ -2100,7 +2100,7 @@ UWorld* UWorld::DuplicateWorldForPIE(const FString& PackageName, UWorld* OwningW
 	if ( !EditorLevelWorld )
 	{
 		EditorLevelWorld = UWorld::FollowWorldRedirectorInPackage(EditorLevelPackage);
-		if ( EditorLevelPackage )
+		if ( EditorLevelWorld )
 		{
 			EditorLevelPackage = EditorLevelWorld->GetOutermost();
 		}
