Xbox One Movie player, using the MF library to play .mp4 files. Currently only plays video stream, not audio.

diff --git a/Engine/Source/Runtime/Engine/Private/Texture2D.cpp b/Engine/Source/Runtime/Engine/Private/Texture2D.cpp
--- a/Engine/Source/Runtime/Engine/Private/Texture2D.cpp
+++ b/Engine/Source/Runtime/Engine/Private/Texture2D.cpp
@@ -1845,7 +1845,6 @@ void FTexture2DResource::LoadMipData()
 			{
 				int32 ActualMipIndex = MipIndex + PendingFirstMip;
 				const FTexture2DMipMap& MipMap = OwnerMips[ActualMipIndex];
-				int32 MipSize = CalcTextureMipMapSize(MipMap.SizeX, MipMap.SizeY, TextureRHI->GetFormat(), 0);
 
 				void* TheMipData = nullptr;
 				uint32 DestPitch;
