Added EQS Logging to Visual Logger #ue4

diff --git a/Engine/Source/Developer/GameplayDebugger/Classes/GameplayDebuggingComponent.h b/Engine/Source/Developer/GameplayDebugger/Classes/GameplayDebuggingComponent.h
--- a/Engine/Source/Developer/GameplayDebugger/Classes/GameplayDebuggingComponent.h
+++ b/Engine/Source/Developer/GameplayDebugger/Classes/GameplayDebuggingComponent.h
@@ -6,10 +6,9 @@
 
 #pragma once
 #include "Components/PrimitiveComponent.h"
-#include "DebugRenderSceneProxy.h"
-#include "GameplayDebuggingControllerComponent.h"
+#include "GameplayDebuggingTypes.h"
 #include "EnvironmentQuery/EQSQueryResultSourceInterface.h"
-//#include "EnvironmentQuery/EnvQueryTypes.h"
+#include "EnvironmentQuery/EnvQueryDebugHelpers.h"
 #include "GameplayDebuggingComponent.generated.h"
 
 #define WITH_EQS 1
@@ -32,34 +31,6 @@ namespace EDebugComponentMessage
 	};
 }
 
-namespace EQSDebug
-{
-	struct FItemData
-	{
-		FString Desc;
-		int32 ItemIdx;
-		float TotalScore;
-
-		TArray<float> TestValues;
-		TArray<float> TestScores;
-	};
-
-	struct FTestData
-	{
-		FString ShortName;
-		FString Detailed;
-	};
-
-	struct FQueryData
-	{
-		TArray<FItemData> Items;
-		TArray<FTestData> Tests;
-		TArray<FDebugRenderSceneProxy::FSphere> SolidSpheres;
-		TArray<FDebugRenderSceneProxy::FText3d> Texts;
-		int32 NumValidItems;
-	};
-}
-
 DECLARE_MULTICAST_DELEGATE_TwoParams(FOnDebuggingTargetChanged, class AActor* /*Owner of debugging component*/, bool /*is being debugged now*/);
 
 UCLASS(config=Engine)
@@ -124,7 +95,9 @@ class GAMEPLAYDEBUGGER_API UGameplayDebuggingComponent : public UPrimitiveCompon
 	TArray<uint8> EQSRepData;
 	
 	/** local EQS debug data, decoded from EQSRepData blob */
-	EQSDebug::FQueryData EQSLocalData;	
+#if  USE_EQS_DEBUGGER || ENABLE_VISUAL_LOG
+	EQSDebug::FQueryData EQSLocalData;
+#endif
 	/** End EQS replication data */
 
 
diff --git a/Engine/Source/Developer/GameplayDebugger/Private/GameplayDebuggingComponent.cpp b/Engine/Source/Developer/GameplayDebugger/Private/GameplayDebuggingComponent.cpp
--- a/Engine/Source/Developer/GameplayDebugger/Private/GameplayDebuggingComponent.cpp
+++ b/Engine/Source/Developer/GameplayDebugger/Private/GameplayDebuggingComponent.cpp
@@ -9,10 +9,8 @@
 #include "EnvironmentQuery/EQSRenderingComponent.h"
 #include "EnvironmentQuery/EnvQueryTest.h"
 
-// @todo this is here only due to circular dependency to AIModule. To be removed
 #include "AIController.h"
 #include "BrainComponent.h"
-//#include "EnvironmentQuery/EnvQueryTypes.h"
 #include "BehaviorTreeDelegates.h"
 #include "Navigation/NavigationComponent.h"
 #include "GameFramework/PlayerState.h"
@@ -564,51 +562,9 @@ const FEnvQueryInstance* UGameplayDebuggingComponent::GetQueryInstance() const
 	return CachedQueryInstance.Get();
 }
 
-FArchive& operator<<(FArchive& Ar, FDebugRenderSceneProxy::FSphere& Data)
-{
-	Ar << Data.Radius;
-	Ar << Data.Location;
-	Ar << Data.Color;
-	return Ar;
-}
-
-FArchive& operator<<(FArchive& Ar, FDebugRenderSceneProxy::FText3d& Data)
-{
-	Ar << Data.Text;
-	Ar << Data.Location;
-	Ar << Data.Color;
-	return Ar;
-}
-
-FArchive& operator<<(FArchive& Ar, EQSDebug::FItemData& Data)
-{
-	Ar << Data.Desc;
-	Ar << Data.ItemIdx;
-	Ar << Data.TotalScore;
-	Ar << Data.TestValues;
-	Ar << Data.TestScores;
-	return Ar;
-}
-
-FArchive& operator<<(FArchive& Ar, EQSDebug::FTestData& Data)
-{
-	Ar << Data.ShortName;
-	Ar << Data.Detailed;
-	return Ar;
-}
-
-FArchive& operator<<(FArchive& Ar, EQSDebug::FQueryData& Data)
-{
-	Ar << Data.Items;
-	Ar << Data.Tests;
-	Ar << Data.SolidSpheres;
-	Ar << Data.Texts;
-	Ar << Data.NumValidItems;
-	return Ar;
-}
-
 void UGameplayDebuggingComponent::OnRep_UpdateEQS()
 {
+#if  USE_EQS_DEBUGGER
 	// decode scoring data
 	if (World && World->GetNetMode() == NM_Client)
 	{
@@ -630,6 +586,7 @@ void UGameplayDebuggingComponent::OnRep_UpdateEQS()
 
 	UpdateBounds();
 	MarkRenderStateDirty();
+#endif //USE_EQS_DEBUGGER
 }
 
 void UGameplayDebuggingComponent::CollectEQSData()
@@ -665,58 +622,8 @@ void UGameplayDebuggingComponent::CollectEQSData()
 	TArray<uint8> UncompressedBuffer;
 	FMemoryWriter ArWriter(UncompressedBuffer);
 
-	// step 1: data for rendering component
-	EQSLocalData.SolidSpheres.Reset();
-	EQSLocalData.Texts.Reset();
-	FEQSSceneProxy::CollectEQSData(this, NULL, EQSLocalData.SolidSpheres, EQSLocalData.Texts);
-
-	// step 2: detailed scoring data for HUD
-	const int32 MaxDetailedItems = 10;
-	const int32 FirstItemIndex = 0;
+	UEnvQueryDebugHelpers::QueryToDebugData(CachedQueryInstance.Get(), EQSLocalData);
 
-	const int32 NumTests = CachedQueryInstance->ItemDetails.IsValidIndex(0) ? CachedQueryInstance->ItemDetails[0].TestResults.Num() : 0;
-	const int32 NumItems = FMath::Min(MaxDetailedItems, CachedQueryInstance->NumValidItems);
-
-	EQSLocalData.Items.Reset(NumItems);
-	EQSLocalData.Tests.Reset(NumTests);
-	EQSLocalData.NumValidItems = CachedQueryInstance->NumValidItems;
-
-	UEnvQueryItemType* ItemCDO = CachedQueryInstance->ItemType.GetDefaultObject();
-	for (int32 ItemIdx = 0; ItemIdx < NumItems; ItemIdx++)
-	{
-		EQSDebug::FItemData ItemInfo;
-		ItemInfo.ItemIdx = ItemIdx + FirstItemIndex;
-		ItemInfo.TotalScore = CachedQueryInstance->Items[ItemInfo.ItemIdx].Score;
-
-		const uint8* ItemData = CachedQueryInstance->RawData.GetTypedData() + CachedQueryInstance->Items[ItemInfo.ItemIdx].DataOffset;
-		ItemInfo.Desc = FString::Printf(TEXT("[%d] %s"), ItemInfo.ItemIdx, *ItemCDO->GetDescription(ItemData));
-
-		ItemInfo.TestValues.Reserve(NumTests);
-		ItemInfo.TestScores.Reserve(NumTests);
-		for (int32 TestIdx = 0; TestIdx < NumTests; TestIdx++)
-		{
-			const float ScoreN = CachedQueryInstance->ItemDetails[ItemInfo.ItemIdx].TestResults[TestIdx];
-			const float ScoreW = CachedQueryInstance->ItemDetails[ItemInfo.ItemIdx].TestWeightedScores[TestIdx];
-
-			ItemInfo.TestValues.Add(ScoreN);
-			ItemInfo.TestScores.Add(ScoreW);
-		}
-
-		EQSLocalData.Items.Add(ItemInfo);
-	}
-
-	{
-		for (int32 TestIdx = 0; TestIdx < NumTests; TestIdx++)
-		{
-			EQSDebug::FTestData TestInfo;
-
-			UEnvQueryTest* TestOb = CachedQueryInstance->Options[CachedQueryInstance->OptionIndex].GetTestObject(TestIdx);
-			TestInfo.ShortName = TestOb->GetDescriptionTitle();
-			TestInfo.Detailed = TestOb->GetDescriptionDetails().ToString().Replace(TEXT("\n"), TEXT(", "));
-			
-			EQSLocalData.Tests.Add(TestInfo);
-		}
-	}
 	ArWriter << EQSLocalData;
 
 	const double Timer2 = FPlatformTime::Seconds();
@@ -1160,13 +1067,13 @@ FPrimitiveSceneProxy* UGameplayDebuggingComponent::CreateSceneProxy()
 	}
 #endif
 
-#if WITH_EQS
+#if USE_EQS_DEBUGGER
 	if (ShouldReplicateData(EAIDebugDrawDataView::EQS) && IsClientEQSSceneProxyEnabled() )
 	{
 		CompositeProxy = CompositeProxy ? CompositeProxy : (new FDebugRenderSceneCompositeProxy(this));
 		CompositeProxy->AddChild(new FEQSSceneProxy(this, TEXT("GameplayDebug"), /*bDrawOnlyWhenSelected=*/false, EQSLocalData.SolidSpheres, EQSLocalData.Texts));
 	}
-#endif // WITH_EQS
+#endif // USE_EQS_DEBUGGER
 
 	return CompositeProxy;
 }
@@ -1183,12 +1090,12 @@ FBoxSphereBounds UGameplayDebuggingComponent::CalcBounds(const FTransform & Loca
 	}
 #endif
 
-#if WITH_EQS
+#if USE_EQS_DEBUGGER
 	if (EQSRepData.Num())
 	{
 		MyBounds = FBox(FVector(-HALF_WORLD_MAX, -HALF_WORLD_MAX, -HALF_WORLD_MAX), FVector(HALF_WORLD_MAX, HALF_WORLD_MAX, HALF_WORLD_MAX));
 	}
-#endif // WITH_EQS
+#endif // USE_EQS_DEBUGGER
 
 	return MyBounds;
 }
diff --git a/Engine/Source/Developer/GameplayDebugger/Private/GameplayDebuggingHUDComponent.cpp b/Engine/Source/Developer/GameplayDebugger/Private/GameplayDebuggingHUDComponent.cpp
--- a/Engine/Source/Developer/GameplayDebugger/Private/GameplayDebuggingHUDComponent.cpp
+++ b/Engine/Source/Developer/GameplayDebugger/Private/GameplayDebuggingHUDComponent.cpp
@@ -362,6 +362,7 @@ void AGameplayDebuggingHUDComponent::DrawEQSData(APlayerController* PC, class UG
 			));
 	}
 
+#if  USE_EQS_DEBUGGER
 	if (DebugComponent->EQSLocalData.NumValidItems > 0)
 	{
 		// draw test weights for best X items
@@ -418,6 +419,7 @@ void AGameplayDebuggingHUDComponent::DrawEQSData(APlayerController* PC, class UG
 			PrintString(DefaultContext, TestDesc);
 		}
 	}
+#endif //USE_EQS_DEBUGGER
 #endif //!(UE_BUILD_SHIPPING || UE_BUILD_TEST)
 }
 
diff --git a/Engine/Source/Developer/LogVisualizer/Private/LogVisualizer.cpp b/Engine/Source/Developer/LogVisualizer/Private/LogVisualizer.cpp
--- a/Engine/Source/Developer/LogVisualizer/Private/LogVisualizer.cpp
+++ b/Engine/Source/Developer/LogVisualizer/Private/LogVisualizer.cpp
@@ -2,6 +2,7 @@
 
 #include "LogVisualizerPCH.h"
 #include "CollisionDebugDrawingPublic.h"
+#include "LogVisualizerDebugActor.h"
 
 #if ENABLE_VISUAL_LOG
 //////////////////////////////////////////////////////////////////////////
@@ -52,6 +53,9 @@ void FLogVisualizer::CloseUI(UWorld* InWorld)
 	{
 		if (LogWindow.IsValid() && (World.IsValid() == false || World == InWorld))
 		{
+			InWorld->DestroyActor(GetHelperActor(InWorld));
+			DebugActor = NULL;
+
 			CleanUp();
 			FSlateApplication::Get().RequestDestroyWindow(LogWindow.Pin().ToSharedRef());
 		}
@@ -77,6 +81,32 @@ void FLogVisualizer::CleanUp()
 	FVisualLog::Get().ClearNewLogsObserver();
 }
 
+class AActor* FLogVisualizer::GetHelperActor(class UWorld* InWorld)
+{
+	UWorld* ActorWorld = DebugActor.IsValid() ? DebugActor->GetWorld() : NULL;
+	if (DebugActor.IsValid() && ActorWorld == InWorld)
+	{
+		return DebugActor.Get();
+	}
+
+	for (FActorIterator It(InWorld); It; ++It)
+	{
+		ALogVisualizerDebugActor* LogVisualizerDebugActor = Cast<ALogVisualizerDebugActor>(*It);
+		if (LogVisualizerDebugActor != NULL)
+		{
+			DebugActor = LogVisualizerDebugActor;
+			return LogVisualizerDebugActor;
+		}
+	}
+
+	FActorSpawnParameters SpawnInfo;
+	SpawnInfo.bNoCollisionFail = true;
+	SpawnInfo.Name = *FString::Printf(TEXT("LogVisualizerDebugActor"));
+	DebugActor = InWorld->SpawnActor<ALogVisualizerDebugActor>(SpawnInfo);
+
+	return DebugActor.Get();
+}
+
 void FLogVisualizer::PullDataFromVisualLog(const FVisualLog& VisualLog)
 {
 	Logs.Reset();
diff --git a/Engine/Source/Developer/LogVisualizer/Private/LogVisualizer.h b/Engine/Source/Developer/LogVisualizer/Private/LogVisualizer.h
--- a/Engine/Source/Developer/LogVisualizer/Private/LogVisualizer.h
+++ b/Engine/Source/Developer/LogVisualizer/Private/LogVisualizer.h
@@ -45,6 +45,8 @@ public:
 
 	UWorld* GetWorld() { return World.Get(); }
 
+	virtual class AActor* GetHelperActor(class UWorld* InWorld);
+
 protected:
 	void OnNewLog(const class AActor* Actor, TSharedPtr<FActorsVisLog> Log);
 	
@@ -57,6 +59,7 @@ private:
 	FVisLogsChangedEvent	LogAddedEvent;
 	TWeakObjectPtr<UWorld>	World;
 	TWeakPtr<SWindow>		LogWindow;
+	TWeakObjectPtr<class ALogVisualizerDebugActor> DebugActor;
 };
 
 #endif //ENABLE_VISUAL_LOG
\ No newline at end of file
diff --git a/Engine/Source/Developer/LogVisualizer/Private/LogVisualizerDebugActor.cpp b/Engine/Source/Developer/LogVisualizer/Private/LogVisualizerDebugActor.cpp
new file mode 100644
--- /dev/null
+++ b/Engine/Source/Developer/LogVisualizer/Private/LogVisualizerDebugActor.cpp
@@ -0,0 +1,15 @@
+// Copyright 1998-2014 Epic Games, Inc. All Rights Reserved.
+
+#include "LogVisualizerPCH.h"
+#include "Engine/GameInstance.h"
+#include "Debug/DebugDrawService.h"
+#include "GameFramework/HUD.h"
+#if WITH_EDITOR
+#include "Editor/EditorEngine.h"
+#endif // WITH_EDITOR
+#include "LogVisualizerDebugActor.h"
+
+ALogVisualizerDebugActor::ALogVisualizerDebugActor(const class FPostConstructInitializeProperties& PCIP)
+	: Super(PCIP)
+{
+}
diff --git a/Engine/Source/Developer/LogVisualizer/Private/LogVisualizerDebugActor.h b/Engine/Source/Developer/LogVisualizer/Private/LogVisualizerDebugActor.h
new file mode 100644
--- /dev/null
+++ b/Engine/Source/Developer/LogVisualizer/Private/LogVisualizerDebugActor.h
@@ -0,0 +1,17 @@
+// Copyright 1998-2014 Epic Games, Inc. All Rights Reserved.
+
+#pragma once
+#include "LogVisualizerDebugActor.generated.h"
+
+/**
+*	Transient actor used to draw visual logger data on level
+*/
+
+DECLARE_MULTICAST_DELEGATE_OneParam(FOnSelectionChanged, class AActor*);
+
+UCLASS(config = Engine, NotBlueprintable, Transient)
+class LOGVISUALIZER_API ALogVisualizerDebugActor : public AActor
+{
+	GENERATED_UCLASS_BODY()
+
+};
diff --git a/Engine/Source/Developer/LogVisualizer/Private/SLogVisualizer.cpp b/Engine/Source/Developer/LogVisualizer/Private/SLogVisualizer.cpp
--- a/Engine/Source/Developer/LogVisualizer/Private/SLogVisualizer.cpp
+++ b/Engine/Source/Developer/LogVisualizer/Private/SLogVisualizer.cpp
@@ -20,6 +20,7 @@
 #endif
 
 #include "GameplayDebuggingComponent.h"
+#include "LogVisualizerModule.h"
 
 #include "SFilterList.h"
 
@@ -769,6 +770,24 @@ void SLogVisualizer::UpdateVisibleEntriesCache(const TSharedPtr<FActorsVisLog>&
 					break;
 				}
 			}
+
+			if (bAddedEntry)
+			{
+				continue;
+			}
+
+			for (const auto CurrentBlock : CurrentEntry->DataBlocks)
+			{
+				const bool bCurrentCategoryPassed = FilterListPtr->IsFilterEnabled(CurrentBlock.Category.ToString(), ELogVerbosity::All);
+				const bool bCurrentTagNamePassed = FilterListPtr->IsFilterEnabled(CurrentBlock.TagName.ToString(), ELogVerbosity::All);
+
+				if (CurrentBlock.Category != NAME_None && (bCurrentCategoryPassed &&	bCurrentTagNamePassed))
+				{
+					OutEntriesCached[CachedLogIndex].CachedEntries.AddUnique(CurrentEntry);
+					bAddedEntry = true;
+					break;
+				}
+			}
 		}
 	}
 	else
@@ -862,6 +881,23 @@ void SLogVisualizer::GetVisibleEntries(const TSharedPtr<FActorsVisLog>& Log, TAr
 						break;
 					}
 				}
+				if (bAddedEntry)
+				{
+					continue;
+				}
+
+				for (const auto CurrentData : CurrentEntry->DataBlocks)
+				{
+					const bool bCurrentCategoryPassed = FilterListPtr->IsFilterEnabled(CurrentData.Category.ToString(), ELogVerbosity::All);
+					const bool bCurrentTagNamePassed = FilterListPtr->IsFilterEnabled(CurrentData.TagName.ToString(), ELogVerbosity::All);
+
+					if (CurrentData.TagName != NAME_None && (bCurrentCategoryPassed &&	bCurrentTagNamePassed))
+					{
+						OutEntriesCached[CachedLogIndex].CachedEntries.AddUnique(CurrentEntry);
+						bAddedEntry = true;
+						break;
+					}
+				}
 			}
 
 		}
@@ -936,6 +972,24 @@ void SLogVisualizer::GetVisibleEntries(const TSharedPtr<FActorsVisLog>& Log, TAr
 					break;
 				}
 			}
+			if (bAddedEntry)
+			{
+				continue;
+			}
+
+			for (const auto CurrentData : LogEntry->DataBlocks)
+			{
+				const bool bCurrentCategoryPassed = bHistogramGraphsFilter || (QuickFilterText.Len() == 0 || CurrentData.Category.ToString().Find(QuickFilterText) != INDEX_NONE);
+				const bool bCurrentTagNamePassed = !bHistogramGraphsFilter || (QuickFilterText.Len() == 0 || CurrentData.TagName.ToString().Find(QuickFilterText) != INDEX_NONE);
+
+				if (bCurrentCategoryPassed &&	bCurrentTagNamePassed)
+				{
+					OutEntries.AddUnique(LogEntry);
+					bAddedEntry = true;
+					break;
+				}
+			}
+
 		} //for (int32 Index = 0; Index < NumEntries; ++Index)
 	} //if (QuickFilterText.Len() > 0)
 	else
@@ -1136,6 +1190,21 @@ void SLogVisualizer::IncrementCurrentLogIndex(int32 IncrementBy)
 				}
 			}
 
+			if (!bShouldShow)
+			{
+				for (const auto CurrentData : Log->Entries[NewEntryIndex]->DataBlocks)
+				{
+					const FName CurrentCategory = CurrentData.Category;
+					const FName CurrentTagName = CurrentData.TagName;
+					if (CurrentCategory == NAME_None ||
+						(FilterListPtr->IsFilterEnabled(CurrentCategory.ToString(), ELogVerbosity::All) && FilterListPtr->IsFilterEnabled(CurrentTagName.ToString(), ELogVerbosity::All)))
+					{
+						bShouldShow = true;
+						break;
+					}
+				}
+			}
+
 			if (bShouldShow)
 			{
 				break;
@@ -1226,6 +1295,18 @@ void SLogVisualizer::AddOrUpdateLog(int32 LogIndex, const FActorsVisLog* Log)
 			const FString DataNameAsString = Log->Entries[EntryIndex]->HistogramSamples[SampleIndex].DataName.ToString();
 			FilterListPtr->AddGraphFilter(GraphNameAsString, DataNameAsString, FColor::White);
 		}
+
+		for (const auto CurrentData : Log->Entries[EntryIndex]->DataBlocks)
+		{
+			int32 Index = UsedCategories.Find(CurrentData.Category.ToString());
+			if (Index == INDEX_NONE)
+			{
+				Index = UsedCategories.Add(CurrentData.Category.ToString());
+				FilterListPtr->AddFilter(CurrentData.Category.ToString(), GetColorForUsedCategory(Index));
+			}
+		}
+
+
 	}
 
 	if (CurrentIndex == INDEX_NONE)
@@ -1671,6 +1752,12 @@ void SLogVisualizer::ShowEntry(const FVisLogEntry* LogEntry)
 
 	SetCurrentViewedTime(LogEntry->TimeStamp);
 
+	UWorld* World = GetWorld();
+	for (auto It = FVisualLog::Get().GetAllExtensions().CreateIterator(); It; ++It)
+	{
+		(*It).Value->OnTimestampChange(LogEntry->TimeStamp, World, FLogVisualizerModule::Get()->GetHelperActor(World));
+	}
+	InvalidateCanvas();
 	LogsLinesWidget->RequestListRefresh();
 }
 
@@ -2100,6 +2187,21 @@ void SLogVisualizer::DrawOnCanvas(UCanvas* Canvas, APlayerController*)
 				}
 			}
 
+			AActor* HelperActor = FLogVisualizerModule::Get()->GetHelperActor(World);
+			for (const auto CurrentData : Entry->DataBlocks)
+			{
+				const FName TagName = CurrentData.TagName;
+				const bool bIsValidByFilter = FilterListPtr->IsFilterEnabled(CurrentData.Category.ToString(), ELogVerbosity::All) && FilterListPtr->IsFilterEnabled(CurrentData.TagName.ToString(), ELogVerbosity::All);
+				if (!bIsValidByFilter)
+				{
+					continue;
+				}
+
+				if (FVisualLogExtensionInterface* Extension = FVisualLog::Get().GetExtensionForTag(TagName))
+				{
+					Extension->DrawData(World, Canvas, HelperActor, TagName, CurrentData, Entry->TimeStamp);
+				}
+			}
 
 			const FVisLogEntry::FElementToDraw* ElementToDraw = Entry->ElementsToDraw.GetTypedData();
 			const int32 ElementsCount = Entry->ElementsToDraw.Num();
diff --git a/Engine/Source/Developer/LogVisualizer/Public/ILogVisualizer.h b/Engine/Source/Developer/LogVisualizer/Public/ILogVisualizer.h
--- a/Engine/Source/Developer/LogVisualizer/Public/ILogVisualizer.h
+++ b/Engine/Source/Developer/LogVisualizer/Public/ILogVisualizer.h
@@ -16,4 +16,7 @@ public:
 
 	/** checks if we have LogVisualizer UI opened for given World */
 	virtual bool IsOpenUI(UWorld* InWorld) = 0;
+
+	/** Helper actor used to render data in given world */
+	virtual class AActor* GetHelperActor(class UWorld* InWorld) = 0;
 };
\ No newline at end of file
diff --git a/Engine/Source/Runtime/AIModule/Classes/EnvironmentQuery/EQSRenderingComponent.h b/Engine/Source/Runtime/AIModule/Classes/EnvironmentQuery/EQSRenderingComponent.h
--- a/Engine/Source/Runtime/AIModule/Classes/EnvironmentQuery/EQSRenderingComponent.h
+++ b/Engine/Source/Runtime/AIModule/Classes/EnvironmentQuery/EQSRenderingComponent.h
@@ -4,6 +4,7 @@
 #include "DebugRenderSceneProxy.h"
 #include "Debug/DebugDrawService.h"
 #include "EnvQueryTypes.h"
+#include "EnvQueryDebugHelpers.h"
 #include "EQSRenderingComponent.generated.h"
 
 class AIMODULE_API FEQSSceneProxy : public FDebugRenderSceneProxy
@@ -18,6 +19,7 @@ public:
 	virtual FPrimitiveViewRelevance GetViewRelevance(const FSceneView* View) override;
 
 	static void CollectEQSData(const UPrimitiveComponent* InComponent, const class IEQSQueryResultSourceInterface* QueryDataSource, TArray<FSphere>& Spheres, TArray<FText3d>& Texts);
+	static void CollectEQSData(const struct FEnvQueryResult* ResultItems, const struct FEnvQueryInstance* QueryInstance, TArray<FSphere>& Spheres, TArray<FText3d>& Texts, bool ShouldDrawFailedItems);
 private:
 	FEnvQueryResult QueryResult;	
 	// can be 0
@@ -38,6 +40,10 @@ class AIMODULE_API UEQSRenderingComponent : public UPrimitiveComponent
 	FString DrawFlagName;
 	bool bDrawOnlyWhenSelected;
 
+#if  USE_EQS_DEBUGGER || ENABLE_VISUAL_LOG
+	EQSDebug::FQueryData DebugData;
+#endif
+
 	virtual FPrimitiveSceneProxy* CreateSceneProxy() override;
 	virtual FBoxSphereBounds CalcBounds(const FTransform &LocalToWorld) const override;
 	virtual void CreateRenderState_Concurrent() override;
diff --git a/Engine/Source/Runtime/AIModule/Classes/EnvironmentQuery/EnvQueryDebugHelpers.h b/Engine/Source/Runtime/AIModule/Classes/EnvironmentQuery/EnvQueryDebugHelpers.h
new file mode 100644
--- /dev/null
+++ b/Engine/Source/Runtime/AIModule/Classes/EnvironmentQuery/EnvQueryDebugHelpers.h
@@ -0,0 +1,133 @@
+// Copyright 1998-2014 Epic Games, Inc. All Rights Reserved.
+
+#pragma once
+#include "DebugRenderSceneProxy.h"
+#include "VisualLoggerExtension.h"
+#include "EnvQueryDebugHelpers.generated.h"
+
+
+#if  USE_EQS_DEBUGGER || ENABLE_VISUAL_LOG
+namespace EQSDebug
+{
+	struct FItemData
+	{
+		FString Desc;
+		int32 ItemIdx;
+		float TotalScore;
+
+		TArray<float> TestValues;
+		TArray<float> TestScores;
+	};
+
+	struct FTestData
+	{
+		FString ShortName;
+		FString Detailed;
+	};
+
+	struct FQueryData
+	{
+		TArray<FItemData> Items;
+		TArray<FTestData> Tests;
+		TArray<FDebugRenderSceneProxy::FSphere> SolidSpheres;
+		TArray<FDebugRenderSceneProxy::FText3d> Texts;
+		int32 NumValidItems;
+		int32 Id;
+		FString Name;
+
+		void Reset()
+		{
+			NumValidItems = 0;
+			Id = INDEX_NONE;
+			Name.Empty();
+			Items.Reset();
+			Tests.Reset();
+			SolidSpheres.Reset();
+			Texts.Reset();
+		}
+	};
+}
+
+FORCEINLINE
+FArchive& operator<<(FArchive& Ar, FDebugRenderSceneProxy::FSphere& Data)
+{
+	Ar << Data.Radius;
+	Ar << Data.Location;
+	Ar << Data.Color;
+	return Ar;
+}
+
+FORCEINLINE
+FArchive& operator<<(FArchive& Ar, FDebugRenderSceneProxy::FText3d& Data)
+{
+	Ar << Data.Text;
+	Ar << Data.Location;
+	Ar << Data.Color;
+	return Ar;
+}
+
+FORCEINLINE
+FArchive& operator<<(FArchive& Ar, EQSDebug::FItemData& Data)
+{
+	Ar << Data.Desc;
+	Ar << Data.ItemIdx;
+	Ar << Data.TotalScore;
+	Ar << Data.TestValues;
+	Ar << Data.TestScores;
+	return Ar;
+}
+
+FORCEINLINE
+FArchive& operator<<(FArchive& Ar, EQSDebug::FTestData& Data)
+{
+	Ar << Data.ShortName;
+	Ar << Data.Detailed;
+	return Ar;
+}
+
+FORCEINLINE
+FArchive& operator<<(FArchive& Ar, EQSDebug::FQueryData& Data)
+{
+	Ar << Data.Items;
+	Ar << Data.Tests;
+	Ar << Data.SolidSpheres;
+	Ar << Data.Texts;
+	Ar << Data.NumValidItems;
+	Ar << Data.Id;
+	Ar << Data.Name;
+	return Ar;
+}
+
+#endif //USE_EQS_DEBUGGER || ENABLE_VISUAL_LOG
+
+#if ENABLE_VISUAL_LOG
+
+#define UE_VLOG_EQS(Query, CategoryName, Verbosity) \
+{ \
+	SCOPE_CYCLE_COUNTER(STAT_VisualLog); \
+	static_assert((ELogVerbosity::Verbosity & ELogVerbosity::VerbosityMask) < ELogVerbosity::NumVerbosity && ELogVerbosity::Verbosity > 0, "Verbosity must be constant and in range."); \
+	if (FVisualLog::Get().IsRecording() && (!FVisualLog::Get().IsAllBlocked() || FVisualLog::Get().InWhitelist(CategoryName.GetCategoryName())) && Cast<AActor>(Query->Owner.Get())) \
+	{ \
+		UE_VLOG(Cast<AActor>(Query->Owner.Get()), CategoryName, Verbosity, TEXT("Executed EQS: \n - Name: '%s' (id=%d, option=%d),\n - All Items: %d,\n - ValidItems: %d"), *Query->QueryName, Query->QueryID, Query->OptionIndex, Query->ItemDetails.Num(), Query->NumValidItems); \
+		UEnvQueryDebugHelpers::LogQuery(Query, CategoryName.GetCategoryName(), ELogVerbosity::Verbosity); \
+	} \
+}
+
+#else
+#define UE_VLOG_EQS(Query, CategoryName, Verbosity)
+
+#endif //ENABLE_VISUAL_LOG
+
+UCLASS(Abstract, CustomConstructor)
+class AIMODULE_API UEnvQueryDebugHelpers : public UObject
+{
+	GENERATED_UCLASS_BODY()
+
+	UEnvQueryDebugHelpers(const class FPostConstructInitializeProperties& PCIP) : Super(PCIP) {}
+#if USE_EQS_DEBUGGER || ENABLE_VISUAL_LOG
+	static void QueryToDebugData(struct FEnvQueryInstance* Query, EQSDebug::FQueryData& EQSLocalData);
+	static void QueryToBlobArray(struct FEnvQueryInstance* Query, TArray<uint8>& BlobArray, bool bUseCompression = false);
+	static void  BlobArrayToDebugData(const TArray<uint8>& BlobArray, EQSDebug::FQueryData& EQSLocalData, bool bUseCompression = false);
+	static void LogQuery(struct FEnvQueryInstance* Query, const FName& CategoryName, ELogVerbosity::Type);
+#endif
+};
diff --git a/Engine/Source/Runtime/AIModule/Classes/VisualLoggerExtension.h b/Engine/Source/Runtime/AIModule/Classes/VisualLoggerExtension.h
new file mode 100644
--- /dev/null
+++ b/Engine/Source/Runtime/AIModule/Classes/VisualLoggerExtension.h
@@ -0,0 +1,33 @@
+// Copyright 1998-2014 Epic Games, Inc. All Rights Reserved.
+#pragma once
+#if ENABLE_VISUAL_LOG
+#	include "VisualLog.h"
+#endif
+#include "VisualLoggerExtension.generated.h"
+
+namespace EVisLogTags
+{
+	const FString TAG_EQS = TEXT("EQS");
+}
+
+#if ENABLE_VISUAL_LOG
+class FVisualLoggerExtension : public FVisualLogExtensionInterface
+{
+public:
+	FVisualLoggerExtension();
+	virtual void OnTimestampChange(float Timestamp, class UWorld* InWorld, class AActor* HelperActor);
+	virtual void DrawData(class UWorld* InWorld, class UCanvas* Canvas, class AActor* HelperActor, const FName& TagName, const struct FVisLogEntry::FDataBlock& DataBlock, float Timestamp);
+
+protected:
+	uint32 CachedEQSId;
+	float CurrentTimestamp;
+};
+#endif //ENABLE_VISUAL_LOG
+
+UCLASS(Abstract, CustomConstructor)
+class AIMODULE_API UVisualLoggerExtension : public UObject
+{
+	GENERATED_UCLASS_BODY()
+
+	UVisualLoggerExtension(const class FPostConstructInitializeProperties& PCIP) : Super(PCIP) {}
+};
diff --git a/Engine/Source/Runtime/AIModule/Private/AIModule.cpp b/Engine/Source/Runtime/AIModule/Private/AIModule.cpp
--- a/Engine/Source/Runtime/AIModule/Private/AIModule.cpp
+++ b/Engine/Source/Runtime/AIModule/Private/AIModule.cpp
@@ -1,7 +1,9 @@
 // Copyright 1998-2014 Epic Games, Inc. All Rights Reserved.
-
 #include "AIModulePrivate.h"
 #include "AISystem.h"
+#if WITH_EDITOR && ENABLE_VISUAL_LOG
+#	include "VisualLoggerExtension.h"
+#endif
 
 DEFINE_LOG_CATEGORY_STATIC(LogAIModule, Log, All);
 
@@ -13,6 +15,11 @@ class FAIModule : public IAIModule
 	
 	virtual UAISystemBase* CreateAISystemInstance(UWorld* World) override;
 	// End IModuleInterface
+
+protected:
+#if WITH_EDITOR && ENABLE_VISUAL_LOG
+	FVisualLoggerExtension	VisualLoggerExtension;
+#endif
 };
 
 IMPLEMENT_MODULE(FAIModule, AIModule)
@@ -23,12 +30,19 @@ void FAIModule::StartupModule()
 #if !(UE_BUILD_SHIPPING || UE_BUILD_TEST)
 	FModuleManager::LoadModulePtr< IModuleInterface >("GameplayDebugger");
 #endif
+
+#if WITH_EDITOR && ENABLE_VISUAL_LOG
+	FVisualLog::Get().RegisterExtension(*EVisLogTags::TAG_EQS, &VisualLoggerExtension);
+#endif
 }
 
 void FAIModule::ShutdownModule()
 {
 	// This function may be called during shutdown to clean up your module.  For modules that support dynamic reloading,
 	// we call this function before unloading the module.
+#if WITH_EDITOR && ENABLE_VISUAL_LOG
+	FVisualLog::Get().UnregisterExtension(*EVisLogTags::TAG_EQS, &VisualLoggerExtension);
+#endif
 }
 
 UAISystemBase* FAIModule::CreateAISystemInstance(UWorld* World)
@@ -37,4 +51,4 @@ UAISystemBase* FAIModule::CreateAISystemInstance(UWorld* World)
 	
 	TSubclassOf<class UAISystemBase> AISystemClass = LoadClass<UAISystemBase>(NULL, *UAISystem::GetAISystemClassName().ToString(), NULL, LOAD_None, NULL);
 	return ConstructObject<UAISystemBase>(AISystemClass, World);
-}
\ No newline at end of file
+}
diff --git a/Engine/Source/Runtime/AIModule/Private/EnvironmentQuery/EQSRenderingComponent.cpp b/Engine/Source/Runtime/AIModule/Private/EnvironmentQuery/EQSRenderingComponent.cpp
--- a/Engine/Source/Runtime/AIModule/Private/EnvironmentQuery/EQSRenderingComponent.cpp
+++ b/Engine/Source/Runtime/AIModule/Private/EnvironmentQuery/EQSRenderingComponent.cpp
@@ -66,6 +66,9 @@ FEQSSceneProxy::FEQSSceneProxy(const UPrimitiveComponent* InComponent, const FSt
 	ViewFlagIndex = uint32(FEngineShowFlags::FindIndexByName(*InViewFlagName));
 	ViewFlagName = InViewFlagName;
 
+	SolidSpheres = Spheres;
+	Texts = InTexts;
+
 	if (InComponent == NULL)
 	{
 		return;
@@ -81,9 +84,6 @@ FEQSSceneProxy::FEQSSceneProxy(const UPrimitiveComponent* InComponent, const FSt
 			return;
 		}
 	}
-
-	SolidSpheres = Spheres;
-	Texts = InTexts;
 }
 
 void FEQSSceneProxy::CollectEQSData(const UPrimitiveComponent* InComponent, const class IEQSQueryResultSourceInterface* InQueryDataSource, TArray<FSphere>& Spheres, TArray<FText3d>& Texts)
@@ -106,6 +106,19 @@ void FEQSSceneProxy::CollectEQSData(const UPrimitiveComponent* InComponent, cons
 	const FEnvQueryResult* ResultItems = QueryDataSource->GetQueryResult();
 	const FEnvQueryInstance* QueryInstance = QueryDataSource->GetQueryInstance();
 
+	FEQSSceneProxy::CollectEQSData(ResultItems, QueryInstance, Spheres, Texts, QueryDataSource->GetShouldDrawFailedItems());
+
+	if (ActorOwner && Spheres.Num() > EQSMaxItemsDrawn)
+	{
+		UE_VLOG(ActorOwner, LogEQS, Warning, TEXT("EQS drawing: too much items to draw! Drawing first %d from set of %d")
+			, EQSMaxItemsDrawn
+			, Spheres.Num()
+			);
+	}
+}
+
+void FEQSSceneProxy::CollectEQSData(const struct FEnvQueryResult* ResultItems, const struct FEnvQueryInstance* QueryInstance, TArray<FSphere>& Spheres, TArray<FText3d>& Texts, bool ShouldDrawFailedItems)
+{
 	if (ResultItems == NULL)
 	{
 		if (QueryInstance == NULL)
@@ -193,7 +206,7 @@ void FEQSSceneProxy::CollectEQSData(const UPrimitiveComponent* InComponent, cons
 		}
 	}
 
-	if (QueryDataSource->GetShouldDrawFailedItems() && QueryInstance)
+	if (ShouldDrawFailedItems && QueryInstance)
 	{
 		const FEQSQueryDebugData& InstanceDebugData = QueryInstance->DebugData;
 		const TArray<FEnvQueryItem>& DebugQueryItems = InstanceDebugData.DebugItems;
@@ -217,21 +230,13 @@ void FEQSSceneProxy::CollectEQSData(const UPrimitiveComponent* InComponent, cons
 		}
 	}
 #endif // USE_EQS_DEBUGGER
-
-	if (ActorOwner && Spheres.Num() > EQSMaxItemsDrawn)
-	{
-		UE_VLOG(ActorOwner, LogEQS, Warning, TEXT("EQS drawing: too much items to draw! Drawing first %d from set of %d")
-			, EQSMaxItemsDrawn
-			, Spheres.Num()
-			);
-	}
 }
 
 void FEQSSceneProxy::DrawDebugLabels(UCanvas* Canvas, APlayerController* PC)
 {
 	if ( !ActorOwner
 		 || (ActorOwner->IsSelected() == false && bDrawOnlyWhenSelected == true)
-		 || QueryDataSource->GetShouldDebugDrawLabels() == false)
+		 || (QueryDataSource && QueryDataSource->GetShouldDebugDrawLabels() == false))
 	{
 		return;
 	}
@@ -271,6 +276,12 @@ UEQSRenderingComponent::UEQSRenderingComponent(const class FPostConstructInitial
 
 FPrimitiveSceneProxy* UEQSRenderingComponent::CreateSceneProxy()
 {
+#if  USE_EQS_DEBUGGER || ENABLE_VISUAL_LOG
+	if (DebugData.SolidSpheres.Num() > 0 || DebugData.Texts.Num() > 0)
+	{
+		return new FEQSSceneProxy(this, DrawFlagName, bDrawOnlyWhenSelected, DebugData.SolidSpheres, DebugData.Texts);
+	}
+#endif
 	return new FEQSSceneProxy(this, DrawFlagName, bDrawOnlyWhenSelected);
 }
 
diff --git a/Engine/Source/Runtime/AIModule/Private/EnvironmentQuery/EnvQueryDebugHelpers.cpp b/Engine/Source/Runtime/AIModule/Private/EnvironmentQuery/EnvQueryDebugHelpers.cpp
new file mode 100644
--- /dev/null
+++ b/Engine/Source/Runtime/AIModule/Private/EnvironmentQuery/EnvQueryDebugHelpers.cpp
@@ -0,0 +1,131 @@
+// Copyright 1998-2014 Epic Games, Inc. All Rights Reserved.
+
+#include "AIModulePrivate.h"
+#include "DebugRenderSceneProxy.h"
+#include "EnvironmentQuery/EQSRenderingComponent.h"
+#include "EnvironmentQuery/EnvQueryDebugHelpers.h"
+#include "EnvironmentQuery/EQSQueryResultSourceInterface.h"
+
+#if USE_EQS_DEBUGGER || ENABLE_VISUAL_LOG
+void UEnvQueryDebugHelpers::QueryToBlobArray(struct FEnvQueryInstance* Query, TArray<uint8>& BlobArray, bool bUseCompression)
+{
+	EQSDebug::FQueryData EQSLocalData;
+	UEnvQueryDebugHelpers::QueryToDebugData(Query, EQSLocalData);
+
+	if (!bUseCompression)
+	{
+		FMemoryWriter ArWriter(BlobArray);
+		ArWriter << EQSLocalData;
+	}
+	else
+	{
+		TArray<uint8> UncompressedBuffer;
+		FMemoryWriter ArWriter(UncompressedBuffer);
+		ArWriter << EQSLocalData;
+
+		const int32 UncompressedSize = UncompressedBuffer.Num();
+		const int32 HeaderSize = sizeof(int32);
+		BlobArray.Init(0, HeaderSize + FMath::TruncToInt(1.1f * UncompressedSize));
+
+		int32 CompressedSize = BlobArray.Num() - HeaderSize;
+		uint8* DestBuffer = BlobArray.GetTypedData();
+		FMemory::Memcpy(DestBuffer, &UncompressedSize, HeaderSize);
+		DestBuffer += HeaderSize;
+
+		FCompression::CompressMemory((ECompressionFlags)(COMPRESS_ZLIB | COMPRESS_BiasMemory), (void*)DestBuffer, CompressedSize, (void*)UncompressedBuffer.GetData(), UncompressedSize);
+
+		BlobArray.SetNum(CompressedSize + HeaderSize, false);
+	}
+}
+
+void UEnvQueryDebugHelpers::QueryToDebugData(struct FEnvQueryInstance* Query, EQSDebug::FQueryData& EQSLocalData)
+{
+	// step 1: data for rendering component
+	EQSLocalData.SolidSpheres.Reset();
+	EQSLocalData.Texts.Reset();
+
+	FEQSSceneProxy::CollectEQSData(Query, Query, EQSLocalData.SolidSpheres, EQSLocalData.Texts, true);
+
+	// step 2: detailed scoring data for HUD
+	const int32 MaxDetailedItems = 10;
+	const int32 FirstItemIndex = 0;
+
+	const int32 NumTests = Query->ItemDetails.IsValidIndex(0) ? Query->ItemDetails[0].TestResults.Num() : 0;
+	const int32 NumItems = FMath::Min(MaxDetailedItems, Query->NumValidItems);
+
+	EQSLocalData.Name = Query->QueryName;
+	EQSLocalData.Id = Query->QueryID;
+	EQSLocalData.Items.Reset(NumItems);
+	EQSLocalData.Tests.Reset(NumTests);
+	EQSLocalData.NumValidItems = Query->NumValidItems;
+
+	UEnvQueryItemType* ItemCDO = Query->ItemType.GetDefaultObject();
+	for (int32 ItemIdx = 0; ItemIdx < NumItems; ItemIdx++)
+	{
+		EQSDebug::FItemData ItemInfo;
+		ItemInfo.ItemIdx = ItemIdx + FirstItemIndex;
+		ItemInfo.TotalScore = Query->Items[ItemInfo.ItemIdx].Score;
+
+		const uint8* ItemData = Query->RawData.GetTypedData() + Query->Items[ItemInfo.ItemIdx].DataOffset;
+		ItemInfo.Desc = FString::Printf(TEXT("[%d] %s"), ItemInfo.ItemIdx, *ItemCDO->GetDescription(ItemData));
+
+		ItemInfo.TestValues.Reserve(NumTests);
+		ItemInfo.TestScores.Reserve(NumTests);
+		for (int32 TestIdx = 0; TestIdx < NumTests; TestIdx++)
+		{
+			const float ScoreN = Query->ItemDetails[ItemInfo.ItemIdx].TestResults[TestIdx];
+			const float ScoreW = Query->ItemDetails[ItemInfo.ItemIdx].TestWeightedScores[TestIdx];
+
+			ItemInfo.TestValues.Add(ScoreN);
+			ItemInfo.TestScores.Add(ScoreW);
+		}
+
+		EQSLocalData.Items.Add(ItemInfo);
+	}
+
+	for (int32 TestIdx = 0; TestIdx < NumTests; TestIdx++)
+	{
+		EQSDebug::FTestData TestInfo;
+
+		UEnvQueryTest* TestOb = Query->Options[Query->OptionIndex].GetTestObject(TestIdx);
+		TestInfo.ShortName = TestOb->GetDescriptionTitle();
+		TestInfo.Detailed = TestOb->GetDescriptionDetails().ToString().Replace(TEXT("\n"), TEXT(", "));
+
+		EQSLocalData.Tests.Add(TestInfo);
+	}
+}
+
+void  UEnvQueryDebugHelpers::BlobArrayToDebugData(const TArray<uint8>& BlobArray, EQSDebug::FQueryData& EQSLocalData, bool bUseCompression)
+{
+	if (!bUseCompression)
+	{
+		FMemoryReader ArReader(BlobArray);
+		ArReader << EQSLocalData;
+	}
+	else
+	{
+		TArray<uint8> UncompressedBuffer;
+		int32 UncompressedSize = 0;
+		const int32 HeaderSize = sizeof(int32);
+		uint8* SrcBuffer = (uint8*)BlobArray.GetTypedData();
+		FMemory::Memcpy(&UncompressedSize, SrcBuffer, HeaderSize);
+		SrcBuffer += HeaderSize;
+		const int32 CompressedSize = BlobArray.Num() - HeaderSize;
+
+		UncompressedBuffer.AddZeroed(UncompressedSize);
+
+		FCompression::UncompressMemory((ECompressionFlags)(COMPRESS_ZLIB), (void*)UncompressedBuffer.GetTypedData(), UncompressedSize, SrcBuffer, CompressedSize);
+		FMemoryReader ArReader(UncompressedBuffer);
+
+		ArReader << EQSLocalData;
+	}
+}
+
+void UEnvQueryDebugHelpers::LogQuery(struct FEnvQueryInstance* Query, const FName& CategoryName, ELogVerbosity::Type Type)
+{
+	TArray<uint8> BlobArray;
+	UEnvQueryDebugHelpers::QueryToBlobArray(Query, BlobArray);
+	FVisualLog::Get().GetEntryToWrite(Cast<AActor>(Query->Owner.Get()))->AddDataBlock(EVisLogTags::TAG_EQS, BlobArray, CategoryName);
+}
+
+#endif //ENABLE_VISUAL_LOG
diff --git a/Engine/Source/Runtime/AIModule/Private/EnvironmentQuery/EnvQueryManager.cpp b/Engine/Source/Runtime/AIModule/Private/EnvironmentQuery/EnvQueryManager.cpp
--- a/Engine/Source/Runtime/AIModule/Private/EnvironmentQuery/EnvQueryManager.cpp
+++ b/Engine/Source/Runtime/AIModule/Private/EnvironmentQuery/EnvQueryManager.cpp
@@ -7,7 +7,7 @@
 #include "EnvironmentQuery/EnvQueryOption.h"
 #include "EnvironmentQuery/EnvQueryManager.h"
 #include "EnvironmentQuery/EQSTestingPawn.h"
-
+#include "EnvironmentQuery/EnvQueryDebugHelpers.h"
 #if WITH_EDITOR
 #include "Engine/Brush.h"
 #include "Editor/EditorEngine.h"
@@ -177,6 +177,8 @@ TSharedPtr<FEnvQueryResult> UEnvQueryManager::RunInstantQuery(const FEnvQueryReq
 		QueryInstance->ExecuteOneStep((double)FLT_MAX);
 	}
 
+	UE_VLOG_EQS(QueryInstance.Get(), LogEQS, All);
+
 #if USE_EQS_DEBUGGER
 	EQSDebugger.StoreQuery(QueryInstance);
 #endif // USE_EQS_DEBUGGER
@@ -248,6 +250,8 @@ void UEnvQueryManager::Tick(float DeltaTime)
 			
 			if (QueryInstance->Status != EEnvQueryStatus::Processing)
 			{
+				UE_VLOG_EQS(QueryInstance.Get(), LogEQS, All);
+
 #if USE_EQS_DEBUGGER
 				EQSDebugger.StoreQuery(QueryInstance);
 #endif // USE_EQS_DEBUGGER
diff --git a/Engine/Source/Runtime/AIModule/Private/VisualLoggerExtension.cpp b/Engine/Source/Runtime/AIModule/Private/VisualLoggerExtension.cpp
new file mode 100644
--- /dev/null
+++ b/Engine/Source/Runtime/AIModule/Private/VisualLoggerExtension.cpp
@@ -0,0 +1,59 @@
+#include "AIModulePrivate.h"
+#include "EnvironmentQuery/EnvQueryDebugHelpers.h"
+#include "EnvironmentQuery/EQSRenderingComponent.h"
+#include "VisualLoggerExtension.h"
+
+#if ENABLE_VISUAL_LOG
+FVisualLoggerExtension::FVisualLoggerExtension()
+	:  CachedEQSId(INDEX_NONE)
+	, CurrentTimestamp(FLT_MIN)
+{
+
+}
+
+void FVisualLoggerExtension::OnTimestampChange(float Timestamp, class UWorld* InWorld, class AActor* HelperActor )
+{
+	if (CurrentTimestamp != Timestamp)
+	{
+		CurrentTimestamp = Timestamp;
+		CachedEQSId = INDEX_NONE;
+		UEQSRenderingComponent* EQSRenderComp = HelperActor->FindComponentByClass<UEQSRenderingComponent>();
+		if (EQSRenderComp)
+		{
+			EQSRenderComp->SetHiddenInGame(true);
+			EQSRenderComp->Deactivate();
+			EQSRenderComp->DebugData.Reset();
+			EQSRenderComp->MarkRenderStateDirty();
+		}
+	}
+}
+
+void FVisualLoggerExtension::DrawData(class UWorld* InWorld, class UCanvas* Canvas, class AActor* HelperActor, const FName& TagName, const FVisLogEntry::FDataBlock& DataBlock, float Timestamp)
+{
+	if (TagName == *EVisLogTags::TAG_EQS)
+	{
+		if (HelperActor)
+		{
+			UEQSRenderingComponent* EQSRenderComp = HelperActor->FindComponentByClass<UEQSRenderingComponent>();
+			if (!EQSRenderComp)
+			{
+				EQSRenderComp = ConstructObject<UEQSRenderingComponent>(UEQSRenderingComponent::StaticClass(), HelperActor);
+				EQSRenderComp->bDrawOnlyWhenSelected = false;
+				EQSRenderComp->RegisterComponent();
+				EQSRenderComp->SetHiddenInGame(true);
+			}
+
+			EQSDebug::FQueryData DebugData;
+			UEnvQueryDebugHelpers::BlobArrayToDebugData(DataBlock.Data, DebugData, false);
+			if (DebugData.Id != CachedEQSId || (EQSRenderComp && EQSRenderComp->bHiddenInGame))
+			{
+				CachedEQSId = DebugData.Id;
+				EQSRenderComp->DebugData = DebugData;
+				EQSRenderComp->Activate();
+				EQSRenderComp->SetHiddenInGame(false);
+				EQSRenderComp->MarkRenderStateDirty();
+			}
+		}
+	}
+}
+#endif //ENABLE_VISUAL_LOG
\ No newline at end of file
diff --git a/Engine/Source/Runtime/Engine/Private/VisualLog.cpp b/Engine/Source/Runtime/Engine/Private/VisualLog.cpp
--- a/Engine/Source/Runtime/Engine/Private/VisualLog.cpp
+++ b/Engine/Source/Runtime/Engine/Private/VisualLog.cpp
@@ -10,6 +10,162 @@ DEFINE_LOG_CATEGORY(LogVisual);
 
 DEFINE_STAT(STAT_VisualLog);
 
+/** The table used to encode a 6 bit value as an ascii character */
+uint8 EncodingAlphabet[64] = { 'A', 'B', 'C', 'D', 'E', 'F', 'G', 'H', 'I', 'J', 'K', 'L',
+'M', 'N', 'O', 'P', 'Q', 'R', 'S', 'T', 'U', 'V', 'W', 'X', 'Y', 'Z', 'a', 'b', 'c', 'd', 'e',
+'f', 'g', 'h', 'i', 'j', 'k', 'l', 'm', 'n', 'o', 'p', 'q', 'r', 's', 't', 'u', 'v', 'w', 'x',
+'y', 'z', '0', '1', '2', '3', '4', '5', '6', '7', '8', '9', '+', '/' };
+
+/** The table used to convert an ascii character into a 6 bit value */
+uint8 DecodingAlphabet[256] = { 0xFF, 0xFF, 0xFF, 0xFF, 0xFF, 0xFF, 0xFF, 0xFF, 0xFF, 0xFF,
+0xFF, 0xFF, 0xFF, 0xFF, 0xFF, 0xFF, 0xFF, 0xFF, 0xFF, 0xFF, 0xFF, 0xFF, 0xFF, 0xFF, 0xFF, 0xFF,
+0xFF, 0xFF, 0xFF, 0xFF, 0xFF, 0xFF, 0xFF, 0xFF, 0xFF, 0xFF, 0xFF, 0xFF, 0xFF, 0xFF, 0xFF, 0xFF,
+0xFF, 0x3E, 0xFF, 0xFF, 0xFF, 0x3F, 0x34, 0x35, 0x36, 0x37, 0x38, 0x39, 0x3A, 0x3B, 0x3C, 0x3D,
+0xFF, 0xFF, 0xFF, 0x00, 0xFF, 0xFF, 0xFF, 0x00, 0x01, 0x02, 0x03, 0x04, 0x05, 0x06, 0x07, 0x08,
+0x09, 0x0A, 0x0B, 0x0C, 0x0D, 0x0E, 0x0F, 0x10, 0x11, 0x12, 0x13, 0x14, 0x15, 0x16, 0x17, 0x18,
+0x19, 0xFF, 0xFF, 0xFF, 0xFF, 0xFF, 0xFF, 0x1A, 0x1B, 0x1C, 0x1D, 0x1E, 0x1F, 0x20, 0x21, 0x22,
+0x23, 0x24, 0x25, 0x26, 0x27, 0x28, 0x29, 0x2A, 0x2B, 0x2C, 0x2D, 0x2E, 0x2F, 0x30, 0x31, 0x32,
+0x33, 0xFF, 0xFF, 0xFF, 0xFF, 0xFF, 0xFF, 0xFF, 0xFF, 0xFF, 0xFF, 0xFF, 0xFF, 0xFF, 0xFF, 0xFF,
+0xFF, 0xFF, 0xFF, 0xFF, 0xFF, 0xFF, 0xFF, 0xFF, 0xFF, 0xFF, 0xFF, 0xFF, 0xFF, 0xFF, 0xFF, 0xFF,
+0xFF, 0xFF, 0xFF, 0xFF, 0xFF, 0xFF, 0xFF, 0xFF, 0xFF, 0xFF, 0xFF, 0xFF, 0xFF, 0xFF, 0xFF, 0xFF,
+0xFF, 0xFF, 0xFF, 0xFF, 0xFF, 0xFF, 0xFF, 0xFF, 0xFF, 0xFF, 0xFF, 0xFF, 0xFF, 0xFF, 0xFF, 0xFF,
+0xFF, 0xFF, 0xFF, 0xFF, 0xFF, 0xFF, 0xFF, 0xFF, 0xFF, 0xFF, 0xFF, 0xFF, 0xFF, 0xFF, 0xFF, 0xFF,
+0xFF, 0xFF, 0xFF, 0xFF, 0xFF, 0xFF, 0xFF, 0xFF, 0xFF, 0xFF, 0xFF, 0xFF, 0xFF, 0xFF, 0xFF, 0xFF,
+0xFF, 0xFF, 0xFF, 0xFF, 0xFF, 0xFF, 0xFF, 0xFF, 0xFF, 0xFF, 0xFF, 0xFF, 0xFF, 0xFF, 0xFF, 0xFF,
+0xFF, 0xFF, 0xFF, 0xFF, 0xFF, 0xFF, 0xFF, 0xFF, 0xFF, 0xFF, 0xFF, 0xFF, 0xFF, 0xFF, 0xFF, 0xFF,
+0xFF, 0xFF, 0xFF, 0xFF, 0xFF, 0xFF };
+
+FString EncodeToBase64(uint8* Source, uint32 Length)
+{
+	// Each 3 uint8 set of data expands to 4 bytes and must be padded to 4 bytes
+	uint32 ExpectedLength = (Length + 2) / 3 * 4;
+	FString OutBuffer;
+	OutBuffer.Empty(ExpectedLength);
+
+	ANSICHAR EncodedBytes[5];
+	// Null terminate this so we can append it easily as a string
+	EncodedBytes[4] = 0;
+	// Loop through the buffer converting 3 bytes of binary data at a time
+	while (Length >= 3)
+	{
+		uint8 A = *Source++;
+		uint8 B = *Source++;
+		uint8 C = *Source++;
+		Length -= 3;
+		// The algorithm takes 24 bits of data (3 bytes) and breaks it into 4 6bit chunks represented as ascii
+		uint32 ByteTriplet = A << 16 | B << 8 | C;
+		// Use the 6bit block to find the representation ascii character for it
+		EncodedBytes[3] = EncodingAlphabet[ByteTriplet & 0x3F];
+		// Shift each time to get the next 6 bit chunk
+		ByteTriplet >>= 6;
+		EncodedBytes[2] = EncodingAlphabet[ByteTriplet & 0x3F];
+		ByteTriplet >>= 6;
+		EncodedBytes[1] = EncodingAlphabet[ByteTriplet & 0x3F];
+		ByteTriplet >>= 6;
+		EncodedBytes[0] = EncodingAlphabet[ByteTriplet & 0x3F];
+		// Now we can append this buffer to our destination string
+		OutBuffer += EncodedBytes;
+	}
+	// Since this algorithm operates on blocks, we may need to pad the last chunks
+	if (Length > 0)
+	{
+		uint8 A = *Source++;
+		uint8 B = 0;
+		uint8 C = 0;
+		// Grab the second character if it is a 2 uint8 finish
+		if (Length == 2)
+		{
+			B = *Source;
+		}
+		uint32 ByteTriplet = A << 16 | B << 8 | C;
+		// Pad with = to make a 4 uint8 chunk
+		EncodedBytes[3] = '=';
+		ByteTriplet >>= 6;
+		// If there's only one 1 uint8 left in the source, then you need 2 pad chars
+		if (Length == 1)
+		{
+			EncodedBytes[2] = '=';
+		}
+		else
+		{
+			EncodedBytes[2] = EncodingAlphabet[ByteTriplet & 0x3F];
+		}
+		// Now encode the remaining bits the same way
+		ByteTriplet >>= 6;
+		EncodedBytes[1] = EncodingAlphabet[ByteTriplet & 0x3F];
+		ByteTriplet >>= 6;
+		EncodedBytes[0] = EncodingAlphabet[ByteTriplet & 0x3F];
+		OutBuffer += EncodedBytes;
+	}
+	return OutBuffer;
+}
+
+bool DecodeFromBase64(const ANSICHAR* Source, uint32 Length, uint8* Dest, uint32& PadCount)
+{
+	PadCount = 0;
+	uint8 DecodedValues[4];
+	while (Length)
+	{
+		// Decode the next 4 BYTEs
+		for (int32 Index = 0; Index < 4; Index++)
+		{
+			// Tell the caller if there were any pad bytes
+			if (*Source == '=')
+			{
+				PadCount++;
+			}
+			DecodedValues[Index] = DecodingAlphabet[(int32)(*Source++)];
+			// Abort on values that we don't understand
+			if (DecodedValues[Index] == -1)
+			{
+				return false;
+			}
+		}
+		Length -= 4;
+		// Rebuild the original 3 bytes from the 4 chunks of 6 bits
+		uint32 OriginalTriplet = DecodedValues[0];
+		OriginalTriplet <<= 6;
+		OriginalTriplet |= DecodedValues[1];
+		OriginalTriplet <<= 6;
+		OriginalTriplet |= DecodedValues[2];
+		OriginalTriplet <<= 6;
+		OriginalTriplet |= DecodedValues[3];
+		// Now we can tear the uint32 into bytes
+		Dest[2] = OriginalTriplet & 0xFF;
+		OriginalTriplet >>= 8;
+		Dest[1] = OriginalTriplet & 0xFF;
+		OriginalTriplet >>= 8;
+		Dest[0] = OriginalTriplet & 0xFF;
+		Dest += 3;
+	}
+	return true;
+}
+
+bool DecodeFromBase64(const FString& Source, TArray<uint8>& Dest)
+{
+	uint32 Length = Source.Len();
+	// Size must be a multiple of 4
+	if (Length % 4)
+	{
+		return false;
+	}
+	// Each 4 uint8 chunk of characters is 3 bytes of data
+	uint32 ExpectedLength = Length / 4 * 3;
+	// Add the number we need for output
+	Dest.AddZeroed(ExpectedLength);
+	uint8* Buffer = Dest.GetTypedData();
+	uint32 PadCount = 0;
+	bool bWasSuccessful = DecodeFromBase64(TCHAR_TO_ANSI(*Source), Length, Dest.GetTypedData(), PadCount);
+	if (bWasSuccessful)
+	{
+		if (PadCount > 0)
+		{
+			Dest.RemoveAt(ExpectedLength - PadCount, PadCount);
+		}
+	}
+	return bWasSuccessful;
+}
+
 //----------------------------------------------------------------------//
 // FVisLogEntry
 //----------------------------------------------------------------------//
@@ -141,6 +297,39 @@ FVisLogEntry::FVisLogEntry(TSharedPtr<FJsonValue> FromJson)
 		}
 	}
 
+	TArray< TSharedPtr<FJsonValue> > JasonDataBlocks = JsonEntryObject->GetArrayField(VisualLogJson::TAG_DATABLOCK);
+	if (JasonDataBlocks.Num() > 0)
+	{
+		DataBlocks.Reserve(JasonDataBlocks.Num());
+		for (int32 SampleIndex = 0; SampleIndex < JasonDataBlocks.Num(); ++SampleIndex)
+		{
+			TSharedPtr<FJsonObject> JsonSampleObject = JasonDataBlocks[SampleIndex]->AsObject();
+			if (JsonSampleObject.IsValid())
+			{
+				FDataBlock& Sample = DataBlocks[DataBlocks.Add(FDataBlock())];
+				Sample.TagName = FName(*(JsonSampleObject->GetStringField(VisualLogJson::TAG_DATABLOCK_TAGNAME)));
+				Sample.Category = FName(*(JsonSampleObject->GetStringField(VisualLogJson::TAG_CATEGORY)));
+				Sample.Verbosity = TEnumAsByte<ELogVerbosity::Type>((uint8)FMath::TruncToInt(JsonSampleObject->GetNumberField(VisualLogJson::TAG_VERBOSITY)));
+
+				// decode data from Base64 string
+				const FString DataBlockAsCompressedString = JsonSampleObject->GetStringField(VisualLogJson::TAG_DATABLOCK_DATA);
+				TArray<uint8> CompressedDataBlock;
+				DecodeFromBase64(DataBlockAsCompressedString, CompressedDataBlock);
+
+				// uncompress decoded data to get final TArray<uint8>
+				int32 UncompressedSize = 0;
+				const int32 HeaderSize = sizeof(int32);
+				uint8* SrcBuffer = (uint8*)CompressedDataBlock.GetTypedData();
+				FMemory::Memcpy(&UncompressedSize, SrcBuffer, HeaderSize);
+				SrcBuffer += HeaderSize;
+				const int32 CompressedSize = CompressedDataBlock.Num() - HeaderSize;
+
+				Sample.Data.AddZeroed(UncompressedSize);
+				FCompression::UncompressMemory((ECompressionFlags)(COMPRESS_ZLIB), (void*)Sample.Data.GetTypedData(), UncompressedSize, SrcBuffer, CompressedSize);
+			}
+		}
+	}
+
 }
 
 TSharedPtr<FJsonValue> FVisLogEntry::ToJson() const
@@ -239,6 +428,35 @@ TSharedPtr<FJsonValue> FVisLogEntry::ToJson() const
 
 	JsonEntryObject->SetArrayField(VisualLogJson::TAG_HISTOGRAMSAMPLES, JsonHistogramSamples);
 	
+	int32 DataBlockIndex = 0;
+	TArray<TSharedPtr<FJsonValue> > JasonDataBlocks;
+	JasonDataBlocks.AddZeroed(DataBlocks.Num());
+	for (const auto CurrentData : DataBlocks)
+	{
+		TSharedPtr<FJsonObject> JsonSampleObject = MakeShareable(new FJsonObject);
+
+		TArray<uint8> CompressedData;
+		const int32 UncompressedSize = CurrentData.Data.Num();
+		const int32 HeaderSize = sizeof(int32);
+		CompressedData.Init(0, HeaderSize + FMath::TruncToInt(1.1f * UncompressedSize));
+
+		int32 CompressedSize = CompressedData.Num() - HeaderSize;
+		uint8* DestBuffer = CompressedData.GetTypedData();
+		FMemory::Memcpy(DestBuffer, &UncompressedSize, HeaderSize);
+		DestBuffer += HeaderSize;
+
+		FCompression::CompressMemory((ECompressionFlags)(COMPRESS_ZLIB | COMPRESS_BiasMemory), (void*)DestBuffer, CompressedSize, (void*)CurrentData.Data.GetData(), UncompressedSize);
+		CompressedData.SetNum(CompressedSize + HeaderSize, true);
+		const FString CurrentDataAsString = EncodeToBase64((uint8*)CompressedData.GetTypedData(), CompressedData.Num());
+
+		JsonSampleObject->SetStringField(VisualLogJson::TAG_CATEGORY, CurrentData.Category.ToString());
+		JsonSampleObject->SetStringField(VisualLogJson::TAG_DATABLOCK_TAGNAME, CurrentData.TagName.ToString());
+		JsonSampleObject->SetStringField(VisualLogJson::TAG_DATABLOCK_DATA, CurrentDataAsString);
+		JsonSampleObject->SetNumberField(VisualLogJson::TAG_VERBOSITY, CurrentData.Verbosity);
+
+		JasonDataBlocks[DataBlockIndex++] = MakeShareable(new FJsonValueObject(JsonSampleObject));
+	}
+	JsonEntryObject->SetArrayField(VisualLogJson::TAG_DATABLOCK, JasonDataBlocks);
 
 	return MakeShareable(new FJsonValueObject(JsonEntryObject));
 }
@@ -290,6 +508,16 @@ void FVisLogEntry::AddHistogramData(const FVector2D& DataSample, const FName& Ca
 	HistogramSamples.Add(Sample);
 }
 
+void FVisLogEntry::AddDataBlock(const FString& TagName, const TArray<uint8>& BlobDataArray, const FName& CategoryName)
+{
+	FDataBlock DataBlock;
+	DataBlock.Category = CategoryName;
+	DataBlock.TagName = *TagName;
+	DataBlock.Data = BlobDataArray;
+
+	DataBlocks.Add(DataBlock);
+}
+
 //----------------------------------------------------------------------//
 // FActorsVisLog
 //----------------------------------------------------------------------//
diff --git a/Engine/Source/Runtime/Engine/Public/VisualLog.h b/Engine/Source/Runtime/Engine/Public/VisualLog.h
--- a/Engine/Source/Runtime/Engine/Public/VisualLog.h
+++ b/Engine/Source/Runtime/Engine/Public/VisualLog.h
@@ -148,6 +148,17 @@ struct ENGINE_API FVisLogEntry
 	};
 	TArray<FHistogramSample>	HistogramSamples;
 
+	struct FDataBlock
+	{
+		FDataBlock() : Verbosity(ELogVerbosity::All) {}
+
+		FName TagName;
+		FName Category;
+		TEnumAsByte<ELogVerbosity::Type> Verbosity;
+		TArray<uint8>	Data;
+	};
+	TArray<FDataBlock>	DataBlocks;
+
 
 	FVisLogEntry(const class AActor* Actor, TArray<TWeakObjectPtr<AActor> >* Children);
 	FVisLogEntry(TSharedPtr<FJsonValue> FromJson);
@@ -166,6 +177,9 @@ struct ENGINE_API FVisLogEntry
 	// histogram sample
 	void AddHistogramData(const FVector2D& DataSample, const FName& CategoryName, const FName& GraphName, const FName& DataName);
 
+	// Custom data block
+	void AddDataBlock(const FString& TagName, const TArray<uint8>& BlobDataArray, const FName& CategoryName);
+
 	// find index of status category
 	FORCEINLINE int32 FindStatusIndex(const FString& CategoryName)
 	{
@@ -277,6 +291,10 @@ namespace VisualLogJson
 	static const FString TAG_HISTOGRAMGRAPHNAME = TEXT("GraphName");
 	static const FString TAG_HISTOGRAMDATANAME = TEXT("DataName");
 
+	static const FString TAG_DATABLOCK = TEXT("DataBlock");
+	static const FString TAG_DATABLOCK_DATA = TEXT("DataBlockData");
+	static const FString TAG_DATABLOCK_TAGNAME = TEXT("DataBlockTagName");
+
 	static const FString TAG_LOGS = TEXT("Logs");
 }
 
@@ -296,6 +314,13 @@ struct ENGINE_API FActorsVisLog
 
 DECLARE_DELEGATE_RetVal(FString, FVisualLogFilenameGetterDelegate);
 
+struct FVisualLogExtensionInterface
+{
+	virtual void OnTimestampChange(float Timestamp, class UWorld* InWorld, class AActor* HelperActor) = 0;
+	virtual void DrawData(class UWorld* InWorld, class UCanvas* Canvas, class AActor* HelperActor, const FName& TagName, const FVisLogEntry::FDataBlock& DataBlock, float Timestamp) = 0;
+};
+
+
 /** This class is to capture all log output even if the Visual Logger is closed */
 class ENGINE_API FVisualLog : public FOutputDevice
 {
@@ -344,6 +369,11 @@ public:
 	/** highly encouradged to use FVisualLogFilenameGetterDelegate::CreateUObject with this */
 	void SetLogFileNameGetter(const FVisualLogFilenameGetterDelegate& InLogFileNameGetter) { LogFileNameGetter = InLogFileNameGetter; }
 
+	TMap<FName, FVisualLogExtensionInterface*> AllExtensions;
+	void RegisterExtension(FName TagName, FVisualLogExtensionInterface* ExtensionInterface) { AllExtensions.Add(TagName, ExtensionInterface); }
+	void UnregisterExtension(FName TagName, FVisualLogExtensionInterface* ExtensionInterface) { AllExtensions.Remove(TagName); }
+	FVisualLogExtensionInterface* GetExtensionForTag(const FName TagName) { return AllExtensions.Contains(TagName) ? AllExtensions[TagName] : NULL; }
+	TMap<FName, FVisualLogExtensionInterface*>& GetAllExtensions() { return AllExtensions; }
 protected:
 
 	FString GetLogFileFullName() const;
@@ -393,4 +423,4 @@ private:
 	#define UE_VLOG_HISTOGRAM(Actor, CategoryName, Verbosity, GraphName, DataName, Data)
 	#define REDIRECT_TO_VLOG(Dest)
 	#define REDIRECT_ACTOR_TO_VLOG(Src, Destination) 
-#endif
+#endif //ENABLE_VISUAL_LOG
