Fixed some centering and scaling issues with wrapped Canvas text

TTP# 343372 - CANVAS: WrappedPrint() is not functioning correctly when used with centering and scaling

Updated some places in the Canvas string rendering that were using the older StringSize function (which has a hard-coded scale of 1) rather than the newer UCanvas::CanvasStringSize (which correctly measures using the current scale).

The centering on Y was trying to center each line on its own height, rather than center the entire block of text on its total height. This has been updated to apply a centering based upon the total height of the text being drawn.

Fixed some places that were doing float -> int32 -> float truncation when measuring.

diff --git a/Engine/Source/Runtime/Engine/Private/UserInterface/Canvas.cpp b/Engine/Source/Runtime/Engine/Private/UserInterface/Canvas.cpp
--- a/Engine/Source/Runtime/Engine/Private/UserInterface/Canvas.cpp
+++ b/Engine/Source/Runtime/Engine/Private/UserInterface/Canvas.cpp
@@ -1042,11 +1042,9 @@ namespace
 		FTextSizingParameters MeasureParameters(Parameters);
 		FString Substring(EndIndex - StartIndex, String + StartIndex);
 		FWrappedStringElement Element(*Substring, 0.0f, 0.0f);
-		int32 X;
-		int32 Y;
-		StringSize(MeasureParameters.DrawFont, X, Y, *Element.Value);
-		Element.LineExtent.X = X;
-		Element.LineExtent.Y = Y;
+		UCanvas::CanvasStringSize(MeasureParameters, *Element.Value);
+		Element.LineExtent.X = MeasureParameters.DrawXL;
+		Element.LineExtent.Y = MeasureParameters.DrawYL;
 		Results.Add(Element);
 	}
 
@@ -1398,8 +1396,32 @@ int32 UCanvas::WrappedPrint(bool Draw, float X, float Y, int32& out_XL, int32& o
 	TArray<FWrappedStringElement> WrappedStrings;
 	WrapString(RenderParms, 0, Text, WrappedStrings);
 
+	TArray<FVector2D> SizedStrings;
+	if (bCenterTextX || bCenterTextY)
+	{
+		SizedStrings.Reserve(WrappedStrings.Num());
+
+		FTextSizingParameters Parameters(Font, ScaleX, ScaleY);
+		for (int32 Idx = 0; Idx < WrappedStrings.Num(); Idx++)
+		{
+			UCanvas::CanvasStringSize(Parameters, *WrappedStrings[Idx].Value);
+			SizedStrings.Emplace(FVector2D(Parameters.DrawXL, Parameters.DrawYL));
+		}
+	}
+
 	float DrawX = OrgX + X;
 	float DrawY = OrgY + Y;
+	if (bCenterTextY)
+	{
+		// Center text about DrawY
+		float MeasuredHeight = 0.f;
+		for (const FVector2D& SizedString : SizedStrings)
+		{
+			MeasuredHeight += SizedString.Y;
+		}
+		DrawY -= (MeasuredHeight * 0.5f);
+	}
+
 	float XL = 0.f;
 	float YL = 0.f;
 	FCanvasTextItem TextItem( FVector2D::ZeroVector, FText::GetEmpty(), Font, DrawColor );
@@ -1408,27 +1430,20 @@ int32 UCanvas::WrappedPrint(bool Draw, float X, float Y, int32& out_XL, int32& o
 	TextItem.FontRenderInfo = RenderInfo;
 	for (int32 Idx = 0; Idx < WrappedStrings.Num(); Idx++)
 	{
-		if (bCenterTextX || bCenterTextY)
+		float LineDrawX = DrawX;
+		float LineDrawY = DrawY;
+
+		if (bCenterTextX)
 		{
 			// Center text about DrawX
-			int32 StringSizeX, StringSizeY;
-			StringSize(Font, StringSizeX, StringSizeY, *WrappedStrings[Idx].Value);
-
-			if (bCenterTextX)
-			{
-				DrawX = OrgX + X - (StringSizeX / 2);
-			}
-			if (bCenterTextY)
-			{
-				DrawY = OrgY + Y - (StringSizeY / 2);
-			}
+			LineDrawX -= (SizedStrings[Idx].X * 0.5f);
 		}
 
 		float LineXL = 0.0f;
 		if( Draw )
 		{
 			TextItem.Text = FText::FromString(WrappedStrings[Idx].Value);
-			Canvas->DrawItem( TextItem, DrawX, DrawY );
+			Canvas->DrawItem( TextItem, LineDrawX, LineDrawY );
 			LineXL = TextItem.DrawnSize.X;
 		}
 		else
@@ -1438,7 +1453,6 @@ int32 UCanvas::WrappedPrint(bool Draw, float X, float Y, int32& out_XL, int32& o
 			ClippedStrLen(Font, ScaleX, ScaleY, TempX, TempY, *WrappedStrings[Idx].Value);
 			LineXL = TempX;
 		}
-		//float const LineXL = Canvas->DrawString( DrawX, DrawY, *WrappedStrings[Idx].Value, Font, DrawColor, ScaleX, ScaleY, 0.0f, NULL, SE_BLEND_Translucent, Draw, 0.f, 1.f, RenderInfo );
 		XL = FMath::Max<float>(XL, LineXL);
 		DrawY += Font->GetMaxCharHeight() * ScaleY;
 		YL += Font->GetMaxCharHeight() * ScaleY;
@@ -1457,12 +1471,11 @@ void UCanvas::StrLen( UFont* InFont, const FString& InText, float& XL, float& YL
 	}
 	else
 	{
-		int32 XLi = 0, YLi = 0;
-
-		StringSize(InFont, XLi, YLi, *InText);
+		FTextSizingParameters Parameters(InFont,1.0f,1.0f);
+		UCanvas::CanvasStringSize(Parameters, *InText);
 
-		XL = XLi;
-		YL = YLi;
+		XL = Parameters.DrawXL;
+		YL = Parameters.DrawYL;
 	}
 }
 
