Removed some redundant text measuring in UCanvas::WrappedPrint

TTP# 343372 - CANVAS: WrappedPrint() is not functioning correctly when used with centering and scaling

FWrappedStringElement already contains the correct line measurements (calculated while wrapping the text), so there's no need to measure again in WrappedPrint.

diff --git a/Engine/Source/Runtime/Engine/Private/UserInterface/Canvas.cpp b/Engine/Source/Runtime/Engine/Private/UserInterface/Canvas.cpp
--- a/Engine/Source/Runtime/Engine/Private/UserInterface/Canvas.cpp
+++ b/Engine/Source/Runtime/Engine/Private/UserInterface/Canvas.cpp
@@ -1396,28 +1396,15 @@ int32 UCanvas::WrappedPrint(bool Draw, float X, float Y, int32& out_XL, int32& o
 	TArray<FWrappedStringElement> WrappedStrings;
 	WrapString(RenderParms, 0, Text, WrappedStrings);
 
-	TArray<FVector2D> SizedStrings;
-	if (bCenterTextX || bCenterTextY)
-	{
-		SizedStrings.Reserve(WrappedStrings.Num());
-
-		FTextSizingParameters Parameters(Font, ScaleX, ScaleY);
-		for (int32 Idx = 0; Idx < WrappedStrings.Num(); Idx++)
-		{
-			UCanvas::CanvasStringSize(Parameters, *WrappedStrings[Idx].Value);
-			SizedStrings.Emplace(FVector2D(Parameters.DrawXL, Parameters.DrawYL));
-		}
-	}
-
 	float DrawX = OrgX + X;
 	float DrawY = OrgY + Y;
 	if (bCenterTextY)
 	{
 		// Center text about DrawY
 		float MeasuredHeight = 0.f;
-		for (const FVector2D& SizedString : SizedStrings)
+		for (const FWrappedStringElement& WrappedString : WrappedStrings)
 		{
-			MeasuredHeight += SizedString.Y;
+			MeasuredHeight += WrappedString.LineExtent.Y;
 		}
 		DrawY -= (MeasuredHeight * 0.5f);
 	}
@@ -1428,7 +1415,7 @@ int32 UCanvas::WrappedPrint(bool Draw, float X, float Y, int32& out_XL, int32& o
 	TextItem.Scale = FVector2D( ScaleX, ScaleY );
 	TextItem.BlendMode = SE_BLEND_Translucent;
 	TextItem.FontRenderInfo = RenderInfo;
-	for (int32 Idx = 0; Idx < WrappedStrings.Num(); Idx++)
+	for (const FWrappedStringElement& WrappedString : WrappedStrings)
 	{
 		float LineDrawX = DrawX;
 		float LineDrawY = DrawY;
@@ -1436,13 +1423,13 @@ int32 UCanvas::WrappedPrint(bool Draw, float X, float Y, int32& out_XL, int32& o
 		if (bCenterTextX)
 		{
 			// Center text about DrawX
-			LineDrawX -= (SizedStrings[Idx].X * 0.5f);
+			LineDrawX -= (WrappedString.LineExtent.X * 0.5f);
 		}
 
 		float LineXL = 0.0f;
 		if( Draw )
 		{
-			TextItem.Text = FText::FromString(WrappedStrings[Idx].Value);
+			TextItem.Text = FText::FromString(WrappedString.Value);
 			Canvas->DrawItem( TextItem, LineDrawX, LineDrawY );
 			LineXL = TextItem.DrawnSize.X;
 		}
@@ -1450,7 +1437,7 @@ int32 UCanvas::WrappedPrint(bool Draw, float X, float Y, int32& out_XL, int32& o
 		{
 			int32 TempX;
 			int32 TempY;
-			ClippedStrLen(Font, ScaleX, ScaleY, TempX, TempY, *WrappedStrings[Idx].Value);
+			ClippedStrLen(Font, ScaleX, ScaleY, TempX, TempY, *WrappedString.Value);
 			LineXL = TempX;
 		}
 		XL = FMath::Max<float>(XL, LineXL);
