Fixed navmesh not generating for navigation bound volumes from levels added to a persistent level in the editor #UE4

In addition:
- exposed UNavigationSystem::OnNavigationBoundsUpdated for BP use and changed the param to more specific ANavMeshBoundsVolume
- made navmesh bounds changes application happen in NavigationSystem's tick rather than instantly
- hid some book-keeping functions from public use

diff --git a/Engine/Source/Runtime/Engine/Classes/AI/Navigation/NavMeshBoundsVolume.h b/Engine/Source/Runtime/Engine/Classes/AI/Navigation/NavMeshBoundsVolume.h
--- a/Engine/Source/Runtime/Engine/Classes/AI/Navigation/NavMeshBoundsVolume.h
+++ b/Engine/Source/Runtime/Engine/Classes/AI/Navigation/NavMeshBoundsVolume.h
@@ -16,6 +16,7 @@ class ANavMeshBoundsVolume : public AVolume
 
 	// Begin AActor Interface
 	virtual void PostInitializeComponents() override;
+	virtual void PostLoad() override;
 	// End AActor Interface
 #if WITH_EDITOR
 	// Begin UObject Interface
diff --git a/Engine/Source/Runtime/Engine/Classes/AI/Navigation/NavigationSystem.h b/Engine/Source/Runtime/Engine/Classes/AI/Navigation/NavigationSystem.h
--- a/Engine/Source/Runtime/Engine/Classes/AI/Navigation/NavigationSystem.h
+++ b/Engine/Source/Runtime/Engine/Classes/AI/Navigation/NavigationSystem.h
@@ -31,7 +31,7 @@ class UCrowdManager;
 class UNavArea;
 class INavLinkCustomInterface;
 class FNavigationOctree;
-class AVolume;
+class ANavMeshBoundsVolume;
 class FNavDataGenerator;
 class AWorldSettings;
 #if WITH_EDITOR
@@ -155,6 +155,9 @@ class ENGINE_API UNavigationSystem : public UBlueprintFunctionLibrary
 
 	TSet<FNavigationDirtyElement> PendingOctreeUpdates;
 
+	UPROPERTY(transient)
+	TArray<ANavMeshBoundsVolume*> PendingNavVolumeUpdates;
+
  	UPROPERTY(/*BlueprintAssignable, */Transient)
 	FOnNavDataRegistered OnNavDataRegisteredEvent;
 
@@ -395,11 +398,15 @@ public:
 	/** adds NavData to registration candidates queue - NavDataRegistrationQueue */
 	void RequestRegistration(ANavigationData* NavData, bool bTriggerRegistrationProcessing = true);
 
+protected:
 	/** Processes registration of candidates queues via RequestRegistration and stored in NavDataRegistrationQueue */
 	void ProcessRegistrationCandidates();
 
 	/** registers NavArea classes awaiting registration in PendingNavAreaRegistration */
 	void ProcessNavAreaPendingRegistration();
+
+	/** used to apply updates of nav volumes in navigation system's tick */
+	void PerformNavigationBoundsUpdate(ANavMeshBoundsVolume* NavVolume);
 	
 	/** @return pointer to ANavigationData instance of given ID, or NULL if it was not found. Note it looks only through registered navigation data */
 	ANavigationData* GetNavDataWithID(const uint16 NavDataID) const;
@@ -409,6 +416,7 @@ public:
 	//----------------------------------------------------------------------//
 	// navigation octree related functions
 	//----------------------------------------------------------------------//
+public:
 	FSetElementId RegisterNavigationRelevantActor(AActor* Actor, int32 UpdateFlags = OctreeUpdate_Default);
 	void UnregisterNavigationRelevantActor(AActor* Actor, int32 UpdateFlags = OctreeUpdate_Default);
 
@@ -492,7 +500,8 @@ public:
 	FORCEINLINE bool IsNavigationBuildingLocked() const { return bNavigationBuildingLocked || bInitialBuildingLockActive; }
 
 	// @todo document
-	void OnNavigationBoundsUpdated(AVolume* NavVolume);
+	UFUNCTION(BlueprintCallable, Category = Navigation)
+	void OnNavigationBoundsUpdated(ANavMeshBoundsVolume* NavVolume);
 
 	/** Used to display "navigation building in progress" notify */
 	bool IsNavigationBuildInProgress(bool bCheckDirtyToo = true);
diff --git a/Engine/Source/Runtime/Engine/Private/AI/Navigation/NavigationSystem.cpp b/Engine/Source/Runtime/Engine/Private/AI/Navigation/NavigationSystem.cpp
--- a/Engine/Source/Runtime/Engine/Private/AI/Navigation/NavigationSystem.cpp
+++ b/Engine/Source/Runtime/Engine/Private/AI/Navigation/NavigationSystem.cpp
@@ -475,6 +475,18 @@ void UNavigationSystem::Tick(float DeltaSeconds)
 		ProcessNavAreaPendingRegistration();
 	}
 
+	if (PendingNavVolumeUpdates.Num() > 0)
+	{
+		for (auto Volume : PendingNavVolumeUpdates)
+		{
+			if (Volume != NULL && Volume->IsPendingKillPending() == false)
+			{
+				PerformNavigationBoundsUpdate(Volume);
+			}
+		}
+		PendingNavVolumeUpdates.Reset();
+	}
+
 	if (PendingOctreeUpdates.Num() > 0)
 	{
 		SCOPE_CYCLE_COUNTER(STAT_Navigation_AddingActorsToNavOctree);
@@ -2363,79 +2375,72 @@ void UNavigationSystem::OnEditorModeChanged(FEdMode* Mode, bool IsEntering)
 #endif
 
 #if WITH_NAVIGATION_GENERATOR
-void UNavigationSystem::OnNavigationBoundsUpdated(AVolume* NavVolume)
+void UNavigationSystem::OnNavigationBoundsUpdated(ANavMeshBoundsVolume* NavVolume)
 {
-	if (Cast<ANavMeshBoundsVolume>(NavVolume) != NULL)
+	if (NavVolume == NULL)
 	{
-		const bool bIsInGame = GIsEditor == false || NavVolume->GetWorld()->IsPlayInEditor();
+		return;
+	}
+
+	PendingNavVolumeUpdates.AddUnique(NavVolume);
+}
+
+void UNavigationSystem::PerformNavigationBoundsUpdate(ANavMeshBoundsVolume* NavVolume)
+{
+	const bool bIsInGame = GIsEditor == false || NavVolume->GetWorld()->IsPlayInEditor();
 
-		if (bNavDataRemovedDueToMissingNavBounds)
+	if (bNavDataRemovedDueToMissingNavBounds)
+	{
+		PopulateNavOctree();
+		bNavDataRemovedDueToMissingNavBounds = false;
+	}
+	
+	if (NavDataSet.Num() == 0)
+	{
+		if (NavDataRegistrationQueue.Num() > 0)
 		{
-			PopulateNavOctree();
-			bNavDataRemovedDueToMissingNavBounds = false;
+			ProcessRegistrationCandidates();
 		}
 
-//#if WITH_EDITOR
-//		if (GEditor)
-//		{
-//			// this is done just in case navigation data has been removed 
-//			for (int32 Index = NavDataSet.Num() - 1; Index >= 0; --Index)
-//			{
-//				if (NavDataSet[Index] == NULL || NavDataSet[Index]->IsPendingKillPending())
-//				{
-//					NavDataSet.RemoveAtSwap(Index, 1, false);
-//				}
-//			}
-//		}
-//#endif // WITH_EDITOR
-	
 		if (NavDataSet.Num() == 0)
 		{
-			if (NavDataRegistrationQueue.Num() > 0)
-			{
-				ProcessRegistrationCandidates();
-			}
+			SpawnMissingNavigationData();
+			ProcessRegistrationCandidates();
 
-			if (NavDataSet.Num() == 0)
+			for (int32 NavDataIndex = 0; NavDataIndex < NavDataSet.Num(); ++NavDataIndex)
 			{
-				SpawnMissingNavigationData();
-				ProcessRegistrationCandidates();
-
-				for (int32 NavDataIndex = 0; NavDataIndex < NavDataSet.Num(); ++NavDataIndex)
+				ANavigationData* NavData = NavDataSet[NavDataIndex];
+				if (NavData != NULL && (NavData->bRebuildAtRuntime || (GIsEditor && !bIsInGame)))
 				{
-					ANavigationData* NavData = NavDataSet[NavDataIndex];
-					if (NavData != NULL && (NavData->bRebuildAtRuntime || (GIsEditor && !bIsInGame)))
-					{
-						NavData->GetGenerator(FNavigationSystem::Create);
-					}
+					NavData->GetGenerator(FNavigationSystem::Create);
 				}
 			}
 		}
+	}
 
 #if WITH_RECAST
-		if (IsNavigationBuildingLocked() == false)
+	if (IsNavigationBuildingLocked() == false)
+	{
+		// rebuild all navmeshes
+		// NOTE: this will most probably be done in editor, if required to be performed at runtime may require some optimizations
+		for (int32 NavDataIndex = 0; NavDataIndex < NavDataSet.Num(); ++NavDataIndex)
 		{
-			// rebuild all navmeshes
-			// NOTE: this will most probably be done in editor, if required to be performed at runtime may require some optimizations
-			for (int32 NavDataIndex = 0; NavDataIndex < NavDataSet.Num(); ++NavDataIndex)
+			ANavigationData* NavData = NavDataSet[NavDataIndex];
+			check(NavData);
+			if (NavData->GetGenerator() != NULL)
 			{
-				ANavigationData* NavData = NavDataSet[NavDataIndex];
-				check(NavData);
-				if (NavData->GetGenerator() != NULL)
+				if ( Cast<ARecastNavMesh>(NavData) != NULL && (bIsInGame == false || ((ARecastNavMesh*)NavData)->bRebuildAtRuntime) )
 				{
-					if ( Cast<ARecastNavMesh>(NavData) != NULL && (bIsInGame == false || ((ARecastNavMesh*)NavData)->bRebuildAtRuntime) )
-					{
-						// there are two possible cases here:
-						// 1. new volume is totally inside original generation bounds
-						// 2. it expands current bounds.
-						// @TODO currently handling every case as 1), 2) is an optimization
-						NavData->GetGenerator(FNavigationSystem::Create)->OnNavigationBoundsUpdated(NavVolume);
-					}
+					// there are two possible cases here:
+					// 1. new volume is totally inside original generation bounds
+					// 2. it expands current bounds.
+					// @TODO currently handling every case as 1), 2) is an optimization
+					NavData->GetGenerator(FNavigationSystem::Create)->OnNavigationBoundsUpdated(NavVolume);
 				}
 			}
 		}
-#endif // WITH_RECAST
 	}
+#endif // WITH_RECAST
 }
 
 void UNavigationSystem::Build()
diff --git a/Engine/Source/Runtime/Engine/Private/NavMeshBoundsVolume.cpp b/Engine/Source/Runtime/Engine/Private/NavMeshBoundsVolume.cpp
--- a/Engine/Source/Runtime/Engine/Private/NavMeshBoundsVolume.cpp
+++ b/Engine/Source/Runtime/Engine/Private/NavMeshBoundsVolume.cpp
@@ -49,3 +49,16 @@ void ANavMeshBoundsVolume::PostInitializeComponents()
 	}
 #endif
 }
+
+void ANavMeshBoundsVolume::PostLoad()
+{
+	Super::PostLoad();
+
+#if WITH_NAVIGATION_GENERATOR
+	UNavigationSystem* NavSys = UNavigationSystem::GetCurrent(GetWorld());
+	if (NavSys && Role == ROLE_Authority)
+	{
+		NavSys->OnNavigationBoundsUpdated(this);
+	}
+#endif
+}
\ No newline at end of file
