NUX Fix - A bool variable doesnÆt automatically convert to text.  Also removed parameters to formating ints to text that make no sense.  Also added converting bytes to text.

diff --git a/Engine/Source/Runtime/Engine/Classes/Kismet/KismetTextLibrary.h b/Engine/Source/Runtime/Engine/Classes/Kismet/KismetTextLibrary.h
--- a/Engine/Source/Runtime/Engine/Classes/Kismet/KismetTextLibrary.h
+++ b/Engine/Source/Runtime/Engine/Classes/Kismet/KismetTextLibrary.h
@@ -105,10 +105,18 @@ class UKismetTextLibrary : public UBlueprintFunctionLibrary
 	UFUNCTION(BlueprintPure, meta=(FriendlyName = "NotEqual, Case Insensitive (text)", CompactNodeTitle = "!="), Category="Utilities|Text")
 	static bool NotEqual_IgnoreCase_TextText(const FText& A, const FText& B);
 
+	/** Converts a boolean value to text, either 'true' or 'false' */
+	UFUNCTION(BlueprintPure, meta=(FriendlyName = "ToText (boolean)", CompactNodeTitle = "->"), Category="Utilities|Text")
+	static FText Conv_BoolToText(bool InBool);
+
+	/** Converts a byte value to text */
+	UFUNCTION(BlueprintPure, meta=(FriendlyName = "ToText (byte)", CompactNodeTitle = "->"), Category="Utilities|Text")
+	static FText Conv_ByteToText(uint8 Value);
+
 	// Default values are duplicated from FNumberFormattingOptions and should be replicated in all functions and in the struct when changed!
 	/* Converts a passed in integer to a text based on formatting options */
 	UFUNCTION(BlueprintPure, meta=(FriendlyName = "ToText (int)", AdvancedDisplay = "1"), Category="Utilities|Text")
-	static FText Conv_IntToText(int32 Value, TEnumAsByte<ERoundingMode> RoundingMode, bool bUseGrouping = true, int32 MinimumIntegralDigits = 1, int32 MaximumIntegralDigits = 324, int32 MinimumFractionalDigits = 0, int32 MaximumFractionalDigits = 3);
+	static FText Conv_IntToText(int32 Value, bool bUseGrouping = true, int32 MinimumIntegralDigits = 1, int32 MaximumIntegralDigits = 324);
 
 	// Default values are duplicated from FNumberFormattingOptions and should be replicated in all functions and in the struct when changed!
 	/* Converts a passed in float to a text based on formatting options */
diff --git a/Engine/Source/Runtime/Engine/Private/KismetTextLibrary.cpp b/Engine/Source/Runtime/Engine/Private/KismetTextLibrary.cpp
--- a/Engine/Source/Runtime/Engine/Private/KismetTextLibrary.cpp
+++ b/Engine/Source/Runtime/Engine/Private/KismetTextLibrary.cpp
@@ -2,6 +2,8 @@
 
 #include "EnginePrivate.h"
 
+#define LOCTEXT_NAMESPACE "Kismet"
+
 //////////////////////////////////////////////////////////////////////////
 // UKismetTextLibrary
 
@@ -85,15 +87,33 @@ bool UKismetTextLibrary::NotEqual_TextText(const FText& A, const FText& B)
 	return !A.EqualTo( B );
 }
 
-FText UKismetTextLibrary::Conv_IntToText(int32 Value, TEnumAsByte<ERoundingMode> RoundingMode, bool bUseGrouping/* = true*/, int32 MinimumIntegralDigits/* = 1*/, int32 MaximumIntegralDigits/* = 324*/, int32 MinimumFractionalDigits/* = 0*/, int32 MaximumFractionalDigits/* = 3*/)
+FText UKismetTextLibrary::Conv_BoolToText(bool InBool)
+{
+	return InBool ? LOCTEXT("True", "true") : LOCTEXT("False", "false");
+}
+
+FText UKismetTextLibrary::Conv_ByteToText(uint8 Value)
+{
+	FNumberFormattingOptions NumberFormatOptions;
+	NumberFormatOptions.UseGrouping = true;
+	NumberFormatOptions.RoundingMode = ERoundingMode::ToZero;
+	NumberFormatOptions.MinimumIntegralDigits = 1;
+	NumberFormatOptions.MaximumIntegralDigits = 3;
+	NumberFormatOptions.MinimumFractionalDigits = 0;
+	NumberFormatOptions.MaximumFractionalDigits = 0;
+
+	return FText::AsNumber(Value, &NumberFormatOptions);
+}
+
+FText UKismetTextLibrary::Conv_IntToText(int32 Value, bool bUseGrouping/* = true*/, int32 MinimumIntegralDigits/* = 1*/, int32 MaximumIntegralDigits/* = 324*/)
 {
 	FNumberFormattingOptions NumberFormatOptions;
 	NumberFormatOptions.UseGrouping = bUseGrouping;
-	NumberFormatOptions.RoundingMode = RoundingMode;
+	NumberFormatOptions.RoundingMode = ERoundingMode::FromZero;
 	NumberFormatOptions.MinimumIntegralDigits = MinimumIntegralDigits;
 	NumberFormatOptions.MaximumIntegralDigits = MaximumIntegralDigits;
-	NumberFormatOptions.MinimumFractionalDigits = MinimumFractionalDigits;
-	NumberFormatOptions.MaximumFractionalDigits = MaximumFractionalDigits;
+	NumberFormatOptions.MinimumFractionalDigits = 0;
+	NumberFormatOptions.MaximumFractionalDigits = 0;
 
 	return FText::AsNumber(Value, &NumberFormatOptions);
 }
@@ -164,3 +184,5 @@ FText UKismetTextLibrary::Format(FText InPattern, TArray<FFormatTextArgument> In
 
 	return FText::Format(InPattern, Arguments);
 }
+
+#undef LOCTEXT_NAMESPACE
