Remove VS2010 bug workaround #if

diff --git a/Engine/Source/Runtime/Engine/Classes/GameFramework/Actor.h b/Engine/Source/Runtime/Engine/Classes/GameFramework/Actor.h
--- a/Engine/Source/Runtime/Engine/Classes/GameFramework/Actor.h
+++ b/Engine/Source/Runtime/Engine/Classes/GameFramework/Actor.h
@@ -1047,11 +1047,7 @@ public:
 		this is a template for no other reason than to delay compilation until USceneComponent is defined
 	*/ 
 	template<class T>
-	static 
-#if !defined(_MSC_VER) || _MSC_VER != 1600  //VS210 bug 
-		FORCEINLINE
-#endif
-	FVector GetActorLocation(T* RootComponent)
+	static FORCEINLINE FVector GetActorLocation(T* RootComponent)
 	{
 		FVector Result(0.f);
 		if( RootComponent != NULL )
@@ -1065,11 +1061,7 @@ public:
 		this is a template for no other reason than to delay compilation until USceneComponent is defined
 	*/ 
 	template<class T>
-	static 
-#if !defined(_MSC_VER) || _MSC_VER != 1600  //VS210 bug 
-		FORCEINLINE
-#endif
-	FRotator GetActorRotation(T* RootComponent)
+	static FORCEINLINE FRotator GetActorRotation(T* RootComponent)
 	{
 		FRotator Result(0,0,0);
 		if( RootComponent != NULL )
@@ -1083,11 +1075,7 @@ public:
 		this is a template for no other reason than to delay compilation until USceneComponent is defined
 	*/ 
 	template<class T>
-	static 
-#if !defined(_MSC_VER) || _MSC_VER != 1600  //VS210 bug 
-		FORCEINLINE
-#endif
-	FVector GetActorScale(T* RootComponent)
+	static FORCEINLINE FVector GetActorScale(T* RootComponent)
 	{
 		FVector Result(1,1,1);
 		if( RootComponent != NULL )
@@ -1101,11 +1089,7 @@ public:
 		this is a template for no other reason than to delay compilation until USceneComponent is defined
 	 */ 
 	template<class T>
-	static 
-#if !defined(_MSC_VER) || _MSC_VER != 1600  //VS210 bug 
-		FORCEINLINE
-#endif
-	FQuat GetActorQuat(T* RootComponent)
+	static FORCEINLINE FQuat GetActorQuat(T* RootComponent)
 	{
 		FQuat Result(ForceInit);
 		if( RootComponent != NULL )
