Deprecate Game Viewport Client FVector2D GetMousePosition() and replaced with bool GetMousePosition(FVector2D&) which returns false when the mouse cursor is not over the viewport
Game Viewport Client now tracks when mouse is over it or not
DeprojectMousePositionToWorld and DeprojectScreenPositionToWorld both now return false when unable to deproject, due lack of mouse cursor, valid scene view, etc.
GetMousePosition now returns false when a mouse is not attached

diff --git a/Engine/Source/Runtime/Engine/Classes/Engine/GameViewportClient.h b/Engine/Source/Runtime/Engine/Classes/Engine/GameViewportClient.h
--- a/Engine/Source/Runtime/Engine/Classes/Engine/GameViewportClient.h
+++ b/Engine/Source/Runtime/Engine/Classes/Engine/GameViewportClient.h
@@ -205,6 +205,9 @@ public:
 	/** set to disable world rendering */
 	uint32 bDisableWorldRendering:1;
 
+	/** Whether the mouse is currently over this game viewport */
+	uint32 bMouseOverViewport:1;
+
 protected:
 	/**
 	 * The splitscreen type that is actually being used; takes into account the number of players and other factors (such as cinematic mode)
@@ -392,13 +395,17 @@ public:
 	 *
 	 * @param	out_ViewportSize	[out] will be filled in with the size of the main viewport
 	 */
-	void GetViewportSize( FVector2D& ViewportSize );
+	void GetViewportSize( FVector2D& ViewportSize ) const;
 
 	/** @return Whether or not the main viewport is fullscreen or windowed. */
-	bool IsFullScreenViewport();
+	bool IsFullScreenViewport() const;
+
+	/** @return mouse position in game viewport coordinates (does not account for splitscreen) */
+	DEPRECATED(4.5, "Use GetMousePosition that returns a boolean if mouse is in window instead.")
+	FVector2D GetMousePosition() const;
 
 	/** @return mouse position in game viewport coordinates (does not account for splitscreen) */
-	FVector2D GetMousePosition();
+	bool GetMousePosition(FVector2D& MousePosition) const;
 
 	/**
 	 * Determine whether a fullscreen viewport should be used in cases where there are multiple players.
diff --git a/Engine/Source/Runtime/Engine/Classes/GameFramework/PlayerController.h b/Engine/Source/Runtime/Engine/Classes/GameFramework/PlayerController.h
--- a/Engine/Source/Runtime/Engine/Classes/GameFramework/PlayerController.h
+++ b/Engine/Source/Runtime/Engine/Classes/GameFramework/PlayerController.h
@@ -360,13 +360,13 @@ public:
 	UFUNCTION(BlueprintCallable, Category="Game|Player", meta=(bTraceComplex=true))
 	bool GetHitResultUnderFingerForObjects(ETouchIndex::Type FingerIndex, const  TArray<TEnumAsByte<EObjectTypeQuery> > & ObjectTypes, bool bTraceComplex, FHitResult& HitResult) const;
 
-	/** Convert current mouse 2D position to World Space 3D position and direction. **/
+	/** Convert current mouse 2D position to World Space 3D position and direction. Returns false if unable to determine value. **/
 	UFUNCTION(BlueprintCallable, Category="Game|Player", meta=(FriendlyName="ConvertMouseLocationToWorldSpace"))
-	void DeprojectMousePositionToWorld(FVector & WorldLocation, FVector & WorldDirection) const;
+	bool DeprojectMousePositionToWorld(FVector & WorldLocation, FVector & WorldDirection) const;
 
-	/** Convert current mouse 2D position to World Space 3D position and direction. **/
+	/** Convert current mouse 2D position to World Space 3D position and direction. Returns false if unable to determine value. **/
 	UFUNCTION(BlueprintCallable, Category = "Game|Player", meta = (FriendlyName = "ConvertScreenLocationToWorldSpace"))
-	void DeprojectScreenPositionToWorld(float ScreenX, float ScreenY, FVector & WorldLocation, FVector & WorldDirection) const;
+	bool DeprojectScreenPositionToWorld(float ScreenX, float ScreenY, FVector & WorldLocation, FVector & WorldDirection) const;
 
 	/**
 	 * Convert a World Space 3D position into a 2D Screen Space position.
diff --git a/Engine/Source/Runtime/Engine/Private/GameViewportClient.cpp b/Engine/Source/Runtime/Engine/Private/GameViewportClient.cpp
--- a/Engine/Source/Runtime/Engine/Private/GameViewportClient.cpp
+++ b/Engine/Source/Runtime/Engine/Private/GameViewportClient.cpp
@@ -369,6 +369,8 @@ void UGameViewportClient::MouseEnter(FViewport* InViewport, int32 x, int32 y)
 {
 	Super::MouseEnter(InViewport, x, y);
 
+	bMouseOverViewport = true;
+
 	if (GetDefault<UInputSettings>()->bUseMouseForTouch && !GetGameViewport()->GetPlayInEditorIsSimulate())
 	{
 		FSlateApplication::Get().SetGameIsFakingTouchEvents(true);
@@ -379,6 +381,8 @@ void UGameViewportClient::MouseLeave(FViewport* InViewport)
 {
 	Super::MouseLeave(InViewport);
 
+	bMouseOverViewport = false;
+
 	if (GetDefault<UInputSettings>()->bUseMouseForTouch)
 	{
 		UE_LOG(LogTemp, Log, TEXT("MouseLeave"));
@@ -391,18 +395,30 @@ void UGameViewportClient::MouseLeave(FViewport* InViewport)
 	}
 }
 
-FVector2D UGameViewportClient::GetMousePosition()
+bool UGameViewportClient::GetMousePosition(FVector2D& MousePosition) const
 {
-	if (Viewport == NULL)
-	{
-		return FVector2D(0.f, 0.f);
-	}
-	else
+	bool bGotMousePosition = false;
+
+	if (Viewport && bMouseOverViewport && FSlateApplication::Get().IsMouseAttached())
 	{
 		FIntPoint MousePos;
 		Viewport->GetMousePos(MousePos);
-		return FVector2D(MousePos);
+		MousePosition = FVector2D(MousePos);
+		bGotMousePosition = true;
 	}
+
+	return bGotMousePosition;
+}
+
+FVector2D UGameViewportClient::GetMousePosition() const
+{
+	FVector2D MousePosition;
+	if (!GetMousePosition(MousePosition))
+	{
+		MousePosition = FVector2D::ZeroVector;
+	}
+
+	return MousePosition;
 }
 
 
@@ -539,9 +555,7 @@ void UGameViewportClient::SetViewport( FViewport* InViewport )
 	}
 }
 
-
-
-void UGameViewportClient::GetViewportSize( FVector2D& out_ViewportSize )
+void UGameViewportClient::GetViewportSize( FVector2D& out_ViewportSize ) const
 {
 	if ( Viewport != NULL )
 	{
@@ -550,13 +564,11 @@ void UGameViewportClient::GetViewportSize( FVector2D& out_ViewportSize )
 	}
 }
 
-bool UGameViewportClient::IsFullScreenViewport()
+bool UGameViewportClient::IsFullScreenViewport() const
 {
 	return Viewport->IsFullscreen();
 }
 
-
-
 bool UGameViewportClient::ShouldForceFullscreenViewport() const
 {
 	bool bResult = false;
diff --git a/Engine/Source/Runtime/Engine/Private/HUD.cpp b/Engine/Source/Runtime/Engine/Private/HUD.cpp
--- a/Engine/Source/Runtime/Engine/Private/HUD.cpp
+++ b/Engine/Source/Runtime/Engine/Private/HUD.cpp
@@ -135,10 +135,13 @@ void AHUD::PostRender()
 		{
 			TArray<FVector2D> ContactPoints;
 
-			// TODO: This needs a check whether the platform has a mouse at all, but there is no way to currently do that
 			if (!FSlateApplication::Get().IsFakingTouchEvents())
 			{
-				ContactPoints.Add(LocalPlayer->ViewportClient->GetMousePosition());
+				FVector2D MousePosition;
+				if (LocalPlayer->ViewportClient->GetMousePosition(MousePosition))
+				{
+					ContactPoints.Add(MousePosition);
+				}
 			}
 
 			for (int32 FingerIndex = 0; FingerIndex < EKeys::NUM_TOUCH_KEYS; ++FingerIndex)
diff --git a/Engine/Source/Runtime/Engine/Private/PlayerController.cpp b/Engine/Source/Runtime/Engine/Private/PlayerController.cpp
--- a/Engine/Source/Runtime/Engine/Private/PlayerController.cpp
+++ b/Engine/Source/Runtime/Engine/Private/PlayerController.cpp
@@ -1700,7 +1700,11 @@ bool APlayerController::GetHitResultUnderCursor(ECollisionChannel TraceChannel,
 	bool bHit = false;
 	if (LocalPlayer)
 	{
-		bHit = GetHitResultAtScreenPosition(LocalPlayer->ViewportClient->GetMousePosition(), TraceChannel, bTraceComplex, HitResult);
+		FVector2D MousePosition;
+		if (LocalPlayer->ViewportClient->GetMousePosition(MousePosition))
+		{
+			bHit = GetHitResultAtScreenPosition(MousePosition, TraceChannel, bTraceComplex, HitResult);
+		}
 	}
 
 	if(!bHit)	//If there was no hit we reset the results. This is redundent but helps Blueprint users
@@ -1717,7 +1721,11 @@ bool APlayerController::GetHitResultUnderCursorByChannel(ETraceTypeQuery TraceCh
 	bool bHit = false;
 	if (LocalPlayer)
 	{
-		bHit = GetHitResultAtScreenPosition(LocalPlayer->ViewportClient->GetMousePosition(), TraceChannel, bTraceComplex, HitResult);
+		FVector2D MousePosition;
+		if (LocalPlayer->ViewportClient->GetMousePosition(MousePosition))
+		{
+			bHit = GetHitResultAtScreenPosition(MousePosition, TraceChannel, bTraceComplex, HitResult);
+		}
 	}
 
 	if(!bHit)	//If there was no hit we reset the results. This is redundent but helps Blueprint users
@@ -1734,7 +1742,11 @@ bool APlayerController::GetHitResultUnderCursorForObjects(const TArray<TEnumAsBy
 	bool bHit = false;
 	if (LocalPlayer)
 	{
-		bHit = GetHitResultAtScreenPosition(LocalPlayer->ViewportClient->GetMousePosition(), ObjectTypes, bTraceComplex, HitResult);
+		FVector2D MousePosition;
+		if (LocalPlayer->ViewportClient->GetMousePosition(MousePosition))
+		{
+			bHit = GetHitResultAtScreenPosition(MousePosition, ObjectTypes, bTraceComplex, HitResult);
+		}
 	}
 
 	if(!bHit)	//If there was no hit we reset the results. This is redundent but helps Blueprint users
@@ -1811,21 +1823,30 @@ bool APlayerController::GetHitResultUnderFingerForObjects(ETouchIndex::Type Fing
 	return bHit;
 }
 
-void APlayerController::DeprojectMousePositionToWorld(FVector & WorldLocation, FVector & WorldDirection) const
+bool APlayerController::DeprojectMousePositionToWorld(FVector & WorldLocation, FVector & WorldDirection) const
 {
+	bool bSuccessfulDeproject = false;
+
 	ULocalPlayer* LocalPlayer = Cast<ULocalPlayer>(Player);
-	if (LocalPlayer != NULL && LocalPlayer->ViewportClient != NULL)
+	if (LocalPlayer && LocalPlayer->ViewportClient )
 	{
-		FVector2D ScreenPosition = LocalPlayer->ViewportClient->GetMousePosition();
-		DeprojectScreenPositionToWorld(ScreenPosition.X, ScreenPosition.Y, WorldLocation, WorldDirection);
+		FVector2D ScreenPosition;
+		if (LocalPlayer->ViewportClient->GetMousePosition(ScreenPosition))
+		{
+			bSuccessfulDeproject = DeprojectScreenPositionToWorld(ScreenPosition.X, ScreenPosition.Y, WorldLocation, WorldDirection);
+		}
 	}
+
+	return bSuccessfulDeproject;
 }
 
-void APlayerController::DeprojectScreenPositionToWorld(float ScreenX, float ScreenY, FVector & WorldLocation, FVector & WorldDirection) const
+bool APlayerController::DeprojectScreenPositionToWorld(float ScreenX, float ScreenY, FVector & WorldLocation, FVector & WorldDirection) const
 {
+	bool bSuccessfulDeproject = false;
+
 	ULocalPlayer* LocalPlayer = Cast<ULocalPlayer>(Player);
 
-	if (LocalPlayer != NULL && LocalPlayer->ViewportClient != NULL && LocalPlayer->ViewportClient->Viewport != NULL)
+	if (LocalPlayer && LocalPlayer->ViewportClient && LocalPlayer->ViewportClient->Viewport)
 	{
 		// Create a view family for the game viewport
 		FSceneViewFamilyContext ViewFamily(FSceneViewFamily::ConstructionValues(
@@ -1843,8 +1864,12 @@ void APlayerController::DeprojectScreenPositionToWorld(float ScreenX, float Scre
 		{
 			const FVector2D ScreenPosition(ScreenX, ScreenY);
 			SceneView->DeprojectFVector2D(ScreenPosition, WorldLocation, WorldDirection);
+
+			bSuccessfulDeproject = true;
 		}
 	}
+
+	return bSuccessfulDeproject;
 }
 
 
@@ -2004,8 +2029,14 @@ void APlayerController::PlayerTick( float DeltaTime )
 	{
 		if (bEnableMouseOverEvents)
 		{
+			FVector2D MousePosition;
 			FHitResult HitResult;
-			const bool bHit = GetHitResultAtScreenPosition(LocalPlayer->ViewportClient->GetMousePosition(), CurrentClickTraceChannel, true, /*out*/ HitResult);
+			bool bHit = false;
+			
+			if (LocalPlayer->ViewportClient->GetMousePosition(MousePosition))
+			{
+				bHit = GetHitResultAtScreenPosition(MousePosition, CurrentClickTraceChannel, true, /*out*/ HitResult);
+			}
 
 			UPrimitiveComponent* PreviousComponent = CurrentClickablePrimitive.Get();
 			UPrimitiveComponent* CurrentComponent = (bHit ? HitResult.Component.Get() : NULL);
@@ -2100,7 +2131,8 @@ bool APlayerController::InputKey(FKey Key, EInputEvent EventType, float AmountDe
 		// TODO: Allow click key(s?) to be defined
 		if (bEnableClickEvents && Key == EKeys::LeftMouseButton)
 		{
-			const FVector2D MousePosition = CastChecked<ULocalPlayer>(Player)->ViewportClient->GetMousePosition();
+			FVector2D MousePosition;
+			verify(CastChecked<ULocalPlayer>(Player)->ViewportClient->GetMousePosition(MousePosition));
 			UPrimitiveComponent* ClickedPrimitive = NULL;
 			if (bEnableMouseOverEvents)
 			{
@@ -4093,17 +4125,23 @@ float APlayerController::GetInputKeyTimeDown(const FKey Key) const
 
 bool APlayerController::GetMousePosition(float& LocationX, float& LocationY) const
 {
-	// TODO: Return false if there is no mouse attached
-	if (Player == nullptr)
+	bool bGotMousePosition = false;
+	ULocalPlayer* LocalPlayer = Cast<ULocalPlayer>(Player);
+
+	if (LocalPlayer)
 	{
-		return false;
-	}
+		FVector2D MousePosition;
+		
+		bGotMousePosition = LocalPlayer->ViewportClient->GetMousePosition(MousePosition);
 
-	const FVector2D MousePosition = CastChecked<ULocalPlayer>(Player)->ViewportClient->GetMousePosition();
-	LocationX = MousePosition.X;
-	LocationY = MousePosition.Y;
+		if (bGotMousePosition)
+		{
+			LocationX = MousePosition.X;
+			LocationY = MousePosition.Y;
+		}
+	}
 
-	return true;
+	return bGotMousePosition;
 }
 
 void APlayerController::GetInputMouseDelta(float& DeltaX, float& DeltaY) const
