Deprecate GetRootPrimitiveComponent.  Use GetRootComponent and cast as appropriate.

diff --git a/Engine/Source/Editor/UnrealEd/Private/Editor/ActorPositioning.cpp b/Engine/Source/Editor/UnrealEd/Private/Editor/ActorPositioning.cpp
--- a/Engine/Source/Editor/UnrealEd/Private/Editor/ActorPositioning.cpp
+++ b/Engine/Source/Editor/UnrealEd/Private/Editor/ActorPositioning.cpp
@@ -45,7 +45,7 @@ bool IsHitIgnored(const FHitResult& InHit, const FSceneView& InSceneView)
 	const auto* Actor = InHit.GetActor();
 	
 	// Try and find a primitive component for the hit
-	const UPrimitiveComponent* PrimitiveComponent = Actor ? Actor->GetRootPrimitiveComponent() : nullptr;
+	const UPrimitiveComponent* PrimitiveComponent = Actor ? Cast<UPrimitiveComponent>(Actor->GetRootComponent()) : nullptr;
 
 	if (!PrimitiveComponent)
 	{
diff --git a/Engine/Source/Editor/UnrealEd/Private/PhysicsManipulationMode.cpp b/Engine/Source/Editor/UnrealEd/Private/PhysicsManipulationMode.cpp
--- a/Engine/Source/Editor/UnrealEd/Private/PhysicsManipulationMode.cpp
+++ b/Engine/Source/Editor/UnrealEd/Private/PhysicsManipulationMode.cpp
@@ -14,7 +14,7 @@ void FPhysicsManipulationEdModeFactory::OnSelectionChanged(FEditorModeTools& Too
 		AActor* SelectedActor = Cast<AActor>(ItemUndergoingChange);
 		if (SelectedActor != NULL)
 		{
-			UPrimitiveComponent* PC = SelectedActor->GetRootPrimitiveComponent();
+			UPrimitiveComponent* PC = Cast<UPrimitiveComponent>(SelectedActor->GetRootComponent());
 			if (PC != NULL && PC->BodyInstance.bSimulatePhysics)
 			{
 				Tools.ActivateMode(FBuiltinEditorModes::EM_Physics);
@@ -96,7 +96,7 @@ bool FPhysicsManipulationEdMode::StartTracking( FEditorViewportClient* InViewpor
 
 		if (SelectedActor != NULL)
 		{
-			UPrimitiveComponent* PC = SelectedActor->GetRootPrimitiveComponent();
+			UPrimitiveComponent* PC = Cast<UPrimitiveComponent>(SelectedActor->GetRootComponent());
 
 			if (PC != NULL && PC->BodyInstance.bSimulatePhysics)
 			{
diff --git a/Engine/Source/Editor/UnrealEd/Private/Tests/EditorAutomationTests.cpp b/Engine/Source/Editor/UnrealEd/Private/Tests/EditorAutomationTests.cpp
--- a/Engine/Source/Editor/UnrealEd/Private/Tests/EditorAutomationTests.cpp
+++ b/Engine/Source/Editor/UnrealEd/Private/Tests/EditorAutomationTests.cpp
@@ -1086,7 +1086,7 @@ bool FStaticMeshValidation::RunTest(const FString& Parameters)
 	AActor* PhysicsActor = FActorFactoryAssetProxy::AddActorForAsset( EditorCubeMesh );
 	PhysicsActor->SetActorRelativeScale3D(FVector(2.0f, 2.0f, .5f));
 	PhysicsActor->GetRootComponent()->SetMobility(EComponentMobility::Movable);
-	PhysicsActor->GetRootPrimitiveComponent()->SetSimulatePhysics(true);
+	CastChecked<UPrimitiveComponent>(PhysicsActor->GetRootComponent())->SetSimulatePhysics(true);
 	PhysicsActor->TeleportTo(FVector(-96.0f, 128.0f, 256.0f), FRotator(0, 0, 0));
 
 	//Skeletal Mesh
diff --git a/Engine/Source/Runtime/AIModule/Private/Navigation/PathFollowingComponent.cpp b/Engine/Source/Runtime/AIModule/Private/Navigation/PathFollowingComponent.cpp
--- a/Engine/Source/Runtime/AIModule/Private/Navigation/PathFollowingComponent.cpp
+++ b/Engine/Source/Runtime/AIModule/Private/Navigation/PathFollowingComponent.cpp
@@ -644,9 +644,10 @@ int32 UPathFollowingComponent::OptimizeSegmentVisibility(int32 StartIndex)
 	MyPawn->GetSimpleCollisionCylinder(Radius, HalfHeight);
 	const FCollisionShape CollisionShape = FCollisionShape::MakeCapsule(Radius, HalfHeight);
 	FCollisionQueryParams CapsuleParams(NAME_NavAreaTrace, false, MyAI->GetCharacter());
-	if (MyPawn->GetRootPrimitiveComponent())
+	UPrimitiveComponent* PrimitiveComponent = Cast<UPrimitiveComponent>(MyPawn->GetRootComponent());
+	if (PrimitiveComponent)
 	{
-		CapsuleParams.AddIgnoredActors(MyPawn->GetRootPrimitiveComponent()->MoveIgnoreActors);
+		CapsuleParams.AddIgnoredActors(PrimitiveComponent->MoveIgnoreActors);
 	}
 
 #endif
diff --git a/Engine/Source/Runtime/Engine/Classes/GameFramework/Actor.h b/Engine/Source/Runtime/Engine/Classes/GameFramework/Actor.h
--- a/Engine/Source/Runtime/Engine/Classes/GameFramework/Actor.h
+++ b/Engine/Source/Runtime/Engine/Classes/GameFramework/Actor.h
@@ -1103,6 +1103,7 @@ public:
 	FORCEINLINE class USceneComponent* GetRootComponent() const { return RootComponent; }
 
 	/** Returns this actor's root component cast to a primitive component */
+	DEPRECATED(4.5, "Use GetRootComponent() and cast manually if needed")
 	class UPrimitiveComponent* GetRootPrimitiveComponent() const;
 
 	/** Sets root component to be the specified component.  NewRootComponent's owner should be this actor.  Returns true if successful. */
diff --git a/Engine/Source/Runtime/Engine/Private/EngineTypes.cpp b/Engine/Source/Runtime/Engine/Private/EngineTypes.cpp
--- a/Engine/Source/Runtime/Engine/Private/EngineTypes.cpp
+++ b/Engine/Source/Runtime/Engine/Private/EngineTypes.cpp
@@ -41,7 +41,7 @@ void FDamageEvent::GetBestHitInfo(AActor const* HitActor, AActor const* HitInsti
 		OutHitInfo.Actor = const_cast<AActor*>(HitActor);
 		OutHitInfo.bBlockingHit = true;
 		OutHitInfo.BoneName = NAME_None;
-		OutHitInfo.Component = HitActor->GetRootPrimitiveComponent();
+		OutHitInfo.Component = Cast<UPrimitiveComponent>(HitActor->GetRootComponent());
 		
 		// assume the actor got hit in the center of his root component
 		OutHitInfo.ImpactPoint = HitActor->GetActorLocation();
@@ -145,7 +145,7 @@ USceneComponent* FComponentReference::GetComponent(AActor* OwningActor) const
 			}
 			else
 			{
-				Result = SearchActor->GetRootPrimitiveComponent();
+				Result = Cast<UPrimitiveComponent>(SearchActor->GetRootComponent());
 			}
 		}
 	}
diff --git a/Engine/Source/Runtime/Engine/Private/MovementComponent.cpp b/Engine/Source/Runtime/Engine/Private/MovementComponent.cpp
--- a/Engine/Source/Runtime/Engine/Private/MovementComponent.cpp
+++ b/Engine/Source/Runtime/Engine/Private/MovementComponent.cpp
@@ -82,7 +82,7 @@ void UMovementComponent::InitializeComponent()
 		AActor* MyActor = GetOwner();
 		if (MyActor)
 		{
-			NewUpdatedComponent = MyActor->GetRootPrimitiveComponent();
+			NewUpdatedComponent = Cast<UPrimitiveComponent>(MyActor->GetRootComponent());
 			if (!NewUpdatedComponent)
 			{
 				FMessageLog("PIE").Warning(FText::Format(LOCTEXT("NoRootPrimitiveWarning", "Movement component {0} must update a PrimitiveComponent, but owning actor '{1}' does not have a root PrimitiveComponent. Auto registration failed."), 
diff --git a/Engine/Source/Runtime/Engine/Private/Pawn.cpp b/Engine/Source/Runtime/Engine/Private/Pawn.cpp
--- a/Engine/Source/Runtime/Engine/Private/Pawn.cpp
+++ b/Engine/Source/Runtime/Engine/Private/Pawn.cpp
@@ -160,7 +160,7 @@ AActor* APawn::GetMovementBaseActor(const APawn* Pawn)
 
 bool APawn::CanBeBaseForCharacter(class APawn* APawn) const
 {
-	UPrimitiveComponent* RootPrimitive = GetRootPrimitiveComponent();
+	UPrimitiveComponent* RootPrimitive = Cast<UPrimitiveComponent>(GetRootComponent());
 	if (RootPrimitive && RootPrimitive->CanCharacterStepUpOn != ECB_Owner)
 	{
 		return RootPrimitive->CanCharacterStepUpOn == ECB_Yes;
@@ -1034,7 +1034,7 @@ void APawn::GetLifetimeReplicatedProps( TArray< FLifetimeProperty > & OutLifetim
 
 void APawn::MoveIgnoreActorAdd(AActor * ActorToIgnore)
 {
-	UPrimitiveComponent * RootPrimitiveComponent = GetRootPrimitiveComponent();
+	UPrimitiveComponent * RootPrimitiveComponent = Cast<UPrimitiveComponent>(GetRootComponent());
 	if( RootPrimitiveComponent )
 	{
 		// Remove dead references first
@@ -1048,7 +1048,7 @@ void APawn::MoveIgnoreActorAdd(AActor * ActorToIgnore)
 
 void APawn::MoveIgnoreActorRemove(AActor * ActorToIgnore)
 {
-	UPrimitiveComponent * RootPrimitiveComponent = GetRootPrimitiveComponent();
+	UPrimitiveComponent * RootPrimitiveComponent = Cast<UPrimitiveComponent>(GetRootComponent());
 	if( RootPrimitiveComponent )
 	{
 		// Remove dead references first
diff --git a/Engine/Source/Runtime/Engine/Private/PhysicsEngine/PhysicsConstraintComponent.cpp b/Engine/Source/Runtime/Engine/Private/PhysicsEngine/PhysicsConstraintComponent.cpp
--- a/Engine/Source/Runtime/Engine/Private/PhysicsEngine/PhysicsConstraintComponent.cpp
+++ b/Engine/Source/Runtime/Engine/Private/PhysicsEngine/PhysicsConstraintComponent.cpp
@@ -62,7 +62,7 @@ UPrimitiveComponent* UPhysicsConstraintComponent::GetComponentInternal(EConstrai
 			// No name specified, use the root component
 			if(ComponentName == NAME_None)
 			{
-				PrimComp = Actor->GetRootPrimitiveComponent();
+				PrimComp = Cast<UPrimitiveComponent>(Actor->GetRootComponent());
 			}
 			// Name specified, see if we can find that property..
 			else
diff --git a/Engine/Source/Runtime/Engine/Private/PhysicsVolume.cpp b/Engine/Source/Runtime/Engine/Private/PhysicsVolume.cpp
--- a/Engine/Source/Runtime/Engine/Private/PhysicsVolume.cpp
+++ b/Engine/Source/Runtime/Engine/Private/PhysicsVolume.cpp
@@ -40,7 +40,7 @@ bool APhysicsVolume::IsOverlapInVolume(const class USceneComponent& TestComponen
 	if (!bPhysicsOnContact)
 	{
 		FVector ClosestPoint(0.f);
-		UPrimitiveComponent* RootPrimitive = GetRootPrimitiveComponent();
+		UPrimitiveComponent* RootPrimitive = Cast<UPrimitiveComponent>(GetRootComponent());
 		const float DistToCollision = RootPrimitive ? RootPrimitive->GetDistanceToCollision(TestComponent.GetComponentLocation(), ClosestPoint) : 0.f;
 		bInsideVolume = (DistToCollision == 0.f);
 	}
