Fix crashing on opening skeletal meshes with invalid attached assets

#TTP 343782 - [ASSET CONTENT PACK]: EGMMilitaryWeap1: CRASH: Opening animations

diff --git a/Engine/Source/Editor/Persona/Private/Persona.cpp b/Engine/Source/Editor/Persona/Private/Persona.cpp
--- a/Engine/Source/Editor/Persona/Private/Persona.cpp
+++ b/Engine/Source/Editor/Persona/Private/Persona.cpp
@@ -659,6 +659,9 @@ void FPersona::InitPersona(const EToolkitMode::Type Mode, const TSharedPtr< clas
 		}
 	}
 
+	// Force validation of preview attached assets (catch case of never doing it if we dont have a valid preview mesh)
+	ValidatePreviewAttachedAssets(NULL);
+
 	UAnimBlueprint* AnimBlueprint = GetAnimBlueprint();
 	PreviewComponent->SetAnimInstanceClass(AnimBlueprint ? AnimBlueprint->GeneratedClass : NULL);
 
diff --git a/Engine/Source/Runtime/Engine/Private/Animation/Skeleton.cpp b/Engine/Source/Runtime/Engine/Private/Animation/Skeleton.cpp
--- a/Engine/Source/Runtime/Engine/Private/Animation/Skeleton.cpp
+++ b/Engine/Source/Runtime/Engine/Private/Animation/Skeleton.cpp
@@ -735,8 +735,9 @@ USkeletalMesh* USkeleton::GetPreviewMesh(bool bFindIfNotSet)
 		{
 			PreviewMesh = Cast<USkeletalMesh>(StaticLoadObject(USkeletalMesh::StaticClass(), NULL, *PreviewMeshStringRef.ToString(), NULL, LOAD_None, NULL));
 		}
+		
 		// if not existing, and if bFindIfNotExisting is true, then try find one
-		else if (bFindIfNotSet)
+		if (!PreviewMesh && bFindIfNotSet)
 		{
 			FARFilter Filter;
 			Filter.ClassNames.Add(USkeletalMesh::StaticClass()->GetFName());
