#ue4 - Fixed ProjectileMovementComponent passing wrong time value to HandleHitWall() (and hence HandleImpact).

diff --git a/Engine/Source/Runtime/Engine/Private/ProjectileMovementComponent.cpp b/Engine/Source/Runtime/Engine/Private/ProjectileMovementComponent.cpp
--- a/Engine/Source/Runtime/Engine/Private/ProjectileMovementComponent.cpp
+++ b/Engine/Source/Runtime/Engine/Private/ProjectileMovementComponent.cpp
@@ -146,7 +146,7 @@ void UProjectileMovementComponent::TickComponent(float DeltaTime, enum ELevelTic
 				// re-calculate end velocity for partial time
 				Velocity = CalculateVelocity(OldVelocity, TimeTick*Hit.Time);
 			}
-			if ( HandleHitWall(Hit, TimeTick*Hit.Time, MoveDelta) )
+			if ( HandleHitWall(Hit, TimeTick, MoveDelta) )
 			{
 				break;
 			}
