Integrate change from John Barrett that adds hooks needed for an internal tool being worked on

diff --git a/Engine/Source/Runtime/Engine/Classes/Engine/Channel.h b/Engine/Source/Runtime/Engine/Classes/Engine/Channel.h
--- a/Engine/Source/Runtime/Engine/Classes/Engine/Channel.h
+++ b/Engine/Source/Runtime/Engine/Classes/Engine/Channel.h
@@ -25,7 +25,7 @@ enum EChannelType
 #define VOICE_CHANNEL_INDEX 1
 
 UCLASS(transient)
-class UChannel : public UObject
+class ENGINE_API UChannel : public UObject
 {
 	GENERATED_UCLASS_BODY()
 	
diff --git a/Engine/Source/Runtime/Engine/Classes/Engine/NetConnection.h b/Engine/Source/Runtime/Engine/Classes/Engine/NetConnection.h
--- a/Engine/Source/Runtime/Engine/Classes/Engine/NetConnection.h
+++ b/Engine/Source/Runtime/Engine/Classes/Engine/NetConnection.h
@@ -491,7 +491,7 @@ public:
 	ENGINE_API virtual void ReceivedRawPacket(void* Data,int32 Count);
 
 	/** Send a raw bunch. */
-	int32 SendRawBunch( FOutBunch& Bunch, bool InAllowMerge );
+	ENGINE_API int32 SendRawBunch( FOutBunch& Bunch, bool InAllowMerge );
 
 	/** @return The driver object */
 	UNetDriver* GetDriver() {return Driver;}
diff --git a/Engine/Source/Runtime/Engine/Classes/GameFramework/Actor.h b/Engine/Source/Runtime/Engine/Classes/GameFramework/Actor.h
--- a/Engine/Source/Runtime/Engine/Classes/GameFramework/Actor.h
+++ b/Engine/Source/Runtime/Engine/Classes/GameFramework/Actor.h
@@ -39,6 +39,10 @@ DECLARE_DYNAMIC_MULTICAST_DELEGATE_OneParam(FActorEndPlaySignature, EEndPlayReas
 
 DECLARE_DELEGATE_FourParams(FMakeNoiseDelegate, AActor*, float, class APawn*, const FVector&);
 
+#if !UE_BUILD_SHIPPING
+DECLARE_DELEGATE_RetVal_ThreeParams(bool, FOnProcessEvent, AActor*, UFunction*, void*);
+#endif
+
 DECLARE_CYCLE_STAT_EXTERN(TEXT("GetComponentsTime"),STAT_GetComponentsTime,STATGROUP_Engine,ENGINE_API);
 
 /**
@@ -2036,7 +2040,13 @@ private:
 	void SetActorLabelInternal( const FString& NewActorLabelDirty, bool bMakeGloballyUniqueFName );
 
 	static FMakeNoiseDelegate MakeNoiseDelegate;
+
 public:
+#if !UE_BUILD_SHIPPING
+	/** Delegate for globally hooking ProccessEvent calls - used by a non-public testing plugin */
+	static FOnProcessEvent ProcessEventDelegate;
+#endif
+
 	static void MakeNoiseImpl(AActor* NoiseMaker, float Loudness, APawn* NoiseInstigator, const FVector& NoiseLocation);
 	static void SetMakeNoiseDelegate(const FMakeNoiseDelegate& NewDelegate);
 };
diff --git a/Engine/Source/Runtime/Engine/Private/Actor.cpp b/Engine/Source/Runtime/Engine/Private/Actor.cpp
--- a/Engine/Source/Runtime/Engine/Private/Actor.cpp
+++ b/Engine/Source/Runtime/Engine/Private/Actor.cpp
@@ -33,6 +33,10 @@ DEFINE_STAT(STAT_GetComponentsTime);
 
 FMakeNoiseDelegate AActor::MakeNoiseDelegate = FMakeNoiseDelegate::CreateStatic(&AActor::MakeNoiseImpl);
 
+#if !UE_BUILD_SHIPPING
+FOnProcessEvent AActor::ProcessEventDelegate;
+#endif
+
 AActor::AActor(const class FPostConstructInitializeProperties& PCIP)
 	: Super(PCIP)
 {
@@ -531,7 +535,14 @@ void AActor::ProcessEvent(UFunction* Function, void* Parameters)
 	#endif
 	if( ((GetWorld() && (GetWorld()->AreActorsInitialized() || bAllowScriptExecution)) || HasAnyFlags(RF_ClassDefaultObject)) && !GIsGarbageCollecting )
 	{
+#if !UE_BUILD_SHIPPING
+		if (!ProcessEventDelegate.IsBound() || !ProcessEventDelegate.Execute(this, Function, Parameters))
+		{
+			Super::ProcessEvent(Function, Parameters);
+		}
+#else
 		Super::ProcessEvent(Function, Parameters);
+#endif
 	}
 }
 
