Purge duplicate AxisConfig entries in cases where the ini files have gotten messed up.

diff --git a/Engine/Source/Runtime/Engine/Classes/GameFramework/InputSettings.h b/Engine/Source/Runtime/Engine/Classes/GameFramework/InputSettings.h
--- a/Engine/Source/Runtime/Engine/Classes/GameFramework/InputSettings.h
+++ b/Engine/Source/Runtime/Engine/Classes/GameFramework/InputSettings.h
@@ -64,6 +64,8 @@ class ENGINE_API UInputSettings
 #if WITH_EDITOR
 	virtual void PostEditChangeChainProperty(struct FPropertyChangedChainEvent& PropertyChangedEvent) override;
 #endif
+
+	virtual void PostInitProperties() override;
 	// End of UObject interface
 
 	void AddActionMapping(const FInputActionKeyMapping& KeyMapping);
diff --git a/Engine/Source/Runtime/Engine/Private/UserInterface/InputSettings.cpp b/Engine/Source/Runtime/Engine/Private/UserInterface/InputSettings.cpp
--- a/Engine/Source/Runtime/Engine/Private/UserInterface/InputSettings.cpp
+++ b/Engine/Source/Runtime/Engine/Private/UserInterface/InputSettings.cpp
@@ -11,6 +11,26 @@ UInputSettings::UInputSettings(const class FPostConstructInitializeProperties& P
 {
 }
 
+void UInputSettings::PostInitProperties()
+{
+	Super::PostInitProperties();
+
+	TMap<FName, int32> UniqueAxisConfigNames;
+	for (int32 Index = 0; Index < AxisConfig.Num(); ++Index)
+	{
+		UniqueAxisConfigNames.Add(AxisConfig[Index].AxisKeyName, Index);
+	}
+
+	for (int32 Index = AxisConfig.Num() - 1; Index >= 0; --Index)
+	{
+		const int32 UniqueAxisIndex = UniqueAxisConfigNames.FindChecked(AxisConfig[Index].AxisKeyName);
+		if (UniqueAxisIndex != Index)
+		{
+			AxisConfig.RemoveAtSwap(Index);
+		}
+	}
+}
+
 #if WITH_EDITOR
 void UInputSettings::PostEditChangeChainProperty(FPropertyChangedChainEvent& PropertyChangedEvent)
 {
