Fix issue where Controller property wasn't replicating to non owners, but needs to when the value is NULL

diff --git a/Engine/Source/Runtime/Engine/Private/Pawn.cpp b/Engine/Source/Runtime/Engine/Private/Pawn.cpp
--- a/Engine/Source/Runtime/Engine/Private/Pawn.cpp
+++ b/Engine/Source/Runtime/Engine/Private/Pawn.cpp
@@ -1027,8 +1027,8 @@ void APawn::GetLifetimeReplicatedProps( TArray< FLifetimeProperty > & OutLifetim
 	Super::GetLifetimeReplicatedProps( OutLifetimeProps );
 
 	DOREPLIFETIME( APawn, PlayerState ); 
+	DOREPLIFETIME( APawn, Controller );
 
-	DOREPLIFETIME_CONDITION( APawn, Controller,			COND_OwnerOnly );
 	DOREPLIFETIME_CONDITION( APawn, RemoteViewPitch, 	COND_SkipOwner );
 }
 
