#ue4 - Fix some HandleImpact() calls getting passed the wrong TimeSlice value in CharacterMovementComponent.

diff --git a/Engine/Source/Runtime/Engine/Private/CharacterMovementComponent.cpp b/Engine/Source/Runtime/Engine/Private/CharacterMovementComponent.cpp
--- a/Engine/Source/Runtime/Engine/Private/CharacterMovementComponent.cpp
+++ b/Engine/Source/Runtime/Engine/Private/CharacterMovementComponent.cpp
@@ -2832,7 +2832,9 @@ void UCharacterMovementComponent::PhysFalling(float deltaTime, int32 Iterations)
 			return;
 		}
 		
+		float LastMoveTimeSlice = timeTick;
 		float subTimeTickRemaining = timeTick * (1.f - Hit.Time);
+		
 		if ( IsSwimming() ) //just entered water
 		{
 			remainingTime += subTimeTickRemaining;
@@ -2867,7 +2869,7 @@ void UCharacterMovementComponent::PhysFalling(float deltaTime, int32 Iterations)
 					}
 				}
 
-				HandleImpact(Hit, timeTick, Adjusted);
+				HandleImpact(Hit, LastMoveTimeSlice, Adjusted);
 				
 				// If we've changed physics mode, abort.
 				if (!HasValidData() || !IsFalling())
@@ -2894,6 +2896,7 @@ void UCharacterMovementComponent::PhysFalling(float deltaTime, int32 Iterations)
 					if (Hit.Time < 1.f)
 					{
 						// hit second wall
+						LastMoveTimeSlice = subTimeTickRemaining;
 						subTimeTickRemaining = subTimeTickRemaining * (1.f - Hit.Time);
 
 						if (IsValidLandingSpot(UpdatedComponent->GetComponentLocation(), Hit))
@@ -2903,7 +2906,7 @@ void UCharacterMovementComponent::PhysFalling(float deltaTime, int32 Iterations)
 							return;
 						}
 
-						HandleImpact(Hit, subTimeTickRemaining, Delta);
+						HandleImpact(Hit, LastMoveTimeSlice, Delta);
 
 						// If we've changed physics mode, abort.
 						if (!HasValidData() || !IsFalling())
@@ -3168,6 +3171,8 @@ void UCharacterMovementComponent::MoveAlongFloor(const FVector& InVelocity, cons
 	FHitResult Hit(1.f);
 	FVector RampVector = ComputeGroundMovementDelta(Delta, CurrentFloor.HitResult, CurrentFloor.bLineTrace);
 	SafeMoveUpdatedComponent(RampVector, CharacterOwner->GetActorRotation(), true, Hit);
+	float LastMoveTimeSlice = DeltaSeconds;
+	
 	if (Hit.bStartPenetrating)
 	{
 		UE_LOG(LogCharacterMovement, Log, TEXT("%s is stuck and failed to move!"), *CharacterOwner->GetName());
@@ -3179,16 +3184,17 @@ void UCharacterMovementComponent::MoveAlongFloor(const FVector& InVelocity, cons
 	if (Hit.IsValidBlockingHit())
 	{
 		// We impacted something (most likely another ramp, but possibly a barrier).
-		float TimeApplied = Hit.Time;
+		float PercentTimeApplied = Hit.Time;
 		if ((Hit.Time > 0.f) && (Hit.Normal.Z > KINDA_SMALL_NUMBER) && IsWalkable(Hit))
 		{
 			// Another walkable ramp.
-			const float InitialTimeRemaining = 1.f - TimeApplied;
-			RampVector = ComputeGroundMovementDelta(Delta * InitialTimeRemaining, Hit, false);
+			const float InitialPercentRemaining = 1.f - PercentTimeApplied;
+			RampVector = ComputeGroundMovementDelta(Delta * InitialPercentRemaining, Hit, false);
+			LastMoveTimeSlice = InitialPercentRemaining * LastMoveTimeSlice;
 			SafeMoveUpdatedComponent(RampVector, CharacterOwner->GetActorRotation(), true, Hit);
 
-			const float SecondHitPercent = Hit.Time * InitialTimeRemaining;
-			TimeApplied = FMath::Clamp(TimeApplied + SecondHitPercent, 0.f, 1.f);
+			const float SecondHitPercent = Hit.Time * InitialPercentRemaining;
+			PercentTimeApplied = FMath::Clamp(PercentTimeApplied + SecondHitPercent, 0.f, 1.f);
 		}
 
 		if (Hit.IsValidBlockingHit())
@@ -3197,11 +3203,11 @@ void UCharacterMovementComponent::MoveAlongFloor(const FVector& InVelocity, cons
 			{
 				// hit a barrier, try to step up
 				const FVector GravDir(0.f, 0.f, -1.f);
-				if (!StepUp(GravDir, Delta * (1.f - TimeApplied), Hit, OutStepDownResult))
+				if (!StepUp(GravDir, Delta * (1.f - PercentTimeApplied), Hit, OutStepDownResult))
 				{
 					UE_LOG(LogCharacterMovement, Verbose, TEXT("- StepUp (ImpactNormal %s, Normal %s"), *Hit.ImpactNormal.ToString(), *Hit.Normal.ToString());
-					HandleImpact(Hit, DeltaSeconds, Delta);
-					SlideAlongSurface(Delta, 1.f - TimeApplied, Hit.Normal, Hit, true);
+					HandleImpact(Hit, LastMoveTimeSlice, RampVector);
+					SlideAlongSurface(Delta, 1.f - PercentTimeApplied, Hit.Normal, Hit, true);
 				}
 				else
 				{
@@ -3212,8 +3218,8 @@ void UCharacterMovementComponent::MoveAlongFloor(const FVector& InVelocity, cons
 			}
 			else if ( Hit.Component.IsValid() && !Hit.Component.Get()->CanCharacterStepUp(CharacterOwner) )
 			{
-				HandleImpact(Hit, DeltaSeconds, Delta);
-				SlideAlongSurface(Delta, 1.f - TimeApplied, Hit.Normal, Hit, true);
+				HandleImpact(Hit, LastMoveTimeSlice, RampVector);
+				SlideAlongSurface(Delta, 1.f - PercentTimeApplied, Hit.Normal, Hit, true);
 			}
 		}
 	}
