#ue4 - Deprecate Character::DoJump() in favor of CharacterMovementComponent::DoJump().

- Character::CheckJumpInput() now checks CanJump() and then calls the component's function.
- Added bool bReplayingMoves to CharacterMovementComponent::DoJump(), for network correction replays.
- If you were overriding Character::DoJump() to do work not restricted by CanJump(), you may now need to override Character::CheckJumpInput() and run your code if bPressedJump is true.
   - See the PlatformerGame sample as an example of this.

diff --git a/Engine/Source/Runtime/Engine/Classes/GameFramework/Character.h b/Engine/Source/Runtime/Engine/Classes/GameFramework/Character.h
--- a/Engine/Source/Runtime/Engine/Classes/GameFramework/Character.h
+++ b/Engine/Source/Runtime/Engine/Classes/GameFramework/Character.h
@@ -431,6 +431,7 @@ protected:
 	 * @param	bReplayingMoves: true if this is being done as part of replaying moves on a locally controlled client after a server correction.
 	 * @return	True if the jump was allowed by CanJump() and if CharacterMovement->Jump() succeeded.
 	 **/ 
+	DEPRECATED(4.5, "Character::DoJump() will be removed, use CharacterMovementComponent::DoJump() instead.")
 	virtual bool DoJump(bool bReplayingMoves);
 
 public:
diff --git a/Engine/Source/Runtime/Engine/Classes/GameFramework/CharacterMovementComponent.h b/Engine/Source/Runtime/Engine/Classes/GameFramework/CharacterMovementComponent.h
--- a/Engine/Source/Runtime/Engine/Classes/GameFramework/CharacterMovementComponent.h
+++ b/Engine/Source/Runtime/Engine/Classes/GameFramework/CharacterMovementComponent.h
@@ -743,8 +743,13 @@ public:
 	/** changes physics based on MovementMode */
 	virtual void StartNewPhysics(float deltaTime, int32 Iterations);
 	
-	/** Perform jump. Note that you should usually trigger a jump through Character::Jump() instead. */
-	virtual bool DoJump();
+	/**
+	 * Perform jump. Called by Character when a jump has been detected because Character->bPressedJump was true. Checks CanJump().
+	 * Note that you should usually trigger a jump through Character::Jump() instead.
+	 * @param	bReplayingMoves: true if this is being done as part of replaying moves on a locally controlled client after a server correction.
+	 * @return	True if the jump was triggered successfully.
+	 */
+	virtual bool DoJump(bool bReplayingMoves);
 
 	/** Queue a pending launch with velocity LaunchVel. */
 	virtual void Launch(FVector const& LaunchVel);
diff --git a/Engine/Source/Runtime/Engine/Private/Character.cpp b/Engine/Source/Runtime/Engine/Private/Character.cpp
--- a/Engine/Source/Runtime/Engine/Private/Character.cpp
+++ b/Engine/Source/Runtime/Engine/Private/Character.cpp
@@ -221,9 +221,10 @@ bool ACharacter::IsJumpProvidingForce() const
 	return (bPressedJump && JumpKeyHoldTime > 0.0f && JumpKeyHoldTime < GetJumpMaxHoldTime());
 }
 
+// Deprecated
 bool ACharacter::DoJump( bool bReplayingMoves )
 {
-	return CanJump() && CharacterMovement->DoJump();
+	return CanJump() && CharacterMovement->DoJump(bReplayingMoves);
 }
 
 
@@ -818,8 +819,8 @@ void ACharacter::CheckJumpInput(float DeltaTime)
 	{
 		// Increment our timer first so calls to IsJumpProvidingForce() will return true
 		JumpKeyHoldTime += DeltaTime;
-		const bool bDidJump = DoJump(bClientUpdating);
-		if(!bWasJumping && bDidJump)
+		const bool bDidJump = CanJump() && CharacterMovement && CharacterMovement->DoJump(bClientUpdating);
+		if (!bWasJumping && bDidJump)
 		{
 			OnJumped();
 		}
diff --git a/Engine/Source/Runtime/Engine/Private/CharacterMovementComponent.cpp b/Engine/Source/Runtime/Engine/Private/CharacterMovementComponent.cpp
--- a/Engine/Source/Runtime/Engine/Private/CharacterMovementComponent.cpp
+++ b/Engine/Source/Runtime/Engine/Private/CharacterMovementComponent.cpp
@@ -427,9 +427,9 @@ FVector UCharacterMovementComponent::GetPawnCapsuleExtent(const EShrinkCapsuleEx
 }
 
 
-bool UCharacterMovementComponent::DoJump()
+bool UCharacterMovementComponent::DoJump(bool bReplayingMoves)
 {
-	if ( CharacterOwner )
+	if ( CharacterOwner && CharacterOwner->CanJump() )
 	{
 		// Don't jump if we can't move up/down.
 		if (!bConstrainToPlane || FMath::Abs(PlaneConstraintNormal.Z) != 1.f)
