#ue4 - Fix typo in ApplyVelocityBraking setting the wrong limit on BrakingDeceleration.

diff --git a/Engine/Source/Runtime/Engine/Private/CharacterMovementComponent.cpp b/Engine/Source/Runtime/Engine/Private/CharacterMovementComponent.cpp
--- a/Engine/Source/Runtime/Engine/Private/CharacterMovementComponent.cpp
+++ b/Engine/Source/Runtime/Engine/Private/CharacterMovementComponent.cpp
@@ -2428,7 +2428,7 @@ void UCharacterMovementComponent::ApplyVelocityBraking(float DeltaTime, float Fr
 	}
 
 	Friction = FMath::Max(0.f, Friction);
-	BrakingDeceleration = FMath::Max(0.f, Friction);
+	BrakingDeceleration = FMath::Max(0.f, BrakingDeceleration);
 
 	if (Friction == 0.f && BrakingDeceleration == 0.f)
 	{
