Provide a matching pair to InitializeComponent in UninitializeComponent

diff --git a/Engine/Source/Runtime/Engine/Classes/Components/ActorComponent.h b/Engine/Source/Runtime/Engine/Classes/Components/ActorComponent.h
--- a/Engine/Source/Runtime/Engine/Classes/Components/ActorComponent.h
+++ b/Engine/Source/Runtime/Engine/Classes/Components/ActorComponent.h
@@ -79,6 +79,9 @@ public:
 	/** Indicates that OnCreatedComponent has been called, but OnDestroyedComponent has not yet */
 	uint32 bHasBeenCreated:1;
 
+	/** Indicates that InitializeComponent has been called, but UninitializeComponent has not yet */
+	uint32 bHasBeenInitialized:1;
+
 	UFUNCTION()
 	void OnRep_IsActive();
 
@@ -254,6 +257,12 @@ public:
 	 */
 	virtual void InitializeComponent();
 
+	/**
+	 * Ends gameplay for this component.
+	 * Called from AActor::EndPlay only if bHasBeenInitialized is true
+	 */
+	virtual void UninitializeComponent();
+
 	/**
 	 * When called, will call the virtual call chain to register all of the tick functions
 	 * Do not override this function or make it virtual
diff --git a/Engine/Source/Runtime/Engine/Classes/GameFramework/Actor.h b/Engine/Source/Runtime/Engine/Classes/GameFramework/Actor.h
--- a/Engine/Source/Runtime/Engine/Classes/GameFramework/Actor.h
+++ b/Engine/Source/Runtime/Engine/Classes/GameFramework/Actor.h
@@ -1551,6 +1551,9 @@ public:
 	/** Iterate over components array and call InitializeComponent */
 	void InitializeComponents();
 
+	/** Iterate over components array and call UninitializeComponent */
+	void UninitializeComponents();
+
 	/** Debug rendering to visualize the component tree for this actor. */
 	void DrawDebugComponents(FColor const& BaseColor=FColor::White) const;
 
diff --git a/Engine/Source/Runtime/Engine/Private/Actor.cpp b/Engine/Source/Runtime/Engine/Private/Actor.cpp
--- a/Engine/Source/Runtime/Engine/Private/Actor.cpp
+++ b/Engine/Source/Runtime/Engine/Private/Actor.cpp
@@ -1481,6 +1481,8 @@ void AActor::OnRep_Instigator() {}
 
 void AActor::EndPlay(const EEndPlayReason::Type EndPlayReason)
 {
+	UninitializeComponents();
+
 	// Dispatch the blueprint events
 	ReceiveEndPlay(EndPlayReason);
 	OnEndPlay.Broadcast(EndPlayReason);
@@ -3129,9 +3131,8 @@ void AActor::InitializeComponents()
 	TArray<UActorComponent*> Components;
 	GetComponents(Components);
 
-	for (int32 Idx = 0; Idx < Components.Num(); Idx++)
+	for (UActorComponent* ActorComp : Components)
 	{
-		UActorComponent* ActorComp = Components[Idx];
 		if (ActorComp->bWantsInitializeComponent && ActorComp->IsRegistered())
 		{
 			ActorComp->InitializeComponent();
@@ -3139,6 +3140,19 @@ void AActor::InitializeComponents()
 	}
 }
 
+void AActor::UninitializeComponents()
+{
+	TArray<UActorComponent*> Components;
+	GetComponents(Components);
+
+	for (UActorComponent* ActorComp : Components)
+	{
+		if (ActorComp->bHasBeenInitialized)
+		{
+			ActorComp->UninitializeComponent();
+		}
+	}
+}
 
 void AActor::DrawDebugComponents(FColor const& BaseColor) const
 {
diff --git a/Engine/Source/Runtime/Engine/Private/ActorComponent.cpp b/Engine/Source/Runtime/Engine/Private/ActorComponent.cpp
--- a/Engine/Source/Runtime/Engine/Private/ActorComponent.cpp
+++ b/Engine/Source/Runtime/Engine/Private/ActorComponent.cpp
@@ -283,6 +283,12 @@ FString UActorComponent::GetReadableName() const
 
 void UActorComponent::BeginDestroy()
 {
+	// Ensure that we call UninitializeComponent before we destroy this component
+	if (bHasBeenInitialized)
+	{
+		UninitializeComponent();
+	}
+
 	// Ensure that we call OnComponentDestroyed before we destroy this component
 	if(bHasBeenCreated)
 	{
@@ -437,13 +443,19 @@ void UActorComponent::OnRegister()
 void UActorComponent::OnUnregister()
 {
 	check(bRegistered);
-
 	bRegistered = false;
 }
 
 void UActorComponent::InitializeComponent()
 {
 	check(bRegistered);
+	bHasBeenInitialized = true;
+}
+
+void UActorComponent::UninitializeComponent()
+{
+	check(bHasBeenInitialized);
+	bHasBeenInitialized = false;
 }
 
 void FActorComponentTickFunction::ExecuteTick(float DeltaTime, enum ELevelTick TickType, ENamedThreads::Type CurrentThread, const FGraphEventRef& MyCompletionGraphEvent)
