Make sure some other things are off for no-gbuffer path.

diff --git a/Engine/Source/Runtime/Engine/Private/ShowFlags.cpp b/Engine/Source/Runtime/Engine/Private/ShowFlags.cpp
--- a/Engine/Source/Runtime/Engine/Private/ShowFlags.cpp
+++ b/Engine/Source/Runtime/Engine/Private/ShowFlags.cpp
@@ -342,7 +342,11 @@ void EngineShowFlagOverride(EShowFlagInitMode ShowFlagInitMode, EViewModeIndex V
 		static const auto ICVar = IConsoleManager::Get().FindTConsoleVariableDataInt(TEXT("r.GBuffer"));
 		if(ICVar->GetValueOnGameThread() == 0)
 		{
+			EngineShowFlags.AmbientOcclusion = 0;
+			EngineShowFlags.Decals = 0;
 			EngineShowFlags.DynamicShadows = 0;
+			EngineShowFlags.GlobalIllumination = 0;
+			EngineShowFlags.ScreenSpaceReflections = 0;
 		}
 	}
 	
