Integrated Licensee optimization to FPrecomputedLightVolume::InterpolateIncidentRadianceBlock, measured ILC update cost go from 103ms -> 63ms with this change in a scene with ~400 moving objects

diff --git a/Engine/Source/Runtime/Engine/Private/PrecomputedLightVolume.cpp b/Engine/Source/Runtime/Engine/Private/PrecomputedLightVolume.cpp
--- a/Engine/Source/Runtime/Engine/Private/PrecomputedLightVolume.cpp
+++ b/Engine/Source/Runtime/Engine/Private/PrecomputedLightVolume.cpp
@@ -318,35 +318,48 @@ void FPrecomputedLightVolume::InterpolateIncidentRadianceBlock(
 		for (int32 SampleIndex = 0; SampleIndex < PotentiallyIntersectingSamples.Num(); SampleIndex++)
 		{
 			const FVolumeLightingSample& VolumeSample = *PotentiallyIntersectingSamples[SampleIndex];
-			
+			const int32 LinearIndexBase = DestCellPosition.Z * DestCellDimensions.Y * DestCellDimensions.X
+				+ DestCellPosition.Y * DestCellDimensions.X
+				+ DestCellPosition.X;
+				
+			const float RadiusSquared = FMath::Square(VolumeSample.Radius);
+			const float WeightBase  = 1.0f / RadiusSquared;
+			const float WeightMultiplier = -1.0f / (RadiusSquared * RadiusSquared);
+				
+			const FVector BaseTranslationFromSample = BoundingBox.Center - BoundingBox.Extent - VolumeSample.Position;
+			const FVector QuerySteps = BoundingBox.Extent / FVector(QueryCellDimensions) * 2;
+			FVector TranslationFromSample = BaseTranslationFromSample;
+
 			for (int32 Z = 0; Z < QueryCellDimensions.Z; Z++)
 			{
 				for (int32 Y = 0; Y < QueryCellDimensions.Y; Y++)
 				{
 					for (int32 X = 0; X < QueryCellDimensions.X; X++)
 					{
-						const int32 LinearIndex = 
-							(Z + DestCellPosition.Z) * DestCellDimensions.Y * DestCellDimensions.X 
-							+ (Y + DestCellPosition.Y) * DestCellDimensions.X
-							+ (X + DestCellPosition.X);
-
-						const FVector WorldPosition = BoundingBox.Center - BoundingBox.Extent + FVector(X, Y, Z) / FVector(QueryCellDimensions) * BoundingBox.Extent * 2;	
-						const float DistanceSquared = (VolumeSample.Position - WorldPosition).SizeSquared();
-						const float RadiusSquared = FMath::Square(VolumeSample.Radius);
+						const float DistanceSquared = TranslationFromSample.SizeSquared();
 
 						if (DistanceSquared < RadiusSquared)
 						{
+							const int32 LinearIndex = LinearIndexBase + (Z * DestCellDimensions.Y + Y) * DestCellDimensions.X + X;
 							// Weight each sample with the fraction that Position is to the center of the sample, and inversely to the sample radius.
 							// The weight goes to 0 when Position is on the bounding radius of the sample, so the interpolated result is continuous.
 							// The sample's radius size is a factor so that smaller samples contribute more than larger, low detail ones.
-							const float SampleWeight = (1.0f - DistanceSquared / RadiusSquared) / RadiusSquared;
+							const float SampleWeight = DistanceSquared * WeightMultiplier + WeightBase;
 							// Accumulate weighted results and the total weight for normalization later
 							AccumulatedIncidentRadiance[LinearIndex] += VolumeSample.Lighting * SampleWeight;
 							AccumulatedSkyBentNormal[LinearIndex] += VolumeSample.GetSkyBentNormalUnpacked() * SampleWeight;
 							AccumulatedWeights[LinearIndex] += SampleWeight;
 						}
+
+						TranslationFromSample.X += QuerySteps.X;
 					}
+
+					TranslationFromSample.X = BaseTranslationFromSample.X;
+					TranslationFromSample.Y += QuerySteps.Y;
 				}
+
+				TranslationFromSample.Y = BaseTranslationFromSample.Y;
+				TranslationFromSample.Z += QuerySteps.Z;
 			}
 		}
 	}
