#ue4 - CalcVelocity runs each iteration of CharacterMovement walking movement.

- Results in more consistent movement across varied framerates, within the limits of MaxSimulationTimeStep.
- Modified ApplyVelocityBraking for better behavior at low framerates, and with minor optimizations in the absense of friction or braking deceleration.
- ApplyRequestedVelocity (for path following, used by CalcVelocity) now works across multiple iterations within the same frame (bHasRequestedVelocity is cleared after the frame, not on first use).

diff --git a/Engine/Source/Runtime/Engine/Classes/GameFramework/CharacterMovementComponent.h b/Engine/Source/Runtime/Engine/Classes/GameFramework/CharacterMovementComponent.h
--- a/Engine/Source/Runtime/Engine/Classes/GameFramework/CharacterMovementComponent.h
+++ b/Engine/Source/Runtime/Engine/Classes/GameFramework/CharacterMovementComponent.h
@@ -435,7 +435,7 @@ public:
 	 * WARNING: if (MaxSimulationTimeStep * MaxSimulationIterations) is too low for the min framerate, the last simulation step may exceed MaxSimulationTimeStep to complete the simulation.
 	 * @see MaxSimulationIterations
 	 */
-	UPROPERTY(Category="Character Movement", EditAnywhere, BlueprintReadWrite, AdvancedDisplay, meta=(ClampMin="0.02", ClampMax="0.50", UIMin="0.02", UIMax="0.50"))
+	UPROPERTY(Category="Character Movement", EditAnywhere, BlueprintReadWrite, AdvancedDisplay, meta=(ClampMin="0.0166", ClampMax="0.50", UIMin="0.0166", UIMax="0.50"))
 	float MaxSimulationTimeStep;
 
 	/**
@@ -445,7 +445,7 @@ public:
 	 *
 	 * WARNING: if (MaxSimulationTimeStep * MaxSimulationIterations) is too low for the min framerate, the last simulation step may exceed MaxSimulationTimeStep to complete the simulation.
 	 */
-	UPROPERTY(Category="Character Movement", EditAnywhere, BlueprintReadWrite, AdvancedDisplay, meta=(ClampMin="1", ClampMax="20", UIMin="1", UIMax="20"))
+	UPROPERTY(Category="Character Movement", EditAnywhere, BlueprintReadWrite, AdvancedDisplay, meta=(ClampMin="1", ClampMax="25", UIMin="1", UIMax="25"))
 	int32 MaxSimulationIterations;
 
 	/** Used in determining if pawn is going off ledge.  If the ledge is "shorter" than this value then the pawn will be able to walk off it. **/
diff --git a/Engine/Source/Runtime/Engine/Private/CharacterMovementComponent.cpp b/Engine/Source/Runtime/Engine/Private/CharacterMovementComponent.cpp
--- a/Engine/Source/Runtime/Engine/Private/CharacterMovementComponent.cpp
+++ b/Engine/Source/Runtime/Engine/Private/CharacterMovementComponent.cpp
@@ -1079,6 +1079,9 @@ void UCharacterMovementComponent::SimulateMovement(float DeltaSeconds)
 		FStepDownResult StepDownResult;
 		MoveSmooth(Velocity, DeltaSeconds, &StepDownResult);
 
+		// consume path following requested velocity
+		bHasRequestedVelocity = false;
+
 		// if simulated gravity, find floor and check if falling
 		const bool bEnableFloorCheck = (!CharacterOwner->bSimGravityDisabled || !bIsSimulatedProxy);
 		if (bEnableFloorCheck && (MovementMode == MOVE_Walking || MovementMode == MOVE_Falling))
@@ -1409,6 +1412,9 @@ void UCharacterMovementComponent::PerformMovement(float DeltaSeconds)
 		// change position
 		StartNewPhysics(DeltaSeconds, 0);
 
+		// consume path following requested velocity
+		bHasRequestedVelocity = false;
+
 		if (!HasValidData())
 		{
 			return;
@@ -2151,8 +2157,6 @@ bool UCharacterMovementComponent::ApplyRequestedMove(FVector& NewAcceleration, f
 {
 	if (bHasRequestedVelocity)
 	{
-		bHasRequestedVelocity = false;
-
 		OutRequestedSpeed = RequestedVelocity.Size();
 		if (OutRequestedSpeed < KINDA_SMALL_NUMBER)
 		{
@@ -2429,8 +2433,10 @@ void UCharacterMovementComponent::ApplyVelocityBraking(float DeltaTime, float Fr
 
 	Friction = FMath::Max(0.f, Friction);
 	BrakingDeceleration = FMath::Max(0.f, BrakingDeceleration);
+	const bool bZeroFriction = (Friction == 0.f);
+	const bool bZeroBraking = (BrakingDeceleration == 0.f);
 
-	if (Friction == 0.f && BrakingDeceleration == 0.f)
+	if (bZeroFriction && bZeroBraking)
 	{
 		return;
 	}
@@ -2440,13 +2446,14 @@ void UCharacterMovementComponent::ApplyVelocityBraking(float DeltaTime, float Fr
 	// subdivide braking to get reasonably consistent results at lower frame rates
 	// (important for packet loss situations w/ networking)
 	float RemainingTime = DeltaTime;
-	const float TimeStep = (1.0f / 33.0f);
+	const float MaxTimeStep = (1.0f / 33.0f);
 
 	// Decelerate to brake to a stop
-	const FVector RevAccel = (-1.f * BrakingDeceleration) * SafeNormalPrecise(Velocity);
+	const FVector RevAccel = (bZeroBraking ? FVector::ZeroVector : (-BrakingDeceleration * SafeNormalPrecise(Velocity)));
 	while( RemainingTime >= MIN_TICK_TIME )
 	{
-		const float dt = ((RemainingTime > TimeStep) ? TimeStep : RemainingTime);
+		// Zero friction uses constant deceleration, so no need for iteration.
+		const float dt = ((RemainingTime > MaxTimeStep && !bZeroFriction) ? FMath::Min(MaxTimeStep, RemainingTime * 0.5f) : RemainingTime);
 		RemainingTime -= dt;
 
 		// apply friction and braking
@@ -2462,7 +2469,7 @@ void UCharacterMovementComponent::ApplyVelocityBraking(float DeltaTime, float Fr
 
 	// Clamp to zero if nearly zero, or if below min threshold and braking.
 	const float VSizeSq = Velocity.SizeSquared();
-	if (VSizeSq <= KINDA_SMALL_NUMBER || (BrakingDeceleration > 0.f && VSizeSq <= FMath::Square(BRAKE_TO_STOP_VELOCITY)))
+	if (VSizeSq <= KINDA_SMALL_NUMBER || (!bZeroBraking && VSizeSq <= FMath::Square(BRAKE_TO_STOP_VELOCITY)))
 	{
 		Velocity = FVector::ZeroVector;
 	}
@@ -3103,10 +3110,10 @@ void UCharacterMovementComponent::RevertMove(const FVector& OldLocation, UPrimit
 	//UE_LOG(LogCharacterMovement, Log, TEXT("Now at %f %f %f"), CharacterOwner->Location.X, CharacterOwner->Location.Y, CharacterOwner->Location.Z);
 	bJustTeleported = false;
 	// if our previous base couldn't have moved or changed in any physics-affecting way, restore it
-	if (OldBase && 
+	if (IsValid(OldBase) && 
 		(!MovementBaseUtility::IsDynamicBase(OldBase) ||
 		 (OldBase->Mobility == EComponentMobility::Static) ||
-		 (OldBase->GetOwner() && OldBase->GetOwner()->GetActorLocation() == PreviousBaseLocation)
+		 (OldBase->GetComponentLocation() == PreviousBaseLocation)
 		)
 	   )
 	{
@@ -3264,25 +3271,8 @@ void UCharacterMovementComponent::PhysWalking(float deltaTime, int32 Iterations)
 		return;
 	}
 
-	// Ensure velocity is horizontal.
-	MaintainHorizontalGroundVelocity();
-	checkf(!Velocity.ContainsNaN(), TEXT("PhysWalking: Velocity contains NaN before CalcVelocity (%s: %s)\n%s"), *GetPathNameSafe(this), *GetPathNameSafe(GetOuter()), *Velocity.ToString());
-
-	//bound acceleration
-	Acceleration.Z = 0.f;
-	if( !HasRootMotion() )
-	{
-		CalcVelocity(deltaTime, GroundFriction, false, BrakingDecelerationWalking);
-		checkf(!Velocity.ContainsNaN(), TEXT("PhysWalking: Velocity contains NaN after CalcVelocity (%s: %s)\n%s"), *GetPathNameSafe(this), *GetPathNameSafe(GetOuter()), *Velocity.ToString());
-	}
-
-	FVector DesiredMove = Velocity;
-	DesiredMove.Z = 0.f;
-
-	const FVector OldLocation = UpdatedComponent->GetComponentLocation();
-	const FFindFloorResult OldFloor = CurrentFloor;
-	UPrimitiveComponent* const OldBase = CharacterOwner->GetMovementBase();
-	const FVector PreviousBaseLocation = (OldBase != NULL) ? OldBase->GetComponentLocation() : FVector::ZeroVector;
+	checkf(!Velocity.ContainsNaN(), TEXT("PhysWalking: Velocity contains NaN before Iteration (%s: %s)\n%s"), *GetPathNameSafe(this), *GetPathNameSafe(GetOuter()), *Velocity.ToString());
+	
 	bJustTeleported = false;
 	bool bCheckedFall = false;
 	bool bTriedLedgeMove = false;
@@ -3292,10 +3282,32 @@ void UCharacterMovementComponent::PhysWalking(float deltaTime, int32 Iterations)
 	while ( (remainingTime >= MIN_TICK_TIME) && (Iterations < MaxSimulationIterations) && (CharacterOwner->Controller || bRunPhysicsWithNoController || HasRootMotion()) )
 	{
 		Iterations++;
+		bJustTeleported = false;
 		const float timeTick = GetSimulationTimeStep(remainingTime, Iterations);
 		remainingTime -= timeTick;
-		const FVector Delta = timeTick * DesiredMove;
-		const FVector subLoc = CharacterOwner->GetActorLocation();
+
+		// Save current values
+		UPrimitiveComponent * const OldBase = GetMovementBase();
+		const FVector OldBaseLocation = (OldBase != NULL) ? OldBase->GetComponentLocation() : FVector::ZeroVector;
+		const FVector OldLocation = UpdatedComponent->GetComponentLocation();
+		const FFindFloorResult OldFloor = CurrentFloor;
+
+		// Ensure velocity is horizontal.
+		MaintainHorizontalGroundVelocity();
+		Velocity.Z = 0.f;
+		const FVector OldVelocity = Velocity;
+
+		// Apply acceleration
+		Acceleration.Z = 0.f;
+		if( !HasRootMotion() )
+		{
+			CalcVelocity(timeTick, GroundFriction, false, BrakingDecelerationWalking);
+		}
+		checkf(!Velocity.ContainsNaN(), TEXT("PhysWalking: Velocity contains NaN after CalcVelocity (%s: %s)\n%s"), *GetPathNameSafe(this), *GetPathNameSafe(GetOuter()), *Velocity.ToString());
+
+		// Compute move parameters
+		const FVector MoveVelocity = Velocity;
+		const FVector Delta = timeTick * MoveVelocity;
 		const bool bZeroDelta = Delta.IsNearlyZero();
 		FStepDownResult StepDownResult;
 
@@ -3305,12 +3317,8 @@ void UCharacterMovementComponent::PhysWalking(float deltaTime, int32 Iterations)
 		}
 		else
 		{
-			// @todo hunting down NaN TTP 304692
-			checkf( !Delta.ContainsNaN(), TEXT("PhysWalking: Delta contains NaN (%s: %s)\ntimeTick = %d Velocity = %s Delta = %s"), *GetPathNameSafe(this), *GetPathNameSafe(GetOuter()), timeTick, *Velocity.ToString(), *Delta.ToString() );
-			// @todo hunting down NaN TTP 304692
-
 			// try to move forward
-			MoveAlongFloor(DesiredMove, timeTick, &StepDownResult);
+			MoveAlongFloor(MoveVelocity, timeTick, &StepDownResult);
 
 			if ( IsFalling() )
 			{
@@ -3318,7 +3326,7 @@ void UCharacterMovementComponent::PhysWalking(float deltaTime, int32 Iterations)
 				const float DesiredDist = Delta.Size();
 				if (DesiredDist > KINDA_SMALL_NUMBER)
 				{
-					const float ActualDist = (CharacterOwner->GetActorLocation() - subLoc).Size2D();
+					const float ActualDist = (CharacterOwner->GetActorLocation() - OldLocation).Size2D();
 					remainingTime += timeTick * (1.f - FMath::Min(1.f,ActualDist/DesiredDist));
 				}
 				StartNewPhysics(remainingTime,Iterations);
@@ -3326,7 +3334,7 @@ void UCharacterMovementComponent::PhysWalking(float deltaTime, int32 Iterations)
 			}
 			else if ( IsSwimming() ) //just entered water
 			{
-				StartSwimming(OldLocation, Velocity, timeTick, remainingTime, Iterations);
+				StartSwimming(OldLocation, OldVelocity, timeTick, remainingTime, Iterations);
 				return;
 			}
 		}
@@ -3352,13 +3360,13 @@ void UCharacterMovementComponent::PhysWalking(float deltaTime, int32 Iterations)
 			if ( !NewDelta.IsZero() )
 			{
 				// first revert this move
-				RevertMove(OldLocation, OldBase, PreviousBaseLocation, OldFloor, false);
+				RevertMove(OldLocation, OldBase, OldBaseLocation, OldFloor, false);
 
 				// avoid repeated ledge moves if the first one fails
 				bTriedLedgeMove = true;
 
 				// Try new movement direction
-				DesiredMove = NewDelta/timeTick;
+				Velocity = NewDelta/timeTick;
 				remainingTime += timeTick;
 				continue;
 			}
@@ -3367,15 +3375,14 @@ void UCharacterMovementComponent::PhysWalking(float deltaTime, int32 Iterations)
 				// see if it is OK to jump
 				// @todo collision : only thing that can be problem is that oldbase has world collision on
 				bool bMustJump = bZeroDelta || (OldBase == NULL || (!OldBase->IsCollisionEnabled() && MovementBaseUtility::IsDynamicBase(OldBase)));
-				if ( (bMustJump || !bCheckedFall) && CheckFall(CurrentFloor.HitResult, Delta, subLoc, remainingTime, timeTick, Iterations, bMustJump) )
+				if ( (bMustJump || !bCheckedFall) && CheckFall(CurrentFloor.HitResult, Delta, OldLocation, remainingTime, timeTick, Iterations, bMustJump) )
 				{
 					return;
 				}
 				bCheckedFall = true;
 
 				// revert this move
-				//UE_LOG(LogCharacterMovement, Log, TEXT("%s REVERT MOVE 1"), *CharacterOwner->GetName());
-				RevertMove(OldLocation, OldBase, PreviousBaseLocation, OldFloor, true);
+				RevertMove(OldLocation, OldBase, OldBaseLocation, OldFloor, true);
 				remainingTime = 0.f;
 				break;
 			}
@@ -3391,7 +3398,7 @@ void UCharacterMovementComponent::PhysWalking(float deltaTime, int32 Iterations)
 					if (IsMovingOnGround())
 					{
 						// If still walking, then fall. If not, assume the user set a different mode they want to keep.
-						StartFalling(Iterations, remainingTime, timeTick, Delta, subLoc);
+						StartFalling(Iterations, remainingTime, timeTick, Delta, OldLocation);
 					}
 					return;
 				}
@@ -3420,7 +3427,7 @@ void UCharacterMovementComponent::PhysWalking(float deltaTime, int32 Iterations)
 			if (!CurrentFloor.IsWalkableFloor() && !CurrentFloor.HitResult.bStartPenetrating)
 			{
 				const bool bMustJump = bJustTeleported || bZeroDelta || (OldBase == NULL || (!OldBase->IsCollisionEnabled() && MovementBaseUtility::IsDynamicBase(OldBase)));
-				if ((bMustJump || !bCheckedFall) && CheckFall(CurrentFloor.HitResult, Delta, subLoc, remainingTime, timeTick, Iterations, bMustJump) )
+				if ((bMustJump || !bCheckedFall) && CheckFall(CurrentFloor.HitResult, Delta, OldLocation, remainingTime, timeTick, Iterations, bMustJump) )
 				{
 					return;
 				}
@@ -3428,24 +3435,27 @@ void UCharacterMovementComponent::PhysWalking(float deltaTime, int32 Iterations)
 			}
 		}
 
+
+		// Allow overlap events and such to change physics state and velocity
+		if (IsMovingOnGround())
+		{
+			// Make velocity reflect actual move
+			if( !bJustTeleported && !HasRootMotion() && timeTick >= MIN_TICK_TIME)
+			{
+				Velocity = (CharacterOwner->GetActorLocation() - OldLocation) / timeTick;
+			}
+		}
+
 		// If we didn't move at all this iteration then abort (since future iterations will also be stuck).
-		if (CharacterOwner->GetActorLocation() == subLoc)
+		if (CharacterOwner->GetActorLocation() == OldLocation)
 		{
 			remainingTime = 0.f;
 			break;
-		}
+		}	
 	}
 
-
-	// Allow overlap events and such to change physics state and velocity
-	if( IsMovingOnGround() )
+	if (IsMovingOnGround())
 	{
-		// Make velocity reflect actual move
-		if( !bJustTeleported && !HasRootMotion() && deltaTime >= MIN_TICK_TIME)
-		{
-			Velocity = (CharacterOwner->GetActorLocation() - OldLocation) / deltaTime;
-		}
-
 		MaintainHorizontalGroundVelocity();
 	}
 }
