[343954] Blueprints - fix Select node not working with array types

diff --git a/Engine/Source/Runtime/Engine/Private/EdGraph/EdGraphPin.cpp b/Engine/Source/Runtime/Engine/Private/EdGraph/EdGraphPin.cpp
--- a/Engine/Source/Runtime/Engine/Private/EdGraph/EdGraphPin.cpp
+++ b/Engine/Source/Runtime/Engine/Private/EdGraph/EdGraphPin.cpp
@@ -161,9 +161,7 @@ void UEdGraphPin::CopyPersistentDataFromOldPin(const UEdGraphPin& SourcePin)
 	// TODO: Better less hacky way of this?
 	if (PinType.PinCategory == TEXT("wildcard"))
 	{
-		PinType.PinCategory = SourcePin.PinType.PinCategory;
-		PinType.PinSubCategory = SourcePin.PinType.PinSubCategory;
-		PinType.PinSubCategoryObject = SourcePin.PinType.PinSubCategoryObject;
+		PinType = SourcePin.PinType;
 	}
 
 	// Copy the links
