Removed 'transient' from ANavigationTestingActor::OtherActor #UE4

diff --git a/Engine/Source/Runtime/Engine/Classes/AI/Navigation/NavigationTestingActor.h b/Engine/Source/Runtime/Engine/Classes/AI/Navigation/NavigationTestingActor.h
--- a/Engine/Source/Runtime/Engine/Classes/AI/Navigation/NavigationTestingActor.h
+++ b/Engine/Source/Runtime/Engine/Classes/AI/Navigation/NavigationTestingActor.h
@@ -108,7 +108,7 @@ class ENGINE_API ANavigationTestingActor : public AActor, public INavAgentInterf
 	UPROPERTY(transient, VisibleAnywhere, BlueprintReadOnly, Category=PathfindingStatus)
 	int32 PathfindingSteps;
 
-	UPROPERTY(transient, EditAnywhere, Category=Pathfinding)
+	UPROPERTY(EditAnywhere, Category=Pathfinding)
 	ANavigationTestingActor* OtherActor;
 
 	/** "None" will result in default filter being used */
