Reduced UE_VLOG overhead when recording for vlogs disabled. #ue4

diff --git a/Engine/Source/Runtime/Engine/Public/VisualLog.h b/Engine/Source/Runtime/Engine/Public/VisualLog.h
--- a/Engine/Source/Runtime/Engine/Public/VisualLog.h
+++ b/Engine/Source/Runtime/Engine/Public/VisualLog.h
@@ -211,7 +211,10 @@ struct ENGINE_API FVisLogEntry
 { \
 	SCOPE_CYCLE_COUNTER(STAT_VisualLog); \
 	static_assert((ELogVerbosity::Verbosity & ELogVerbosity::VerbosityMask) < ELogVerbosity::NumVerbosity && ELogVerbosity::Verbosity > 0, "Verbosity must be constant and in range."); \
-	FVisualLog::Get().LogLine(Actor, CategoryName.GetCategoryName(), ELogVerbosity::Verbosity, FString::Printf(Format, ##__VA_ARGS__)); \
+	if (FVisualLog::Get().IsRecording() && (!FVisualLog::Get().IsAllBlocked() || FVisualLog::Get().InWhitelist(CategoryName.GetCategoryName()))) \
+	{ \
+		FVisualLog::Get().LogLine(Actor, CategoryName.GetCategoryName(), ELogVerbosity::Verbosity, FString::Printf(Format, ##__VA_ARGS__)); \
+	} \
 	if (UE_LOG_CHECK_COMPILEDIN_VERBOSITY(CategoryName, Verbosity)) \
 	{ \
 		if (!CategoryName.IsSuppressed(ELogVerbosity::Verbosity) && (Actor) != NULL) \
