#ttp 334858 Android : Support ADPCM

Added support for ADPCM decoding.
Added 'sound type' information to USoundWave
- Android uses this to detect which type of asset it needs to deal with
Added support for selected encoding type via project settings

diff --git a/Engine/Source/Developer/Android/AndroidTargetPlatform/AndroidTargetPlatform.Build.cs b/Engine/Source/Developer/Android/AndroidTargetPlatform/AndroidTargetPlatform.Build.cs
--- a/Engine/Source/Developer/Android/AndroidTargetPlatform/AndroidTargetPlatform.Build.cs
+++ b/Engine/Source/Developer/Android/AndroidTargetPlatform/AndroidTargetPlatform.Build.cs
@@ -32,6 +32,8 @@ public class AndroidTargetPlatform : ModuleRules
 			PrivateIncludePathModuleNames.Add("TextureCompressor");		//@todo android: AndroidTargetPlatform.Build
 		}
 
+        Definitions.Add("WITH_OGGVORBIS=1");
+
 		PrivateIncludePaths.AddRange(
 			new string[]
 			{
diff --git a/Engine/Source/Developer/Android/AndroidTargetPlatform/Private/AndroidTargetPlatform.inl b/Engine/Source/Developer/Android/AndroidTargetPlatform/Private/AndroidTargetPlatform.inl
--- a/Engine/Source/Developer/Android/AndroidTargetPlatform/Private/AndroidTargetPlatform.inl
+++ b/Engine/Source/Developer/Android/AndroidTargetPlatform/Private/AndroidTargetPlatform.inl
@@ -275,9 +275,28 @@ const FTextureLODSettings& FAndroidTargetPlatform<TPlatformProperties>::GetTextu
 template<class TPlatformProperties>
 FName FAndroidTargetPlatform<TPlatformProperties>::GetWaveFormat( class USoundWave* Wave ) const
 {
-	static FName NAME_OGG(TEXT("OGG"));		//@todo android: probably not ogg
+	FString audioSetting;
+	if (!GConfig->GetString(TEXT("/Script/UnrealEd.CookerSettings"), TEXT("AndroidAudio"), audioSetting, GEngineIni))
+	{
+		audioSetting = TEXT("DEFAULT");
+	}
 
-	return NAME_OGG;
+#if WITH_OGGVORBIS
+	if (audioSetting == TEXT("OGG") || audioSetting == TEXT("Default"))
+	{
+		static FName NAME_OGG(TEXT("OGG"));
+		return NAME_OGG;
+	}
+#else
+	if (audioSetting == TEXT("OGG"))
+	{
+		UE_LOG(LogAudio, Error, TEXT("Attemped to select Ogg Vorbis encoding when the cooker is built without Ogg Vorbis support."));
+	}
+#endif
+	
+	// Otherwise return ADPCM as it'll either be option '2' or 'default' depending on WITH_OGGVORBIS config
+	static FName NAME_ADPCM(TEXT("ADPCM"));
+	return NAME_ADPCM;
 }
 
 #endif //WITH_ENGINE
diff --git a/Engine/Source/Editor/UnrealEd/Private/Commandlets/CookCommandlet.cpp b/Engine/Source/Editor/UnrealEd/Private/Commandlets/CookCommandlet.cpp
--- a/Engine/Source/Editor/UnrealEd/Private/Commandlets/CookCommandlet.cpp
+++ b/Engine/Source/Editor/UnrealEd/Private/Commandlets/CookCommandlet.cpp
@@ -33,6 +33,7 @@ UCookerSettings::UCookerSettings(const FPostConstructInitializeProperties& PCIP)
 	DefaultPVRTCQuality = 0;
 	bSupportMetal = false;
 	bSupportOpenGLES2 = true;
+	AndroidAudio = EAndroidAudio::Default;
 }
 
 
diff --git a/Engine/Source/Editor/UnrealEd/Public/CookerSettings.h b/Engine/Source/Editor/UnrealEd/Public/CookerSettings.h
--- a/Engine/Source/Editor/UnrealEd/Public/CookerSettings.h
+++ b/Engine/Source/Editor/UnrealEd/Public/CookerSettings.h
@@ -9,6 +9,17 @@
 #include "CookerSettings.generated.h"
 
 
+UENUM()
+namespace EAndroidAudio
+{
+	enum Type
+	{
+		Default = 0 UMETA(DisplayName="Default",ToolTip="This option selects the default encoder."),
+		OGG = 1 UMETA(DisplayName="Ogg Vorbis", ToolTip="Selects Ogg Vorbis encoding."),
+		ADPCM = 2 UMETA(DisplayName="ADPCM",ToolTip="This option selects ADPCM lossless encoding.")
+	};
+}
+
 /**
  * Implements per-project cooker settings exposed to the editor
  */
@@ -31,4 +42,8 @@ public:
 	/** Whether or not to add support for OpenGL ES2 (if this is false, then your game should specify minimum IOS8 version and A7, in the plist) */
 	UPROPERTY(config, EditAnywhere, Category = IOS)
 	bool bSupportOpenGLES2;
+
+	/** Android Audio encoding options */
+	UPROPERTY(config, EditAnywhere, Category = Android, meta = (DisplayName = "Audio encoding"))
+	TEnumAsByte<EAndroidAudio::Type> AndroidAudio;
 };
diff --git a/Engine/Source/Runtime/Android/AndroidAudio/Private/AndroidAudioBuffer.cpp b/Engine/Source/Runtime/Android/AndroidAudio/Private/AndroidAudioBuffer.cpp
--- a/Engine/Source/Runtime/Android/AndroidAudio/Private/AndroidAudioBuffer.cpp
+++ b/Engine/Source/Runtime/Android/AndroidAudio/Private/AndroidAudioBuffer.cpp
@@ -103,7 +103,6 @@ FSLESSoundBuffer* FSLESSoundBuffer::CreateQueuedBuffer( FSLESAudioDevice* AudioD
  */
 FSLESSoundBuffer* FSLESSoundBuffer::CreateNativeBuffer( FSLESAudioDevice* AudioDevice, USoundWave* InWave )
 {
-#if WITH_OGGVORBIS
 	// Check to see if thread has finished decompressing on the other thread
 	if( InWave->AudioDecompressor != NULL )
 	{
@@ -113,7 +112,6 @@ FSLESSoundBuffer* FSLESSoundBuffer::CreateNativeBuffer( FSLESAudioDevice* AudioD
 		delete InWave->AudioDecompressor;
 		InWave->AudioDecompressor = NULL;
 	}
-#endif	//WITH_OGGVORBIS
 
 	FSLESSoundBuffer* Buffer = NULL;
 
@@ -216,5 +214,16 @@ FSLESSoundBuffer* FSLESSoundBuffer::Init(  FSLESAudioDevice* AudioDevice ,USound
 bool FSLESSoundBuffer::ReadCompressedData( uint8* Destination, bool bLooping )
 {
 	ensure( DecompressionState);
-	return( DecompressionState->ReadCompressedData( Destination, bLooping, MONO_PCM_BUFFER_SIZE * NumChannels ) );
+	return(DecompressionState->ReadCompressedData(Destination, bLooping, DecompressionState->GetStreamBufferSize() * NumChannels));
+}
+
+/**
+* Returns the size for a real time/streaming buffer based on decompressor
+*
+* @return Size of buffer in bytes for a single channel or 0 if no decompression state
+*/
+int FSLESSoundBuffer::GetRTBufferSize(void)
+{
+	check(DecompressionState != NULL);
+	return DecompressionState ? DecompressionState->GetStreamBufferSize() : 0;
 }
diff --git a/Engine/Source/Runtime/Android/AndroidAudio/Private/AndroidAudioDevice.cpp b/Engine/Source/Runtime/Android/AndroidAudio/Private/AndroidAudioDevice.cpp
--- a/Engine/Source/Runtime/Android/AndroidAudio/Private/AndroidAudioDevice.cpp
+++ b/Engine/Source/Runtime/Android/AndroidAudio/Private/AndroidAudioDevice.cpp
@@ -7,6 +7,7 @@
 
 #include "AndroidAudioDevice.h"
 #include "VorbisAudioInfo.h"
+#include "ADPCMAudioInfo.h"
 #include "AudioEffect.h"
 #include "Engine.h"
 
@@ -157,17 +158,34 @@ FSoundSource* FSLESAudioDevice::CreateSoundSource()
 bool FSLESAudioDevice::HasCompressedAudioInfoClass(USoundWave* SoundWave)
 {
 #if WITH_OGGVORBIS
-	return true;
-#else
-	return false;
+	if (SoundWave->CompressionName.IsValid() && SoundWave->CompressionName == TEXT("OGG"))
+	{
+		return true;
+	}
 #endif
+
+	if (SoundWave->CompressionName.IsValid() && SoundWave->CompressionName == TEXT("ADPCM"))
+	{
+		return true;
+	}
+
+	return false;
+
 }
 
 class ICompressedAudioInfo* FSLESAudioDevice::CreateCompressedAudioInfo(USoundWave* SoundWave)
 {
 #if WITH_OGGVORBIS
-	return new FVorbisAudioInfo();
-#else
-	return NULL;
+	if (SoundWave->CompressionName.IsValid() && SoundWave->CompressionName == TEXT("OGG"))
+	{
+		return new FVorbisAudioInfo();
+	}
 #endif
+
+	if (SoundWave->CompressionName.IsValid() && SoundWave->CompressionName == TEXT("ADPCM"))
+	{
+		return new FADPCMAudioInfo();
+	}
+
+	return NULL;
 }
diff --git a/Engine/Source/Runtime/Android/AndroidAudio/Private/AndroidAudioSource.cpp b/Engine/Source/Runtime/Android/AndroidAudio/Private/AndroidAudioSource.cpp
--- a/Engine/Source/Runtime/Android/AndroidAudio/Private/AndroidAudioSource.cpp
+++ b/Engine/Source/Runtime/Android/AndroidAudio/Private/AndroidAudioSource.cpp
@@ -118,7 +118,7 @@ bool FSLESSoundSource::EnqueuePCMBuffer( bool bLoop)
 		result = (*SL_PlayerBufferQueue)->RegisterCallback(SL_PlayerBufferQueue, OpenSLBufferQueueCallback, (void*)this);
 		if(result != SL_RESULT_SUCCESS) { UE_LOG( LogAndroidAudio, Fatal,TEXT("FAILED OPENSL BUFFER QUEUE RegisterCallback 0x%x "), result); return false;}
 	}
-		
+	
 	result = (*SL_PlayerBufferQueue)->Enqueue(SL_PlayerBufferQueue, Buffer->AudioData, Buffer->GetSize() );
 	if(result != SL_RESULT_SUCCESS) { UE_LOG( LogAndroidAudio, Fatal,TEXT("FAILED OPENSL BUFFER Enqueue SL_PlayerBufferQueue 0x%x params( %p, %d)"), result, Buffer->AudioData, int32(Buffer->GetSize())); return false;}
 
@@ -138,11 +138,13 @@ bool FSLESSoundSource::EnqueuePCMRTBuffer( bool bLoop )
 	FMemory::Memzero( AudioBuffers, sizeof( SLESAudioBuffer ) * 2 );
 
 	// Set up double buffer area to decompress to
-	AudioBuffers[0].AudioData = ( uint8* )FMemory::Malloc( MONO_PCM_BUFFER_SIZE * Buffer->NumChannels );
-	AudioBuffers[0].AudioDataSize = MONO_PCM_BUFFER_SIZE * Buffer->NumChannels;
+	const int bufferSize = Buffer->GetRTBufferSize() * Buffer->NumChannels;
+
+	AudioBuffers[0].AudioData = (uint8*)FMemory::Malloc(bufferSize);
+	AudioBuffers[0].AudioDataSize = bufferSize;
 
-	AudioBuffers[1].AudioData = ( uint8* )FMemory::Malloc( MONO_PCM_BUFFER_SIZE * Buffer->NumChannels );
-	AudioBuffers[1].AudioDataSize = MONO_PCM_BUFFER_SIZE * Buffer->NumChannels;
+	AudioBuffers[1].AudioData = (uint8*)FMemory::Malloc(bufferSize);
+	AudioBuffers[1].AudioDataSize = bufferSize;
 
 	// Decompress two buffers worth of data to fill up the queue
 	Buffer->ReadCompressedData( AudioBuffers[0].AudioData, bLoop );
diff --git a/Engine/Source/Runtime/Android/AndroidAudio/Public/AndroidAudioDevice.h b/Engine/Source/Runtime/Android/AndroidAudio/Public/AndroidAudioDevice.h
--- a/Engine/Source/Runtime/Android/AndroidAudio/Public/AndroidAudioDevice.h
+++ b/Engine/Source/Runtime/Android/AndroidAudio/Public/AndroidAudioDevice.h
@@ -97,6 +97,11 @@ public:
 		return( BufferSize ); 
 	}
 
+	/**
+	 * Returns the size for a real time/streaming buffer based on decompressor
+	 */
+	int GetRTBufferSize(void);
+		
 	/** Audio device this buffer is attached to */
 	FSLESAudioDevice*			AudioDevice;
 	/** Data */
@@ -222,8 +227,19 @@ public:
 
 	virtual FName GetRuntimeFormat(USoundWave* SoundWave) override
 	{
-		static FName NAME_OGG(TEXT("OGG"));
-		return NAME_OGG;
+		
+		if (SoundWave->CompressionName.IsNone())
+		{
+#if WITH_OGGVORBIS
+			static FName NAME_OGG(TEXT("OGG"));		//@todo android: probably not ogg
+			return NAME_OGG;
+#else
+			static FName NAME_ADPCM(TEXT("ADPCM"));
+			return NAME_ADPCM;
+#endif
+		}
+
+		return SoundWave->CompressionName;
 	}
 
 	virtual bool HasCompressedAudioInfoClass(USoundWave* SoundWave) override;
diff --git a/Engine/Source/Runtime/Core/Public/UObject/ObjectVersion.h b/Engine/Source/Runtime/Core/Public/UObject/ObjectVersion.h
--- a/Engine/Source/Runtime/Core/Public/UObject/ObjectVersion.h
+++ b/Engine/Source/Runtime/Core/Public/UObject/ObjectVersion.h
@@ -596,6 +596,8 @@ enum EUnrealEngineObjectUE4Version
 	VER_UE4_ENUM_CLASS_SUPPORT,
 	// Fixup widget animation class
 	VER_UE4_FIXUP_WIDGET_ANIMATION_CLASS,
+	// USoundWave objects now contain details about compression scheme used.
+	VER_UE4_SOUND_COMPRESSION_TYPE_ADDED,
 
 	// -----<new versions can be added before this line>-------------------------------------------------
 	// - this needs to be the last line (see note below)
diff --git a/Engine/Source/Runtime/Engine/Classes/Sound/SoundWave.h b/Engine/Source/Runtime/Engine/Classes/Sound/SoundWave.h
--- a/Engine/Source/Runtime/Engine/Classes/Sound/SoundWave.h
+++ b/Engine/Source/Runtime/Engine/Classes/Sound/SoundWave.h
@@ -256,6 +256,9 @@ public:
 	/** cooked streaming platform data for this sound */
 	TMap<FString, FStreamedAudioPlatformData*> CookedPlatformData;
 
+	/** Codec used to compress/encode this audio data */
+	FName CompressionName;
+
 	// Begin UObject interface. 
 	virtual void Serialize( FArchive& Ar ) override;
 	virtual void PostInitProperties() override;
diff --git a/Engine/Source/Runtime/Engine/Private/ADPCMAudioInfo.cpp b/Engine/Source/Runtime/Engine/Private/ADPCMAudioInfo.cpp
new file mode 100644
--- /dev/null
+++ b/Engine/Source/Runtime/Engine/Private/ADPCMAudioInfo.cpp
@@ -0,0 +1,154 @@
+// Copyright 1998-2014 Epic Games, Inc. All Rights Reserved.
+
+
+#include "EnginePrivate.h"
+#include "ADPCMAudioInfo.h"
+#include "IAudioFormat.h"
+
+
+namespace ADPCM
+{
+	void DecodeBlock(const uint8* EncodedADPCMBlock, int32 BlockSize, int16* DecodedPCMData);
+	void DecodeBlockStereo(const uint8* EncodedADPCMBlockLeft, const uint8* EncodedADPCMBlockRight, int32 BlockSize, int16* DecodedPCMData);
+}
+
+FADPCMAudioInfo::FADPCMAudioInfo(void)
+{
+
+}
+
+FADPCMAudioInfo::~FADPCMAudioInfo(void)
+{
+
+}
+
+bool FADPCMAudioInfo::ReadCompressedInfo(const uint8* InSrcBufferData, uint32 InSrcBufferDataSize, struct FSoundQualityInfo* QualityInfo)
+{
+	check(InSrcBufferData);
+	check(QualityInfo);
+
+	SrcBufferData = InSrcBufferData;
+	SrcBufferDataSize = InSrcBufferDataSize;
+	SrcBufferOffset = 0;
+
+	if (!WaveInfo.ReadWaveInfo((uint8*)SrcBufferData, SrcBufferDataSize))
+	{
+		UE_LOG(LogAudio, Warning, TEXT("WaveInfo.ReadWaveInfo Failed"));
+		return false;
+	}
+
+	const uint32 PreambleSize = 7;
+	BlockSize = *WaveInfo.pBlockAlign;
+
+	// (BlockSize - PreambleSize) * 2 (samples per byte) + 2 (preamble samples)
+	UncompressedBlockSize = (2 + (BlockSize - PreambleSize) * 2) * sizeof(int16);
+	CompressedBlockSize = BlockSize;
+
+	// Calculate the preferred mono buffer size here based on the same calculations as we do for Vorbis but with our blocksizes taken into account
+	const uint32 uncompressedBlockSamples = (2 + (BlockSize - PreambleSize) * 2);
+	const uint32 targetBlocks = MONO_PCM_BUFFER_SAMPLES / uncompressedBlockSamples;
+	StreamBufferSize = targetBlocks * UncompressedBlockSize;
+	StreamBufferSizeInBlocks = CompressedBlockSize * targetBlocks;
+
+	if (QualityInfo)
+	{
+		QualityInfo->SampleRate = *WaveInfo.pSamplesPerSec;
+		QualityInfo->NumChannels = *WaveInfo.pChannels;
+		QualityInfo->SampleDataSize = (WaveInfo.SampleDataSize / CompressedBlockSize) * UncompressedBlockSize;
+		//QualityInfo->Duration = (float)TrueSampleCount / QualityInfo->SampleRate;
+	}
+
+	return true;
+}
+
+bool FADPCMAudioInfo::ReadCompressedData(uint8* Destination, bool bLooping, uint32 BufferSize)
+{
+	check(Destination);
+
+	const uint8* EncodedADPCMBlockStart = reinterpret_cast<uint8*>(WaveInfo.SampleDataStart);
+	uint8* EncodedADPCMBlock = (uint8*)EncodedADPCMBlockStart + SrcBufferOffset;
+	const uint8* EncodedDataEnd = EncodedADPCMBlockStart + WaveInfo.SampleDataSize;
+
+	// See if we are still completely inside the source data..
+	if (EncodedADPCMBlock + StreamBufferSizeInBlocks < EncodedDataEnd)
+	{
+		// and if so do a straight decode
+		while (BufferSize)
+		{
+			ADPCM::DecodeBlock(EncodedADPCMBlock, CompressedBlockSize, (int16*)Destination);
+			EncodedADPCMBlock += CompressedBlockSize;
+			SrcBufferOffset += CompressedBlockSize;
+			Destination += UncompressedBlockSize;
+			BufferSize -= UncompressedBlockSize;
+		}
+	}
+	else
+	{
+		// otherwise we copy what we have left
+
+		while (EncodedADPCMBlock < EncodedDataEnd)
+		{
+			ADPCM::DecodeBlock(EncodedADPCMBlock, CompressedBlockSize, (int16*)Destination);
+			EncodedADPCMBlock += CompressedBlockSize;
+			Destination += UncompressedBlockSize;
+			BufferSize -= UncompressedBlockSize;
+		}
+
+		// and then decide if we zero pad or reset and loop
+		if (bLooping)
+		{
+			EncodedADPCMBlock = (uint8*)EncodedADPCMBlockStart;
+			SrcBufferOffset = 0;
+			while (BufferSize)
+			{
+				ADPCM::DecodeBlock(EncodedADPCMBlock, CompressedBlockSize, (int16*)Destination);
+				EncodedADPCMBlock += CompressedBlockSize;
+				SrcBufferOffset += CompressedBlockSize;
+				Destination += UncompressedBlockSize;
+				BufferSize -= UncompressedBlockSize;
+			}
+		}
+		else
+		{
+			FMemory::Memzero(Destination, BufferSize);
+		}
+	}
+
+	return true;
+}
+
+void FADPCMAudioInfo::ExpandFile(uint8* DstBuffer, struct FSoundQualityInfo* QualityInfo)
+{
+	check(DstBuffer);
+
+	if (*WaveInfo.pChannels == 1)
+	{
+		const uint8* EncodedADPCMBlockStart = reinterpret_cast<uint8*>(WaveInfo.SampleDataStart);
+		uint8 * EncodedADPCMBlock = (uint8*)EncodedADPCMBlockStart;
+		while (EncodedADPCMBlock < EncodedADPCMBlockStart + (WaveInfo.SampleDataSize))
+		{
+			ADPCM::DecodeBlock(EncodedADPCMBlock, CompressedBlockSize, (int16*)DstBuffer);
+			EncodedADPCMBlock += CompressedBlockSize;
+			DstBuffer += UncompressedBlockSize;
+		}
+	}
+	else
+	{
+		const uint8* EncodedADPCMBlockLeftStart = reinterpret_cast<uint8*>(WaveInfo.SampleDataStart);
+		const uint8* EncodedADPCMBlockRightStart = reinterpret_cast<uint8*>(WaveInfo.SampleDataStart + (WaveInfo.SampleDataSize / 2));
+		uint8 * EncodedADPCMLeftBlock = (uint8*)EncodedADPCMBlockLeftStart;
+		uint8 * EncodedADPCMRightBlock = (uint8*)EncodedADPCMBlockRightStart;
+		while (EncodedADPCMLeftBlock < EncodedADPCMBlockLeftStart + (WaveInfo.SampleDataSize / 2))
+		{
+			ADPCM::DecodeBlockStereo(EncodedADPCMLeftBlock, EncodedADPCMRightBlock, CompressedBlockSize, (int16*)DstBuffer);
+			EncodedADPCMLeftBlock += CompressedBlockSize;
+			EncodedADPCMRightBlock += CompressedBlockSize;
+			DstBuffer += (UncompressedBlockSize * 2);
+		}
+	}
+}
+
+int FADPCMAudioInfo::GetStreamBufferSize() const
+{
+	return StreamBufferSize;
+}
diff --git a/Engine/Source/Runtime/Engine/Private/AudioDecompress.cpp b/Engine/Source/Runtime/Engine/Private/AudioDecompress.cpp
--- a/Engine/Source/Runtime/Engine/Private/AudioDecompress.cpp
+++ b/Engine/Source/Runtime/Engine/Private/AudioDecompress.cpp
@@ -7,6 +7,276 @@
 #include "Sound/SoundWave.h"
 #include "IAudioFormat.h"
 
+//////////////////////////////////////////////////////////////////////////
+// Copied from IOS - probably want to split and share
+//////////////////////////////////////////////////////////////////////////
+#define NUM_ADAPTATION_TABLE 16
+#define NUM_ADAPTATION_COEFF 7
+#define SOUND_SOURCE_FREE 0
+#define SOUND_SOURCE_LOCKED 1
+
+namespace
+{
+	template <typename T, uint32 B>
+	inline T SignExtend(const T ValueToExtend)
+	{
+		struct { T ExtendedValue : B; } SignExtender;
+		return SignExtender.ExtendedValue = ValueToExtend;
+	}
+
+	template <typename T>
+	inline T ReadFromByteStream(const uint8* ByteStream, int32& ReadIndex, bool bLittleEndian = true)
+	{
+		T ValueRaw = 0;
+
+		if (bLittleEndian)
+		{
+#if PLATFORM_LITTLE_ENDIAN
+			for (int32 ByteIndex = 0; ByteIndex < sizeof(T); ++ByteIndex)
+#else
+			for (int32 ByteIndex = sizeof(T) - 1; ByteIndex >= 0; --ByteIndex)
+#endif // PLATFORM_LITTLE_ENDIAN
+			{
+				ValueRaw |= ByteStream[ReadIndex++] << 8 * ByteIndex;
+			}
+		}
+		else
+		{
+#if PLATFORM_LITTLE_ENDIAN
+			for (int32 ByteIndex = sizeof(T) - 1; ByteIndex >= 0; --ByteIndex)
+#else
+			for (int32 ByteIndex = 0; ByteIndex < sizeof(T); ++ByteIndex)
+#endif // PLATFORM_LITTLE_ENDIAN
+			{
+				ValueRaw |= ByteStream[ReadIndex++] << 8 * ByteIndex;
+			}
+		}
+
+		return ValueRaw;
+	}
+
+	template <typename T>
+	inline void WriteToByteStream(T Value, uint8* ByteStream, int32& WriteIndex, bool bLittleEndian = true)
+	{
+		if (bLittleEndian)
+		{
+#if PLATFORM_LITTLE_ENDIAN
+			for (int32 ByteIndex = 0; ByteIndex < sizeof(T); ++ByteIndex)
+#else
+			for (int32 ByteIndex = sizeof(T) - 1; ByteIndex >= 0; --ByteIndex)
+#endif // PLATFORM_LITTLE_ENDIAN
+			{
+				ByteStream[WriteIndex++] = (Value >> (8 * ByteIndex)) & 0xFF;
+			}
+		}
+		else
+		{
+#if PLATFORM_LITTLE_ENDIAN
+			for (int32 ByteIndex = sizeof(T) - 1; ByteIndex >= 0; --ByteIndex)
+#else
+			for (int32 ByteIndex = 0; ByteIndex < sizeof(T); ++ByteIndex)
+#endif // PLATFORM_LITTLE_ENDIAN
+			{
+				ByteStream[WriteIndex++] = (Value >> (8 * ByteIndex)) & 0xFF;
+			}
+		}
+	}
+
+	template <typename T>
+	inline T ReadFromArray(const T* ElementArray, int32& ReadIndex, int32 NumElements, int32 IndexStride = 1)
+	{
+		T OutputValue = 0;
+
+		if (ReadIndex >= 0 && ReadIndex < NumElements)
+		{
+			OutputValue = ElementArray[ReadIndex];
+			ReadIndex += IndexStride;
+		}
+
+		return OutputValue;
+	}
+
+	inline bool LockSourceChannel(int32* ChannelLock)
+	{
+		check(ChannelLock != NULL);
+		return FPlatformAtomics::InterlockedCompareExchange(ChannelLock, SOUND_SOURCE_LOCKED, SOUND_SOURCE_FREE) == SOUND_SOURCE_FREE;
+	}
+
+	inline void UnlockSourceChannel(int32* ChannelLock)
+	{
+		check(ChannelLock != NULL);
+		int32 Result = FPlatformAtomics::InterlockedCompareExchange(ChannelLock, SOUND_SOURCE_FREE, SOUND_SOURCE_LOCKED);
+
+		check(Result == SOUND_SOURCE_LOCKED);
+	}
+
+} // end namespace
+
+namespace ADPCM
+{
+	template <typename T>
+	static void GetAdaptationTable(T(&OutAdaptationTable)[NUM_ADAPTATION_TABLE])
+	{
+		// Magic values as specified by standard
+		static T AdaptationTable[] =
+		{
+			230, 230, 230, 230, 307, 409, 512, 614,
+			768, 614, 512, 409, 307, 230, 230, 230
+		};
+
+		FMemory::Memcpy(&OutAdaptationTable, AdaptationTable, sizeof(AdaptationTable));
+	}
+
+	template <typename T>
+	static void GetAdaptationCoefficients(T(&OutAdaptationCoefficient1)[NUM_ADAPTATION_COEFF], T(&OutAdaptationCoefficient2)[NUM_ADAPTATION_COEFF])
+	{
+		// Magic values as specified by standard
+		static T AdaptationCoefficient1[] =
+		{
+			256, 512, 0, 192, 240, 460, 392
+		};
+		static T AdaptationCoefficient2[] =
+		{
+			0, -256, 0, 64, 0, -208, -232
+		};
+
+		FMemory::Memcpy(&OutAdaptationCoefficient1, AdaptationCoefficient1, sizeof(AdaptationCoefficient1));
+		FMemory::Memcpy(&OutAdaptationCoefficient2, AdaptationCoefficient2, sizeof(AdaptationCoefficient2));
+	}
+
+	struct FAdaptationContext
+	{
+	public:
+		// Adaptation constants
+		int32 AdaptationTable[NUM_ADAPTATION_TABLE];
+		int32 AdaptationCoefficient1[NUM_ADAPTATION_COEFF];
+		int32 AdaptationCoefficient2[NUM_ADAPTATION_COEFF];
+
+		int32 AdaptationDelta;
+		int32 Coefficient1;
+		int32 Coefficient2;
+		int32 Sample1;
+		int32 Sample2;
+
+		FAdaptationContext() :
+			AdaptationDelta(0),
+			Coefficient1(0),
+			Coefficient2(0),
+			Sample1(0),
+			Sample2(0)
+		{
+			GetAdaptationTable(AdaptationTable);
+			GetAdaptationCoefficients(AdaptationCoefficient1, AdaptationCoefficient2);
+		}
+	};
+
+	int16 DecodeNibble(FAdaptationContext& Context, uint8 EncodedNibble)
+	{
+		int32 PredictedSample = (Context.Sample1 * Context.Coefficient1 + Context.Sample2 * Context.Coefficient2) / 256;
+		PredictedSample += SignExtend<int8, 4>(EncodedNibble) * Context.AdaptationDelta;
+		PredictedSample = FMath::Clamp(PredictedSample, -32768, 32767);
+
+		// Shuffle samples for the next iteration
+		Context.Sample2 = Context.Sample1;
+		Context.Sample1 = static_cast<int16>(PredictedSample);
+		Context.AdaptationDelta = (Context.AdaptationDelta * Context.AdaptationTable[EncodedNibble]) / 256;
+		Context.AdaptationDelta = FMath::Max(Context.AdaptationDelta, 16);
+
+		return Context.Sample1;
+	}
+
+	void DecodeBlock(const uint8* EncodedADPCMBlock, int32 BlockSize, int16* DecodedPCMData)
+	{
+		FAdaptationContext Context;
+		int32 ReadIndex = 0;
+		int32 WriteIndex = 0;
+
+		uint8 CoefficientIndex = ReadFromByteStream<uint8>(EncodedADPCMBlock, ReadIndex);
+		Context.AdaptationDelta = ReadFromByteStream<int16>(EncodedADPCMBlock, ReadIndex);
+		Context.Sample1 = ReadFromByteStream<int16>(EncodedADPCMBlock, ReadIndex);
+		Context.Sample2 = ReadFromByteStream<int16>(EncodedADPCMBlock, ReadIndex);
+		Context.Coefficient1 = Context.AdaptationCoefficient1[CoefficientIndex];
+		Context.Coefficient2 = Context.AdaptationCoefficient2[CoefficientIndex];
+
+		// The first two samples are sent directly to the output in reverse order, as per the standard
+		DecodedPCMData[WriteIndex++] = Context.Sample2;
+		DecodedPCMData[WriteIndex++] = Context.Sample1;
+
+		uint8 EncodedNibblePair = 0;
+		uint8 EncodedNibble = 0;
+		while (ReadIndex < BlockSize)
+		{
+			// Read from the byte stream and advance the read head.
+			EncodedNibblePair = ReadFromByteStream<uint8>(EncodedADPCMBlock, ReadIndex);
+
+			EncodedNibble = (EncodedNibblePair >> 4) & 0x0F;
+			DecodedPCMData[WriteIndex++] = DecodeNibble(Context, EncodedNibble);
+
+			EncodedNibble = EncodedNibblePair & 0x0F;
+			DecodedPCMData[WriteIndex++] = DecodeNibble(Context, EncodedNibble);
+		}
+	}
+
+	// Decode two PCM streams and interleave as stereo data
+	void DecodeBlockStereo(const uint8* EncodedADPCMBlockLeft, const uint8* EncodedADPCMBlockRight, int32 BlockSize, int16* DecodedPCMData)
+	{
+		FAdaptationContext ContextLeft;
+		FAdaptationContext ContextRight;
+		int32 ReadIndexLeft = 0;
+		int32 ReadIndexRight = 0;
+		int32 WriteIndex = 0;
+
+		uint8 CoefficientIndexLeft = ReadFromByteStream<uint8>(EncodedADPCMBlockLeft, ReadIndexLeft);
+		ContextLeft.AdaptationDelta = ReadFromByteStream<int16>(EncodedADPCMBlockLeft, ReadIndexLeft);
+		ContextLeft.Sample1 = ReadFromByteStream<int16>(EncodedADPCMBlockLeft, ReadIndexLeft);
+		ContextLeft.Sample2 = ReadFromByteStream<int16>(EncodedADPCMBlockLeft, ReadIndexLeft);
+		ContextLeft.Coefficient1 = ContextLeft.AdaptationCoefficient1[CoefficientIndexLeft];
+		ContextLeft.Coefficient2 = ContextLeft.AdaptationCoefficient2[CoefficientIndexLeft];
+
+		uint8 CoefficientIndexRight = ReadFromByteStream<uint8>(EncodedADPCMBlockRight, ReadIndexRight);
+		ContextRight.AdaptationDelta = ReadFromByteStream<int16>(EncodedADPCMBlockRight, ReadIndexRight);
+		ContextRight.Sample1 = ReadFromByteStream<int16>(EncodedADPCMBlockRight, ReadIndexRight);
+		ContextRight.Sample2 = ReadFromByteStream<int16>(EncodedADPCMBlockRight, ReadIndexRight);
+		ContextRight.Coefficient1 = ContextRight.AdaptationCoefficient1[CoefficientIndexRight];
+		ContextRight.Coefficient2 = ContextRight.AdaptationCoefficient2[CoefficientIndexRight];
+
+		// The first two samples from each stream are sent directly to the output in reverse order, as per the standard
+		DecodedPCMData[WriteIndex++] = ContextLeft.Sample2;
+		DecodedPCMData[WriteIndex++] = ContextRight.Sample2;
+		DecodedPCMData[WriteIndex++] = ContextLeft.Sample1;
+		DecodedPCMData[WriteIndex++] = ContextRight.Sample1;
+
+		uint8 EncodedNibblePairLeft = 0;
+		uint8 EncodedNibbleLeft = 0;
+		uint8 EncodedNibblePairRight = 0;
+		uint8 EncodedNibbleRight = 0;
+
+		while (ReadIndexLeft < BlockSize)
+		{
+			// Read from the byte stream and advance the read head.
+			EncodedNibblePairLeft = ReadFromByteStream<uint8>(EncodedADPCMBlockLeft, ReadIndexLeft);
+			EncodedNibblePairRight = ReadFromByteStream<uint8>(EncodedADPCMBlockRight, ReadIndexRight);
+
+			EncodedNibbleLeft = (EncodedNibblePairLeft >> 4) & 0x0F;
+			DecodedPCMData[WriteIndex++] = DecodeNibble(ContextLeft, EncodedNibbleLeft);
+
+			EncodedNibbleRight = (EncodedNibblePairRight >> 4) & 0x0F;
+			DecodedPCMData[WriteIndex++] = DecodeNibble(ContextRight, EncodedNibbleRight);
+
+			EncodedNibbleLeft = EncodedNibblePairLeft & 0x0F;
+			DecodedPCMData[WriteIndex++] = DecodeNibble(ContextLeft, EncodedNibbleLeft);
+
+			EncodedNibbleRight = EncodedNibblePairRight & 0x0F;
+			DecodedPCMData[WriteIndex++] = DecodeNibble(ContextRight, EncodedNibbleRight);
+		}
+
+	}
+} // end namespace ADPCM
+
+//////////////////////////////////////////////////////////////////////////
+// End of copied section
+//////////////////////////////////////////////////////////////////////////
+
 /**
  * Worker for decompression on a separate thread
  */
diff --git a/Engine/Source/Runtime/Engine/Private/SoundWave.cpp b/Engine/Source/Runtime/Engine/Private/SoundWave.cpp
--- a/Engine/Source/Runtime/Engine/Private/SoundWave.cpp
+++ b/Engine/Source/Runtime/Engine/Private/SoundWave.cpp
@@ -145,6 +145,16 @@ void USoundWave::Serialize( FArchive& Ar )
 		UE_LOG(LogAudio, Fatal, TEXT("This platform requires cooked packages, and audio data was not cooked into %s."), *GetFullName());
 	}
 
+	if (Ar.IsCooking())
+	{
+		CompressionName = Ar.CookingTarget()->GetWaveFormat(this);
+	}
+	
+	if (Ar.UE4Ver() >= VER_UE4_SOUND_COMPRESSION_TYPE_ADDED)
+	{
+		Ar << CompressionName;
+	}
+
 	if (bCooked)
 	{
 		// Only want to cook/load full data if we don't support streaming
diff --git a/Engine/Source/Runtime/Engine/Public/ADPCMAudioInfo.h b/Engine/Source/Runtime/Engine/Public/ADPCMAudioInfo.h
new file mode 100644
--- /dev/null
+++ b/Engine/Source/Runtime/Engine/Public/ADPCMAudioInfo.h
@@ -0,0 +1,38 @@
+// Copyright 1998-2014 Epic Games, Inc. All Rights Reserved.
+
+/*=============================================================================
+	ADPCMAudioInfo.h: Unreal audio ADPCM decompression interface object.
+=============================================================================*/
+
+#pragma once
+
+#include "AudioDecompress.h"
+
+class FADPCMAudioInfo : public ICompressedAudioInfo
+{
+public:
+	ENGINE_API FADPCMAudioInfo(void);
+	ENGINE_API virtual ~FADPCMAudioInfo(void);
+
+	// ICompressedAudioInfo Interface
+	ENGINE_API virtual bool ReadCompressedInfo(const uint8* InSrcBufferData, uint32 InSrcBufferDataSize, struct FSoundQualityInfo* QualityInfo);
+	ENGINE_API virtual bool ReadCompressedData(uint8* Destination, bool bLooping, uint32 BufferSize);
+	ENGINE_API virtual void SeekToTime(const float SeekTime) {};
+	ENGINE_API virtual void ExpandFile(uint8* DstBuffer, struct FSoundQualityInfo* QualityInfo);
+	ENGINE_API virtual void EnableHalfRate(bool HalfRate) {};
+	virtual uint32 GetSourceBufferSize() const { return SrcBufferDataSize; }
+	virtual bool UsesVorbisChannelOrdering() const { return false; }
+	virtual int GetStreamBufferSize() const;
+	// End of ICompressedAudioInfo Interface
+
+	FWaveModInfo WaveInfo;
+	const uint8*	SrcBufferData;
+	uint32			SrcBufferDataSize;
+	uint32			SrcBufferOffset;
+
+	uint32 UncompressedBlockSize;
+	uint32 CompressedBlockSize;
+	uint32 BlockSize;
+	int32 StreamBufferSize;
+	uint32 StreamBufferSizeInBlocks;
+};
diff --git a/Engine/Source/Runtime/Engine/Public/AudioDecompress.h b/Engine/Source/Runtime/Engine/Public/AudioDecompress.h
--- a/Engine/Source/Runtime/Engine/Public/AudioDecompress.h
+++ b/Engine/Source/Runtime/Engine/Public/AudioDecompress.h
@@ -70,6 +70,11 @@ public:
 	 */
 	virtual bool UsesVorbisChannelOrdering() const = 0;
 
+	/**
+	* Gets the preferred size for a streaming buffer for this decompression scheme
+	*/
+	virtual int GetStreamBufferSize() const = 0;
+
 	////////////////////////////////////////////////////////////////
 	// Following functions are optional if streaming is supported //
 	////////////////////////////////////////////////////////////////
@@ -109,6 +114,7 @@ public:
 	virtual int32 GetCurrentChunkOffset() const {return -1;}
 };
 
+
 /**
  * Asynchronous audio decompression
  */
diff --git a/Engine/Source/Runtime/Engine/Public/OpusAudioInfo.h b/Engine/Source/Runtime/Engine/Public/OpusAudioInfo.h
--- a/Engine/Source/Runtime/Engine/Public/OpusAudioInfo.h
+++ b/Engine/Source/Runtime/Engine/Public/OpusAudioInfo.h
@@ -29,6 +29,7 @@ public:
 	ENGINE_API virtual void EnableHalfRate(bool HalfRate) override {};
 	virtual uint32 GetSourceBufferSize() const override { return SrcBufferDataSize;}
 	virtual bool UsesVorbisChannelOrdering() const override { return false; }
+	virtual int GetStreamBufferSize() const { return  MONO_PCM_BUFFER_SIZE; }
 
 	virtual bool SupportsStreaming() const override {return true;}
 	virtual bool StreamCompressedInfo(USoundWave* Wave, struct FSoundQualityInfo* QualityInfo) override;
diff --git a/Engine/Source/Runtime/Engine/Public/VorbisAudioInfo.h b/Engine/Source/Runtime/Engine/Public/VorbisAudioInfo.h
--- a/Engine/Source/Runtime/Engine/Public/VorbisAudioInfo.h
+++ b/Engine/Source/Runtime/Engine/Public/VorbisAudioInfo.h
@@ -80,6 +80,7 @@ public:
 	virtual uint32 GetSourceBufferSize() const { return SrcBufferDataSize;}
 
 	virtual bool UsesVorbisChannelOrdering() const override { return true; }
+	virtual int GetStreamBufferSize() const { return MONO_PCM_BUFFER_SIZE; }
 
 	struct FVorbisFileWrapper* VFWrapper;
 	const uint8*		SrcBufferData;
diff --git a/Engine/Source/Runtime/Windows/XAudio2/Private/XAudio2Device.cpp b/Engine/Source/Runtime/Windows/XAudio2/Private/XAudio2Device.cpp
--- a/Engine/Source/Runtime/Windows/XAudio2/Private/XAudio2Device.cpp
+++ b/Engine/Source/Runtime/Windows/XAudio2/Private/XAudio2Device.cpp
@@ -244,11 +244,21 @@ class ICompressedAudioInfo* FXAudio2Device::CreateCompressedAudioInfo(USoundWave
 	{
 		return new FOpusAudioInfo();
 	}
+
 #if WITH_OGGVORBIS
-	return new FVorbisAudioInfo();
+	if (SoundWave->CompressionName.IsNone() || SoundWave->CompressionName == TEXT("OGG"))
+	{
+
+		return new FVorbisAudioInfo();
+	}
+	else
+	{
+		return NULL;
+	}
 #else
 	return NULL;
 #endif
+	
 }
 
 /**  
diff --git a/Engine/Source/Runtime/Windows/XAudio2/Public/XAudio2Device.h b/Engine/Source/Runtime/Windows/XAudio2/Public/XAudio2Device.h
--- a/Engine/Source/Runtime/Windows/XAudio2/Public/XAudio2Device.h
+++ b/Engine/Source/Runtime/Windows/XAudio2/Public/XAudio2Device.h
@@ -100,13 +100,18 @@ class FXAudio2Device : public FAudioDevice
 		{
 			return NAME_OPUS;
 		}
+		if (SoundWave->CompressionName.IsNone())
+		{
 #if WITH_OGGVORBIS
-		static FName NAME_OGG(TEXT("OGG"));
-		return NAME_OGG;
+			static FName NAME_OGG(TEXT("OGG"));
+			return NAME_OGG;
 #else //WITH_OGGVORBIS
-		static FName NAME_XMA(TEXT("XMA"));
-		return NAME_XMA;
+			static FName NAME_XMA(TEXT("XMA"));
+			return NAME_XMA;
 #endif //WITH_OGGVORBIS
+		}
+
+		return SoundWave->CompressionName;
 	}
 
 	virtual bool HasCompressedAudioInfoClass(USoundWave* SoundWave) override;
