Made all loaded and created notify objects transactional to fix undo not working. This affected all notify properties.
#ttp 343842 - PERSONA: Unable to Undo/Redo changing Anim Notify Color

diff --git a/Engine/Source/Editor/Persona/Private/SAnimNotifyPanel.cpp b/Engine/Source/Editor/Persona/Private/SAnimNotifyPanel.cpp
--- a/Engine/Source/Editor/Persona/Private/SAnimNotifyPanel.cpp
+++ b/Engine/Source/Editor/Persona/Private/SAnimNotifyPanel.cpp
@@ -1836,7 +1836,7 @@ void SAnimNotifyTrack::CreateNewNotifyAtCursor(FString NewNotifyName, UClass* No
 
 	if( NotifyClass )
 	{
-		class UObject * AnimNotifyClass = NewObject<UObject>(Sequence, NotifyClass);
+		class UObject* AnimNotifyClass = ConstructObject<UObject>(NotifyClass, Sequence, NAME_None, RF_Transactional);
 		NewEvent.NotifyStateClass = Cast<UAnimNotifyState>(AnimNotifyClass);
 		NewEvent.Notify = Cast<UAnimNotify>(AnimNotifyClass);
 
diff --git a/Engine/Source/Runtime/Engine/Classes/Animation/AnimNotifies/AnimNotify.h b/Engine/Source/Runtime/Engine/Classes/Animation/AnimNotifies/AnimNotify.h
--- a/Engine/Source/Runtime/Engine/Classes/Animation/AnimNotifies/AnimNotify.h
+++ b/Engine/Source/Runtime/Engine/Classes/Animation/AnimNotifies/AnimNotify.h
@@ -65,6 +65,10 @@ class ENGINE_API UAnimNotify : public UObject
 	 */
 	virtual class UWorld* GetWorld() const override;
 
+	/** UObject Interface */
+	virtual void PostLoad();
+	/** End UObject Interface */
+
 private:
 	/* The mesh we're currently triggering a UAnimNotify for (so we can retrieve per instance information) */
 	class USkeletalMeshComponent* MeshContext;
diff --git a/Engine/Source/Runtime/Engine/Classes/Animation/AnimNotifies/AnimNotifyState.h b/Engine/Source/Runtime/Engine/Classes/Animation/AnimNotifies/AnimNotifyState.h
--- a/Engine/Source/Runtime/Engine/Classes/Animation/AnimNotifies/AnimNotifyState.h
+++ b/Engine/Source/Runtime/Engine/Classes/Animation/AnimNotifies/AnimNotifyState.h
@@ -63,6 +63,10 @@ class ENGINE_API UAnimNotifyState : public UObject
 	 *	@param	OwnerEvent		The event that 'owns' this AnimNotify.
 	 */
 	virtual void AnimNotifyEventChanged(USkeletalMeshComponent* MeshComp, UAnimSequenceBase* Animation, FAnimNotifyEvent * OwnerEvent) {}
+
+	/** UObject Interface */
+	virtual void PostLoad();
+	/** End UObject Interface */
 };
 
 
diff --git a/Engine/Source/Runtime/Engine/Private/Animation/AnimNotify.cpp b/Engine/Source/Runtime/Engine/Private/Animation/AnimNotify.cpp
--- a/Engine/Source/Runtime/Engine/Private/Animation/AnimNotify.cpp
+++ b/Engine/Source/Runtime/Engine/Private/Animation/AnimNotify.cpp
@@ -50,3 +50,12 @@ FString UAnimNotify::GetNotifyName_Implementation() const
 	
 	return NotifyName;
 }
+
+void UAnimNotify::PostLoad()
+{
+	Super::PostLoad();
+#if WITH_EDITOR
+	// Ensure that all loaded notifies are transactional
+	SetFlags(GetFlags() | RF_Transactional);
+#endif
+}
diff --git a/Engine/Source/Runtime/Engine/Private/Animation/AnimNotifyState.cpp b/Engine/Source/Runtime/Engine/Private/Animation/AnimNotifyState.cpp
--- a/Engine/Source/Runtime/Engine/Private/Animation/AnimNotifyState.cpp
+++ b/Engine/Source/Runtime/Engine/Private/Animation/AnimNotifyState.cpp
@@ -52,3 +52,13 @@ FString UAnimNotifyState::GetNotifyName_Implementation() const
 	return NotifyName;
 }
 
+void UAnimNotifyState::PostLoad()
+{
+	Super::PostLoad();
+#if WITH_EDITOR
+	// Ensure that all loaded notifies are transactional
+	SetFlags(GetFlags() | RF_Transactional);
+#endif
+}
+
+
diff --git a/Engine/Source/Runtime/Engine/Private/Animation/AnimNotifyState_TimedParticleEffect.cpp b/Engine/Source/Runtime/Engine/Private/Animation/AnimNotifyState_TimedParticleEffect.cpp
--- a/Engine/Source/Runtime/Engine/Private/Animation/AnimNotifyState_TimedParticleEffect.cpp
+++ b/Engine/Source/Runtime/Engine/Private/Animation/AnimNotifyState_TimedParticleEffect.cpp
@@ -105,14 +105,17 @@ FString UAnimNotifyState_TimedParticleEffect::GetNotifyName_Implementation() con
 #if WITH_EDITORONLY_DATA
 void UAnimNotifyState_TimedParticleEffect::PreEditChange(UProperty* PropertyAboutToChange)
 {
-	if(PropertyAboutToChange->GetName() == GET_MEMBER_NAME_STRING_CHECKED(UAnimNotifyState_TimedParticleEffect, PSTemplate) && PSTemplate != NULL)
+	if(PropertyAboutToChange)
 	{
-		PreviousPSTemplates.Add(PSTemplate);
-	}
+		if(PropertyAboutToChange->GetName() == GET_MEMBER_NAME_STRING_CHECKED(UAnimNotifyState_TimedParticleEffect, PSTemplate) && PSTemplate != NULL)
+		{
+			PreviousPSTemplates.Add(PSTemplate);
+		}
 
-	if(PropertyAboutToChange->GetName() == GET_MEMBER_NAME_STRING_CHECKED(UAnimNotifyState_TimedParticleEffect, SocketName) && SocketName != FName(TEXT("None")))
-	{
-		PreviousSocketNames.Add(SocketName);
+		if(PropertyAboutToChange->GetName() == GET_MEMBER_NAME_STRING_CHECKED(UAnimNotifyState_TimedParticleEffect, SocketName) && SocketName != FName(TEXT("None")))
+		{
+			PreviousSocketNames.Add(SocketName);
+		}
 	}
 }
 #endif
