UMG - Working on Input and overriding input for widgets so that users can perform actions based on controller interaction in the UI.

diff --git a/Engine/Source/Editor/BlueprintGraph/Private/K2ActionMenuBuilder.cpp b/Engine/Source/Editor/BlueprintGraph/Private/K2ActionMenuBuilder.cpp
--- a/Engine/Source/Editor/BlueprintGraph/Private/K2ActionMenuBuilder.cpp
+++ b/Engine/Source/Editor/BlueprintGraph/Private/K2ActionMenuBuilder.cpp
@@ -841,7 +841,7 @@ void FK2ActionMenuBuilder::GetContextAllowedNodeTypes(FBlueprintGraphActionListB
 		{
 			AddSpawnActorNodeAction(ContextMenuBuilder, K2ActionCategories::GenericFunctionCategory);
 
-			if (Class->IsChildOf(AActor::StaticClass()))
+			if ( Blueprint->SupportsInputEvents() )
 			{
 				GetInputNodes(ContextMenuBuilder, true);
 			}
diff --git a/Engine/Source/Runtime/Engine/Classes/Engine/Blueprint.h b/Engine/Source/Runtime/Engine/Classes/Engine/Blueprint.h
--- a/Engine/Source/Runtime/Engine/Classes/Engine/Blueprint.h
+++ b/Engine/Source/Runtime/Engine/Classes/Engine/Blueprint.h
@@ -517,6 +517,9 @@ public:
 	/** @return true if the blueprint supports event binding for multicast delegates */
 	virtual bool AllowsDynamicBinding() const;
 
+	/** @return true if the blueprint supports event binding for input events */
+	virtual bool SupportsInputEvents() const;
+
 	bool ChangeOwnerOfTemplates();
 
 #endif	//#if WITH_EDITOR
diff --git a/Engine/Source/Runtime/Engine/Classes/Engine/InputDelegateBinding.h b/Engine/Source/Runtime/Engine/Classes/Engine/InputDelegateBinding.h
--- a/Engine/Source/Runtime/Engine/Classes/Engine/InputDelegateBinding.h
+++ b/Engine/Source/Runtime/Engine/Classes/Engine/InputDelegateBinding.h
@@ -33,5 +33,5 @@ class ENGINE_API UInputDelegateBinding : public UDynamicBlueprintBinding
 
 	virtual void BindToInputComponent(UInputComponent* InputComponent) const { };
 
-	static void BindInputDelegates(UBlueprintGeneratedClass* BGClass, UInputComponent* InputComponent);
+	static void BindInputDelegates(const UBlueprintGeneratedClass* BGClass, UInputComponent* InputComponent);
 };
diff --git a/Engine/Source/Runtime/Engine/Classes/Kismet/KismetInputLibrary.h b/Engine/Source/Runtime/Engine/Classes/Kismet/KismetInputLibrary.h
--- a/Engine/Source/Runtime/Engine/Classes/Kismet/KismetInputLibrary.h
+++ b/Engine/Source/Runtime/Engine/Classes/Kismet/KismetInputLibrary.h
@@ -2,6 +2,8 @@
 
 #pragma once
 
+#include "SlateCore.h"
+
 #include "KismetInputLibrary.generated.h"
 
 UCLASS(MinimalAPI)
@@ -12,4 +14,117 @@ class UKismetInputLibrary : public UBlueprintFunctionLibrary
 	/** Calibrate the tilt for the input device */
 	UFUNCTION(BlueprintCallable, Category="Input")
 	static void CalibrateTilt();
+
+	/**
+	 * Test if the input key are equal (A == B)
+	 * @param A - The key to compare against
+	 * @param B - The key to compare
+	 * @returns True if the key are equal, false otherwise
+	 */
+	UFUNCTION(BlueprintPure, meta=(FriendlyName = "Equal (Key)", CompactNodeTitle = "=="), Category="Utilities|Key")
+	static bool EqualEqual_KeyKey(FKey A, FKey B);
+
+	/**
+	 * Returns whether or not this character is an auto-repeated keystroke
+	 *
+	 * @return  True if this character is a repeat
+	 */
+	UFUNCTION(BlueprintPure, meta=( FriendlyName = "IsRepeat" ), Category="Utilities|InputEvent")
+	static bool InputEvent_IsRepeat(const FInputEvent& Input);
+
+	/**
+	 * Returns true if either shift key was down when this event occurred
+	 *
+	 * @return  True if shift is pressed
+	 */
+	UFUNCTION(BlueprintPure, meta=( FriendlyName = "IsShiftDown" ), Category="Utilities|InputEvent")
+	static bool InputEvent_IsShiftDown(const FInputEvent& Input);
+
+	/**
+	 * Returns true if left shift key was down when this event occurred
+	 *
+	 * @return True if left shift is pressed.
+	 */
+	UFUNCTION(BlueprintPure, meta=( FriendlyName = "IsLeftShiftDown" ), Category="Utilities|InputEvent")
+	static bool InputEvent_IsLeftShiftDown(const FInputEvent& Input);
+
+	/**
+	 * Returns true if right shift key was down when this event occurred
+	 *
+	 * @return True if right shift is pressed.
+	 */
+	UFUNCTION(BlueprintPure, meta=( FriendlyName = "IsRightShiftDown" ), Category="Utilities|InputEvent")
+	static bool InputEvent_IsRightShiftDown(const FInputEvent& Input);
+
+	/**
+	 * Returns true if either control key was down when this event occurred
+	 *
+	 * @return  True if control is pressed
+	 */
+	UFUNCTION(BlueprintPure, meta=( FriendlyName = "IsControlDown" ), Category="Utilities|InputEvent")
+	static bool InputEvent_IsControlDown(const FInputEvent& Input);
+
+	/**
+	 * Returns true if left control key was down when this event occurred
+	 *
+	 * @return  True if left control is pressed
+	 */
+	UFUNCTION(BlueprintPure, meta=( FriendlyName = "IsLeftControlDown" ), Category="Utilities|InputEvent")
+	static bool InputEvent_IsLeftControlDown(const FInputEvent& Input);
+
+	/**
+	 * Returns true if left control key was down when this event occurred
+	 *
+	 * @return  True if left control is pressed
+	 */
+	UFUNCTION(BlueprintPure, meta=( FriendlyName = "IsRightControlDown" ), Category="Utilities|InputEvent")
+	static bool InputEvent_IsRightControlDown(const FInputEvent& Input);
+
+	/**
+	 * Returns true if either alt key was down when this event occurred
+	 *
+	 * @return  True if alt is pressed
+	 */
+	UFUNCTION(BlueprintPure, meta=( FriendlyName = "IsAltDown" ), Category="Utilities|InputEvent")
+	static bool InputEvent_IsAltDown(const FInputEvent& Input);
+
+	/**
+	 * Returns true if left alt key was down when this event occurred
+	 *
+	 * @return  True if left alt is pressed
+	 */
+	UFUNCTION(BlueprintPure, meta=( FriendlyName = "IsLeftAltDown" ), Category="Utilities|InputEvent")
+	static bool InputEvent_IsLeftAltDown(const FInputEvent& Input);
+
+	/**
+	 * Returns true if right alt key was down when this event occurred
+	 *
+	 * @return  True if right alt is pressed
+	 */
+	UFUNCTION(BlueprintPure, meta=( FriendlyName = "IsRightAltDown" ), Category="Utilities|InputEvent")
+	static bool InputEvent_IsRightAltDown(const FInputEvent& Input);
+
+	/**
+	 * Returns true if either command key was down when this event occurred
+	 *
+	 * @return  True if command is pressed
+	 */
+	UFUNCTION(BlueprintPure, meta=( FriendlyName = "IsCommandDown" ), Category="Utilities|InputEvent")
+	static bool InputEvent_IsCommandDown(const FInputEvent& Input);
+
+	/**
+	 * Returns true if left command key was down when this event occurred
+	 *
+	 * @return  True if left command is pressed
+	 */
+	UFUNCTION(BlueprintPure, meta=( FriendlyName = "IsLeftCommandDown" ), Category="Utilities|InputEvent")
+	static bool InputEvent_IsLeftCommandDown(const FInputEvent& Input);
+
+	/**
+	 * Returns true if right command key was down when this event occurred
+	 *
+	 * @return  True if right command is pressed
+	 */
+	UFUNCTION(BlueprintPure, meta=( FriendlyName = "IsRightCommandDown" ), Category="Utilities|InputEvent")
+	static bool InputEvent_IsRightCommandDown(const FInputEvent& Input);
 };
diff --git a/Engine/Source/Runtime/Engine/Private/Blueprint.cpp b/Engine/Source/Runtime/Engine/Private/Blueprint.cpp
--- a/Engine/Source/Runtime/Engine/Private/Blueprint.cpp
+++ b/Engine/Source/Runtime/Engine/Private/Blueprint.cpp
@@ -699,6 +699,11 @@ bool UBlueprint::AllowsDynamicBinding() const
 	return FBlueprintEditorUtils::IsActorBased(this);
 }
 
+bool UBlueprint::SupportsInputEvents() const
+{
+	return FBlueprintEditorUtils::IsActorBased(this);
+}
+
 struct FBlueprintInnerHelper
 {
 	template<typename TOBJ, typename TARR>
diff --git a/Engine/Source/Runtime/Engine/Private/InputDelegateBinding.cpp b/Engine/Source/Runtime/Engine/Private/InputDelegateBinding.cpp
--- a/Engine/Source/Runtime/Engine/Private/InputDelegateBinding.cpp
+++ b/Engine/Source/Runtime/Engine/Private/InputDelegateBinding.cpp
@@ -15,7 +15,7 @@ UInputDelegateBinding::UInputDelegateBinding(const class FPostConstructInitializ
 {
 }
 
-void UInputDelegateBinding::BindInputDelegates(UBlueprintGeneratedClass* BGClass, UInputComponent* InputComponent)
+void UInputDelegateBinding::BindInputDelegates(const UBlueprintGeneratedClass* BGClass, UInputComponent* InputComponent)
 {
 	static UClass* InputBindingClasses[] = { 
 												UInputActionDelegateBinding::StaticClass(), 
diff --git a/Engine/Source/Runtime/Engine/Private/KismetInputLibrary.cpp b/Engine/Source/Runtime/Engine/Private/KismetInputLibrary.cpp
--- a/Engine/Source/Runtime/Engine/Private/KismetInputLibrary.cpp
+++ b/Engine/Source/Runtime/Engine/Private/KismetInputLibrary.cpp
@@ -16,10 +16,79 @@ UKismetInputLibrary::UKismetInputLibrary(const class FPostConstructInitializePro
 {
 }
 
-
 void UKismetInputLibrary::CalibrateTilt()
 {
 	GEngine->Exec(NULL, TEXT("CALIBRATEMOTION"));
 }
 
+bool UKismetInputLibrary::EqualEqual_KeyKey(FKey A, FKey B)
+{
+	return A == B;
+}
+
+bool UKismetInputLibrary::InputEvent_IsRepeat(const FInputEvent& Input)
+{
+	return Input.IsRepeat();
+}
+
+bool UKismetInputLibrary::InputEvent_IsShiftDown(const FInputEvent& Input)
+{
+	return Input.IsShiftDown();
+}
+
+bool UKismetInputLibrary::InputEvent_IsLeftShiftDown(const FInputEvent& Input)
+{
+	return Input.IsLeftShiftDown();
+}
+
+bool UKismetInputLibrary::InputEvent_IsRightShiftDown(const FInputEvent& Input)
+{
+	return Input.IsRightShiftDown();
+}
+
+bool UKismetInputLibrary::InputEvent_IsControlDown(const FInputEvent& Input)
+{
+	return Input.IsControlDown();
+}
+
+bool UKismetInputLibrary::InputEvent_IsLeftControlDown(const FInputEvent& Input)
+{
+	return Input.IsLeftControlDown();
+}
+
+bool UKismetInputLibrary::InputEvent_IsRightControlDown(const FInputEvent& Input)
+{
+	return Input.IsRightControlDown();
+}
+
+bool UKismetInputLibrary::InputEvent_IsAltDown(const FInputEvent& Input)
+{
+	return Input.IsAltDown();
+}
+
+bool UKismetInputLibrary::InputEvent_IsLeftAltDown(const FInputEvent& Input)
+{
+	return Input.IsLeftAltDown();
+}
+
+bool UKismetInputLibrary::InputEvent_IsRightAltDown(const FInputEvent& Input)
+{
+	return Input.IsRightAltDown();
+}
+
+bool UKismetInputLibrary::InputEvent_IsCommandDown(const FInputEvent& Input)
+{
+	return Input.IsCommandDown();
+}
+
+bool UKismetInputLibrary::InputEvent_IsLeftCommandDown(const FInputEvent& Input)
+{
+	return Input.IsLeftCommandDown();
+}
+
+bool UKismetInputLibrary::InputEvent_IsRightCommandDown(const FInputEvent& Input)
+{
+	return Input.IsRightCommandDown();
+}
+
 #undef LOCTEXT_NAMESPACE
\ No newline at end of file
