#ue4 - Add optimized path for TransformNormalToShapeSpace() when rotation is identity, and avoid Load-hit-stores with the result.

- 100% of my test cases w/ non-uniform scale and rotated assets had identiy rotation within PhysX.

diff --git a/Engine/Source/Runtime/Engine/Private/Collision/CollisionConversions.cpp b/Engine/Source/Runtime/Engine/Private/Collision/CollisionConversions.cpp
--- a/Engine/Source/Runtime/Engine/Private/Collision/CollisionConversions.cpp
+++ b/Engine/Source/Runtime/Engine/Private/Collision/CollisionConversions.cpp
@@ -40,29 +40,44 @@ static void CheckHitResultNormal(const FHitResult& OutResult, const TCHAR* Messa
 #endif // !(UE_BUILD_SHIPPING || UE_BUILD_TEST)
 
 
+static FORCEINLINE bool PxQuatIsIdentity(PxQuat const& Q)
+{
+	return
+		Q.x == 0.f &&
+		Q.y == 0.f &&
+		Q.z == 0.f &&
+		Q.w == 1.f;
+}
+
 
 /** Helper to transform a normal when non-uniform scale is present. */
-// TODO: I'm sure there is a more efficient approach.
-static void TransformNormalToShapeSpace(const PxMeshScale& meshScale, const PxVec3& nIn, PxVec3& nOut)
+static PxVec3 TransformNormalToShapeSpace(const PxMeshScale& meshScale, const PxVec3& nIn)
 {
 	// Uniform scale makes this unnecessary
 	if (meshScale.scale.x == meshScale.scale.y &&
 		meshScale.scale.x == meshScale.scale.z)
 	{
-		nOut = nIn;
-		return;
+		return nIn;
 	}
+	
+	if (PxQuatIsIdentity(meshScale.rotation))
+	{
+		// Inverse transpose: inverse is 1/scale, transpose = original when rotation is identity.
+		const PxVec3 tmp = PxVec3(nIn.x / meshScale.scale.x, nIn.y / meshScale.scale.y, nIn.z / meshScale.scale.z);
+		const PxReal denom = 1.0f / tmp.magnitude();
+		return tmp * denom;
+	}
+	else
+	{
+		const PxMat33 rot(meshScale.rotation);
+		const PxMat33 diagonal = PxMat33::createDiagonal(meshScale.scale);
+		const PxMat33 vertex2Shape = (rot.getTranspose() * diagonal) * rot;
 
-	PxMat33 R(meshScale.rotation);
-	PxMat33 vertex2ShapeSkew = R.getTranspose();
-	PxMat33 diagonal = PxMat33::createDiagonal(meshScale.scale);
-	vertex2ShapeSkew = vertex2ShapeSkew * diagonal;
-	vertex2ShapeSkew = vertex2ShapeSkew * R;
-
-	PxMat33 shape2VertexSkew = vertex2ShapeSkew.getInverse();
-	const PxVec3 tmp = shape2VertexSkew.transformTranspose(nIn);
-	const PxReal Denom = 1.0f / tmp.magnitude();
-	nOut = tmp * Denom;
+		const PxMat33 shape2Vertex = vertex2Shape.getInverse();
+		const PxVec3 tmp = shape2Vertex.transformTranspose(nIn);
+		const PxReal denom = 1.0f / tmp.magnitude();
+		return tmp * denom;
+	}
 }
 
 
@@ -117,12 +132,12 @@ static bool FindGeomOpposingNormal(const PxLocationHit& PHit, const FVector& Dir
 		const PxVec3 V1 = PVerts[I1];
 		const PxVec3 V2 = PVerts[I2];
 
-		// Find normal of triangle, and convert into world space
-		PxVec3 PLocalTriNormal = ((V1 - V0).cross(V2 - V0)).getNormalized();
-		TransformNormalToShapeSpace(PTriMeshGeom.scale, PLocalTriNormal, PLocalTriNormal);
+		// Find normal of triangle (local space), and account for non-uniform scale
+		const PxVec3 PTempNormal = ((V1 - V0).cross(V2 - V0)).getNormalized();
+		const PxVec3 PLocalTriNormal = TransformNormalToShapeSpace(PTriMeshGeom.scale, PTempNormal);
 	
-		const PxTransform PShapeWorldPose = PxShapeExt::getGlobalPose(*PHit.shape, *PHit.actor); 
-
+		// Convert to world space
+		const PxTransform PShapeWorldPose = PxShapeExt::getGlobalPose(*PHit.shape, *PHit.actor);
 		const PxVec3 PWorldTriNormal = PShapeWorldPose.rotate(PLocalTriNormal);
 		OutNormal = P2UVector(PWorldTriNormal);
 
@@ -151,14 +166,12 @@ static bool FindGeomOpposingNormal(const PxLocationHit& PHit, const FVector& Dir
 		bool bSuccess = PConvexMeshGeom.convexMesh->getPolygonData(PolyIndex, PPoly);
 		if(bSuccess)
 		{
-			PxVec3 PPlaneNormal(PPoly.mPlane[0], PPoly.mPlane[1], PPoly.mPlane[2]);
+			// Account for non-uniform scale in local space normal.
+			const PxVec3 PPlaneNormal(PPoly.mPlane[0], PPoly.mPlane[1], PPoly.mPlane[2]);
+			const PxVec3 PLocalPolyNormal = TransformNormalToShapeSpace(PConvexMeshGeom.scale, PPlaneNormal.getNormalized());			
 
-			// Convert to local space, taking scale into account.
-			// TODO: can we just change the hit normal within the physx code where we choose the most opposing normal, and avoid doing this here again?
-			PxVec3 PLocalPolyNormal;
-			TransformNormalToShapeSpace(PConvexMeshGeom.scale, PPlaneNormal.getNormalized(), PLocalPolyNormal);			
+			// Convert to world space
 			const PxTransform PShapeWorldPose = PxShapeExt::getGlobalPose(*PHit.shape, *PHit.actor); 
-
 			const PxVec3 PWorldPolyNormal = PShapeWorldPose.rotate(PLocalPolyNormal);
 			OutNormal = P2UVector(PWorldPolyNormal);
 
