#UE4

- Adding a getter for the player's HUD on FLocalPlayerContext

diff --git a/Engine/Source/Runtime/Engine/Classes/Engine/LocalPlayer.h b/Engine/Source/Runtime/Engine/Classes/Engine/LocalPlayer.h
--- a/Engine/Source/Runtime/Engine/Classes/Engine/LocalPlayer.h
+++ b/Engine/Source/Runtime/Engine/Classes/Engine/LocalPlayer.h
@@ -57,6 +57,16 @@ struct ENGINE_API FLocalPlayerContext
 		return CastChecked<T>(GetPlayerState(), ECastCheckedType::NullAllowed);
 	}
 
+	/** Getter for this player's HUD */
+	class AHUD* GetHUD() const;
+
+	/** Templated Getter for the HUD */
+	template<class T>
+	FORCEINLINE T* GetHUD() const
+	{
+		return CastChecked<T>(GetHUD(), ECastCheckedType::NullAllowed);
+	}
+
 private:	
 
 	/* Set the local player. */
diff --git a/Engine/Source/Runtime/Engine/Private/LocalPlayer.cpp b/Engine/Source/Runtime/Engine/Private/LocalPlayer.cpp
--- a/Engine/Source/Runtime/Engine/Private/LocalPlayer.cpp
+++ b/Engine/Source/Runtime/Engine/Private/LocalPlayer.cpp
@@ -110,6 +110,13 @@ APlayerState* FLocalPlayerContext::GetPlayerState() const
 	return PC ? PC->PlayerState : NULL;
 }
 
+AHUD* FLocalPlayerContext::GetHUD() const
+{
+	check(LocalPlayer.IsValid())
+	APlayerController* PC = LocalPlayer->PlayerController;
+	return PC ? PC->MyHUD : NULL;
+}
+
 void FLocalPlayerContext::SetLocalPlayer( const ULocalPlayer* InLocalPlayer )
 {
 	LocalPlayer = InLocalPlayer;
