#ue4 - Improved comments around Crouching in Character. Changed displayname of NavAgentProps to make CanCrouch/CanJump etc easier to find.

diff --git a/Engine/Source/Runtime/Engine/Classes/GameFramework/Character.h b/Engine/Source/Runtime/Engine/Classes/GameFramework/Character.h
--- a/Engine/Source/Runtime/Engine/Classes/GameFramework/Character.h
+++ b/Engine/Source/Runtime/Engine/Classes/GameFramework/Character.h
@@ -405,7 +405,10 @@ protected:
 
 public:
 
-	/** True if jump is actively providing a force; i.e. jump key is held and the time it has been held is less than JumpMaxHoldTime. */
+	/**
+	 * True if jump is actively providing a force; i.e. jump key is held and the time it has been held is less than JumpMaxHoldTime.
+	 * @see CharacterMovement->IsFalling
+	 */
 	UFUNCTION(BlueprintCallable, Category="Pawn|Character")
 	virtual bool IsJumpProvidingForce() const;
 
@@ -477,11 +480,21 @@ public:
 		@PARAM Impact describes the blocking hit. */
 	virtual void MoveBlockedBy(const FHitResult& Impact) {};
 
-	/** Make the character start crouching.	 */
+	/**
+	 * Request the character to start crouching. The request is processed on the next update of the CharacterMovementComponent.
+	 * @see OnStartCrouch
+	 * @see IsCrouched
+	 * @see CharacterMovement->WantsToCrouch
+	 */
 	UFUNCTION(BlueprintCallable, Category="Pawn|Character", meta=(HidePin="bClientSimulation"))
 	virtual void Crouch(bool bClientSimulation = false);
 
-	/** Make the character stop crouching.	 */
+	/**
+	 * Request the character to stop crouching. The request is processed on the next update of the CharacterMovementComponent.
+	 * @see OnEndCrouch
+	 * @see IsCrouched
+	 * @see CharacterMovement->WantsToCrouch
+	 */
 	UFUNCTION(BlueprintCallable, Category="Pawn|Character", meta=(HidePin="bClientSimulation"))
 	virtual void UnCrouch(bool bClientSimulation = false);
 
diff --git a/Engine/Source/Runtime/Engine/Classes/GameFramework/CharacterMovementComponent.h b/Engine/Source/Runtime/Engine/Classes/GameFramework/CharacterMovementComponent.h
--- a/Engine/Source/Runtime/Engine/Classes/GameFramework/CharacterMovementComponent.h
+++ b/Engine/Source/Runtime/Engine/Classes/GameFramework/CharacterMovementComponent.h
@@ -519,8 +519,8 @@ public:
 	UPROPERTY()
 	uint32 bCheatFlying:1;
 
-	/** If true, try to crouch (or keep crouching). If false, stop crouching if crouched. */
-	UPROPERTY()
+	/** If true, try to crouch (or keep crouching) on next update. If false, try to stop crouching on next update. */
+	UPROPERTY(Category="Character Movement", VisibleInstanceOnly, BlueprintReadOnly)
 	uint32 bWantsToCrouch:1;
 
 	/**
diff --git a/Engine/Source/Runtime/Engine/Classes/GameFramework/NavMovementComponent.h b/Engine/Source/Runtime/Engine/Classes/GameFramework/NavMovementComponent.h
--- a/Engine/Source/Runtime/Engine/Classes/GameFramework/NavMovementComponent.h
+++ b/Engine/Source/Runtime/Engine/Classes/GameFramework/NavMovementComponent.h
@@ -18,8 +18,8 @@ class ENGINE_API UNavMovementComponent : public UMovementComponent
 {
 	GENERATED_UCLASS_BODY()
 
-	/** Properties that define how navigation can interact with this component. */
-	UPROPERTY(EditAnywhere, BlueprintReadWrite, Category=MovementComponent)
+	/** Properties that define how the component can move. */
+	UPROPERTY(EditAnywhere, BlueprintReadWrite, Category=MovementComponent, meta=(DisplayName="Movement Capabilities", Keywords="Nav Agent"))
 	FNavAgentProperties NavAgentProps;
 
 protected:
