WITH_EDITOR and WITH_EDITORONLY_DATA cleanup

diff --git a/Engine/Plugins/Experimental/Paper2D/Source/Paper2D/Classes/PaperFlipbook.h b/Engine/Plugins/Experimental/Paper2D/Source/Paper2D/Classes/PaperFlipbook.h
--- a/Engine/Plugins/Experimental/Paper2D/Source/Paper2D/Classes/PaperFlipbook.h
+++ b/Engine/Plugins/Experimental/Paper2D/Source/Paper2D/Classes/PaperFlipbook.h
@@ -87,7 +87,7 @@ public:
 	void InvalidateCachedData();
 
 	// UObject interface
-#if WITH_EDITORONLY_DATA
+#if WITH_EDITOR
 	virtual void Serialize(FArchive& Ar) override;
 	virtual void PostLoad() override;
 	virtual void PostEditChangeProperty(FPropertyChangedEvent& PropertyChangedEvent) override;
diff --git a/Engine/Plugins/Experimental/Paper2D/Source/Paper2D/Classes/PaperSprite.h b/Engine/Plugins/Experimental/Paper2D/Source/Paper2D/Classes/PaperSprite.h
--- a/Engine/Plugins/Experimental/Paper2D/Source/Paper2D/Classes/PaperSprite.h
+++ b/Engine/Plugins/Experimental/Paper2D/Source/Paper2D/Classes/PaperSprite.h
@@ -155,7 +155,7 @@ public:
 
 public:
 
-#if WITH_EDITORONLY_DATA
+#if WITH_EDITOR
 	// UObject interface
 	virtual void PostEditChangeProperty(FPropertyChangedEvent& PropertyChangedEvent) override;
 	virtual void Serialize(FArchive& Ar) override;
diff --git a/Engine/Plugins/Experimental/Paper2D/Source/Paper2D/Classes/PaperTerrainComponent.h b/Engine/Plugins/Experimental/Paper2D/Source/Paper2D/Classes/PaperTerrainComponent.h
--- a/Engine/Plugins/Experimental/Paper2D/Source/Paper2D/Classes/PaperTerrainComponent.h
+++ b/Engine/Plugins/Experimental/Paper2D/Source/Paper2D/Classes/PaperTerrainComponent.h
@@ -50,7 +50,7 @@ protected:
 public:
 	// UObject interface
 	virtual const UObject* AdditionalStatObject() const override;
-#if WITH_EDITORONLY_DATA
+#if WITH_EDITOR
 	virtual void PostEditChangeProperty(FPropertyChangedEvent& PropertyChangedEvent) override;
 #endif
 	// End of UObject interface
diff --git a/Engine/Plugins/Experimental/Paper2D/Source/Paper2D/Classes/PaperTerrainSplineComponent.h b/Engine/Plugins/Experimental/Paper2D/Source/Paper2D/Classes/PaperTerrainSplineComponent.h
--- a/Engine/Plugins/Experimental/Paper2D/Source/Paper2D/Classes/PaperTerrainSplineComponent.h
+++ b/Engine/Plugins/Experimental/Paper2D/Source/Paper2D/Classes/PaperTerrainSplineComponent.h
@@ -13,7 +13,7 @@ class PAPER2D_API UPaperTerrainSplineComponent : public USplineComponent
 
 public:
 	// UObject interface
-#if WITH_EDITORONLY_DATA
+#if WITH_EDITOR
 	virtual void PostEditChangeProperty(FPropertyChangedEvent& PropertyChangedEvent) override;
 #endif
 	// End of UObject interface
diff --git a/Engine/Plugins/Experimental/Paper2D/Source/Paper2D/Classes/PaperTileLayer.h b/Engine/Plugins/Experimental/Paper2D/Source/Paper2D/Classes/PaperTileLayer.h
--- a/Engine/Plugins/Experimental/Paper2D/Source/Paper2D/Classes/PaperTileLayer.h
+++ b/Engine/Plugins/Experimental/Paper2D/Source/Paper2D/Classes/PaperTileLayer.h
@@ -93,7 +93,7 @@ public:
 	void ConvertToTileSetPerCell();
 
 	// UObject interface
-#if WITH_EDITORONLY_DATA
+#if WITH_EDITOR
 	virtual void PostEditChangeProperty(struct FPropertyChangedEvent& PropertyChangedEvent) override;
 #endif
 	// End of UObject interface
diff --git a/Engine/Plugins/Experimental/Paper2D/Source/Paper2D/Private/PaperFlipbook.cpp b/Engine/Plugins/Experimental/Paper2D/Source/Paper2D/Private/PaperFlipbook.cpp
--- a/Engine/Plugins/Experimental/Paper2D/Source/Paper2D/Private/PaperFlipbook.cpp
+++ b/Engine/Plugins/Experimental/Paper2D/Source/Paper2D/Private/PaperFlipbook.cpp
@@ -84,7 +84,7 @@ UPaperSprite* UPaperFlipbook::GetSpriteAtFrame(int32 FrameIndex) const
 	return KeyFrames.IsValidIndex(FrameIndex) ? KeyFrames[FrameIndex].Sprite : nullptr;
 }
 
-#if WITH_EDITORONLY_DATA
+#if WITH_EDITOR
 void UPaperFlipbook::Serialize(FArchive& Ar)
 {
 	Super::Serialize(Ar);
diff --git a/Engine/Plugins/Experimental/Paper2D/Source/Paper2D/Private/PaperSprite.cpp b/Engine/Plugins/Experimental/Paper2D/Source/Paper2D/Private/PaperSprite.cpp
--- a/Engine/Plugins/Experimental/Paper2D/Source/Paper2D/Private/PaperSprite.cpp
+++ b/Engine/Plugins/Experimental/Paper2D/Source/Paper2D/Private/PaperSprite.cpp
@@ -5,7 +5,7 @@
 #include "PhysicsEngine/BodySetup.h"
 #include "PaperCustomVersion.h"
 
-#if WITH_EDITORONLY_DATA
+#if WITH_EDITOR
 
 #include "PhysicsEngine/BodySetup2D.h"
 #include "PaperSpriteAtlas.h"
@@ -257,7 +257,7 @@ UPaperSprite::UPaperSprite(const FPostConstructInitializeProperties& PCIP)
 	
 	AlternateMaterialSplitIndex = INDEX_NONE;
 
-#if WITH_EDITORONLY_DATA
+#if WITH_EDITOR
 	PivotMode = ESpritePivotMode::Center_Center;
 	bSnapPivotToPixelGrid = true;
 
@@ -277,7 +277,7 @@ UPaperSprite::UPaperSprite(const FPostConstructInitializeProperties& PCIP)
 	AlternateMaterial = OpaqueMaterialRef.Object;
 }
 
-#if WITH_EDITORONLY_DATA
+#if WITH_EDITOR
 void UPaperSprite::PostEditChangeProperty(FPropertyChangedEvent& PropertyChangedEvent)
 {
 	//@TODO: Determine when these are really needed, as they're seriously expensive!
@@ -1306,7 +1306,7 @@ void UPaperSprite::QuerySupportedSockets(TArray<FComponentSocketDescription>& Ou
 	}
 }
 
-#if WITH_EDITORONLY_DATA
+#if WITH_EDITOR
 void UPaperSprite::Serialize(FArchive& Ar)
 {
 	Super::Serialize(Ar);
diff --git a/Engine/Plugins/Experimental/Paper2D/Source/Paper2D/Private/PaperTileLayer.cpp b/Engine/Plugins/Experimental/Paper2D/Source/Paper2D/Private/PaperTileLayer.cpp
--- a/Engine/Plugins/Experimental/Paper2D/Source/Paper2D/Private/PaperTileLayer.cpp
+++ b/Engine/Plugins/Experimental/Paper2D/Source/Paper2D/Private/PaperTileLayer.cpp
@@ -65,7 +65,7 @@ void UPaperTileLayer::ReallocateAndCopyMap()
 //	BakeMap();
 }
 
-#if WITH_EDITORONLY_DATA
+#if WITH_EDITOR
 void UPaperTileLayer::PostEditChangeProperty(struct FPropertyChangedEvent& PropertyChangedEvent)
 {
 	FName PropertyName = (PropertyChangedEvent.Property != NULL) ? PropertyChangedEvent.Property->GetFName() : NAME_None;
diff --git a/Engine/Plugins/Experimental/Paper2D/Source/Paper2D/Private/Terrain/PaperTerrainComponent.cpp b/Engine/Plugins/Experimental/Paper2D/Source/Paper2D/Private/Terrain/PaperTerrainComponent.cpp
--- a/Engine/Plugins/Experimental/Paper2D/Source/Paper2D/Private/Terrain/PaperTerrainComponent.cpp
+++ b/Engine/Plugins/Experimental/Paper2D/Source/Paper2D/Private/Terrain/PaperTerrainComponent.cpp
@@ -95,7 +95,7 @@ const UObject* UPaperTerrainComponent::AdditionalStatObject() const
 	return TerrainMaterial;
 }
 
-#if WITH_EDITORONLY_DATA
+#if WITH_EDITOR
 void UPaperTerrainComponent::PostEditChangeProperty(FPropertyChangedEvent& PropertyChangedEvent)
 {
 	if (PropertyChangedEvent.Property != nullptr)
diff --git a/Engine/Plugins/Experimental/Paper2D/Source/Paper2D/Private/Terrain/PaperTerrainSplineComponent.cpp b/Engine/Plugins/Experimental/Paper2D/Source/Paper2D/Private/Terrain/PaperTerrainSplineComponent.cpp
--- a/Engine/Plugins/Experimental/Paper2D/Source/Paper2D/Private/Terrain/PaperTerrainSplineComponent.cpp
+++ b/Engine/Plugins/Experimental/Paper2D/Source/Paper2D/Private/Terrain/PaperTerrainSplineComponent.cpp
@@ -10,7 +10,7 @@ UPaperTerrainSplineComponent::UPaperTerrainSplineComponent(const FPostConstructI
 {
 }
 
-#if WITH_EDITORONLY_DATA
+#if WITH_EDITOR
 void UPaperTerrainSplineComponent::PostEditChangeProperty(FPropertyChangedEvent& PropertyChangedEvent)
 {
 	OnSplineEdited.ExecuteIfBound();
diff --git a/Engine/Source/Developer/GameplayDebugger/Private/GameplayDebuggingHUDComponent.cpp b/Engine/Source/Developer/GameplayDebugger/Private/GameplayDebuggingHUDComponent.cpp
--- a/Engine/Source/Developer/GameplayDebugger/Private/GameplayDebuggingHUDComponent.cpp
+++ b/Engine/Source/Developer/GameplayDebugger/Private/GameplayDebuggingHUDComponent.cpp
@@ -24,9 +24,11 @@ AGameplayDebuggingHUDComponent::AGameplayDebuggingHUDComponent(const class FPost
 	, EngineShowFlags(EShowFlagInitMode::ESFIM_Game)
 {
 	World = NULL;
-#if WITH_EDITORONLY_DATA
+#if WITH_EDITOR
 	SetIsTemporarilyHiddenInEditor(true);
 	SetActorHiddenInGame(false);
+#endif
+#if WITH_EDITORONLY_DATA
 	bHiddenEdLevel = true;
 	bHiddenEdLayer = true;
 	bHiddenEd = true;
diff --git a/Engine/Source/Developer/GameplayDebugger/Private/GameplayDebuggingReplicator.cpp b/Engine/Source/Developer/GameplayDebugger/Private/GameplayDebuggingReplicator.cpp
--- a/Engine/Source/Developer/GameplayDebugger/Private/GameplayDebuggingReplicator.cpp
+++ b/Engine/Source/Developer/GameplayDebugger/Private/GameplayDebuggingReplicator.cpp
@@ -27,10 +27,12 @@ AGameplayDebuggingReplicator::AGameplayDebuggingReplicator(const class FPostCons
 	TSubobjectPtr<USceneComponent> SceneComponent = PCIP.CreateDefaultSubobject<USceneComponent>(this, TEXT("SceneComponent"));
 	RootComponent = SceneComponent;
 
-#if WITH_EDITORONLY_DATA
+#if WITH_EDITOR
 	SetTickableWhenPaused(true);
 	SetIsTemporarilyHiddenInEditor(false);
 	SetActorHiddenInGame(false);
+#endif
+#if WITH_EDITORONLY_DATA
 	bHiddenEdLevel = false;
 	bHiddenEdLayer = false;
 	bHiddenEd = false;
diff --git a/Engine/Source/Runtime/Engine/Classes/Animation/AnimNotifies/AnimNotifyState_TimedParticleEffect.h b/Engine/Source/Runtime/Engine/Classes/Animation/AnimNotifies/AnimNotifyState_TimedParticleEffect.h
--- a/Engine/Source/Runtime/Engine/Classes/Animation/AnimNotifies/AnimNotifyState_TimedParticleEffect.h
+++ b/Engine/Source/Runtime/Engine/Classes/Animation/AnimNotifies/AnimNotifyState_TimedParticleEffect.h
@@ -44,7 +44,10 @@ class UAnimNotifyState_TimedParticleEffect : public UAnimNotifyState
 
 	UPROPERTY(transient)
 	TArray<FName> PreviousSocketNames;
+	
+#endif
 
+#if WITH_EDITOR
 	virtual void PreEditChange(UProperty* PropertyAboutToChange) override;
 #endif
 
diff --git a/Engine/Source/Runtime/Engine/Classes/Camera/CameraComponent.h b/Engine/Source/Runtime/Engine/Classes/Camera/CameraComponent.h
--- a/Engine/Source/Runtime/Engine/Classes/Camera/CameraComponent.h
+++ b/Engine/Source/Runtime/Engine/Classes/Camera/CameraComponent.h
@@ -77,8 +77,9 @@ protected:
 	// The camera mesh to show visually where the camera is placed
 	UPROPERTY(transient)
 	class UStaticMeshComponent* ProxyMeshComponent;
-
+#endif
 public:
+#if WITH_EDITOR
 	// Refreshes the visual components to match the component state
 	ENGINE_API virtual void RefreshVisualRepresentation();
 
diff --git a/Engine/Source/Runtime/Engine/Classes/Components/SplineComponent.h b/Engine/Source/Runtime/Engine/Classes/Components/SplineComponent.h
--- a/Engine/Source/Runtime/Engine/Classes/Components/SplineComponent.h
+++ b/Engine/Source/Runtime/Engine/Classes/Components/SplineComponent.h
@@ -123,7 +123,7 @@ class ENGINE_API USplineComponent : public USceneComponent
 	// UObject interface
 	virtual void PostLoad() override;
 	virtual void PostEditImport() override;
-#if WITH_EDITORONLY_DATA
+#if WITH_EDITOR
 	virtual void PostEditChangeProperty(FPropertyChangedEvent& PropertyChangedEvent) override;
 	// End of UObject interface
 #endif
diff --git a/Engine/Source/Runtime/Engine/Classes/Engine/Blueprint.h b/Engine/Source/Runtime/Engine/Classes/Engine/Blueprint.h
--- a/Engine/Source/Runtime/Engine/Classes/Engine/Blueprint.h
+++ b/Engine/Source/Runtime/Engine/Classes/Engine/Blueprint.h
@@ -545,7 +545,7 @@ public:
 	 */
 	static bool GetBlueprintHierarchyFromClass(const UClass* InClass, TArray<UBlueprint*>& OutBlueprintParents);
 
-#if WITH_EDITORONLY_DATA
+#if WITH_EDITOR
 	template<class TFieldType>
 	static FName GetFieldNameFromClassByGuid(const UClass* InClass, const FGuid VarGuid)
 	{
@@ -632,7 +632,7 @@ public:
 
 };
 
-#if WITH_EDITORONLY_DATA
+#if WITH_EDITOR
 template<>
 inline FName UBlueprint::GetFieldNameFromClassByGuid<UFunction>(const UClass* InClass, const FGuid FunctionGuid)
 {
@@ -644,4 +644,4 @@ inline bool UBlueprint::GetGuidFromClassByFieldName<UFunction>(const UClass* InC
 {
 	return GetFunctionGuidFromClassByFieldName(InClass, FunctionName, FunctionGuid);
 }
-#endif // #if WITH_EDITORONLY_DATA
+#endif // #if WITH_EDITOR
diff --git a/Engine/Source/Runtime/Engine/Classes/Engine/StaticMesh.h b/Engine/Source/Runtime/Engine/Classes/Engine/StaticMesh.h
--- a/Engine/Source/Runtime/Engine/Classes/Engine/StaticMesh.h
+++ b/Engine/Source/Runtime/Engine/Classes/Engine/StaticMesh.h
@@ -175,7 +175,7 @@ struct FStaticMeshSourceModel
 #if WITH_EDITOR
 	/** Imported raw mesh data. Optional for all but the first LOD. */
 	class FRawMeshBulkData* RawMeshBulkData;
-#endif // #if WITH_EDITORONLY_DATA
+#endif // #if WITH_EDITOR
 
 	/** Settings applied when building the mesh. */
 	UPROPERTY(EditAnywhere, Category=BuildSettings)
@@ -198,8 +198,10 @@ struct FStaticMeshSourceModel
 	/** Destructor. */
 	ENGINE_API ~FStaticMeshSourceModel();
 
+#if WITH_EDITOR
 	/** Serializes bulk data. */
 	void SerializeBulkData(FArchive& Ar, UObject* Owner);
+#endif
 };
 
 /**
@@ -618,7 +620,7 @@ public:
 
 	void EnforceLightmapRestrictions();
 
-#if WITH_EDITORONLY_DATA
+#if WITH_EDITOR
 
 	/**
 	 * Returns true if LODs of this static mesh may share texture lightmaps.
@@ -655,5 +657,5 @@ private:
 	 * Caches derived renderable data.
 	 */
 	void CacheDerivedData();
-#endif // #if WITH_EDITORONLY_DATA
+#endif // #if WITH_EDITOR
 };
diff --git a/Engine/Source/Runtime/Engine/Classes/Engine/Texture.h b/Engine/Source/Runtime/Engine/Classes/Engine/Texture.h
--- a/Engine/Source/Runtime/Engine/Classes/Engine/Texture.h
+++ b/Engine/Source/Runtime/Engine/Classes/Engine/Texture.h
@@ -110,7 +110,7 @@ struct FTextureSource
 	/** Default constructor. */
 	ENGINE_API FTextureSource();
 
-#if WITH_EDITORONLY_DATA
+#if WITH_EDITOR
 	/**
 	 * Initialize the source data with the given size, number of mips, and format.
 	 * @param NewSizeX - Width of the texture source data.
@@ -207,7 +207,7 @@ private:
 	/** Which mips are locked, if any. */
 	uint32 LockedMips;
 
-#if WITH_EDITORONLY_DATA
+#if WITH_EDITOR
 	/** Return true if the source art is not png compressed but could be. */
 	bool CanPNGCompress() const;
 
@@ -307,7 +307,7 @@ struct FTexturePlatformData
 	/** Serialization for cooked builds. */
 	void SerializeCooked(FArchive& Ar, class UTexture* Owner);
 
-#if WITH_EDITORONLY_DATA
+#if WITH_EDITOR
 	void Cache(
 		class UTexture& InTexture,
 		const struct FTextureBuildSettings& InSettings,
@@ -584,7 +584,7 @@ public:
 	 */
 	ENGINE_API void SerializeCookedPlatformData(class FArchive& Ar);
 
-#if WITH_EDITORONLY_DATA
+#if WITH_EDITOR
 	/**
 	 * Caches platform data for the texture.
 	 */
@@ -679,11 +679,11 @@ public:
 	 */
 	bool HasHDRSource() const
 	{
-#if WITH_EDITORONLY_DATA
+#if WITH_EDITOR
 		return ((Source.GetFormat() == TSF_BGRE8) || (Source.GetFormat() == TSF_RGBA16F));
 #else
 		return false;
-#endif // WITH_EDITORONLY_DATA
+#endif // WITH_EDITOR
 	}
 
 
diff --git a/Engine/Source/Runtime/Engine/Classes/Engine/Texture2D.h b/Engine/Source/Runtime/Engine/Classes/Engine/Texture2D.h
--- a/Engine/Source/Runtime/Engine/Classes/Engine/Texture2D.h
+++ b/Engine/Source/Runtime/Engine/Classes/Engine/Texture2D.h
@@ -62,11 +62,11 @@ public:
 	 */
 	FORCEINLINE FIntPoint GetImportedSize() const
 	{
-#if WITH_EDITORONLY_DATA
+#if WITH_EDITOR
 		return FIntPoint(Source.GetSizeX(),Source.GetSizeY());
-#else // #if WITH_EDITORONLY_DATA
+#else // #if WITH_EDITOR
 		return ImportedSize;
-#endif // #if WITH_EDITORONLY_DATA
+#endif // #if WITH_EDITOR
 	}
 
 private:
@@ -431,7 +431,7 @@ public:
 	/** creates and initializes a new Texture2D with the requested settings */
 	ENGINE_API static class UTexture2D* CreateTransient(int32 InSizeX, int32 InSizeY, EPixelFormat InFormat = PF_B8G8R8A8);
 
-#if WITH_EDITORONLY_DATA
+#if WITH_EDITOR
 	/**
 	 * Legacy serialization.
 	 */
diff --git a/Engine/Source/Runtime/Engine/Classes/Engine/TextureCube.h b/Engine/Source/Runtime/Engine/Classes/Engine/TextureCube.h
--- a/Engine/Source/Runtime/Engine/Classes/Engine/TextureCube.h
+++ b/Engine/Source/Runtime/Engine/Classes/Engine/TextureCube.h
@@ -82,7 +82,7 @@ public:
 	 */
 	virtual uint32 CalcTextureMemorySizeEnum( ETextureMipCount Enum ) const override;
 
-#if WITH_EDITORONLY_DATA
+#if WITH_EDITOR
 	/**
 	* Return maximum dimension for this texture type.
 	*/
diff --git a/Engine/Source/Runtime/Engine/Classes/Materials/MaterialFunction.h b/Engine/Source/Runtime/Engine/Classes/Materials/MaterialFunction.h
--- a/Engine/Source/Runtime/Engine/Classes/Materials/MaterialFunction.h
+++ b/Engine/Source/Runtime/Engine/Classes/Materials/MaterialFunction.h
@@ -81,7 +81,7 @@ public:
 	/** Appends textures referenced by the expressions in this function. */
 	ENGINE_API void AppendReferencedTextures(TArray<UTexture*>& InOutTextures) const;
 
-#if WITH_EDITORONLY_DATA
+#if WITH_EDITOR
 	ENGINE_API UMaterial* GetPreviewMaterial();
 
 	void UpdateInputOutputTypes();
@@ -90,7 +90,9 @@ public:
 	 * Checks whether a Material Function is arranged in the old style, with inputs flowing from right to left
 	 */
 	bool HasFlippedCoordinates() const;
+#endif
 
+#if WITH_EDITORONLY_DATA
 	UPROPERTY(AssetRegistrySearchable)
 	uint32 CombinedInputTypes;
 
diff --git a/Engine/Source/Runtime/Engine/Classes/Sound/SoundCue.h b/Engine/Source/Runtime/Engine/Classes/Sound/SoundCue.h
--- a/Engine/Source/Runtime/Engine/Classes/Sound/SoundCue.h
+++ b/Engine/Source/Runtime/Engine/Classes/Sound/SoundCue.h
@@ -100,7 +100,7 @@ public:
 	T* ConstructSoundNode(TSubclassOf<USoundNode> SoundNodeClass = T::StaticClass(), bool bSelectNewNode = true)
 	{
 		T* SoundNode = ConstructObject<T>(SoundNodeClass, this);
-#if WITH_EDITORONLY_DATA
+#if WITH_EDITOR
 		AllNodes.Add(SoundNode);
 		SetupSoundNode(SoundNode, bSelectNewNode);
 #endif // WITH_EDITORONLY_DATA
@@ -137,7 +137,7 @@ public:
 	 */
 	ENGINE_API void RecursiveFindAttenuation( USoundNode* Node, TArray<class USoundNodeAttenuation*> &OutNodes );
 
-#if WITH_EDITORONLY_DATA
+#if WITH_EDITOR
 	/** Create the basic sound graph */
 	ENGINE_API void CreateGraph();
 
diff --git a/Engine/Source/Runtime/Engine/Classes/Sound/SoundNode.h b/Engine/Source/Runtime/Engine/Classes/Sound/SoundNode.h
--- a/Engine/Source/Runtime/Engine/Classes/Sound/SoundNode.h
+++ b/Engine/Source/Runtime/Engine/Classes/Sound/SoundNode.h
@@ -65,11 +65,10 @@ public:
 	// Begin UObject Interface
 #if WITH_EDITOR
 	virtual void PostEditChangeProperty(FPropertyChangedEvent& PropertyChangedEvent) override;
-#endif
-#if WITH_EDITORONLY_DATA
+
 	virtual void Serialize(FArchive& Ar) override;
 	static void AddReferencedObjects(UObject* InThis, FReferenceCollector& Collector);
-#endif //WITH_EDITORONLY_DATA
+#endif //WITH_EDITOR
 	// End UObject Interface
 
 	//
diff --git a/Engine/Source/Runtime/Engine/Private/AmbientSound.cpp b/Engine/Source/Runtime/Engine/Private/AmbientSound.cpp
--- a/Engine/Source/Runtime/Engine/Private/AmbientSound.cpp
+++ b/Engine/Source/Runtime/Engine/Private/AmbientSound.cpp
@@ -96,7 +96,7 @@ void AAmbientSound::PostRegisterAllComponents()
 FString AAmbientSound::GetInternalSoundCueName()
 {
 	FString CueName;
-#if WITH_EDITORONLY_DATA
+#if WITH_EDITOR
 	CueName = GetActorLabel();
 #endif
 	if (CueName.Len() == 0)
@@ -242,7 +242,7 @@ void AAmbientSound::MigrateSoundNodeInstance()
 #endif //WITH_EDITOR
 	}
 
-#if WITH_EDITORONLY_DATA
+#if WITH_EDITOR
 	SoundCue->LinkGraphNodesFromSoundNodes();
 #endif //WITH_EDITORONLY_DATA
 }
diff --git a/Engine/Source/Runtime/Engine/Private/Animation/AnimInstance.cpp b/Engine/Source/Runtime/Engine/Private/Animation/AnimInstance.cpp
--- a/Engine/Source/Runtime/Engine/Private/Animation/AnimInstance.cpp
+++ b/Engine/Source/Runtime/Engine/Private/Animation/AnimInstance.cpp
@@ -321,7 +321,7 @@ bool UAnimInstance::UpdateSnapshotAndSkipRemainingUpdate()
 
 void UAnimInstance::UpdateAnimation(float DeltaSeconds)
 {
-#if WITH_EDITORONLY_DATA
+#if WITH_EDITOR
 	if (GIsEditor)
 	{
 		// Reset the anim graph visualization
diff --git a/Engine/Source/Runtime/Engine/Private/Animation/AnimNotifyState_TimedParticleEffect.cpp b/Engine/Source/Runtime/Engine/Private/Animation/AnimNotifyState_TimedParticleEffect.cpp
--- a/Engine/Source/Runtime/Engine/Private/Animation/AnimNotifyState_TimedParticleEffect.cpp
+++ b/Engine/Source/Runtime/Engine/Private/Animation/AnimNotifyState_TimedParticleEffect.cpp
@@ -102,7 +102,7 @@ FString UAnimNotifyState_TimedParticleEffect::GetNotifyName_Implementation() con
 	return UAnimNotifyState::GetNotifyName_Implementation();
 }
 
-#if WITH_EDITORONLY_DATA
+#if WITH_EDITOR
 void UAnimNotifyState_TimedParticleEffect::PreEditChange(UProperty* PropertyAboutToChange)
 {
 	if(PropertyAboutToChange)
diff --git a/Engine/Source/Runtime/Engine/Private/Animation/AnimNotifyState_Trail.cpp b/Engine/Source/Runtime/Engine/Private/Animation/AnimNotifyState_Trail.cpp
--- a/Engine/Source/Runtime/Engine/Private/Animation/AnimNotifyState_Trail.cpp
+++ b/Engine/Source/Runtime/Engine/Private/Animation/AnimNotifyState_Trail.cpp
@@ -189,7 +189,7 @@ void UAnimNotifyState_Trail::NotifyEnd(class USkeletalMeshComponent * MeshComp,
 
 bool UAnimNotifyState_Trail::ValidateInput(class USkeletalMeshComponent * MeshComp, bool bReportErrors/* =false */)
 {
-#if WITH_EDITORONLY_DATA
+#if WITH_EDITOR
 	bool bError = false;
 	FText FirstSocketEqualsNoneErrorText = FText::Format( LOCTEXT("NoneFirstSocket", "{0}: Must set First Socket Name."), FText::FromString(GetName()));
 	FText SecondSocketEqualsNoneErrorText = FText::Format( LOCTEXT("NoneSecondSocket", "{0}: Must set Second Socket Name."), FText::FromString(GetName()));
@@ -241,7 +241,7 @@ bool UAnimNotifyState_Trail::ValidateInput(class USkeletalMeshComponent * MeshCo
 	return bError;
 #else
 	return false;
-#endif // WITH_EDITORONLY_DATA
+#endif // WITH_EDITOR
 }
 
 #undef  LOCTEXT_NAMESPACE
\ No newline at end of file
diff --git a/Engine/Source/Runtime/Engine/Private/Animation/AnimSequence.cpp b/Engine/Source/Runtime/Engine/Private/Animation/AnimSequence.cpp
--- a/Engine/Source/Runtime/Engine/Private/Animation/AnimSequence.cpp
+++ b/Engine/Source/Runtime/Engine/Private/Animation/AnimSequence.cpp
@@ -1384,7 +1384,7 @@ bool UAnimSequence::CopyAnimSequenceProperties(UAnimSequence* SourceAnimSeq, UAn
 
 bool UAnimSequence::CopyNotifies(UAnimSequence* SourceAnimSeq, UAnimSequence* DestAnimSeq)
 {
-#if WITH_EDITORONLY_DATA
+#if WITH_EDITOR
 	// Abort if source == destination.
 	if( SourceAnimSeq == DestAnimSeq )
 	{
@@ -1493,7 +1493,7 @@ bool UAnimSequence::CopyNotifies(UAnimSequence* SourceAnimSeq, UAnimSequence* De
 		FMessageDialog::Open( EAppMsgType::Ok, FText::Format(
 			NSLOCTEXT("UnrealEd", "SomeNotifiesWereNotCopiedF", "Because the destination sequence was shorter, {0} notifies were not copied."), FText::AsNumber(SourceAnimSeq->Notifies.Num() - NewNotifyIndices.Num())) );
 	}
-#endif // WITH_EDITORONLY_DATA
+#endif // WITH_EDITOR
 
 	return true;
 }
diff --git a/Engine/Source/Runtime/Engine/Private/Animation/AnimSequenceBase.cpp b/Engine/Source/Runtime/Engine/Private/Animation/AnimSequenceBase.cpp
--- a/Engine/Source/Runtime/Engine/Private/Animation/AnimSequenceBase.cpp
+++ b/Engine/Source/Runtime/Engine/Private/Animation/AnimSequenceBase.cpp
@@ -240,7 +240,7 @@ void UAnimSequenceBase::PostLoad()
 	// Ensure notifies are sorted.
 	SortNotifies();
 
-#if WITH_EDITORONLY_DATA
+#if WITH_EDITOR
 	InitializeNotifyTrack();
 	UpdateAnimNotifyTrackCache();
 
diff --git a/Engine/Source/Runtime/Engine/Private/Animation/Skeleton.cpp b/Engine/Source/Runtime/Engine/Private/Animation/Skeleton.cpp
--- a/Engine/Source/Runtime/Engine/Private/Animation/Skeleton.cpp
+++ b/Engine/Source/Runtime/Engine/Private/Animation/Skeleton.cpp
@@ -685,7 +685,7 @@ void USkeleton::CollectAnimationNotifies()
 	//AnimationNotifies.Empty();
 	TArray<FAssetData> AssetList;
 	AssetRegistryModule.Get().GetAssetsByClass(UAnimSequenceBase::StaticClass()->GetFName(), AssetList, true);
-
+#if WITH_EDITOR
 	// do not clear AnimationNotifies. We can't remove old ones yet. 
 	FString CurrentSkeletonName = FAssetData(this).GetExportTextName();
 	for (auto Iter = AssetList.CreateConstIterator(); Iter; ++Iter)
@@ -706,6 +706,7 @@ void USkeleton::CollectAnimationNotifies()
 			}
 		}
 	}
+#endif
 }
 
 void USkeleton::AddNewAnimationNotify(FName NewAnimNotifyName)
@@ -833,6 +834,8 @@ void USkeleton::AddBoneToLOD(int32 LODIndex, int32 BoneIndex)
 	}
 }
 
+#if WITH_EDITOR
+
 void USkeleton::RemoveBonesFromSkeleton( const TArray<FName>& BonesToRemove, bool bRemoveChildBones )
 {
 	TArray<int32> BonesRemoved = ReferenceSkeleton.RemoveBonesByName(BonesToRemove);
@@ -875,6 +878,8 @@ void USkeleton::UnregisterOnSkeletonHierarchyChanged(void * Unregister)
 	OnSkeletonHierarchyChanged.RemoveAll(Unregister);
 }
 
+#endif
+
 void USkeleton::AddSlotNodeName(FName SlotNodeName)
 {
 	SlotNodeNames.AddUnique(SlotNodeName);
diff --git a/Engine/Source/Runtime/Engine/Private/ArrowComponent.cpp b/Engine/Source/Runtime/Engine/Private/ArrowComponent.cpp
--- a/Engine/Source/Runtime/Engine/Private/ArrowComponent.cpp
+++ b/Engine/Source/Runtime/Engine/Private/ArrowComponent.cpp
@@ -304,7 +304,7 @@ public:
 		FPrimitiveViewRelevance Result;
 		Result.bDrawRelevance = IsShown(View) && (View->Family->EngineShowFlags.BillboardSprites);
 		Result.bDynamicRelevance = true;
-#if WITH_EDITORONLY_DATA
+#if WITH_EDITOR
 		if (bTreatAsASprite)
 		{
 			if ( GIsEditor && SpriteCategoryIndex != INDEX_NONE && SpriteCategoryIndex < View->SpriteCategoryVisibility.Num() && !View->SpriteCategoryVisibility[ SpriteCategoryIndex ] )
diff --git a/Engine/Source/Runtime/Engine/Private/Atmosphere/Atmosphere.cpp b/Engine/Source/Runtime/Engine/Private/Atmosphere/Atmosphere.cpp
--- a/Engine/Source/Runtime/Engine/Private/Atmosphere/Atmosphere.cpp
+++ b/Engine/Source/Runtime/Engine/Private/Atmosphere/Atmosphere.cpp
@@ -111,7 +111,7 @@ void UAtmosphericFogComponent::PostLoad()
 
 	if ( !IsTemplate() )
 	{
-#if WITH_EDITORONLY_DATA
+#if WITH_EDITOR
 		if (TransmittanceTexture_DEPRECATED)
 		{
 			// Copy data from Previous data
@@ -173,7 +173,7 @@ void UAtmosphericFogComponent::PostLoad()
 	}
 }
 
-#if WITH_EDITORONLY_DATA
+#if WITH_EDITOR
 #include "UnrealEd.h"
 
 class FAtmospherePrecomputeDataHandler : public FTickableEditorObject
@@ -238,7 +238,7 @@ void UAtmosphericFogComponent::InitResource()
 			BeginInitResource( InscatterResource );
 		}
 	}
-#if WITH_EDITORONLY_DATA
+#if WITH_EDITOR
 	else if (!PrecomputeDataHandler)
 	{
 		PrecomputeDataHandler = new FAtmospherePrecomputeDataHandler(this);
@@ -293,7 +293,7 @@ void UAtmosphericFogComponent::ReleaseResource()
 		InscatterResource = NULL;
 	}
 
-#if WITH_EDITORONLY_DATA
+#if WITH_EDITOR
 	if (PrecomputeDataHandler)
 	{
 		delete PrecomputeDataHandler;
@@ -527,7 +527,7 @@ void UAtmosphericFogComponent::PostInterpChange(UProperty* PropertyThatChanged)
 
 void UAtmosphericFogComponent::StartPrecompute()
 {
-#if WITH_EDITORONLY_DATA
+#if WITH_EDITOR
 	if (GIsEditor)
 	{
 		PrecomputeCounter.Reset();
diff --git a/Engine/Source/Runtime/Engine/Private/BillboardComponent.cpp b/Engine/Source/Runtime/Engine/Private/BillboardComponent.cpp
--- a/Engine/Source/Runtime/Engine/Private/BillboardComponent.cpp
+++ b/Engine/Source/Runtime/Engine/Private/BillboardComponent.cpp
@@ -288,11 +288,11 @@ private:
 	FLinearColor Color;
 	const uint32 bIsScreenSizeScaled : 1;
 	uint32 bIsActorLocked : 1;
-#if WITH_EDITOR
+#if WITH_EDITORONLY_DATA
 	int32 SpriteCategoryIndex;
 	bool bUseInEditorScaling;
 	float EditorScale;
-#endif // #if WITH_EDITOR
+#endif // #if WITH_EDITORONLY_DATA
 };
 
 UBillboardComponent::UBillboardComponent(const class FPostConstructInitializeProperties& PCIP)
diff --git a/Engine/Source/Runtime/Engine/Private/Blueprint.cpp b/Engine/Source/Runtime/Engine/Private/Blueprint.cpp
--- a/Engine/Source/Runtime/Engine/Private/Blueprint.cpp
+++ b/Engine/Source/Runtime/Engine/Private/Blueprint.cpp
@@ -1188,7 +1188,7 @@ bool UBlueprint::Modify(bool bAlwaysMarkDirty)
 
 #endif
 
-#if WITH_EDITORONLY_DATA
+#if WITH_EDITOR
 
 FName UBlueprint::GetFunctionNameFromClassByGuid(const UClass* InClass, const FGuid FunctionGuid)
 {
@@ -1200,4 +1200,4 @@ bool UBlueprint::GetFunctionGuidFromClassByFieldName(const UClass* InClass, cons
 	return FBlueprintEditorUtils::GetFunctionGuidFromClassByFieldName(InClass, FunctionName, FunctionGuid);
 }
 
-#endif //WITH_EDITORONLY_DATA
\ No newline at end of file
+#endif //WITH_EDITOR
\ No newline at end of file
diff --git a/Engine/Source/Runtime/Engine/Private/BrushComponent.cpp b/Engine/Source/Runtime/Engine/Private/BrushComponent.cpp
--- a/Engine/Source/Runtime/Engine/Private/BrushComponent.cpp
+++ b/Engine/Source/Runtime/Engine/Private/BrushComponent.cpp
@@ -13,7 +13,7 @@
 
 DEFINE_LOG_CATEGORY_STATIC(LogBrushComponent, Log, All);
 
-#if WITH_EDITOR
+#if WITH_EDITORONLY_DATA
 struct FModelWireVertex
 {
 	FVector Position;
@@ -30,11 +30,13 @@ public:
 	FModelWireVertexBuffer(UModel* InModel):
 		NumVertices(0)
 	{
+#if WITH_EDITOR
 		Polys.Append(InModel->Polys->Element);
 		for(int32 PolyIndex = 0;PolyIndex < InModel->Polys->Element.Num();PolyIndex++)
 		{
 			NumVertices += InModel->Polys->Element[PolyIndex].Vertices.Num();
 		}
+#endif
 	}
 
 	// FRenderResource interface.
@@ -81,11 +83,13 @@ public:
 	FModelWireIndexBuffer(UModel* InModel):
 		NumEdges(0)
 	{
+#if WITH_EDITOR
 		Polys.Append(InModel->Polys->Element);
 		for(int32 PolyIndex = 0;PolyIndex < InModel->Polys->Element.Num();PolyIndex++)
 		{
 			NumEdges += InModel->Polys->Element[PolyIndex].Vertices.Num();
 		}
+#endif
 	}
 
 	// FRenderResource interface.
@@ -119,7 +123,7 @@ private:
 	TArray<FPoly> Polys;
 	uint32 NumEdges;
 };
-#endif // WITH_EDITOR
+#endif // WITH_EDITORONLY_DATA
 
 class FBrushSceneProxy : public FPrimitiveSceneProxy
 {
diff --git a/Engine/Source/Runtime/Engine/Private/Camera/CameraComponent.cpp b/Engine/Source/Runtime/Engine/Private/Camera/CameraComponent.cpp
--- a/Engine/Source/Runtime/Engine/Private/Camera/CameraComponent.cpp
+++ b/Engine/Source/Runtime/Engine/Private/Camera/CameraComponent.cpp
@@ -32,7 +32,7 @@ UCameraComponent::UCameraComponent(const class FPostConstructInitializePropertie
 
 void UCameraComponent::OnRegister()
 {
-#if WITH_EDITORONLY_DATA
+#if WITH_EDITOR
 	if (ProxyMeshComponent == NULL)
 	{
 		ProxyMeshComponent = ConstructObject<UStaticMeshComponent>(UStaticMeshComponent::StaticClass(), GetOuter(), NAME_None, RF_Transactional);
@@ -86,7 +86,7 @@ void UCameraComponent::OnUnregister()
 #endif
 }
 
-#if WITH_EDITORONLY_DATA
+#if WITH_EDITOR
 void UCameraComponent::RefreshVisualRepresentation()
 {
 	if (DrawFrustum != NULL)
diff --git a/Engine/Source/Runtime/Engine/Private/DestructibleMesh.cpp b/Engine/Source/Runtime/Engine/Private/DestructibleMesh.cpp
--- a/Engine/Source/Runtime/Engine/Private/DestructibleMesh.cpp
+++ b/Engine/Source/Runtime/Engine/Private/DestructibleMesh.cpp
@@ -496,7 +496,7 @@ bool CreateSubmeshFromSMSection(FStaticMeshLODResources& RenderMesh, int32 Subme
 
 bool UDestructibleMesh::BuildFractureSettingsFromStaticMesh(UStaticMesh* StaticMesh)
 {
-#if WITH_APEX && WITH_EDITORONLY_DATA
+#if WITH_APEX && WITH_EDITOR
 	// Make sure we have the authoring data container
 	CreateFractureSettings();
 
@@ -576,14 +576,14 @@ bool UDestructibleMesh::BuildFractureSettingsFromStaticMesh(UStaticMesh* StaticM
 
 	FractureSettings->SetRootMesh(Triangles, Materials, Submeshes, MeshPartitions, true);
 	return true;
-#else // WITH_APEX && WITH_EDITORONLY_DATA
+#else // WITH_APEX && WITH_EDITOR
 	return false;
-#endif // WITH_APEX && WITH_EDITORONLY_DATA
+#endif // WITH_APEX && WITH_EDITOR
 }
 
 ENGINE_API bool UDestructibleMesh::BuildFromStaticMesh( UStaticMesh& StaticMesh )
 {
-#if WITH_EDITORONLY_DATA
+#if WITH_EDITOR
 	PreEditChange(NULL);
 
 	// Import the StaticMesh
@@ -604,13 +604,13 @@ ENGINE_API bool UDestructibleMesh::BuildFromStaticMesh( UStaticMesh& StaticMesh
 
 	PostEditChange();
 	MarkPackageDirty();
-#endif // WITH_EDITORONLY_DATA
+#endif // WITH_EDITOR
 	return true;
 }
 
 ENGINE_API bool UDestructibleMesh::SetupChunksFromStaticMeshes( const TArray<UStaticMesh*>& ChunkMeshes )
 {
-#if WITH_EDITORONLY_DATA
+#if WITH_EDITOR
 	if (SourceStaticMesh == NULL)
 	{
 		UE_LOG(LogDestructible, Warning, TEXT("Unable to import FBX as level 1 chunks if the DM was not created from a static mesh."));
@@ -637,6 +637,6 @@ ENGINE_API bool UDestructibleMesh::SetupChunksFromStaticMeshes( const TArray<USt
 
 	PostEditChange();
 	MarkPackageDirty();
-#endif // WITH_EDITORONLY_DATA
+#endif // WITH_EDITOR
 	return true;
 }
diff --git a/Engine/Source/Runtime/Engine/Private/DocumentationActor.cpp b/Engine/Source/Runtime/Engine/Private/DocumentationActor.cpp
--- a/Engine/Source/Runtime/Engine/Private/DocumentationActor.cpp
+++ b/Engine/Source/Runtime/Engine/Private/DocumentationActor.cpp
@@ -2,7 +2,7 @@
 
 #include "EnginePrivate.h"
 #include "Engine/DocumentationActor.h"
-#if WITH_EDITORONLY_DATA
+#if WITH_EDITOR
 #include "Components/MaterialBillboardComponent.h"
 #include "IDocumentation.h"
 #endif
@@ -28,9 +28,9 @@ bool ADocumentationActor::OpenDocumentLink() const
 {
 	bool bOpened = false;
 
-#if WITH_EDITORONLY_DATA
+#if WITH_EDITOR
 	bOpened = IDocumentation::Get()->Open(DocumentLink);
-#endif //WITH_EDITORONLY_DATA
+#endif //WITH_EDITOR
 
 	return bOpened;
 }
diff --git a/Engine/Source/Runtime/Engine/Private/ImageUtils.cpp b/Engine/Source/Runtime/Engine/Private/ImageUtils.cpp
--- a/Engine/Source/Runtime/Engine/Private/ImageUtils.cpp
+++ b/Engine/Source/Runtime/Engine/Private/ImageUtils.cpp
@@ -123,7 +123,7 @@ void FImageUtils::ImageResize(int32 SrcWidth, int32 SrcHeight, const TArray<FCol
  */
 UTexture2D* FImageUtils::CreateTexture2D(int32 SrcWidth, int32 SrcHeight, const TArray<FColor> &SrcData, UObject* Outer, const FString& Name, const EObjectFlags &Flags, const FCreateTexture2DParameters& InParams)
 {
-#if WITH_EDITORONLY_DATA
+#if WITH_EDITOR
 	UTexture2D* Tex2D;
 
 	Tex2D = CastChecked<UTexture2D>( StaticConstructObject(UTexture2D::StaticClass(), Outer, FName(*Name), Flags) );
diff --git a/Engine/Source/Runtime/Engine/Private/Interpolation.cpp b/Engine/Source/Runtime/Engine/Private/Interpolation.cpp
--- a/Engine/Source/Runtime/Engine/Private/Interpolation.cpp
+++ b/Engine/Source/Runtime/Engine/Private/Interpolation.cpp
@@ -8027,7 +8027,7 @@ void UInterpTrackSound::PreviewUpdateTrack(float NewPosition, UInterpTrackInst*
 	bool bJump = !( MatineeActor->bIsPlaying );
 	UpdateTrack(NewPosition, TrInst, bJump);
 
-#if WITH_EDITORONLY_DATA
+#if WITH_EDITOR
 	bool bTimeChangedDrastically = !FMath::IsNearlyEqual(NewPosition, MatineeActor->InterpPosition);
 	if ( bTimeChangedDrastically && MatineeActor->bIsScrubbing )
 	{
diff --git a/Engine/Source/Runtime/Engine/Private/Landscape/LandscapeGizmoActor.cpp b/Engine/Source/Runtime/Engine/Private/Landscape/LandscapeGizmoActor.cpp
--- a/Engine/Source/Runtime/Engine/Private/Landscape/LandscapeGizmoActor.cpp
+++ b/Engine/Source/Runtime/Engine/Private/Landscape/LandscapeGizmoActor.cpp
@@ -129,7 +129,7 @@ public:
 		HeightmapRenderProxy(NULL),
 		GizmoRenderProxy(NULL)
 	{
-#if WITH_EDITORONLY_DATA	
+#if WITH_EDITOR	
 		ALandscapeGizmoActiveActor* Gizmo = Cast<ALandscapeGizmoActiveActor>(InComponent->GetOwner());
 		if (Gizmo && Gizmo->GizmoMeshMaterial && Gizmo->GizmoDataMaterial && Gizmo->GetRootComponent())
 		{
diff --git a/Engine/Source/Runtime/Engine/Private/LightMap.cpp b/Engine/Source/Runtime/Engine/Private/LightMap.cpp
--- a/Engine/Source/Runtime/Engine/Private/LightMap.cpp
+++ b/Engine/Source/Runtime/Engine/Private/LightMap.cpp
@@ -115,7 +115,7 @@ FString ULightMapTexture2D::GetDesc()
 	return FString::Printf( TEXT("Lightmap: %dx%d [%s]"), GetSizeX(), GetSizeY(), GPixelFormats[GetPixelFormat()].Name );
 }
 
-#if WITH_EDITORONLY_DATA
+#if WITH_EDITOR
 static void DumpLightmapSizeOnDisk()
 {
 	UE_LOG(LogLightMap,Log,TEXT("Lightmap size on disk"));
@@ -137,7 +137,7 @@ static FAutoConsoleCommand CmdDumpLightmapSizeOnDisk(
 	TEXT("Dumps the size of all loaded lightmaps on disk (source and platform data)"),
 	FConsoleCommandDelegate::CreateStatic(DumpLightmapSizeOnDisk)
 	);
-#endif // #if WITH_EDITORONLY_DATA
+#endif // #if WITH_EDITOR
 
 /** Lightmap resolution scaling factors for debugging.  The defaults are to use the original resolution unchanged. */
 float TextureMappingDownsampleFactor0 = 1.0f;
diff --git a/Engine/Source/Runtime/Engine/Private/Materials/Material.cpp b/Engine/Source/Runtime/Engine/Private/Materials/Material.cpp
--- a/Engine/Source/Runtime/Engine/Private/Materials/Material.cpp
+++ b/Engine/Source/Runtime/Engine/Private/Materials/Material.cpp
@@ -2511,9 +2511,9 @@ void UMaterial::UpdateExpressionDynamicParameterNames(const UMaterialExpression*
 			UMaterialExpressionDynamicParameter* CheckParam = Cast<UMaterialExpressionDynamicParameter>(Expressions[ExpIndex]);
 			if (CheckParam->CopyDynamicParameterNames(DynParam))
 			{
-#if WITH_EDITORONLY_DATA
+#if WITH_EDITOR
 				CheckParam->GraphNode->ReconstructNode();
-#endif // WITH_EDITORONLY_DATA
+#endif // WITH_EDITOR
 			}
 		}
 	}
diff --git a/Engine/Source/Runtime/Engine/Private/Materials/MaterialExpressions.cpp b/Engine/Source/Runtime/Engine/Private/Materials/MaterialExpressions.cpp
--- a/Engine/Source/Runtime/Engine/Private/Materials/MaterialExpressions.cpp
+++ b/Engine/Source/Runtime/Engine/Private/Materials/MaterialExpressions.cpp
@@ -6144,7 +6144,7 @@ UMaterialFunction::UMaterialFunction(const class FPostConstructInitializePropert
 
 }
 
-#if WITH_EDITORONLY_DATA
+#if WITH_EDITOR
 UMaterial* UMaterialFunction::GetPreviewMaterial()
 {
 	if( NULL == PreviewMaterial )
@@ -6249,13 +6249,12 @@ void UMaterialFunction::PostLoad()
 		}
 	}
 
-#if WITH_EDITORONLY_DATA
+#if WITH_EDITOR
 	if (CombinedOutputTypes == 0)
 	{
 		UpdateInputOutputTypes();
 	}
-#endif
-#if WITH_EDITOR
+
 	if (GIsEditor)
 	{
 		// Clean up any removed material expression classes	
@@ -6617,7 +6616,7 @@ void UMaterialFunction::AppendReferencedTextures(TArray<UTexture*>& InOutTexture
 	}
 }
 
-#if WITH_EDITORONLY_DATA
+#if WITH_EDITOR
 bool UMaterialFunction::HasFlippedCoordinates() const
 {
 	uint32 ReversedInputCount = 0;
diff --git a/Engine/Source/Runtime/Engine/Private/Particles/ParticleComponents.cpp b/Engine/Source/Runtime/Engine/Private/Particles/ParticleComponents.cpp
--- a/Engine/Source/Runtime/Engine/Private/Particles/ParticleComponents.cpp
+++ b/Engine/Source/Runtime/Engine/Private/Particles/ParticleComponents.cpp
@@ -799,7 +799,7 @@ int32	UParticleLODLevel::CalculateMaxActiveParticleCount()
 
 void UParticleLODLevel::ConvertToSpawnModule()
 {
-#if WITH_EDITORONLY_DATA
+#if WITH_EDITOR
 	// Move the required module SpawnRate and Burst information to a new SpawnModule.
 	if (SpawnModule)
 	{
@@ -833,7 +833,7 @@ void UParticleLODLevel::ConvertToSpawnModule()
 	}
 
 	MarkPackageDirty();
-#endif	//#if WITH_EDITORONLY_DATA
+#endif	//#if WITH_EDITOR
 }
 
 
diff --git a/Engine/Source/Runtime/Engine/Private/SCS_Node.cpp b/Engine/Source/Runtime/Engine/Private/SCS_Node.cpp
--- a/Engine/Source/Runtime/Engine/Private/SCS_Node.cpp
+++ b/Engine/Source/Runtime/Engine/Private/SCS_Node.cpp
@@ -3,7 +3,7 @@
 #include "EnginePrivate.h"
 #include "BlueprintUtilities.h"
 #include "LatentActions.h"
-#if WITH_EDITORONLY_DATA
+#if WITH_EDITOR
 #include "Kismet2/BlueprintEditorUtils.h"
 #endif
 #include "DeferRegisterComponents.h"
diff --git a/Engine/Source/Runtime/Engine/Private/SceneView.cpp b/Engine/Source/Runtime/Engine/Private/SceneView.cpp
--- a/Engine/Source/Runtime/Engine/Private/SceneView.cpp
+++ b/Engine/Source/Runtime/Engine/Private/SceneView.cpp
@@ -1173,7 +1173,7 @@ FSceneViewFamily::FSceneViewFamily( const ConstructionValues& CVS )
 	}
 #endif
 
-#if !WITH_EDITORONLY_DATA
+#if !WITH_EDITOR
 	// Console shader compilers don't set instruction count, 
 	// Also various console-specific rendering paths haven't been tested with shader complexity
 	check(!EngineShowFlags.ShaderComplexity);
diff --git a/Engine/Source/Runtime/Engine/Private/SoundClass.cpp b/Engine/Source/Runtime/Engine/Private/SoundClass.cpp
--- a/Engine/Source/Runtime/Engine/Private/SoundClass.cpp
+++ b/Engine/Source/Runtime/Engine/Private/SoundClass.cpp
@@ -3,7 +3,7 @@
 
 #include "EnginePrivate.h"
 #include "SoundDefinitions.h"
-#if WITH_EDITORONLY_DATA
+#if WITH_EDITOR
 #include "UnrealEd.h"
 #endif
 #if WITH_EDITOR
diff --git a/Engine/Source/Runtime/Engine/Private/SoundCue.cpp b/Engine/Source/Runtime/Engine/Private/SoundCue.cpp
--- a/Engine/Source/Runtime/Engine/Private/SoundCue.cpp
+++ b/Engine/Source/Runtime/Engine/Private/SoundCue.cpp
@@ -6,7 +6,7 @@
 #include "Sound/SoundNodeWavePlayer.h"
 #include "Sound/SoundNodeAttenuation.h"
 #include "Sound/SoundWave.h"
-#if WITH_EDITORONLY_DATA
+#if WITH_EDITOR
 #include "Kismet2/BlueprintEditorUtils.h"
 #include "UnrealEd.h"
 #endif
@@ -354,7 +354,7 @@ const FAttenuationSettings* USoundCue::GetAttenuationSettingsToApply() const
 	return Super::GetAttenuationSettingsToApply();
 }
 
-#if WITH_EDITORONLY_DATA
+#if WITH_EDITOR
 void USoundCue::CreateGraph()
 {
 	check(SoundCueGraph == NULL);
diff --git a/Engine/Source/Runtime/Engine/Private/SoundNode.cpp b/Engine/Source/Runtime/Engine/Private/SoundNode.cpp
--- a/Engine/Source/Runtime/Engine/Private/SoundNode.cpp
+++ b/Engine/Source/Runtime/Engine/Private/SoundNode.cpp
@@ -5,7 +5,7 @@
 #include "Sound/SoundNode.h"
 #include "Sound/SoundCue.h"
 
-#if WITH_EDITORONLY_DATA
+#if WITH_EDITOR
 #include "UnrealEd.h"
 #endif
 
@@ -17,7 +17,7 @@ USoundNode::USoundNode(const class FPostConstructInitializeProperties& PCIP)
 {
 }
 
-#if WITH_EDITORONLY_DATA
+#if WITH_EDITOR
 void USoundNode::Serialize(FArchive& Ar)
 {
 	Super::Serialize(Ar);
@@ -36,7 +36,7 @@ void USoundNode::AddReferencedObjects(UObject* InThis, FReferenceCollector& Coll
 
 	Super::AddReferencedObjects(InThis, Collector);
 }
-#endif //WITH_EDITORONLY_DATA
+#endif //WITH_EDITOR
 
 UPTRINT USoundNode::GetNodeWaveInstanceHash(const UPTRINT ParentWaveInstanceHash, const USoundNode* ChildNode, const uint32 ChildIndex)
 {
@@ -85,7 +85,7 @@ void USoundNode::InsertChildNode( int32 Index )
 	if (MaxChildNodes > ChildNodes.Num())
 	{
 		ChildNodes.InsertZeroed( Index );
-#if WITH_EDITORONLY_DATA
+#if WITH_EDITOR
 		GraphNode->CreateInputPin();
 #endif //WITH_EDITORONLY_DATA
 	}
diff --git a/Engine/Source/Runtime/Engine/Private/SoundNodeRandom.cpp b/Engine/Source/Runtime/Engine/Private/SoundNodeRandom.cpp
--- a/Engine/Source/Runtime/Engine/Private/SoundNodeRandom.cpp
+++ b/Engine/Source/Runtime/Engine/Private/SoundNodeRandom.cpp
@@ -75,9 +75,9 @@ void USoundNodeRandom::ParseNodes( FAudioDevice* AudioDevice, const UPTRINT Node
 		FixHasBeenUsedArray();  // for now prob need this until resave packages has occurred
 	}
 
-#if WITH_EDITORONLY_DATA
+#if WITH_EDITOR
 	bool bIsPIESound = (GEditor != nullptr) && ((GEditor->bIsSimulatingInEditor || GEditor->PlayWorld != NULL) && ActiveSound.World != NULL);
-#endif //WITH_EDITORONLY_DATA
+#endif //WITH_EDITOR
 
 	// Pick a random child node and save the index.
 	if( *RequiresInitialization )
@@ -85,7 +85,7 @@ void USoundNodeRandom::ParseNodes( FAudioDevice* AudioDevice, const UPTRINT Node
 		NodeIndex = 0;
 		float WeightSum = 0.0f;
 
-#if WITH_EDITORONLY_DATA
+#if WITH_EDITOR
 		if (bIsPIESound)
 		{
 			// Find the first available index - needed if there is only one
@@ -94,14 +94,14 @@ void USoundNodeRandom::ParseNodes( FAudioDevice* AudioDevice, const UPTRINT Node
 				NodeIndex++;
 			}
 		}
-#endif //WITH_EDITORONLY_DATA
+#endif //WITH_EDITOR
 
 		// only calculate the weights that have not been used and use that set for the random choice
 		for( int32 i = 0; i < Weights.Num(); ++i )
 		{
-#if WITH_EDITORONLY_DATA
+#if WITH_EDITOR
 			if (!bIsPIESound || !PIEHiddenNodes.Contains(i))
-#endif //WITH_EDITORONLY_DATA
+#endif //WITH_EDITOR
 			{
 				if( ( bRandomizeWithoutReplacement == false ) ||  ( HasBeenUsed[ i ] != true ) )
 				{
@@ -113,9 +113,9 @@ void USoundNodeRandom::ParseNodes( FAudioDevice* AudioDevice, const UPTRINT Node
 		float Weight = FMath::FRand() * WeightSum;
 		for( int32 i = 0; i < ChildNodes.Num() && i < Weights.Num(); ++i )
 		{
-#if WITH_EDITORONLY_DATA
+#if WITH_EDITOR
 			if(!bIsPIESound || !PIEHiddenNodes.Contains(i) )
-#endif //WITH_EDITORONLY_DATA
+#endif //WITH_EDITOR
 			{
 				if( bRandomizeWithoutReplacement && ( Weights[ i ] >= Weight ) && ( HasBeenUsed[ i ] != true ) )
 				{
@@ -143,9 +143,9 @@ void USoundNodeRandom::ParseNodes( FAudioDevice* AudioDevice, const UPTRINT Node
 
 	// check to see if we have used up our random sounds
 	if( bRandomizeWithoutReplacement && ( HasBeenUsed.Num() > 0 ) && ( NumRandomUsed >= HasBeenUsed.Num()
-#if WITH_EDITORONLY_DATA
+#if WITH_EDITOR
 		|| (bIsPIESound && NumRandomUsed >= (HasBeenUsed.Num() - PIEHiddenNodes.Num()))
-#endif //WITH_EDITORONLY_DATA
+#endif //WITH_EDITOR
 		)	)
 	{
 		// reset all of the children nodes
diff --git a/Engine/Source/Runtime/Engine/Private/SoundNodeSwitch.cpp b/Engine/Source/Runtime/Engine/Private/SoundNodeSwitch.cpp
--- a/Engine/Source/Runtime/Engine/Private/SoundNodeSwitch.cpp
+++ b/Engine/Source/Runtime/Engine/Private/SoundNodeSwitch.cpp
@@ -5,7 +5,7 @@
 #include "SoundDefinitions.h"
 #include "Sound/SoundNodeSwitch.h"
 
-#if WITH_EDITORONLY_DATA
+#if WITH_EDITOR
 #include "UnrealEd.h"
 #endif
 
diff --git a/Engine/Source/Runtime/Engine/Private/SplineComponent.cpp b/Engine/Source/Runtime/Engine/Private/SplineComponent.cpp
--- a/Engine/Source/Runtime/Engine/Private/SplineComponent.cpp
+++ b/Engine/Source/Runtime/Engine/Private/SplineComponent.cpp
@@ -438,7 +438,7 @@ void USplineComponent::ApplyComponentInstanceData(TSharedPtr<FComponentInstanceD
 	}
 }
 
-#if WITH_EDITORONLY_DATA
+#if WITH_EDITOR
 void USplineComponent::PostEditChangeProperty(FPropertyChangedEvent& PropertyChangedEvent)
 {
 	if (PropertyChangedEvent.Property != nullptr)
diff --git a/Engine/Source/Runtime/Engine/Private/StaticMesh.cpp b/Engine/Source/Runtime/Engine/Private/StaticMesh.cpp
--- a/Engine/Source/Runtime/Engine/Private/StaticMesh.cpp
+++ b/Engine/Source/Runtime/Engine/Private/StaticMesh.cpp
@@ -716,11 +716,13 @@ FStaticMeshRenderData::FStaticMeshRenderData()
 
 void FStaticMeshRenderData::Serialize(FArchive& Ar, UStaticMesh* Owner, bool bCooked)
 {
-#if WITH_EDITORONLY_DATA
+#if WITH_EDITOR
 	if (Ar.IsSaving())
 	{
 		ResolveSectionInfo(Owner);
 	}
+#endif
+#if WITH_EDITORONLY_DATA
 	if (!bCooked)
 	{
 		Ar << WedgeMap;
@@ -751,9 +753,9 @@ void FStaticMeshRenderData::Serialize(FArchive& Ar, UStaticMesh* Owner, bool bCo
 
 void FStaticMeshRenderData::InitResources(UStaticMesh* Owner)
 {
-#if WITH_EDITORONLY_DATA
+#if WITH_EDITOR
 	ResolveSectionInfo(Owner);
-#endif // #if WITH_EDITORONLY_DATA
+#endif // #if WITH_EDITOR
 
 	for (int32 LODIndex = 0; LODIndex < LODResources.Num(); ++LODIndex)
 	{
@@ -778,7 +780,7 @@ void FStaticMeshRenderData::AllocateLODResources(int32 NumLODs)
 	}
 }
 
-#if WITH_EDITORONLY_DATA
+#if WITH_EDITOR
 /**
  * Calculates the view distance that a mesh should be displayed at.
  * @param MaxDeviation - The maximum surface-deviation between the reduced geometry and the original. This value should be acquired from Simplygon
@@ -1235,7 +1237,7 @@ void FStaticMeshRenderData::Cache(UStaticMesh* Owner, const FStaticMeshLODSettin
 		FPlatformAtomics::InterlockedAdd(&StaticMeshDerivedDataTimings::BuildCycles, T1-T0);
 	}
 }
-#endif // #if WITH_EDITORONLY_DATA
+#endif // #if WITH_EDITOR
 
 /*-----------------------------------------------------------------------------
 UStaticMesh
@@ -1567,13 +1569,13 @@ FStaticMeshSourceModel::~FStaticMeshSourceModel()
 #endif // #if WITH_EDITOR
 	}
 
-#if WITH_EDITORONLY_DATA
+#if WITH_EDITOR
 void FStaticMeshSourceModel::SerializeBulkData(FArchive& Ar, UObject* Owner)
 {
 	check(RawMeshBulkData != NULL);
 	RawMeshBulkData->Serialize(Ar, Owner);
 }
-#endif // #if WITH_EDITORONLY_DATA
+#endif // #if WITH_EDITOR
 
 /*------------------------------------------------------------------------------
 	FMeshSectionInfoMap
@@ -1668,7 +1670,7 @@ void FMeshSectionInfoMap::Serialize(FArchive& Ar)
 
 #endif // #if WITH_EDITORONLY_DATA
 
-#if WITH_EDITORONLY_DATA
+#if WITH_EDITOR
 static FStaticMeshRenderData& GetPlatformStaticMeshRenderData(UStaticMesh* Mesh, const ITargetPlatform* Platform)
 {
 	check(Mesh && Mesh->RenderData);
@@ -1754,7 +1756,7 @@ void UStaticMesh::Serialize(FArchive& Ar)
 		Ar << NavCollision;
 	}
 
-#if WITH_EDITORONLY_DATA
+#if WITH_EDITOR
 	if (Ar.UE4Ver() < VER_UE4_STATIC_MESH_REFACTOR)
 	{
 		int64 StartCycles = FPlatformTime::Cycles();
@@ -1767,7 +1769,9 @@ void UStaticMesh::Serialize(FArchive& Ar)
 			FPlatformTime::ToSeconds(EndCycles - StartCycles)
 			);
 	}
+#endif
 
+#if WITH_EDITORONLY_DATA
 	if( !StripFlags.IsEditorDataStripped() )
 	{
 		if ( Ar.IsLoading() && Ar.UE4Ver() < VER_UE4_DEPRECATED_STATIC_MESH_THUMBNAIL_PROPERTIES_REMOVED )
@@ -1836,7 +1840,7 @@ void UStaticMesh::Serialize(FArchive& Ar)
 		Ar << Sockets;
 	}
 
-#if WITH_EDITORONLY_DATA
+#if WITH_EDITOR
 	if (Ar.UE4Ver() >= VER_UE4_STATIC_MESH_REFACTOR && !StripFlags.IsEditorDataStripped())
 	{
 		for (int32 i = 0; i < SourceModels.Num(); ++i)
@@ -1850,7 +1854,7 @@ void UStaticMesh::Serialize(FArchive& Ar)
 		// created until postload and it is needed for bounding information
 		bRequiresLODDistanceConversion = Ar.UE4Ver() < VER_UE4_STATIC_MESH_SCREEN_SIZE_LODS;
 	}
-#endif // #if WITH_EDITORONLY_DATA
+#endif // #if WITH_EDITOR
 
 	// Inline the derived data for cooked builds. Never include render data when
 	// counting memory as it is included by GetResourceSize.
@@ -1862,7 +1866,7 @@ void UStaticMesh::Serialize(FArchive& Ar)
 			RenderData->Serialize(Ar, this, bCooked);
 		}
 
-#if WITH_EDITORONLY_DATA
+#if WITH_EDITOR
 		else if (Ar.IsSaving())
 		{
 			FStaticMeshRenderData& PlatformRenderData = GetPlatformStaticMeshRenderData(this, Ar.CookingTarget());
@@ -1912,7 +1916,7 @@ void UStaticMesh::PostLoad()
 {
 	Super::PostLoad();
 
-#if WITH_EDITORONLY_DATA
+#if WITH_EDITOR
 	// Needs to happen before 'CacheDerivedData'
 	if ( GetLinkerUE4Version() < VER_UE4_BUILD_SCALE_VECTOR )
 	{
@@ -1937,7 +1941,7 @@ void UStaticMesh::PostLoad()
 	{
 		FixupZeroTriangleSections();
 	}
-#endif // #if WITH_EDITORONLY_DATA
+#endif // #if WITH_EDITOR
 
 	EnforceLightmapRestrictions();
 	
@@ -2163,7 +2167,7 @@ void UStaticMesh::CreateNavCollision()
 void UStaticMesh::GetVertexColorData(TMap<FVector, FColor>& VertexColorData)
 {
 	VertexColorData.Empty();
-#if WITH_EDITORONLY_DATA
+#if WITH_EDITOR
 	// What LOD to get vertex colors from.  
 	// Currently mesh painting only allows for painting on the first lod.
 	const uint32 PaintingMeshLODIndex = 0;
@@ -2201,7 +2205,7 @@ void UStaticMesh::GetVertexColorData(TMap<FVector, FColor>& VertexColorData)
  */
 void UStaticMesh::SetVertexColorData(const TMap<FVector, FColor>& VertexColorData)
 {
-#if WITH_EDITORONLY_DATA
+#if WITH_EDITOR
 	// What LOD to get vertex colors from.  
 	// Currently mesh painting only allows for painting on the first lod.
 	const uint32 PaintingMeshLODIndex = 0;
@@ -2238,7 +2242,7 @@ void UStaticMesh::SetVertexColorData(const TMap<FVector, FColor>& VertexColorDat
 		SourceModels[PaintingMeshLODIndex].RawMeshBulkData->SaveRawMesh(Mesh);
 	}
 	// TODO_STATICMESH: Build?
-#endif // #if WITH_EDITORONLY_DATA
+#endif // #if WITH_EDITOR
 }
 
 void UStaticMesh::EnforceLightmapRestrictions()
@@ -2527,7 +2531,7 @@ FStaticMeshLODResources& UStaticMesh::GetLODForExport( int32 LODIndex )
 	return RenderData->LODResources[LODIndex];
 }
 
-#if WITH_EDITORONLY_DATA
+#if WITH_EDITOR
 bool UStaticMesh::CanLODsShareStaticLighting() const
 {
 	bool bCanShareData = true;
diff --git a/Engine/Source/Runtime/Engine/Private/StaticMeshBuild.cpp b/Engine/Source/Runtime/Engine/Private/StaticMeshBuild.cpp
--- a/Engine/Source/Runtime/Engine/Private/StaticMeshBuild.cpp
+++ b/Engine/Source/Runtime/Engine/Private/StaticMeshBuild.cpp
@@ -7,16 +7,16 @@
 #include "EnginePrivate.h"
 #include "StaticMeshResources.h"
 
-#if WITH_EDITORONLY_DATA
+#if WITH_EDITOR
 #include "RawMesh.h"
 #include "MeshUtilities.h"
 #include "TargetPlatform.h"
 #include "GenericOctree.h"
-#endif // #if WITH_EDITORONLY_DATA
+#endif // #if WITH_EDITOR
 
 #define LOCTEXT_NAMESPACE "StaticMeshEditor"
 
-#if WITH_EDITORONLY_DATA
+#if WITH_EDITOR
 /**
  * Check the render data for the provided mesh and return true if the mesh
  * contains degenerate tangent bases.
@@ -46,11 +46,11 @@ static bool HasBadTangents(UStaticMesh* Mesh)
 	}
 	return bHasBadTangents;
 }
-#endif // #if WITH_EDITORONLY_DATA
+#endif // #if WITH_EDITOR
 
 void UStaticMesh::Build(bool bSilent)
 {
-#if WITH_EDITORONLY_DATA
+#if WITH_EDITOR
 	if (IsTemplate())
 		return;
 
@@ -155,7 +155,7 @@ void UStaticMesh::Build(bool bSilent)
 	Remapping of painted vertex colors.
 ------------------------------------------------------------------------------*/
 
-#if WITH_EDITORONLY_DATA
+#if WITH_EDITOR
 /** Helper struct for the mesh component vert position octree */
 struct FStaticMeshComponentVertPosOctreeSemantics
 {
@@ -335,13 +335,13 @@ void RemapPaintedVertexColors(
 		}
 	}
 }
-#endif // #if WITH_EDITORONLY_DATA
+#endif // #if WITH_EDITOR
 
 /*------------------------------------------------------------------------------
 	Conversion of legacy source data.
 ------------------------------------------------------------------------------*/
 
-#if WITH_EDITORONLY_DATA
+#if WITH_EDITOR
 
 struct FStaticMeshTriangle
 {
@@ -1067,6 +1067,6 @@ void UStaticMesh::FixupZeroTriangleSections()
 	}
 }
 
-#endif // #if WITH_EDITORONLY_DATA
+#endif // #if WITH_EDITOR
 
 #undef LOCTEXT_NAMESPACE
diff --git a/Engine/Source/Runtime/Engine/Private/StaticMeshComponent.cpp b/Engine/Source/Runtime/Engine/Private/StaticMeshComponent.cpp
--- a/Engine/Source/Runtime/Engine/Private/StaticMeshComponent.cpp
+++ b/Engine/Source/Runtime/Engine/Private/StaticMeshComponent.cpp
@@ -884,7 +884,7 @@ bool UStaticMeshComponent::FixupOverrideColorsIfNecessary( bool bRebuildingStati
 
 void UStaticMeshComponent::PrivateFixupOverrideColors( const TArray<int32>& LODsToUpdate )
 {
-#if WITH_EDITORONLY_DATA
+#if WITH_EDITOR
 	UE_LOG(LogStaticMesh,Verbose,TEXT("Fixing up override colors for %s [%s]"),*GetPathName(),*StaticMesh->GetPathName());
 	for ( TArray<int32>::TConstIterator LODIdxIter( LODsToUpdate ); LODIdxIter; ++LODIdxIter )
 	{
@@ -911,7 +911,7 @@ void UStaticMeshComponent::PrivateFixupOverrideColors( const TArray<int32>& LODs
 		BeginInitResource( CurCompLODInfo.OverrideVertexColors );
 	}
 	StaticMeshDerivedDataKey = StaticMesh->RenderData->DerivedDataKey;
-#endif // WITH_EDITORONLY_DATA
+#endif // WITH_EDITOR
 }
 
 void UStaticMeshComponent::InitResources()
diff --git a/Engine/Source/Runtime/Engine/Private/Texture.cpp b/Engine/Source/Runtime/Engine/Private/Texture.cpp
--- a/Engine/Source/Runtime/Engine/Private/Texture.cpp
+++ b/Engine/Source/Runtime/Engine/Private/Texture.cpp
@@ -246,7 +246,9 @@ void UTexture::LegacySerialize(FArchive& Ar, FStripDataFlags& StripFlags)
 		Source.bPNGCompressed = true;
 		break;
 	case TSAT_DDSFile:
+#if WITH_EDITOR
 		Source.ConvertFromLegacyDDS();
+#endif
 		break;
 	}
 
@@ -310,10 +312,10 @@ void UTexture::PostLoad()
 
 	if( !IsTemplate() )
 	{
-#if WITH_EDITORONLY_DATA
+#if WITH_EDITOR
 		// Update cached LOD bias.
 		UpdateCachedLODBias();
-#endif // #if WITH_EDITORONLY_DATA
+#endif // #if WITH_EDITOR
 
 		// The texture will be cached by the cubemap it is contained within on consoles.
 		UTextureCube* CubeMap = Cast<UTextureCube>(GetOuter());
@@ -401,7 +403,7 @@ void UTexture::PreSave()
 
 	Super::PreSave();
 
-#if WITH_EDITORONLY_DATA
+#if WITH_EDITOR
 	if (DeferCompression)
 	{
 		GWarn->StatusUpdate( 0, 0, FText::Format( NSLOCTEXT("UnrealEd", "SavingPackage_CompressingTexture", "Compressing texture:  {0}"), FText::FromString(GetName()) ) );
@@ -412,7 +414,7 @@ void UTexture::PreSave()
 	GWarn->StatusUpdate( 0, 0, FText::Format( NSLOCTEXT("UnrealEd", "SavingPackage_CompressingSourceArt", "Compressing source art for texture:  {0}"), FText::FromString(GetName()) ) );
 	Source.Compress();
 
-#endif // #if WITH_EDITORONLY_DATA
+#endif // #if WITH_EDITOR
 }
 
 
@@ -482,7 +484,7 @@ bool UTexture::ForceUpdateTextureStreaming()
 		const bool bOldOnlyStreamInTextures = CVarOnlyStreamInTextures->GetInt() != 0;
 		CVarOnlyStreamInTextures->Set(false);
 
-#if WITH_EDITORONLY_DATA
+#if WITH_EDITOR
 		for( TObjectIterator<UTexture2D> It; It; ++It )
 		{
 			UTexture* Texture = *It;
@@ -490,7 +492,7 @@ bool UTexture::ForceUpdateTextureStreaming()
 			// Update cached LOD bias.
 			Texture->UpdateCachedLODBias();
 		}
-#endif // #if WITH_EDITORONLY_DATA
+#endif // #if WITH_EDITOR
 
 		// Make sure we iterate over all textures by setting it to high value.
 		IStreamingManager::Get().SetNumIterationsForNextFrame( 100 );
@@ -527,7 +529,7 @@ FTextureSource::FTextureSource()
 {
 }
 
-#if WITH_EDITORONLY_DATA
+#if WITH_EDITOR
 
 void FTextureSource::Init(
 		int32 NewSizeX,
@@ -932,4 +934,4 @@ uint32 UTexture::GetMaximumDimension() const
 {
 	return GetMax2DTextureDimension();
 }
-#endif // #if WITH_EDITORONLY_DATA
+#endif // #if WITH_EDITOR
diff --git a/Engine/Source/Runtime/Engine/Private/Texture2D.cpp b/Engine/Source/Runtime/Engine/Private/Texture2D.cpp
--- a/Engine/Source/Runtime/Engine/Private/Texture2D.cpp
+++ b/Engine/Source/Runtime/Engine/Private/Texture2D.cpp
@@ -175,7 +175,7 @@ bool UTexture2D::GetResourceMemSettings(int32 FirstMipIdx, int32& OutSizeX, int3
 	return false;
 }
 
-#if WITH_EDITORONLY_DATA
+#if WITH_EDITOR
 void UTexture2D::LegacySerialize(FArchive& Ar, FStripDataFlags& StripDataFlags)
 {
 	check(Ar.UE4Ver() < VER_UE4_TEXTURE_SOURCE_ART_REFACTOR);
@@ -272,7 +272,7 @@ void UTexture2D::LegacySerialize(FArchive& Ar, FStripDataFlags& StripDataFlags)
 		SRGB = false;
 	}
 }
-#endif // #if WITH_EDITORONLY_DATA
+#endif // #if WITH_EDITOR
 
 void UTexture2D::Serialize(FArchive& Ar)
 {
@@ -288,19 +288,19 @@ void UTexture2D::Serialize(FArchive& Ar)
 	}
 
 	/** Legacy serialization. */
-#if WITH_EDITORONLY_DATA
+#if WITH_EDITOR
 	if (Ar.UE4Ver() < VER_UE4_TEXTURE_SOURCE_ART_REFACTOR)
 	{
 		UTexture2D::LegacySerialize(Ar, StripDataFlags);
 	}
-#endif // #if WITH_EDITORONLY_DATA
+#endif // #if WITH_EDITOR
 
 	if (Ar.IsCooking() || bCooked)
 	{
 		SerializeCookedPlatformData(Ar);
 	}
 
-#if WITH_EDITORONLY_DATA	
+#if WITH_EDITOR	
 	if (Ar.IsLoading() && !Ar.IsTransacting() && !bCooked && !(GetOutermost()->PackageFlags & PKG_ReloadingForCooker))
 	{
 		// The composite texture may not have been loaded yet. We have to defer caching platform
@@ -310,7 +310,7 @@ void UTexture2D::Serialize(FArchive& Ar)
 			BeginCachePlatformData();
 		}
 	}
-#endif // #if WITH_EDITORONLY_DATA
+#endif // #if WITH_EDITOR
 }
 
 float UTexture2D::GetLastRenderTimeForStreaming()
@@ -365,7 +365,7 @@ void UTexture2D::PostEditChangeProperty(FPropertyChangedEvent& PropertyChangedEv
 float UTexture2D::GetAverageBrightness(bool bIgnoreTrueBlack, bool bUseGrayscale)
 {
 	float AvgBrightness = -1.0f;
-#if WITH_EDITORONLY_DATA
+#if WITH_EDITOR
 	TArray<uint8> RawData;
 	// use the source art if it exists
 	if (Source.IsValid() && Source.GetFormat() == TSF_BGRA8)
@@ -424,7 +424,7 @@ float UTexture2D::GetAverageBrightness(bool bIgnoreTrueBlack, bool bUseGrayscale
 			AvgBrightness = PixelSum / Divisor;
 		}
 	}
-#endif // #if WITH_EDITORONLY_DATA
+#endif // #if WITH_EDITOR
 	return AvgBrightness;
 }
 
@@ -457,14 +457,14 @@ void UTexture2D::CancelPendingTextureStreaming()
 
 void UTexture2D::PostLoad()
 {
-#if WITH_EDITORONLY_DATA
+#if WITH_EDITOR
 	ImportedSize = FIntPoint(Source.GetSizeX(),Source.GetSizeY());
 
 	if (FApp::CanEverRender())
 	{
 		FinishCachePlatformData();
 	}
-#endif // #if WITH_EDITORONLY_DATA
+#endif // #if WITH_EDITOR
 
 	// Route postload, which will update bIsStreamable as UTexture::PostLoad calls UpdateResource.
 	Super::PostLoad();
@@ -483,7 +483,7 @@ void UTexture2D::PreSave()
 }
 void UTexture2D::GetAssetRegistryTags(TArray<FAssetRegistryTag>& OutTags) const
 {
-#if WITH_EDITORONLY_DATA
+#if WITH_EDITOR
 	int32 SizeX = Source.GetSizeX();
 	int32 SizeY = Source.GetSizeY();
 #else
@@ -508,10 +508,10 @@ void UTexture2D::UpdateResource()
 		FPlatformProcess::Sleep(0);
 	}
 
-#if WITH_EDITORONLY_DATA
+#if WITH_EDITOR
 	// Recache platform data if the source has changed.
 	CachePlatformData();
-#endif // #if WITH_EDITORONLY_DATA
+#endif // #if WITH_EDITOR
 
 	// Route to super.
 	Super::UpdateResource();
@@ -555,9 +555,9 @@ FString UTexture2D::GetDesc()
 	uint32 EffectiveSizeX;
 	uint32 EffectiveSizeY;
 
-#if WITH_EDITORONLY_DATA
+#if WITH_EDITOR
 	UpdateCachedLODBias();
-#endif //#if WITH_EDITORONLY_DATA
+#endif //#if WITH_EDITOR
 
 	GSystemSettings.TextureLODSettings.ComputeInGameMaxResolution(GetCachedLODBias(), *this, EffectiveSizeX, EffectiveSizeY);
 
@@ -675,7 +675,7 @@ bool UTexture2D::UpdateStreamingStatus( bool bWaitForMipFading /*= false*/ )
 			}
 			bHasPendingRequestInFlight = false;
 
-#if WITH_EDITORONLY_DATA
+#if WITH_EDITOR
 			if ( Texture2DResource && Texture2DResource->DidDerivedDataRequestFail() )
 			{
 				// A streaming request to the derived data cache failed. That means the texture itself needs to be updated.
@@ -689,7 +689,7 @@ bool UTexture2D::UpdateStreamingStatus( bool bWaitForMipFading /*= false*/ )
 					ForceRebuildPlatformData();
 				}
 			}
-#endif // #if WITH_EDITORONLY_DATA
+#endif // #if WITH_EDITOR
 		}
 	}
 	return bHasPendingRequestInFlight;
@@ -766,7 +766,7 @@ bool UTexture2D::GetSourceArtCRC(uint32& OutSourceCRC)
 {
 	bool bResult = false;
 	TArray<uint8> RawData;
-#if WITH_EDITORONLY_DATA
+#if WITH_EDITOR
 	// use the source art if it exists
 	if (Source.IsValid())
 	{
@@ -783,14 +783,14 @@ bool UTexture2D::GetSourceArtCRC(uint32& OutSourceCRC)
 		OutSourceCRC = FCrc::MemCrc_DEPRECATED((void*)(RawData.GetData()), RawData.Num());
 		bResult = true;
 	}
-#endif // #if WITH_EDITORONLY_DATA
+#endif // #if WITH_EDITOR
 	return bResult;
 }
 
 bool UTexture2D::HasSameSourceArt(UTexture2D* InTexture)
 {
 	bool bResult = false;
-#if WITH_EDITORONLY_DATA
+#if WITH_EDITOR
 	TArray<uint8> RawData1;
 	TArray<uint8> RawData2;
 	int32 SizeX = 0;
@@ -814,7 +814,7 @@ bool UTexture2D::HasSameSourceArt(UTexture2D* InTexture)
 			bResult = true;
 		}
 	}
-#endif // #if WITH_EDITORONLY_DATA
+#endif // #if WITH_EDITOR
 	return bResult;
 }
 
diff --git a/Engine/Source/Runtime/Engine/Private/TextureCube.cpp b/Engine/Source/Runtime/Engine/Private/TextureCube.cpp
--- a/Engine/Source/Runtime/Engine/Private/TextureCube.cpp
+++ b/Engine/Source/Runtime/Engine/Private/TextureCube.cpp
@@ -20,7 +20,7 @@ void UTextureCube::Serialize(FArchive& Ar)
 
 	bool bCooked = false;
 
-#if WITH_EDITORONLY_DATA
+#if WITH_EDITOR
 	// Handle legacy serialization.
 	if (Ar.UE4Ver() < VER_UE4_TEXTURE_SOURCE_ART_REFACTOR)
 	{
@@ -40,7 +40,7 @@ void UTextureCube::Serialize(FArchive& Ar)
 		Source.UseHashAsGuid();
 	}
 	else
-#endif // #if WITH_EDITORONLY_DATA
+#endif // #if WITH_EDITOR
 	{
 		FStripDataFlags StripFlags(Ar);
 		bCooked = Ar.IsCooking();
@@ -52,26 +52,26 @@ void UTextureCube::Serialize(FArchive& Ar)
 		}
 	}
 
-#if WITH_EDITORONLY_DATA
+#if WITH_EDITOR
 	if (Ar.IsLoading() && !Ar.IsTransacting() && !bCooked)
 	{
 		BeginCachePlatformData();
 	}
-#endif // #if WITH_EDITORONLY_DATA
+#endif // #if WITH_EDITOR
 }
 
 void UTextureCube::PostLoad()
 {
-#if WITH_EDITORONLY_DATA
+#if WITH_EDITOR
 	FinishCachePlatformData();
-#endif // #if WITH_EDITORONLY_DATA
+#endif // #if WITH_EDITOR
 
 	Super::PostLoad();
 }
 
 void UTextureCube::GetAssetRegistryTags(TArray<FAssetRegistryTag>& OutTags) const
 {
-#if WITH_EDITORONLY_DATA
+#if WITH_EDITOR
 	int32 SizeX = Source.GetSizeX();
 	int32 SizeY = Source.GetSizeY();
 #else
@@ -87,10 +87,10 @@ void UTextureCube::GetAssetRegistryTags(TArray<FAssetRegistryTag>& OutTags) cons
 
 void UTextureCube::UpdateResource()
 {
-#if WITH_EDITORONLY_DATA
+#if WITH_EDITOR
 	// Recache platform data if the source has changed.
 	CachePlatformData();
-#endif // #if WITH_EDITORONLY_DATA
+#endif // #if WITH_EDITOR
 
 	// Route to super.
 	Super::UpdateResource();
@@ -371,9 +371,9 @@ SIZE_T UTextureCube::GetResourceSize(EResourceSizeMode::Type Mode)
 	return CalcTextureMemorySizeEnum(TMC_ResidentMips);
 }
 
-#if WITH_EDITORONLY_DATA
+#if WITH_EDITOR
 uint32 UTextureCube::GetMaximumDimension() const
 {
 	return GetMaxCubeTextureDimension();
 }
-#endif // #if WITH_EDITORONLY_DATA
\ No newline at end of file
+#endif // #if WITH_EDITOR
\ No newline at end of file
diff --git a/Engine/Source/Runtime/Engine/Private/TextureDerivedData.cpp b/Engine/Source/Runtime/Engine/Private/TextureDerivedData.cpp
--- a/Engine/Source/Runtime/Engine/Private/TextureDerivedData.cpp
+++ b/Engine/Source/Runtime/Engine/Private/TextureDerivedData.cpp
@@ -12,16 +12,15 @@ enum
 	NUM_INLINE_DERIVED_MIPS = 7,
 };
 
-#if WITH_EDITORONLY_DATA
+#if WITH_EDITOR
 
 #include "UObjectAnnotation.h"
 #include "DerivedDataCacheInterface.h"
 #include "DerivedDataPluginInterface.h"
-#include "TextureCompressorModule.h"
 #include "DDSLoader.h"
 #include "RenderUtils.h"
 #include "TargetPlatform.h"
-
+#include "TextureCompressorModule.h"
 #include "ImageCore.h"
 
 /*------------------------------------------------------------------------------
@@ -97,14 +96,10 @@ namespace TextureDerivedDataTimings
 	};
 }
 
-#endif // #if WITH_EDITORONLY_DATA
-
 /*------------------------------------------------------------------------------
 	Derived data key generation.
 ------------------------------------------------------------------------------*/
 
-#if WITH_EDITORONLY_DATA
-
 /**
  * Serialize build settings for use when generating the derived data key.
  */
@@ -267,13 +262,13 @@ static void GetTextureDerivedDataKey(
 	GetTextureDerivedDataKeyFromSuffix(KeySuffix, OutKey);
 }
 
-#endif // #if WITH_EDITORONLY_DATA
+#endif // #if WITH_EDITOR
 
 /*------------------------------------------------------------------------------
 	Texture compression.
 ------------------------------------------------------------------------------*/
 
-#if WITH_EDITORONLY_DATA
+#if WITH_EDITOR
 
 /**
  * Sets texture build settings.
@@ -461,13 +456,13 @@ static void PutDerivedDataInCache(
 	UE_LOG(LogTexture,Verbose,TEXT("%s  Derived Data: %d bytes"),*LogString,RawDerivedData.Num());
 }
 
-#endif // #if WITH_EDITORONLY_DATA
+#endif // #if WITH_EDITOR
 
 /*------------------------------------------------------------------------------
 	Derived data.
 ------------------------------------------------------------------------------*/
 
-#if WITH_EDITORONLY_DATA
+#if WITH_EDITOR
 
 class FTextureStatusMessageContext : public FScopedSlowTask
 {
@@ -1066,7 +1061,7 @@ bool FTexturePlatformData::TryInlineMipData()
 	return true;
 }
 
-#endif // #if WITH_EDITORONLY_DATA
+#endif // #if WITH_EDITOR
 
 FTexturePlatformData::FTexturePlatformData()
 	: SizeX(0)
@@ -1081,7 +1076,7 @@ FTexturePlatformData::FTexturePlatformData()
 
 FTexturePlatformData::~FTexturePlatformData()
 {
-#if WITH_EDITORONLY_DATA
+#if WITH_EDITOR
 	if (AsyncTask)
 	{
 		AsyncTask->EnsureCompletion();
@@ -1095,12 +1090,12 @@ bool FTexturePlatformData::TryLoadMips(int32 FirstMipToLoad, void** OutMipData)
 {
 	int32 NumMipsCached = 0;
 
-#if WITH_EDITORONLY_DATA
+#if WITH_EDITOR
 	TArray<uint8> TempData;
 	FAsyncMipHandles AsyncHandles;
 	FDerivedDataCacheInterface& DDC = GetDerivedDataCacheRef();
 	BeginLoadDerivedMips(Mips, FirstMipToLoad, AsyncHandles);
-#endif // #if WITH_EDITORONLY_DATA
+#endif // #if WITH_EDITOR
 
 	// Load remaining mips (if any) from bulk data.
 	for (int32 MipIndex = FirstMipToLoad; MipIndex < Mips.Num(); ++MipIndex)
@@ -1117,7 +1112,7 @@ bool FTexturePlatformData::TryLoadMips(int32 FirstMipToLoad, void** OutMipData)
 		}
 	}
 
-#if WITH_EDITORONLY_DATA
+#if WITH_EDITOR
 	// Wait for async DDC gets.
 	for (int32 MipIndex = FirstMipToLoad; MipIndex < Mips.Num(); ++MipIndex)
 	{
@@ -1143,7 +1138,7 @@ bool FTexturePlatformData::TryLoadMips(int32 FirstMipToLoad, void** OutMipData)
 			TempData.Reset();
 		}
 	}
-#endif // #if WITH_EDITORONLY_DATA
+#endif // #if WITH_EDITOR
 
 	if (NumMipsCached != (Mips.Num() - FirstMipToLoad))
 	{
@@ -1162,7 +1157,7 @@ bool FTexturePlatformData::TryLoadMips(int32 FirstMipToLoad, void** OutMipData)
 	return true;
 }
 
-#if WITH_EDITORONLY_DATA
+#if WITH_EDITOR
 bool FTexturePlatformData::AreDerivedMipsAvailable() const
 {
 	bool bMipsAvailable = true;
@@ -1177,7 +1172,7 @@ bool FTexturePlatformData::AreDerivedMipsAvailable() const
 	}
 	return bMipsAvailable;
 }
-#endif // #if WITH_EDITORONLY_DATA
+#endif // #if WITH_EDITOR
 
 static void SerializePlatformData(
 	FArchive& Ar,
@@ -1315,21 +1310,21 @@ void FAsyncStreamDerivedMipWorker::DoWork()
 
 void UTexture2D::GetMipData(int32 FirstMipToLoad, void** OutMipData)
 {
-#if WITH_EDITORONLY_DATA
+#if WITH_EDITOR
 	TextureDerivedDataTimings::FScopedMeasurement Timer(TextureDerivedDataTimings::GetMipDataTime);
-#endif // #if WITH_EDITORONLY_DATA
+#endif // #if WITH_EDITOR
 	if (PlatformData->TryLoadMips(FirstMipToLoad, OutMipData) == false)
 	{
 		// Unable to load mips from the cache. Rebuild the texture and try again.
 		UE_LOG(LogTexture,Warning,TEXT("GetMipData failed for %s (%s)"),
 			*GetPathName(), GPixelFormats[GetPixelFormat()].Name);
-#if WITH_EDITORONLY_DATA
+#if WITH_EDITOR
 		ForceRebuildPlatformData();
 		if (PlatformData->TryLoadMips(FirstMipToLoad, OutMipData) == false)
 		{
 			UE_LOG(LogTexture,Error,TEXT("Failed to build texture %s."), *GetPathName());
 		}
-#endif // #if WITH_EDITORONLY_DATA
+#endif // #if WITH_EDITOR
 	}
 }
 
@@ -1352,7 +1347,7 @@ void UTexture::CleanupCachedCookedPlatformData()
 	}
 }
 
-#if WITH_EDITORONLY_DATA
+#if WITH_EDITOR
 void UTexture::CachePlatformData(bool bAsyncCache)
 {
 	FTexturePlatformData** PlatformDataLinkPtr = GetRunningPlatformData();
@@ -1603,7 +1598,7 @@ void UTexture::MarkPlatformDataTransient()
 		}
 	}
 }
-#endif // #if WITH_EDITORONLY_DATA
+#endif // #if WITH_EDITOR
 
 void UTexture::CleanupCachedRunningPlatformData()
 {
@@ -1632,7 +1627,7 @@ void UTexture::SerializeCookedPlatformData(FArchive& Ar)
 	UEnum* PixelFormatEnum = UTexture::GetPixelFormatEnum();
 
 
-#if WITH_EDITORONLY_DATA
+#if WITH_EDITOR
 	TextureDerivedDataTimings::FScopedMeasurement Timer(TextureDerivedDataTimings::SerializeCookedTime);
 	if (Ar.IsCooking() && Ar.IsPersistent())
 	{
@@ -1691,7 +1686,7 @@ void UTexture::SerializeCookedPlatformData(FArchive& Ar)
 		Ar << PixelFormatName;
 	}
 	else
-#endif // #if WITH_EDITORONLY_DATA
+#endif // #if WITH_EDITOR
 	{
 
 		FTexturePlatformData** RunningPlatformDataPtr = GetRunningPlatformData();
diff --git a/Engine/Source/Runtime/Engine/Private/TextureRenderTarget2D.cpp b/Engine/Source/Runtime/Engine/Private/TextureRenderTarget2D.cpp
--- a/Engine/Source/Runtime/Engine/Private/TextureRenderTarget2D.cpp
+++ b/Engine/Source/Runtime/Engine/Private/TextureRenderTarget2D.cpp
@@ -128,7 +128,7 @@ FString UTextureRenderTarget2D::GetDesc()
 UTexture2D* UTextureRenderTarget2D::ConstructTexture2D(UObject* Outer, const FString& NewTexName, EObjectFlags ObjectFlags, uint32 Flags, TArray<uint8>* AlphaOverride)
 {
 	UTexture2D* Result = NULL;
-#if WITH_EDITORONLY_DATA
+#if WITH_EDITOR
 	// Check render target size is valid and power of two.
 	const bool bIsValidSize = (SizeX != 0 && !(SizeX & (SizeX - 1)) &&
 		SizeY != 0 && !(SizeY & (SizeY - 1)));
diff --git a/Engine/Source/Runtime/Engine/Private/TextureRenderTargetCube.cpp b/Engine/Source/Runtime/Engine/Private/TextureRenderTargetCube.cpp
--- a/Engine/Source/Runtime/Engine/Private/TextureRenderTargetCube.cpp
+++ b/Engine/Source/Runtime/Engine/Private/TextureRenderTargetCube.cpp
@@ -111,7 +111,7 @@ UTextureCube* UTextureRenderTargetCube::ConstructTextureCube(
 )
 {
 	UTextureCube* Result = NULL;
-#if WITH_EDITORONLY_DATA
+#if WITH_EDITOR
 	// Check render target size is valid and power of two.
 	if (SizeX != 0 && !(SizeX & (SizeX - 1)))
 	{
@@ -183,7 +183,7 @@ UTextureCube* UTextureRenderTargetCube::ConstructTextureCube(
 			Result->PostEditChange();
 		}
 	}
-#endif // #if WITH_EDITORONLY_DATA
+#endif // #if WITH_EDITOR
 	return Result;
 }
 
diff --git a/Engine/Source/Runtime/Engine/Private/UnrealEngine.cpp b/Engine/Source/Runtime/Engine/Private/UnrealEngine.cpp
--- a/Engine/Source/Runtime/Engine/Private/UnrealEngine.cpp
+++ b/Engine/Source/Runtime/Engine/Private/UnrealEngine.cpp
@@ -1192,7 +1192,7 @@ void LoadSpecialMaterial(const FString& MaterialName, UMaterial*& Material, bool
 		// to set the bit
 		else if (!Material->bUsedAsSpecialEngineMaterial && bCheckUsage) 
 		{
-#if !WITH_EDITORONLY_DATA
+#if !WITH_EDITOR
 			// consoles must have the flag set properly in the editor
 			UE_LOG(LogEngine, Fatal,TEXT("The special material (%s) was not marked with bUsedAsSpecialEngineMaterial. Make sure this flag is set in the editor, save the package, and compile shaders for this platform"), *MaterialName);
 #else
@@ -4754,7 +4754,7 @@ bool UEngine::HandleObjCommand( const TCHAR* Cmd, FOutputDevice& Ar )
 				for ( TFieldIterator<UProperty> It(Obj->GetClass()); It; ++It )
 				{
 					Value.Empty();
-#if WITH_EDITORONLY_DATA
+#if WITH_EDITOR
 					if ( HiddenCategories.Num() )
 					{
 						const FString Category = FObjectEditorUtils::GetCategory(*It);
@@ -4799,7 +4799,7 @@ bool UEngine::HandleObjCommand( const TCHAR* Cmd, FOutputDevice& Ar )
 							continue;
 						}
 					}
-#endif // #if WITH_EDITORONLY_DATA
+#endif // #if WITH_EDITOR
 					if ( LastOwnerClass != It->GetOwnerClass() )
 					{
 						LastOwnerClass = It->GetOwnerClass();
diff --git a/Engine/Source/Runtime/Engine/Private/World.cpp b/Engine/Source/Runtime/Engine/Private/World.cpp
--- a/Engine/Source/Runtime/Engine/Private/World.cpp
+++ b/Engine/Source/Runtime/Engine/Private/World.cpp
@@ -2901,7 +2901,7 @@ void UWorld::CleanupWorld(bool bSessionEnded, bool bCleanupResources, UWorld* Ne
 
 	NetworkActors.Empty();
 
-#if WITH_EDITORONLY_DATA
+#if WITH_EDITOR
 	// If we're server traveling, we need to break the reference dependency here (caused by levelscript)
 	// to avoid a GC crash for not cleaning up the gameinfo referenced by levelscript
 	if( IsGameWorld() && !GIsEditor && !IsRunningCommandlet() && bSessionEnded && bCleanupResources )
@@ -2919,7 +2919,7 @@ void UWorld::CleanupWorld(bool bSessionEnded, bool bCleanupResources, UWorld* Ne
 			}
 		}
 	}
-#endif //WITH_EDITORONLY_DATA
+#endif //WITH_EDITOR
 
 #if ENABLE_VISUAL_LOG
 	FVisualLog::Get().Cleanup();
diff --git a/Engine/Source/Runtime/Engine/Public/StaticMeshResources.h b/Engine/Source/Runtime/Engine/Public/StaticMeshResources.h
--- a/Engine/Source/Runtime/Engine/Public/StaticMeshResources.h
+++ b/Engine/Source/Runtime/Engine/Public/StaticMeshResources.h
@@ -779,7 +779,8 @@ private:
 #if WITH_EDITORONLY_DATA
 	/** Allow the editor to explicitly update section information. */
 	friend class FLevelOfDetailSettingsLayout;
-
+#endif // #if WITH_EDITORONLY_DATA
+#if WITH_EDITOR
 	/** Resolve all per-section settings. */
 	ENGINE_API void ResolveSectionInfo(UStaticMesh* Owner);
 #endif // #if WITH_EDITORONLY_DATA
@@ -1008,7 +1009,7 @@ protected:
 	float GetScreenSize(int32 LODIndex) const;
 };
 
-#if WITH_EDITORONLY_DATA
+#if WITH_EDITOR
 /**
  * Remaps painted vertex colors when the renderable mesh has changed.
  * @param InPaintedVertices - The original position and normal for each painted vertex.
@@ -1024,4 +1025,4 @@ ENGINE_API void RemapPaintedVertexColors(
 	const FStaticMeshVertexBuffer* OptionalVertexBuffer,
 	TArray<FColor>& OutOverrideColors
 	);
-#endif // #if WITH_EDITORONLY_DATA
+#endif // #if WITH_EDITOR
diff --git a/Engine/Source/Runtime/Renderer/Private/AtmosphereRendering.cpp b/Engine/Source/Runtime/Renderer/Private/AtmosphereRendering.cpp
--- a/Engine/Source/Runtime/Renderer/Private/AtmosphereRendering.cpp
+++ b/Engine/Source/Runtime/Renderer/Private/AtmosphereRendering.cpp
@@ -1573,7 +1573,7 @@ FAtmosphericFogSceneInfo::FAtmosphericFogSceneInfo(UAtmosphericFogComponent* InC
 	, SunDiscScale(InComponent->SunDiscScale)
 	, RenderFlag(EAtmosphereRenderFlag::E_EnableAll)
 	, InscatterAltitudeSampleNum(InComponent->PrecomputeParams.InscatterAltitudeSampleNum)
-#if WITH_EDITOR
+#if WITH_EDITORONLY_DATA
 	, bNeedRecompute(false)
 	, MaxScatteringOrder(InComponent->PrecomputeParams.MaxScatteringOrder)
 	, AtmospherePhase(0)
