Fixed compilation issues in navigation-generator-less builds #UE4

diff --git a/Engine/Source/Runtime/Engine/Classes/AI/Navigation/NavigationSystem.h b/Engine/Source/Runtime/Engine/Classes/AI/Navigation/NavigationSystem.h
--- a/Engine/Source/Runtime/Engine/Classes/AI/Navigation/NavigationSystem.h
+++ b/Engine/Source/Runtime/Engine/Classes/AI/Navigation/NavigationSystem.h
@@ -405,9 +405,11 @@ protected:
 	/** registers NavArea classes awaiting registration in PendingNavAreaRegistration */
 	void ProcessNavAreaPendingRegistration();
 
+#if WITH_NAVIGATION_GENERATOR
 	/** used to apply updates of nav volumes in navigation system's tick */
 	void PerformNavigationBoundsUpdate(ANavMeshBoundsVolume* NavVolume);
-	
+#endif // WITH_NAVIGATION_GENERATOR
+
 	/** @return pointer to ANavigationData instance of given ID, or NULL if it was not found. Note it looks only through registered navigation data */
 	ANavigationData* GetNavDataWithID(const uint16 NavDataID) const;
 
diff --git a/Engine/Source/Runtime/Engine/Private/AI/Navigation/NavigationSystem.cpp b/Engine/Source/Runtime/Engine/Private/AI/Navigation/NavigationSystem.cpp
--- a/Engine/Source/Runtime/Engine/Private/AI/Navigation/NavigationSystem.cpp
+++ b/Engine/Source/Runtime/Engine/Private/AI/Navigation/NavigationSystem.cpp
@@ -475,6 +475,7 @@ void UNavigationSystem::Tick(float DeltaSeconds)
 		ProcessNavAreaPendingRegistration();
 	}
 
+#if WITH_NAVIGATION_GENERATOR
 	if (PendingNavVolumeUpdates.Num() > 0)
 	{
 		for (auto Volume : PendingNavVolumeUpdates)
@@ -486,6 +487,7 @@ void UNavigationSystem::Tick(float DeltaSeconds)
 		}
 		PendingNavVolumeUpdates.Reset();
 	}
+#endif // WITH_NAVIGATION_GENERATOR
 
 	if (PendingOctreeUpdates.Num() > 0)
 	{
