UMG - Moving the GetKey function.

diff --git a/Engine/Source/Runtime/Engine/Classes/Kismet/KismetInputLibrary.h b/Engine/Source/Runtime/Engine/Classes/Kismet/KismetInputLibrary.h
--- a/Engine/Source/Runtime/Engine/Classes/Kismet/KismetInputLibrary.h
+++ b/Engine/Source/Runtime/Engine/Classes/Kismet/KismetInputLibrary.h
@@ -15,6 +15,14 @@ class UKismetInputLibrary : public UBlueprintFunctionLibrary
 	UFUNCTION(BlueprintCallable, Category="Input")
 	static void CalibrateTilt();
 
+	/**
+	 * Returns the key for this event.
+	 *
+	 * @return  Key name
+	 */
+	UFUNCTION(BlueprintCallable, Category="Keyboard")
+	static FKey GetKey(const FKeyboardEvent& Input);
+
 	/**
 	 * Test if the input key are equal (A == B)
 	 * @param A - The key to compare against
diff --git a/Engine/Source/Runtime/Engine/Private/KismetInputLibrary.cpp b/Engine/Source/Runtime/Engine/Private/KismetInputLibrary.cpp
--- a/Engine/Source/Runtime/Engine/Private/KismetInputLibrary.cpp
+++ b/Engine/Source/Runtime/Engine/Private/KismetInputLibrary.cpp
@@ -21,6 +21,11 @@ void UKismetInputLibrary::CalibrateTilt()
 	GEngine->Exec(NULL, TEXT("CALIBRATEMOTION"));
 }
 
+FKey UKismetInputLibrary::GetKey(const FKeyboardEvent& Input)
+{
+	return Input.GetKey();
+}
+
 bool UKismetInputLibrary::EqualEqual_KeyKey(FKey A, FKey B)
 {
 	return A == B;
