UMG - Exposing many more functions on structures used in events for slate for UMG.  Working on input forwarding.

diff --git a/Engine/Plugins/Script/ScriptPlugin/Source/ScriptPlugin/ScriptPlugin.Build.cs b/Engine/Plugins/Script/ScriptPlugin/Source/ScriptPlugin/ScriptPlugin.Build.cs
--- a/Engine/Plugins/Script/ScriptPlugin/Source/ScriptPlugin/ScriptPlugin.Build.cs
+++ b/Engine/Plugins/Script/ScriptPlugin/Source/ScriptPlugin/ScriptPlugin.Build.cs
@@ -28,6 +28,7 @@ namespace UnrealBuildTool.Rules
 					"CoreUObject",
 					"Engine",
 					"InputCore",
+					"SlateCore",
 					// ... add other public dependencies that you statically link with here ...
 				}
 				);
diff --git a/Engine/Source/Runtime/Engine/Classes/Kismet/KismetInputLibrary.h b/Engine/Source/Runtime/Engine/Classes/Kismet/KismetInputLibrary.h
--- a/Engine/Source/Runtime/Engine/Classes/Kismet/KismetInputLibrary.h
+++ b/Engine/Source/Runtime/Engine/Classes/Kismet/KismetInputLibrary.h
@@ -15,14 +15,6 @@ class UKismetInputLibrary : public UBlueprintFunctionLibrary
 	UFUNCTION(BlueprintCallable, Category="Input")
 	static void CalibrateTilt();
 
-	/**
-	 * Returns the key for this event.
-	 *
-	 * @return  Key name
-	 */
-	UFUNCTION(BlueprintCallable, Category="Keyboard")
-	static FKey GetKey(const FKeyboardEvent& Input);
-
 	/**
 	 * Test if the input key are equal (A == B)
 	 * @param A - The key to compare against
@@ -135,4 +127,77 @@ class UKismetInputLibrary : public UBlueprintFunctionLibrary
 	 */
 	UFUNCTION(BlueprintPure, meta=( FriendlyName = "IsRightCommandDown" ), Category="Utilities|InputEvent")
 	static bool InputEvent_IsRightCommandDown(const FInputEvent& Input);
+
+
+	/**
+	 * Returns the key for this event.
+	 *
+	 * @return  Key name
+	 */
+	UFUNCTION(BlueprintCallable, Category="Utilities|KeyboardEvent")
+	static FKey GetKey(const FKeyboardEvent& Input);
+
+
+	/** @return The position of the cursor in screen space */
+	UFUNCTION(BlueprintPure, meta=( FriendlyName = "GetScreenSpacePosition" ), Category="Utilities|PointerEvent")
+	static FVector2D PointerEvent_GetScreenSpacePosition(const FPointerEvent& Input);
+
+	/** @return The position of the cursor in screen space last time we handled an input event */
+	UFUNCTION(BlueprintPure, meta=( FriendlyName = "GetLastScreenSpacePosition" ), Category="Utilities|PointerEvent")
+	static FVector2D PointerEvent_GetLastScreenSpacePosition(const FPointerEvent& Input);
+
+	/** @return the distance the mouse traveled since the last event was handled. */
+	UFUNCTION(BlueprintPure, meta=( FriendlyName = "GetCursorDelta" ), Category="Utilities|PointerEvent")
+	static FVector2D PointerEvent_GetCursorDelta(const FPointerEvent& Input);
+
+	/** Mouse buttons that are currently pressed */
+	UFUNCTION(BlueprintPure, meta=( FriendlyName = "IsMouseButtonDown" ), Category="Utilities|PointerEvent")
+		static bool PointerEvent_IsMouseButtonDown(const FPointerEvent& Input, FKey MouseButton);
+
+	/** Mouse button that caused this event to be raised (possibly EB_None) */
+	UFUNCTION(BlueprintPure, meta=( FriendlyName = "GetEffectingButton" ), Category="Utilities|PointerEvent")
+	static FKey PointerEvent_GetEffectingButton(const FPointerEvent& Input);
+
+	/** How much did the mouse wheel turn since the last mouse event */
+	UFUNCTION(BlueprintPure, meta=( FriendlyName = "GetWheelDelta" ), Category="Utilities|PointerEvent")
+	static float PointerEvent_GetWheelDelta(const FPointerEvent& Input);
+
+	/** @return The index of the user that caused the event */
+	UFUNCTION(BlueprintPure, meta=( FriendlyName = "GetUserIndex" ), Category="Utilities|PointerEvent")
+	static int32 PointerEvent_GetUserIndex(const FPointerEvent& Input);
+
+	/** @return The unique identifier of the pointer (e.g., finger index) */
+	UFUNCTION(BlueprintPure, meta=( FriendlyName = "GetPointerIndex" ), Category="Utilities|PointerEvent")
+	static uint32 PointerEvent_GetPointerIndex(const FPointerEvent& Input);
+
+	/** @return The index of the touch pad that generated this event (for platforms with multiple touch pads per user) */
+	UFUNCTION(BlueprintPure, meta=( FriendlyName = "GetTouchpadIndex" ), Category="Utilities|PointerEvent")
+	static uint32 PointerEvent_GetTouchpadIndex(const FPointerEvent& Input);
+
+	/** @return Is this event a result from a touch (as opposed to a mouse) */
+	UFUNCTION(BlueprintPure, meta=( FriendlyName = "IsTouchEvent" ), Category="Utilities|PointerEvent")
+	static uint32 PointerEvent_IsTouchEvent(const FPointerEvent& Input);
+
+	//TODO UMG Support GetGestureType()
+
+	///** @return The type of touch gesture */
+	//UFUNCTION(BlueprintPure, meta=( FriendlyName = "GetGestureType" ), Category="Utilities|PointerEvent")
+	//static EGestureEvent::Type PointerEvent_GetGestureType(const FPointerEvent& Input);
+
+	/** @return The change in gesture value since the last gesture event of the same type. */
+	UFUNCTION(BlueprintPure, meta=( FriendlyName = "GetGestureDelta" ), Category="Utilities|PointerEvent")
+	static FVector2D PointerEvent_GetGestureDelta(const FPointerEvent& Input);
+
+
+	/** @return The controller button that caused this event */
+	UFUNCTION(BlueprintPure, meta=( FriendlyName = "GetEffectingButton" ), Category="Utilities|ControllerEvent")
+	static FKey ControllerEvent_GetEffectingButton(const FControllerEvent& Input);
+
+	/** @return The index of the user that caused the event */
+	UFUNCTION(BlueprintPure, meta=( FriendlyName = "GetUserIndex" ), Category="Utilities|ControllerEvent")
+	static uint32 ControllerEvent_GetUserIndex(const FControllerEvent& Input);
+
+	/** @return Analog value between 0 and 1.  1 being fully pressed, 0 being not pressed at all */
+	UFUNCTION(BlueprintPure, meta=( FriendlyName = "GetAnalogValue" ), Category="Utilities|ControllerEvent")
+	static float ControllerEvent_GetAnalogValue(const FControllerEvent& Input);
 };
diff --git a/Engine/Source/Runtime/Engine/Private/KismetInputLibrary.cpp b/Engine/Source/Runtime/Engine/Private/KismetInputLibrary.cpp
--- a/Engine/Source/Runtime/Engine/Private/KismetInputLibrary.cpp
+++ b/Engine/Source/Runtime/Engine/Private/KismetInputLibrary.cpp
@@ -21,10 +21,6 @@ void UKismetInputLibrary::CalibrateTilt()
 	GEngine->Exec(NULL, TEXT("CALIBRATEMOTION"));
 }
 
-FKey UKismetInputLibrary::GetKey(const FKeyboardEvent& Input)
-{
-	return Input.GetKey();
-}
 
 bool UKismetInputLibrary::EqualEqual_KeyKey(FKey A, FKey B)
 {
@@ -96,4 +92,87 @@ bool UKismetInputLibrary::InputEvent_IsRightCommandDown(const FInputEvent& Input
 	return Input.IsRightCommandDown();
 }
 
+
+FKey UKismetInputLibrary::GetKey(const FKeyboardEvent& Input)
+{
+	return Input.GetKey();
+}
+
+
+FVector2D UKismetInputLibrary::PointerEvent_GetScreenSpacePosition(const FPointerEvent& Input)
+{
+	return Input.GetScreenSpacePosition();
+}
+
+FVector2D UKismetInputLibrary::PointerEvent_GetLastScreenSpacePosition(const FPointerEvent& Input)
+{
+	return Input.GetLastScreenSpacePosition();
+}
+
+FVector2D UKismetInputLibrary::PointerEvent_GetCursorDelta(const FPointerEvent& Input)
+{
+	return Input.GetCursorDelta();
+}
+
+bool UKismetInputLibrary::PointerEvent_IsMouseButtonDown(const FPointerEvent& Input, FKey MouseButton)
+{
+	return Input.IsMouseButtonDown(MouseButton);
+}
+
+FKey UKismetInputLibrary::PointerEvent_GetEffectingButton(const FPointerEvent& Input)
+{
+	return Input.GetEffectingButton();
+}
+
+float UKismetInputLibrary::PointerEvent_GetWheelDelta(const FPointerEvent& Input)
+{
+	return Input.GetWheelDelta();
+}
+
+int32 UKismetInputLibrary::PointerEvent_GetUserIndex(const FPointerEvent& Input)
+{
+	return Input.GetUserIndex();
+}
+
+uint32 UKismetInputLibrary::PointerEvent_GetPointerIndex(const FPointerEvent& Input)
+{
+	return Input.GetPointerIndex();
+}
+
+uint32 UKismetInputLibrary::PointerEvent_GetTouchpadIndex(const FPointerEvent& Input)
+{
+	return Input.GetTouchpadIndex();
+}
+
+uint32 UKismetInputLibrary::PointerEvent_IsTouchEvent(const FPointerEvent& Input)
+{
+	return Input.IsTouchEvent();
+}
+
+//EGestureEvent::Type UKismetInputLibrary::PointerEvent_GetGestureType(const FPointerEvent& Input)
+//{
+//	return Input.GetGestureType();
+//}
+
+FVector2D UKismetInputLibrary::PointerEvent_GetGestureDelta(const FPointerEvent& Input)
+{
+	return Input.GetGestureDelta();
+}
+
+
+FKey UKismetInputLibrary::ControllerEvent_GetEffectingButton(const FControllerEvent& Input)
+{
+	return Input.GetEffectingButton();
+}
+
+uint32 UKismetInputLibrary::ControllerEvent_GetUserIndex(const FControllerEvent& Input)
+{
+	return Input.GetUserIndex();
+}
+
+float UKismetInputLibrary::ControllerEvent_GetAnalogValue(const FControllerEvent& Input)
+{
+	return Input.GetAnalogValue();
+}
+
 #undef LOCTEXT_NAMESPACE
\ No newline at end of file
diff --git a/Engine/Source/Runtime/UMG/Private/UserWidget.cpp b/Engine/Source/Runtime/UMG/Private/UserWidget.cpp
--- a/Engine/Source/Runtime/UMG/Private/UserWidget.cpp
+++ b/Engine/Source/Runtime/UMG/Private/UserWidget.cpp
@@ -71,79 +71,37 @@ class SViewportWidgetHost : public SCompoundWidget
 
 	virtual FReply OnKeyDown(const FGeometry& MyGeometry, const FKeyboardEvent& InKeyboardEvent) override
 	{
-		FReply Result = UserWidget->OnKeyDown(MyGeometry, InKeyboardEvent);
-		if ( bModal && !Result.IsEventHandled() )
-		{
-			return FReply::Handled();
-		}
-
-		return Result;
+		return bModal ? FReply::Handled() : FReply::Unhandled();
 	}
 	
 	virtual FReply OnKeyUp(const FGeometry& MyGeometry, const FKeyboardEvent& InKeyboardEvent) override
 	{
-		FReply Result = UserWidget->OnKeyUp(MyGeometry, InKeyboardEvent);
-		if ( bModal && !Result.IsEventHandled() )
-		{
-			return FReply::Handled();
-		}
-
-		return Result;
+		return bModal ? FReply::Handled() : FReply::Unhandled();
 	}
 
 	virtual FReply OnMouseButtonDown(const FGeometry& MyGeometry, const FPointerEvent& MouseEvent) override
 	{
-		FReply Result = UserWidget->OnMouseButtonDown(MyGeometry, MouseEvent);
-		if ( bModal && !Result.IsEventHandled() )
-		{
-			return FReply::Handled().CaptureMouse(AsShared());
-		}
-
-		return Result;
+		return bModal ? FReply::Handled().CaptureMouse(AsShared()) : FReply::Unhandled();
 	}
 
 	virtual FReply OnMouseButtonUp(const FGeometry& MyGeometry, const FPointerEvent& MouseEvent) override
 	{
-		FReply Result = UserWidget->OnMouseButtonUp(MyGeometry, MouseEvent);
-		if ( bModal && !Result.IsEventHandled() )
-		{
-			return FReply::Handled().ReleaseMouseCapture();
-		}
-
-		return Result;
+		return bModal ? FReply::Handled().ReleaseMouseCapture() : FReply::Unhandled();
 	}
 
 	virtual FReply OnMouseButtonDoubleClick(const FGeometry& MyGeometry, const FPointerEvent& MouseEvent) override
 	{
-		FReply Result = UserWidget->OnMouseButtonDoubleClick(MyGeometry, MouseEvent);
-		if ( bModal && !Result.IsEventHandled() )
-		{
-			return FReply::Handled();
-		}
-
-		return Result;
+		return bModal ? FReply::Handled() : FReply::Unhandled();
 	}
 
 	virtual FReply OnMouseMove(const FGeometry& MyGeometry, const FPointerEvent& MouseEvent) override
 	{
-		FReply Result = UserWidget->OnMouseMove(MyGeometry, MouseEvent);
-		if ( bModal && !Result.IsEventHandled() )
-		{
-			return FReply::Handled();
-		}
-
-		return Result;
+		return bModal ? FReply::Handled() : FReply::Unhandled();
 	}
 
 	virtual FReply OnMouseWheel(const FGeometry& MyGeometry, const FPointerEvent& MouseEvent) override
 	{
-		FReply Result = UserWidget->OnMouseWheel(MyGeometry, MouseEvent);
-		if ( bModal && !Result.IsEventHandled() )
-		{
-			return FReply::Handled();
-		}
-
-		return Result;
+		return bModal ? FReply::Handled() : FReply::Unhandled();
 	}
 
 	bool IsModal() const { return bModal; }
