FViewport can now convert desktop coordinates to normalized local coordinates and vice versa.

#github http://github.com/EpicGames/UnrealEngine/pull/265

diff --git a/Engine/Source/Runtime/Engine/Private/Slate/SceneViewport.cpp b/Engine/Source/Runtime/Engine/Private/Slate/SceneViewport.cpp
--- a/Engine/Source/Runtime/Engine/Private/Slate/SceneViewport.cpp
+++ b/Engine/Source/Runtime/Engine/Private/Slate/SceneViewport.cpp
@@ -239,6 +239,25 @@ void FSceneViewport::ProcessInput( float DeltaTime )
 
 }
 
+FVector2D FSceneViewport::VirtualDesktopPixelToViewport(FIntPoint VirtualDesktopPointPx) const
+{
+	// Virtual Desktop Pixel to local slate unit
+	const FVector2D TransformedPoint = CachedGeometry.AbsoluteToLocal(FVector2D(VirtualDesktopPointPx.X, VirtualDesktopPointPx.Y));
+
+	// Pixels to normalized coordinates
+	return FVector2D( TransformedPoint.X / SizeX, TransformedPoint.Y / SizeY );
+}
+
+FIntPoint FSceneViewport::ViewportToVirtualDesktopPixel(FVector2D ViewportCoordinate) const
+{
+	// Normalized to pixels transform
+	const FVector2D LocalCoordinateInSu = FVector2D( ViewportCoordinate.X * SizeX, ViewportCoordinate.Y * SizeY );
+	// Local slate unit to virtual desktop pixel.
+	const FVector2D TransformedPoint = FVector2D( CachedGeometry.LocalToAbsolute( LocalCoordinateInSu ) );
+
+	return FIntPoint( FMath::TruncToInt(TransformedPoint.X), FMath::TruncToInt(TransformedPoint.Y) );
+}
+
 void FSceneViewport::OnDrawViewport( const FGeometry& AllottedGeometry, const FSlateRect& MyClippingRect, FSlateWindowElementList& OutDrawElements, int32 LayerId, const FWidgetStyle& InWidgetStyle, bool bParentEnabled )
 {
 	// Switch to the viewport clients world before resizing
diff --git a/Engine/Source/Runtime/Engine/Private/UnrealClient.cpp b/Engine/Source/Runtime/Engine/Private/UnrealClient.cpp
--- a/Engine/Source/Runtime/Engine/Private/UnrealClient.cpp
+++ b/Engine/Source/Runtime/Engine/Private/UnrealClient.cpp
@@ -716,6 +716,8 @@ public:
 	virtual void	GetMousePos( FIntPoint& MousePosition, const bool bLocalPosition = true) { MousePosition = FIntPoint(0, 0); }
 	virtual void	SetMouse(int32 x, int32 y) { }
 	virtual void	ProcessInput( float DeltaTime ) { }
+	virtual FVector2D VirtualDesktopPixelToViewport(FIntPoint VirtualDesktopPointPx) const override { return FVector2D::ZeroVector; }
+	virtual FIntPoint ViewportToVirtualDesktopPixel(FVector2D ViewportCoordinate) const override { return FIntPoint::ZeroValue; }
 	virtual void InvalidateDisplay() { }
 	virtual void DeferInvalidateHitProxy() { }
 	virtual FViewportFrame* GetViewportFrame() { return 0; }
diff --git a/Engine/Source/Runtime/Engine/Public/Slate/SceneViewport.h b/Engine/Source/Runtime/Engine/Public/Slate/SceneViewport.h
--- a/Engine/Source/Runtime/Engine/Public/Slate/SceneViewport.h
+++ b/Engine/Source/Runtime/Engine/Public/Slate/SceneViewport.h
@@ -83,6 +83,9 @@ public:
 	 */
 	virtual void ProcessInput( float DeltaTime ) override;
 
+	virtual FVector2D VirtualDesktopPixelToViewport(FIntPoint VirtualDesktopPointPx) const override;
+	virtual FIntPoint ViewportToVirtualDesktopPixel(FVector2D ViewportCoordinate) const override;
+
 	/**
 	 * Called when the viewport should be invalidated and redrawn                   
 	 */
diff --git a/Engine/Source/Runtime/Engine/Public/UnrealClient.h b/Engine/Source/Runtime/Engine/Public/UnrealClient.h
--- a/Engine/Source/Runtime/Engine/Public/UnrealClient.h
+++ b/Engine/Source/Runtime/Engine/Public/UnrealClient.h
@@ -284,6 +284,24 @@ public:
 	virtual bool IsFullscreen()	const { return WindowMode == EWindowMode::Fullscreen || WindowMode == EWindowMode::WindowedFullscreen || WindowMode == EWindowMode::WindowedMirror; }
 	virtual void ProcessInput( float DeltaTime ) = 0;
 
+	/**
+	 * Transforms a virtual desktop pixel (the origin is in the primary screen's top left corner) to the local space of this viewport
+	 *
+	 * @param VirtualDesktopPointPx   Coordinate on the virtual desktop in pixel units. Desktop is considered virtual because multiple monitors are supported.
+	 *
+	 * @return The transformed pixel. It is normalized to the range [0, 1]
+	 */
+	virtual FVector2D VirtualDesktopPixelToViewport(FIntPoint VirtualDesktopPointPx) const = 0;
+
+	/**
+	 * Transforms a coordinate in the local space of this viewport into a virtual desktop pixel.
+	 *
+	 * @param ViewportCoordiante    Normalized coordniate in the rate [0..1]; (0,0) is upper left and (1,1) is lower right.
+	 *
+	 * @return the transformed coordinate. It is in virtual desktop pixels.
+	 */
+	virtual FIntPoint ViewportToVirtualDesktopPixel(FVector2D ViewportCoordinate) const = 0;
+
 	/**
 	 * @return A canvas that can be used while this viewport is being drawn to render debug elements on top of everything else
 	 */
