Add MaxDepenetrationVelocity override per body instance

diff --git a/Engine/Source/Runtime/Engine/Classes/PhysicsEngine/BodyInstance.h b/Engine/Source/Runtime/Engine/Classes/PhysicsEngine/BodyInstance.h
--- a/Engine/Source/Runtime/Engine/Classes/PhysicsEngine/BodyInstance.h
+++ b/Engine/Source/Runtime/Engine/Classes/PhysicsEngine/BodyInstance.h
@@ -211,6 +211,14 @@ protected:
 	UPROPERTY(EditAnywhere, AdvancedDisplay, BlueprintReadOnly, Category=Physics, meta=(editcondition="bOverrideWalkableSlopeOnInstance"))
 	struct FWalkableSlopeOverride WalkableSlopeOverride;
 
+	/** Whether this body instance has its own custom MaxDepenetrationVelocity*/
+	UPROPERTY(EditAnywhere, AdvancedDisplay, BlueprintReadOnly, Category = Physics)
+	uint32 bOverrideMaxDepenetrationVelocity : 1;
+
+	/** The maximum velocity used to depenetrate this object*/
+	UPROPERTY(EditAnywhere, AdvancedDisplay, BlueprintReadOnly, Category = Physics, meta = (editcondition = "bOverrideMaxDepenetrationVelocity", ClampMin = "0.0", UIMin = "0.0"))
+	float MaxDepenetrationVelocity;
+
 	/**	Allows you to override the PhysicalMaterial to use for simple collision on this body. */
 	UPROPERTY(EditAnywhere, AdvancedDisplay, BlueprintReadOnly, Category=Collision)
 	class UPhysicalMaterial* PhysMaterialOverride;
diff --git a/Engine/Source/Runtime/Engine/Private/PhysicsEngine/BodyInstance.cpp b/Engine/Source/Runtime/Engine/Private/PhysicsEngine/BodyInstance.cpp
--- a/Engine/Source/Runtime/Engine/Private/PhysicsEngine/BodyInstance.cpp
+++ b/Engine/Source/Runtime/Engine/Private/PhysicsEngine/BodyInstance.cpp
@@ -205,6 +205,8 @@ FBodyInstance::FBodyInstance()
 , CustomLockedAxis(FVector::ZeroVector)
 , DOFConstraint(NULL)
 , bOverrideWalkableSlopeOnInstance(false)
+, bOverrideMaxDepenetrationVelocity(false)
+, MaxDepenetrationVelocity(0.f)
 , PhysMaterialOverride(NULL)
 , COMNudge(ForceInit)
 , SleepFamily(SF_Normal)
@@ -1088,7 +1090,7 @@ void FBodyInstance::InitBody(UBodySetup* Setup, const FTransform& Transform, UPr
 
 			SetMaxAngularVelocity(MaxAngularVelocity, false);
 
-			SetMaxDepenetrationVelocity(UPhysicsSettings::Get()->MaxDepenetrationVelocity);
+			SetMaxDepenetrationVelocity(bOverrideMaxDepenetrationVelocity ? MaxDepenetrationVelocity : UPhysicsSettings::Get()->MaxDepenetrationVelocity);
 
 			//@TODO: BOX2D: Determine if sleep threshold and solver settings can be configured per-body or not
 #if 0
@@ -1285,7 +1287,7 @@ void FBodyInstance::InitBody(UBodySetup* Setup, const FTransform& Transform, UPr
 
 		SetMaxAngularVelocity(MaxAngularVelocity, false);
 
-		SetMaxDepenetrationVelocity(UPhysicsSettings::Get()->MaxDepenetrationVelocity);
+		SetMaxDepenetrationVelocity(bOverrideMaxDepenetrationVelocity ? MaxDepenetrationVelocity : UPhysicsSettings::Get()->MaxDepenetrationVelocity);
 
 		// Set initial velocity 
 		if(bUseSimulate)
@@ -2740,7 +2742,9 @@ void FBodyInstance::SetMaxAngularVelocity(float NewMaxAngVel, bool bAddToCurrent
 
 void FBodyInstance::SetMaxDepenetrationVelocity(float MaxVelocity)
 {
+	MaxDepenetrationVelocity = MaxVelocity;
 #if WITH_PHYSX
+	//0 from both instance and project settings means 
 	float UseMaxVelocity = MaxVelocity == 0.f ? PX_MAX_F32 : MaxVelocity;
 
 	if (PxRigidDynamic* PRigidDynamic = GetPxRigidDynamic())
@@ -2748,6 +2752,7 @@ void FBodyInstance::SetMaxDepenetrationVelocity(float MaxVelocity)
 		PRigidDynamic->setMaxDepenetrationVelocity(UseMaxVelocity);
 	}
 #endif
+
 }
 
 
