#ue4 - Made CharacterMovementComponent.bCrouchMovesCharacterDown VisibleInstanceOnly to make it less confusing, since it changes based on state.

diff --git a/Engine/Source/Runtime/Engine/Classes/GameFramework/CharacterMovementComponent.h b/Engine/Source/Runtime/Engine/Classes/GameFramework/CharacterMovementComponent.h
--- a/Engine/Source/Runtime/Engine/Classes/GameFramework/CharacterMovementComponent.h
+++ b/Engine/Source/Runtime/Engine/Classes/GameFramework/CharacterMovementComponent.h
@@ -527,7 +527,7 @@ public:
 	 * If true, crouching should lower the center of the player and keep the base of the capsule in place. If false, the base of the capsule moves up and the center stays in place.
 	 * By default, this is set to true when walking and false otherwise when the movement mode changes. Feel free to override the behavior when the movement mode changes.
 	 */
-	UPROPERTY(Category="Character Movement", VisibleAnywhere, BlueprintReadWrite, AdvancedDisplay)
+	UPROPERTY(Category="Character Movement", VisibleInstanceOnly, BlueprintReadWrite, AdvancedDisplay)
 	uint32 bCrouchMovesCharacterDown:1;
 
 	/** If true, the pawn ignores the effects of changes in its base's rotation on its rotation. */
