Adding audio setting to restrict the amount of streaming sounds that can play concurrently

Changed audio streaming manager add sound source function so that sources are ordered by priority and those that go over the stream limit are stopped
Added a function to FWaveInstance to check both whether streaming is supported and if it's using a streaming wave
Added a new function to the audio streaming manager to check whether a sound source should be created for a wave instance to avoid initialising sources that will not play
Set default amount of streams to 2 in the base engine ini
Added detail customization for AudioSettings so that maximum stream count setting is only shown to those using the experimental streaming switch
Fixed some unsafe sections of the Opus decompression code
Non-Unity build include fix

diff --git a/Engine/Config/BaseEngine.ini b/Engine/Config/BaseEngine.ini
--- a/Engine/Config/BaseEngine.ini
+++ b/Engine/Config/BaseEngine.ini
@@ -1388,6 +1388,7 @@ CommonAudioPoolSize=0
 [/Script/Engine.AudioSettings]
 DefaultSoundClassName=/Engine/EngineSounds/Master.Master
 LowPassFilterResonance=0.9
+MaximumConcurrentStreams=2
 
 [/Script/Engine.SoundGroups]
 +SoundGroupProfiles=(SoundGroup=SOUNDGROUP_Default, bAlwaysDecompressOnLoad=false, DecompressedDuration=5)
diff --git a/Engine/Source/Editor/DetailCustomizations/Private/AudioSettingsDetails.cpp b/Engine/Source/Editor/DetailCustomizations/Private/AudioSettingsDetails.cpp
new file mode 100644
--- /dev/null
+++ b/Engine/Source/Editor/DetailCustomizations/Private/AudioSettingsDetails.cpp
@@ -0,0 +1,17 @@
+// Copyright 1998-2014 Epic Games, Inc. All Rights Reserved.
+
+#include "DetailCustomizationsPrivatePCH.h"
+#include "AudioSettingsDetails.h"
+
+TSharedRef<IDetailCustomization> FAudioSettingsDetails::MakeInstance()
+{
+	return MakeShareable(new FAudioSettingsDetails);
+}
+
+void FAudioSettingsDetails::CustomizeDetails(IDetailLayoutBuilder& DetailBuilder)
+{
+	if (!GetDefault<UEditorExperimentalSettings>()->bShowAudioStreamingOptions)
+	{
+		DetailBuilder.HideProperty(GET_MEMBER_NAME_CHECKED(UAudioSettings, MaximumConcurrentStreams));
+	}
+}
diff --git a/Engine/Source/Editor/DetailCustomizations/Private/AudioSettingsDetails.h b/Engine/Source/Editor/DetailCustomizations/Private/AudioSettingsDetails.h
new file mode 100644
--- /dev/null
+++ b/Engine/Source/Editor/DetailCustomizations/Private/AudioSettingsDetails.h
@@ -0,0 +1,14 @@
+// Copyright 1998-2014 Epic Games, Inc. All Rights Reserved.
+
+#pragma once
+
+class FAudioSettingsDetails : public IDetailCustomization
+{
+public:
+	/** Makes a new instance of this detail layout class for a specific detail view requesting it */
+	static TSharedRef<IDetailCustomization> MakeInstance();
+
+private:
+	/** IDetailCustomization interface */
+	virtual void CustomizeDetails(IDetailLayoutBuilder& DetailBuilder) override;
+};
\ No newline at end of file
diff --git a/Engine/Source/Editor/DetailCustomizations/Private/DetailCustomizations.cpp b/Engine/Source/Editor/DetailCustomizations/Private/DetailCustomizations.cpp
--- a/Engine/Source/Editor/DetailCustomizations/Private/DetailCustomizations.cpp
+++ b/Engine/Source/Editor/DetailCustomizations/Private/DetailCustomizations.cpp
@@ -80,6 +80,7 @@
 #include "MediaSoundWaveCustomization.h"
 #include "MediaTextureCustomization.h"
 #include "SoundWaveDetails.h"
+#include "AudioSettingsDetails.h"
 
 IMPLEMENT_MODULE( FDetailCustomizationsModule, DetailCustomizations );
 
@@ -221,6 +222,7 @@ void FDetailCustomizationsModule::RegisterObjectCustomizations()
 	RegisterCustomClassLayout("PhysicsConstraintComponent", FOnGetDetailCustomizationInstance::CreateStatic(&FPhysicsConstraintComponentDetails::MakeInstance));
 	RegisterCustomClassLayout("CollisionProfile", FOnGetDetailCustomizationInstance::CreateStatic(&FCollisionProfileDetails::MakeInstance));
 	RegisterCustomClassLayout("PhysicsSettings", FOnGetDetailCustomizationInstance::CreateStatic(&FPhysicsSettingsDetails::MakeInstance));
+	RegisterCustomClassLayout("AudioSettings", FOnGetDetailCustomizationInstance::CreateStatic(&FAudioSettingsDetails::MakeInstance));
 
 	RegisterCustomClassLayout("ParticleModuleRequired", FOnGetDetailCustomizationInstance::CreateStatic(&FParticleModuleRequiredDetails::MakeInstance));
 	RegisterCustomClassLayout("ParticleModuleSubUV", FOnGetDetailCustomizationInstance::CreateStatic(&FParticleModuleSubUVDetails::MakeInstance));
diff --git a/Engine/Source/Editor/DetailCustomizations/Private/SoundWaveDetails.cpp b/Engine/Source/Editor/DetailCustomizations/Private/SoundWaveDetails.cpp
--- a/Engine/Source/Editor/DetailCustomizations/Private/SoundWaveDetails.cpp
+++ b/Engine/Source/Editor/DetailCustomizations/Private/SoundWaveDetails.cpp
@@ -2,9 +2,6 @@
 
 #include "DetailCustomizationsPrivatePCH.h"
 #include "SoundWaveDetails.h"
-//#include "EditorExperimentalSettings.h"
-
-#define LOCTEXT_NAMESPACE "AmbientSoundDetails"
 
 TSharedRef<IDetailCustomization> FSoundWaveDetails::MakeInstance()
 {
@@ -19,5 +16,3 @@ void FSoundWaveDetails::CustomizeDetails(IDetailLayoutBuilder& DetailBuilder)
 		DetailBuilder.HideProperty(GET_MEMBER_NAME_CHECKED(USoundWave, StreamingPriority));
 	}
 }
-
-#undef LOCTEXT_NAMESPACE
diff --git a/Engine/Source/Editor/GameProjectGeneration/Private/SNewProjectWizard.cpp b/Engine/Source/Editor/GameProjectGeneration/Private/SNewProjectWizard.cpp
--- a/Engine/Source/Editor/GameProjectGeneration/Private/SNewProjectWizard.cpp
+++ b/Engine/Source/Editor/GameProjectGeneration/Private/SNewProjectWizard.cpp
@@ -7,6 +7,7 @@
 #include "SourceCodeNavigation.h"
 #include "SScrollBorder.h"
 #include "TemplateCategory.h"
+#include "GameProjectGenerationModule.h"
 
 #define LOCTEXT_NAMESPACE "NewProjectWizard"
 
diff --git a/Engine/Source/Runtime/Engine/Classes/Sound/AudioSettings.h b/Engine/Source/Runtime/Engine/Classes/Sound/AudioSettings.h
--- a/Engine/Source/Runtime/Engine/Classes/Sound/AudioSettings.h
+++ b/Engine/Source/Runtime/Engine/Classes/Sound/AudioSettings.h
@@ -22,4 +22,8 @@ class ENGINE_API UAudioSettings
 	
 	UPROPERTY(config, EditAnywhere, Category="Audio", AdvancedDisplay, meta=(ClampMin=0.1,ClampMax=1.5))
 	float LowPassFilterResonance;
+
+	/** How many streaming sounds can be played at the same time (if more are played they will be sorted by priority) */
+	UPROPERTY(config, EditAnywhere, Category="Audio", meta=(ClampMin=0))
+	int32 MaximumConcurrentStreams;
 };
diff --git a/Engine/Source/Runtime/Engine/Private/Audio.cpp b/Engine/Source/Runtime/Engine/Private/Audio.cpp
--- a/Engine/Source/Runtime/Engine/Private/Audio.cpp
+++ b/Engine/Source/Runtime/Engine/Private/Audio.cpp
@@ -107,9 +107,9 @@ FName FSoundBuffer::GetSoundClassName()
 			USoundWave* Wave = Cast<USoundWave>(*It);
 			if (Wave && Wave->ResourceID == ResourceID)
 			{
-				if (Cue && Cue->SoundClassObject)
+				if (Wave->SoundClassObject)
 				{
-					return Cue->SoundClassObject->GetFName();
+					return Wave->SoundClassObject->GetFName();
 				}
 				else
 				{
@@ -443,6 +443,11 @@ float FWaveInstance::GetActualVolume() const
 	return Volume * VolumeMultiplier;
 }
 
+bool FWaveInstance::IsStreaming() const
+{
+	return FPlatformProperties::SupportsAudioStreaming() && WaveData != nullptr && WaveData->IsStreaming();
+}
+
 /*-----------------------------------------------------------------------------
 	WaveModInfo implementation - downsampling of wave files.
 -----------------------------------------------------------------------------*/
diff --git a/Engine/Source/Runtime/Engine/Private/AudioDevice.cpp b/Engine/Source/Runtime/Engine/Private/AudioDevice.cpp
--- a/Engine/Source/Runtime/Engine/Private/AudioDevice.cpp
+++ b/Engine/Source/Runtime/Engine/Private/AudioDevice.cpp
@@ -1773,7 +1773,9 @@ void FAudioDevice::StartSources( TArray<FWaveInstance*>& WaveInstances, int32 Fi
 		if(	bGameTicking || WaveInstance->bIsUISound )
 		{
 			FSoundSource* Source = WaveInstanceSourceMap.FindRef( WaveInstance );
-			if( !Source )
+			if( !Source &&
+			  ( !WaveInstance->IsStreaming() ||
+				IStreamingManager::Get().GetAudioStreamingManager().CanCreateSoundSource(WaveInstance) ) )
 			{
 				check( FreeSources.Num() );
 				Source = FreeSources.Pop();
@@ -1798,10 +1800,16 @@ void FAudioDevice::StartSources( TArray<FWaveInstance*>& WaveInstances, int32 Fi
 					FreeSources.Add( Source );
 				}
 			}
-			else
+			else if (Source)
 			{
 				Source->Update();
 			}
+			else
+			{
+				// This can happen if the streaming manager determines that this sound should not be started.
+				// We stop the wave instance to prevent it from attempting to initialize every frame
+				WaveInstance->StopWithoutNotification();
+			}
 		}
 	}
 }
diff --git a/Engine/Source/Runtime/Engine/Private/AudioStreaming.cpp b/Engine/Source/Runtime/Engine/Private/AudioStreaming.cpp
--- a/Engine/Source/Runtime/Engine/Private/AudioStreaming.cpp
+++ b/Engine/Source/Runtime/Engine/Private/AudioStreaming.cpp
@@ -8,6 +8,7 @@ AudioStreaming.cpp: Implementation of audio streaming classes.
 #include "AudioStreaming.h"
 #include "Sound/SoundWave.h"
 #include "DerivedDataCacheInterface.h"
+#include "Sound/AudioSettings.h"
 
 /*------------------------------------------------------------------------------
 	Streaming chunks from the derived data cache.
@@ -476,15 +477,69 @@ bool FAudioStreamingManager::IsStreamingInProgress(const USoundWave* SoundWave)
 	return false;
 }
 
+bool FAudioStreamingManager::CanCreateSoundSource(const FWaveInstance* WaveInstance) const
+{
+	if (WaveInstance && WaveInstance->IsStreaming())
+	{
+		int32 MaxStreams = GetDefault<UAudioSettings>()->MaximumConcurrentStreams;
+
+		if ( StreamingSoundSources.Num() < MaxStreams )
+		{
+			return true;
+		}
+		else
+		{
+			for (int32 Index = 0; Index < StreamingSoundSources.Num(); ++Index)
+			{
+				const FSoundSource* ExistingSource = StreamingSoundSources[Index];
+				const FWaveInstance* ExistingWaveInst = ExistingSource->GetWaveInstance();
+				if (!ExistingWaveInst || !ExistingWaveInst->WaveData
+					|| ExistingWaveInst->WaveData->StreamingPriority < WaveInstance->WaveData->StreamingPriority)
+				{
+					return Index < MaxStreams;
+				}
+			}
+
+			return false;
+		}
+	}
+	return true;
+}
+
 void FAudioStreamingManager::AddStreamingSoundSource(FSoundSource* SoundSource)
 {
-	if (FPlatformProperties::SupportsAudioStreaming())
+	const FWaveInstance* WaveInstance = SoundSource->GetWaveInstance();
+	if (WaveInstance && WaveInstance->IsStreaming())
 	{
-		const FWaveInstance* WaveInstance = SoundSource->GetWaveInstance();
-		if (WaveInstance && WaveInstance->WaveData && WaveInstance->WaveData->IsStreaming())
+		int32 MaxStreams = GetDefault<UAudioSettings>()->MaximumConcurrentStreams;
+
+		// Add source sorted by priority so we can easily iterate over the amount of streams
+		// that are allowed
+		int32 OrderedIndex = -1;
+		for (int32 Index = 0; Index < StreamingSoundSources.Num() && Index < MaxStreams; ++Index)
+		{
+			const FSoundSource* ExistingSource = StreamingSoundSources[Index];
+			const FWaveInstance* ExistingWaveInst = ExistingSource->GetWaveInstance();
+			if (!ExistingWaveInst || !ExistingWaveInst->WaveData
+				|| ExistingWaveInst->WaveData->StreamingPriority < WaveInstance->WaveData->StreamingPriority)
+			{
+				OrderedIndex = Index;
+				break;
+			}
+		}
+		if (OrderedIndex != -1)
+		{
+			StreamingSoundSources.Insert(SoundSource, OrderedIndex);
+		}
+		else if (StreamingSoundSources.Num() < MaxStreams)
 		{
 			StreamingSoundSources.AddUnique(SoundSource);
 		}
+
+		for (int32 Index = StreamingSoundSources.Num()-1; Index >= MaxStreams; --Index)
+		{
+			StreamingSoundSources[Index]->Stop();
+		}
 	}
 }
 
diff --git a/Engine/Source/Runtime/Engine/Private/OpusAudioInfo.cpp b/Engine/Source/Runtime/Engine/Private/OpusAudioInfo.cpp
--- a/Engine/Source/Runtime/Engine/Private/OpusAudioInfo.cpp
+++ b/Engine/Source/Runtime/Engine/Private/OpusAudioInfo.cpp
@@ -211,6 +211,7 @@ bool FOpusAudioInfo::ReadCompressedData(uint8* Destination, bool bLooping, uint3
 		if (DecodedSamples < 0)
 		{
 			LastPCMByteSize = 0;
+			ZeroBuffer(Destination + RawPCMOffset, BufferSize - RawPCMOffset);
 			return false;
 		}
 		else
@@ -339,6 +340,7 @@ bool FOpusAudioInfo::StreamCompressedData(uint8* Destination, bool bLooping, uin
 		if (DecodedSamples < 0)
 		{
 			LastPCMByteSize = 0;
+			ZeroBuffer(Destination + RawPCMOffset, BufferSize - RawPCMOffset);
 			return false;
 		}
 		else
@@ -399,7 +401,11 @@ bool FOpusAudioInfo::StreamCompressedData(uint8* Destination, bool bLooping, uin
 
 int32 FOpusAudioInfo::DecompressToPCMBuffer(uint16 FrameSize)
 {
-	check(SrcBufferOffset + FrameSize <= SrcBufferDataSize);
+	if (SrcBufferOffset + FrameSize > SrcBufferDataSize)
+	{
+		// if frame size is too large, something has gone wrong
+		return -1;
+	}
 
 	const uint8* SrcPtr = SrcBufferData + SrcBufferOffset;
 	SrcBufferOffset += FrameSize;
diff --git a/Engine/Source/Runtime/Engine/Public/Audio.h b/Engine/Source/Runtime/Engine/Public/Audio.h
--- a/Engine/Source/Runtime/Engine/Public/Audio.h
+++ b/Engine/Source/Runtime/Engine/Public/Audio.h
@@ -262,6 +262,11 @@ struct ENGINE_API FWaveInstance
 
 	/** Returns the actual volume the wave instance will play at */
 	float GetActualVolume() const;
+
+	/**
+	 * Checks whether wave is streaming and streaming is supported
+	 */
+	bool IsStreaming() const;
 };
 
 inline uint32 GetTypeHash( FWaveInstance* A ) { return A->TypeHash; }
diff --git a/Engine/Source/Runtime/Engine/Public/AudioStreaming.h b/Engine/Source/Runtime/Engine/Public/AudioStreaming.h
--- a/Engine/Source/Runtime/Engine/Public/AudioStreaming.h
+++ b/Engine/Source/Runtime/Engine/Public/AudioStreaming.h
@@ -203,6 +203,7 @@ struct FAudioStreamingManager : public IAudioStreamingManager
 	virtual void RemoveStreamingSoundWave(USoundWave* SoundWave) override;
 	virtual bool IsManagedStreamingSoundWave(const USoundWave* SoundWave) const override;
 	virtual bool IsStreamingInProgress(const USoundWave* SoundWave) override;
+	virtual bool CanCreateSoundSource(const FWaveInstance* WaveInstance) const override;
 	virtual void AddStreamingSoundSource(FSoundSource* SoundSource) override;
 	virtual void RemoveStreamingSoundSource(FSoundSource* SoundSource) override;
 	virtual bool IsManagedStreamingSoundSource(const FSoundSource* SoundSource) const override;
diff --git a/Engine/Source/Runtime/Engine/Public/ContentStreaming.h b/Engine/Source/Runtime/Engine/Public/ContentStreaming.h
--- a/Engine/Source/Runtime/Engine/Public/ContentStreaming.h
+++ b/Engine/Source/Runtime/Engine/Public/ContentStreaming.h
@@ -492,6 +492,8 @@ struct IAudioStreamingManager : public IStreamingManager
 	/** Returns true if this Sound Wave is currently streaming a chunk. */
 	virtual bool IsStreamingInProgress(const USoundWave* SoundWave) = 0;
 
+	virtual bool CanCreateSoundSource(const FWaveInstance* WaveInstance) const = 0;
+
 	/** Adds a new Sound Source to the streaming manager. */
 	virtual void AddStreamingSoundSource(FSoundSource* SoundSource) = 0;
 
