Fix for GetCompletionHandle returning NULL (https://udn.unrealengine.com/questions/217350/crash-found-in-uskeletalmeshcomponentrefreshbonetr.html)

diff --git a/Engine/Source/Runtime/Engine/Private/TickTaskManager.cpp b/Engine/Source/Runtime/Engine/Private/TickTaskManager.cpp
--- a/Engine/Source/Runtime/Engine/Private/TickTaskManager.cpp
+++ b/Engine/Source/Runtime/Engine/Private/TickTaskManager.cpp
@@ -313,6 +313,7 @@ private:
 				UE_LOG(LogTick, Log, TEXT("tick %6d %2d %s"),GFrameCounter, (int32)CurrentThread, *Target->DiagnosticMessage());
 			}
 			Target->ExecuteTick(Context.DeltaSeconds, Context.TickType, CurrentThread, MyCompletionGraphEvent);
+			Target->CompletionHandle = NULL; // Allow the old completion handle to be recycled
 		}
 	};
 };
@@ -401,6 +402,7 @@ public:
 				TickFunction->TickVisitedGFrameCounter = GFrameCounter;
 				TickFunction->TickQueuedGFrameCounter = GFrameCounter;
 				TickFunction->ExecuteTick(InContext.DeltaSeconds, InContext.TickType, ENamedThreads::GameThread, FGraphEventRef());
+				TickFunction->CompletionHandle = NULL; // Allow the old completion handle to be recycled
 			}
 		}
 		check(!NewlySpawnedTickFunctions.Num()); // We don't support new spawns during pause ticks
@@ -986,7 +988,6 @@ void FTickFunction::UnRegisterTickFunction()
 	{
 		FTickTaskManager::Get().RemoveTickFunction(this);
 		bRegistered = false;
-		CompletionHandle = NULL; // allow the old completion handle to be recycled
 	}
 }
 
@@ -999,10 +1000,6 @@ void FTickFunction::SetTickFunctionEnable(bool bInEnabled)
 		TickTaskLevel->RemoveTickFunction(this);
 		bTickEnabled = bInEnabled;
 		TickTaskLevel->AddTickFunction(this);
-		if (!bTickEnabled)
-		{
-			CompletionHandle = NULL; // allow the old completion handle to be recycled
-		}
 	}
 	else
 	{
@@ -1045,7 +1042,6 @@ void FTickFunction::QueueTickFunction(const struct FTickContext& TickContext)
 	if (TickVisitedGFrameCounter != GFrameCounter)
 	{
 		TickVisitedGFrameCounter = GFrameCounter;
-		CompletionHandle = NULL; // allow the old completion handle to be recycled
 		if (bTickEnabled && (!EnableParent || EnableParent->bTickEnabled))
 		{
 			ETickingGroup MaxPrerequisiteTickGroup =  ETickingGroup(0);
@@ -1119,7 +1115,6 @@ void FTickFunction::QueueTickFunctionParallel(const FTickContext& TickContext, T
 	if (bProcessTick)
 	{
 		check(bRegistered);
-		CompletionHandle = NULL; // allow the old completion handle to be recycled
 		if (bTickEnabled && (!EnableParent || EnableParent->bTickEnabled))
 		{
 			ETickingGroup MaxPrerequisiteTickGroup =  ETickingGroup(0);
