Move body welding into primitive component and fix initialization order issues.

diff --git a/Engine/Source/Runtime/Engine/Classes/Components/PrimitiveComponent.h b/Engine/Source/Runtime/Engine/Classes/Components/PrimitiveComponent.h
--- a/Engine/Source/Runtime/Engine/Classes/Components/PrimitiveComponent.h
+++ b/Engine/Source/Runtime/Engine/Classes/Components/PrimitiveComponent.h
@@ -1001,6 +1001,11 @@ public:
 	/** Called when a component is 'damaged', allowing for component class specific behaviour */
 	virtual void ReceiveComponentDamage(float DamageAmount, FDamageEvent const& DamageEvent, AController* EventInstigator, AActor* DamageCauser);
 
+#if WITH_BODY_WELDING
+	/** Used to physically weld objects at runtime. See decleration in USceneComponent for full description */
+	virtual bool WeldPhysicsBody(USceneComponent * ChildComponent, bool bWeld = true, FName ParentBoneName = NAME_None, FName ChildBoneName = NAME_None) override;
+#endif
+
 protected:
 
 	/** Give the static mesh component recreate render state context access to Create/DestroyRenderState_Concurrent(). */
diff --git a/Engine/Source/Runtime/Engine/Classes/Components/SceneComponent.h b/Engine/Source/Runtime/Engine/Classes/Components/SceneComponent.h
--- a/Engine/Source/Runtime/Engine/Classes/Components/SceneComponent.h
+++ b/Engine/Source/Runtime/Engine/Classes/Components/SceneComponent.h
@@ -708,6 +708,21 @@ public:
 	/** Is component (not including attachments) relevant for navigation updates? */
 	virtual bool IsNavigationRelevant(bool bSkipCollisionEnabledCheck = false) const;
 
+#if WITH_BODY_WELDING
+
+	/**
+	* Attempts to weld/unweld the child component to this component in a physically correct way. This re-calculates mass distribution and geometry to achieve correct physics.
+	* This function assumes ChildComponent has already been attached to this component and is now doing the physics fix-up.
+	* The function uses the attachment hierarchy to easily attach/detach multiple components together
+	* @param ChildComponent the component to be physically attached - assumed to already be attached in the scene hierarchy
+	* @param bWeld determines whether to weld or unweld. This flag is used for creating or breaking welded pieces at runtime
+	* @param ParentBoneName potentially used by skeletal mesh to weld child to specific bone on this component
+	* @param ChildBoneName potentially used by skeletal mesh to weld this component to specific bone in child
+	* @return whether or not the weld/unweld was succesfull
+	*/
+	virtual bool WeldPhysicsBody(USceneComponent * ChildComponent, bool bWeld = true, FName ParentBoneName = NAME_None, FName ChildBoneName = NAME_None) { return false;  }
+#endif
+
 protected:
 	/**
 	 * Called after a child scene component is attached to this component.
diff --git a/Engine/Source/Runtime/Engine/Private/PhysicsEngine/BodyInstance.cpp b/Engine/Source/Runtime/Engine/Private/PhysicsEngine/BodyInstance.cpp
--- a/Engine/Source/Runtime/Engine/Private/PhysicsEngine/BodyInstance.cpp
+++ b/Engine/Source/Runtime/Engine/Private/PhysicsEngine/BodyInstance.cpp
@@ -899,6 +899,26 @@ void FBodyInstance::InitBody(UBodySetup* Setup, const FTransform& Transform, UPr
 	// whenever update filter, this check will trigger, it seems it's best to know when we initialize body instead. 
 	check(BodySetup.IsValid());
 
+#if WITH_BODY_WELDING
+	{
+		UPrimitiveComponent * ParentPrimComponent = PrimComp ? Cast<UPrimitiveComponent>(PrimComp->AttachParent) : NULL;
+		
+		//if we have a parent we will now do the weld and exit any further initialization
+		if (ParentPrimComponent && PrimComp->GetWorld() && PrimComp->GetWorld()->IsGameWorld())
+		{
+			/*FTransform RelativeTM = FTransform(PrimComp->RelativeRotation, PrimComp->RelativeLocation, PrimComp->RelativeScale3D);
+			FBodyInstance * OwnerBody = ParentPrimComponent->GetBodyInstance();
+			OwnerBody->Weld(this, RelativeTM);*/
+			if (ParentPrimComponent->WeldPhysicsBody(PrimComp))
+			{
+				//welded so stop any further initialization
+				return;
+			}
+			
+		}
+	}
+#endif
+
 	USkeletalMeshComponent* SkelMeshComp = Cast<USkeletalMeshComponent>(OwnerComponent.Get());
 	if (SkelMeshComp != nullptr)
 	{
diff --git a/Engine/Source/Runtime/Engine/Private/PrimitiveComponent.cpp b/Engine/Source/Runtime/Engine/Private/PrimitiveComponent.cpp
--- a/Engine/Source/Runtime/Engine/Private/PrimitiveComponent.cpp
+++ b/Engine/Source/Runtime/Engine/Private/PrimitiveComponent.cpp
@@ -327,6 +327,7 @@ void UPrimitiveComponent::OnAttachmentChanged()
 void UPrimitiveComponent::OnChildAttached(USceneComponent* ChildComponent)
 {
 	Super::OnChildAttached(ChildComponent);
+	/*
 #if WITH_BODY_WELDING
 	//We want to support attaching physically simulated rigid bodies to create a single rigid body.
 	if (UPrimitiveComponent * PrimitiveChild = Cast<UPrimitiveComponent>(ChildComponent))
@@ -343,8 +344,54 @@ void UPrimitiveComponent::OnChildAttached(USceneComponent* ChildComponent)
 		}
 	}
 #endif
+	*/
 }
 
+#if WITH_BODY_WELDING
+bool UPrimitiveComponent::WeldPhysicsBody(USceneComponent * ChildComponent, bool bWeld, FName ParentBoneName, FName ChildBoneName)
+{
+	//assumes that child is already attached
+	if (ChildComponent->AttachParent != this)
+	{
+		return false;
+	}
+
+	if (UPrimitiveComponent * ChildPrim = Cast<UPrimitiveComponent>(ChildComponent))
+	{
+		//traverse the hierarchy to weld/unweld from root
+		USceneComponent * RootComponent = this;
+		FTransform RelativeTM = FTransform(ChildComponent->RelativeRotation, ChildComponent->RelativeLocation, ChildComponent->RelativeScale3D);
+
+		for (; RootComponent->AttachParent; RootComponent = RootComponent->AttachParent)
+		{
+			RelativeTM = RelativeTM * FTransform(RootComponent->RelativeRotation, RootComponent->RelativeLocation, RootComponent->RelativeScale3D);
+		}
+
+		if (UPrimitiveComponent * RootPrim = Cast<UPrimitiveComponent>(RootComponent))
+		{
+			//we support welding of kinematic actors to simulated parents.
+			//welding simulated to anything doesn't make sense and we'd rather use a physical constraint
+			//welding kinematic to kinematic is not needed as the attachment hierarchy already handles moving the object around
+			if (RootPrim->IsSimulatingPhysics())
+			{
+				FBodyInstance * RootBody = RootPrim->GetBodyInstance(ParentBoneName);
+				FBodyInstance * ChildBody = ChildPrim->GetBodyInstance(ChildBoneName);
+
+				if (RootBody && ChildBody && bWeld)
+				{
+
+					RootBody->Weld(ChildBody, RelativeTM);
+					return true;
+				}
+			}
+		}
+	}
+
+	
+	return false;
+}
+#endif
+
 void UPrimitiveComponent::DestroyRenderState_Concurrent()
 {
 	// Remove the primitive from the scene.
