Make attenuation settings properties blueprint read write

diff --git a/Engine/Source/Runtime/Engine/Classes/Sound/SoundAttenuation.h b/Engine/Source/Runtime/Engine/Classes/Sound/SoundAttenuation.h
--- a/Engine/Source/Runtime/Engine/Classes/Sound/SoundAttenuation.h
+++ b/Engine/Source/Runtime/Engine/Classes/Sound/SoundAttenuation.h
@@ -46,34 +46,34 @@ struct ENGINE_API FAttenuationSettings
 	GENERATED_USTRUCT_BODY()
 
 	/* Enable attenuation via volume. */
-	UPROPERTY(EditAnywhere, Category=Attenuation )
+	UPROPERTY(EditAnywhere, BlueprintReadWrite, Category=Attenuation )
 	uint32 bAttenuate:1;
 
 	/* Enable the source to be positioned in 3D. */
-	UPROPERTY(EditAnywhere, Category=Attenuation )
+	UPROPERTY(EditAnywhere, BlueprintReadWrite, Category=Attenuation )
 	uint32 bSpatialize:1;
 
 	/* Enable attenuation via low pass filter. */
-	UPROPERTY(EditAnywhere, Category=LowPassFilter )
+	UPROPERTY(EditAnywhere, BlueprintReadWrite, Category=LowPassFilter )
 	uint32 bAttenuateWithLPF:1;
 
 	/* The type of volume versus distance algorithm to use for the attenuation model. */
-	UPROPERTY(EditAnywhere, Category=Attenuation )
+	UPROPERTY(EditAnywhere, BlueprintReadWrite, Category=Attenuation )
 	TEnumAsByte<enum ESoundDistanceModel> DistanceAlgorithm;
 
 	UPROPERTY()
 	TEnumAsByte<enum ESoundDistanceCalc> DistanceType_DEPRECATED;
 
 	/* The shape of the attenuation volume. */
-	UPROPERTY(EditAnywhere, Category=Attenuation )
+	UPROPERTY(EditAnywhere, BlueprintReadWrite, Category=Attenuation )
 	TEnumAsByte<enum EAttenuationShape::Type> AttenuationShape;
 
 	/* The volume at maximum distance in deciBels. */
-	UPROPERTY(EditAnywhere, Category=Attenuation, meta=(DisplayName = "dB Attenuation At Max", ClampMax = "0" ))
+	UPROPERTY(EditAnywhere, BlueprintReadWrite, Category=Attenuation, meta=(DisplayName = "dB Attenuation At Max", ClampMax = "0" ))
 	float dBAttenuationAtMax;
 
 	/** At what distance we start treating the sound source as spatialized */
-	UPROPERTY(EditAnywhere, Category=Attenuation, meta=(ClampMin = "0", EditCondition="bSpatialize", DisplayName="Non-Spatialized Radius"))
+	UPROPERTY(EditAnywhere, BlueprintReadWrite, Category=Attenuation, meta=(ClampMin = "0", EditCondition="bSpatialize", DisplayName="Non-Spatialized Radius"))
 	float OmniRadius;
 
 	UPROPERTY()
@@ -88,23 +88,23 @@ struct ENGINE_API FAttenuationSettings
 	   Box     - X, Y, and Z are the Box's dimensions
 	   Cone    - X is Cone Radius, Y is Cone Angle, Z is Cone Falloff Angle
 	*/
-	UPROPERTY(EditAnywhere, Category=Attenuation)
+	UPROPERTY(EditAnywhere, BlueprintReadWrite, Category=Attenuation)
 	FVector AttenuationShapeExtents;
 
 	/* The distance back from the sound's origin to begin the cone when using the cone attenuation shape. */
-	UPROPERTY(EditAnywhere, Category=Attenuation, meta=(ClampMin = "0"))
+	UPROPERTY(EditAnywhere, BlueprintReadWrite, Category=Attenuation, meta=(ClampMin = "0"))
 	float ConeOffset;
 
 	/* The distance over which falloff occurs. */
-	UPROPERTY(EditAnywhere, Category=Attenuation, meta=(ClampMin = "0"))
+	UPROPERTY(EditAnywhere, BlueprintReadWrite, Category=Attenuation, meta=(ClampMin = "0"))
 	float FalloffDistance;
 
 	/* The range at which to start applying a low pass filter. */
-	UPROPERTY(EditAnywhere, Category=LowPassFilter )
+	UPROPERTY(EditAnywhere, BlueprintReadWrite, Category=LowPassFilter )
 	float LPFRadiusMin;
 
 	/* The range at which to apply the maximum amount of low pass filter. */
-	UPROPERTY(EditAnywhere, Category=LowPassFilter )
+	UPROPERTY(EditAnywhere, BlueprintReadWrite, Category=LowPassFilter )
 	float LPFRadiusMax;
 
 	FAttenuationSettings()
