* minor optimization

diff --git a/Engine/Source/Runtime/Engine/Private/Materials/MaterialInstance.cpp b/Engine/Source/Runtime/Engine/Private/Materials/MaterialInstance.cpp
--- a/Engine/Source/Runtime/Engine/Private/Materials/MaterialInstance.cpp
+++ b/Engine/Source/Runtime/Engine/Private/Materials/MaterialInstance.cpp
@@ -639,11 +639,13 @@ bool UMaterialInstance::GetRefractionSettings(float& OutBiasValue) const
 	FName ParamName;
 	if( GetLinkerUE4Version() >= VER_UE4_REFRACTION_BIAS_TO_REFRACTION_DEPTH_BIAS )
 	{
-		ParamName = FName(TEXT("RefractionDepthBias"));
+		static FName NAME_RefractionDepthBias(TEXT("RefractionDepthBias"));
+		ParamName = NAME_RefractionDepthBias;
 	}
 	else
 	{
-		ParamName = FName(TEXT("RefractionBias"));
+		static FName NAME_RefractionBias(TEXT("RefractionBias"));
+		ParamName = NAME_RefractionBias;
 	}
 
 	const FScalarParameterValue* BiasParameterValue = GameThread_FindParameterByName(ScalarParameterValues, ParamName);
