Fix simple as complex flag when welding bodies

diff --git a/Engine/Source/Runtime/Engine/Classes/PhysicsEngine/BodyInstance.h b/Engine/Source/Runtime/Engine/Classes/PhysicsEngine/BodyInstance.h
--- a/Engine/Source/Runtime/Engine/Classes/PhysicsEngine/BodyInstance.h
+++ b/Engine/Source/Runtime/Engine/Classes/PhysicsEngine/BodyInstance.h
@@ -491,7 +491,7 @@ public:
 	void UpdatePhysicalMaterials();
 
 	/** Update the instances collision filtering data */
-	void UpdatePhysicsFilterData();
+	void UpdatePhysicsFilterData(bool bForceSimpleAsComplex = false);
 
 	friend FArchive& operator<<(FArchive& Ar,FBodyInstance& BodyInst);
 
diff --git a/Engine/Source/Runtime/Engine/Private/PhysicsEngine/BodyInstance.cpp b/Engine/Source/Runtime/Engine/Private/PhysicsEngine/BodyInstance.cpp
--- a/Engine/Source/Runtime/Engine/Private/PhysicsEngine/BodyInstance.cpp
+++ b/Engine/Source/Runtime/Engine/Private/PhysicsEngine/BodyInstance.cpp
@@ -600,9 +600,8 @@ ECollisionEnabled::Type FBodyInstance::GetCollisionEnabled() const
 	}
 }
 
-
 /** Update the filter data on the physics shapes, based on the owning component flags. */
-void FBodyInstance::UpdatePhysicsFilterData()
+void FBodyInstance::UpdatePhysicsFilterData(bool bForceSimpleAsComplex)
 {
 	// Do nothing if no physics actor
 	if (!IsValidBodyInstance())
@@ -667,8 +666,8 @@ void FBodyInstance::UpdatePhysicsFilterData()
 	}
 #endif
 
-	const bool bUseComplexAsSimple = (BodySetup.Get()->CollisionTraceFlag == CTF_UseComplexAsSimple);
-	const bool bUseSimpleAsComplex = (BodySetup.Get()->CollisionTraceFlag == CTF_UseSimpleAsComplex);
+	const bool bUseComplexAsSimple = !bForceSimpleAsComplex && (BodySetup.Get()->CollisionTraceFlag == CTF_UseComplexAsSimple);
+	const bool bUseSimpleAsComplex = bForceSimpleAsComplex || (BodySetup.Get()->CollisionTraceFlag == CTF_UseSimpleAsComplex);
 
 #if WITH_PHYSX
 	if (PxRigidActor* PActor = GetPxRigidActor())
@@ -1479,8 +1478,7 @@ void FBodyInstance::Weld(FBodyInstance* TheirBody, const FTransform& RelativeTM)
 	UpdatePhysicalMaterials();
 
 	// Set the filter data on the shapes (call this after setting BodyData because it uses that pointer)
-	UpdatePhysicsFilterData();
-
+	UpdatePhysicsFilterData(true);
 
 	UpdateMassProperties();
 	// Update damping
