Automatically weld objects when calling Attach through BP

diff --git a/Engine/Source/Runtime/Engine/Private/Actor.cpp b/Engine/Source/Runtime/Engine/Private/Actor.cpp
--- a/Engine/Source/Runtime/Engine/Private/Actor.cpp
+++ b/Engine/Source/Runtime/Engine/Private/Actor.cpp
@@ -1109,6 +1109,7 @@ void AActor::AttachRootComponentTo(USceneComponent* InParent, FName InSocketName
 	if(RootComponent && InParent)
 	{
 		RootComponent->AttachTo(InParent, InSocketName, AttachLocationType );
+		InParent->WeldPhysicsBody(RootComponent);
 
 		AttachmentReplication.AttachParent = InParent->GetAttachmentRootActor();
 		AttachmentReplication.LocationOffset = RootComponent->RelativeLocation;
@@ -1140,6 +1141,7 @@ void AActor::OnRep_AttachmentReplication()
 				RootComponent->RelativeScale3D = AttachmentReplication.RelativeScale3D;
 
 				RootComponent->UpdateComponentToWorld();
+				ParentComponent->WeldPhysicsBody(RootComponent, true, AttachmentReplication.AttachSocket);
 			}
 		}
 	}
@@ -1157,6 +1159,7 @@ void AActor::AttachRootComponentToActor(AActor* InParentActor, FName InSocketNam
 		if (ParentRootComponent)
 		{
 			RootComponent->AttachTo(ParentRootComponent, InSocketName, AttachLocationType );
+			ParentRootComponent->WeldPhysicsBody(RootComponent, true, InSocketName);
 
 			AttachmentReplication.AttachParent = InParentActor;
 			AttachmentReplication.LocationOffset = RootComponent->RelativeLocation;
