Fix build for missing weld ifdef

diff --git a/Engine/Source/Runtime/Engine/Private/Actor.cpp b/Engine/Source/Runtime/Engine/Private/Actor.cpp
--- a/Engine/Source/Runtime/Engine/Private/Actor.cpp
+++ b/Engine/Source/Runtime/Engine/Private/Actor.cpp
@@ -1109,7 +1109,9 @@ void AActor::AttachRootComponentTo(USceneComponent* InParent, FName InSocketName
 	if(RootComponent && InParent)
 	{
 		RootComponent->AttachTo(InParent, InSocketName, AttachLocationType );
+#if WITH_BODY_WELDING
 		InParent->WeldPhysicsBody(RootComponent);
+#endif
 
 		AttachmentReplication.AttachParent = InParent->GetAttachmentRootActor();
 		AttachmentReplication.LocationOffset = RootComponent->RelativeLocation;
@@ -1141,7 +1143,9 @@ void AActor::OnRep_AttachmentReplication()
 				RootComponent->RelativeScale3D = AttachmentReplication.RelativeScale3D;
 
 				RootComponent->UpdateComponentToWorld();
+#if WITH_BODY_WELDING
 				ParentComponent->WeldPhysicsBody(RootComponent, true, AttachmentReplication.AttachSocket);
+#endif
 			}
 		}
 	}
@@ -1159,7 +1163,9 @@ void AActor::AttachRootComponentToActor(AActor* InParentActor, FName InSocketNam
 		if (ParentRootComponent)
 		{
 			RootComponent->AttachTo(ParentRootComponent, InSocketName, AttachLocationType );
+#if WITH_BODY_WELDING
 			ParentRootComponent->WeldPhysicsBody(RootComponent, true, InSocketName);
+#endif
 
 			AttachmentReplication.AttachParent = InParentActor;
 			AttachmentReplication.LocationOffset = RootComponent->RelativeLocation;
