#ue4 - Jump() calls during character movement update will now be triggered on the next update rather than discarded.

- Moved the ClearJumpInput() call to within PerformMovement().

TTP 343881

diff --git a/Engine/Source/Runtime/Engine/Private/CharacterMovementComponent.cpp b/Engine/Source/Runtime/Engine/Private/CharacterMovementComponent.cpp
--- a/Engine/Source/Runtime/Engine/Private/CharacterMovementComponent.cpp
+++ b/Engine/Source/Runtime/Engine/Private/CharacterMovementComponent.cpp
@@ -794,8 +794,6 @@ void UCharacterMovementComponent::TickComponent(float DeltaTime, enum ELevelTick
 			{
 				ReplicateMoveToServer(DeltaTime, Acceleration);
 			}
-
-			CharacterOwner->ClearJumpInput();
 		}
 		else if (CharacterOwner->GetRemoteRole() == ROLE_AutonomousProxy)
 		{
@@ -1409,6 +1407,9 @@ void UCharacterMovementComponent::PerformMovement(float DeltaSeconds)
 		// NaN tracking
 		checkf(!Velocity.ContainsNaN(), TEXT("UCharacterMovementComponent::PerformMovement: Velocity contains NaN (%s: %s)\n%s"), *GetPathNameSafe(this), *GetPathNameSafe(GetOuter()), *Velocity.ToString());
 
+		// Clear jump input now, to allow movement events to trigger it for next update.
+		CharacterOwner->ClearJumpInput();
+
 		// change position
 		StartNewPhysics(DeltaSeconds, 0);
 
@@ -4911,11 +4912,11 @@ bool UCharacterMovementComponent::ClientUpdatePositionAfterServerUpdate()
 	}
 
 	// Save important values that might get affected by the replay.
-	float SavedAnalogInputModifier = AnalogInputModifier;
-	FRootMotionMovementParams BackupRootMotionParams = RootMotionParams;
-	bool bRealJump = CharacterOwner->bPressedJump;
-	bool bRealCrouch = bWantsToCrouch;
-	bool bRealForceMaxAccel = bForceMaxAccel;
+	const float SavedAnalogInputModifier = AnalogInputModifier;
+	const FRootMotionMovementParams BackupRootMotionParams = RootMotionParams;
+	const bool bRealJump = CharacterOwner->bPressedJump;
+	const bool bRealCrouch = bWantsToCrouch;
+	const bool bRealForceMaxAccel = bForceMaxAccel;
 	CharacterOwner->bClientWasFalling = (MovementMode == MOVE_Falling);
 	CharacterOwner->bClientUpdating = true;
 	bForceNextFloorCheck = true;
@@ -5590,10 +5591,9 @@ void UCharacterMovementComponent::MoveAutonomous
 
 	Acceleration = ConstrainInputAcceleration(NewAccel);
 	AnalogInputModifier = ComputeAnalogInputModifier();
+	
 	PerformMovement(DeltaTime);
 
-	CharacterOwner->ClearJumpInput();
-
 	// If not playing root motion, tick animations after physics. We do this here to keep events, notifies, states and transitions in sync with client updates.
 	if( !CharacterOwner->bClientUpdating && !CharacterOwner->IsPlayingRootMotion() && CharacterOwner->Mesh )
 	{
