ANIM: LOOKAT: Make sure look at node can be exposed as pin.

diff --git a/Engine/Source/Runtime/Engine/Classes/Animation/BoneControllers/AnimNode_LookAt.h b/Engine/Source/Runtime/Engine/Classes/Animation/BoneControllers/AnimNode_LookAt.h
--- a/Engine/Source/Runtime/Engine/Classes/Animation/BoneControllers/AnimNode_LookAt.h
+++ b/Engine/Source/Runtime/Engine/Classes/Animation/BoneControllers/AnimNode_LookAt.h
@@ -36,7 +36,7 @@ struct ENGINE_API FAnimNode_LookAt : public FAnimNode_SkeletalControlBase
 	FBoneReference LookAtBone;
 
 	/** Target Location in world space if LookAtBone is empty */
-	UPROPERTY(EditAnywhere, BlueprintReadWrite, Category=SkeletalControl)
+	UPROPERTY(EditAnywhere, BlueprintReadWrite, Category=SkeletalControl, meta=(PinHiddenByDefault))
 	FVector LookAtLocation;
 
 	/** Look at axis, which axis to align to look at point */
