Niagara updates; rendering refactor, noise node, various other fixes, changes and updates

diff --git a/Engine/Source/Editor/GraphEditor/Private/NiagaraPins/SGraphPinVector4.cpp b/Engine/Source/Editor/GraphEditor/Private/NiagaraPins/SGraphPinVector4.cpp
--- a/Engine/Source/Editor/GraphEditor/Private/NiagaraPins/SGraphPinVector4.cpp
+++ b/Engine/Source/Editor/GraphEditor/Private/NiagaraPins/SGraphPinVector4.cpp
@@ -248,7 +248,7 @@ void SGraphPinVector4::OnChangedValueTextBox_0(float NewValue, ETextCommit::Type
 
 	FString DefaultValue;
 	//Update X value
-	DefaultValue = ValueStr + FString(TEXT(",")) + GetValue(TextBox_1) + FString(TEXT(",")) + GetValue(TextBox_2) + GetValue(TextBox_3);
+	DefaultValue = ValueStr + FString(TEXT(",")) + GetValue(TextBox_1) + FString(TEXT(",")) + GetValue(TextBox_2) + FString(TEXT(",")) + GetValue(TextBox_3);
 	//Set new default value
 	GraphPinObj->GetSchema()->TrySetDefaultValue(*GraphPinObj, DefaultValue);
 }
diff --git a/Engine/Source/Editor/NiagaraEditor/Private/NiagaraCompiler.cpp b/Engine/Source/Editor/NiagaraEditor/Private/NiagaraCompiler.cpp
--- a/Engine/Source/Editor/NiagaraEditor/Private/NiagaraCompiler.cpp
+++ b/Engine/Source/Editor/NiagaraEditor/Private/NiagaraCompiler.cpp
@@ -162,33 +162,14 @@ static int32 EvaluateGraph(FNiagaraCompilerContext& Context, UEdGraphPin* Pin)
 						{
 							UEdGraphPin* SrcPin = OpNode->FindPinChecked(UNiagaraNodeOp::InPinNames[SrcIndex]);
 							NewOp.Src[SrcIndex] = EvaluateGraph(Context,SrcPin);
+							if (NewOp.Src[SrcIndex] & 0x40000000)
+							{
+								NewOp.SrcOperandTypeVector |= (1 << SrcIndex);
+							}
 						}
 					}
 
-					while (!IsValidOp(Context, NewOp) && VectorVM::GetOpCodeInfo(NewOp.OpIndex).BaseOpcode == OpInfo.BaseOpcode)
-					{
-						NewOp.OpIndex++;
-					}
-
-
-					if (OpInfo.IsCommutative() && !IsValidOp(Context, NewOp))
-					{
-						int32 Temp = NewOp.Src[0];
-						NewOp.Src[0] = NewOp.Src[1];
-						NewOp.Src[1] = Temp;
-						NewOp.OpIndex = OpNode->OpIndex;
-						while (!IsValidOp(Context, NewOp) && VectorVM::GetOpCodeInfo(NewOp.OpIndex).BaseOpcode == OpInfo.BaseOpcode)
-						{
-							NewOp.OpIndex++;
-						}
-					}
-					check(NewOp.OpIndex < VectorVM::EOp::NumOpcodes);
-
-
-					if (IsValidOp(Context, NewOp))
-					{
-						ExpressionIndex = Context.Expressions.Add(NewOp);
-					}
+					ExpressionIndex = Context.Expressions.Add(NewOp);
 				}
 				else if (Node->IsA(UNiagaraNodeGetAttr::StaticClass()))
 				{
@@ -295,10 +276,11 @@ void FNiagaraEditorModule::CompileScript(UNiagaraScript* ScriptToCompile)
 		{
 			// Generate a pass-thru op.
 			FNiagaraExpr PassThru;
-			PassThru.OpIndex = VectorVM::EOp::addi;
+			PassThru.OpIndex = VectorVM::EOp::add;
 			PassThru.Src[0] = AttrIndex | 0x80000000;
 			PassThru.Src[1] = 0 | 0x40000000;
 			PassThru.Src[2] = INDEX_NONE;
+			PassThru.SrcOperandTypeVector = 0x2;
 			OutputExpressions[AttrIndex] = Context.Expressions.Add(PassThru);
 		}
 	}
@@ -313,7 +295,8 @@ void FNiagaraEditorModule::CompileScript(UNiagaraScript* ScriptToCompile)
 		VectorVM::FVectorVMOpInfo const& OpInfo = VectorVM::GetOpCodeInfo(Expr.OpIndex);
 		for (int32 k = 0; k < 3; ++k)
 		{
-			if (OpInfo.SrcTypes[k] == VectorVM::EOpSrc::Register
+			//if (OpInfo.SrcTypes[k] == VectorVM::EOpSrc::Register
+			if ((Expr.SrcOperandTypeVector & (1 << k)) == 0 && OpInfo.SrcTypes[k] != VectorVM::EOpSrc::Invalid
 				&& !IsAttrIndex(Expr.Src[k]))
 			{
 				check(Context.Expressions.IsValidIndex(Expr.Src[k]));
@@ -335,7 +318,8 @@ void FNiagaraEditorModule::CompileScript(UNiagaraScript* ScriptToCompile)
 
 		for (int32 j = 0; j < 3; ++j)
 		{
-			if (OpInfo.SrcTypes[j] == VectorVM::EOpSrc::Register)
+			//if (OpInfo.SrcTypes[j] == VectorVM::EOpSrc::Register)
+			if ((Expr.SrcOperandTypeVector & (1 << j)) == 0 && OpInfo.SrcTypes[j] != VectorVM::EOpSrc::Invalid)
 			{
 				if (IsAttrIndex(Expr.Src[j]))
 				{
@@ -346,11 +330,12 @@ void FNiagaraEditorModule::CompileScript(UNiagaraScript* ScriptToCompile)
 					Expr.Src[j] = RegisterAssignments[Expr.Src[j]];
 				}
 			}
-			else if (OpInfo.SrcTypes[j] == VectorVM::EOpSrc::Const)
+			//else if (OpInfo.SrcTypes[j] == VectorVM::EOpSrc::Const)
+			else if ((Expr.SrcOperandTypeVector & (1 << j)) == 1)
 			{
 				// VectorVM::EOpSrc::Const.
 				Expr.Src[j] = ConstIndexFromExpressionIndex(Expr.Src[j]);
-				Expr.SrcOperandTypeVector |= 1<<j;
+				//Expr.SrcOperandTypeVector |= 1<<j;
 			}
 		}
 
@@ -417,7 +402,7 @@ void FNiagaraEditorModule::CompileScript(UNiagaraScript* ScriptToCompile)
 			{
 				if (OutputExpressions[j] == i)
 				{
-					Code.Add(VectorVM::EOp::addi);
+					Code.Add(VectorVM::EOp::add);
 					Code.Add(VectorVM::FirstOutputRegister + j);
 					Code.Add(0x0);	// all inputs are registers
 					Code.Add(DestRegister);
diff --git a/Engine/Source/Runtime/Engine/Classes/Components/NiagaraComponent.h b/Engine/Source/Runtime/Engine/Classes/Components/NiagaraComponent.h
--- a/Engine/Source/Runtime/Engine/Classes/Components/NiagaraComponent.h
+++ b/Engine/Source/Runtime/Engine/Classes/Components/NiagaraComponent.h
@@ -1,9 +1,16 @@
 // Copyright 1998-2014 Epic Games, Inc. All Rights Reserved.
 
-
 #pragma once
 #include "NiagaraComponent.generated.h"
 
+UENUM()
+enum ERenderModuleType
+{
+	Sprites = 0,
+	Ribbon
+};
+
+
 UCLASS()
 class ENGINE_API UNiagaraComponent : public UPrimitiveComponent
 {
@@ -18,6 +25,10 @@ class ENGINE_API UNiagaraComponent : public UPrimitiveComponent
 	UPROPERTY(EditAnywhere, Category = NiagaraComponent)
 	class UNiagaraScript* SpawnScript;
 
+	/** The method for rendering the simulated data */
+	UPROPERTY(EditAnywhere, Category=Rendering)
+	TEnumAsByte<ERenderModuleType> RenderModuleType;
+
 	/** Material with which to render particles. */
 	UPROPERTY(EditAnywhere, Category=Rendering)
 	UMaterialInterface* Material;
@@ -55,3 +66,48 @@ public:
 
 
 
+
+
+
+/**
+* Scene proxy for drawing niagara particle simulations.
+*/
+class FNiagaraSceneProxy : public FPrimitiveSceneProxy
+{
+public:
+
+	FNiagaraSceneProxy(const UNiagaraComponent* InComponent);
+	~FNiagaraSceneProxy();
+
+	/** Called on render thread to assign new dynamic data */
+	void SetDynamicData_RenderThread(struct FNiagaraDynamicDataBase* NewDynamicData);
+	class NiagaraEffectRenderer *GetEffectRenderer() { return EffectRenderer; }
+
+private:
+	void ReleaseRenderThreadResources();
+
+	// FPrimitiveSceneProxy interface.
+	virtual void CreateRenderThreadResources() override;
+
+	virtual void OnActorPositionChanged() override;
+	virtual void OnTransformChanged() override;
+
+	virtual void PreRenderView(const FSceneViewFamily* ViewFamily, const uint32 VisibilityMap, int32 FrameNumber) override;
+	virtual void DrawDynamicElements(FPrimitiveDrawInterface* PDI, const FSceneView* View) override;
+	virtual void GetDynamicMeshElements(const TArray<const FSceneView*>& Views, const FSceneViewFamily& ViewFamily, uint32 VisibilityMap, FMeshElementCollector& Collector) const override;
+
+	virtual FPrimitiveViewRelevance GetViewRelevance(const FSceneView* View)  override;
+	/*
+	virtual bool CanBeOccluded() const override
+	{
+	return !MaterialRelevance.bDisableDepthTest;
+	}
+	*/
+	virtual uint32 GetMemoryFootprint() const override;
+
+	uint32 GetAllocatedSize() const;
+
+
+private:
+	class NiagaraEffectRenderer *EffectRenderer;
+};
diff --git a/Engine/Source/Runtime/Engine/Private/NiagaraComponent.cpp b/Engine/Source/Runtime/Engine/Private/NiagaraComponent.cpp
--- a/Engine/Source/Runtime/Engine/Private/NiagaraComponent.cpp
+++ b/Engine/Source/Runtime/Engine/Private/NiagaraComponent.cpp
@@ -7,371 +7,111 @@
 #include "ParticleHelper.h"
 #include "Particles/ParticleResources.h"
 #include "Engine/NiagaraConstants.h"
+#include "Engine/NiagaraEffectRenderer.h"
 
 DECLARE_CYCLE_STAT(TEXT("Tick"),STAT_NiagaraTick,STATGROUP_Niagara);
 DECLARE_CYCLE_STAT(TEXT("Simulate"),STAT_NiagaraSimulate,STATGROUP_Niagara);
 DECLARE_CYCLE_STAT(TEXT("Spawn + Kill"),STAT_NiagaraSpawnAndKill,STATGROUP_Niagara);
 DECLARE_CYCLE_STAT(TEXT("Gen Verts"),STAT_NiagaraGenerateVertices,STATGROUP_Niagara);
-DECLARE_CYCLE_STAT(TEXT("PreRenderView"),STAT_NiagaraPreRenderView,STATGROUP_Niagara);
-DECLARE_CYCLE_STAT(TEXT("Render"),STAT_NiagaraRender,STATGROUP_Niagara);
 DECLARE_DWORD_COUNTER_STAT(TEXT("NumParticles"),STAT_NiagaraNumParticles,STATGROUP_Niagara);
 
 
 DEFINE_LOG_CATEGORY_STATIC(LogNiagaraComponent, All, All);
 
-/** Struct used to pass dynamic data from game thread to render thread */
-struct FNiagaraDynamicData
-{
-	/** Vertex data */
-	TArray<FParticleSpriteVertex> VertexData;
-};
-
-class FNiagaraMeshCollectorResources : public FOneFrameResource
-{
-public:
-	FParticleSpriteVertexFactory VertexFactory;
-	FParticleSpriteUniformBufferRef UniformBuffer;
 
-	virtual ~FNiagaraMeshCollectorResources()
-	{
-		VertexFactory.ReleaseResource();
-	}
-};
-
-/**
- * Scene proxy for drawing niagara particle simulations.
- */
-class FNiagaraSceneProxy : public FPrimitiveSceneProxy
-{
-public:
-
-	FNiagaraSceneProxy(const UNiagaraComponent* InComponent)
+FNiagaraSceneProxy::FNiagaraSceneProxy(const UNiagaraComponent* InComponent)
 		:	FPrimitiveSceneProxy(InComponent)
-		,	DynamicData(NULL)
-		, VertexFactory(PVFT_Sprite, InComponent->GetWorld()->FeatureLevel)
-	{
-		Material = InComponent->Material;
-		if (!Material || !Material->CheckMaterialUsage(MATUSAGE_ParticleSprites))
-		{
-			Material = UMaterial::GetDefaultMaterial(MD_Surface);
-		}
-		check(Material);
-		MaterialRelevance = Material->GetRelevance();
-	}
-
-	~FNiagaraSceneProxy()
-	{
-		ReleaseRenderThreadResources();
-	}
-
-	/** Called on render thread to assign new dynamic data */
-	void SetDynamicData_RenderThread(FNiagaraDynamicData* NewDynamicData)
-	{
-		check(IsInRenderingThread());
-
-		if(DynamicData)
-		{
-			delete DynamicData;
-			DynamicData = NULL;
-		}
-
-		DynamicData = NewDynamicData;
-	}
-
-private:
-	void ReleaseRenderThreadResources()
-	{
-		VertexFactory.ReleaseResource();
-		PerViewUniformBuffers.Empty();
-		WorldSpacePrimitiveUniformBuffer.ReleaseResource();
-	}
-
-	// FPrimitiveSceneProxy interface.
-	virtual void CreateRenderThreadResources() override
-	{
-		VertexFactory.InitResource();
-
-		UniformParameters.AxisLockRight = FVector4(0.0f, 0.0f, 0.0f, 0.0f);
-		UniformParameters.AxisLockUp = FVector4(0.0f, 0.0f, 0.0f, 0.0f);
-		UniformParameters.RotationScale = 1.0f;
-		UniformParameters.RotationBias = 0.0f;
-		UniformParameters.TangentSelector = FVector4(0.0f, 0.0f, 0.0f, 1.0f);
-		UniformParameters.InvDeltaSeconds = 0.0f;
-		UniformParameters.SubImageSize = FVector4(1.0f,1.0f,1.0f,1.0f);
-		UniformParameters.NormalsType = 0;
-		UniformParameters.NormalsSphereCenter = FVector4(0.0f, 0.0f, 0.0f, 1.0f);
-		UniformParameters.NormalsCylinderUnitDirection = FVector4(0.0f, 0.0f, 1.0f, 0.0f);
-		UniformParameters.PivotOffset = FVector2D(0.0f, 0.0f);
-	}
-
-	virtual void OnActorPositionChanged() override
+{
+	if (InComponent->RenderModuleType == Sprites)
 	{
-		WorldSpacePrimitiveUniformBuffer.ReleaseResource();
+		EffectRenderer = new NiagaraEffectRendererSprites(InComponent, this);
 	}
-
-	virtual void OnTransformChanged() override
+	else  if (InComponent->RenderModuleType == Ribbon)
 	{
-		WorldSpacePrimitiveUniformBuffer.ReleaseResource();
+		EffectRenderer = new NiagaraEffectRendererRibbon(InComponent, this);
 	}
 
-	virtual void GetDynamicMeshElements(const TArray<const FSceneView*>& Views, const FSceneViewFamily& ViewFamily, uint32 VisibilityMap, FMeshElementCollector& Collector) const override
-	{
-		SCOPE_CYCLE_COUNTER(STAT_NiagaraRender);
+}
 
-		check(DynamicData && DynamicData->VertexData.Num());
+FNiagaraSceneProxy::~FNiagaraSceneProxy()
+{
+	ReleaseRenderThreadResources();
+}
 
-		const bool bIsWireframe = ViewFamily.EngineShowFlags.Wireframe;
-		FMaterialRenderProxy* MaterialRenderProxy = Material->GetRenderProxy(IsSelected(),IsHovered());
+/** Called on render thread to assign new dynamic data */
+void FNiagaraSceneProxy::SetDynamicData_RenderThread(FNiagaraDynamicDataBase* NewDynamicData)
+{
+	EffectRenderer->SetDynamicData_RenderThread(NewDynamicData);
+	return;
+}
 
-		int32 SizeInBytes = DynamicData->VertexData.GetTypeSize() * DynamicData->VertexData.Num();
-		FGlobalDynamicVertexBuffer::FAllocation LocalDynamicVertexAllocation = FGlobalDynamicVertexBuffer::Get().Allocate(SizeInBytes);
 
-		if (LocalDynamicVertexAllocation.IsValid())
-		{
-			// Update the primitive uniform buffer if needed.
-			if (!WorldSpacePrimitiveUniformBuffer.IsInitialized())
-			{
-				FPrimitiveUniformShaderParameters PrimitiveUniformShaderParameters = GetPrimitiveUniformShaderParameters(
-					FMatrix::Identity,
-					GetActorPosition(),
-					GetBounds(),
-					GetLocalBounds(),
-					ReceivesDecals(),
-					false,
-					UseEditorDepthTest()
-					);
-				WorldSpacePrimitiveUniformBuffer.SetContents(PrimitiveUniformShaderParameters);
-				WorldSpacePrimitiveUniformBuffer.InitResource();
-			}
+void FNiagaraSceneProxy::ReleaseRenderThreadResources()
+{
+	EffectRenderer->ReleaseRenderThreadResources();
+	return;
+}
 
-			// Copy the vertex data over.
-			FMemory::Memcpy(LocalDynamicVertexAllocation.Buffer, DynamicData->VertexData.GetData(), SizeInBytes);
+// FPrimitiveSceneProxy interface.
+void FNiagaraSceneProxy::CreateRenderThreadResources()
+{
+	EffectRenderer->CreateRenderThreadResources();
+	return;
+}
 
-			// Compute the per-view uniform buffers.
-			for (int32 ViewIndex = 0; ViewIndex < Views.Num(); ViewIndex++)
-			{
-				if (VisibilityMap & (1 << ViewIndex))
-				{
-					const FSceneView* View = Views[ViewIndex];
-
-					FNiagaraMeshCollectorResources& CollectorResources = Collector.AllocateOneFrameResource<FNiagaraMeshCollectorResources>();
-					FParticleSpriteUniformParameters PerViewUniformParameters = UniformParameters;
-
-					// Collector.AllocateOneFrameResource uses default ctor, initialize the vertex factory
-					CollectorResources.VertexFactory.SetFeatureLevel(ViewFamily.GetFeatureLevel());
-					CollectorResources.VertexFactory.SetParticleFactoryType(PVFT_Sprite);
-
-					PerViewUniformParameters.MacroUVParameters = FVector4(0.0f,0.0f,1.0f,1.0f);
-					CollectorResources.UniformBuffer = FParticleSpriteUniformBufferRef::CreateUniformBufferImmediate(PerViewUniformParameters, UniformBuffer_SingleFrame);
-
-					CollectorResources.VertexFactory.InitResource();
-					CollectorResources.VertexFactory.SetSpriteUniformBuffer(CollectorResources.UniformBuffer);
-					CollectorResources.VertexFactory.SetInstanceBuffer(
-						LocalDynamicVertexAllocation.VertexBuffer,
-						LocalDynamicVertexAllocation.VertexOffset,
-						sizeof(FParticleSpriteVertex),
-						true
-						);
-					CollectorResources.VertexFactory.SetDynamicParameterBuffer(NULL,0,0,true);
-
-					FMeshBatch& MeshBatch = Collector.AllocateMesh();
-					MeshBatch.VertexFactory = &CollectorResources.VertexFactory;
-					MeshBatch.CastShadow = CastsDynamicShadow();
-					MeshBatch.bUseAsOccluder = false;
-					MeshBatch.ReverseCulling = IsLocalToWorldDeterminantNegative();
-					MeshBatch.Type = PT_TriangleList;
-					MeshBatch.DepthPriorityGroup = GetDepthPriorityGroup(View);
-					MeshBatch.bCanApplyViewModeOverrides = true;
-					MeshBatch.bUseWireframeSelectionColoring = IsSelected();
-
-					if (bIsWireframe)
-					{
-						MeshBatch.MaterialRenderProxy = UMaterial::GetDefaultMaterial(MD_Surface)->GetRenderProxy(IsSelected(),IsHovered());
-					}
-					else
-					{
-						MeshBatch.MaterialRenderProxy = MaterialRenderProxy;
-					}
-
-					FMeshBatchElement& MeshElement = MeshBatch.Elements[0];
-					MeshElement.IndexBuffer = &GParticleIndexBuffer;
-					MeshElement.FirstIndex = 0;
-					MeshElement.NumPrimitives = 2;
-					MeshElement.NumInstances = DynamicData->VertexData.Num();
-					MeshElement.MinVertexIndex = 0;
-					MeshElement.MaxVertexIndex = MeshElement.NumInstances * 4 - 1;
-					MeshElement.PrimitiveUniformBufferResource = &WorldSpacePrimitiveUniformBuffer;
-
-					Collector.AddMesh(ViewIndex, MeshBatch);
-				}
-			}
-		}
-	}
+void FNiagaraSceneProxy::OnActorPositionChanged()
+{
+	//WorldSpacePrimitiveUniformBuffer.ReleaseResource();
+}
 
-	virtual void PreRenderView(const FSceneViewFamily* ViewFamily, const uint32 VisibilityMap, int32 FrameNumber) override
-	{
-		SCOPE_CYCLE_COUNTER(STAT_NiagaraPreRenderView);
-		
-		check(DynamicData && DynamicData->VertexData.Num());
+void FNiagaraSceneProxy::OnTransformChanged()
+{
+	//WorldSpacePrimitiveUniformBuffer.ReleaseResource();
+}
 
-		int32 SizeInBytes = DynamicData->VertexData.GetTypeSize() * DynamicData->VertexData.Num();
-		DynamicVertexAllocation = FGlobalDynamicVertexBuffer::Get().Allocate(SizeInBytes);
-		if (DynamicVertexAllocation.IsValid())
-		{
-			// Copy the vertex data over.
-			FMemory::Memcpy(DynamicVertexAllocation.Buffer, DynamicData->VertexData.GetData(), SizeInBytes);
-			VertexFactory.SetInstanceBuffer(
-				DynamicVertexAllocation.VertexBuffer,
-				DynamicVertexAllocation.VertexOffset,
-				sizeof(FParticleSpriteVertex),
-				true
-				);
-			VertexFactory.SetDynamicParameterBuffer(NULL, 0, 0, true);
-			// Compute the per-view uniform buffers.
-			const int32 NumViews = ViewFamily->Views.Num();
-			uint32 ViewBit = 1;
-			for (int32 ViewIndex = 0; ViewIndex < NumViews; ++ViewIndex, ViewBit <<= 1)
-			{
-				FParticleSpriteUniformBufferRef* SpriteViewUniformBufferPtr = new(PerViewUniformBuffers) FParticleSpriteUniformBufferRef();
-				if (VisibilityMap & ViewBit)
-				{
-					FParticleSpriteUniformParameters PerViewUniformParameters = UniformParameters;
-					PerViewUniformParameters.MacroUVParameters = FVector4(0.0f, 0.0f, 1.0f, 1.0f);
-					*SpriteViewUniformBufferPtr = FParticleSpriteUniformBufferRef::CreateUniformBufferImmediate(PerViewUniformParameters, UniformBuffer_SingleFrame);
-				}
-			}
+void FNiagaraSceneProxy::PreRenderView(const FSceneViewFamily* ViewFamily, const uint32 VisibilityMap, int32 FrameNumber)
+{
+	EffectRenderer->PreRenderView(ViewFamily, VisibilityMap, FrameNumber);
+	return;
 
-			// Update the primitive uniform buffer if needed.
-			if (!WorldSpacePrimitiveUniformBuffer.IsInitialized())
-			{
-				FPrimitiveUniformShaderParameters PrimitiveUniformShaderParameters = GetPrimitiveUniformShaderParameters(
-					FMatrix::Identity,
-					GetActorPosition(),
-					GetBounds(),
-					GetLocalBounds(),
-					ReceivesDecals(),
-					false,
-					UseEditorDepthTest()
-					);
-				WorldSpacePrimitiveUniformBuffer.SetContents(PrimitiveUniformShaderParameters);
-				WorldSpacePrimitiveUniformBuffer.InitResource();
-			}
-		}
-	}
+}
 		  
-	virtual void DrawDynamicElements(FPrimitiveDrawInterface* PDI,const FSceneView* View) override
-	{
-		SCOPE_CYCLE_COUNTER(STAT_NiagaraRender);
-
-		check(DynamicData && DynamicData->VertexData.Num());
-
-		const bool bIsWireframe = View->Family->EngineShowFlags.Wireframe;
-		FMaterialRenderProxy* MaterialRenderProxy = Material->GetRenderProxy(IsSelected(),IsHovered());
-
-		if (DynamicVertexAllocation.IsValid()
-			&& (bIsWireframe || !PDI->IsMaterialIgnored(MaterialRenderProxy, View->GetFeatureLevel())))
-		{
-			SCOPED_DRAW_EVENT(NiagaraRender, DEC_SCENE_ITEMS);
-
-			FMeshBatch MeshBatch;
-			MeshBatch.VertexFactory = &VertexFactory;
-			MeshBatch.CastShadow = CastsDynamicShadow();
-			MeshBatch.bUseAsOccluder = false;
-			MeshBatch.ReverseCulling = IsLocalToWorldDeterminantNegative();
-			MeshBatch.Type = PT_TriangleList;
-			MeshBatch.DepthPriorityGroup = GetDepthPriorityGroup(View);
-			MeshBatch.LCI = NULL;
-			if (bIsWireframe)
-			{
-				MeshBatch.MaterialRenderProxy = UMaterial::GetDefaultMaterial(MD_Surface)->GetRenderProxy(IsSelected(),IsHovered());
-			}
-			else
-			{
-				MeshBatch.MaterialRenderProxy = MaterialRenderProxy;
-			}
+void FNiagaraSceneProxy::DrawDynamicElements(FPrimitiveDrawInterface* PDI, const FSceneView* View) 
+{
+	EffectRenderer->DrawDynamicElements(PDI, View);
+	return;
+}
 
-			FMeshBatchElement& MeshElement = MeshBatch.Elements[0];
-			MeshElement.IndexBuffer = &GParticleIndexBuffer;
-			MeshElement.FirstIndex = 0;
-			MeshElement.NumPrimitives = 2;
-			MeshElement.NumInstances = DynamicData->VertexData.Num();
-			MeshElement.MinVertexIndex = 0;
-			MeshElement.MaxVertexIndex = MeshElement.NumInstances * 4 - 1;
-			MeshElement.PrimitiveUniformBufferResource = &WorldSpacePrimitiveUniformBuffer;
-
-			int32 ViewIndex = View->Family->Views.Find(View);
-			check(PerViewUniformBuffers.IsValidIndex(ViewIndex) && PerViewUniformBuffers[ViewIndex]);
-			VertexFactory.SetSpriteUniformBuffer(PerViewUniformBuffers[ViewIndex]);
-
-			DrawRichMesh(
-				PDI, 
-				MeshBatch, 
-				FLinearColor(1.0f, 0.0f, 0.0f),	//WireframeColor,
-				FLinearColor(1.0f, 1.0f, 0.0f),	//LevelColor,
-				FLinearColor(1.0f, 1.0f, 1.0f),	//PropertyColor,		
-				this,
-				GIsEditor && (View->Family->EngineShowFlags.Selection) ? IsSelected() : false,
-				bIsWireframe
-				);
-		}
+FPrimitiveViewRelevance FNiagaraSceneProxy::GetViewRelevance(const FSceneView* View)
+{
+	return EffectRenderer->GetViewRelevance(View);
+}
+/*
+virtual bool CanBeOccluded() const override
+{
+	return !MaterialRelevance.bDisableDepthTest;
+}
+*/
 
-		/*
-		for(int32 i=0; i<DynamicData->VertexData.Num(); i++)
-		{
-			FParticleSpriteVertex& Vertex = DynamicData->VertexData[i];
-			PDI->DrawPoint(Vertex.Position, Vertex.Color, 1.0f, SDPG_World);
-		}
-		*/
-	}
+uint32 FNiagaraSceneProxy::GetMemoryFootprint() const
+{ 
+	return (sizeof(*this) + GetAllocatedSize()); 
+}
 
-	virtual FPrimitiveViewRelevance GetViewRelevance(const FSceneView* View)  override
-	{
-		FPrimitiveViewRelevance Result;
-		bool bHasDynamicData = DynamicData && DynamicData->VertexData.Num() > 0;
-
-		Result.bDrawRelevance = bHasDynamicData && IsShown(View) && View->Family->EngineShowFlags.Particles;
-		Result.bShadowRelevance = bHasDynamicData && IsShadowCast(View);
-		Result.bDynamicRelevance = bHasDynamicData;
-		Result.bNeedsPreRenderView = Result.bDrawRelevance || Result.bShadowRelevance;
-		if (bHasDynamicData && View->Family->EngineShowFlags.Bounds)
-		{
-			Result.bOpaqueRelevance = true;
-		}
-		MaterialRelevance.SetPrimitiveViewRelevance(Result);
+uint32 FNiagaraSceneProxy::GetAllocatedSize() const
+{ 
+	uint32 DynamicDataSize = 0;
+	return FPrimitiveSceneProxy::GetAllocatedSize() + EffectRenderer->GetDynamicDataSize();
+}
 
-		return Result;
-	}
 
-	virtual bool CanBeOccluded() const override
-	{
-		return !MaterialRelevance.bDisableDepthTest;
-	}
+void FNiagaraSceneProxy::GetDynamicMeshElements(const TArray<const FSceneView*>& Views, const FSceneViewFamily& ViewFamily, uint32 VisibilityMap, FMeshElementCollector& Collector) const
+{
+	EffectRenderer->GetDynamicMeshElements(Views, ViewFamily, VisibilityMap, Collector);
+}
 
-	virtual uint32 GetMemoryFootprint() const override 
-	{ 
-		return (sizeof(*this) + GetAllocatedSize()); 
-	}
 
-	uint32 GetAllocatedSize() const 
-	{ 
-		uint32 DynamicDataSize = 0;
-		if (DynamicData)
-		{
-			DynamicDataSize = sizeof(FNiagaraDynamicData) + DynamicData->VertexData.GetAllocatedSize();
-		}
-		return FPrimitiveSceneProxy::GetAllocatedSize() + DynamicDataSize;
-	}
 
-private:
-	UMaterialInterface* Material;
-	FNiagaraDynamicData* DynamicData;
-	mutable TUniformBuffer<FPrimitiveUniformShaderParameters> WorldSpacePrimitiveUniformBuffer;
-	FParticleSpriteVertexFactory VertexFactory;
-	FParticleSpriteUniformParameters UniformParameters;
-	TArray<FParticleSpriteUniformBufferRef, TInlineAllocator<2> > PerViewUniformBuffers;
-	FGlobalDynamicVertexBuffer::FAllocation DynamicVertexAllocation;
-	FMaterialRelevance MaterialRelevance;
-};
 
 namespace ENiagaraVectorAttr
 {
@@ -471,16 +211,13 @@ public:
 		INC_DWORD_STAT_BY(STAT_NiagaraNumParticles, NumParticles);
 	}
 
-	FNiagaraDynamicData* GetDynamicData()
-	{
-		SCOPE_CYCLE_COUNTER(STAT_NiagaraGenerateVertices);
-		FNiagaraDynamicData* DynamicData = new FNiagaraDynamicData;
-		GenerateVertexData(DynamicData->VertexData);
-		return DynamicData;
-	}
-
 	FBox GetBounds() const { return CachedBounds; }
 
+
+	int GetNumParticles() const { return NumParticles; }
+	int GetNumVectorsPerAttribute() const { return NumVectorsPerAttribute; }
+	const TArray<FVector4> &GetParticles() { return ParticleBuffers[BufferIndex]; }
+
 	void SetConstant(int Idx, const FVector4 &Value)
 	{
 		check(Idx < ConstantTable.Num());
@@ -675,35 +412,6 @@ private:
 		}
 	}
 
-	/** Update render data buffer from attributes */
-	void GenerateVertexData(TArray<FParticleSpriteVertex>& RenderData)
-	{
-		RenderData.Reset(NumParticles);
-		CachedBounds.Init();
-
-
-		float MaxSize = 0.0f;
-		FVector4* Particles = ParticleBuffers[BufferIndex].GetTypedData();
-		int32 AttrStride = NumVectorsPerAttribute;
-		for (int32 ParticleIndex = 0; ParticleIndex < NumParticles; ParticleIndex++)
-		{
-			FVector4* Particle = Particles + ParticleIndex;
-
-			FParticleSpriteVertex& NewVertex = *new(RenderData)FParticleSpriteVertex;
-			NewVertex.Position = Particle[AttrStride*ENiagaraVectorAttr::Position];
-			NewVertex.OldPosition = NewVertex.Position;
-			CachedBounds += NewVertex.Position;
-			NewVertex.Color = FLinearColor(Particle[AttrStride*ENiagaraVectorAttr::Color]);
-			NewVertex.ParticleId = static_cast<float>(ParticleIndex) / NumParticles;
-			NewVertex.RelativeTime = Particle[AttrStride*ENiagaraVectorAttr::RelativeTime].X;
-			NewVertex.Size = FVector2D(1.0f, 1.0f);
-			MaxSize = FMath::Max(MaxSize, NewVertex.Size.X);
-			MaxSize = FMath::Max(MaxSize, NewVertex.Size.Y);
-			NewVertex.Rotation = Particle[AttrStride*ENiagaraVectorAttr::Rotation].X;
-			NewVertex.SubImageIndex = 0.f;
-		}
-		CachedBounds.ExpandBy(MaxSize);
-	}
 };
 
 //////////////////////////////////////////////////////////////////////////
@@ -713,8 +421,9 @@ UNiagaraComponent::UNiagaraComponent(const class FPostConstructInitializePropert
 {
 	PrimaryComponentTick.bCanEverTick = true;
 	bTickInEditor = true;
-
+	RenderModuleType = ERenderModuleType::Sprites;
 	SpawnRate = 20.f;
+	Simulation = nullptr;
 }
 
 
@@ -747,6 +456,7 @@ void UNiagaraComponent::OnRegister()
 	{
 		Simulation = new FNiagaraSimulation(this);
 	}
+	VectorVM::Init();
 }
 
 
@@ -765,25 +475,29 @@ void UNiagaraComponent::SendRenderDynamicData_Concurrent()
 {
 	if (Simulation && SceneProxy)
 	{
-		FNiagaraDynamicData* DynamicData = Simulation->GetDynamicData();
+		FNiagaraSceneProxy *NiagaraProxy = static_cast<FNiagaraSceneProxy*>(SceneProxy);
+		FNiagaraDynamicDataBase* DynamicData = NiagaraProxy->GetEffectRenderer()->GenerateVertexData(Simulation->GetNumVectorsPerAttribute(), Simulation->GetParticles());
+		//FNiagaraDynamicDataBase* DynamicData = Simulation->GetDynamicData();
 
 		ENQUEUE_UNIQUE_RENDER_COMMAND_TWOPARAMETER(
 			FSendNiagaraDynamicData,
-			FNiagaraSceneProxy*,NiagaraSceneProxy,(FNiagaraSceneProxy*)SceneProxy,
-			FNiagaraDynamicData*,DynamicData,DynamicData,
+			NiagaraEffectRenderer*, EffectRenderer, (NiagaraEffectRenderer*)NiagaraProxy->GetEffectRenderer(),
+			FNiagaraDynamicDataBase*,DynamicData,DynamicData,
 		{
-			NiagaraSceneProxy->SetDynamicData_RenderThread(DynamicData);
+			EffectRenderer->SetDynamicData_RenderThread(DynamicData);
 		});
 	}
 }
 
 int32 UNiagaraComponent::GetNumMaterials() const
 {
-	return 1;
+	return 0;
 }
 
 UMaterialInterface* UNiagaraComponent::GetMaterial(int32 ElementIndex) const
 {
+	return NULL;
+
 	UMaterialInterface* RequestedMaterial = NULL;
 	if (ElementIndex == 0)
 	{
diff --git a/Engine/Source/Runtime/VectorVM/Private/VectorVM.cpp b/Engine/Source/Runtime/VectorVM/Private/VectorVM.cpp
--- a/Engine/Source/Runtime/VectorVM/Private/VectorVM.cpp
+++ b/Engine/Source/Runtime/VectorVM/Private/VectorVM.cpp
@@ -19,6 +19,7 @@ DEFINE_LOG_CATEGORY_STATIC(LogVectorVM, All, All);
 #define SRCOP_CCR 0x06
 #define SRCOP_CCC 0x07
 
+
 /**
  * Context information passed around during VM execution.
  */
@@ -338,15 +339,15 @@ struct FVectorKernelSin : public TUnaryVectorKernel<FVectorKernelSin>
 	}
 };
 
-struct FVectorKernelSin4 : public TUnaryVectorKernel<FVectorKernelSin>
+struct FVectorKernelSin4 : public TUnaryVectorKernel<FVectorKernelSin4>
 {
 	static void FORCEINLINE DoKernel(VectorRegister* RESTRICT Dst, VectorRegister Src0)
 	{
 		float const* FloatSrc0 = reinterpret_cast<float const*>(&Src0);
-		float sn1 = FMath::Sin(*FloatSrc0++ * 3.14f);		// [0;1] takes us through half a period
-		float sn2 = FMath::Sin(*FloatSrc0++ * 3.14f);		// [0;1] takes us through half a period
-		float sn3 = FMath::Sin(*FloatSrc0++ * 3.14f);		// [0;1] takes us through half a period
-		float sn4 = FMath::Sin(*FloatSrc0++ * 3.14f);		// [0;1] takes us through half a period
+		float sn1 = FMath::Sin(FloatSrc0[0] * 3.14f);		// [0;1] takes us through half a period
+		float sn2 = FMath::Sin(FloatSrc0[1] * 3.14f);		// [0;1] takes us through half a period
+		float sn3 = FMath::Sin(FloatSrc0[2] * 3.14f);		// [0;1] takes us through half a period
+		float sn4 = FMath::Sin(FloatSrc0[3] * 3.14f);		// [0;1] takes us through half a period
 		*Dst = MakeVectorRegister(sn1, sn2, sn3, sn4);
 	}
 };
@@ -430,7 +431,91 @@ struct FVectorKernelPow : public TBinaryVectorKernel<FVectorKernelPow>
 };
 
 
+struct FVectorKernelNoise : public TUnaryVectorKernel<FVectorKernelNoise>
+{
+	static VectorRegister RandomTable[10][10][10];
+
+	static void FORCEINLINE DoKernel(VectorRegister* RESTRICT Dst, VectorRegister Src0)
+	{
+		const VectorRegister One = MakeVectorRegister(1.0f, 1.0f, 1.0f, 1.0f);
+		const VectorRegister OneHalf = MakeVectorRegister(0.5f, 0.5f, 0.5f, 0.5f);
+		VectorRegister RndVals[2][2][2];
+		*Dst = MakeVectorRegister(0.0f, 0.0f, 0.0f, 0.0f);
+
+		float const* FloatSrc = reinterpret_cast<float const*>(&Src0);
+
+		for (uint32 i = 1; i < 2; i++)
+		{
+			float FX = FloatSrc[0] / 5 / (1<<i);
+			float FY = FloatSrc[1] / 5 / (1<<i);
+			float FZ = FloatSrc[2] / 5 / (1<<i);
+			uint32 X = FMath::Abs((int32)(FX) % 8);
+			uint32 Y = FMath::Abs( (int32)(FY) % 8 );
+			uint32 Z = FMath::Abs( (int32)(FZ) % 8 );
+
+			for (int32 z = 0; z < 2; z++)
+			{
+				for (int32 y = 0; y < 2; y++)
+				{
+					for (int32 x = 0; x < 2; x++)
+					{
+						RndVals[x][y][z] = RandomTable[X+x][Y+y][Z+z];
+					}
+				}
+			}
+
+			const uint32 IntCoords[3] = { FX, FY, FZ };
+			const float Fractionals[3] = { FX - IntCoords[0], FY - IntCoords[1], FZ - IntCoords[2] };
+			VectorRegister Alpha = MakeVectorRegister(Fractionals[0], Fractionals[0], Fractionals[0], Fractionals[0]);
+
+			VectorRegister OneMinusAlpha = VectorSubtract(One, Alpha);
+			VectorRegister XV1 = VectorAdd(VectorMultiply(RndVals[0][0][0], Alpha), VectorMultiply(RndVals[1][0][0], OneMinusAlpha));
+			VectorRegister XV2 = VectorAdd(VectorMultiply(RndVals[0][1][0], Alpha), VectorMultiply(RndVals[1][1][0], OneMinusAlpha));
+			VectorRegister XV3 = VectorAdd(VectorMultiply(RndVals[0][0][1], Alpha), VectorMultiply(RndVals[1][0][1], OneMinusAlpha));
+			VectorRegister XV4 = VectorAdd(VectorMultiply(RndVals[0][1][1], Alpha), VectorMultiply(RndVals[1][1][1], OneMinusAlpha));
+
+			Alpha = MakeVectorRegister(Fractionals[1], Fractionals[1], Fractionals[1], Fractionals[1]);
+			OneMinusAlpha = VectorSubtract(One, Alpha);
+			VectorRegister YV1 = VectorAdd(VectorMultiply(XV1, Alpha), VectorMultiply(XV2, OneMinusAlpha));
+			VectorRegister YV2 = VectorAdd(VectorMultiply(XV3, Alpha), VectorMultiply(XV4, OneMinusAlpha));
+
+
+			Alpha = MakeVectorRegister(Fractionals[2], Fractionals[2], Fractionals[2], Fractionals[2]);
+			OneMinusAlpha = VectorSubtract(One, Alpha);
+			VectorRegister ZV = VectorAdd(VectorMultiply(YV1, Alpha), VectorMultiply(YV2, OneMinusAlpha));
+
+			*Dst = VectorAdd(*Dst, ZV);
+		}
+	}
+};
+
+VectorRegister FVectorKernelNoise::RandomTable[10][10][10];
+
+
+void VectorVM::Init()
+{
+	static bool Inited = false;
+	if (Inited == false)
+	{
+		for (int z = 0; z < 10; z++)
+		{
+			for (int y = 0; y < 10; y++)
+			{
+				for (int x = 0; x < 10; x++)
+				{
+					float f1 = ((float)FMath::Rand() / (float)RAND_MAX) * 2 - 1;
+					float f2 = ((float)FMath::Rand() / (float)RAND_MAX) * 2 - 1;
+					float f3 = ((float)FMath::Rand() / (float)RAND_MAX) * 2 - 1;
+					float f4 = ((float)FMath::Rand() / (float)RAND_MAX) * 2 - 1;
+
+					FVectorKernelNoise::RandomTable[x][y][z] = MakeVectorRegister(f1, f2, f3, f4);
+				}
+			}
+		}
 
+		Inited = true;
+	}
+}
 
 void VectorVM::Exec(
 	uint8 const* Code,
@@ -498,6 +583,7 @@ void VectorVM::Exec(
 			case EOp::cross: FVectorKernelCross::Exec(Context); break;
 			case EOp::normalize: FVectorKernelNormalize::Exec(Context); break;
 			case EOp::random: FVectorKernelRandom::Exec(Context); break;
+			case EOp::noise: FVectorKernelNoise::Exec(Context); break;
 
 			// Execution always terminates with a "done" opcode.
 			case EOp::done:
@@ -599,6 +685,11 @@ namespace VectorVM
 		FVectorVMOpInfo(EOp::sin4, EOpFlags::Implemented, EOpSrc::Register, EOpSrc::Invalid, EOpSrc::Invalid, TEXT("Sin4")),
 		FVectorVMOpInfo(EOp::sin4, EOpFlags::None, EOpSrc::Const, EOpSrc::Invalid, EOpSrc::Invalid, TEXT("sin4i")),
 
+		FVectorVMOpInfo(EOp::add, EOpFlags::None, EOpSrc::Const, EOpSrc::Const, EOpSrc::Invalid, TEXT("addii")),
+		FVectorVMOpInfo(EOp::sub, EOpFlags::None, EOpSrc::Const, EOpSrc::Const, EOpSrc::Invalid, TEXT("subii")),
+		FVectorVMOpInfo(EOp::mul, EOpFlags::None, EOpSrc::Const, EOpSrc::Const, EOpSrc::Invalid, TEXT("mulii")),
+		FVectorVMOpInfo(EOp::noise, EOpFlags::Implemented, EOpSrc::Register, EOpSrc::Invalid, EOpSrc::Invalid, TEXT("Noise")),
+
 		FVectorVMOpInfo(EOp::add, EOpFlags::None, EOpSrc::Invalid, EOpSrc::Invalid, EOpSrc::Invalid, TEXT("invalid"))
 	};
 } // namespace VectorVM
diff --git a/Engine/Source/Runtime/VectorVM/Public/VectorVM.h b/Engine/Source/Runtime/VectorVM/Public/VectorVM.h
--- a/Engine/Source/Runtime/VectorVM/Public/VectorVM.h
+++ b/Engine/Source/Runtime/VectorVM/Public/VectorVM.h
@@ -93,6 +93,10 @@ namespace VectorVM
 			lengthi,
 			sin4, 
 			sin4i,
+			addii,
+			subii,
+			mulii,
+			noise,
 			NumOpcodes
 		};
 	} // namespace EOp
@@ -170,4 +174,7 @@ namespace VectorVM
 		int32 NumVectors
 		);
 
+	VECTORVM_API void Init();
+
+
 } // namespace VectorVM
