Added missing file.

diff --git a/Engine/Source/Runtime/Engine/Classes/Engine/NiagaraEffectRenderer.h b/Engine/Source/Runtime/Engine/Classes/Engine/NiagaraEffectRenderer.h
new file mode 100644
--- /dev/null
+++ b/Engine/Source/Runtime/Engine/Classes/Engine/NiagaraEffectRenderer.h
@@ -0,0 +1,856 @@
+// Copyright 1998-2014 Epic Games, Inc. All Rights Reserved.
+
+/*==============================================================================
+NiagaraEffectRenderer.h: Base class for Niagara render modules
+==============================================================================*/
+
+DECLARE_CYCLE_STAT(TEXT("PreRenderView"), STAT_NiagaraPreRenderView, STATGROUP_Niagara);
+DECLARE_CYCLE_STAT(TEXT("Render Total"), STAT_NiagaraRender, STATGROUP_Niagara);
+DECLARE_CYCLE_STAT(TEXT("Render Sprites"), STAT_NiagaraRenderSprites, STATGROUP_Niagara);
+DECLARE_CYCLE_STAT(TEXT("Render Ribbons"), STAT_NiagaraRenderRibbons, STATGROUP_Niagara);
+
+
+/** TODO: This should not be declared here but in its own file */
+struct FNiagaraParticleAttribute
+{
+	FName Name;
+	FVector4 *DataPtr;
+};
+
+class FNiagaraEmitterParticleData
+{
+public:
+	FNiagaraEmitterParticleData() 
+	{
+		AddAttribute("Position");
+		AddAttribute("Velocity");
+		AddAttribute("Color");
+	}
+
+	~FNiagaraEmitterParticleData() {}
+	
+	void Allocate(int NumParticles)
+	{
+	}
+
+	FVector4 *GetAttributeData(FName Name)
+	{
+	}
+
+	void AddAttribute(FName NewAttrName)
+	{
+	}
+
+private:
+	TArray<FVector4> ParticleBuffers[2];
+	TArray<FNiagaraParticleAttribute> Attributes;
+	TMap<FName, FNiagaraParticleAttribute*> AttrMap;
+};
+
+
+
+/** Struct used to pass dynamic data from game thread to render thread */
+struct FNiagaraDynamicDataBase
+{
+};
+
+
+struct FNiagaraDynamicDataSprites : public FNiagaraDynamicDataBase
+{
+	TArray<FParticleSpriteVertex> VertexData;
+};
+
+struct FNiagaraDynamicDataRibbon : public FNiagaraDynamicDataBase
+{
+	TArray<FParticleBeamTrailVertex> VertexData;
+};
+
+
+
+/* Mesh collector classes */
+class FNiagaraMeshCollectorResourcesSprite : public FOneFrameResource
+{
+public:
+	FParticleSpriteVertexFactory VertexFactory;
+	FParticleSpriteUniformBufferRef UniformBuffer;
+
+	virtual ~FNiagaraMeshCollectorResourcesSprite()
+	{
+		VertexFactory.ReleaseResource();
+	}
+};
+
+
+class FNiagaraMeshCollectorResourcesRibbon : public FOneFrameResource
+{
+public:
+	 FParticleBeamTrailVertexFactory VertexFactory;
+	 FParticleBeamTrailUniformBufferRef UniformBuffer;
+
+	virtual ~FNiagaraMeshCollectorResourcesRibbon()
+	{
+		VertexFactory.ReleaseResource();
+	}
+};
+
+
+
+
+
+class NiagaraEffectRenderer
+{
+public:
+	virtual void GetDynamicMeshElements(const TArray<const FSceneView*>& Views, const FSceneViewFamily& ViewFamily, uint32 VisibilityMap, FMeshElementCollector& Collector) const = 0;
+
+	virtual void PreRenderView(const FSceneViewFamily* ViewFamily, const uint32 VisibilityMap, int32 FrameNumber) = 0;
+	virtual void SetDynamicData_RenderThread(FNiagaraDynamicDataBase* NewDynamicData) = 0;
+	virtual void DrawDynamicElements(FPrimitiveDrawInterface* PDI, const FSceneView* View) = 0;
+	virtual void CreateRenderThreadResources() = 0;
+	virtual void ReleaseRenderThreadResources() = 0;
+	virtual FNiagaraDynamicDataBase *GenerateVertexData(int NumParticles, const TArray<FVector4> &Particles) = 0;
+	virtual int GetDynamicDataSize() = 0;
+
+	virtual bool HasDynamicData() = 0;
+
+	virtual FPrimitiveViewRelevance GetViewRelevance(const FSceneView* View)
+	{
+		FPrimitiveViewRelevance Result;
+		bool bHasDynamicData = HasDynamicData();
+
+		Result.bDrawRelevance = bHasDynamicData && SceneProxy->IsShown(View) && View->Family->EngineShowFlags.Particles;
+		Result.bShadowRelevance = bHasDynamicData && SceneProxy->IsShadowCast(View);
+		Result.bDynamicRelevance = bHasDynamicData;
+		Result.bNeedsPreRenderView = Result.bDrawRelevance || Result.bShadowRelevance;
+		if (bHasDynamicData && View->Family->EngineShowFlags.Bounds)
+		{
+			Result.bOpaqueRelevance = true;
+		}
+		MaterialRelevance.SetPrimitiveViewRelevance(Result);
+
+		return Result;
+	}
+
+
+
+	~NiagaraEffectRenderer() {}
+
+protected:
+	NiagaraEffectRenderer(const UNiagaraComponent *InComponent, const FPrimitiveSceneProxy *Proxy)	
+	{
+		SceneProxy = static_cast<const FNiagaraSceneProxy*>(Proxy);
+		Material = InComponent->Material;
+		if (!Material || !Material->CheckMaterialUsage(MATUSAGE_ParticleSprites))
+		{
+			Material = UMaterial::GetDefaultMaterial(MD_Surface);
+		}
+		check(Material);
+		MaterialRelevance = Material->GetRelevance();
+	}
+
+	const FNiagaraSceneProxy *SceneProxy;
+	UMaterialInterface* Material;
+
+private:
+	FMaterialRelevance MaterialRelevance;
+};
+
+
+
+class NiagaraEffectRendererSprites : public NiagaraEffectRenderer
+{
+public:	
+	NiagaraEffectRendererSprites(const UNiagaraComponent* InComponent, const FPrimitiveSceneProxy *Proxy) :
+		NiagaraEffectRenderer(InComponent, Proxy),
+		VertexFactory(PVFT_Sprite, InComponent->GetWorld()->FeatureLevel),
+		DynamicDataRender(NULL)
+	{
+	}
+
+	~NiagaraEffectRendererSprites()
+	{
+		ReleaseRenderThreadResources();
+	}
+
+
+	virtual void ReleaseRenderThreadResources() override
+	{
+		VertexFactory.ReleaseResource();
+		PerViewUniformBuffers.Empty();
+		WorldSpacePrimitiveUniformBuffer.ReleaseResource();
+	}
+
+	// FPrimitiveSceneProxy interface.
+	virtual void CreateRenderThreadResources() override
+	{
+		VertexFactory.InitResource();
+
+		UniformParameters.AxisLockRight = FVector4(0.0f, 0.0f, 0.0f, 0.0f);
+		UniformParameters.AxisLockUp = FVector4(0.0f, 0.0f, 0.0f, 0.0f);
+		UniformParameters.RotationScale = 1.0f;
+		UniformParameters.RotationBias = 0.0f;
+		UniformParameters.TangentSelector = FVector4(0.0f, 0.0f, 0.0f, 1.0f);
+		UniformParameters.InvDeltaSeconds = 0.0f;
+		UniformParameters.SubImageSize = FVector4(1.0f, 1.0f, 1.0f, 1.0f);
+		UniformParameters.NormalsType = 0;
+		UniformParameters.NormalsSphereCenter = FVector4(0.0f, 0.0f, 0.0f, 1.0f);
+		UniformParameters.NormalsCylinderUnitDirection = FVector4(0.0f, 0.0f, 1.0f, 0.0f);
+		UniformParameters.PivotOffset = FVector2D(0.0f, 0.0f);
+	}
+
+
+	virtual void PreRenderView(const FSceneViewFamily* ViewFamily, const uint32 VisibilityMap, int32 FrameNumber) override
+	{
+		SCOPE_CYCLE_COUNTER(STAT_NiagaraPreRenderView);
+
+		check(DynamicDataRender && DynamicDataRender->VertexData.Num());
+
+		int32 SizeInBytes = DynamicDataRender->VertexData.GetTypeSize() * DynamicDataRender->VertexData.Num();
+		DynamicVertexAllocation = FGlobalDynamicVertexBuffer::Get().Allocate(SizeInBytes);
+		if (DynamicVertexAllocation.IsValid())
+		{
+			// Copy the vertex data over.
+			FMemory::Memcpy(DynamicVertexAllocation.Buffer, DynamicDataRender->VertexData.GetData(), SizeInBytes);
+			VertexFactory.SetInstanceBuffer(
+				DynamicVertexAllocation.VertexBuffer,
+				DynamicVertexAllocation.VertexOffset,
+				sizeof(FParticleSpriteVertex),
+				true
+				);
+			VertexFactory.SetDynamicParameterBuffer(NULL, 0, 0, true);
+			// Compute the per-view uniform buffers.
+			const int32 NumViews = ViewFamily->Views.Num();
+			uint32 ViewBit = 1;
+			for (int32 ViewIndex = 0; ViewIndex < NumViews; ++ViewIndex, ViewBit <<= 1)
+			{
+				FParticleSpriteUniformBufferRef* SpriteViewUniformBufferPtr = new(PerViewUniformBuffers)FParticleSpriteUniformBufferRef();
+				if (VisibilityMap & ViewBit)
+				{
+					FParticleSpriteUniformParameters PerViewUniformParameters = UniformParameters;
+					PerViewUniformParameters.MacroUVParameters = FVector4(0.0f, 0.0f, 1.0f, 1.0f);
+					*SpriteViewUniformBufferPtr = FParticleSpriteUniformBufferRef::CreateUniformBufferImmediate(PerViewUniformParameters, UniformBuffer_SingleFrame);
+				}
+			}
+
+			// Update the primitive uniform buffer if needed.
+			if (!WorldSpacePrimitiveUniformBuffer.IsInitialized())
+			{
+				FPrimitiveUniformShaderParameters PrimitiveUniformShaderParameters = GetPrimitiveUniformShaderParameters(
+					FMatrix::Identity,
+					SceneProxy->GetActorPosition(),
+					SceneProxy->GetBounds(),
+					SceneProxy->GetLocalBounds(),
+					SceneProxy->ReceivesDecals(),
+					false,
+					false, 
+					false
+					);
+				WorldSpacePrimitiveUniformBuffer.SetContents(PrimitiveUniformShaderParameters);
+				WorldSpacePrimitiveUniformBuffer.InitResource();
+			}
+		}
+	}
+
+
+	virtual void DrawDynamicElements(FPrimitiveDrawInterface* PDI, const FSceneView* View) override
+	{
+		SCOPE_CYCLE_COUNTER(STAT_NiagaraRender);
+
+		check(DynamicDataRender && DynamicDataRender->VertexData.Num());
+
+		const bool bIsWireframe = View->Family->EngineShowFlags.Wireframe;
+		FMaterialRenderProxy* MaterialRenderProxy = Material->GetRenderProxy(SceneProxy->IsSelected(), SceneProxy->IsHovered());
+
+		if (DynamicVertexAllocation.IsValid()
+			&& (bIsWireframe || !PDI->IsMaterialIgnored(MaterialRenderProxy, View->GetFeatureLevel())))
+		{
+			SCOPED_DRAW_EVENT(NiagaraRender, DEC_SCENE_ITEMS);
+
+			FMeshBatch MeshBatch;
+			MeshBatch.VertexFactory = &VertexFactory;
+			MeshBatch.CastShadow = SceneProxy->CastsDynamicShadow();
+			MeshBatch.bUseAsOccluder = false;
+			MeshBatch.ReverseCulling = SceneProxy->IsLocalToWorldDeterminantNegative();
+			MeshBatch.Type = PT_TriangleList;
+			MeshBatch.DepthPriorityGroup = SceneProxy->GetDepthPriorityGroup(View);
+			MeshBatch.LCI = NULL;
+			if (bIsWireframe)
+			{
+				MeshBatch.MaterialRenderProxy = UMaterial::GetDefaultMaterial(MD_Surface)->GetRenderProxy(SceneProxy->IsSelected(), SceneProxy->IsHovered());
+			}
+			else
+			{
+				MeshBatch.MaterialRenderProxy = MaterialRenderProxy;
+			}
+
+			FMeshBatchElement& MeshElement = MeshBatch.Elements[0];
+			MeshElement.IndexBuffer = &GParticleIndexBuffer;
+			MeshElement.FirstIndex = 0;
+			MeshElement.NumPrimitives = 2;
+			MeshElement.NumInstances = DynamicDataRender->VertexData.Num();
+			MeshElement.MinVertexIndex = 0;
+			MeshElement.MaxVertexIndex = MeshElement.NumInstances * 4 - 1;
+			MeshElement.PrimitiveUniformBufferResource = &WorldSpacePrimitiveUniformBuffer;
+
+			int32 ViewIndex = View->Family->Views.Find(View);
+			check(PerViewUniformBuffers.IsValidIndex(ViewIndex) && PerViewUniformBuffers[ViewIndex]);
+			VertexFactory.SetSpriteUniformBuffer(PerViewUniformBuffers[ViewIndex]);
+
+			DrawRichMesh(
+				PDI,
+				MeshBatch,
+				FLinearColor(1.0f, 0.0f, 0.0f),	//WireframeColor,
+				FLinearColor(1.0f, 1.0f, 0.0f),	//LevelColor,
+				FLinearColor(1.0f, 1.0f, 1.0f),	//PropertyColor,		
+				SceneProxy,
+				GIsEditor && (View->Family->EngineShowFlags.Selection) ? SceneProxy->IsSelected() : false,
+				bIsWireframe
+				);
+		}
+
+		/*
+		for(int32 i=0; i<DynamicDataRender->VertexData.Num(); i++)
+		{
+		FParticleSpriteVertex& Vertex = DynamicDataRender->VertexData[i];
+		PDI->DrawPoint(Vertex.Position, Vertex.Color, 1.0f, SDPG_World);
+		}
+		*/
+	}
+
+
+	virtual void GetDynamicMeshElements(const TArray<const FSceneView*>& Views, const FSceneViewFamily& ViewFamily, uint32 VisibilityMap, FMeshElementCollector& Collector) const override
+	{
+		SCOPE_CYCLE_COUNTER(STAT_NiagaraRender);
+		SCOPE_CYCLE_COUNTER(STAT_NiagaraRenderSprites);
+
+		check(DynamicDataRender && DynamicDataRender->VertexData.Num());
+
+		const bool bIsWireframe = ViewFamily.EngineShowFlags.Wireframe;
+		FMaterialRenderProxy* MaterialRenderProxy = Material->GetRenderProxy(SceneProxy->IsSelected(), SceneProxy->IsHovered());
+
+		int32 SizeInBytes = DynamicDataRender->VertexData.GetTypeSize() * DynamicDataRender->VertexData.Num();
+		FGlobalDynamicVertexBuffer::FAllocation LocalDynamicVertexAllocation = FGlobalDynamicVertexBuffer::Get().Allocate(SizeInBytes);
+
+		if (LocalDynamicVertexAllocation.IsValid())
+		{
+			// Update the primitive uniform buffer if needed.
+			if (!WorldSpacePrimitiveUniformBuffer.IsInitialized())
+			{
+				FPrimitiveUniformShaderParameters PrimitiveUniformShaderParameters = GetPrimitiveUniformShaderParameters(
+					FMatrix::Identity,
+					SceneProxy->GetActorPosition(),
+					SceneProxy->GetBounds(),
+					SceneProxy->GetLocalBounds(),
+					SceneProxy->ReceivesDecals(),
+					false,
+					SceneProxy->UseEditorDepthTest()
+					);
+				WorldSpacePrimitiveUniformBuffer.SetContents(PrimitiveUniformShaderParameters);
+				WorldSpacePrimitiveUniformBuffer.InitResource();
+			}
+
+			// Copy the vertex data over.
+			FMemory::Memcpy(LocalDynamicVertexAllocation.Buffer, DynamicDataRender->VertexData.GetData(), SizeInBytes);
+
+			// Compute the per-view uniform buffers.
+			for (int32 ViewIndex = 0; ViewIndex < Views.Num(); ViewIndex++)
+			{
+				if (VisibilityMap & (1 << ViewIndex))
+				{
+					const FSceneView* View = Views[ViewIndex];
+
+					FNiagaraMeshCollectorResourcesSprite& CollectorResources = Collector.AllocateOneFrameResource<FNiagaraMeshCollectorResourcesSprite>();
+					FParticleSpriteUniformParameters PerViewUniformParameters = UniformParameters;
+
+					// Collector.AllocateOneFrameResource uses default ctor, initialize the vertex factory
+					CollectorResources.VertexFactory.SetFeatureLevel(ViewFamily.GetFeatureLevel());
+					CollectorResources.VertexFactory.SetParticleFactoryType(PVFT_Sprite);
+
+					PerViewUniformParameters.MacroUVParameters = FVector4(0.0f, 0.0f, 1.0f, 1.0f);
+					CollectorResources.UniformBuffer = FParticleSpriteUniformBufferRef::CreateUniformBufferImmediate(PerViewUniformParameters, UniformBuffer_SingleFrame);
+
+					CollectorResources.VertexFactory.InitResource();
+					CollectorResources.VertexFactory.SetSpriteUniformBuffer(CollectorResources.UniformBuffer);
+					CollectorResources.VertexFactory.SetInstanceBuffer(
+						LocalDynamicVertexAllocation.VertexBuffer,
+						LocalDynamicVertexAllocation.VertexOffset,
+						sizeof(FParticleSpriteVertex),
+						true
+						);
+					CollectorResources.VertexFactory.SetDynamicParameterBuffer(NULL, 0, 0, true);
+
+					FMeshBatch& MeshBatch = Collector.AllocateMesh();
+					MeshBatch.VertexFactory = &CollectorResources.VertexFactory;
+					MeshBatch.CastShadow = SceneProxy->CastsDynamicShadow();
+					MeshBatch.bUseAsOccluder = false;
+					MeshBatch.ReverseCulling = SceneProxy->IsLocalToWorldDeterminantNegative();
+					MeshBatch.Type = PT_TriangleList;
+					MeshBatch.DepthPriorityGroup = SceneProxy->GetDepthPriorityGroup(View);
+					MeshBatch.bCanApplyViewModeOverrides = true;
+					MeshBatch.bUseWireframeSelectionColoring = SceneProxy->IsSelected();
+
+					if (bIsWireframe)
+					{
+						MeshBatch.MaterialRenderProxy = UMaterial::GetDefaultMaterial(MD_Surface)->GetRenderProxy(SceneProxy->IsSelected(), SceneProxy->IsHovered());
+					}
+					else
+					{
+						MeshBatch.MaterialRenderProxy = MaterialRenderProxy;
+					}
+
+					FMeshBatchElement& MeshElement = MeshBatch.Elements[0];
+					MeshElement.IndexBuffer = &GParticleIndexBuffer;
+					MeshElement.FirstIndex = 0;
+					MeshElement.NumPrimitives = 2;
+					MeshElement.NumInstances = DynamicDataRender->VertexData.Num();
+					MeshElement.MinVertexIndex = 0;
+					MeshElement.MaxVertexIndex = MeshElement.NumInstances * 4 - 1;
+					MeshElement.PrimitiveUniformBufferResource = &WorldSpacePrimitiveUniformBuffer;
+
+					Collector.AddMesh(ViewIndex, MeshBatch);
+				}
+			}
+		}
+	}
+
+
+	/** Update render data buffer from attributes */
+	FNiagaraDynamicDataBase *GenerateVertexData(int NumParticles, const TArray<FVector4> &ParticleArray) override
+	{
+		FNiagaraDynamicDataSprites *DynamicData = new FNiagaraDynamicDataSprites;
+		TArray<FParticleSpriteVertex>& RenderData = DynamicData->VertexData;
+
+		RenderData.Reset(NumParticles);
+		//CachedBounds.Init();
+
+		const FVector4* Particles = ParticleArray.GetTypedData();
+		int32 AttrStride = NumParticles;
+		for (int32 ParticleIndex = 0; ParticleIndex < NumParticles; ParticleIndex++)
+		{
+			const FVector4* Particle = Particles + ParticleIndex;
+
+			FParticleSpriteVertex& NewVertex = *new(RenderData)FParticleSpriteVertex;
+			NewVertex.Position = Particle[AttrStride * 0];
+			NewVertex.OldPosition = NewVertex.Position;
+			//CachedBounds += NewVertex.Position;
+			NewVertex.Color = FLinearColor(Particle[AttrStride * 2]);
+			NewVertex.ParticleId = static_cast<float>(ParticleIndex) / NumParticles;
+			NewVertex.RelativeTime = Particle[AttrStride*4].X;
+			NewVertex.Size = FVector2D(1.0f, 1.0f);
+			NewVertex.Rotation = Particle[AttrStride*3].X;
+			NewVertex.SubImageIndex = 0.f;
+		}
+		//CachedBounds.ExpandBy(MaxSize);
+
+		return DynamicData;
+	}
+
+
+
+	virtual void SetDynamicData_RenderThread(FNiagaraDynamicDataBase* NewDynamicData) override
+	{
+		check(IsInRenderingThread());
+
+		if (DynamicDataRender)
+		{
+			delete DynamicDataRender;
+			DynamicDataRender = NULL;
+		}
+		DynamicDataRender = static_cast<FNiagaraDynamicDataSprites*>(NewDynamicData);
+	}
+
+	int GetDynamicDataSize()
+	{
+		uint32 Size = sizeof(FNiagaraDynamicDataSprites);
+		if (DynamicDataRender)
+		{
+			Size += DynamicDataRender->VertexData.GetAllocatedSize();
+		}
+
+		return Size;
+	}
+
+	bool HasDynamicData()
+	{
+		return DynamicDataRender && DynamicDataRender->VertexData.Num() > 0;
+	}
+
+private:
+	FNiagaraDynamicDataSprites *DynamicDataRender;
+	mutable TUniformBuffer<FPrimitiveUniformShaderParameters> WorldSpacePrimitiveUniformBuffer;
+	FParticleSpriteVertexFactory VertexFactory;
+	FParticleSpriteUniformParameters UniformParameters;
+	TArray<FParticleSpriteUniformBufferRef, TInlineAllocator<2> > PerViewUniformBuffers;
+	FGlobalDynamicVertexBuffer::FAllocation DynamicVertexAllocation;
+};
+
+
+
+
+
+
+
+
+
+
+class NiagaraEffectRendererRibbon : public NiagaraEffectRenderer
+{
+public:
+	NiagaraEffectRendererRibbon(const UNiagaraComponent* InComponent, FPrimitiveSceneProxy *Proxy) :
+		NiagaraEffectRenderer(InComponent, Proxy),
+		VertexFactory(PVFT_BeamTrail, InComponent->GetWorld()->FeatureLevel),
+		DynamicDataRender(NULL)
+	{
+	}
+
+	~NiagaraEffectRendererRibbon()
+	{
+		ReleaseRenderThreadResources();
+	}
+
+
+	virtual void ReleaseRenderThreadResources() override
+	{
+		VertexFactory.ReleaseResource();
+		PerViewUniformBuffers.Empty();
+		WorldSpacePrimitiveUniformBuffer.ReleaseResource();
+	}
+
+	// FPrimitiveSceneProxy interface.
+	virtual void CreateRenderThreadResources() override
+	{
+		VertexFactory.InitResource();
+		UniformParameters.CameraUp = FVector4(0.0f, 0.0f, 1.0f, 0.0f);
+		UniformParameters.CameraRight = FVector4(1.0f, 0.0f, 0.0f, 0.0f);
+		UniformParameters.ScreenAlignment = FVector4(0.0f, 0.0f, 0.0f, 0.0f);
+	}
+
+	virtual void GetDynamicMeshElements(const TArray<const FSceneView*>& Views, const FSceneViewFamily& ViewFamily, uint32 VisibilityMap, FMeshElementCollector& Collector) const override
+	{
+		SCOPE_CYCLE_COUNTER(STAT_NiagaraRender);
+		SCOPE_CYCLE_COUNTER(STAT_NiagaraRenderRibbons);
+
+		check(DynamicDataRender && DynamicDataRender->VertexData.Num());
+
+		const bool bIsWireframe = ViewFamily.EngineShowFlags.Wireframe;
+		FMaterialRenderProxy* MaterialRenderProxy = Material->GetRenderProxy(SceneProxy->IsSelected(), SceneProxy->IsHovered());
+
+		int32 SizeInBytes = DynamicDataRender->VertexData.GetTypeSize() * DynamicDataRender->VertexData.Num();
+		FGlobalDynamicVertexBuffer::FAllocation LocalDynamicVertexAllocation = FGlobalDynamicVertexBuffer::Get().Allocate(SizeInBytes);
+
+		if (LocalDynamicVertexAllocation.IsValid())
+		{
+			// Update the primitive uniform buffer if needed.
+			if (!WorldSpacePrimitiveUniformBuffer.IsInitialized())
+			{
+				FPrimitiveUniformShaderParameters PrimitiveUniformShaderParameters = GetPrimitiveUniformShaderParameters(
+					FMatrix::Identity,
+					SceneProxy->GetActorPosition(),
+					SceneProxy->GetBounds(),
+					SceneProxy->GetLocalBounds(),
+					SceneProxy->ReceivesDecals(),
+					false,
+					SceneProxy->UseEditorDepthTest()
+					);
+				WorldSpacePrimitiveUniformBuffer.SetContents(PrimitiveUniformShaderParameters);
+				WorldSpacePrimitiveUniformBuffer.InitResource();
+			}
+
+			// Copy the vertex data over.
+			FMemory::Memcpy(LocalDynamicVertexAllocation.Buffer, DynamicDataRender->VertexData.GetData(), SizeInBytes);
+
+			// Compute the per-view uniform buffers.
+			for (int32 ViewIndex = 0; ViewIndex < Views.Num(); ViewIndex++)
+			{
+				if (VisibilityMap & (1 << ViewIndex))
+				{
+					const FSceneView* View = Views[ViewIndex];
+
+					FNiagaraMeshCollectorResourcesRibbon& CollectorResources = Collector.AllocateOneFrameResource<FNiagaraMeshCollectorResourcesRibbon>();
+					FParticleBeamTrailUniformParameters PerViewUniformParameters = UniformParameters;
+
+					// Collector.AllocateOneFrameResource uses default ctor, initialize the vertex factory
+					CollectorResources.VertexFactory.SetFeatureLevel(ViewFamily.GetFeatureLevel());
+					CollectorResources.VertexFactory.SetParticleFactoryType(PVFT_BeamTrail);
+
+					CollectorResources.UniformBuffer = FParticleBeamTrailUniformBufferRef::CreateUniformBufferImmediate(PerViewUniformParameters, UniformBuffer_SingleFrame);
+
+					CollectorResources.VertexFactory.InitResource();
+					CollectorResources.VertexFactory.SetBeamTrailUniformBuffer(PerViewUniformBuffers[ViewIndex]);
+					CollectorResources.VertexFactory.SetVertexBuffer(DynamicVertexAllocation.VertexBuffer, DynamicVertexAllocation.VertexOffset, sizeof(FParticleBeamTrailVertex));
+					CollectorResources.VertexFactory.SetDynamicParameterBuffer(NULL, 0, 0);
+
+					FMeshBatch& MeshBatch = Collector.AllocateMesh();
+					MeshBatch.VertexFactory = &CollectorResources.VertexFactory;
+					MeshBatch.CastShadow = SceneProxy->CastsDynamicShadow();
+					MeshBatch.bUseAsOccluder = false;
+					MeshBatch.ReverseCulling = SceneProxy->IsLocalToWorldDeterminantNegative();
+					MeshBatch.bDisableBackfaceCulling = true;
+					MeshBatch.Type = PT_TriangleStrip;
+					MeshBatch.DepthPriorityGroup = SceneProxy->GetDepthPriorityGroup(View);
+					MeshBatch.bCanApplyViewModeOverrides = true;
+					MeshBatch.bUseWireframeSelectionColoring = SceneProxy->IsSelected();
+
+					if (bIsWireframe)
+					{
+						MeshBatch.MaterialRenderProxy = UMaterial::GetDefaultMaterial(MD_Surface)->GetRenderProxy(SceneProxy->IsSelected(), SceneProxy->IsHovered());
+					}
+					else
+					{
+						MeshBatch.MaterialRenderProxy = MaterialRenderProxy;
+					}
+
+					FMeshBatchElement& MeshElement = MeshBatch.Elements[0];
+					MeshElement.IndexBuffer = &GParticleIndexBuffer;
+					MeshElement.FirstIndex = 0;
+					MeshElement.NumPrimitives = (DynamicDataRender->VertexData.Num()-1) / 2;
+					MeshElement.NumInstances = 1;
+					MeshElement.MinVertexIndex = 0;
+					MeshElement.MaxVertexIndex = DynamicDataRender->VertexData.Num() - 1;
+					MeshElement.PrimitiveUniformBufferResource = &WorldSpacePrimitiveUniformBuffer;
+
+					Collector.AddMesh(ViewIndex, MeshBatch);
+				}
+			}
+		}
+	}
+
+	virtual void PreRenderView(const FSceneViewFamily* ViewFamily, const uint32 VisibilityMap, int32 FrameNumber) override
+	{
+		SCOPE_CYCLE_COUNTER(STAT_NiagaraPreRenderView);
+
+		check(DynamicDataRender && DynamicDataRender->VertexData.Num());
+
+		int32 SizeInBytes = DynamicDataRender->VertexData.GetTypeSize() * DynamicDataRender->VertexData.Num();
+		DynamicVertexAllocation = FGlobalDynamicVertexBuffer::Get().Allocate(SizeInBytes);
+		if (DynamicVertexAllocation.IsValid())
+		{
+			// Copy the vertex data over.
+			FMemory::Memcpy(DynamicVertexAllocation.Buffer, DynamicDataRender->VertexData.GetData(), SizeInBytes);
+			VertexFactory.SetVertexBuffer(DynamicVertexAllocation.VertexBuffer, DynamicVertexAllocation.VertexOffset, sizeof(FParticleBeamTrailVertex));
+			VertexFactory.SetDynamicParameterBuffer(NULL, 0, 0);
+			// Compute the per-view uniform buffers.
+			const int32 NumViews = ViewFamily->Views.Num();
+			uint32 ViewBit = 1;
+			for (int32 ViewIndex = 0; ViewIndex < NumViews; ++ViewIndex, ViewBit <<= 1)
+			{
+				FParticleBeamTrailUniformBufferRef* SpriteViewUniformBufferPtr = new(PerViewUniformBuffers)FParticleBeamTrailUniformBufferRef();
+				if (VisibilityMap & ViewBit)
+				{
+					FParticleBeamTrailUniformParameters PerViewUniformParameters = UniformParameters;
+					*SpriteViewUniformBufferPtr = FParticleBeamTrailUniformBufferRef::CreateUniformBufferImmediate(PerViewUniformParameters, UniformBuffer_SingleFrame);
+				}
+			}
+
+			// Update the primitive uniform buffer if needed.
+			if (!WorldSpacePrimitiveUniformBuffer.IsInitialized())
+			{
+				FPrimitiveUniformShaderParameters PrimitiveUniformShaderParameters = GetPrimitiveUniformShaderParameters(
+					FMatrix::Identity,
+					SceneProxy->GetActorPosition(),
+					SceneProxy->GetBounds(),
+					SceneProxy->GetLocalBounds(),
+					SceneProxy->ReceivesDecals(),
+					false,
+					false,
+					false
+					);
+				WorldSpacePrimitiveUniformBuffer.SetContents(PrimitiveUniformShaderParameters);
+				WorldSpacePrimitiveUniformBuffer.InitResource();
+			}
+		}
+	}
+
+
+	virtual void DrawDynamicElements(FPrimitiveDrawInterface* PDI, const FSceneView* View) override
+	{
+		SCOPE_CYCLE_COUNTER(STAT_NiagaraRender);
+
+		check(DynamicDataRender && DynamicDataRender->VertexData.Num());
+
+		const bool bIsWireframe = View->Family->EngineShowFlags.Wireframe;
+		FMaterialRenderProxy* MaterialRenderProxy = Material->GetRenderProxy(SceneProxy->IsSelected(), SceneProxy->IsHovered());
+
+		if (DynamicVertexAllocation.IsValid()
+			&& (bIsWireframe || !PDI->IsMaterialIgnored(MaterialRenderProxy, View->GetFeatureLevel())))
+		{
+			SCOPED_DRAW_EVENT(NiagaraRender, DEC_SCENE_ITEMS);
+
+			FMeshBatch MeshBatch;
+			MeshBatch.VertexFactory = &VertexFactory;
+			MeshBatch.CastShadow = SceneProxy->CastsDynamicShadow();
+			MeshBatch.bUseAsOccluder = false;
+			MeshBatch.bDisableBackfaceCulling = true;
+			MeshBatch.ReverseCulling = SceneProxy->IsLocalToWorldDeterminantNegative();
+			MeshBatch.Type = PT_TriangleStrip;
+			MeshBatch.DepthPriorityGroup = SceneProxy->GetDepthPriorityGroup(View);
+			MeshBatch.LCI = NULL;
+			if (bIsWireframe)
+			{
+				MeshBatch.MaterialRenderProxy = UMaterial::GetDefaultMaterial(MD_Surface)->GetRenderProxy(SceneProxy->IsSelected(), SceneProxy->IsHovered());
+			}
+			else
+			{
+				MeshBatch.MaterialRenderProxy = MaterialRenderProxy;
+			}
+
+			FMeshBatchElement& MeshElement = MeshBatch.Elements[0];
+			MeshElement.IndexBuffer = &GParticleIndexBuffer;
+			MeshElement.FirstIndex = 0;
+			MeshElement.NumPrimitives = DynamicDataRender->VertexData.Num()/2;
+			MeshElement.NumInstances = 1;
+			MeshElement.MinVertexIndex = 0;
+			MeshElement.MaxVertexIndex = DynamicDataRender->VertexData.Num() - 1;
+			MeshElement.PrimitiveUniformBufferResource = &WorldSpacePrimitiveUniformBuffer;
+
+			int32 ViewIndex = View->Family->Views.Find(View);
+			check(PerViewUniformBuffers.IsValidIndex(ViewIndex) && PerViewUniformBuffers[ViewIndex]);
+			VertexFactory.SetBeamTrailUniformBuffer(PerViewUniformBuffers[ViewIndex]);
+
+			DrawRichMesh(
+				PDI,
+				MeshBatch,
+				FLinearColor(1.0f, 0.0f, 0.0f),	//WireframeColor,
+				FLinearColor(1.0f, 1.0f, 0.0f),	//LevelColor,
+				FLinearColor(1.0f, 1.0f, 1.0f),	//PropertyColor,		
+				SceneProxy,
+				GIsEditor && (View->Family->EngineShowFlags.Selection) ? SceneProxy->IsSelected() : false,
+				bIsWireframe
+				);
+		}
+
+		/*
+		for(int32 i=0; i<DynamicDataRender->VertexData.Num(); i++)
+		{
+		FParticleSpriteVertex& Vertex = DynamicDataRender->VertexData[i];
+		PDI->DrawPoint(Vertex.Position, Vertex.Color, 1.0f, SDPG_World);
+		}
+		*/
+	}
+
+
+
+	/** Update render data buffer from attributes */
+	FNiagaraDynamicDataBase *GenerateVertexData(int NumParticles, const TArray<FVector4> &ParticleArray) override
+	{
+		FNiagaraDynamicDataRibbon *DynamicData = new FNiagaraDynamicDataRibbon;
+		TArray<FParticleBeamTrailVertex>& RenderData = DynamicData->VertexData;
+
+		RenderData.Reset(NumParticles*2);
+		//CachedBounds.Init();
+
+		const FVector4* Particles = ParticleArray.GetTypedData();
+		int32 AttrStride = NumParticles;
+
+
+		// build a sorted list by age, so we always get particles in order 
+		// regardless of them being moved around due to dieing and spawning
+		TArray<int32> SortedIndices;
+		for (int32 Idx = 0; Idx < NumParticles; Idx++)
+		{
+			SortedIndices.Add(Idx);
+		}
+
+		SortedIndices.Sort(
+			[&Particles, AttrStride](const int32& A, const int32& B) {
+				return Particles[A + AttrStride * 4].X < Particles[B + AttrStride * 4].X;
+				}
+		);
+
+
+		FVector2D UVs[4] = { FVector2D(0.0f, 0.0f), FVector2D(1.0f, 0.0f), FVector2D(1.0f, 1.0f), FVector2D(0.0f, 1.0f) };
+		
+		FVector PrevPos, PrevPos2, PrevDir(0.0f, 0.0f, 0.1f);
+		for (int32 i = 0; i<SortedIndices.Num() - 1; i++)
+		{
+			const FVector4* Particle = Particles + SortedIndices[i];
+			const FVector4* Particle2 = Particles + SortedIndices[i+1];
+
+			const FVector ParticlePos = Particle[AttrStride * 0];
+			FVector ParticleDir = Particle2[AttrStride * 0] - ParticlePos;
+			if (ParticleDir.Size()<=SMALL_NUMBER)
+			{
+				ParticleDir = PrevDir*0.1f;
+			}
+
+			FVector ParticleRight = FVector::CrossProduct(ParticleDir, FVector(0.0f, 0.0f, 1.0f));
+			ParticleRight.Normalize();
+
+			if (i == 0)
+			{
+				AddRibbonVert(RenderData, ParticlePos + ParticleRight, Particle, AttrStride, UVs[0]);
+				AddRibbonVert(RenderData, ParticlePos - ParticleRight, Particle, AttrStride, UVs[1]);
+			}
+			else
+			{
+				AddRibbonVert(RenderData, PrevPos2, Particle, AttrStride, UVs[0]);
+				AddRibbonVert(RenderData, PrevPos, Particle, AttrStride, UVs[1]);
+			}
+
+			AddRibbonVert(RenderData, ParticlePos - ParticleRight + ParticleDir, Particle, AttrStride, UVs[2]);
+			AddRibbonVert(RenderData, ParticlePos + ParticleRight + ParticleDir, Particle, AttrStride, UVs[3]);
+			PrevPos = ParticlePos - ParticleRight + ParticleDir;
+			PrevPos2 = ParticlePos + ParticleRight + ParticleDir;
+			PrevDir = ParticleDir;
+		}
+
+		return DynamicData;
+	}
+
+
+	void AddRibbonVert(TArray<FParticleBeamTrailVertex>& RenderData, FVector ParticlePos, const FVector4 *Particle, int32 AttrStride, FVector2D UV1)
+	{
+		FParticleBeamTrailVertex& NewVertex = *new(RenderData)FParticleBeamTrailVertex;
+		NewVertex.Position = ParticlePos;
+		NewVertex.OldPosition = NewVertex.Position;
+		NewVertex.Color = FLinearColor(Particle[AttrStride * 2]);
+		NewVertex.ParticleId = 0;
+		NewVertex.RelativeTime = Particle[AttrStride * 4].X;
+		NewVertex.Size = FVector2D(1.0f, 1.0f);
+		NewVertex.Rotation = Particle[AttrStride * 3].X;
+		NewVertex.SubImageIndex = 0.f;
+		NewVertex.Tex_U = UV1.X;
+		NewVertex.Tex_V = UV1.Y;
+		NewVertex.Tex_U2 = UV1.X;
+		NewVertex.Tex_V2 = UV1.Y;
+	}
+
+
+	virtual void SetDynamicData_RenderThread(FNiagaraDynamicDataBase* NewDynamicData) override
+	{
+		check(IsInRenderingThread());
+
+		if (DynamicDataRender)
+		{
+			delete DynamicDataRender;
+			DynamicDataRender = NULL;
+		}
+		DynamicDataRender = static_cast<FNiagaraDynamicDataRibbon*>(NewDynamicData);
+	}
+
+	int GetDynamicDataSize()
+	{
+		uint32 Size = sizeof(FNiagaraDynamicDataRibbon);
+		if (DynamicDataRender)
+		{
+			Size += DynamicDataRender->VertexData.GetAllocatedSize();
+		}
+
+		return Size;
+	}
+
+	bool HasDynamicData()
+	{
+		return DynamicDataRender && DynamicDataRender->VertexData.Num() > 0;
+	}
+
+
+private:
+	FNiagaraDynamicDataRibbon *DynamicDataRender;
+	mutable TUniformBuffer<FPrimitiveUniformShaderParameters> WorldSpacePrimitiveUniformBuffer;
+	FParticleBeamTrailVertexFactory VertexFactory;
+	FParticleBeamTrailUniformParameters UniformParameters;
+	TArray<FParticleBeamTrailUniformBufferRef, TInlineAllocator<2> > PerViewUniformBuffers;
+	FGlobalDynamicVertexBuffer::FAllocation DynamicVertexAllocation;
+};
+
+
