#ue4 - Fix ANavigationObjectBase::Validate() ignoring BSP geometry because it checked for Hit.GetActor(), which is NULL for BSP.

TTP 343829

diff --git a/Engine/Source/Runtime/Engine/Private/NavigationObjectBase.cpp b/Engine/Source/Runtime/Engine/Private/NavigationObjectBase.cpp
--- a/Engine/Source/Runtime/Engine/Private/NavigationObjectBase.cpp
+++ b/Engine/Source/Runtime/Engine/Private/NavigationObjectBase.cpp
@@ -191,7 +191,7 @@ void ANavigationObjectBase::Validate()
 		const FVector TraceStart = GetActorLocation();
 		const FVector TraceEnd = GetActorLocation() - FVector(0.f,0.f, 4.f * CapsuleComponent->GetScaledCapsuleHalfHeight());
 		GetWorld()->SweepSingle(Hit, TraceStart, TraceEnd, FQuat::Identity, ECC_Pawn, FCollisionShape::MakeBox(Slice), FCollisionQueryParams(NAME_None, false, this));
-		if( Hit.GetActor() != NULL )
+		if( Hit.bBlockingHit )
 		{
 			const FVector HitLocation = TraceStart + (TraceEnd - TraceStart) * Hit.Time;
 			FVector Dest = HitLocation + FVector(0.f,0.f,CapsuleComponent->GetScaledCapsuleHalfHeight()-2.f);
