Fix initialization order

diff --git a/Engine/Source/Runtime/Engine/Classes/Engine/NiagaraEffectRenderer.h b/Engine/Source/Runtime/Engine/Classes/Engine/NiagaraEffectRenderer.h
--- a/Engine/Source/Runtime/Engine/Classes/Engine/NiagaraEffectRenderer.h
+++ b/Engine/Source/Runtime/Engine/Classes/Engine/NiagaraEffectRenderer.h
@@ -161,8 +161,8 @@ class NiagaraEffectRendererSprites : public NiagaraEffectRenderer
 public:	
 	NiagaraEffectRendererSprites(const UNiagaraComponent* InComponent, const FPrimitiveSceneProxy *Proxy) :
 		NiagaraEffectRenderer(InComponent, Proxy),
-		VertexFactory(PVFT_Sprite, InComponent->GetWorld()->FeatureLevel),
-		DynamicDataRender(NULL)
+		DynamicDataRender(NULL),
+		VertexFactory(PVFT_Sprite, InComponent->GetWorld()->FeatureLevel)		
 	{
 	}
 
@@ -497,8 +497,8 @@ class NiagaraEffectRendererRibbon : public NiagaraEffectRenderer
 public:
 	NiagaraEffectRendererRibbon(const UNiagaraComponent* InComponent, FPrimitiveSceneProxy *Proxy) :
 		NiagaraEffectRenderer(InComponent, Proxy),
-		VertexFactory(PVFT_BeamTrail, InComponent->GetWorld()->FeatureLevel),
-		DynamicDataRender(NULL)
+		DynamicDataRender(NULL),
+		VertexFactory(PVFT_BeamTrail, InComponent->GetWorld()->FeatureLevel)
 	{
 	}
 
