Fix build.

diff --git a/Engine/Source/Runtime/Engine/Classes/Engine/NiagaraEffectRenderer.h b/Engine/Source/Runtime/Engine/Classes/Engine/NiagaraEffectRenderer.h
--- a/Engine/Source/Runtime/Engine/Classes/Engine/NiagaraEffectRenderer.h
+++ b/Engine/Source/Runtime/Engine/Classes/Engine/NiagaraEffectRenderer.h
@@ -35,6 +35,7 @@ public:
 
 	FVector4 *GetAttributeData(FName Name)
 	{
+		return nullptr
 	}
 
 	void AddAttribute(FName NewAttrName)
