fixed time units while loging during origin rebasing

diff --git a/Engine/Source/Runtime/Engine/Private/World.cpp b/Engine/Source/Runtime/Engine/Private/World.cpp
--- a/Engine/Source/Runtime/Engine/Private/World.cpp
+++ b/Engine/Source/Runtime/Engine/Private/World.cpp
@@ -3997,7 +3997,7 @@ bool UWorld::SetNewWorldOrigin(FIntVector InNewOriginLocation)
 	const double CurrentTime = FPlatformTime::Seconds();
 	const float TimeTaken = CurrentTime - MoveStartTime;
 	UE_LOG(LogLevel, Log, TEXT("WORLD TRANSLATION END {%d, %d, %d} took %.4f ms"),
-		OriginLocation.X, OriginLocation.Y, OriginLocation.Z, TimeTaken);
+		OriginLocation.X, OriginLocation.Y, OriginLocation.Z, TimeTaken*1000);
 	
 	return true;
 }
