Changes to FRecastRenderingSceneProxy to work with GetDynamicMeshElements and to render navmesh with this function. #ue4

diff --git a/Engine/Source/Developer/GameplayDebugger/GameplayDebugger.Build.cs b/Engine/Source/Developer/GameplayDebugger/GameplayDebugger.Build.cs
--- a/Engine/Source/Developer/GameplayDebugger/GameplayDebugger.Build.cs
+++ b/Engine/Source/Developer/GameplayDebugger/GameplayDebugger.Build.cs
@@ -32,6 +32,7 @@ namespace UnrealBuildTool.Rules
 					"Engine",    
                     "RenderCore",
                     "RHI",
+                    "ShaderCore",
                     "AIModule",  // it have to be here for now. It'll be changed to remove any dependency to AIModule in future
 				}
                 );
diff --git a/Engine/Source/Developer/GameplayDebugger/Private/GameplayDebuggingComponent.cpp b/Engine/Source/Developer/GameplayDebugger/Private/GameplayDebuggingComponent.cpp
--- a/Engine/Source/Developer/GameplayDebugger/Private/GameplayDebuggingComponent.cpp
+++ b/Engine/Source/Developer/GameplayDebugger/Private/GameplayDebuggingComponent.cpp
@@ -735,7 +735,7 @@ namespace NavMeshDebug
 
 	struct FAreaPolys
 	{
-		TArray<int16> Indices;
+		TArray<int32> Indices;
 		FColor Color;
 	};
 
diff --git a/Engine/Source/Runtime/Engine/Private/BatchedElements.cpp b/Engine/Source/Runtime/Engine/Private/BatchedElements.cpp
--- a/Engine/Source/Runtime/Engine/Private/BatchedElements.cpp
+++ b/Engine/Source/Runtime/Engine/Private/BatchedElements.cpp
@@ -221,17 +221,26 @@ void FBatchedElements::AddReserveVertices(int32 NumMeshVerts)
  * Reserves space in line vertex array
  * 
  * @param NumLines - number of lines to reserve space for
+ * @param bDepthBiased - whether reserving depth-biased lines or non-biased lines
+ * @param bThickLines - whether reserving regular lines or thick lines
  */
-void FBatchedElements::AddReserveLines(int32 NumLines, bool bDepthBiased)
+void FBatchedElements::AddReserveLines(int32 NumLines, bool bDepthBiased, bool bThickLines)
 {
-	if (!bDepthBiased)
+	if (!bThickLines)
 	{
-		LineVertices.Reserve( LineVertices.Num() + NumLines * 2 );
+		if (!bDepthBiased)
+		{
+			LineVertices.Reserve(LineVertices.Num() + NumLines * 2);
+		}
+		else
+		{
+			WireTris.Reserve(WireTris.Num() + NumLines);
+			WireTriVerts.Reserve(WireTriVerts.Num() + NumLines * 3);
+		}
 	}
 	else
 	{
-		WireTris.Reserve( WireTris.Num() + NumLines );
-		WireTriVerts.Reserve( WireTriVerts.Num() + NumLines * 3 );
+		ThickLines.Reserve(ThickLines.Num() + NumLines * 2);
 	}
 }
 
diff --git a/Engine/Source/Runtime/Engine/Private/NavMeshRenderingComponent.cpp b/Engine/Source/Runtime/Engine/Private/NavMeshRenderingComponent.cpp
--- a/Engine/Source/Runtime/Engine/Private/NavMeshRenderingComponent.cpp
+++ b/Engine/Source/Runtime/Engine/Private/NavMeshRenderingComponent.cpp
@@ -86,6 +86,9 @@ void UNavMeshRenderingComponent::GatherData(struct FNavMeshSceneProxyData* Curre
 	CurrentData->bEnableDrawing = NavMesh->bEnableDrawing;
 	CurrentData->bNeedsNewData = false;
 
+	static const auto CVarUseGetMeshElements = IConsoleManager::Get().FindTConsoleVariableDataInt(TEXT("r.UseGetMeshElements"));
+	const bool bUseGetMeshElements = !GRHICommandList.Bypass() || CVarUseGetMeshElements->GetValueOnGameThread() != 0;
+
 	if (CurrentData && NavMesh && NavMesh->bEnableDrawing)
 	{
 		FHitProxyId HitProxyId = FHitProxyId();
@@ -117,9 +120,18 @@ void UNavMeshRenderingComponent::GatherData(struct FNavMeshSceneProxyData* Curre
 				const TArray<int32>& MeshIndices = CurrentData->NavMeshGeometry.AreaIndices[AreaIdx];
 				for (int32 Idx=0; Idx<MeshIndices.Num(); Idx += 3)
 				{
-					CurrentData->BatchedElements.AddLine(MeshVerts[MeshIndices[Idx+0]] + CurrentData->NavMeshDrawOffset, MeshVerts[MeshIndices[Idx+1]] + CurrentData->NavMeshDrawOffset, NavMeshRenderColor_Recast_TriangleEdges,HitProxyId,DefaultEdges_LineThickness, 0, true);
-					CurrentData->BatchedElements.AddLine(MeshVerts[MeshIndices[Idx+1]] + CurrentData->NavMeshDrawOffset, MeshVerts[MeshIndices[Idx+2]] + CurrentData->NavMeshDrawOffset, NavMeshRenderColor_Recast_TriangleEdges,HitProxyId,DefaultEdges_LineThickness, 0, true);
-					CurrentData->BatchedElements.AddLine(MeshVerts[MeshIndices[Idx+2]] + CurrentData->NavMeshDrawOffset, MeshVerts[MeshIndices[Idx+0]] + CurrentData->NavMeshDrawOffset, NavMeshRenderColor_Recast_TriangleEdges,HitProxyId,DefaultEdges_LineThickness, 0, true);
+					if (bUseGetMeshElements)
+					{
+						CurrentData->ThickLineItems.Add(FNavMeshSceneProxyData::FDebugThickLine(MeshVerts[MeshIndices[Idx + 0]] + CurrentData->NavMeshDrawOffset, MeshVerts[MeshIndices[Idx + 1]] + CurrentData->NavMeshDrawOffset, NavMeshRenderColor_Recast_TriangleEdges, DefaultEdges_LineThickness));
+						CurrentData->ThickLineItems.Add(FNavMeshSceneProxyData::FDebugThickLine(MeshVerts[MeshIndices[Idx + 1]] + CurrentData->NavMeshDrawOffset, MeshVerts[MeshIndices[Idx + 2]] + CurrentData->NavMeshDrawOffset, NavMeshRenderColor_Recast_TriangleEdges, DefaultEdges_LineThickness));
+						CurrentData->ThickLineItems.Add(FNavMeshSceneProxyData::FDebugThickLine(MeshVerts[MeshIndices[Idx + 2]] + CurrentData->NavMeshDrawOffset, MeshVerts[MeshIndices[Idx + 0]] + CurrentData->NavMeshDrawOffset, NavMeshRenderColor_Recast_TriangleEdges, DefaultEdges_LineThickness));
+					}
+					else
+					{
+						CurrentData->BatchedElements.AddLine(MeshVerts[MeshIndices[Idx + 0]] + CurrentData->NavMeshDrawOffset, MeshVerts[MeshIndices[Idx + 1]] + CurrentData->NavMeshDrawOffset, NavMeshRenderColor_Recast_TriangleEdges, HitProxyId, DefaultEdges_LineThickness, 0, true);
+						CurrentData->BatchedElements.AddLine(MeshVerts[MeshIndices[Idx + 1]] + CurrentData->NavMeshDrawOffset, MeshVerts[MeshIndices[Idx + 2]] + CurrentData->NavMeshDrawOffset, NavMeshRenderColor_Recast_TriangleEdges, HitProxyId, DefaultEdges_LineThickness, 0, true);
+						CurrentData->BatchedElements.AddLine(MeshVerts[MeshIndices[Idx + 2]] + CurrentData->NavMeshDrawOffset, MeshVerts[MeshIndices[Idx + 0]] + CurrentData->NavMeshDrawOffset, NavMeshRenderColor_Recast_TriangleEdges, HitProxyId, DefaultEdges_LineThickness, 0, true);
+					}
 				}
 			}
 		}
@@ -156,23 +168,47 @@ void UNavMeshRenderingComponent::GatherData(struct FNavMeshSceneProxyData* Curre
 			FVector UR(OuterBBox.Max.X, OuterBBox.Max.Y, DrawZ);
 			FVector UL(LL.X, UR.Y, DrawZ);
 			FVector LR(UR.X, LL.Y, DrawZ);
-			CurrentData->BatchedElements.AddLine(LL, UL, NavMeshRenderColor_TileBounds,HitProxyId,DefaultEdges_LineThickness, 0, true);
-			CurrentData->BatchedElements.AddLine(UL, UR, NavMeshRenderColor_TileBounds,HitProxyId,DefaultEdges_LineThickness, 0, true);
-			CurrentData->BatchedElements.AddLine(UR, LR, NavMeshRenderColor_TileBounds,HitProxyId,DefaultEdges_LineThickness, 0, true);
-			CurrentData->BatchedElements.AddLine(LR, LL, NavMeshRenderColor_TileBounds,HitProxyId,DefaultEdges_LineThickness, 0, true);
+			if (bUseGetMeshElements)
+			{
+				CurrentData->ThickLineItems.Add(FNavMeshSceneProxyData::FDebugThickLine(LL, UL, NavMeshRenderColor_TileBounds, DefaultEdges_LineThickness));
+				CurrentData->ThickLineItems.Add(FNavMeshSceneProxyData::FDebugThickLine(UL, UR, NavMeshRenderColor_TileBounds, DefaultEdges_LineThickness));
+				CurrentData->ThickLineItems.Add(FNavMeshSceneProxyData::FDebugThickLine(UR, LR, NavMeshRenderColor_TileBounds, DefaultEdges_LineThickness));
+				CurrentData->ThickLineItems.Add(FNavMeshSceneProxyData::FDebugThickLine(LR, LL, NavMeshRenderColor_TileBounds, DefaultEdges_LineThickness));
+			}
+			else
+			{
+				CurrentData->BatchedElements.AddLine(LL, UL, NavMeshRenderColor_TileBounds, HitProxyId, DefaultEdges_LineThickness, 0, true);
+				CurrentData->BatchedElements.AddLine(UL, UR, NavMeshRenderColor_TileBounds, HitProxyId, DefaultEdges_LineThickness, 0, true);
+				CurrentData->BatchedElements.AddLine(UR, LR, NavMeshRenderColor_TileBounds, HitProxyId, DefaultEdges_LineThickness, 0, true);
+				CurrentData->BatchedElements.AddLine(LR, LL, NavMeshRenderColor_TileBounds, HitProxyId, DefaultEdges_LineThickness, 0, true);
+			}
 
 			float const TileSizeX = (LL.X - LR.X) / NumTilesX;
 			float const TileSizeY = (LL.Y - UL.Y) / NumTilesY;
 
-			for (int32 Idx=1; Idx<NumTilesX; ++Idx)
+			for (int32 Idx = 1; Idx < NumTilesX; ++Idx)
 			{
-				float const DrawX = FMath::Lerp(LL.X, LR.X, (float)Idx/(float)NumTilesX);
-				CurrentData->BatchedElements.AddLine(FVector(DrawX, LL.Y, DrawZ), FVector(DrawX, UL.Y, DrawZ), NavMeshRenderColor_TileBounds,HitProxyId,DefaultEdges_LineThickness, 0, true);
+				float const DrawX = FMath::Lerp(LL.X, LR.X, (float)Idx / (float)NumTilesX);
+				if (bUseGetMeshElements)
+				{
+					CurrentData->ThickLineItems.Add(FNavMeshSceneProxyData::FDebugThickLine(FVector(DrawX, LL.Y, DrawZ), FVector(DrawX, UL.Y, DrawZ), NavMeshRenderColor_TileBounds, DefaultEdges_LineThickness));
+				}
+				else
+				{
+					CurrentData->BatchedElements.AddLine(FVector(DrawX, LL.Y, DrawZ), FVector(DrawX, UL.Y, DrawZ), NavMeshRenderColor_TileBounds, HitProxyId, DefaultEdges_LineThickness, 0, true);
+				}
 			}
 			for (int32 Idx=1; Idx<NumTilesY; ++Idx)
 			{
 				float const DrawY = FMath::Lerp(LL.Y, UL.Y, (float)Idx/(float)NumTilesY);
-				CurrentData->BatchedElements.AddLine(FVector(LL.X, DrawY, DrawZ), FVector(LR.X, DrawY, DrawZ), NavMeshRenderColor_TileBounds,HitProxyId,DefaultEdges_LineThickness, 0, true);
+				if (bUseGetMeshElements)
+				{
+					CurrentData->ThickLineItems.Add(FNavMeshSceneProxyData::FDebugThickLine(FVector(LL.X, DrawY, DrawZ), FVector(LR.X, DrawY, DrawZ), NavMeshRenderColor_TileBounds, DefaultEdges_LineThickness));
+				}
+				else
+				{
+					CurrentData->BatchedElements.AddLine(FVector(LL.X, DrawY, DrawZ), FVector(LR.X, DrawY, DrawZ), NavMeshRenderColor_TileBounds, HitProxyId, DefaultEdges_LineThickness, 0, true);
+				}
 			}
 		}
 
@@ -328,16 +364,16 @@ void UNavMeshRenderingComponent::GatherData(struct FNavMeshSceneProxyData* Curre
 				}
 
 				FNavMeshSceneProxyData::FDebugMeshData DebugMeshData;
+				DebugMeshData.ClusterColor = GetClusterColor(Idx);
 				for (int32 VertIdx=0; VertIdx < MeshVerts.Num(); ++VertIdx)
 				{
-					AddVertexHelper(DebugMeshData, MeshVerts[VertIdx] + CurrentData->NavMeshDrawOffset);
+					AddVertexHelper(DebugMeshData, MeshVerts[VertIdx] + CurrentData->NavMeshDrawOffset, DebugMeshData.ClusterColor);
 				}
 				for (int32 TriIdx=0; TriIdx < MeshIndices.Num(); TriIdx+=3)
 				{
 					AddTriangleHelper(DebugMeshData, MeshIndices[TriIdx], MeshIndices[TriIdx+1], MeshIndices[TriIdx+2]);
 				}
 
-				DebugMeshData.ClusterColor = GetClusterColor(Idx);
 				CurrentData->MeshBuilders.Add(DebugMeshData);
 			}
 		}
@@ -355,7 +391,7 @@ void UNavMeshRenderingComponent::GatherData(struct FNavMeshSceneProxyData* Curre
 				FNavMeshSceneProxyData::FDebugMeshData DebugMeshData;
 				for (int32 VertIdx=0; VertIdx < MeshVerts.Num(); ++VertIdx)
 				{
-					AddVertexHelper(DebugMeshData, MeshVerts[VertIdx] + CurrentData->NavMeshDrawOffset);
+					AddVertexHelper(DebugMeshData, MeshVerts[VertIdx] + CurrentData->NavMeshDrawOffset, CurrentData->NavMeshColors[AreaType]);
 				}
 				for (int32 TriIdx=0; TriIdx < MeshIndices.Num(); TriIdx+=3)
 				{
@@ -373,7 +409,7 @@ void UNavMeshRenderingComponent::GatherData(struct FNavMeshSceneProxyData* Curre
 			FNavMeshSceneProxyData::FDebugMeshData DebugMeshData;
 			for (int32 VertIdx=0; VertIdx < CurrentData->PathCollidingGeomVerts.Num(); ++VertIdx)
 			{
-				AddVertexHelper(DebugMeshData, CurrentData->PathCollidingGeomVerts[VertIdx]);
+				AddVertexHelper(DebugMeshData, CurrentData->PathCollidingGeomVerts[VertIdx], NavMeshRenderColor_PathCollidingGeom);
 			}
 			DebugMeshData.Indices.Append(CurrentData->PathCollidingGeomIndices);
 			DebugMeshData.ClusterColor = NavMeshRenderColor_PathCollidingGeom;
@@ -385,7 +421,7 @@ void UNavMeshRenderingComponent::GatherData(struct FNavMeshSceneProxyData* Curre
 			FNavMeshSceneProxyData::FDebugMeshData DebugMeshData;
 			for (int32 VertIdx=0; VertIdx < MeshVerts.Num(); ++VertIdx)
 			{
-				AddVertexHelper(DebugMeshData, MeshVerts[VertIdx] + CurrentData->NavMeshDrawOffset);
+				AddVertexHelper(DebugMeshData, MeshVerts[VertIdx] + CurrentData->NavMeshDrawOffset, NavMeshRenderColor_RecastTileBeingRebuilt);
 			}
 			DebugMeshData.Indices.Append(CurrentData->NavMeshGeometry.BuiltMeshIndices);
 			DebugMeshData.ClusterColor = NavMeshRenderColor_RecastTileBeingRebuilt;
diff --git a/Engine/Source/Runtime/Engine/Private/NavMeshRenderingHelpers.h b/Engine/Source/Runtime/Engine/Private/NavMeshRenderingHelpers.h
--- a/Engine/Source/Runtime/Engine/Private/NavMeshRenderingHelpers.h
+++ b/Engine/Source/Runtime/Engine/Private/NavMeshRenderingHelpers.h
@@ -5,10 +5,14 @@
 #include "DynamicMeshBuilder.h"
 #include "AI/Navigation/RecastHelpers.h"
 #include "BatchedElements.h"
+#include "MeshBatch.h"
+#include "Engine/TextureDefines.h"
+#include "SceneManagement.h"
 #include "AI/Navigation/NavMeshRenderingComponent.h"
 #include "Materials/Material.h"
 #include "MaterialShared.h"
 
+
 static const FColor NavMeshRenderColor_Recast_TriangleEdges(255,255,255);
 static const FColor NavMeshRenderColor_Recast_TileEdges(16,16,16,32);
 static const FColor NavMeshRenderColor_Recast_NavMeshEdges(32,63,0,220);
@@ -92,6 +96,19 @@ struct FNavMeshSceneProxyData : public TSharedFromThis<FNavMeshSceneProxyData, E
 	};
 	TArray<FDebugMeshData> MeshBuilders;
 
+	struct FDebugThickLine : FDebugRenderSceneProxy::FDebugLine
+	{
+		FDebugThickLine(const FVector &InStart, const FVector &InEnd, const FColor &InColor, float InThickness)
+			: FDebugLine(InStart, InEnd, InColor)
+			, Thickness(InThickness)
+		{
+
+		}
+
+		float Thickness;
+	};
+	TArray<FDebugThickLine> ThickLineItems;;
+
 	struct FDebugText
 	{
 		const FVector Location;
@@ -146,6 +163,93 @@ FORCEINLINE bool LineInCorrectDistance(const FVector& Start, const FVector& End,
 	return true;
 }
 
+class FNavMeshIndexBuffer : public FIndexBuffer
+{
+public:
+	TArray<int32> Indices;
+
+	virtual void InitRHI() override
+	{
+		if (Indices.Num() > 0)
+		{
+			FRHIResourceCreateInfo CreateInfo;
+			IndexBufferRHI = RHICreateIndexBuffer(sizeof(int32), Indices.Num() * sizeof(int32), BUF_Static, CreateInfo);
+
+			// Write the indices to the index buffer.
+			void* Buffer = RHILockIndexBuffer(IndexBufferRHI, 0, Indices.Num() * sizeof(int32), RLM_WriteOnly);
+			FMemory::Memcpy(Buffer, Indices.GetTypedData(), Indices.Num() * sizeof(int32));
+			RHIUnlockIndexBuffer(IndexBufferRHI);
+		}
+	}
+};
+
+/** Vertex Buffer */
+class FNavMeshVertexBuffer : public FVertexBuffer
+{
+public:
+	TArray<FDynamicMeshVertex> Vertices;
+
+	virtual void InitRHI() override
+	{
+		if (Vertices.Num() > 0)
+		{
+			FRHIResourceCreateInfo CreateInfo;
+			VertexBufferRHI = RHICreateVertexBuffer(Vertices.Num() * sizeof(FDynamicMeshVertex), BUF_Static, CreateInfo);
+
+			// Copy the vertex data into the vertex buffer.
+			void* VertexBufferData = RHILockVertexBuffer(VertexBufferRHI, 0, Vertices.Num() * sizeof(FDynamicMeshVertex), RLM_WriteOnly);
+			FMemory::Memcpy(VertexBufferData, Vertices.GetTypedData(), Vertices.Num() * sizeof(FDynamicMeshVertex));
+			RHIUnlockVertexBuffer(VertexBufferRHI);
+		}
+	}
+
+};
+
+/** Vertex Factory */
+class FNavMeshVertexFactory : public FLocalVertexFactory
+{
+public:
+
+	FNavMeshVertexFactory()
+	{}
+
+
+	/** Initialization */
+	void Init(const FNavMeshVertexBuffer* VertexBuffer)
+	{
+		if (IsInRenderingThread())
+		{
+			// Initialize the vertex factory's stream components.
+			DataType NewData;
+			NewData.PositionComponent = STRUCTMEMBER_VERTEXSTREAMCOMPONENT(VertexBuffer, FDynamicMeshVertex, Position, VET_Float3);
+			NewData.TextureCoordinates.Add(
+				FVertexStreamComponent(VertexBuffer, STRUCT_OFFSET(FDynamicMeshVertex, TextureCoordinate), sizeof(FDynamicMeshVertex), VET_Float2)
+				);
+			NewData.TangentBasisComponents[0] = STRUCTMEMBER_VERTEXSTREAMCOMPONENT(VertexBuffer, FDynamicMeshVertex, TangentX, VET_PackedNormal);
+			NewData.TangentBasisComponents[1] = STRUCTMEMBER_VERTEXSTREAMCOMPONENT(VertexBuffer, FDynamicMeshVertex, TangentZ, VET_PackedNormal);
+			SetData(NewData);
+		}
+		else
+		{
+			ENQUEUE_UNIQUE_RENDER_COMMAND_TWOPARAMETER(
+				InitNavMeshVertexFactory,
+				FNavMeshVertexFactory*, VertexFactory, this,
+				const FNavMeshVertexBuffer*, VertexBuffer, VertexBuffer,
+				{
+				// Initialize the vertex factory's stream components.
+				DataType NewData;
+				NewData.PositionComponent = STRUCTMEMBER_VERTEXSTREAMCOMPONENT(VertexBuffer, FDynamicMeshVertex, Position, VET_Float3);
+				NewData.TextureCoordinates.Add(
+					FVertexStreamComponent(VertexBuffer, STRUCT_OFFSET(FDynamicMeshVertex, TextureCoordinate), sizeof(FDynamicMeshVertex), VET_Float2)
+					);
+				NewData.TangentBasisComponents[0] = STRUCTMEMBER_VERTEXSTREAMCOMPONENT(VertexBuffer, FDynamicMeshVertex, TangentX, VET_PackedNormal);
+				NewData.TangentBasisComponents[1] = STRUCTMEMBER_VERTEXSTREAMCOMPONENT(VertexBuffer, FDynamicMeshVertex, TangentZ, VET_PackedNormal);
+				VertexFactory->SetData(NewData);
+				});
+		}
+	}
+};
+
 
 class FRecastRenderingSceneProxy : public FDebugRenderSceneProxy
 {
@@ -161,10 +265,65 @@ public:
 			ProxyData = *InProxyData;
 		}
 		RenderingComponent = Cast<const UNavMeshRenderingComponent>(InComponent);
+
+		const int32 NumberOfMeshes = ProxyData.MeshBuilders.Num();
+		if (!NumberOfMeshes)
+		{
+			return;
+		}
+
+		MeshColors.Reserve(NumberOfMeshes);
+		const FMaterialRenderProxy* ParentMaterial = GEngine->DebugMeshMaterial->GetRenderProxy(false);
+		for (int32 Index = 0; Index < NumberOfMeshes; ++Index)
+		{
+			const auto& CurrentMeshBuilder = ProxyData.MeshBuilders[Index];
+
+			FMeshBatchElement Element;
+			Element.FirstIndex = IndexBuffer.Indices.Num();
+			Element.NumPrimitives = CurrentMeshBuilder.Indices.Num() / 3;
+			Element.MinVertexIndex = VertexBuffer.Vertices.Num();
+			Element.MaxVertexIndex = Element.MinVertexIndex + CurrentMeshBuilder.Vertices.Num() - 1;
+			Element.IndexBuffer = &IndexBuffer;
+			MeshBatchElements.Add(Element);
+
+			MeshColors.Add(FColoredMaterialRenderProxy(ParentMaterial, CurrentMeshBuilder.ClusterColor));
+
+			VertexBuffer.Vertices.Append( CurrentMeshBuilder.Vertices );
+			IndexBuffer.Indices.Append( CurrentMeshBuilder.Indices );
+		}
+
+		if (ProxyData.bDrawPathCollidingGeometry && ProxyData.PathCollidingGeomVerts.Num() > 0)
+		{
+			FMeshBatchElement Element;
+			Element.FirstIndex = IndexBuffer.Indices.Num();
+			Element.MinVertexIndex = VertexBuffer.Vertices.Num();
+			Element.IndexBuffer = &IndexBuffer;
+
+			for (const auto& CurrentLoc : ProxyData.PathCollidingGeomVerts)
+			{
+				VertexBuffer.Vertices.Add(FDynamicMeshVertex(CurrentLoc));
+			}
+			IndexBuffer.Indices.Append(ProxyData.PathCollidingGeomIndices);
+
+			Element.NumPrimitives = ProxyData.PathCollidingGeomIndices.Num() / 3;
+			Element.MaxVertexIndex = VertexBuffer.Vertices.Num() - 1;
+			MeshBatchElements.Add(Element);
+
+			MeshColors.Add(FColoredMaterialRenderProxy(ParentMaterial, NavMeshRenderColor_PathCollidingGeom));
+		}
+
+		VertexFactory.Init(&VertexBuffer);
+
+		BeginInitResource(&IndexBuffer);
+		BeginInitResource(&VertexBuffer);
+		BeginInitResource(&VertexFactory);
 	}
 
 	virtual ~FRecastRenderingSceneProxy() 
 	{
+		VertexBuffer.ReleaseResource();
+		IndexBuffer.ReleaseResource();
+		VertexFactory.ReleaseResource();
 		ProxyData.Reset();
 	}
 
@@ -325,6 +484,146 @@ public:
 			);
 	}
 
+	virtual void GetDynamicMeshElements(const TArray<const FSceneView*>& Views, const FSceneViewFamily& ViewFamily, uint32 VisibilityMap, FMeshElementCollector& Collector) const override
+	{
+		QUICK_SCOPE_CYCLE_COUNTER(STAT_RecastRenderingSceneProxy_GetDynamicMeshElements);
+
+		if (!ProxyData.bEnableDrawing) //check if we have any data to render
+		{
+			return;
+		}
+
+		for (int32 ViewIndex = 0; ViewIndex < Views.Num(); ViewIndex++)
+		{
+			if (VisibilityMap & (1 << ViewIndex))
+			{
+				const FSceneView* View = Views[ViewIndex];
+				FPrimitiveDrawInterface* PDI = Collector.GetPDI(ViewIndex);
+
+				const bool bSkipDistanceCheck = GIsEditor && (GEngine->GetDebugLocalPlayer() == NULL);
+
+				// Draw Mesh
+				for (int32 Index = 0; Index < MeshBatchElements.Num(); ++Index)
+				{
+					if (MeshBatchElements[Index].NumPrimitives == 0)
+					{
+						continue;
+					}
+
+					FMeshBatch& Mesh = Collector.AllocateMesh();
+					FMeshBatchElement& BatchElement = Mesh.Elements[0];
+					BatchElement = MeshBatchElements[Index];
+					BatchElement.PrimitiveUniformBuffer = CreatePrimitiveUniformBufferImmediate(GetLocalToWorld(), GetBounds(), GetLocalBounds(), false, UseEditorDepthTest());
+
+					Mesh.bWireframe = false;
+					Mesh.VertexFactory = &VertexFactory;
+					Mesh.MaterialRenderProxy = &MeshColors[Index];
+					Mesh.ReverseCulling = IsLocalToWorldDeterminantNegative();
+					Mesh.Type = PT_TriangleList;
+					Mesh.DepthPriorityGroup = SDPG_World;
+					Mesh.bCanApplyViewModeOverrides = false;
+					Collector.AddMesh(ViewIndex, Mesh);
+				}
+
+				int32 Num = ProxyData.NavMeshEdgeLines.Num();
+				PDI->AddReserveLines(SDPG_World, Num, false, false);
+				PDI->AddReserveLines(SDPG_Foreground, Num, false, true);
+				for (int32 Index = 0; Index < Num; ++Index)
+				{
+					const FDebugLine &Line = ProxyData.NavMeshEdgeLines[Index];
+					if (LineInView(Line.Start, Line.End, View, false))
+					{
+						if (LineInCorrectDistance(Line.Start, Line.End, View))
+						{
+							PDI->DrawLine(Line.Start, Line.End, Line.Color, SDPG_World, NavMeshEdges_LineThickness, 0, true);
+						}
+						else if (GIsEditor)
+						{
+							PDI->DrawLine(Line.Start, Line.End, Line.Color, SDPG_Foreground, DefaultEdges_LineThickness, 0, true);
+						}
+					}
+				}
+
+
+				Num = ProxyData.ClusterLinkLines.Num();
+				PDI->AddReserveLines(SDPG_World, Num, false, false);
+				PDI->AddReserveLines(SDPG_Foreground, Num, false, true);
+				for (int32 Index = 0; Index < Num; ++Index)
+				{
+					const FDebugLine &Line = ProxyData.ClusterLinkLines[Index];
+					if (LineInView(Line.Start, Line.End, View, false))
+					{
+						if (LineInCorrectDistance(Line.Start, Line.End, View))
+						{
+							PDI->DrawLine(Line.Start, Line.End, Line.Color, SDPG_World, ClusterLinkLines_LineThickness, 0, true);
+						}
+						else if (GIsEditor)
+						{
+							PDI->DrawLine(Line.Start, Line.End, Line.Color, SDPG_Foreground, DefaultEdges_LineThickness, 0, true);
+						}
+					}
+				}
+
+				Num = ProxyData.TileEdgeLines.Num();
+				PDI->AddReserveLines(SDPG_World, Num, false, false);
+				PDI->AddReserveLines(SDPG_Foreground, Num, false, true);
+				for (int32 Index = 0; Index < Num; ++Index)
+				{
+					const FDebugLine &Line = ProxyData.TileEdgeLines[Index];
+					if (LineInView(Line.Start, Line.End, View, false))
+					{
+						if (LineInCorrectDistance(Line.Start, Line.End, View))
+						{
+							PDI->DrawLine(Line.Start, Line.End, Line.Color, SDPG_World, PolyEdges_LineThickness, 0, true);
+						}
+						else if (GIsEditor)
+						{
+							PDI->DrawLine(Line.Start, Line.End, Line.Color, SDPG_Foreground, DefaultEdges_LineThickness, 0, true);
+						}
+					}
+				}
+
+				Num = ProxyData.NavLinkLines.Num();
+				PDI->AddReserveLines(SDPG_World, Num, false, false);
+				PDI->AddReserveLines(SDPG_Foreground, Num, false, true);
+				for (int32 Index = 0; Index < Num; ++Index)
+				{
+					const FDebugLine &Line = ProxyData.NavLinkLines[Index];
+					if (LineInView(Line.Start, Line.End, View, false))
+					{
+						if (LineInCorrectDistance(Line.Start, Line.End, View))
+						{
+							PDI->DrawLine(Line.Start, Line.End, Line.Color, SDPG_World, LinkLines_LineThickness, 0, true);
+						}
+						else if (GIsEditor)
+						{
+							PDI->DrawLine(Line.Start, Line.End, Line.Color, SDPG_Foreground, DefaultEdges_LineThickness, 0, true);
+						}
+					}
+				}
+
+				Num = ProxyData.ThickLineItems.Num();
+				PDI->AddReserveLines(SDPG_Foreground, Num, false, true);
+				for (int32 Index = 0; Index < Num; ++Index)
+				{
+					const auto &Line = ProxyData.ThickLineItems[Index];
+					if (LineInView(Line.Start, Line.End, View, false))
+					{
+						if (LineInCorrectDistance(Line.Start, Line.End, View))
+						{
+							PDI->DrawLine(Line.Start, Line.End, Line.Color, SDPG_World, Line.Thickness, 0, true);
+						}
+						else if (GIsEditor)
+						{
+							PDI->DrawLine(Line.Start, Line.End, Line.Color, SDPG_Foreground, DefaultEdges_LineThickness, 0, true);
+						}
+					}
+				}
+
+			}
+		}
+	}
+
 	FORCEINLINE bool PointInView(const FVector& Location, const FSceneView* View)
 	{
 		return View->ViewFrustum.IntersectBox(Location, FVector::ZeroVector);
@@ -372,11 +671,18 @@ public:
 			+ ProxyData.TileEdgeLines.GetAllocatedSize() + ProxyData.NavMeshEdgeLines.GetAllocatedSize() + ProxyData.NavLinkLines.GetAllocatedSize()
 			+ ProxyData.PathCollidingGeomIndices.GetAllocatedSize() + ProxyData.PathCollidingGeomVerts.GetAllocatedSize()
 			+ ProxyData.DebugLabels.GetAllocatedSize() + ProxyData.ClusterLinkLines.GetAllocatedSize() + ProxyData.MeshBuilders.GetAllocatedSize()
-			+ ProxyData.BatchedElements.GetAllocatedSize());
+			+ ProxyData.BatchedElements.GetAllocatedSize() + IndexBuffer.Indices.GetAllocatedSize() + VertexBuffer.Vertices.GetAllocatedSize() + MeshColors.GetAllocatedSize() + MeshBatchElements.GetAllocatedSize());
 	}
 
 private:
 	FNavMeshSceneProxyData ProxyData;
+	
+	FNavMeshIndexBuffer IndexBuffer;
+	FNavMeshVertexBuffer VertexBuffer;
+	FNavMeshVertexFactory VertexFactory;
+	TArray <FColoredMaterialRenderProxy> MeshColors;
+	TArray<FMeshBatchElement> MeshBatchElements;
+
 	FDebugDrawDelegate DebugTextDrawingDelegate;
 	TWeakObjectPtr<UNavMeshRenderingComponent> RenderingComponent;
 	uint32 bRequestedData : 1;
diff --git a/Engine/Source/Runtime/Engine/Public/BatchedElements.h b/Engine/Source/Runtime/Engine/Public/BatchedElements.h
--- a/Engine/Source/Runtime/Engine/Public/BatchedElements.h
+++ b/Engine/Source/Runtime/Engine/Public/BatchedElements.h
@@ -147,8 +147,9 @@ public:
 	 * 
 	 * @param NumLines - number of lines to reserve space for
 	 * @param bDepthBiased - whether reserving depth-biased lines or non-biased lines
+	 * @param bThickLines - whether reserving regular lines or thick lines
 	 */
-	void AddReserveLines(int32 NumLines, bool bDepthBiased = false);
+	void AddReserveLines(int32 NumLines, bool bDepthBiased = false, bool bThickLines = false);
 
 	/** Adds a sprite to the batch. */
 	void AddSprite(
diff --git a/Engine/Source/Runtime/Engine/Public/SceneManagement.h b/Engine/Source/Runtime/Engine/Public/SceneManagement.h
--- a/Engine/Source/Runtime/Engine/Public/SceneManagement.h
+++ b/Engine/Source/Runtime/Engine/Public/SceneManagement.h
@@ -1012,7 +1012,7 @@ public:
 
 	virtual void RegisterDynamicResource(FDynamicPrimitiveResource* DynamicResource) = 0;
 
-	virtual void AddReserveLines(uint8 DepthPriorityGroup, int32 NumLines, bool bDepthBiased = false) = 0;
+	virtual void AddReserveLines(uint8 DepthPriorityGroup, int32 NumLines, bool bDepthBiased = false, bool bThickLines = false) = 0;
 
 	virtual void DrawSprite(
 		const FVector& Position,
@@ -1111,7 +1111,7 @@ public:
 	~FSimpleElementCollector();
 
 	virtual void SetHitProxy(HHitProxy* HitProxy);
-	virtual void AddReserveLines(uint8 DepthPriorityGroup, int32 NumLines, bool bDepthBiased = false) {}
+	virtual void AddReserveLines(uint8 DepthPriorityGroup, int32 NumLines, bool bDepthBiased = false, bool bThickLines = false) {}
 
 	virtual void DrawSprite(
 		const FVector& Position,
diff --git a/Engine/Source/Runtime/Renderer/Private/DynamicPrimitiveDrawing.h b/Engine/Source/Runtime/Renderer/Private/DynamicPrimitiveDrawing.h
--- a/Engine/Source/Runtime/Renderer/Private/DynamicPrimitiveDrawing.h
+++ b/Engine/Source/Runtime/Renderer/Private/DynamicPrimitiveDrawing.h
@@ -56,7 +56,7 @@ public:
 	virtual void RegisterDynamicResource(FDynamicPrimitiveResource* DynamicResource) override;
 	virtual bool IsMaterialIgnored(const FMaterialRenderProxy* MaterialRenderProxy, ERHIFeatureLevel::Type InFeatureLevel) const override;
 	virtual int32 DrawMesh(const FMeshBatch& Mesh) override;
-	virtual void AddReserveLines(uint8 DepthPriorityGroup, int32 NumLines, bool bDepthBiased = false) override;
+	virtual void AddReserveLines(uint8 DepthPriorityGroup, int32 NumLines, bool bDepthBiased = false, bool bThickLines = false) override;
 	virtual void DrawSprite(
 		const FVector& Position,
 		float SizeX,
@@ -216,7 +216,7 @@ public:
 	virtual bool IsHitTesting() override;
 	virtual void SetHitProxy(HHitProxy* HitProxy) override;
 	virtual void RegisterDynamicResource(FDynamicPrimitiveResource* DynamicResource) override;
-	virtual void AddReserveLines(uint8 DepthPriorityGroup, int32 NumLines, bool bDepthBiased = false) override;
+	virtual void AddReserveLines(uint8 DepthPriorityGroup, int32 NumLines, bool bDepthBiased = false, bool bThickLines = false) override;
 	virtual void DrawSprite(
 		const FVector& Position,
 		float SizeX,
diff --git a/Engine/Source/Runtime/Renderer/Private/DynamicPrimitiveDrawing.inl b/Engine/Source/Runtime/Renderer/Private/DynamicPrimitiveDrawing.inl
--- a/Engine/Source/Runtime/Renderer/Private/DynamicPrimitiveDrawing.inl
+++ b/Engine/Source/Runtime/Renderer/Private/DynamicPrimitiveDrawing.inl
@@ -160,11 +160,11 @@ void TDynamicPrimitiveDrawer<DrawingPolicyFactoryType>::DrawSprite(
 }
 
 template<typename DrawingPolicyFactoryType>
-void TDynamicPrimitiveDrawer<DrawingPolicyFactoryType>::AddReserveLines(uint8 DepthPriorityGroup, int32 NumLines, bool bDepthBiased)
+void TDynamicPrimitiveDrawer<DrawingPolicyFactoryType>::AddReserveLines(uint8 DepthPriorityGroup, int32 NumLines, bool bDepthBiased, bool bThickLines)
 {
 	if (DrawingPolicyFactoryType::bAllowSimpleElements)
 	{
-		BatchedElements.AddReserveLines(NumLines, bDepthBiased);
+		BatchedElements.AddReserveLines(NumLines, bDepthBiased, bThickLines);
 	}
 }
 
@@ -463,11 +463,11 @@ inline void FViewElementPDI::DrawSprite(
 	);
 }
 
-inline void FViewElementPDI::AddReserveLines(uint8 DepthPriorityGroup, int32 NumLines, bool bDepthBiased)
+inline void FViewElementPDI::AddReserveLines(uint8 DepthPriorityGroup, int32 NumLines, bool bDepthBiased, bool bThickLines)
 {
 	FBatchedElements& Elements = GetElements(DepthPriorityGroup);
 
-	Elements.AddReserveLines( NumLines, bDepthBiased );
+	Elements.AddReserveLines(NumLines, bDepthBiased, bThickLines);
 }
 
 inline void FViewElementPDI::DrawLine(
