Fix for TTP #286286 "The command "VisLog" crashes the game." #ue4

diff --git a/Engine/Source/Runtime/Engine/Private/VisualLog.cpp b/Engine/Source/Runtime/Engine/Private/VisualLog.cpp
--- a/Engine/Source/Runtime/Engine/Private/VisualLog.cpp
+++ b/Engine/Source/Runtime/Engine/Private/VisualLog.cpp
@@ -823,11 +823,6 @@ public:
 				FVisualLog::Get().SetIsRecording(false);
 				return true;
 			}
-			else if (Command == TEXT("exit"))
-			{
-				FLogVisualizerModule::Get()->CloseUI(InWorld);
-				return true;
-			}
 			else if (Command == TEXT("disableallbut"))
 			{
 				FString Category = FParse::Token(Cmd, 1);
@@ -835,11 +830,18 @@ public:
 				FVisualLog::Get().AddCategortyToWhiteList(*Category);
 				return true;
 			}
+#if WITH_EDITOR
+			else if (Command == TEXT("exit"))
+			{
+				FLogVisualizerModule::Get()->CloseUI(InWorld);
+				return true;
+			}
 			else
 			{
 				FLogVisualizerModule::Get()->SummonUI(InWorld);
 				return true;
 			}
+#endif
 #else
 			UE_LOG(LogVisual, Warning, TEXT("Unable to open LogVisualizer - logs are disabled"));
 #endif
