Fixed a crash when running a synth benchmark

diff --git a/Engine/Source/Runtime/Engine/Private/SceneView.cpp b/Engine/Source/Runtime/Engine/Private/SceneView.cpp
--- a/Engine/Source/Runtime/Engine/Private/SceneView.cpp
+++ b/Engine/Source/Runtime/Engine/Private/SceneView.cpp
@@ -1142,7 +1142,7 @@ void FSceneView::ConfigureBufferVisualizationSettings()
 
 ERHIFeatureLevel::Type FSceneView::GetFeatureLevel() const
 { 
-	return Family->GetFeatureLevel();
+	return (Family != nullptr) ? Family->GetFeatureLevel() : GRHIFeatureLevel;
 }
 
 EShaderPlatform FSceneView::GetShaderPlatform() const
