MSAA should be anything with sample count > 1.

diff --git a/Engine/Source/Runtime/Engine/Private/SceneView.cpp b/Engine/Source/Runtime/Engine/Private/SceneView.cpp
--- a/Engine/Source/Runtime/Engine/Private/SceneView.cpp
+++ b/Engine/Source/Runtime/Engine/Private/SceneView.cpp
@@ -989,7 +989,7 @@ void FSceneView::EndFinalPostprocessSettings()
 		if( !Family->EngineShowFlags.PostProcessing || !Family->EngineShowFlags.AntiAliasing || Quality <= 0
 			// Disable antialiasing in GammaLDR mode to avoid jittering.
 			|| (FeatureLevel == ERHIFeatureLevel::ES2 && MobileHDRCvar->GetValueOnGameThread() == 0)
-			|| (FeatureLevel == ERHIFeatureLevel::ES2 && (MSAAValue == 2 || MSAAValue == 4)))
+			|| (FeatureLevel == ERHIFeatureLevel::ES2 && (MSAAValue > 1)))
 		{
 			FinalPostProcessSettings.AntiAliasingMethod = AAM_None;
 		}
