Detach welded bodies works and also fix crash when BodySetup is NULL and trying to weld

diff --git a/Engine/Source/Runtime/Engine/Classes/PhysicsEngine/BodyInstance.h b/Engine/Source/Runtime/Engine/Classes/PhysicsEngine/BodyInstance.h
--- a/Engine/Source/Runtime/Engine/Classes/PhysicsEngine/BodyInstance.h
+++ b/Engine/Source/Runtime/Engine/Classes/PhysicsEngine/BodyInstance.h
@@ -313,10 +313,23 @@ public:
 
 
 #if WITH_BODY_WELDING
+
+#if WITH_PHYSX
+	TMap<physx::PxShape*, class UPrimitiveComponent*> ShapeToComponentMap;
+#endif
+
 	/** 
-	 * Takes two body instances and welds them together to create a single simulated rigid body
+	 * Takes two body instances and welds them together to create a single simulated rigid body. Returns true if success.
 	 */
-	void Weld(FBodyInstance* Body, const FTransform& RelativeTM);
+	bool Weld(FBodyInstance* Body, const FTransform& RelativeTM);
+
+	/** 
+	 * Takes a welded body and unwelds it. This function does not create the new body, it only removes the old one */
+	void UnWeld(class UPrimitiveComponent * Body);
+
+	/**
+	 * After adding/removing shapes call this function to update mass distribution etc... */
+	void PostShapeChange();
 #endif
 
 	/**
diff --git a/Engine/Source/Runtime/Engine/Classes/PhysicsEngine/BodySetup.h b/Engine/Source/Runtime/Engine/Classes/PhysicsEngine/BodySetup.h
--- a/Engine/Source/Runtime/Engine/Classes/PhysicsEngine/BodySetup.h
+++ b/Engine/Source/Runtime/Engine/Classes/PhysicsEngine/BodySetup.h
@@ -588,7 +588,7 @@ public:
 	 *   Add the shapes defined by this body setup to the supplied PxRigidBody. 
 	 */
 #if WITH_BODY_WELDING
-	void AddShapesToRigidActor(physx::PxRigidActor* PDestActor, FVector& Scale3D, const FTransform* RelativeTM = NULL);
+	void AddShapesToRigidActor(physx::PxRigidActor* PDestActor, FVector& Scale3D, const FTransform* RelativeTM = NULL, TArray<physx::PxShape*> * NewShapes = NULL);
 #else
 	void AddShapesToRigidActor(physx::PxRigidActor* PDestActor, FVector& Scale3D);
 #endif
diff --git a/Engine/Source/Runtime/Engine/Private/Actor.cpp b/Engine/Source/Runtime/Engine/Private/Actor.cpp
--- a/Engine/Source/Runtime/Engine/Private/Actor.cpp
+++ b/Engine/Source/Runtime/Engine/Private/Actor.cpp
@@ -1193,6 +1193,12 @@ void AActor::DetachRootComponentFromParent(bool bMaintainWorldPosition)
 {
 	if(RootComponent)
 	{
+		USceneComponent * RootComponent = GetRootComponent();
+		if (RootComponent->AttachParent)
+		{
+			RootComponent->AttachParent->WeldPhysicsBody(RootComponent, false);
+		}
+
 		GetRootComponent()->DetachFromParent(bMaintainWorldPosition);
 
 		AttachmentReplication.AttachParent = NULL;
diff --git a/Engine/Source/Runtime/Engine/Private/PhysicsEngine/BodyInstance.cpp b/Engine/Source/Runtime/Engine/Private/PhysicsEngine/BodyInstance.cpp
--- a/Engine/Source/Runtime/Engine/Private/PhysicsEngine/BodyInstance.cpp
+++ b/Engine/Source/Runtime/Engine/Private/PhysicsEngine/BodyInstance.cpp
@@ -1451,29 +1451,107 @@ void FBodyInstance::TermBody()
 }
 
 #if WITH_BODY_WELDING
-void FBodyInstance::Weld(FBodyInstance* TheirBody, const FTransform& RelativeTM)
+bool FBodyInstance::Weld(FBodyInstance* TheirBody, const FTransform& RelativeTM)
 {
 	check(TheirBody);
+	if (TheirBody->BodySetup.IsValid() == false)	//attach actor can be called before body has been initialized. In this case just return false
+	{
+		return false;
+	}
 
 	//@TODO: BOX2D: Implement Weld
 
 #if WITH_PHYSX
 	
-	//UBodySetup* LocalSpaceSetup = TheirBody->BodySetup->CreateSpaceCopy(RelativeTM);
+	//first time we're welding with root so we need to grab the shapes
+	if (ShapeToComponentMap.Num() == 0)
+	{
+		int32 NumSyncShapes = 0;
+		TArray<PxShape *> PShapes = GetAllShapes(NumSyncShapes);
+
+		for (int32 ShapeIdx = 0; ShapeIdx < PShapes.Num(); ++ShapeIdx)
+		{
+			PxShape* PShape = PShapes[ShapeIdx];
+			UPrimitiveComponent *& PrimComp = ShapeToComponentMap.FindOrAdd(PShape);
+			PrimComp = OwnerComponent.Get();
+		}
+	}
+
+	TArray<PxShape *> PNewShapes;
+
 
 	//child body gets placed into the same scenes as parent body
 	if (PxRigidActor* MyBody = RigidActorSync)
 	{
-		TheirBody->BodySetup->AddShapesToRigidActor(MyBody, Scale3D, &RelativeTM);
+		TheirBody->BodySetup->AddShapesToRigidActor(MyBody, Scale3D, &RelativeTM, &PNewShapes);
 	}
 
 	if (PxRigidActor* MyBody = RigidActorAsync)
 	{
-		TheirBody->BodySetup->AddShapesToRigidActor(MyBody, Scale3D, &RelativeTM);
+		TheirBody->BodySetup->AddShapesToRigidActor(MyBody, Scale3D, &RelativeTM, &PNewShapes);
 	}
 
 
-	
+	for (int32 ShapeIdx = 0; ShapeIdx < PNewShapes.Num(); ++ShapeIdx)
+	{
+		PxShape* PShape = PNewShapes[ShapeIdx];
+		UPrimitiveComponent *& PrimComp = ShapeToComponentMap.FindOrAdd(PShape);
+		PrimComp = TheirBody->OwnerComponent.Get();
+	}
+
+	PostShapeChange();
+
+	//remove their body from scenes
+	TermBodyHelper(TheirBody->SceneIndexSync, TheirBody->RigidActorSync, TheirBody);
+	TermBodyHelper(TheirBody->SceneIndexAsync, TheirBody->RigidActorAsync, TheirBody);
+
+	return true;
+#endif
+
+	return false;
+}
+
+void FBodyInstance::UnWeld(UPrimitiveComponent* ChildComponent)
+{
+	check(ChildComponent);
+
+	//@TODO: BOX2D: Implement Weld
+
+#if WITH_PHYSX
+
+	int32 NumSyncShapes = 0;
+	TArray<PxShape *> PShapes = GetAllShapes(NumSyncShapes);
+
+	for (int32 ShapeIdx = 0; ShapeIdx < NumSyncShapes; ++ShapeIdx)
+	{
+		PxShape* PShape = PShapes[ShapeIdx];
+		if (UPrimitiveComponent ** PrimCompPtrPtr = ShapeToComponentMap.Find(PShape))
+		{
+			if (*PrimCompPtrPtr == ChildComponent)
+			{
+				RigidActorSync->detachShape(*PShape);
+			}
+		}
+	}
+
+	for (int32 ShapeIdx = NumSyncShapes; ShapeIdx < PShapes.Num(); ++ShapeIdx)
+	{
+		PxShape* PShape = PShapes[ShapeIdx];
+		if (UPrimitiveComponent ** PrimCompPtrPtr = ShapeToComponentMap.Find(PShape))
+		{
+			if (*PrimCompPtrPtr == ChildComponent)
+			{
+				RigidActorAsync->detachShape(*PShape);
+			}
+		}
+	}
+
+	PostShapeChange();
+#endif
+}
+
+void FBodyInstance::PostShapeChange()
+{
 	// Apply correct physical materials to shape we created.
 	UpdatePhysicalMaterials();
 
@@ -1483,11 +1561,6 @@ void FBodyInstance::Weld(FBodyInstance* TheirBody, const FTransform& RelativeTM)
 	UpdateMassProperties();
 	// Update damping
 	UpdateDampingProperties();
-
-	//remove their body from scenes
-	TermBodyHelper(TheirBody->SceneIndexSync, TheirBody->RigidActorSync, TheirBody);
-	TermBodyHelper(TheirBody->SceneIndexAsync, TheirBody->RigidActorAsync, TheirBody);
-#endif
 }
 #endif
 
diff --git a/Engine/Source/Runtime/Engine/Private/PhysicsEngine/BodySetup.cpp b/Engine/Source/Runtime/Engine/Private/PhysicsEngine/BodySetup.cpp
--- a/Engine/Source/Runtime/Engine/Private/PhysicsEngine/BodySetup.cpp
+++ b/Engine/Source/Runtime/Engine/Private/PhysicsEngine/BodySetup.cpp
@@ -276,7 +276,7 @@ void SetupNonUniformHelper(FVector & Scale3D, float & MinScale, float & MinScale
 }
 
 #if WITH_BODY_WELDING
-void UBodySetup::AddShapesToRigidActor(PxRigidActor* PDestActor, FVector& Scale3D, const FTransform * RelativeTM /* = NULL */)
+void UBodySetup::AddShapesToRigidActor(PxRigidActor* PDestActor, FVector& Scale3D, const FTransform * RelativeTM /* = NULL */, TArray<physx::PxShape*> * NewShapes /* = NULL */ )
 #else
 void UBodySetup::AddShapesToRigidActor(PxRigidActor* PDestActor, FVector& Scale3D)
 #endif
@@ -340,6 +340,10 @@ void UBodySetup::AddShapesToRigidActor(PxRigidActor* PDestActor, FVector& Scale3
 
 				ensure(PLocalPose.isValid());
 				PxShape* NewShape = PDestActor->createShape(PSphereGeom, *PDefaultMat, PLocalPose);
+				if (NewShapes)
+				{
+					NewShapes->Add(NewShape);
+				}
 
 				const float ContactOffset = FMath::Min(MaxContactOffset, ContactOffsetFactor * PSphereGeom.radius);
 				NewShape->setContactOffset(ContactOffset);
@@ -370,6 +374,10 @@ void UBodySetup::AddShapesToRigidActor(PxRigidActor* PDestActor, FVector& Scale3
 
 				ensure(PLocalPose.isValid());
 				PxShape* NewShape = PDestActor->createShape(PBoxGeom, *PDefaultMat, PLocalPose);
+				if (NewShapes)
+				{
+					NewShapes->Add(NewShape);
+				}
 
 				const float ContactOffset = FMath::Min(MaxContactOffset, ContactOffsetFactor * PBoxGeom.halfExtents.minElement());
 				NewShape->setContactOffset(ContactOffset);
@@ -412,6 +420,10 @@ void UBodySetup::AddShapesToRigidActor(PxRigidActor* PDestActor, FVector& Scale3
 
 					ensure(PLocalPose.isValid());
 					PxShape* NewShape = PDestActor->createShape(PCapsuleGeom, *PDefaultMat, PLocalPose);
+					if (NewShapes)
+					{
+						NewShapes->Add(NewShape);
+					}
 
 					const float ContactOffset = FMath::Min(MaxContactOffset, ContactOffsetFactor * PCapsuleGeom.radius);
 					NewShape->setContactOffset(ContactOffset);
@@ -453,6 +465,10 @@ void UBodySetup::AddShapesToRigidActor(PxRigidActor* PDestActor, FVector& Scale3
 
 						ensure(PLocalPose.isValid());
 						PxShape* NewShape = PDestActor->createShape(PConvexGeom, *PDefaultMat, PLocalPose);
+						if (NewShapes)
+						{
+							NewShapes->Add(NewShape);
+						}
 
 						const float ContactOffset = FMath::Min(MaxContactOffset, ContactOffsetFactor * PBoundsExtents.minElement());
 						NewShape->setContactOffset(ContactOffset);
