Fix build missing ifdef for body weld

diff --git a/Engine/Source/Runtime/Engine/Private/Actor.cpp b/Engine/Source/Runtime/Engine/Private/Actor.cpp
--- a/Engine/Source/Runtime/Engine/Private/Actor.cpp
+++ b/Engine/Source/Runtime/Engine/Private/Actor.cpp
@@ -1194,10 +1194,12 @@ void AActor::DetachRootComponentFromParent(bool bMaintainWorldPosition)
 	if(RootComponent)
 	{
 		USceneComponent * RootComponent = GetRootComponent();
+#if WITH_BODY_WELDING
 		if (RootComponent->AttachParent)
 		{
 			RootComponent->AttachParent->WeldPhysicsBody(RootComponent, false);
 		}
+#endif
 
 		GetRootComponent()->DetachFromParent(bMaintainWorldPosition);
 
diff --git a/Engine/Source/Runtime/Engine/Private/PhysicsEngine/BodySetup.cpp b/Engine/Source/Runtime/Engine/Private/PhysicsEngine/BodySetup.cpp
--- a/Engine/Source/Runtime/Engine/Private/PhysicsEngine/BodySetup.cpp
+++ b/Engine/Source/Runtime/Engine/Private/PhysicsEngine/BodySetup.cpp
@@ -340,10 +340,12 @@ void UBodySetup::AddShapesToRigidActor(PxRigidActor* PDestActor, FVector& Scale3
 
 				ensure(PLocalPose.isValid());
 				PxShape* NewShape = PDestActor->createShape(PSphereGeom, *PDefaultMat, PLocalPose);
+#if WITH_BODY_WELDING
 				if (NewShapes)
 				{
 					NewShapes->Add(NewShape);
 				}
+#endif
 
 				const float ContactOffset = FMath::Min(MaxContactOffset, ContactOffsetFactor * PSphereGeom.radius);
 				NewShape->setContactOffset(ContactOffset);
@@ -374,10 +376,12 @@ void UBodySetup::AddShapesToRigidActor(PxRigidActor* PDestActor, FVector& Scale3
 
 				ensure(PLocalPose.isValid());
 				PxShape* NewShape = PDestActor->createShape(PBoxGeom, *PDefaultMat, PLocalPose);
+#if WITH_BODY_WELDING
 				if (NewShapes)
 				{
 					NewShapes->Add(NewShape);
 				}
+#endif
 
 				const float ContactOffset = FMath::Min(MaxContactOffset, ContactOffsetFactor * PBoxGeom.halfExtents.minElement());
 				NewShape->setContactOffset(ContactOffset);
@@ -420,10 +424,12 @@ void UBodySetup::AddShapesToRigidActor(PxRigidActor* PDestActor, FVector& Scale3
 
 					ensure(PLocalPose.isValid());
 					PxShape* NewShape = PDestActor->createShape(PCapsuleGeom, *PDefaultMat, PLocalPose);
+#if WITH_BODY_WELDING
 					if (NewShapes)
 					{
 						NewShapes->Add(NewShape);
 					}
+#endif
 
 					const float ContactOffset = FMath::Min(MaxContactOffset, ContactOffsetFactor * PCapsuleGeom.radius);
 					NewShape->setContactOffset(ContactOffset);
@@ -465,10 +471,12 @@ void UBodySetup::AddShapesToRigidActor(PxRigidActor* PDestActor, FVector& Scale3
 
 						ensure(PLocalPose.isValid());
 						PxShape* NewShape = PDestActor->createShape(PConvexGeom, *PDefaultMat, PLocalPose);
+#if WITH_BODY_WELDING
 						if (NewShapes)
 						{
 							NewShapes->Add(NewShape);
 						}
+#endif
 
 						const float ContactOffset = FMath::Min(MaxContactOffset, ContactOffsetFactor * PBoundsExtents.minElement());
 						NewShape->setContactOffset(ContactOffset);
