Missing file for unwelding

diff --git a/Engine/Source/Runtime/Engine/Private/PrimitiveComponent.cpp b/Engine/Source/Runtime/Engine/Private/PrimitiveComponent.cpp
--- a/Engine/Source/Runtime/Engine/Private/PrimitiveComponent.cpp
+++ b/Engine/Source/Runtime/Engine/Private/PrimitiveComponent.cpp
@@ -377,11 +377,47 @@ bool UPrimitiveComponent::WeldPhysicsBody(USceneComponent * ChildComponent, bool
 				FBodyInstance * RootBody = RootPrim->GetBodyInstance(ParentBoneName);
 				FBodyInstance * ChildBody = ChildPrim->GetBodyInstance(ChildBoneName);
 
-				if (RootBody && ChildBody && bWeld)
+				if (RootBody && ChildBody)
 				{
+					if (bWeld)
+					{
+						return RootBody->Weld(ChildBody, RelativeTM);
+					}
+					else 
+					{
+						//we first init the new root which is ChildBody
+						USceneComponent * OldParent = ChildPrim->AttachParent;	//we need to temporarily NULL it out to create a new body then put it back so that we can properly NULL it out later
+						ChildPrim->AttachParent = NULL;
+						ChildBody->InitBody(ChildPrim->GetBodySetup(), ChildPrim->GetComponentToWorld(), ChildPrim, ChildPrim->GetWorld()->GetPhysicsScene());
+						ChildPrim->AttachParent = OldParent;
+
+						//we need to find all of ChildComponent's children
+						TArray<USceneComponent*> Children;
+						ChildComponent->GetChildrenComponents(true, Children);
+
+						//then remove child and all its children from original root
+						RootBody->UnWeld(ChildPrim);
+						for (USceneComponent * Child : Children)
+						{
+							if (UPrimitiveComponent * Prim = Cast<UPrimitiveComponent>(Child))
+							{
+								RootBody->UnWeld(Prim);
+							}
+						}
+						
+						bool bSuccess = true;
+						//then weld all of child's children to child
+						for (USceneComponent * Child : Children)
+						{
+							if (UPrimitiveComponent * Prim = Cast<UPrimitiveComponent>(Child))
+							{
+								FTransform RelativeTM = Prim->GetComponentToWorld().GetRelativeTransform(ChildPrim->GetComponentToWorld());
+								bSuccess &= ChildBody->Weld(Prim->GetBodyInstance(), RelativeTM);
+							}
+						}
 
-					RootBody->Weld(ChildBody, RelativeTM);
-					return true;
+						return bSuccess;
+					}
 				}
 			}
 		}
