Fix non-monolithic build

diff --git a/Engine/Source/Runtime/Engine/Classes/Engine/BlueprintCore.h b/Engine/Source/Runtime/Engine/Classes/Engine/BlueprintCore.h
--- a/Engine/Source/Runtime/Engine/Classes/Engine/BlueprintCore.h
+++ b/Engine/Source/Runtime/Engine/Classes/Engine/BlueprintCore.h
@@ -5,7 +5,7 @@
 #include "BlueprintCore.generated.h"
 
 UCLASS()
-class UBlueprintCore : public UObject
+class ENGINE_API UBlueprintCore : public UObject
 {
 	GENERATED_UCLASS_BODY()
 
