Crash fixes for TTP #344362 (BLOCKER: EDITOR: CRASH: MAIN: Editor crashes when opening certain asset types)

diff --git a/Engine/Source/Runtime/Engine/Private/Slate/DebugCanvas.cpp b/Engine/Source/Runtime/Engine/Private/Slate/DebugCanvas.cpp
--- a/Engine/Source/Runtime/Engine/Private/Slate/DebugCanvas.cpp
+++ b/Engine/Source/Runtime/Engine/Private/Slate/DebugCanvas.cpp
@@ -27,7 +27,7 @@ class FCanvasProxy
 {
 public:
 	FCanvasProxy( FRenderTarget* RenderTarget, UWorld* InWorld )
-		: Canvas(RenderTarget, NULL, InWorld, InWorld->FeatureLevel)
+		: Canvas(RenderTarget, NULL, InWorld, InWorld ? InWorld->FeatureLevel : GRHIFeatureLevel)
 	{
 		// Do not allow the canvas to be flushed outside of our debug rendering path
 		Canvas.SetAllowedModes( 0 );
diff --git a/Engine/Source/Runtime/Engine/Private/UnrealClient.cpp b/Engine/Source/Runtime/Engine/Private/UnrealClient.cpp
--- a/Engine/Source/Runtime/Engine/Private/UnrealClient.cpp
+++ b/Engine/Source/Runtime/Engine/Private/UnrealClient.cpp
@@ -1070,7 +1070,8 @@ void FViewport::Draw( bool bShouldPresent /*= true */)
 					GameThread.Waits = 0;
 				}
 
-				FCanvas Canvas(this, NULL, ViewportClient->GetWorld(), ViewportClient->GetWorld()->FeatureLevel);
+				auto World = ViewportClient->GetWorld();
+				FCanvas Canvas(this, NULL, World, World ? World->FeatureLevel : GRHIFeatureLevel);
 				{
 					ViewportClient->Draw(this, &Canvas);
 				}
@@ -1173,7 +1174,8 @@ const TArray<FColor>& FViewport::GetRawHitProxyData(FIntRect InRect)
 		});
 
 		// Let the viewport client draw its hit proxies.
-		FCanvas Canvas(&HitProxyMap, &HitProxyMap, ViewportClient->GetWorld(), ViewportClient->GetWorld()->FeatureLevel);
+		auto World = ViewportClient->GetWorld();
+		FCanvas Canvas(&HitProxyMap, &HitProxyMap, World, World ? World->FeatureLevel : GRHIFeatureLevel);
 		{
 			ViewportClient->Draw(this, &Canvas);
 		}
diff --git a/Engine/Source/Runtime/Engine/Private/UserInterface/Canvas.cpp b/Engine/Source/Runtime/Engine/Private/UserInterface/Canvas.cpp
--- a/Engine/Source/Runtime/Engine/Private/UserInterface/Canvas.cpp
+++ b/Engine/Source/Runtime/Engine/Private/UserInterface/Canvas.cpp
@@ -37,7 +37,6 @@ FCanvas::FCanvas(FRenderTarget* InRenderTarget, FHitProxyConsumer* InHitProxyCon
 {
 	Construct();
 
-	check(InWorld);
 	if (InWorld)
 	{
 		CurrentRealTime = InWorld->GetRealTimeSeconds();
