#ue4 - Exposed SceneComponent::bShouldUpdatePhysicsVolume and the PhysicsVolumeChangedDelegate to blueprints.

- PhysicsVolumeChangedDelegate is now a dynamic multicast delegate instead of a single delegate.
- bShouldUpdatePhysicsVolume is actually hidden from blueprints by default on SceneComponent and shown on PrimitiveComponents.
   - The flag and functionality should probably be on PrimitiveComponent itself, but we don't want to break current users by moving it.

diff --git a/Engine/Source/Runtime/Engine/Classes/Components/PrimitiveComponent.h b/Engine/Source/Runtime/Engine/Classes/Components/PrimitiveComponent.h
--- a/Engine/Source/Runtime/Engine/Classes/Components/PrimitiveComponent.h
+++ b/Engine/Source/Runtime/Engine/Classes/Components/PrimitiveComponent.h
@@ -87,7 +87,7 @@ DECLARE_DYNAMIC_MULTICAST_DELEGATE_TwoParams( FComponentOnInputTouchEndSignature
 DECLARE_DYNAMIC_MULTICAST_DELEGATE_TwoParams( FComponentBeginTouchOverSignature, ETouchIndex::Type, FingerIndex, UPrimitiveComponent*, TouchedComponent );
 DECLARE_DYNAMIC_MULTICAST_DELEGATE_TwoParams( FComponentEndTouchOverSignature, ETouchIndex::Type, FingerIndex, UPrimitiveComponent*, TouchedComponent );
 
-UCLASS(abstract, HideCategories=(Mobility))
+UCLASS(abstract, HideCategories=(Mobility), ShowCategories=(PhysicsVolume))
 class ENGINE_API UPrimitiveComponent : public USceneComponent
 {
 	GENERATED_UCLASS_BODY()
diff --git a/Engine/Source/Runtime/Engine/Classes/Components/SceneComponent.h b/Engine/Source/Runtime/Engine/Classes/Components/SceneComponent.h
--- a/Engine/Source/Runtime/Engine/Classes/Components/SceneComponent.h
+++ b/Engine/Source/Runtime/Engine/Classes/Components/SceneComponent.h
@@ -67,11 +67,11 @@ FORCEINLINE EMoveComponentFlags operator&(EMoveComponentFlags Arg1,EMoveComponen
 FORCEINLINE void operator&=(EMoveComponentFlags& Dest,EMoveComponentFlags Arg)					{ Dest = EMoveComponentFlags(Dest & Arg); }
 FORCEINLINE void operator|=(EMoveComponentFlags& Dest,EMoveComponentFlags Arg)					{ Dest = EMoveComponentFlags(Dest | Arg); }
 
-DECLARE_DELEGATE_OneParam(FPhysicsVolumeChanged, class APhysicsVolume*);
+DECLARE_DYNAMIC_MULTICAST_DELEGATE_OneParam(FPhysicsVolumeChanged, class APhysicsVolume*, NewVolume);
 
 
 /** A SceneComponent has a transform and supports attachment, but has no rendering or collision capabilities. Useful as a 'dummy' component in the hierarchy to offset others. */
-UCLASS(ClassGroup=Utility, BlueprintType, HideCategories=(Trigger), meta=(BlueprintSpawnableComponent))
+UCLASS(ClassGroup=Utility, BlueprintType, HideCategories=(Trigger, PhysicsVolume), meta=(BlueprintSpawnableComponent))
 class ENGINE_API USceneComponent : public UActorComponent
 {
 	GENERATED_UCLASS_BODY()
@@ -113,8 +113,8 @@ public:
 	UPROPERTY(EditDefaultsOnly, BlueprintReadOnly, Category=Rendering)
 	uint32 bHiddenInGame:1;
 
-	/** Whether or not it should update PhysicsVolume. **/
-	UPROPERTY()
+	/** Whether or not PhysicsVolume should be updated when the component is moved. For PrimitiveComponents, bGenerateOverlapEvents must be true as well. */
+	UPROPERTY(EditAnywhere, BlueprintReadWrite, Category=PhysicsVolume)
 	uint32 bShouldUpdatePhysicsVolume:1;
 
 	/** If true, a change in the bounds of this component will call trigger a streaming data rebuild */
@@ -433,7 +433,8 @@ public:
 	virtual void SetHiddenInGame(bool NewHidden, bool bPropagateToChildren=false);
 
 public:
-	/** Will be called when PhysicsVolume has been changed **/
+	/** Delegate that will be called when PhysicsVolume has been changed **/
+	UPROPERTY(BlueprintAssignable, Category=PhysicsVolume)
 	FPhysicsVolumeChanged PhysicsVolumeChangedDelegate;
 
 	// Begin ActorComponent interface
@@ -476,7 +477,7 @@ public:
 	 * Assumes that the level's Dynamics member is locked, which will always be the case during
 	 * a call to UWorld::Tick; if not locked, no actor-actor collision checking is performed.
 	 *
-	 * Updates actor's PhysicsVolume and touching interactions.
+	 * Updates actor's PhysicsVolume and touching interactions (if bShouldUpdatePhysicsVolume is true).
 	 * 
 	 * @param Delta			The desired location change in world space.
 	 * @param NewRotation	The new desired rotation in world space.
@@ -598,6 +599,7 @@ public:
 	/** 
 	 * Get the PhysicsVolume overlapping this component.
 	 */
+	UFUNCTION(BlueprintCallable, Category=PhysicsVolume)
 	APhysicsVolume* GetPhysicsVolume() const;
 
 
diff --git a/Engine/Source/Runtime/Engine/Classes/GameFramework/MovementComponent.h b/Engine/Source/Runtime/Engine/Classes/GameFramework/MovementComponent.h
--- a/Engine/Source/Runtime/Engine/Classes/GameFramework/MovementComponent.h
+++ b/Engine/Source/Runtime/Engine/Classes/GameFramework/MovementComponent.h
@@ -133,6 +133,7 @@ public:
 	virtual APhysicsVolume* GetPhysicsVolume() const;
 
 	/** Delegate when PhysicsVolume of UpdatedComponent has been changed **/
+	UFUNCTION()
 	virtual void PhysicsVolumeChanged(class APhysicsVolume* NewVolume);
 
 	/** add updated component to physics volume */
diff --git a/Engine/Source/Runtime/Engine/Private/MovementComponent.cpp b/Engine/Source/Runtime/Engine/Private/MovementComponent.cpp
--- a/Engine/Source/Runtime/Engine/Private/MovementComponent.cpp
+++ b/Engine/Source/Runtime/Engine/Private/MovementComponent.cpp
@@ -39,7 +39,7 @@ void UMovementComponent::SetUpdatedComponent(UPrimitiveComponent* NewUpdatedComp
 	{
 		UpdatedComponent->bShouldUpdatePhysicsVolume = false;
 		UpdatedComponent->SetPhysicsVolume(NULL, true);
-		UpdatedComponent->PhysicsVolumeChangedDelegate.Unbind();
+		UpdatedComponent->PhysicsVolumeChangedDelegate.RemoveDynamic(this, &UMovementComponent::PhysicsVolumeChanged);
 
 		// remove from tick prerequisite
 		UpdatedComponent->PrimaryComponentTick.RemovePrerequisite(this, PrimaryComponentTick); 
@@ -51,7 +51,7 @@ void UMovementComponent::SetUpdatedComponent(UPrimitiveComponent* NewUpdatedComp
 	{
 		UpdatedComponent->bShouldUpdatePhysicsVolume = true;
 		UpdatedComponent->UpdatePhysicsVolume(true);
-		UpdatedComponent->PhysicsVolumeChangedDelegate.BindUObject(this, &UMovementComponent::PhysicsVolumeChanged);
+		UpdatedComponent->PhysicsVolumeChangedDelegate.AddUniqueDynamic(this, &UMovementComponent::PhysicsVolumeChanged);
 		
 		// force ticks after movement component updates
 		UpdatedComponent->PrimaryComponentTick.AddPrerequisite(this, PrimaryComponentTick); 
diff --git a/Engine/Source/Runtime/Engine/Private/SceneComponent.cpp b/Engine/Source/Runtime/Engine/Private/SceneComponent.cpp
--- a/Engine/Source/Runtime/Engine/Private/SceneComponent.cpp
+++ b/Engine/Source/Runtime/Engine/Private/SceneComponent.cpp
@@ -1151,7 +1151,7 @@ void USceneComponent::SetPhysicsVolume( APhysicsVolume * NewVolume,  bool bTrigg
 			if( PhysicsVolume )
 			{
 				PhysicsVolume->ActorLeavingVolume(A);
-				PhysicsVolumeChangedDelegate.ExecuteIfBound(NewVolume);
+				PhysicsVolumeChangedDelegate.Broadcast(NewVolume);
 			}
 			PhysicsVolume = NewVolume;
 			if( PhysicsVolume )
@@ -1168,7 +1168,7 @@ void USceneComponent::SetPhysicsVolume( APhysicsVolume * NewVolume,  bool bTrigg
 
 void USceneComponent::BeginDestroy()
 {
-	PhysicsVolumeChangedDelegate.Unbind();
+	PhysicsVolumeChangedDelegate.Clear();
 
 	Super::BeginDestroy();
 }
