#ue4 - Expose add/remove/clear functions for CameraLensEffects on PlayerCameraManager. Make sure to call UpdateLocation() on all camera effects.

TTPs 342122, 342118

diff --git a/Engine/Source/Runtime/Engine/Classes/Camera/PlayerCameraManager.h b/Engine/Source/Runtime/Engine/Classes/Camera/PlayerCameraManager.h
--- a/Engine/Source/Runtime/Engine/Classes/Camera/PlayerCameraManager.h
+++ b/Engine/Source/Runtime/Engine/Classes/Camera/PlayerCameraManager.h
@@ -518,6 +518,9 @@ protected:
 	 */
 	virtual void UpdateViewTarget(FTViewTarget& OutVT, float DeltaTime);
 
+	/** Update any attached camera lens effects **/
+	virtual void UpdateCameraLensEffects( const FTViewTarget& OutVT );
+
 public:
 	/** Set a new ViewTarget with optional transition time */
 	void SetViewTarget(class AActor* NewViewTarget, FViewTargetTransitionParams TransitionParams = FViewTargetTransitionParams());
@@ -529,12 +532,15 @@ public:
 	virtual class AEmitterCameraLensEffectBase* FindCameraLensEffect(TSubclassOf<class AEmitterCameraLensEffectBase> LensEffectEmitterClass);
 	
 	/** Initiates a camera lens effect of the given class on this camera. */
+	UFUNCTION(BlueprintCallable, Category = "Camera")
 	virtual AEmitterCameraLensEffectBase* AddCameraLensEffect(TSubclassOf<class AEmitterCameraLensEffectBase> LensEffectEmitterClass);
 	
 	/** Removes this particular lens effect from the camera. */
+	UFUNCTION(BlueprintCallable, Category = "Camera")
 	virtual void RemoveCameraLensEffect(class AEmitterCameraLensEffectBase* Emitter);
 	
 	/** Removes all Camera Lens Effects. */
+	UFUNCTION(BlueprintCallable, Category = "Camera")
 	virtual void ClearCameraLensEffects();
 
 	//
diff --git a/Engine/Source/Runtime/Engine/Classes/GameFramework/PlayerController.h b/Engine/Source/Runtime/Engine/Classes/GameFramework/PlayerController.h
--- a/Engine/Source/Runtime/Engine/Classes/GameFramework/PlayerController.h
+++ b/Engine/Source/Runtime/Engine/Classes/GameFramework/PlayerController.h
@@ -710,6 +710,10 @@ public:
 	UFUNCTION(unreliable, client, BlueprintCallable, Category="Game|Feedback")
 	void ClientSpawnCameraLensEffect(TSubclassOf<class AEmitterCameraLensEffectBase>  LensEffectEmitterClass);
 
+	/** Removes all Camera Lens Effects. */
+	UFUNCTION(reliable, client, BlueprintCallable, Category="Game|Feedback")
+	virtual void ClientClearCameraLensEffects();
+
 	/** Stop camera animation on client. */
 	UFUNCTION(reliable, client)
 	void ClientStopCameraAnim(class UCameraAnim* AnimToStop);
diff --git a/Engine/Source/Runtime/Engine/Private/Particles/ParticleComponents.cpp b/Engine/Source/Runtime/Engine/Private/Particles/ParticleComponents.cpp
--- a/Engine/Source/Runtime/Engine/Private/Particles/ParticleComponents.cpp
+++ b/Engine/Source/Runtime/Engine/Private/Particles/ParticleComponents.cpp
@@ -6010,8 +6010,7 @@ void AEmitterCameraLensEffectBase::UpdateLocation(const FVector& CamLoc, const F
 
 	//UE_LOG(LogParticles, Warning, TEXT("DistAdjustedForFOV: %f  BaseFOV: %f  CamFOVDeg: %f"), DistAdjustedForFOV, BaseFOV, CamFOVDeg );
 
-	SetActorLocation( CamLoc + X * DistAdjustedForFOV, false );
-	SetActorRotation( NewRot );
+	SetActorLocationAndRotation( CamLoc + X * DistAdjustedForFOV, NewRot, false );
 }
 
 void AEmitterCameraLensEffectBase::Destroyed()
diff --git a/Engine/Source/Runtime/Engine/Private/PlayerCameraManager.cpp b/Engine/Source/Runtime/Engine/Private/PlayerCameraManager.cpp
--- a/Engine/Source/Runtime/Engine/Private/PlayerCameraManager.cpp
+++ b/Engine/Source/Runtime/Engine/Private/PlayerCameraManager.cpp
@@ -572,6 +572,20 @@ void APlayerCameraManager::UpdateViewTarget(FTViewTarget& OutVT, float DeltaTime
 
 	// Synchronize the actor with the view target results
 	SetActorLocationAndRotation(OutVT.POV.Location, OutVT.POV.Rotation, false);
+
+	UpdateCameraLensEffects(OutVT);
+}
+
+
+void APlayerCameraManager::UpdateCameraLensEffects(const FTViewTarget& OutVT)
+{
+	for (int32 Idx=0; Idx<CameraLensEffects.Num(); ++Idx)
+	{
+		if (CameraLensEffects[Idx] != NULL)
+		{
+			CameraLensEffects[Idx]->UpdateLocation(OutVT.POV.Location, OutVT.POV.Rotation, OutVT.POV.FOV);
+		}
+	}
 }
 
 
diff --git a/Engine/Source/Runtime/Engine/Private/PlayerController.cpp b/Engine/Source/Runtime/Engine/Private/PlayerController.cpp
--- a/Engine/Source/Runtime/Engine/Private/PlayerController.cpp
+++ b/Engine/Source/Runtime/Engine/Private/PlayerController.cpp
@@ -3573,6 +3573,14 @@ void APlayerController::ClientSpawnCameraLensEffect_Implementation( TSubclassOf<
 	}
 }
 
+void APlayerController::ClientClearCameraLensEffects_Implementation()
+{
+	if (PlayerCameraManager != NULL)
+	{
+		PlayerCameraManager->ClearCameraLensEffects();
+	}
+}
+
 void APlayerController::ReceivedGameModeClass(TSubclassOf<AGameMode> GameModeClass)
 {
 }
