#ue4 - Deprecate unused NotifyKilledBy() from AController.

diff --git a/Engine/Source/Runtime/Engine/Classes/GameFramework/Controller.h b/Engine/Source/Runtime/Engine/Classes/GameFramework/Controller.h
--- a/Engine/Source/Runtime/Engine/Classes/GameFramework/Controller.h
+++ b/Engine/Source/Runtime/Engine/Classes/GameFramework/Controller.h
@@ -233,13 +233,7 @@ public:
 	 */
 	virtual void PawnPendingDestroy(class APawn* inPawn);
 
-	/**
-	 * Notification of someone in the game being killed (to be overridden for things classes like AI)
-	 * @param	AController * Killer - The controller of the player who did the killing
-	 * @param	AController * KilledPlayer - The controller of the player who got killed
-	 * @param	APawn * KilledPawn - The pawn of the player who got killed
-	 * @param	const UDamageType * DamageType - The damage type used to kill the player
-	 */
+	DEPRECATED(4.5, "NotifyKilled is unused in engine code and will be removed from AController.")
 	virtual void NotifyKilled(AController* Killer, AController* KilledPlayer, APawn* KilledPawn, const UDamageType* DamageType) {}
 
 	/** Called when this controller instigates ANY damage */
