#ue4 - Change log level of MTD warning in physics collision conversion code, it's common when shapes have same location.

- Will be following up to prevent this erorr.

diff --git a/Engine/Source/Runtime/Engine/Private/Collision/CollisionConversions.cpp b/Engine/Source/Runtime/Engine/Private/Collision/CollisionConversions.cpp
--- a/Engine/Source/Runtime/Engine/Private/Collision/CollisionConversions.cpp
+++ b/Engine/Source/Runtime/Engine/Private/Collision/CollisionConversions.cpp
@@ -579,7 +579,7 @@ static bool ConvertOverlappedShapeToImpactHit(const PxLocationHit& PHit, const F
 		OutResult.PenetrationDepth = 0.f;
 		if (!bFiniteNormal)
 		{
-			UE_LOG(LogPhysics, Warning, TEXT("ConvertOverlappedShapeToImpactHit: MTD returned NaN :( normal: (X:%f, Y:%f, Z:%f)"), PHit.normal.x, PHit.normal.y, PHit.normal.z);
+			UE_LOG(LogPhysics, Verbose, TEXT("ConvertOverlappedShapeToImpactHit: MTD returned NaN :( normal: (X:%f, Y:%f, Z:%f)"), PHit.normal.x, PHit.normal.y, PHit.normal.z);
 		}
 	}
 
