#ue4 - Calls to UpdateOverlaps() should abort if currently within an FScopedMovementUpdate.

- Such calls could invalidate the cached pending updates.
- This was a potential issue during Character Crouch()/UnCrouch() since it changed the capsule size, triggering UpdateOverlaps().

diff --git a/Engine/Source/Runtime/Engine/Private/PrimitiveComponent.cpp b/Engine/Source/Runtime/Engine/Private/PrimitiveComponent.cpp
--- a/Engine/Source/Runtime/Engine/Private/PrimitiveComponent.cpp
+++ b/Engine/Source/Runtime/Engine/Private/PrimitiveComponent.cpp
@@ -2110,6 +2110,11 @@ void UPrimitiveComponent::UpdateOverlaps(TArray<FOverlapInfo> const* PendingOver
 {
 	QUICK_SCOPE_CYCLE_COUNTER(STAT_UpdateOverlaps); 
 
+	if (IsDeferringMovementUpdates())
+	{
+		return;
+	}
+
 	// first, dispatch any pending overlaps
 	if (bGenerateOverlapEvents && IsCollisionEnabled())
 	{
diff --git a/Engine/Source/Runtime/Engine/Private/SceneComponent.cpp b/Engine/Source/Runtime/Engine/Private/SceneComponent.cpp
--- a/Engine/Source/Runtime/Engine/Private/SceneComponent.cpp
+++ b/Engine/Source/Runtime/Engine/Private/SceneComponent.cpp
@@ -173,10 +173,9 @@ class FScopedMovementUpdate* USceneComponent::GetCurrentScopedMovement() const
 
 bool USceneComponent::IsDeferringMovementUpdates() const
 {
-	FScopedMovementUpdate* CurrentScopedUpdate = GetCurrentScopedMovement();
-	if (CurrentScopedUpdate)
+	if (ScopedMovementStack.Num() > 0)
 	{
-		checkSlow(CurrentScopedUpdate->IsDeferringUpdates());
+		checkSlow(ScopedMovementStack.Last()->IsDeferringUpdates());
 		return true;
 	}
 	
@@ -1211,6 +1210,11 @@ void USceneComponent::UpdateOverlaps(TArray<FOverlapInfo> const* PendingOverlaps
 {
 	QUICK_SCOPE_CYCLE_COUNTER(STAT_UpdateOverlaps); 
 
+	if (IsDeferringMovementUpdates())
+	{
+		return;
+	}
+	
 	if (bShouldUpdatePhysicsVolume)
 	{
 		UpdatePhysicsVolume(bDoNotifies);
@@ -1596,7 +1600,7 @@ FScopedMovementUpdate::FScopedMovementUpdate( class USceneComponent* Component,
 			FScopedMovementUpdate* OuterScope = Component->GetCurrentScopedMovement();
 			if (OuterScope && OuterScope->bDeferUpdates)
 			{
-				if (s_ScopedWarningCount < 100)
+				if (s_ScopedWarningCount < 100 || (GFrameCounter & 31) == 0)
 				{
 					s_ScopedWarningCount++;
 					UE_LOG(LogSceneComponent, Error, TEXT("FScopedMovementUpdate attempting to use immediate updates within deferred scope, will use deferred updates instead."));
