Header cleanup: Engine module, public headers.

diff --git a/Engine/Source/Editor/LevelEditor/Private/LevelViewportActions.cpp b/Engine/Source/Editor/LevelEditor/Private/LevelViewportActions.cpp
--- a/Engine/Source/Editor/LevelEditor/Private/LevelViewportActions.cpp
+++ b/Engine/Source/Editor/LevelEditor/Private/LevelViewportActions.cpp
@@ -5,6 +5,7 @@
 #include "LevelViewportActions.h"
 #include "EditorShowFlags.h"
 #include "StatsData.h"
+#include "BufferVisualizationData.h"
 
 #define LOCTEXT_NAMESPACE "LevelViewportActions"
 
diff --git a/Engine/Source/Editor/LevelEditor/Private/SLevelViewport.cpp b/Engine/Source/Editor/LevelEditor/Private/SLevelViewport.cpp
--- a/Engine/Source/Editor/LevelEditor/Private/SLevelViewport.cpp
+++ b/Engine/Source/Editor/LevelEditor/Private/SLevelViewport.cpp
@@ -27,6 +27,7 @@
 #include "SCaptureRegionWidget.h"
 #include "Settings.h"
 #include "SceneOutlinerTreeItems.h"
+#include "BufferVisualizationData.h"
 
 static const FName LevelEditorName("LevelEditor");
 
diff --git a/Engine/Source/Editor/LevelEditor/Private/SLevelViewportToolBar.cpp b/Engine/Source/Editor/LevelEditor/Private/SLevelViewportToolBar.cpp
--- a/Engine/Source/Editor/LevelEditor/Private/SLevelViewportToolBar.cpp
+++ b/Engine/Source/Editor/LevelEditor/Private/SLevelViewportToolBar.cpp
@@ -20,6 +20,7 @@
 #include "Editor/UnrealEd/Public/SEditorViewportToolBarButton.h"
 #include "Editor/UnrealEd/Public/SEditorViewportViewMenu.h"
 #include "StatsData.h"
+#include "BufferVisualizationData.h"
 
 #define LOCTEXT_NAMESPACE "LevelViewportToolBar"
 
diff --git a/Engine/Source/Editor/UnrealEd/Private/Editor.cpp b/Engine/Source/Editor/UnrealEd/Private/Editor.cpp
--- a/Engine/Source/Editor/UnrealEd/Private/Editor.cpp
+++ b/Engine/Source/Editor/UnrealEd/Private/Editor.cpp
@@ -53,6 +53,7 @@
 
 #include "ComponentReregisterContext.h"
 #include "EngineModule.h"
+#include "RendererInterface.h"
 
 #if PLATFORM_WINDOWS
 // For WAVEFORMATEXTENSIBLE
diff --git a/Engine/Source/Editor/UnrealEd/Private/EditorExporters.cpp b/Engine/Source/Editor/UnrealEd/Private/EditorExporters.cpp
--- a/Engine/Source/Editor/UnrealEd/Private/EditorExporters.cpp
+++ b/Engine/Source/Editor/UnrealEd/Private/EditorExporters.cpp
@@ -25,6 +25,7 @@
 #include "UnrealExporter.h"
 #include "EngineModule.h"
 #include "InstancedFoliage.h"
+#include "RendererInterface.h"
 
 DEFINE_LOG_CATEGORY_STATIC(LogEditorExporters, Log, All);
 
diff --git a/Engine/Source/Editor/UnrealEd/Private/EditorViewportClient.cpp b/Engine/Source/Editor/UnrealEd/Private/EditorViewportClient.cpp
--- a/Engine/Source/Editor/UnrealEd/Private/EditorViewportClient.cpp
+++ b/Engine/Source/Editor/UnrealEd/Private/EditorViewportClient.cpp
@@ -15,6 +15,7 @@
 #include "IAnalyticsProvider.h"
 #include "Matinee/MatineeActor.h"
 #include "EngineModule.h"
+#include "RendererInterface.h"
 
 #define LOCTEXT_NAMESPACE "EditorViewportClient"
 
diff --git a/Engine/Source/Editor/UnrealEd/Private/ThumbnailRendering/AnimBlueprintThumbnailRenderer.cpp b/Engine/Source/Editor/UnrealEd/Private/ThumbnailRendering/AnimBlueprintThumbnailRenderer.cpp
--- a/Engine/Source/Editor/UnrealEd/Private/ThumbnailRendering/AnimBlueprintThumbnailRenderer.cpp
+++ b/Engine/Source/Editor/UnrealEd/Private/ThumbnailRendering/AnimBlueprintThumbnailRenderer.cpp
@@ -5,6 +5,7 @@
 // FPreviewScene derived helpers for rendering
 #include "ThumbnailHelpers.h"
 #include "EngineModule.h"
+#include "RendererInterface.h"
 
 UAnimBlueprintThumbnailRenderer::UAnimBlueprintThumbnailRenderer(const class FPostConstructInitializeProperties& PCIP)
 	: Super(PCIP)
diff --git a/Engine/Source/Editor/UnrealEd/Private/ThumbnailRendering/AnimSequenceThumbnailRenderer.cpp b/Engine/Source/Editor/UnrealEd/Private/ThumbnailRendering/AnimSequenceThumbnailRenderer.cpp
--- a/Engine/Source/Editor/UnrealEd/Private/ThumbnailRendering/AnimSequenceThumbnailRenderer.cpp
+++ b/Engine/Source/Editor/UnrealEd/Private/ThumbnailRendering/AnimSequenceThumbnailRenderer.cpp
@@ -5,6 +5,7 @@
 // FPreviewScene derived helpers for rendering
 #include "ThumbnailHelpers.h"
 #include "EngineModule.h"
+#include "RendererInterface.h"
 
 UAnimSequenceThumbnailRenderer::UAnimSequenceThumbnailRenderer(const class FPostConstructInitializeProperties& PCIP)
 	: Super(PCIP)
diff --git a/Engine/Source/Editor/UnrealEd/Private/ThumbnailRendering/BlendSpaceThumbnailRenderer.cpp b/Engine/Source/Editor/UnrealEd/Private/ThumbnailRendering/BlendSpaceThumbnailRenderer.cpp
--- a/Engine/Source/Editor/UnrealEd/Private/ThumbnailRendering/BlendSpaceThumbnailRenderer.cpp
+++ b/Engine/Source/Editor/UnrealEd/Private/ThumbnailRendering/BlendSpaceThumbnailRenderer.cpp
@@ -5,6 +5,7 @@
 // FPreviewScene derived helpers for rendering
 #include "ThumbnailHelpers.h"
 #include "EngineModule.h"
+#include "RendererInterface.h"
 
 UBlendSpaceThumbnailRenderer::UBlendSpaceThumbnailRenderer(const class FPostConstructInitializeProperties& PCIP)
 	: Super(PCIP)
diff --git a/Engine/Source/Editor/UnrealEd/Private/ThumbnailRendering/BlueprintThumbnailRenderer.cpp b/Engine/Source/Editor/UnrealEd/Private/ThumbnailRendering/BlueprintThumbnailRenderer.cpp
--- a/Engine/Source/Editor/UnrealEd/Private/ThumbnailRendering/BlueprintThumbnailRenderer.cpp
+++ b/Engine/Source/Editor/UnrealEd/Private/ThumbnailRendering/BlueprintThumbnailRenderer.cpp
@@ -5,6 +5,7 @@
 // FPreviewScene derived helpers for rendering
 #include "ThumbnailHelpers.h"
 #include "EngineModule.h"
+#include "RendererInterface.h"
 
 UBlueprintThumbnailRenderer::UBlueprintThumbnailRenderer(const class FPostConstructInitializeProperties& PCIP)
 	: Super(PCIP)
diff --git a/Engine/Source/Editor/UnrealEd/Private/ThumbnailRendering/LevelThumbnailRenderer.cpp b/Engine/Source/Editor/UnrealEd/Private/ThumbnailRendering/LevelThumbnailRenderer.cpp
--- a/Engine/Source/Editor/UnrealEd/Private/ThumbnailRendering/LevelThumbnailRenderer.cpp
+++ b/Engine/Source/Editor/UnrealEd/Private/ThumbnailRendering/LevelThumbnailRenderer.cpp
@@ -2,6 +2,7 @@
 
 #include "UnrealEd.h"
 #include "EngineModule.h"
+#include "RendererInterface.h"
 
 ULevelThumbnailRenderer::ULevelThumbnailRenderer(const class FPostConstructInitializeProperties& PCIP)
 	: Super(PCIP)
diff --git a/Engine/Source/Editor/UnrealEd/Private/ThumbnailRendering/MaterialFunctionThumbnailRenderer.cpp b/Engine/Source/Editor/UnrealEd/Private/ThumbnailRendering/MaterialFunctionThumbnailRenderer.cpp
--- a/Engine/Source/Editor/UnrealEd/Private/ThumbnailRendering/MaterialFunctionThumbnailRenderer.cpp
+++ b/Engine/Source/Editor/UnrealEd/Private/ThumbnailRendering/MaterialFunctionThumbnailRenderer.cpp
@@ -7,6 +7,8 @@
 // FPreviewScene derived helpers for rendering
 #include "ThumbnailHelpers.h"
 #include "EngineModule.h"
+#include "RendererInterface.h"
+#include "ShaderCompiler.h"
 
 UMaterialFunctionThumbnailRenderer::UMaterialFunctionThumbnailRenderer(const class FPostConstructInitializeProperties& PCIP)
 	: Super(PCIP)
diff --git a/Engine/Source/Editor/UnrealEd/Private/ThumbnailRendering/MaterialInstanceThumbnailRenderer.cpp b/Engine/Source/Editor/UnrealEd/Private/ThumbnailRendering/MaterialInstanceThumbnailRenderer.cpp
--- a/Engine/Source/Editor/UnrealEd/Private/ThumbnailRendering/MaterialInstanceThumbnailRenderer.cpp
+++ b/Engine/Source/Editor/UnrealEd/Private/ThumbnailRendering/MaterialInstanceThumbnailRenderer.cpp
@@ -5,6 +5,7 @@
 // FPreviewScene derived helpers for rendering
 #include "ThumbnailHelpers.h"
 #include "EngineModule.h"
+#include "RendererInterface.h"
 
 UMaterialInstanceThumbnailRenderer::UMaterialInstanceThumbnailRenderer(const class FPostConstructInitializeProperties& PCIP)
 	: Super(PCIP)
diff --git a/Engine/Source/Editor/UnrealEd/Private/ThumbnailRendering/ParticleSystemThumbnailRenderer.cpp b/Engine/Source/Editor/UnrealEd/Private/ThumbnailRendering/ParticleSystemThumbnailRenderer.cpp
--- a/Engine/Source/Editor/UnrealEd/Private/ThumbnailRendering/ParticleSystemThumbnailRenderer.cpp
+++ b/Engine/Source/Editor/UnrealEd/Private/ThumbnailRendering/ParticleSystemThumbnailRenderer.cpp
@@ -7,6 +7,7 @@
 
 #include "Particles/ParticleSystem.h"
 #include "EngineModule.h"
+#include "RendererInterface.h"
 
 UParticleSystemThumbnailRenderer::UParticleSystemThumbnailRenderer(const class FPostConstructInitializeProperties& PCIP)
 	: Super(PCIP)
diff --git a/Engine/Source/Editor/UnrealEd/Private/ThumbnailRendering/SkeletalMeshThumbnailRenderer.cpp b/Engine/Source/Editor/UnrealEd/Private/ThumbnailRendering/SkeletalMeshThumbnailRenderer.cpp
--- a/Engine/Source/Editor/UnrealEd/Private/ThumbnailRendering/SkeletalMeshThumbnailRenderer.cpp
+++ b/Engine/Source/Editor/UnrealEd/Private/ThumbnailRendering/SkeletalMeshThumbnailRenderer.cpp
@@ -5,6 +5,7 @@
 // FPreviewScene derived helpers for rendering
 #include "ThumbnailHelpers.h"
 #include "EngineModule.h"
+#include "RendererInterface.h"
 
 USkeletalMeshThumbnailRenderer::USkeletalMeshThumbnailRenderer(const class FPostConstructInitializeProperties& PCIP)
 	: Super(PCIP)
diff --git a/Engine/Source/Editor/UnrealEd/Private/ThumbnailRendering/StaticMeshThumbnailRenderer.cpp b/Engine/Source/Editor/UnrealEd/Private/ThumbnailRendering/StaticMeshThumbnailRenderer.cpp
--- a/Engine/Source/Editor/UnrealEd/Private/ThumbnailRendering/StaticMeshThumbnailRenderer.cpp
+++ b/Engine/Source/Editor/UnrealEd/Private/ThumbnailRendering/StaticMeshThumbnailRenderer.cpp
@@ -5,6 +5,7 @@
 // FPreviewScene derived helpers for rendering
 #include "ThumbnailHelpers.h"
 #include "EngineModule.h"
+#include "RendererInterface.h"
 
 UStaticMeshThumbnailRenderer::UStaticMeshThumbnailRenderer(const class FPostConstructInitializeProperties& PCIP)
 	: Super(PCIP)
diff --git a/Engine/Source/Editor/UnrealEd/Private/ThumbnailRendering/WorldThumbnailRenderer.cpp b/Engine/Source/Editor/UnrealEd/Private/ThumbnailRendering/WorldThumbnailRenderer.cpp
--- a/Engine/Source/Editor/UnrealEd/Private/ThumbnailRendering/WorldThumbnailRenderer.cpp
+++ b/Engine/Source/Editor/UnrealEd/Private/ThumbnailRendering/WorldThumbnailRenderer.cpp
@@ -1,9 +1,9 @@
 // Copyright 1998-2014 Epic Games, Inc. All Rights Reserved.
 
 #include "UnrealEd.h"
-#include "Runtime/Engine/Public/EngineModule.h"
-#include "Runtime/Engine/Public/ContentStreaming.h"
-
+#include "EngineModule.h"
+#include "ContentStreaming.h"
+#include "RendererInterface.h"
 UWorldThumbnailRenderer::UWorldThumbnailRenderer(const class FPostConstructInitializeProperties& PCIP)
 	: Super(PCIP)
 {
diff --git a/Engine/Source/Runtime/Engine/Classes/Engine/BlueprintGeneratedClass.h b/Engine/Source/Runtime/Engine/Classes/Engine/BlueprintGeneratedClass.h
--- a/Engine/Source/Runtime/Engine/Classes/Engine/BlueprintGeneratedClass.h
+++ b/Engine/Source/Runtime/Engine/Classes/Engine/BlueprintGeneratedClass.h
@@ -1,6 +1,7 @@
 // Copyright 1998-2014 Epic Games, Inc. All Rights Reserved.
 
 #pragma once
+#include "MeshBatch.h"
 #include "BlueprintGeneratedClass.generated.h"
 
 class UEdGraphPin;
@@ -146,6 +147,9 @@ struct ENGINE_API FBlueprintDebugData
 	FBlueprintDebugData()
 	{
 	}
+
+	~FBlueprintDebugData()
+	{ }
 #if WITH_EDITORONLY_DATA
 
 protected:
diff --git a/Engine/Source/Runtime/Engine/Classes/Engine/NiagaraEffectRenderer.h b/Engine/Source/Runtime/Engine/Classes/Engine/NiagaraEffectRenderer.h
--- a/Engine/Source/Runtime/Engine/Classes/Engine/NiagaraEffectRenderer.h
+++ b/Engine/Source/Runtime/Engine/Classes/Engine/NiagaraEffectRenderer.h
@@ -3,6 +3,7 @@
 /*==============================================================================
 NiagaraEffectRenderer.h: Base class for Niagara render modules
 ==============================================================================*/
+#include "SceneUtils.h"
 
 DECLARE_CYCLE_STAT(TEXT("PreRenderView"), STAT_NiagaraPreRenderView, STATGROUP_Niagara);
 DECLARE_CYCLE_STAT(TEXT("Render Total"), STAT_NiagaraRender, STATGROUP_Niagara);
diff --git a/Engine/Source/Runtime/Engine/Private/GPUSkinCache.cpp b/Engine/Source/Runtime/Engine/Private/GPUSkinCache.cpp
--- a/Engine/Source/Runtime/Engine/Private/GPUSkinCache.cpp
+++ b/Engine/Source/Runtime/Engine/Private/GPUSkinCache.cpp
@@ -12,6 +12,7 @@
 #include "SkeletalRenderGPUSkin.h"
 #include "GPUSkinCache.h"
 #include "ShaderParameterUtils.h"
+#include "SceneUtils.h"
 
 #define GPUSKINCACHE_FRAMES	3
 
diff --git a/Engine/Source/Runtime/Engine/Private/GPUSort.cpp b/Engine/Source/Runtime/Engine/Private/GPUSort.cpp
--- a/Engine/Source/Runtime/Engine/Private/GPUSort.cpp
+++ b/Engine/Source/Runtime/Engine/Private/GPUSort.cpp
@@ -11,6 +11,7 @@
 #include "ShaderParameters.h"
 #include "ShaderParameterUtils.h"
 #include "GlobalShader.h"
+#include "SceneUtils.h"
 
 DEFINE_LOG_CATEGORY_STATIC(LogGPUSort, Log, All);
 
diff --git a/Engine/Source/Runtime/Engine/Private/GameEngine.cpp b/Engine/Source/Runtime/Engine/Private/GameEngine.cpp
--- a/Engine/Source/Runtime/Engine/Private/GameEngine.cpp
+++ b/Engine/Source/Runtime/Engine/Private/GameEngine.cpp
@@ -27,6 +27,7 @@
 #include "SynthBenchmark.h"
 
 #include "IHeadMountedDisplay.h"
+#include "RendererInterface.h"
 
 DEFINE_LOG_CATEGORY_STATIC(LogEngine, Log, All);
 
diff --git a/Engine/Source/Runtime/Engine/Private/GameViewportClient.cpp b/Engine/Source/Runtime/Engine/Private/GameViewportClient.cpp
--- a/Engine/Source/Runtime/Engine/Private/GameViewportClient.cpp
+++ b/Engine/Source/Runtime/Engine/Private/GameViewportClient.cpp
@@ -21,6 +21,8 @@
 #include "HighResScreenshot.h"
 #include "Particles/ParticleSystemComponent.h"
 #include "Runtime/GameLiveStreaming/Public/IGameLiveStreaming.h"
+#include "BufferVisualizationData.h"
+#include "RendererInterface.h"
 
 /** This variable allows forcing full screen of the first player controller viewport, even if there are multiple controllers plugged in and no cinematic playing. */
 bool GForceFullscreen = false;
diff --git a/Engine/Source/Runtime/Engine/Private/NiagaraComponent.cpp b/Engine/Source/Runtime/Engine/Private/NiagaraComponent.cpp
--- a/Engine/Source/Runtime/Engine/Private/NiagaraComponent.cpp
+++ b/Engine/Source/Runtime/Engine/Private/NiagaraComponent.cpp
@@ -8,6 +8,8 @@
 #include "Particles/ParticleResources.h"
 #include "Engine/NiagaraConstants.h"
 #include "Engine/NiagaraEffectRenderer.h"
+#include "MeshBatch.h"
+#include "SceneUtils.h"
 
 DECLARE_CYCLE_STAT(TEXT("Tick"),STAT_NiagaraTick,STATGROUP_Niagara);
 DECLARE_CYCLE_STAT(TEXT("Simulate"),STAT_NiagaraSimulate,STATGROUP_Niagara);
diff --git a/Engine/Source/Runtime/Engine/Private/Particles/ParticleCurveTexture.cpp b/Engine/Source/Runtime/Engine/Private/Particles/ParticleCurveTexture.cpp
--- a/Engine/Source/Runtime/Engine/Private/Particles/ParticleCurveTexture.cpp
+++ b/Engine/Source/Runtime/Engine/Private/Particles/ParticleCurveTexture.cpp
@@ -14,6 +14,7 @@
 #include "RHIStaticStates.h"
 #include "GlobalShader.h"
 #include "FXSystem.h"
+#include "SceneUtils.h"
 
 /** The texture size allocated for particle curves. */
 extern const int32 GParticleCurveTextureSizeX = 256;
diff --git a/Engine/Source/Runtime/Engine/Private/Particles/ParticleGpuSimulation.cpp b/Engine/Source/Runtime/Engine/Private/Particles/ParticleGpuSimulation.cpp
--- a/Engine/Source/Runtime/Engine/Private/Particles/ParticleGpuSimulation.cpp
+++ b/Engine/Source/Runtime/Engine/Private/Particles/ParticleGpuSimulation.cpp
@@ -28,6 +28,8 @@
 #include "Particles/ParticleSpriteEmitter.h"
 #include "Particles/ParticleSystemComponent.h"
 #include "VectorField/VectorField.h"
+#include "SceneUtils.h"
+#include "MeshBatch.h"
 
 /*------------------------------------------------------------------------------
 	Constants to tune memory and performance for GPU particle simulation.
diff --git a/Engine/Source/Runtime/Engine/Private/Particles/ParticleSortingGPU.cpp b/Engine/Source/Runtime/Engine/Private/Particles/ParticleSortingGPU.cpp
--- a/Engine/Source/Runtime/Engine/Private/Particles/ParticleSortingGPU.cpp
+++ b/Engine/Source/Runtime/Engine/Private/Particles/ParticleSortingGPU.cpp
@@ -13,6 +13,7 @@
 #include "ShaderParameterUtils.h"
 #include "GlobalShader.h"
 #include "../GPUSort.h"
+#include "SceneUtils.h"
 
 /*------------------------------------------------------------------------------
 	Shaders used to generate particle sort keys.
diff --git a/Engine/Source/Runtime/Engine/Private/Particles/ParticleSystemRender.cpp b/Engine/Source/Runtime/Engine/Private/Particles/ParticleSystemRender.cpp
--- a/Engine/Source/Runtime/Engine/Private/Particles/ParticleSystemRender.cpp
+++ b/Engine/Source/Runtime/Engine/Private/Particles/ParticleSystemRender.cpp
@@ -18,6 +18,8 @@
 #include "Particles/ParticleModuleRequired.h"
 #include "Particles/ParticleSpriteEmitter.h"
 #include "Particles/ParticleSystemComponent.h"
+#include "ParticleBeamTrailVertexFactory.h"
+#include "MeshBatch.h"
 
 /** 
  * Whether to track particle rendering stats.  
diff --git a/Engine/Source/Runtime/Engine/Private/PrimitiveDrawingUtils.cpp b/Engine/Source/Runtime/Engine/Private/PrimitiveDrawingUtils.cpp
--- a/Engine/Source/Runtime/Engine/Private/PrimitiveDrawingUtils.cpp
+++ b/Engine/Source/Runtime/Engine/Private/PrimitiveDrawingUtils.cpp
@@ -5,6 +5,9 @@
 #include "SceneManagement.h"
 #include "StaticMeshResources.h"
 #include "EngineModule.h"
+#include "SceneUtils.h"
+#include "MeshBatch.h"
+#include "RendererInterface.h"
 
 /** Emits draw events for a given FMeshBatch and the FPrimitiveSceneProxy corresponding to that mesh element. */
 void EmitMeshDrawEvents(const FPrimitiveSceneProxy* PrimitiveSceneProxy, const FMeshBatch& Mesh)
diff --git a/Engine/Source/Runtime/Engine/Private/RendererSupport.cpp b/Engine/Source/Runtime/Engine/Private/RendererSupport.cpp
--- a/Engine/Source/Runtime/Engine/Private/RendererSupport.cpp
+++ b/Engine/Source/Runtime/Engine/Private/RendererSupport.cpp
@@ -11,6 +11,7 @@
 #include "GlobalShader.h"
 #include "Slate.h"
 #include "EngineModule.h"
+#include "RendererInterface.h"
 
 /** Clears and optionally backs up all references to renderer module classes in other modules, particularly engine. */
 void ClearReferencesToRendererModuleClasses(
diff --git a/Engine/Source/Runtime/Engine/Private/Rendering/SubsurfaceProfile.cpp b/Engine/Source/Runtime/Engine/Private/Rendering/SubsurfaceProfile.cpp
--- a/Engine/Source/Runtime/Engine/Private/Rendering/SubsurfaceProfile.cpp
+++ b/Engine/Source/Runtime/Engine/Private/Rendering/SubsurfaceProfile.cpp
@@ -2,6 +2,7 @@
 
 #include "EnginePrivate.h"
 #include "SeparableSSS.h"
+#include "RendererInterface.h"
 
 DEFINE_LOG_CATEGORY_STATIC(LogSubsurfaceProfile, Log, All);
 
diff --git a/Engine/Source/Runtime/Engine/Private/SceneUtils.cpp b/Engine/Source/Runtime/Engine/Private/SceneUtils.cpp
--- a/Engine/Source/Runtime/Engine/Private/SceneUtils.cpp
+++ b/Engine/Source/Runtime/Engine/Private/SceneUtils.cpp
@@ -2,6 +2,7 @@
 
 #include "EnginePrivate.h"
 #include "DynamicRHI.h"
+#include "SceneUtils.h"
 
 #if WANTS_DRAW_MESH_EVENTS && PLATFORM_SUPPORTS_DRAW_MESH_EVENTS
 
@@ -10,15 +11,7 @@ bool FDrawEvent::IsInRenderingThread_Internal()
 	return IsInRenderingThread();
 }
 
-FDrawEvent::~FDrawEvent()
-{
-	if (bDrawEventHasBeenEmitted)
-	{
-		GDynamicRHI->PopEvent();
-	}
-}
-
-void FDrawEvent::Start(const FColor& Color, const TCHAR* Fmt, ...)
+ENGINE_API void FDrawEvent::Start(const FColor& Color, const TCHAR* Fmt, ...)
 {
 	if (GRHICommandList.Bypass()) // todo interleave draw evens into command buffers
 	{
diff --git a/Engine/Source/Runtime/Engine/Private/SceneView.cpp b/Engine/Source/Runtime/Engine/Private/SceneView.cpp
--- a/Engine/Source/Runtime/Engine/Private/SceneView.cpp
+++ b/Engine/Source/Runtime/Engine/Private/SceneView.cpp
@@ -11,6 +11,8 @@
 #include "StereoRendering.h"
 #include "HighResScreenshot.h"
 #include "Slate/SceneViewport.h"
+#include "RendererInterface.h"
+#include "BufferVisualizationData.h"
 
 DEFINE_LOG_CATEGORY(LogBufferVisualization);
 
diff --git a/Engine/Source/Runtime/Engine/Private/TileRendering.cpp b/Engine/Source/Runtime/Engine/Private/TileRendering.cpp
--- a/Engine/Source/Runtime/Engine/Private/TileRendering.cpp
+++ b/Engine/Source/Runtime/Engine/Private/TileRendering.cpp
@@ -8,6 +8,8 @@
 #include "TileRendering.h"
 #include "EngineModule.h"
 #include "LocalVertexFactory.h"
+#include "MeshBatch.h"
+#include "RendererInterface.h"
 
 /** 
 * vertex data for a screen quad 
diff --git a/Engine/Source/Runtime/Engine/Private/UnrealClient.cpp b/Engine/Source/Runtime/Engine/Private/UnrealClient.cpp
--- a/Engine/Source/Runtime/Engine/Private/UnrealClient.cpp
+++ b/Engine/Source/Runtime/Engine/Private/UnrealClient.cpp
@@ -16,6 +16,7 @@
 #include "EngineModule.h"
 #include "EngineModule.h"
 #include "ContentStreaming.h"
+#include "SceneUtils.h"
 
 DEFINE_LOG_CATEGORY_STATIC(LogClient, Log, All);
 
diff --git a/Engine/Source/Runtime/Engine/Private/UnrealEngine.cpp b/Engine/Source/Runtime/Engine/Private/UnrealEngine.cpp
--- a/Engine/Source/Runtime/Engine/Private/UnrealEngine.cpp
+++ b/Engine/Source/Runtime/Engine/Private/UnrealEngine.cpp
@@ -79,6 +79,7 @@
 #include "UnrealExporter.h"
 #include "ComponentReregisterContext.h"
 #include "ContentStreaming.h"
+#include "BufferVisualizationData.h"
 
 #include "Engine/GameInstance.h"
 #include "HotReloadInterface.h"
diff --git a/Engine/Source/Runtime/Engine/Private/UserInterface/Canvas.cpp b/Engine/Source/Runtime/Engine/Private/UserInterface/Canvas.cpp
--- a/Engine/Source/Runtime/Engine/Private/UserInterface/Canvas.cpp
+++ b/Engine/Source/Runtime/Engine/Private/UserInterface/Canvas.cpp
@@ -12,6 +12,7 @@
 
 #include "IHeadMountedDisplay.h"
 #include "Debug/ReporterGraph.h"
+#include "SceneUtils.h"
 
 DEFINE_LOG_CATEGORY_STATIC(LogCanvas, Log, All);
 
diff --git a/Engine/Source/Runtime/Engine/Private/UserInterface/VisualizeRT.cpp b/Engine/Source/Runtime/Engine/Private/UserInterface/VisualizeRT.cpp
--- a/Engine/Source/Runtime/Engine/Private/UserInterface/VisualizeRT.cpp
+++ b/Engine/Source/Runtime/Engine/Private/UserInterface/VisualizeRT.cpp
@@ -6,6 +6,7 @@
 
 #include "EnginePrivate.h"
 #include "Slate.h"
+#include "RendererInterface.h"
 
 #define ENABLE_IMAGES	0
 
diff --git a/Engine/Source/Runtime/Engine/Private/VectorField.cpp b/Engine/Source/Runtime/Engine/Private/VectorField.cpp
--- a/Engine/Source/Runtime/Engine/Private/VectorField.cpp
+++ b/Engine/Source/Runtime/Engine/Private/VectorField.cpp
@@ -17,6 +17,7 @@
 #include "VectorField/VectorField.h"
 #include "VectorField/VectorFieldAnimated.h"
 #include "VectorField/VectorFieldStatic.h"
+#include "SceneUtils.h"
 
 #define MAX_GLOBAL_VECTOR_FIELDS (16)
 DEFINE_LOG_CATEGORY(LogVectorField)
diff --git a/Engine/Source/Runtime/Engine/Public/SceneManagement.h b/Engine/Source/Runtime/Engine/Public/SceneManagement.h
--- a/Engine/Source/Runtime/Engine/Public/SceneManagement.h
+++ b/Engine/Source/Runtime/Engine/Public/SceneManagement.h
@@ -7,34 +7,21 @@
 #pragma once
 
 // Includes the draw mesh macros
-#include "SceneUtils.h"
 #include "UniformBuffer.h"
-#include "BufferVisualizationData.h"
 #include "ConvexVolume.h"
-#include "SystemSettings.h"
 #include "Engine/TextureLightProfile.h"
-#include "Engine/World.h"
-#include "RendererInterface.h"
-#include "Engine/GameViewportClient.h"
 #include "SceneTypes.h"
 #include "SceneView.h"
-#include "FinalPostProcessSettings.h"
-#include "PrimitiveUniformShaderParameters.h"
-#include "MeshBatch.h"
-#include "PrimitiveViewRelevance.h"
-#include "SceneInterface.h"
 #include "RHIDefinitions.h"
 #include "ChunkedArray.h"
 #include "BatchedElements.h"
+#include "MeshBatch.h"
+#include "RendererInterface.h"
 
 // Forward declarations.
-class FSceneViewFamily;
 class FLightSceneInfo;
 class ULightComponent;
 class UDecalComponent;
-class FIndexBuffer;
-class FVertexFactory;
-class ISceneViewExtension;
 class HHitProxy;
 struct FDynamicMeshVertex;
 
diff --git a/Engine/Source/Runtime/Engine/Public/SceneUtils.h b/Engine/Source/Runtime/Engine/Public/SceneUtils.h
--- a/Engine/Source/Runtime/Engine/Public/SceneUtils.h
+++ b/Engine/Source/Runtime/Engine/Public/SceneUtils.h
@@ -50,7 +50,13 @@ inline void appSetCounterValue(const TCHAR* CounterName, float Value) {}
 		/**
 		 * Terminate the event based upon scope
 		 */
-		 ~FDrawEvent();
+		~FDrawEvent()
+		{
+			if (bDrawEventHasBeenEmitted)
+			{
+				GDynamicRHI->PopEvent();
+			}
+		}
 
 		/**
 		 * Function for logging a PIX event with var args
diff --git a/Engine/Source/Runtime/GameLiveStreaming/GameLiveStreaming.cpp b/Engine/Source/Runtime/GameLiveStreaming/GameLiveStreaming.cpp
--- a/Engine/Source/Runtime/GameLiveStreaming/GameLiveStreaming.cpp
+++ b/Engine/Source/Runtime/GameLiveStreaming/GameLiveStreaming.cpp
@@ -7,6 +7,7 @@
 #include "ScreenRendering.h"
 #include "RenderCore.h"
 #include "RHIStaticStates.h"
+#include "RendererInterface.h"
 
 IMPLEMENT_MODULE( FGameLiveStreaming, GameLiveStreaming );
 
diff --git a/Engine/Source/Runtime/Renderer/Private/AtmosphereRendering.cpp b/Engine/Source/Runtime/Renderer/Private/AtmosphereRendering.cpp
--- a/Engine/Source/Runtime/Renderer/Private/AtmosphereRendering.cpp
+++ b/Engine/Source/Runtime/Renderer/Private/AtmosphereRendering.cpp
@@ -10,6 +10,8 @@
 #include "AtmosphereRendering.h"
 #include "../../Engine/Private/Atmosphere/Atmosphere.h"
 #include "AtmosphereTextures.h"
+#include "ShowFlags.h"
+#include "SceneUtils.h"
 
 class FAtmosphereShaderPrecomputeTextureParameters
 {
diff --git a/Engine/Source/Runtime/Renderer/Private/CompositionLighting/CompositionLighting.cpp b/Engine/Source/Runtime/Renderer/Private/CompositionLighting/CompositionLighting.cpp
--- a/Engine/Source/Runtime/Renderer/Private/CompositionLighting/CompositionLighting.cpp
+++ b/Engine/Source/Runtime/Renderer/Private/CompositionLighting/CompositionLighting.cpp
@@ -20,6 +20,7 @@
 #include "PostProcessWeightedSampleSum.h"
 #include "PostProcessTemporalAA.h"
 #include "PostProcessSubsurface.h"
+#include "SceneUtils.h"
 
 /** The global center for all deferred lighting activities. */
 FCompositionLighting GCompositionLighting;
diff --git a/Engine/Source/Runtime/Renderer/Private/CompositionLighting/PostProcessAmbient.cpp b/Engine/Source/Runtime/Renderer/Private/CompositionLighting/PostProcessAmbient.cpp
--- a/Engine/Source/Runtime/Renderer/Private/CompositionLighting/PostProcessAmbient.cpp
+++ b/Engine/Source/Runtime/Renderer/Private/CompositionLighting/PostProcessAmbient.cpp
@@ -9,6 +9,7 @@
 #include "SceneFilterRendering.h"
 #include "PostProcessing.h"
 #include "PostProcessAmbient.h"
+#include "SceneUtils.h"
 
 /** Encapsulates the post processing ambient pixel shader. */
 class FPostProcessAmbientPS : public FGlobalShader
diff --git a/Engine/Source/Runtime/Renderer/Private/CompositionLighting/PostProcessAmbientOcclusion.cpp b/Engine/Source/Runtime/Renderer/Private/CompositionLighting/PostProcessAmbientOcclusion.cpp
--- a/Engine/Source/Runtime/Renderer/Private/CompositionLighting/PostProcessAmbientOcclusion.cpp
+++ b/Engine/Source/Runtime/Renderer/Private/CompositionLighting/PostProcessAmbientOcclusion.cpp
@@ -9,6 +9,7 @@
 #include "SceneFilterRendering.h"
 #include "PostProcessing.h"
 #include "PostProcessAmbientOcclusion.h"
+#include "SceneUtils.h"
 
 static TAutoConsoleVariable<int32> CVarAmbientOcclusionSampleSetQuality(
 	TEXT("r.AmbientOcclusionSampleSetQuality"),
diff --git a/Engine/Source/Runtime/Renderer/Private/CompositionLighting/PostProcessDeferredDecals.cpp b/Engine/Source/Runtime/Renderer/Private/CompositionLighting/PostProcessDeferredDecals.cpp
--- a/Engine/Source/Runtime/Renderer/Private/CompositionLighting/PostProcessDeferredDecals.cpp
+++ b/Engine/Source/Runtime/Renderer/Private/CompositionLighting/PostProcessDeferredDecals.cpp
@@ -10,6 +10,7 @@
 #include "PostProcessing.h"
 #include "PostProcessDeferredDecals.h"
 #include "ScreenRendering.h"
+#include "SceneUtils.h"
 
 static TAutoConsoleVariable<float> CVarStencilSizeThreshold(
 	TEXT("r.Decal.StencilSizeThreshold"),
diff --git a/Engine/Source/Runtime/Renderer/Private/CompositionLighting/PostProcessLpvIndirect.cpp b/Engine/Source/Runtime/Renderer/Private/CompositionLighting/PostProcessLpvIndirect.cpp
--- a/Engine/Source/Runtime/Renderer/Private/CompositionLighting/PostProcessLpvIndirect.cpp
+++ b/Engine/Source/Runtime/Renderer/Private/CompositionLighting/PostProcessLpvIndirect.cpp
@@ -17,6 +17,7 @@
 #include "PostProcessLpvIndirect.h"
 #include "LightPropagationVolume.h"
 #include "UniformBuffer.h"
+#include "SceneUtils.h"
 
 IMPLEMENT_UNIFORM_BUFFER_STRUCT(FLpvReadUniformBufferParameters,TEXT("LpvRead"));
 typedef TUniformBufferRef<FLpvReadUniformBufferParameters> FLpvReadUniformBufferRef;
diff --git a/Engine/Source/Runtime/Renderer/Private/CompositionLighting/PostProcessPassThrough.cpp b/Engine/Source/Runtime/Renderer/Private/CompositionLighting/PostProcessPassThrough.cpp
--- a/Engine/Source/Runtime/Renderer/Private/CompositionLighting/PostProcessPassThrough.cpp
+++ b/Engine/Source/Runtime/Renderer/Private/CompositionLighting/PostProcessPassThrough.cpp
@@ -9,6 +9,7 @@
 #include "SceneFilterRendering.h"
 #include "PostProcessPassThrough.h"
 #include "PostProcessing.h"
+#include "SceneUtils.h"
 
 /** Encapsulates a simple copy pixel shader. */
 class FPostProcessPassThroughPS : public FGlobalShader
diff --git a/Engine/Source/Runtime/Renderer/Private/CustomDepthRendering.cpp b/Engine/Source/Runtime/Renderer/Private/CustomDepthRendering.cpp
--- a/Engine/Source/Runtime/Renderer/Private/CustomDepthRendering.cpp
+++ b/Engine/Source/Runtime/Renderer/Private/CustomDepthRendering.cpp
@@ -11,6 +11,7 @@
 #include "PostProcessing.h"
 #include "RenderingCompositionGraph.h"
 #include "SceneFilterRendering.h"
+#include "SceneUtils.h"
 
 /*-----------------------------------------------------------------------------
 	FCustomDepthPrimSet
diff --git a/Engine/Source/Runtime/Renderer/Private/DeferredShadingRenderer.cpp b/Engine/Source/Runtime/Renderer/Private/DeferredShadingRenderer.cpp
--- a/Engine/Source/Runtime/Renderer/Private/DeferredShadingRenderer.cpp
+++ b/Engine/Source/Runtime/Renderer/Private/DeferredShadingRenderer.cpp
@@ -15,6 +15,8 @@
 #include "OneColorShader.h"
 #include "CompositionLighting/PostProcessDeferredDecals.h"
 #include "LightPropagationVolume.h"
+#include "DeferredShadingRenderer.h"
+#include "SceneUtils.h"
 
 TAutoConsoleVariable<int32> CVarEarlyZPass(
 	TEXT("r.EarlyZPass"),
diff --git a/Engine/Source/Runtime/Renderer/Private/DistanceFieldSurfaceCacheLighting.cpp b/Engine/Source/Runtime/Renderer/Private/DistanceFieldSurfaceCacheLighting.cpp
--- a/Engine/Source/Runtime/Renderer/Private/DistanceFieldSurfaceCacheLighting.cpp
+++ b/Engine/Source/Runtime/Renderer/Private/DistanceFieldSurfaceCacheLighting.cpp
@@ -15,6 +15,7 @@
 #include "DistanceFieldLightingShared.h"
 #include "PostProcessAmbientOcclusion.h"
 #include "RHICommandList.h"
+#include "SceneUtils.h"
 
 int32 GDistanceFieldAO = 1;
 FAutoConsoleVariableRef CVarDistanceFieldAO(
diff --git a/Engine/Source/Runtime/Renderer/Private/DistortionRendering.cpp b/Engine/Source/Runtime/Renderer/Private/DistortionRendering.cpp
--- a/Engine/Source/Runtime/Renderer/Private/DistortionRendering.cpp
+++ b/Engine/Source/Runtime/Renderer/Private/DistortionRendering.cpp
@@ -11,6 +11,7 @@
 #include "PostProcessing.h"
 #include "RenderingCompositionGraph.h"
 #include "SceneFilterRendering.h"
+#include "SceneUtils.h"
 
 /**
 * A pixel shader for rendering the full screen refraction pass
diff --git a/Engine/Source/Runtime/Renderer/Private/DrawingPolicy.cpp b/Engine/Source/Runtime/Renderer/Private/DrawingPolicy.cpp
--- a/Engine/Source/Runtime/Renderer/Private/DrawingPolicy.cpp
+++ b/Engine/Source/Runtime/Renderer/Private/DrawingPolicy.cpp
@@ -6,6 +6,7 @@
 
 #include "RendererPrivate.h"
 #include "ScenePrivate.h"
+#include "SceneUtils.h"
 
 int32 GEmitMeshDrawEvent = 0;
 static FAutoConsoleVariableRef CVarEmitMeshDrawEvent(
diff --git a/Engine/Source/Runtime/Renderer/Private/FogRendering.cpp b/Engine/Source/Runtime/Renderer/Private/FogRendering.cpp
--- a/Engine/Source/Runtime/Renderer/Private/FogRendering.cpp
+++ b/Engine/Source/Runtime/Renderer/Private/FogRendering.cpp
@@ -6,6 +6,7 @@
 
 #include "RendererPrivate.h"
 #include "ScenePrivate.h"
+#include "SceneUtils.h"
 
 #if !(UE_BUILD_SHIPPING || UE_BUILD_TEST)
 static TAutoConsoleVariable<float> CVarFogStartDistance(
diff --git a/Engine/Source/Runtime/Renderer/Private/ForwardBasePassRendering.cpp b/Engine/Source/Runtime/Renderer/Private/ForwardBasePassRendering.cpp
--- a/Engine/Source/Runtime/Renderer/Private/ForwardBasePassRendering.cpp
+++ b/Engine/Source/Runtime/Renderer/Private/ForwardBasePassRendering.cpp
@@ -6,6 +6,7 @@
 
 #include "RendererPrivate.h"
 #include "ScenePrivate.h"
+#include "SceneUtils.h"
 
 #define IMPLEMENT_FORWARD_SHADING_BASEPASS_LIGHTMAPPED_SHADER_TYPE(LightMapPolicyType,LightMapPolicyName) \
 	typedef TBasePassForForwardShadingVS< LightMapPolicyType, LDR_GAMMA_32 > TBasePassForForwardShadingVS##LightMapPolicyName##LDRGamma32; \
diff --git a/Engine/Source/Runtime/Renderer/Private/ForwardShadingRenderer.cpp b/Engine/Source/Runtime/Renderer/Private/ForwardShadingRenderer.cpp
--- a/Engine/Source/Runtime/Renderer/Private/ForwardShadingRenderer.cpp
+++ b/Engine/Source/Runtime/Renderer/Private/ForwardShadingRenderer.cpp
@@ -11,7 +11,7 @@
 #include "PostProcessing.h"
 #include "SceneFilterRendering.h"
 #include "PostProcessMobile.h"
-
+#include "SceneUtils.h"
 
 FForwardShadingSceneRenderer::FForwardShadingSceneRenderer(const FSceneViewFamily* InViewFamily,FHitProxyConsumer* HitProxyConsumer)
 	:	FSceneRenderer(InViewFamily, HitProxyConsumer)
diff --git a/Engine/Source/Runtime/Renderer/Private/ForwardTranslucentRendering.cpp b/Engine/Source/Runtime/Renderer/Private/ForwardTranslucentRendering.cpp
--- a/Engine/Source/Runtime/Renderer/Private/ForwardTranslucentRendering.cpp
+++ b/Engine/Source/Runtime/Renderer/Private/ForwardTranslucentRendering.cpp
@@ -8,7 +8,7 @@
 #include "ScenePrivate.h"
 #include "ScreenRendering.h"
 #include "SceneFilterRendering.h"
-
+#include "SceneUtils.h"
 
 /** Pixel shader used to copy scene color into another texture so that materials can read from scene color with a node. */
 class FForwardCopySceneAlphaPS : public FGlobalShader
diff --git a/Engine/Source/Runtime/Renderer/Private/GPUBenchmark.cpp b/Engine/Source/Runtime/Renderer/Private/GPUBenchmark.cpp
--- a/Engine/Source/Runtime/Renderer/Private/GPUBenchmark.cpp
+++ b/Engine/Source/Runtime/Renderer/Private/GPUBenchmark.cpp
@@ -8,6 +8,7 @@
 #include "ScenePrivate.h"
 #include "SceneFilterRendering.h"
 #include "GPUBenchmark.h"
+#include "SceneUtils.h"
 
 static uint32 GBenchmarkResolution = 512;
 
diff --git a/Engine/Source/Runtime/Renderer/Private/LightFunctionRendering.cpp b/Engine/Source/Runtime/Renderer/Private/LightFunctionRendering.cpp
--- a/Engine/Source/Runtime/Renderer/Private/LightFunctionRendering.cpp
+++ b/Engine/Source/Runtime/Renderer/Private/LightFunctionRendering.cpp
@@ -7,6 +7,7 @@
 #include "RendererPrivate.h"
 #include "ScenePrivate.h"
 #include "LightRendering.h"
+#include "SceneUtils.h"
 
 /**
  * A vertex shader for projecting a light function onto the scene.
diff --git a/Engine/Source/Runtime/Renderer/Private/LightMapDensityRendering.cpp b/Engine/Source/Runtime/Renderer/Private/LightMapDensityRendering.cpp
--- a/Engine/Source/Runtime/Renderer/Private/LightMapDensityRendering.cpp
+++ b/Engine/Source/Runtime/Renderer/Private/LightMapDensityRendering.cpp
@@ -6,6 +6,7 @@
 
 #include "RendererPrivate.h"
 #include "ScenePrivate.h"
+#include "SceneUtils.h"
 
 //-----------------------------------------------------------------------------
 
diff --git a/Engine/Source/Runtime/Renderer/Private/LightPropagationVolume.cpp b/Engine/Source/Runtime/Renderer/Private/LightPropagationVolume.cpp
--- a/Engine/Source/Runtime/Renderer/Private/LightPropagationVolume.cpp
+++ b/Engine/Source/Runtime/Renderer/Private/LightPropagationVolume.cpp
@@ -14,6 +14,7 @@
 #include "ScenePrivate.h"
 #include "LightPropagationVolume.h"
 #include "UniformBuffer.h"
+#include "SceneUtils.h"
 
 static TAutoConsoleVariable<int32> CVarLightPropagationVolume(
 	TEXT("r.LightPropagationVolume"),
diff --git a/Engine/Source/Runtime/Renderer/Private/LightPropagationVolumeVisualisation.cpp b/Engine/Source/Runtime/Renderer/Private/LightPropagationVolumeVisualisation.cpp
--- a/Engine/Source/Runtime/Renderer/Private/LightPropagationVolumeVisualisation.cpp
+++ b/Engine/Source/Runtime/Renderer/Private/LightPropagationVolumeVisualisation.cpp
@@ -13,6 +13,7 @@
 #include "RendererPrivate.h"
 #include "ScenePrivate.h"
 #include "LightPropagationVolume.h"
+#include "SceneUtils.h"
 
 // ----------------------------------------------------------------------------
 
diff --git a/Engine/Source/Runtime/Renderer/Private/LightRendering.cpp b/Engine/Source/Runtime/Renderer/Private/LightRendering.cpp
--- a/Engine/Source/Runtime/Renderer/Private/LightRendering.cpp
+++ b/Engine/Source/Runtime/Renderer/Private/LightRendering.cpp
@@ -11,6 +11,7 @@
 #include "ShaderParameters.h"
 #include "LightRendering.h"
 #include "LightPropagationVolume.h"
+#include "SceneUtils.h"
 
 IMPLEMENT_UNIFORM_BUFFER_STRUCT(FDeferredLightUniformStruct,TEXT("DeferredLightUniforms"));
 
diff --git a/Engine/Source/Runtime/Renderer/Private/LightShaftRendering.cpp b/Engine/Source/Runtime/Renderer/Private/LightShaftRendering.cpp
--- a/Engine/Source/Runtime/Renderer/Private/LightShaftRendering.cpp
+++ b/Engine/Source/Runtime/Renderer/Private/LightShaftRendering.cpp
@@ -9,6 +9,7 @@
 #include "SceneFilterRendering.h"
 #include "PostProcess/PostProcessing.h"
 #include "PostProcess/PostProcessTemporalAA.h"
+#include "SceneUtils.h"
 
 
 /** Tweaked values from UE3 implementation **/
diff --git a/Engine/Source/Runtime/Renderer/Private/PostProcess/PostProcessAA.cpp b/Engine/Source/Runtime/Renderer/Private/PostProcess/PostProcessAA.cpp
--- a/Engine/Source/Runtime/Renderer/Private/PostProcess/PostProcessAA.cpp
+++ b/Engine/Source/Runtime/Renderer/Private/PostProcess/PostProcessAA.cpp
@@ -9,6 +9,7 @@
 #include "SceneFilterRendering.h"
 #include "PostProcessAA.h"
 #include "PostProcessing.h"
+#include "SceneUtils.h"
 
 /** Encapsulates the post processing anti aliasing pixel shader. */
 // Quality 1..6
diff --git a/Engine/Source/Runtime/Renderer/Private/PostProcess/PostProcessBloomSetup.cpp b/Engine/Source/Runtime/Renderer/Private/PostProcess/PostProcessBloomSetup.cpp
--- a/Engine/Source/Runtime/Renderer/Private/PostProcess/PostProcessBloomSetup.cpp
+++ b/Engine/Source/Runtime/Renderer/Private/PostProcess/PostProcessBloomSetup.cpp
@@ -11,6 +11,7 @@
 #include "PostProcessing.h"
 #include "PostProcessHistogram.h"
 #include "PostProcessEyeAdaptation.h"
+#include "SceneUtils.h"
 
 /** Encapsulates the post processing bloom threshold pixel shader. */
 class FPostProcessBloomSetupPS : public FGlobalShader
diff --git a/Engine/Source/Runtime/Renderer/Private/PostProcess/PostProcessBokehDOF.cpp b/Engine/Source/Runtime/Renderer/Private/PostProcess/PostProcessBokehDOF.cpp
--- a/Engine/Source/Runtime/Renderer/Private/PostProcess/PostProcessBokehDOF.cpp
+++ b/Engine/Source/Runtime/Renderer/Private/PostProcess/PostProcessBokehDOF.cpp
@@ -10,6 +10,7 @@
 #include "PostProcessPassThrough.h"
 #include "PostProcessing.h"
 #include "PostProcessBokehDOF.h"
+#include "SceneUtils.h"
 
 
 
diff --git a/Engine/Source/Runtime/Renderer/Private/PostProcess/PostProcessBokehDOFRecombine.cpp b/Engine/Source/Runtime/Renderer/Private/PostProcess/PostProcessBokehDOFRecombine.cpp
--- a/Engine/Source/Runtime/Renderer/Private/PostProcess/PostProcessBokehDOFRecombine.cpp
+++ b/Engine/Source/Runtime/Renderer/Private/PostProcess/PostProcessBokehDOFRecombine.cpp
@@ -11,6 +11,7 @@
 #include "PostProcessing.h"
 #include "PostProcessBokehDOFRecombine.h"
 #include "PostProcessBokehDOF.h"
+#include "SceneUtils.h"
 
 
 /**
diff --git a/Engine/Source/Runtime/Renderer/Private/PostProcess/PostProcessBusyWait.cpp b/Engine/Source/Runtime/Renderer/Private/PostProcess/PostProcessBusyWait.cpp
--- a/Engine/Source/Runtime/Renderer/Private/PostProcess/PostProcessBusyWait.cpp
+++ b/Engine/Source/Runtime/Renderer/Private/PostProcess/PostProcessBusyWait.cpp
@@ -9,6 +9,7 @@
 #include "SceneFilterRendering.h"
 #include "PostProcessBusyWait.h"
 #include "PostProcessing.h"
+#include "SceneUtils.h"
 
 #if !(UE_BUILD_SHIPPING || UE_BUILD_TEST)
 static TAutoConsoleVariable<float> CVarSetGPUBusyWait(
diff --git a/Engine/Source/Runtime/Renderer/Private/PostProcess/PostProcessCombineLUTs.cpp b/Engine/Source/Runtime/Renderer/Private/PostProcess/PostProcessCombineLUTs.cpp
--- a/Engine/Source/Runtime/Renderer/Private/PostProcess/PostProcessCombineLUTs.cpp
+++ b/Engine/Source/Runtime/Renderer/Private/PostProcess/PostProcessCombineLUTs.cpp
@@ -10,6 +10,7 @@
 #include "PostProcessCombineLUTs.h"
 #include "PostProcessing.h"
 #include "ScreenRendering.h"
+#include "SceneUtils.h"
 
 
 // CVars
diff --git a/Engine/Source/Runtime/Renderer/Private/PostProcess/PostProcessCompositeEditorPrimitives.cpp b/Engine/Source/Runtime/Renderer/Private/PostProcess/PostProcessCompositeEditorPrimitives.cpp
--- a/Engine/Source/Runtime/Renderer/Private/PostProcess/PostProcessCompositeEditorPrimitives.cpp
+++ b/Engine/Source/Runtime/Renderer/Private/PostProcess/PostProcessCompositeEditorPrimitives.cpp
@@ -5,6 +5,7 @@
 #include "PostProcessCompositeEditorPrimitives.h"
 #include "PostProcessing.h"
 #include "SceneFilterRendering.h"
+#include "SceneUtils.h"
 
 
 // temporary
diff --git a/Engine/Source/Runtime/Renderer/Private/PostProcess/PostProcessDOF.cpp b/Engine/Source/Runtime/Renderer/Private/PostProcess/PostProcessDOF.cpp
--- a/Engine/Source/Runtime/Renderer/Private/PostProcess/PostProcessDOF.cpp
+++ b/Engine/Source/Runtime/Renderer/Private/PostProcess/PostProcessDOF.cpp
@@ -10,6 +10,7 @@
 #include "PostProcessBokehDOF.h"
 #include "PostProcessDOF.h"
 #include "PostProcessing.h"
+#include "SceneUtils.h"
 
 
 /** Encapsulates the DOF setup pixel shader. */
diff --git a/Engine/Source/Runtime/Renderer/Private/PostProcess/PostProcessDownsample.cpp b/Engine/Source/Runtime/Renderer/Private/PostProcess/PostProcessDownsample.cpp
--- a/Engine/Source/Runtime/Renderer/Private/PostProcess/PostProcessDownsample.cpp
+++ b/Engine/Source/Runtime/Renderer/Private/PostProcess/PostProcessDownsample.cpp
@@ -10,6 +10,7 @@
 #include "PostProcessDownsample.h"
 #include "PostProcessing.h"
 #include "PostProcessWeightedSampleSum.h"
+#include "SceneUtils.h"
 
 /** Encapsulates the post processing down sample pixel shader. */
 template <uint32 Method>
diff --git a/Engine/Source/Runtime/Renderer/Private/PostProcess/PostProcessEyeAdaptation.cpp b/Engine/Source/Runtime/Renderer/Private/PostProcess/PostProcessEyeAdaptation.cpp
--- a/Engine/Source/Runtime/Renderer/Private/PostProcess/PostProcessEyeAdaptation.cpp
+++ b/Engine/Source/Runtime/Renderer/Private/PostProcess/PostProcessEyeAdaptation.cpp
@@ -9,6 +9,7 @@
 #include "SceneFilterRendering.h"
 #include "PostProcessEyeAdaptation.h"
 #include "PostProcessing.h"
+#include "SceneUtils.h"
 
 /** Encapsulates the post processing eye adaptation pixel shader. */
 class FPostProcessEyeAdaptationPS : public FGlobalShader
diff --git a/Engine/Source/Runtime/Renderer/Private/PostProcess/PostProcessGBufferHints.cpp b/Engine/Source/Runtime/Renderer/Private/PostProcess/PostProcessGBufferHints.cpp
--- a/Engine/Source/Runtime/Renderer/Private/PostProcess/PostProcessGBufferHints.cpp
+++ b/Engine/Source/Runtime/Renderer/Private/PostProcess/PostProcessGBufferHints.cpp
@@ -11,6 +11,7 @@
 #include "PostProcessing.h"
 #include "PostProcessHistogram.h"
 #include "PostProcessEyeAdaptation.h"
+#include "SceneUtils.h"
 
 /** Encapsulates the post processing eye adaptation pixel shader. */
 class FPostProcessGBufferHintsPS : public FGlobalShader
diff --git a/Engine/Source/Runtime/Renderer/Private/PostProcess/PostProcessHMD.cpp b/Engine/Source/Runtime/Renderer/Private/PostProcess/PostProcessHMD.cpp
--- a/Engine/Source/Runtime/Renderer/Private/PostProcess/PostProcessHMD.cpp
+++ b/Engine/Source/Runtime/Renderer/Private/PostProcess/PostProcessHMD.cpp
@@ -14,7 +14,7 @@
 #include "PostProcessEyeAdaptation.h"
 #include "IHeadMountedDisplay.h"
 #include "RHICommandList.h"
-
+#include "SceneUtils.h"
 
 /** The filter vertex declaration resource type. */
 class FDistortionVertexDeclaration : public FRenderResource
diff --git a/Engine/Source/Runtime/Renderer/Private/PostProcess/PostProcessHistogram.cpp b/Engine/Source/Runtime/Renderer/Private/PostProcess/PostProcessHistogram.cpp
--- a/Engine/Source/Runtime/Renderer/Private/PostProcess/PostProcessHistogram.cpp
+++ b/Engine/Source/Runtime/Renderer/Private/PostProcess/PostProcessHistogram.cpp
@@ -9,6 +9,7 @@
 #include "PostProcessHistogram.h"
 #include "PostProcessing.h"
 #include "PostProcessEyeAdaptation.h"
+#include "SceneUtils.h"
 
 
 /** Encapsulates the post processing histogram compute shader. */
diff --git a/Engine/Source/Runtime/Renderer/Private/PostProcess/PostProcessHistogramReduce.cpp b/Engine/Source/Runtime/Renderer/Private/PostProcess/PostProcessHistogramReduce.cpp
--- a/Engine/Source/Runtime/Renderer/Private/PostProcess/PostProcessHistogramReduce.cpp
+++ b/Engine/Source/Runtime/Renderer/Private/PostProcess/PostProcessHistogramReduce.cpp
@@ -9,6 +9,7 @@
 #include "SceneFilterRendering.h"
 #include "PostProcessHistogramReduce.h"
 #include "PostProcessing.h"
+#include "SceneUtils.h"
 
 /** Encapsulates the post processing histogram reduce compute shader. */
 class FPostProcessHistogramReducePS : public FGlobalShader
diff --git a/Engine/Source/Runtime/Renderer/Private/PostProcess/PostProcessLensBlur.cpp b/Engine/Source/Runtime/Renderer/Private/PostProcess/PostProcessLensBlur.cpp
--- a/Engine/Source/Runtime/Renderer/Private/PostProcess/PostProcessLensBlur.cpp
+++ b/Engine/Source/Runtime/Renderer/Private/PostProcess/PostProcessLensBlur.cpp
@@ -10,6 +10,7 @@
 #include "PostProcessPassThrough.h"
 #include "PostProcessing.h"
 #include "PostProcessLensBlur.h"
+#include "SceneUtils.h"
 
 /** Encapsulates the post processing vertex shader. */
 class FPostProcessLensBlurVS : public FGlobalShader
diff --git a/Engine/Source/Runtime/Renderer/Private/PostProcess/PostProcessLensFlares.cpp b/Engine/Source/Runtime/Renderer/Private/PostProcess/PostProcessLensFlares.cpp
--- a/Engine/Source/Runtime/Renderer/Private/PostProcess/PostProcessLensFlares.cpp
+++ b/Engine/Source/Runtime/Renderer/Private/PostProcess/PostProcessLensFlares.cpp
@@ -9,6 +9,7 @@
 #include "SceneFilterRendering.h"
 #include "PostProcessLensFlares.h"
 #include "PostProcessing.h"
+#include "SceneUtils.h"
 
 /** Encapsulates a simple copy pixel shader. */
 class FPostProcessLensFlareBasePS : public FGlobalShader
diff --git a/Engine/Source/Runtime/Renderer/Private/PostProcess/PostProcessMaterial.cpp b/Engine/Source/Runtime/Renderer/Private/PostProcess/PostProcessMaterial.cpp
--- a/Engine/Source/Runtime/Renderer/Private/PostProcess/PostProcessMaterial.cpp
+++ b/Engine/Source/Runtime/Renderer/Private/PostProcess/PostProcessMaterial.cpp
@@ -11,6 +11,7 @@
 #include "PostProcessing.h"
 #include "PostProcessEyeAdaptation.h"
 #include "../../../Engine/Public/TileRendering.h"
+#include "SceneUtils.h"
 
 class FPostProcessMaterialVS : public FMaterialShader
 {
diff --git a/Engine/Source/Runtime/Renderer/Private/PostProcess/PostProcessMobile.cpp b/Engine/Source/Runtime/Renderer/Private/PostProcess/PostProcessMobile.cpp
--- a/Engine/Source/Runtime/Renderer/Private/PostProcess/PostProcessMobile.cpp
+++ b/Engine/Source/Runtime/Renderer/Private/PostProcess/PostProcessMobile.cpp
@@ -12,6 +12,7 @@
 #include "PostProcessing.h"
 #include "PostProcessHistogram.h"
 #include "PostProcessEyeAdaptation.h"
+#include "SceneUtils.h"
 
 //
 // BLOOM SETUP
diff --git a/Engine/Source/Runtime/Renderer/Private/PostProcess/PostProcessMorpheus.cpp b/Engine/Source/Runtime/Renderer/Private/PostProcess/PostProcessMorpheus.cpp
--- a/Engine/Source/Runtime/Renderer/Private/PostProcess/PostProcessMorpheus.cpp
+++ b/Engine/Source/Runtime/Renderer/Private/PostProcess/PostProcessMorpheus.cpp
@@ -12,6 +12,7 @@
 #include "PostProcessHistogram.h"
 #include "PostProcessEyeAdaptation.h"
 #include "IHeadMountedDisplay.h"
+#include "SceneUtils.h"
 
 DEFINE_LOG_CATEGORY_STATIC(LogMorpheusHMDPostProcess, All, All);
 
diff --git a/Engine/Source/Runtime/Renderer/Private/PostProcess/PostProcessMotionBlur.cpp b/Engine/Source/Runtime/Renderer/Private/PostProcess/PostProcessMotionBlur.cpp
--- a/Engine/Source/Runtime/Renderer/Private/PostProcess/PostProcessMotionBlur.cpp
+++ b/Engine/Source/Runtime/Renderer/Private/PostProcess/PostProcessMotionBlur.cpp
@@ -9,6 +9,7 @@
 #include "SceneFilterRendering.h"
 #include "PostProcessMotionBlur.h"
 #include "PostProcessing.h"
+#include "SceneUtils.h"
 
 #if !(UE_BUILD_SHIPPING || UE_BUILD_TEST)
 static TAutoConsoleVariable<int32> CVarMotionBlurFiltering(
diff --git a/Engine/Source/Runtime/Renderer/Private/PostProcess/PostProcessNoiseBlur.cpp b/Engine/Source/Runtime/Renderer/Private/PostProcess/PostProcessNoiseBlur.cpp
--- a/Engine/Source/Runtime/Renderer/Private/PostProcess/PostProcessNoiseBlur.cpp
+++ b/Engine/Source/Runtime/Renderer/Private/PostProcess/PostProcessNoiseBlur.cpp
@@ -9,6 +9,7 @@
 #include "SceneFilterRendering.h"
 #include "PostProcessNoiseBlur.h"
 #include "PostProcessing.h"
+#include "SceneUtils.h"
 
 /** Encapsulates the post processing down sample pixel shader. */
 template <uint32 Method>
diff --git a/Engine/Source/Runtime/Renderer/Private/PostProcess/PostProcessSceneColorFringe.cpp b/Engine/Source/Runtime/Renderer/Private/PostProcess/PostProcessSceneColorFringe.cpp
--- a/Engine/Source/Runtime/Renderer/Private/PostProcess/PostProcessSceneColorFringe.cpp
+++ b/Engine/Source/Runtime/Renderer/Private/PostProcess/PostProcessSceneColorFringe.cpp
@@ -10,6 +10,7 @@
 #include "PostProcessAA.h"
 #include "PostProcessing.h"
 #include "PostProcessSceneColorFringe.h"
+#include "SceneUtils.h"
 
 
 /*-----------------------------------------------------------------------------
diff --git a/Engine/Source/Runtime/Renderer/Private/PostProcess/PostProcessSelectionOutline.cpp b/Engine/Source/Runtime/Renderer/Private/PostProcess/PostProcessSelectionOutline.cpp
--- a/Engine/Source/Runtime/Renderer/Private/PostProcess/PostProcessSelectionOutline.cpp
+++ b/Engine/Source/Runtime/Renderer/Private/PostProcess/PostProcessSelectionOutline.cpp
@@ -9,6 +9,7 @@
 #include "SceneFilterRendering.h"
 #include "PostProcessSelectionOutline.h"
 #include "OneColorShader.h"
+#include "SceneUtils.h"
 
 ///////////////////////////////////////////
 // FRCPassPostProcessSelectionOutlineColor
diff --git a/Engine/Source/Runtime/Renderer/Private/PostProcess/PostProcessSubsurface.cpp b/Engine/Source/Runtime/Renderer/Private/PostProcess/PostProcessSubsurface.cpp
--- a/Engine/Source/Runtime/Renderer/Private/PostProcess/PostProcessSubsurface.cpp
+++ b/Engine/Source/Runtime/Renderer/Private/PostProcess/PostProcessSubsurface.cpp
@@ -9,6 +9,7 @@
 #include "SceneFilterRendering.h"
 #include "PostProcessSubsurface.h"
 #include "PostProcessing.h"
+#include "SceneUtils.h"
 
 ENGINE_API const IPooledRenderTarget* GetSubsufaceProfileTexture_RT(FRHICommandListImmediate& RHICmdList);
 
diff --git a/Engine/Source/Runtime/Renderer/Private/PostProcess/PostProcessTemporalAA.cpp b/Engine/Source/Runtime/Renderer/Private/PostProcess/PostProcessTemporalAA.cpp
--- a/Engine/Source/Runtime/Renderer/Private/PostProcess/PostProcessTemporalAA.cpp
+++ b/Engine/Source/Runtime/Renderer/Private/PostProcess/PostProcessTemporalAA.cpp
@@ -11,6 +11,7 @@
 #include "PostProcessAmbientOcclusion.h"
 #include "PostProcessTonemap.h"
 #include "PostProcessing.h"
+#include "SceneUtils.h"
 
 static float TemporalHalton( int32 Index, int32 Base )
 {
diff --git a/Engine/Source/Runtime/Renderer/Private/PostProcess/PostProcessTestImage.cpp b/Engine/Source/Runtime/Renderer/Private/PostProcess/PostProcessTestImage.cpp
--- a/Engine/Source/Runtime/Renderer/Private/PostProcess/PostProcessTestImage.cpp
+++ b/Engine/Source/Runtime/Renderer/Private/PostProcess/PostProcessTestImage.cpp
@@ -12,6 +12,7 @@
 #include "PostProcessHistogram.h"
 #include "PostProcessEyeAdaptation.h"
 #include "PostProcessCombineLUTs.h"
+#include "SceneUtils.h"
 
 /** Encapsulates the post processing eye adaptation pixel shader. */
 class FPostProcessTestImagePS : public FGlobalShader
diff --git a/Engine/Source/Runtime/Renderer/Private/PostProcess/PostProcessTonemap.cpp b/Engine/Source/Runtime/Renderer/Private/PostProcess/PostProcessTonemap.cpp
--- a/Engine/Source/Runtime/Renderer/Private/PostProcess/PostProcessTonemap.cpp
+++ b/Engine/Source/Runtime/Renderer/Private/PostProcess/PostProcessTonemap.cpp
@@ -12,6 +12,7 @@
 #include "PostProcessing.h"
 #include "PostProcessEyeAdaptation.h"
 #include "PostProcessCombineLUTs.h"
+#include "SceneUtils.h"
 
 
 //
diff --git a/Engine/Source/Runtime/Renderer/Private/PostProcess/PostProcessUIBlur.cpp b/Engine/Source/Runtime/Renderer/Private/PostProcess/PostProcessUIBlur.cpp
--- a/Engine/Source/Runtime/Renderer/Private/PostProcess/PostProcessUIBlur.cpp
+++ b/Engine/Source/Runtime/Renderer/Private/PostProcess/PostProcessUIBlur.cpp
@@ -9,6 +9,7 @@
 #include "SceneFilterRendering.h"
 #include "PostProcessing.h"
 #include "PostProcessUIBlur.h"
+#include "SceneUtils.h"
 
 
 /** Encapsulates the post processing eye adaptation pixel shader. */
diff --git a/Engine/Source/Runtime/Renderer/Private/PostProcess/PostProcessUpscale.cpp b/Engine/Source/Runtime/Renderer/Private/PostProcess/PostProcessUpscale.cpp
--- a/Engine/Source/Runtime/Renderer/Private/PostProcess/PostProcessUpscale.cpp
+++ b/Engine/Source/Runtime/Renderer/Private/PostProcess/PostProcessUpscale.cpp
@@ -11,6 +11,7 @@
 #include "PostProcessing.h"
 #include "PostProcessHistogram.h"
 #include "PostProcessEyeAdaptation.h"
+#include "SceneUtils.h"
 
 static TAutoConsoleVariable<float> CVarScreenPercentageSoftness(
 	TEXT("r.ScreenPercentageSoftness"),
diff --git a/Engine/Source/Runtime/Renderer/Private/PostProcess/PostProcessVisualizeBuffer.cpp b/Engine/Source/Runtime/Renderer/Private/PostProcess/PostProcessVisualizeBuffer.cpp
--- a/Engine/Source/Runtime/Renderer/Private/PostProcess/PostProcessVisualizeBuffer.cpp
+++ b/Engine/Source/Runtime/Renderer/Private/PostProcess/PostProcessVisualizeBuffer.cpp
@@ -11,6 +11,7 @@
 #include "PostProcessing.h"
 #include "PostProcessHistogram.h"
 #include "PostProcessEyeAdaptation.h"
+#include "SceneUtils.h"
 
 /** Encapsulates the post processing Buffer visualization pixel shader. */
 template<bool bDrawingTile>
diff --git a/Engine/Source/Runtime/Renderer/Private/PostProcess/PostProcessVisualizeHDR.cpp b/Engine/Source/Runtime/Renderer/Private/PostProcess/PostProcessVisualizeHDR.cpp
--- a/Engine/Source/Runtime/Renderer/Private/PostProcess/PostProcessVisualizeHDR.cpp
+++ b/Engine/Source/Runtime/Renderer/Private/PostProcess/PostProcessVisualizeHDR.cpp
@@ -12,6 +12,7 @@
 #include "PostProcessHistogram.h"
 #include "PostProcessEyeAdaptation.h"
 #include "PostProcessTonemap.h"
+#include "SceneUtils.h"
 
 /** Encapsulates the post processing eye adaptation pixel shader. */
 class FPostProcessVisualizeHDRPS : public FGlobalShader
diff --git a/Engine/Source/Runtime/Renderer/Private/PostProcess/PostProcessWeightedSampleSum.cpp b/Engine/Source/Runtime/Renderer/Private/PostProcess/PostProcessWeightedSampleSum.cpp
--- a/Engine/Source/Runtime/Renderer/Private/PostProcess/PostProcessWeightedSampleSum.cpp
+++ b/Engine/Source/Runtime/Renderer/Private/PostProcess/PostProcessWeightedSampleSum.cpp
@@ -8,6 +8,7 @@
 #include "ScenePrivate.h"
 #include "SceneFilterRendering.h"
 #include "PostProcessWeightedSampleSum.h"
+#include "SceneUtils.h"
 
 /**
  * A pixel shader which filters a texture.
diff --git a/Engine/Source/Runtime/Renderer/Private/PostProcess/PostProcessing.cpp b/Engine/Source/Runtime/Renderer/Private/PostProcess/PostProcessing.cpp
--- a/Engine/Source/Runtime/Renderer/Private/PostProcess/PostProcessing.cpp
+++ b/Engine/Source/Runtime/Renderer/Private/PostProcess/PostProcessing.cpp
@@ -45,6 +45,7 @@
 
 #include "PostProcessMorpheus.h"
 #include "IHeadMountedDisplay.h"
+#include "BufferVisualizationData.h"
 
 
 /** The global center for all post processing activities. */
diff --git a/Engine/Source/Runtime/Renderer/Private/PostProcess/RenderTargetPool.h b/Engine/Source/Runtime/Renderer/Private/PostProcess/RenderTargetPool.h
--- a/Engine/Source/Runtime/Renderer/Private/PostProcess/RenderTargetPool.h
+++ b/Engine/Source/Runtime/Renderer/Private/PostProcess/RenderTargetPool.h
@@ -5,6 +5,7 @@
 =============================================================================*/
 
 #pragma once
+#include "RendererInterface.h"
 
 /** The reference to a pooled render target, use like this: TRefCountPtr<IPooledRenderTarget> */
 struct FPooledRenderTarget : public IPooledRenderTarget
diff --git a/Engine/Source/Runtime/Renderer/Private/PostProcess/SceneRenderTargets.cpp b/Engine/Source/Runtime/Renderer/Private/PostProcess/SceneRenderTargets.cpp
--- a/Engine/Source/Runtime/Renderer/Private/PostProcess/SceneRenderTargets.cpp
+++ b/Engine/Source/Runtime/Renderer/Private/PostProcess/SceneRenderTargets.cpp
@@ -9,6 +9,7 @@
 #include "SceneFilterRendering.h"
 #include "ReflectionEnvironment.h"
 #include "LightPropagationVolume.h"
+#include "SceneUtils.h"
 
 // for LightPropagationVolume feature, could be exposed
 const int ReflectiveShadowMapResolution = 256;
diff --git a/Engine/Source/Runtime/Renderer/Private/PostProcess/SceneRenderTargets.h b/Engine/Source/Runtime/Renderer/Private/PostProcess/SceneRenderTargets.h
--- a/Engine/Source/Runtime/Renderer/Private/PostProcess/SceneRenderTargets.h
+++ b/Engine/Source/Runtime/Renderer/Private/PostProcess/SceneRenderTargets.h
@@ -11,6 +11,8 @@
 #include "../SystemTextures.h"
 #include "RHIStaticStates.h"
 
+struct IPooledRenderTarget;
+
 /** Number of cube map shadow depth surfaces that will be created and used for rendering one pass point light shadows. */
 static const int32 NumCubeShadowDepthSurfaces = 5;
 
diff --git a/Engine/Source/Runtime/Renderer/Private/PostProcess/ScreenSpaceReflections.cpp b/Engine/Source/Runtime/Renderer/Private/PostProcess/ScreenSpaceReflections.cpp
--- a/Engine/Source/Runtime/Renderer/Private/PostProcess/ScreenSpaceReflections.cpp
+++ b/Engine/Source/Runtime/Renderer/Private/PostProcess/ScreenSpaceReflections.cpp
@@ -11,6 +11,7 @@
 #include "ScreenSpaceReflections.h"
 #include "PostProcessTemporalAA.h"
 #include "PostProcessAmbientOcclusion.h"
+#include "SceneUtils.h"
 
 static TAutoConsoleVariable<int32> CVarSSRQuality(
 	TEXT("r.SSR.Quality"),
diff --git a/Engine/Source/Runtime/Renderer/Private/PostProcess/VisualizeTexture.cpp b/Engine/Source/Runtime/Renderer/Private/PostProcess/VisualizeTexture.cpp
--- a/Engine/Source/Runtime/Renderer/Private/PostProcess/VisualizeTexture.cpp
+++ b/Engine/Source/Runtime/Renderer/Private/PostProcess/VisualizeTexture.cpp
@@ -11,6 +11,7 @@
 #include "VisualizeTexture.h"
 #include "PostProcessing.h"
 #include "PostProcessWeightedSampleSum.h"
+#include "SceneUtils.h"
 
 
 /** A pixel shader which filters a texture. */
diff --git a/Engine/Source/Runtime/Renderer/Private/PostProcess/VisualizeTexture.h b/Engine/Source/Runtime/Renderer/Private/PostProcess/VisualizeTexture.h
--- a/Engine/Source/Runtime/Renderer/Private/PostProcess/VisualizeTexture.h
+++ b/Engine/Source/Runtime/Renderer/Private/PostProcess/VisualizeTexture.h
@@ -6,6 +6,8 @@
 
 #pragma once
 
+struct IPooledRenderTarget;
+
 class FVisualizeTexture
 {
 public:
diff --git a/Engine/Source/Runtime/Renderer/Private/ReflectionEnvironment.cpp b/Engine/Source/Runtime/Renderer/Private/ReflectionEnvironment.cpp
--- a/Engine/Source/Runtime/Renderer/Private/ReflectionEnvironment.cpp
+++ b/Engine/Source/Runtime/Renderer/Private/ReflectionEnvironment.cpp
@@ -17,6 +17,7 @@
 #include "ReflectionEnvironment.h"
 #include "ShaderParameterUtils.h"
 #include "LightRendering.h"
+#include "SceneUtils.h"
 
 /** Tile size for the reflection environment compute shader, tweaked for 680 GTX. */
 const int32 GReflectionEnvironmentTileSizeX = 16;
diff --git a/Engine/Source/Runtime/Renderer/Private/ReflectionEnvironmentCapture.cpp b/Engine/Source/Runtime/Renderer/Private/ReflectionEnvironmentCapture.cpp
--- a/Engine/Source/Runtime/Renderer/Private/ReflectionEnvironmentCapture.cpp
+++ b/Engine/Source/Runtime/Renderer/Private/ReflectionEnvironmentCapture.cpp
@@ -13,6 +13,7 @@
 #include "ScreenRendering.h"
 #include "ReflectionEnvironment.h"
 #include "ReflectionEnvironmentCapture.h"
+#include "SceneUtils.h"
 
 /** Near plane to use when capturing the scene. */
 float GReflectionCaptureNearPlane = 5;
diff --git a/Engine/Source/Runtime/Renderer/Private/SceneOcclusion.cpp b/Engine/Source/Runtime/Renderer/Private/SceneOcclusion.cpp
--- a/Engine/Source/Runtime/Renderer/Private/SceneOcclusion.cpp
+++ b/Engine/Source/Runtime/Renderer/Private/SceneOcclusion.cpp
@@ -10,6 +10,7 @@
 #include "SceneOcclusion.h"
 #include "ScreenRendering.h"
 #include "SceneFilterRendering.h"
+#include "SceneUtils.h"
 
 /*-----------------------------------------------------------------------------
 	Globals
diff --git a/Engine/Source/Runtime/Renderer/Private/ScenePostProcessing.cpp b/Engine/Source/Runtime/Renderer/Private/ScenePostProcessing.cpp
--- a/Engine/Source/Runtime/Renderer/Private/ScenePostProcessing.cpp
+++ b/Engine/Source/Runtime/Renderer/Private/ScenePostProcessing.cpp
@@ -8,6 +8,7 @@
 #include "ScenePrivate.h"
 #include "SceneFilterRendering.h"
 #include "PostProcessing.h"
+#include "SceneUtils.h"
 
 /** Encapsulates the gamma correction pixel shader. */
 class FGammaCorrectionPS : public FGlobalShader
diff --git a/Engine/Source/Runtime/Renderer/Private/ScenePrivate.h b/Engine/Source/Runtime/Renderer/Private/ScenePrivate.h
--- a/Engine/Source/Runtime/Renderer/Private/ScenePrivate.h
+++ b/Engine/Source/Runtime/Renderer/Private/ScenePrivate.h
@@ -7,6 +7,8 @@
 #ifndef __SCENEPRIVATE_H__
 #define __SCENEPRIVATE_H__
 
+class SceneRenderingAllocator;
+
 class SceneRenderingBitArrayAllocator
 	: public TInlineAllocator<4,SceneRenderingAllocator>
 {
diff --git a/Engine/Source/Runtime/Renderer/Private/SceneRendering.cpp b/Engine/Source/Runtime/Renderer/Private/SceneRendering.cpp
--- a/Engine/Source/Runtime/Renderer/Private/SceneRendering.cpp
+++ b/Engine/Source/Runtime/Renderer/Private/SceneRendering.cpp
@@ -15,6 +15,7 @@
 #include "FXSystem.h"
 #include "SceneViewExtension.h"
 #include "PostProcessBusyWait.h"
+#include "SceneUtils.h"
 
 /*-----------------------------------------------------------------------------
 	Globals
diff --git a/Engine/Source/Runtime/Renderer/Private/SceneVisibility.cpp b/Engine/Source/Runtime/Renderer/Private/SceneVisibility.cpp
--- a/Engine/Source/Runtime/Renderer/Private/SceneVisibility.cpp
+++ b/Engine/Source/Runtime/Renderer/Private/SceneVisibility.cpp
@@ -9,6 +9,7 @@
 #include "ScenePrivate.h"
 #include "FXSystem.h"
 #include "../../Engine/Private/SkeletalRenderGPUSkin.h"		// GPrevPerBoneMotionBlur
+#include "SceneUtils.h"
 
 /*------------------------------------------------------------------------------
 	Globals
diff --git a/Engine/Source/Runtime/Renderer/Private/ShaderComplexityRendering.cpp b/Engine/Source/Runtime/Renderer/Private/ShaderComplexityRendering.cpp
--- a/Engine/Source/Runtime/Renderer/Private/ShaderComplexityRendering.cpp
+++ b/Engine/Source/Runtime/Renderer/Private/ShaderComplexityRendering.cpp
@@ -8,6 +8,7 @@ ShaderComplexityRendering.cpp: Contains definitions for rendering the shader com
 #include "ScenePrivate.h"
 #include "SceneFilterRendering.h"
 #include "PostProcessVisualizeComplexity.h"
+#include "SceneUtils.h"
 
 /**
  * Gets the maximum shader complexity count from the ini settings.
diff --git a/Engine/Source/Runtime/Renderer/Private/ShadowRendering.cpp b/Engine/Source/Runtime/Renderer/Private/ShadowRendering.cpp
--- a/Engine/Source/Runtime/Renderer/Private/ShadowRendering.cpp
+++ b/Engine/Source/Runtime/Renderer/Private/ShadowRendering.cpp
@@ -8,6 +8,7 @@
 #include "ScenePrivate.h"
 #include "TextureLayout.h"
 #include "LightPropagationVolume.h"
+#include "SceneUtils.h"
 
 static TAutoConsoleVariable<float> CVarCSMShadowDepthBias(
 	TEXT("r.Shadow.CSMDepthBias"),
diff --git a/Engine/Source/Runtime/Renderer/Private/TiledDeferredLightRendering.cpp b/Engine/Source/Runtime/Renderer/Private/TiledDeferredLightRendering.cpp
--- a/Engine/Source/Runtime/Renderer/Private/TiledDeferredLightRendering.cpp
+++ b/Engine/Source/Runtime/Renderer/Private/TiledDeferredLightRendering.cpp
@@ -11,6 +11,7 @@
 #include "ShaderParameters.h"
 #include "LightRendering.h"
 #include "ScreenRendering.h"
+#include "SceneUtils.h"
 
 /** 
  * Maximum number of lights that can be handled by tiled deferred in a single compute shader pass.
diff --git a/Engine/Source/Runtime/Renderer/Private/TranslucentLighting.cpp b/Engine/Source/Runtime/Renderer/Private/TranslucentLighting.cpp
--- a/Engine/Source/Runtime/Renderer/Private/TranslucentLighting.cpp
+++ b/Engine/Source/Runtime/Renderer/Private/TranslucentLighting.cpp
@@ -11,6 +11,7 @@
 #include "SceneFilterRendering.h"
 #include "ScreenRendering.h"
 #include "AmbientCubemapParameters.h"
+#include "SceneUtils.h"
 
 class FMaterial;
 
diff --git a/Engine/Source/Runtime/Renderer/Private/TranslucentRendering.cpp b/Engine/Source/Runtime/Renderer/Private/TranslucentRendering.cpp
--- a/Engine/Source/Runtime/Renderer/Private/TranslucentRendering.cpp
+++ b/Engine/Source/Runtime/Renderer/Private/TranslucentRendering.cpp
@@ -9,6 +9,7 @@
 #include "ScreenRendering.h"
 #include "SceneFilterRendering.h"
 #include "LightPropagationVolume.h"
+#include "SceneUtils.h"
 
 static void SetTranslucentRenderTargetAndState(FRHICommandList& RHICmdList, const FViewInfo& View, bool bSeparateTranslucencyPass)
 {
diff --git a/Engine/Source/Runtime/Renderer/Private/VelocityRendering.cpp b/Engine/Source/Runtime/Renderer/Private/VelocityRendering.cpp
--- a/Engine/Source/Runtime/Renderer/Private/VelocityRendering.cpp
+++ b/Engine/Source/Runtime/Renderer/Private/VelocityRendering.cpp
@@ -6,6 +6,7 @@
 
 #include "RendererPrivate.h"
 #include "../../Engine/Private/SkeletalRenderGPUSkin.h"		// GPrevPerBoneMotionBlur
+#include "SceneUtils.h"
 
 //=============================================================================
 /** Encapsulates the Velocity vertex shader. */
diff --git a/Engine/Source/Runtime/SlateCore/Public/Brushes/SlateNoResource.h b/Engine/Source/Runtime/SlateCore/Public/Brushes/SlateNoResource.h
--- a/Engine/Source/Runtime/SlateCore/Public/Brushes/SlateNoResource.h
+++ b/Engine/Source/Runtime/SlateCore/Public/Brushes/SlateNoResource.h
@@ -1,7 +1,7 @@
 // Copyright 1998-2014 Epic Games, Inc. All Rights Reserved.
 
 #pragma once
-
+#include "Styling/SlateBrush.h"
 
 /**
  * A resource that has no appearance
diff --git a/Engine/Source/Runtime/SlateCore/Public/Layout/Margin.h b/Engine/Source/Runtime/SlateCore/Public/Layout/Margin.h
--- a/Engine/Source/Runtime/SlateCore/Public/Layout/Margin.h
+++ b/Engine/Source/Runtime/SlateCore/Public/Layout/Margin.h
@@ -1,7 +1,7 @@
 // Copyright 1998-2014 Epic Games, Inc. All Rights Reserved.
 
 #pragma once
-
+#include "Types/SlateEnums.h"
 #include "Margin.generated.h"
 
 
diff --git a/Engine/Source/Runtime/SlateCore/Public/Styling/CoreStyle.h b/Engine/Source/Runtime/SlateCore/Public/Styling/CoreStyle.h
--- a/Engine/Source/Runtime/SlateCore/Public/Styling/CoreStyle.h
+++ b/Engine/Source/Runtime/SlateCore/Public/Styling/CoreStyle.h
@@ -2,6 +2,7 @@
 
 #pragma once
 
+struct FSlateDynamicImageBrush;
 
 /**
  * Core slate style
diff --git a/Engine/Source/Runtime/SlateCore/Public/Styling/ISlateStyle.h b/Engine/Source/Runtime/SlateCore/Public/Styling/ISlateStyle.h
--- a/Engine/Source/Runtime/SlateCore/Public/Styling/ISlateStyle.h
+++ b/Engine/Source/Runtime/SlateCore/Public/Styling/ISlateStyle.h
@@ -1,5 +1,12 @@
 // Copyright 1998-2014 Epic Games, Inc. All Rights Reserved.
 #pragma once
+#include "Styling/StyleDefaults.h"
+
+struct FMargin;
+struct FSlateColor;
+struct FSlateFontInfo;
+struct FSlateSound;
+struct FSlateWidgetStyle;
 
 /**
  * 
diff --git a/Engine/Source/Runtime/SlateCore/Public/Styling/SlateColor.h b/Engine/Source/Runtime/SlateCore/Public/Styling/SlateColor.h
--- a/Engine/Source/Runtime/SlateCore/Public/Styling/SlateColor.h
+++ b/Engine/Source/Runtime/SlateCore/Public/Styling/SlateColor.h
@@ -1,7 +1,7 @@
 // Copyright 1998-2014 Epic Games, Inc. All Rights Reserved.
 
 #pragma once
-
+#include "Styling/WidgetStyle.h"
 #include "SlateColor.generated.h"
 
 
diff --git a/Engine/Source/Runtime/SlateCore/Public/Styling/StyleDefaults.h b/Engine/Source/Runtime/SlateCore/Public/Styling/StyleDefaults.h
--- a/Engine/Source/Runtime/SlateCore/Public/Styling/StyleDefaults.h
+++ b/Engine/Source/Runtime/SlateCore/Public/Styling/StyleDefaults.h
@@ -1,5 +1,10 @@
 // Copyright 1998-2014 Epic Games, Inc. All Rights Reserved.
 #pragma once
+#include "Brushes/SlateNoResource.h"
+#include "Fonts/SlateFontInfo.h"
+
+struct FMargin;
+struct FSlateSound;
 
 class SLATECORE_API FStyleDefaults
 {
diff --git a/Engine/Source/Runtime/SlateRHIRenderer/Private/SlateRHIRenderer.cpp b/Engine/Source/Runtime/SlateRHIRenderer/Private/SlateRHIRenderer.cpp
--- a/Engine/Source/Runtime/SlateRHIRenderer/Private/SlateRHIRenderer.cpp
+++ b/Engine/Source/Runtime/SlateRHIRenderer/Private/SlateRHIRenderer.cpp
@@ -4,6 +4,7 @@
 #include "StereoRendering.h"
 #include "Runtime/Engine/Public/Features/ILiveStreamingService.h"
 #include "SlateNativeTextureResource.h"
+#include "SceneUtils.h"
 
 DECLARE_CYCLE_STAT(TEXT("Map Staging Buffer"),STAT_MapStagingBuffer,STATGROUP_CrashTracker);
 DECLARE_CYCLE_STAT(TEXT("Generate Capture Buffer"),STAT_GenerateCaptureBuffer,STATGROUP_CrashTracker);
diff --git a/Engine/Source/Runtime/Windows/D3D11RHI/Private/D3D11Commands.cpp b/Engine/Source/Runtime/Windows/D3D11RHI/Private/D3D11Commands.cpp
--- a/Engine/Source/Runtime/Windows/D3D11RHI/Private/D3D11Commands.cpp
+++ b/Engine/Source/Runtime/Windows/D3D11RHI/Private/D3D11Commands.cpp
@@ -17,6 +17,7 @@
 #include "RHICommandList.h"
 #include "RHIStaticStates.h"
 #include "ShaderParameterUtils.h"
+#include "SceneUtils.h"
 
 #if PLATFORM_DESKTOP
 // For Depth Bounds Test interface
diff --git a/Engine/Source/Runtime/Windows/D3D11RHI/Private/D3D11Util.cpp b/Engine/Source/Runtime/Windows/D3D11RHI/Private/D3D11Util.cpp
--- a/Engine/Source/Runtime/Windows/D3D11RHI/Private/D3D11Util.cpp
+++ b/Engine/Source/Runtime/Windows/D3D11RHI/Private/D3D11Util.cpp
@@ -6,6 +6,7 @@
 
 #include "D3D11RHIPrivate.h"
 #include "EngineModule.h"
+#include "RendererInterface.h"
 
 #define D3DERR(x) case x: ErrorCodeText = TEXT(#x); break;
 #define LOCTEXT_NAMESPACE "Developer.MessageLog"
