DataTable can use UserDefinedStructure.
When UDS is modified, DT stroes it's data serialized (with tags).

diff --git a/Engine/Config/BaseEditor.ini b/Engine/Config/BaseEditor.ini
--- a/Engine/Config/BaseEditor.ini
+++ b/Engine/Config/BaseEditor.ini
@@ -261,6 +261,7 @@ SupportWebsiteURL="http://unrealengine.com/support"
 [UserDefinedStructure]
 bUseUserDefinedStructure=true
 bShowDefaultValuePropertyEditor=true
+bCanBeBasedOnTableRow=false
 
 [EpicSurvey]
 Source="CDN"
diff --git a/Engine/Source/Editor/DataTableEditor/Private/DataTableEditor.cpp b/Engine/Source/Editor/DataTableEditor/Private/DataTableEditor.cpp
--- a/Engine/Source/Editor/DataTableEditor/Private/DataTableEditor.cpp
+++ b/Engine/Source/Editor/DataTableEditor/Private/DataTableEditor.cpp
@@ -211,6 +211,8 @@ TSharedRef<SVerticalBox> FDataTableEditor::CreateContentBox()
 
 void FDataTableEditor::OnDataTableReloaded() 
 {
+	CachedDataTable.Reset();
+
 	TSharedRef<SVerticalBox> ContentBox = CreateContentBox();
 	
 	GridPanelOwner->SetContent(ContentBox);
diff --git a/Engine/Source/Editor/KismetCompiler/Private/UserDefinedStructureCompilerUtils.cpp b/Engine/Source/Editor/KismetCompiler/Private/UserDefinedStructureCompilerUtils.cpp
--- a/Engine/Source/Editor/KismetCompiler/Private/UserDefinedStructureCompilerUtils.cpp
+++ b/Engine/Source/Editor/KismetCompiler/Private/UserDefinedStructureCompilerUtils.cpp
@@ -282,7 +282,7 @@ void FUserDefinedStructureCompilerUtils::CompileStruct(class UUserDefinedStruct*
 			UUserDefinedStruct* ChangedStruct = ChangedStructs[StructIdx];
 			if (ChangedStruct)
 			{
-				FStructureEditorUtils::FStructEditorManager::Get().PreChange(ChangedStruct);
+				FStructureEditorUtils::BroadcastPreChange(ChangedStruct);
 				FUserDefinedStructureCompilerInner::ReplaceStructWithTempDuplicate(ChangedStruct, BlueprintsToRecompile, ChangedStructs);
 				ChangedStruct->Status = EUserDefinedStructureStatus::UDSS_Dirty;
 			}
@@ -350,12 +350,11 @@ void FUserDefinedStructureCompilerUtils::CompileStruct(class UUserDefinedStruct*
 			FBlueprintEditorUtils::MarkBlueprintAsStructurallyModified(*BPIter);
 		}
 
-		// UPDATE UI
 		for (auto ChangedStruct : ChangedStructs)
 		{
 			if (ChangedStruct)
 			{
-				FStructureEditorUtils::FStructEditorManager::Get().PostChange(ChangedStruct);
+				FStructureEditorUtils::BroadcastPostChange(ChangedStruct);
 				ChangedStruct->MarkPackageDirty();
 			}
 		}
diff --git a/Engine/Source/Editor/UnrealEd/Private/Kismet2/StructureEditorUtils.cpp b/Engine/Source/Editor/UnrealEd/Private/Kismet2/StructureEditorUtils.cpp
--- a/Engine/Source/Editor/UnrealEd/Private/Kismet2/StructureEditorUtils.cpp
+++ b/Engine/Source/Editor/UnrealEd/Private/Kismet2/StructureEditorUtils.cpp
@@ -10,6 +10,8 @@
 #include "Editor/UnrealEd/Public/Kismet2/CompilerResultsLog.h"
 #include "Editor/UnrealEd/Public/EditorModes.h"
 #include "Editor/KismetCompiler/Public/KismetCompilerModule.h"
+#include "Toolkits/AssetEditorManager.h"
+#include "Editor/DataTableEditor/Public/IDataTableEditor.h"
 
 #define LOCTEXT_NAMESPACE "Structure"
 
@@ -710,4 +712,52 @@ bool FStructureEditorUtils::SetNativeBase(UUserDefinedStruct* Struct, UScriptStr
 	return bResult;
 }
 
+struct FReinstanceDataTableHelper
+{
+	// TODO: shell we cache the dependency?
+	static TArray<UDataTable*> GetTablesDependentOnStruct(UUserDefinedStruct* Struct)
+	{
+		TArray<UDataTable*> Result;
+		if (Struct && (Struct->GetSuperStruct() == FTableRowBase::StaticStruct()))
+		{
+			TArray<UObject*> DataTables;
+			GetObjectsOfClass(UDataTable::StaticClass(), DataTables);
+			for (auto DataTableObj : DataTables)
+			{
+				auto DataTable = Cast<UDataTable>(DataTableObj);
+				if (DataTable && (Struct == DataTable->RowStruct))
+				{
+					Result.Add(DataTable);
+				}
+			}
+		}
+		return Result;
+	}
+};
+
+void FStructureEditorUtils::BroadcastPreChange(UUserDefinedStruct* Struct)
+{
+	FStructureEditorUtils::FStructEditorManager::Get().PreChange(Struct);
+	auto DataTables = FReinstanceDataTableHelper::GetTablesDependentOnStruct(Struct);
+	for (auto DataTable : DataTables)
+	{
+		DataTable->CleanBeforeStructChange();
+	}
+}
+
+void FStructureEditorUtils::BroadcastPostChange(UUserDefinedStruct* Struct)
+{
+	FStructureEditorUtils::FStructEditorManager::Get().PostChange(Struct);
+	auto DataTables = FReinstanceDataTableHelper::GetTablesDependentOnStruct(Struct);
+	for (auto DataTable : DataTables)
+	{
+		DataTable->RestoreAfterStructChange();
+		auto DataTableEditor = static_cast<IDataTableEditor*>(FAssetEditorManager::Get().FindEditorForAsset(DataTable, false));
+		if (DataTableEditor)
+		{
+			DataTableEditor->OnDataTableReloaded();
+		}
+	}
+}
+
 #undef LOCTEXT_NAMESPACE
diff --git a/Engine/Source/Editor/UnrealEd/Public/Kismet2/StructureEditorUtils.h b/Engine/Source/Editor/UnrealEd/Public/Kismet2/StructureEditorUtils.h
--- a/Engine/Source/Editor/UnrealEd/Public/Kismet2/StructureEditorUtils.h
+++ b/Engine/Source/Editor/UnrealEd/Public/Kismet2/StructureEditorUtils.h
@@ -138,6 +138,10 @@ public:
 	/** Can the structure be a member variable for a BPGClass or BPGStruct */
 	static EStructureError IsStructureValid(const UScriptStruct* Struct, const UStruct* RecursionParent = NULL, FString* OutMsg = NULL);
 
+	/** called after UDS was changed by editor*/
 	static void OnStructureChanged(UUserDefinedStruct* Struct);
+
+	static void BroadcastPreChange(UUserDefinedStruct* Struct);
+	static void BroadcastPostChange(UUserDefinedStruct* Struct);
 };
 
diff --git a/Engine/Source/Runtime/Engine/Classes/Engine/DataTable.h b/Engine/Source/Runtime/Engine/Classes/Engine/DataTable.h
--- a/Engine/Source/Runtime/Engine/Classes/Engine/DataTable.h
+++ b/Engine/Source/Runtime/Engine/Classes/Engine/DataTable.h
@@ -121,6 +121,14 @@ class UDataTable : public UObject
 
 #if WITH_EDITOR || HACK_HEADER_GENERATOR
 
+private:
+	//when RowStruct is being modified, row data is stored serialized with tags
+	TArray<uint8> RowsSerializedWithTags;
+
+public:
+	ENGINE_API void CleanBeforeStructChange();
+	ENGINE_API void RestoreAfterStructChange();
+
 	/** Output entire contents of table as a string */
 	ENGINE_API FString GetTableAsString();
 
@@ -164,6 +172,10 @@ class UDataTable : public UObject
 	ENGINE_API static bool IsSupportedTableProperty(const UProperty* InProp);
 
 	// End UDataTable interface
+
+private:
+	void SaveStructData(FArchive& Ar);
+	void LoadStructData(FArchive& Ar);
 };
 
 /** Handle to a particular row in a table*/
diff --git a/Engine/Source/Runtime/Engine/Private/DataTable.cpp b/Engine/Source/Runtime/Engine/Private/DataTable.cpp
--- a/Engine/Source/Runtime/Engine/Private/DataTable.cpp
+++ b/Engine/Source/Runtime/Engine/Private/DataTable.cpp
@@ -3,6 +3,10 @@
 #include "EnginePrivate.h"
 #include "Json.h"
 
+#if WITH_EDITOR
+#include "StructureEditorUtils.h"
+#endif //WITH_EDITOR
+
 DEFINE_LOG_CATEGORY(LogDataTable);
 
 ENGINE_API const FString FDataTableRowHandle::Unknown(TEXT("UNKNOWN"));
@@ -93,58 +97,68 @@ TArray<FName> UDataTable::GetStructPropertyNames(UStruct* InStruct)
 
 //////////////////////////////////////////////////////////////////////////
 
-void UDataTable::Serialize( FArchive& Ar )
+void UDataTable::LoadStructData(FArchive& Ar)
 {
-	Super::Serialize(Ar); // When loading, this should load our RowStruct!	
+	UScriptStruct* LoadUsingStruct = RowStruct;
+	if (LoadUsingStruct == NULL)
+	{
+		UE_LOG(LogDataTable, Error, TEXT("Missing RowStruct while loading DataTable '%s'!"), *GetPathName());
+		LoadUsingStruct = FTableRowBase::StaticStruct();
+	}
 
-	if(Ar.IsLoading())
+	int32 NumRows;
+	Ar << NumRows;
+
+	for (int32 RowIdx = 0; RowIdx < NumRows; RowIdx++)
 	{
-		UScriptStruct* LoadUsingStruct = RowStruct;
-		if(LoadUsingStruct == NULL)
-		{
-			UE_LOG(LogDataTable, Error, TEXT("Missing RowStruct while loading DataTable '%s'!"), *GetPathName());
-			LoadUsingStruct = FTableRowBase::StaticStruct();
-		}
+		// Load row name
+		FName RowName;
+		Ar << RowName;
 
-		int32 NumRows;
+		// Load row data
+		uint8* RowData = (uint8*)FMemory::Malloc(LoadUsingStruct->PropertiesSize);
+		LoadUsingStruct->InitializeScriptStruct(RowData);
+		// And be sure to call DestroyScriptStruct later
+		LoadUsingStruct->SerializeTaggedProperties(Ar, RowData, LoadUsingStruct, NULL);
+
+		// Add to map
+		RowMap.Add(RowName, RowData);
+	}
+}
+
+void UDataTable::SaveStructData(FArchive& Ar)
+{
+	// Don't even try to save rows if no RowStruct
+	if (RowStruct != NULL)
+	{
+		int32 NumRows = RowMap.Num();
 		Ar << NumRows;
 
-		for(int32 RowIdx=0; RowIdx<NumRows; RowIdx++)
+		// Now iterate over rows in the map
+		for (auto RowIt = RowMap.CreateIterator(); RowIt; ++RowIt)
 		{
-			// Load row name
-			FName RowName;
+			// Save out name
+			FName RowName = RowIt.Key();
 			Ar << RowName;
 
-			// Load row data
-			uint8* RowData = (uint8*)FMemory::Malloc(LoadUsingStruct->PropertiesSize);
-			LoadUsingStruct->InitializeScriptStruct(RowData);
-			// And be sure to call DestroyScriptStruct later
-			LoadUsingStruct->SerializeTaggedProperties(Ar, RowData, LoadUsingStruct, NULL);
-
-			// Add to map
-			RowMap.Add(RowName, RowData);
+			// Save out data
+			uint8* RowData = RowIt.Value();
+			RowStruct->SerializeTaggedProperties(Ar, RowData, RowStruct, NULL);
 		}
 	}
-	else if(Ar.IsSaving())
-	{
-		// Don't even try to save rows if no RowStruct
-		if(RowStruct != NULL)
-		{
-			int32 NumRows = RowMap.Num();
-			Ar << NumRows;
+}
 
-			// Now iterate over rows in the map
-			for ( auto RowIt = RowMap.CreateIterator(); RowIt; ++RowIt )
-			{
-				// Save out name
-				FName RowName = RowIt.Key();
-				Ar << RowName;
+void UDataTable::Serialize( FArchive& Ar )
+{
+	Super::Serialize(Ar); // When loading, this should load our RowStruct!	
 
-				// Save out data
-				uint8* RowData = RowIt.Value();
-				RowStruct->SerializeTaggedProperties(Ar, RowData, RowStruct, NULL);
-			}
-		}
+	if(Ar.IsLoading())
+	{
+		LoadStructData(Ar);
+	}
+	else if(Ar.IsSaving())
+	{
+		SaveStructData(Ar);
 	}
 }
 
@@ -225,6 +239,36 @@ UProperty* UDataTable::FindTableProperty(const FName& PropertyName) const
 
 #if WITH_EDITOR || HACK_HEADER_GENERATOR
 
+struct FPropertyDisplayNameHelper
+{
+	static inline FString Get(const UProperty* Prop, const FString& DefaultName)
+	{
+		static const FName DisplayNameKey(TEXT("DisplayName"));
+		return (Prop && Prop->HasMetaData(DisplayNameKey)) ? Prop->GetMetaData(DisplayNameKey) : DefaultName;
+	}
+};
+
+void UDataTable::CleanBeforeStructChange()
+{
+	RowsSerializedWithTags.Reset();
+	{
+		FMemoryWriter MemoryWriter(RowsSerializedWithTags);
+		SaveStructData(MemoryWriter);
+	}
+	EmptyTable();
+	Modify();
+}
+
+void UDataTable::RestoreAfterStructChange()
+{
+	EmptyTable();
+	{
+		FMemoryReader MemoryReader(RowsSerializedWithTags);
+		LoadStructData(MemoryReader);
+	}
+	RowsSerializedWithTags.Empty();
+}
+
 FString UDataTable::GetTableAsString()
 {
 	FString Result;
@@ -349,8 +393,6 @@ bool UDataTable::WriteTableAsJSON(const TSharedRef< TJsonWriter<TCHAR, TPrettyJs
 /** Get array of UProperties that corresponds to columns in the table */
 TArray<UProperty*> UDataTable::GetTablePropertyArray(const FString& FirstRowString, UStruct* InRowStruct, TArray<FString>& OutProblems)
 {
-	static const FName DisplayNameKey(TEXT("DisplayName"));
-
 	TArray<UProperty*> ColumnProps;
 
 	// Get list of all expected properties from the struct
@@ -380,7 +422,7 @@ TArray<UProperty*> UDataTable::GetTablePropertyArray(const FString& FirstRowStri
 
 				for (TFieldIterator<UProperty> It(InRowStruct); It && !ColumnProp; ++It)
 				{
-					const FString DisplayName = *It ? (*It)->GetMetaData(DisplayNameKey) : FString();
+					const auto DisplayName = FPropertyDisplayNameHelper::Get(*It, FString());
 					ColumnProp = (!DisplayName.IsEmpty() && (DisplayName == ColumnNameStrings[ColIdx])) ? *It : NULL;
 				}
 
@@ -419,9 +461,7 @@ TArray<UProperty*> UDataTable::GetTablePropertyArray(const FString& FirstRowStri
 	for(int32 PropIdx=0; PropIdx < ExpectedPropNames.Num(); PropIdx++)
 	{
 		const UProperty* const ColumnProp = FindField<UProperty>(InRowStruct, ExpectedPropNames[PropIdx]);
-		const FString DisplayName = (ColumnProp && ColumnProp->HasMetaData(DisplayNameKey))
-			? ColumnProp->GetMetaData(DisplayNameKey) 
-			: ExpectedPropNames[PropIdx].ToString();
+		const FString DisplayName = FPropertyDisplayNameHelper::Get(ColumnProp, ExpectedPropNames[PropIdx].ToString());
 		OutProblems.Add(FString::Printf(TEXT("Expected column '%s' not found in input."), *DisplayName));
 	}
 
@@ -499,6 +539,10 @@ TArray<FString> UDataTable::CreateTableFromCSVString(const FString& InString)
 		RowStruct->InitializeScriptStruct(RowData);
 		// And be sure to call DestroyScriptStruct later
 
+#if WITH_EDITOR
+		FStructureEditorUtils::Fill_MakeStructureDefaultValue(Cast<const UUserDefinedStruct>(RowStruct), RowData);
+#endif // WITH_EDITOR
+
 		// Add to row map
 		RowMap.Add(RowName, RowData);
 
@@ -512,7 +556,9 @@ TArray<FString> UDataTable::CreateTableFromCSVString(const FString& InString)
 			// If we failed, output a problem string
 			if(Error.Len() > 0)
 			{
-				FString ColumnName = (ColumnProp != NULL) ? ColumnProp->GetName() : FString(TEXT("NONE"));
+				FString ColumnName = (ColumnProp != NULL) 
+					? FPropertyDisplayNameHelper::Get(ColumnProp, ColumnProp->GetName())
+					: FString(TEXT("NONE"));
 				OutProblems.Add(FString::Printf(TEXT("Problem assigning string '%s' to property '%s' on row '%s' : %s"), *CellStrings[CellIdx], *ColumnName, *RowName.ToString(), *Error));
 			}
 		}
@@ -541,16 +587,13 @@ TArray<FString> UDataTable::CreateTableFromJSONString(const FString& InString)
 
 TArray<FString> UDataTable::GetColumnTitles() const
 {
-	static const FName DisplayNameKey(TEXT("DisplayName"));
 	TArray<FString> Result;
 	Result.Add(TEXT("Name"));
 	for (TFieldIterator<UProperty> It(RowStruct); It; ++It)
 	{
 		UProperty* Prop = *It;
 		check(Prop != NULL);
-		const FString DisplayName = Prop->HasMetaData(DisplayNameKey)
-			? Prop->GetMetaData(DisplayNameKey)
-			: Prop->GetName();
+		const FString DisplayName = FPropertyDisplayNameHelper::Get(Prop, Prop->GetName());
 		Result.Add(DisplayName);
 	}
 	return Result;
diff --git a/Engine/Source/Runtime/Engine/Private/UserDefinedStruct.cpp b/Engine/Source/Runtime/Engine/Private/UserDefinedStruct.cpp
--- a/Engine/Source/Runtime/Engine/Private/UserDefinedStruct.cpp
+++ b/Engine/Source/Runtime/Engine/Private/UserDefinedStruct.cpp
@@ -43,6 +43,9 @@ void UUserDefinedStruct::GetAssetRegistryTags(TArray<FAssetRegistryTag>& OutTags
 	Super::GetAssetRegistryTags(OutTags);
 
 	OutTags.Add(FAssetRegistryTag(TEXT("Tooltip"), FStructureEditorUtils::GetTooltip(this), FAssetRegistryTag::TT_Hidden));
+
+	const FName SuperStructName = SuperStruct ? SuperStruct->GetFName() : NAME_None;
+	OutTags.Add(FAssetRegistryTag(TEXT("BaseStruct"), SuperStructName.ToString(), FAssetRegistryTag::TT_Alphabetical));
 }
 
 UProperty* UUserDefinedStruct::CustomFindProperty(const FName Name) const
