Header cleanup: Engine module, public headers.

diff --git a/Engine/Source/Runtime/Engine/Classes/Engine/NiagaraEffectRenderer.h b/Engine/Source/Runtime/Engine/Classes/Engine/NiagaraEffectRenderer.h
--- a/Engine/Source/Runtime/Engine/Classes/Engine/NiagaraEffectRenderer.h
+++ b/Engine/Source/Runtime/Engine/Classes/Engine/NiagaraEffectRenderer.h
@@ -4,6 +4,7 @@
 NiagaraEffectRenderer.h: Base class for Niagara render modules
 ==============================================================================*/
 #include "SceneUtils.h"
+#include "ParticleBeamTrailVertexFactory.h"
 
 DECLARE_CYCLE_STAT(TEXT("PreRenderView"), STAT_NiagaraPreRenderView, STATGROUP_Niagara);
 DECLARE_CYCLE_STAT(TEXT("Render Total"), STAT_NiagaraRender, STATGROUP_Niagara);
diff --git a/Engine/Source/Runtime/Engine/Private/Particles/ParticleBeamTrailVertexFactory.cpp b/Engine/Source/Runtime/Engine/Private/Particles/ParticleBeamTrailVertexFactory.cpp
--- a/Engine/Source/Runtime/Engine/Private/Particles/ParticleBeamTrailVertexFactory.cpp
+++ b/Engine/Source/Runtime/Engine/Private/Particles/ParticleBeamTrailVertexFactory.cpp
@@ -7,6 +7,7 @@
 #include "EnginePrivate.h"
 #include "ParticleDefinitions.h"
 #include "ShaderParameterUtils.h"
+#include "ParticleBeamTrailVertexFactory.h"
 
 IMPLEMENT_UNIFORM_BUFFER_STRUCT(FParticleBeamTrailUniformParameters,TEXT("BeamTrailVF"));
 
diff --git a/Engine/Source/Runtime/Engine/Public/ParticleEmitterInstances.h b/Engine/Source/Runtime/Engine/Public/ParticleEmitterInstances.h
--- a/Engine/Source/Runtime/Engine/Public/ParticleEmitterInstances.h
+++ b/Engine/Source/Runtime/Engine/Public/ParticleEmitterInstances.h
@@ -8,6 +8,7 @@
 #pragma once
 
 #include "Distributions.h"
+#include "ParticleHelper.h"
 #include "Distributions/DistributionFloat.h"
 #include "Distributions/DistributionVector.h"
 #include "ParticleHelper.h"
@@ -33,8 +34,6 @@ class UParticleModuleBeamTarget;
 class UParticleModuleBeamNoise;
 
 class UParticleModuleTrailSource;
-class UParticleModuleTrailSpawn;
-class UParticleModuleTrailTaper;
 
 class UParticleModuleSpawnPerUnit;
 
@@ -42,11 +41,6 @@ class UParticleModuleOrientationAxisLock;
 
 class UParticleLODLevel;
 
-class FParticleSystemSceneProxy;
-class FParticleDynamicData;
-struct FDynamicBeam2EmitterData;
-struct FDynamicTrail2EmitterData;
-
 struct FLODBurstFired
 {
 	TArray<bool> Fired;
diff --git a/Engine/Source/Runtime/Engine/Public/ParticleHelper.h b/Engine/Source/Runtime/Engine/Public/ParticleHelper.h
--- a/Engine/Source/Runtime/Engine/Public/ParticleHelper.h
+++ b/Engine/Source/Runtime/Engine/Public/ParticleHelper.h
@@ -7,7 +7,6 @@
 #pragma once
 
 #include "ParticleVertexFactory.h"
-#include "ParticleBeamTrailVertexFactory.h"
 #include "MeshParticleVertexFactory.h"
 #include "MaterialShared.h"
 #include "Materials/MaterialInterface.h"
@@ -129,7 +128,6 @@ class UParticleLODLevel;
 class USkeletalMeshComponent;
 
 class FParticleSystemSceneProxy;
-class FParticleDynamicData;
 struct FDynamicBeam2EmitterData;
 struct FDynamicTrail2EmitterData;
 
