Fixup multiline comments that are meant to be documentation. Need to start with /** for doxygen to parse them.

diff --git a/Engine/Source/Developer/AutomationWindow/Private/AutomationPresetManager.h b/Engine/Source/Developer/AutomationWindow/Private/AutomationPresetManager.h
--- a/Engine/Source/Developer/AutomationWindow/Private/AutomationPresetManager.h
+++ b/Engine/Source/Developer/AutomationWindow/Private/AutomationPresetManager.h
@@ -146,7 +146,7 @@ public:
 	 */
 	void SavePreset( const AutomationPresetRef Preset);
 
-	/*
+	/**
 	* Load all Presets from disk.
 	*/
 	void LoadPresets( );
diff --git a/Engine/Source/Developer/LauncherServices/Private/Profiles/LauncherProfileManager.h b/Engine/Source/Developer/LauncherServices/Private/Profiles/LauncherProfileManager.h
--- a/Engine/Source/Developer/LauncherServices/Private/Profiles/LauncherProfileManager.h
+++ b/Engine/Source/Developer/LauncherServices/Private/Profiles/LauncherProfileManager.h
@@ -73,12 +73,12 @@ public:
 
 protected:
 
-	/*
+	/**
 	 * Loads all the device groups from a config file
 	 */
 	void LoadDeviceGroups( );
 
-	/*
+	/**
 	 * Load all profiles from disk.
 	 */
 	void LoadProfiles( );
@@ -92,12 +92,12 @@ protected:
 	 */
 	ILauncherDeviceGroupPtr ParseDeviceGroup( const FString& GroupString );
 
-	/*
+	/**
 	 * Saves all the device groups to a config file
 	 */
 	void SaveDeviceGroups( );
 
-	/*
+	/**
 	 * Saves all profiles to disk.
 	 */
 	void SaveProfiles( );
diff --git a/Engine/Source/Developer/LauncherServices/Public/Interfaces/ILauncherProfileManager.h b/Engine/Source/Developer/LauncherServices/Public/Interfaces/ILauncherProfileManager.h
--- a/Engine/Source/Developer/LauncherServices/Public/Interfaces/ILauncherProfileManager.h
+++ b/Engine/Source/Developer/LauncherServices/Public/Interfaces/ILauncherProfileManager.h
@@ -53,7 +53,7 @@ public:
 	 */
 	virtual void AddDeviceGroup( const ILauncherDeviceGroupRef& DeviceGroup ) = 0;
 
-	/*
+	/**
 	 * Create a new device group and maintains a reference for its future usage.
 	 *
 	 * @return The device group created.
diff --git a/Engine/Source/Developer/Profiler/Private/ProfilerDataProvider.h b/Engine/Source/Developer/Profiler/Private/ProfilerDataProvider.h
--- a/Engine/Source/Developer/Profiler/Private/ProfilerDataProvider.h
+++ b/Engine/Source/Developer/Profiler/Private/ProfilerDataProvider.h
@@ -27,7 +27,7 @@ public:
 	{}
 
 protected:
-	/*
+	/**
 	 * Method for creating a duplicated copy of this data provider instance with a particular group of frames.
 	 * Data is stored in the instance of FArrayDataProvider class.
 	 * Data source for the SEventGraph widget.
@@ -163,7 +163,7 @@ public:
 		return FIntPoint( StartIndex, FMath::Min(ElapsedFrameTimes.Num(),EndIndex+1) );
 	}
 
-	/*
+	/**
 	 * @return An instance of FIntPoint 
 	 * where X is an index of first sample in the frame, and 
 	 * where Y is an index of last sample in the frame + 1. 
@@ -269,4 +269,4 @@ protected:
 	FProfilerSampleArray Collection;
 
 	uint64 ChildrenIndicesMemoryUsage;
-};
\ No newline at end of file
+};
diff --git a/Engine/Source/Developer/Profiler/Private/ProfilerDataSource.h b/Engine/Source/Developer/Profiler/Private/ProfilerDataSource.h
--- a/Engine/Source/Developer/Profiler/Private/ProfilerDataSource.h
+++ b/Engine/Source/Developer/Profiler/Private/ProfilerDataSource.h
@@ -437,7 +437,7 @@ public:
 /**	Type definition for type of the cached values. */
 typedef float TGraphDataType;
 
-/* 
+/** 
  * A specialized view of the a data provider. Provides access only to the specified group of data. 
  * This class allows accessing data in linear way which may be used to draw a line graph.
  */
@@ -531,7 +531,7 @@ protected:
 
 class FGraphDataSourceVerticalAggregate;
 
-/* 
+/** 
  * A specialized view of a few data providers. Provides access only to the specified group of data. 
  * Data is interpolated for 60 frames per second.
  * This class allows accessing data in linear way which may be used to draw a combined line graph with min, max and average values.
@@ -958,7 +958,7 @@ struct FEventGraphConsts
 };
 
 // TODO: Rename to FProfilerEvent
-/*
+/**
  * Contains the same data as the profiler sample with some additions, doesn't depend on the other classes like profiler metadata or profiler aggregates.
  * Modeled to be Slate compatible, thus it inherits from TSharedFromThis.
  */
@@ -2192,7 +2192,7 @@ public:
 	TMap<FGuid,FEventGraphDataProperties> EventGraphProperties;
 };
 
-/* 
+/** 
  * Provides access only to the profiler samples specified by a frame index or frame indices.
  * This class allows accessing root and child samples which may be used to create an event graph.
  */
diff --git a/Engine/Source/Developer/Profiler/Private/Widgets/SEventGraph.h b/Engine/Source/Developer/Profiler/Private/Widgets/SEventGraph.h
--- a/Engine/Source/Developer/Profiler/Private/Widgets/SEventGraph.h
+++ b/Engine/Source/Developer/Profiler/Private/Widgets/SEventGraph.h
@@ -221,7 +221,7 @@ protected:
 	/** Called by STreeView to generate a table row for the specified item. */
 	TSharedRef< ITableRow > EventGraph_OnGenerateRow( FEventGraphSamplePtr EventPtr, const TSharedRef< STableViewBase >& OwnerTable );
 
-	/*
+	/**
 	 * Called by STreeView to retrieves the children for the specified parent item.
 	 * @param InParent		- The parent node to retrieve the children from.
 	 * @param out_Children	- List of children for the parent node.
@@ -983,4 +983,4 @@ protected:
 
 	/** Name of the event that should be drawn as highlighted. */
 	FName HighlightedEventName;
-};
\ No newline at end of file
+};
diff --git a/Engine/Source/Developer/Profiler/Private/Widgets/SFiltersAndPresets.h b/Engine/Source/Developer/Profiler/Private/Widgets/SFiltersAndPresets.h
--- a/Engine/Source/Developer/Profiler/Private/Widgets/SFiltersAndPresets.h
+++ b/Engine/Source/Developer/Profiler/Private/Widgets/SFiltersAndPresets.h
@@ -327,7 +327,7 @@ protected:
 	/** Called by STreeView to generate a table row for the specified item. */
 	TSharedRef< ITableRow > GroupAndStatTree_OnGenerateRow( FGroupOrStatNodePtr TreeNode, const TSharedRef< STableViewBase >& OwnerTable );
 
-	/*
+	/**
 	 * Called by STreeView to retrieves the children for the specified parent item.
 	 * @param InParent		- The parent node to retrieve the children from.
 	 * @param out_Children	- List of children for the parent node.
@@ -423,4 +423,4 @@ protected:
 
 	/** Holds the visibility of each stat type. */
 	bool bStatTypeIsVisible[EProfilerSampleTypes::InvalidOrMax];
-};
\ No newline at end of file
+};
diff --git a/Engine/Source/Developer/ScreenShotComparisonTools/Private/ScreenShotToolsModule.cpp b/Engine/Source/Developer/ScreenShotComparisonTools/Private/ScreenShotToolsModule.cpp
--- a/Engine/Source/Developer/ScreenShotComparisonTools/Private/ScreenShotToolsModule.cpp
+++ b/Engine/Source/Developer/ScreenShotComparisonTools/Private/ScreenShotToolsModule.cpp
@@ -37,7 +37,7 @@ public:
 
 public:
 
-	/*
+	/**
 	 * Shutdown the screen shot manager tools module
 	 */
 	virtual void ShutdownModule( ) override
@@ -45,7 +45,7 @@ public:
 		ScreenShotManager.Reset();
 	}
 
-	/*
+	/**
 	* Startup the screen shot manager tools module
 	*/
 	void StartupModule() override
diff --git a/Engine/Source/Developer/SessionLauncher/Private/Models/SessionLauncherModel.h b/Engine/Source/Developer/SessionLauncher/Private/Models/SessionLauncherModel.h
--- a/Engine/Source/Developer/SessionLauncher/Private/Models/SessionLauncherModel.h
+++ b/Engine/Source/Developer/SessionLauncher/Private/Models/SessionLauncherModel.h
@@ -172,7 +172,7 @@ public:
 
 protected:
 
-	/*
+	/**
 	 * Load all profiles from disk.
 	 */
 	void LoadConfig( )
@@ -209,7 +209,7 @@ protected:
 		SelectProfile(LoadedSelectedProfile);
 	}
 
-	/*
+	/**
 	 * Saves all profiles to disk.
 	 */
 	void SaveConfig( )
diff --git a/Engine/Source/Developer/StandaloneRenderer/Private/OpenGL/SlateOpenGLShaders.h b/Engine/Source/Developer/StandaloneRenderer/Private/OpenGL/SlateOpenGLShaders.h
--- a/Engine/Source/Developer/StandaloneRenderer/Private/OpenGL/SlateOpenGLShaders.h
+++ b/Engine/Source/Developer/StandaloneRenderer/Private/OpenGL/SlateOpenGLShaders.h
@@ -20,7 +20,7 @@ protected:
 protected:
 	GLuint ShaderID;
 };
-/*
+/**
  * Represents an OpenGL vertex shader.  Do not inherit from this for specific shaders.  Inherit from FSlateOpenGLProgram
  */
 class FSlateOpenGLVS : public FSlateOpenGLShader
@@ -31,7 +31,7 @@ public:
 	virtual void Create( const FString& Filename );
 };
 
-/*
+/**
  * Represents an OpenGL vertex shader.  Do not inherit from this for specific shaders.  Inherit from FSlateOpenGLProgram
  */
 class FSlateOpenGLPS : public FSlateOpenGLShader
diff --git a/Engine/Source/Developer/TargetDeviceServices/Private/TargetDeviceProxy.h b/Engine/Source/Developer/TargetDeviceServices/Private/TargetDeviceProxy.h
--- a/Engine/Source/Developer/TargetDeviceServices/Private/TargetDeviceProxy.h
+++ b/Engine/Source/Developer/TargetDeviceServices/Private/TargetDeviceProxy.h
@@ -51,7 +51,7 @@ public:
 		return LastUpdateTime;
 	}
 
-	/*
+	/**
 	 * Updates the proxy's information from the given device service response.
 	 *
 	 * @param Message The message containing the response.
diff --git a/Engine/Source/Editor/BehaviorTreeEditor/Classes/BehaviorTreeGraphNode.h b/Engine/Source/Editor/BehaviorTreeEditor/Classes/BehaviorTreeGraphNode.h
--- a/Engine/Source/Editor/BehaviorTreeEditor/Classes/BehaviorTreeGraphNode.h
+++ b/Engine/Source/Editor/BehaviorTreeEditor/Classes/BehaviorTreeGraphNode.h
@@ -61,7 +61,7 @@ class UBehaviorTreeGraphNode : public UEdGraphNode
 
 	virtual bool IsSubNode() const { return false; }
 
-	/*
+	/**
 	 * Finds the difference in Behavior Tree properties
 	 * 
 	 * @param Struct The struct of the class we are looking at
diff --git a/Engine/Source/Editor/DeviceProfileEditor/Private/SDeviceProfileCreateProfilePanel.h b/Engine/Source/Editor/DeviceProfileEditor/Private/SDeviceProfileCreateProfilePanel.h
--- a/Engine/Source/Editor/DeviceProfileEditor/Private/SDeviceProfileCreateProfilePanel.h
+++ b/Engine/Source/Editor/DeviceProfileEditor/Private/SDeviceProfileCreateProfilePanel.h
@@ -59,7 +59,7 @@ public:
 	 */
 	void HandleBaseProfileSelectionChanged( UDeviceProfile* NewSelection, ESelectInfo::Type SelectInfo );
 
-	/*
+	/**
 	 * Set the base combo box content.
 	 *
 	 * @return The device profile parent name
@@ -86,7 +86,7 @@ public:
 	 */
 	void HandleProfileTypeChanged( TSharedPtr<FString> NewSelection, ESelectInfo::Type SelectInfo );
 	
-	/*
+	/**
 	 * Set the profile type combo box content.
 	 *
 	 * @return The name of the class the created profile will be.
@@ -148,4 +148,4 @@ private:
 };
 
 
-#undef LOCTEXT_NAMESPACE
\ No newline at end of file
+#undef LOCTEXT_NAMESPACE
diff --git a/Engine/Source/Editor/Documentation/Private/UDNParser.h b/Engine/Source/Editor/Documentation/Private/UDNParser.h
--- a/Engine/Source/Editor/Documentation/Private/UDNParser.h
+++ b/Engine/Source/Editor/Documentation/Private/UDNParser.h
@@ -167,7 +167,7 @@ private:
 	void NavigateToLink( FString AdditionalContent );
 
 private:
-	/* 
+	/** 
 	 * Parses a code link embedded in the doc.
 	 * Allows us to specify files in code to link to in one of 2 ways. In both cases the last 2 parameters are line and column.
 	 * [Project based link](CODELINK:Private/[PROJECT]File.cpp, 29, 5)
@@ -179,7 +179,7 @@ private:
 	 */
 	bool ParseCodeLink(FString &InternalLink);
 		
-	/* 
+	/** 
 	 * Parses an asset link embedded in the doc.
 	 * Allows us to specify assets to either highlight or edit in the editor
 	 * (ASSETLINK:SELECT,MyCharacter)
diff --git a/Engine/Source/Editor/EditorWidgets/Public/SAssetSearchBox.h b/Engine/Source/Editor/EditorWidgets/Public/SAssetSearchBox.h
--- a/Engine/Source/Editor/EditorWidgets/Public/SAssetSearchBox.h
+++ b/Engine/Source/Editor/EditorWidgets/Public/SAssetSearchBox.h
@@ -2,7 +2,7 @@
 
 #pragma once
 
-/*
+/**
  * A widget to provide a search box with a filtered dropdown menu.
  */
 
diff --git a/Engine/Source/Editor/GeometryMode/Classes/GeomModifier.h b/Engine/Source/Editor/GeometryMode/Classes/GeomModifier.h
--- a/Engine/Source/Editor/GeometryMode/Classes/GeomModifier.h
+++ b/Engine/Source/Editor/GeometryMode/Classes/GeomModifier.h
@@ -44,7 +44,7 @@ public:
 	/** @return		true if the delta was handled by this editor mode tool. */
 	virtual bool InputDelta(class FEditorViewportClient* InViewportClient, FViewport* InViewport, FVector& InDrag, FRotator& InRot, FVector& InScale);
 
-	/*
+	/**
 	 * Drawing functions to allow modifiers to have better control over the screen.
 	 */
 	virtual void Render(const FSceneView* View,FViewport* Viewport,FPrimitiveDrawInterface* PDI);
diff --git a/Engine/Source/Editor/GraphEditor/Public/SGraphActionMenu.h b/Engine/Source/Editor/GraphEditor/Public/SGraphActionMenu.h
--- a/Engine/Source/Editor/GraphEditor/Public/SGraphActionMenu.h
+++ b/Engine/Source/Editor/GraphEditor/Public/SGraphActionMenu.h
@@ -205,7 +205,7 @@ public:
 	/** Get category child actions for the passed in action */
 	void GetCategorySubActions(TWeakPtr<FGraphActionNode> InAction, TArray<TSharedPtr<FEdGraphSchemaAction>>& OutActions) const;
 
-	/*
+	/**
 	 * Selects an non-creation item in the list, searching by FName, deselects if name is none 
 	 *
 	 * @param	ItemName		The name of the item to select
@@ -255,7 +255,7 @@ protected:
 	/** Create widget for the supplied node */
 	TSharedRef<ITableRow> MakeWidget( TSharedPtr<FGraphActionNode> InItem, const TSharedRef<STableViewBase>& OwnerTable, bool bIsReadOnly );
 
-	/*
+	/**
 	 * Called when tree item is selected 
 	 *
 	 * @param	InSelectedItem	The action node that is being selected
@@ -263,7 +263,7 @@ protected:
 	 */
 	void OnItemSelected( TSharedPtr< FGraphActionNode > InSelectedItem, ESelectInfo::Type SelectInfo );
 
-	/*
+	/**
 	 * Executes the selection delegate providing it has been bound, and the provided action node given is valid and is an action node
 	 *
 	 * @param InselectedItem	The graph action node selected
diff --git a/Engine/Source/Editor/Kismet/Private/SSCSEditor.h b/Engine/Source/Editor/Kismet/Private/SSCSEditor.h
--- a/Engine/Source/Editor/Kismet/Private/SSCSEditor.h
+++ b/Engine/Source/Editor/Kismet/Private/SSCSEditor.h
@@ -468,14 +468,14 @@ public:
 	/** Select the given tree node */
 	void SelectNode(FSCSEditorTreeNodePtrType InNodeToSelect, bool IsCntrlDown);
 
-	/*
+	/**
 	 * Function to save current state of SimpleConstructionScript and nodes associated with it.
 	 *
 	 * @param: SimpleContructionScript object reference.
 	 */
 	static void SaveSCSCurrentState( USimpleConstructionScript* SCSObj );
 
-	/*
+	/**
 	 * Function to save the current state of SCS_Node and its children
 	 *
 	 * @param: Reference of the SCS_Node to be saved
diff --git a/Engine/Source/Editor/Kismet/Public/BlueprintEditor.h b/Engine/Source/Editor/Kismet/Public/BlueprintEditor.h
--- a/Engine/Source/Editor/Kismet/Public/BlueprintEditor.h
+++ b/Engine/Source/Editor/Kismet/Public/BlueprintEditor.h
@@ -125,7 +125,7 @@ public:
 	 */
 	void InitBlueprintEditor(const EToolkitMode::Type Mode, const TSharedPtr< class IToolkitHost >& InitToolkitHost, const TArray<class UBlueprint*>& InBlueprints, bool bShouldOpenInDefaultsMode);
 
-	/*
+	/**
 	 * Set transactional flag on SCSNodes and its children.
 	 *
 	 * @param: Node reference to set transactional flag.
@@ -420,7 +420,7 @@ public:
 	/** Handles spawning a graph node in the current graph using the passed in gesture */
 	FReply OnSpawnGraphNodeByShortcut(FInputGesture InGesture, const FVector2D& InPosition, UEdGraph* InGraph);
 
-	/* 
+	/** 
 	 * Perform the actual promote to variable action on the given pin in the given blueprint.
 	 *
 	 * @param	InBlueprint	The blueprint in which to create the variable.
diff --git a/Engine/Source/Editor/LandscapeEditor/Private/LandscapeEdModeTools.h b/Engine/Source/Editor/LandscapeEditor/Private/LandscapeEdModeTools.h
--- a/Engine/Source/Editor/LandscapeEditor/Private/LandscapeEdModeTools.h
+++ b/Engine/Source/Editor/LandscapeEditor/Private/LandscapeEdModeTools.h
@@ -1212,7 +1212,7 @@ public:
 };
 
 
-/*
+/**
  * FLandscapeToolPaintBase - base class for painting tools
  *		ToolTarget - the target for the tool (weight or heightmaap)
  *		StrokeClass - the class that implements the behavior for a mouse stroke applying the tool.
diff --git a/Engine/Source/Editor/LevelEditor/Private/LevelEditorContextMenu.cpp b/Engine/Source/Editor/LevelEditor/Private/LevelEditorContextMenu.cpp
--- a/Engine/Source/Editor/LevelEditor/Private/LevelEditorContextMenu.cpp
+++ b/Engine/Source/Editor/LevelEditor/Private/LevelEditorContextMenu.cpp
@@ -124,7 +124,7 @@ FSelectedActorInfo FLevelEditorContextMenuImpl::SelectionInfo;
 
 struct FLevelScriptEventMenuHelper
 {
-	/*
+	/**
 	* Fills in menu options for events that can be associated with that actors's blueprint in the level script blueprint
 	*
 	* @param MenuBuilder	The menu to add items to
diff --git a/Engine/Source/Editor/LevelEditor/Public/LightmapResRatioAdjust.h b/Engine/Source/Editor/LevelEditor/Public/LightmapResRatioAdjust.h
--- a/Engine/Source/Editor/LevelEditor/Public/LightmapResRatioAdjust.h
+++ b/Engine/Source/Editor/LevelEditor/Public/LightmapResRatioAdjust.h
@@ -36,7 +36,7 @@ public:
 
 	bool ApplyRatioAdjustment();
 
-	/* 
+	/** 
 	* Add levels in a world to an array based on the LevelOptions setting
 	*
 	* @param InLevelOptions	Level options used to decided whether to add levels to the array
diff --git a/Engine/Source/Editor/Persona/Private/PersonaMeshDetails.h b/Engine/Source/Editor/Persona/Private/PersonaMeshDetails.h
--- a/Engine/Source/Editor/Persona/Private/PersonaMeshDetails.h
+++ b/Engine/Source/Editor/Persona/Private/PersonaMeshDetails.h
@@ -3,7 +3,7 @@
 #pragma once
 
 #include "IDetailCustomNodeBuilder.h"
-/*
+/**
  * Struct to uniquely identify clothing applied to a material section
  * Contains index into the ClothingAssets array and the submesh index.
  */
diff --git a/Engine/Source/Editor/PropertyEditor/Private/PropertyNode.cpp b/Engine/Source/Editor/PropertyEditor/Private/PropertyNode.cpp
--- a/Engine/Source/Editor/PropertyEditor/Private/PropertyNode.cpp
+++ b/Engine/Source/Editor/PropertyEditor/Private/PropertyNode.cpp
@@ -1592,7 +1592,7 @@ void FPropertyNode::ResetToDefault( FNotifyHook* InNotifyHook )
 	}
 }
 
-/*
+/**
  * Helper function to obtain the display name for an enum property
  * @param InEnum		The enum whose metadata to pull from
  * @param DisplayName	The name of the enum value to adjust
@@ -2386,4 +2386,4 @@ bool FPropertyNode::GenerateRestrictionToolTip(const FString& Value, FText& OutT
 		}
 	}
 	return bRestricted;
-}
\ No newline at end of file
+}
diff --git a/Engine/Source/Editor/PropertyEditor/Private/PropertyNode.h b/Engine/Source/Editor/PropertyEditor/Private/PropertyNode.h
--- a/Engine/Source/Editor/PropertyEditor/Private/PropertyNode.h
+++ b/Engine/Source/Editor/PropertyEditor/Private/PropertyNode.h
@@ -706,7 +706,7 @@ protected:
 	*/
 	FString GetDefaultValueAsString();
 
-	/*
+	/**
 	 * Helper function to obtain the display name for an enum property
 	 * @param InEnum		The enum whose metadata to pull from
 	 * @param DisplayName	The name of the enum value to adjust
diff --git a/Engine/Source/Editor/PropertyEditor/Public/PropertyHandle.h b/Engine/Source/Editor/PropertyEditor/Public/PropertyHandle.h
--- a/Engine/Source/Editor/PropertyEditor/Public/PropertyHandle.h
+++ b/Engine/Source/Editor/PropertyEditor/Public/PropertyHandle.h
@@ -48,7 +48,7 @@ public:
 	 */
 	virtual UProperty* GetProperty() const = 0;
 
-	/*
+	/**
 	 * Gets the property we should use to read meta-data
 	 */
 	virtual UProperty* GetMetaDataProperty() const = 0;
diff --git a/Engine/Source/Editor/StaticMeshEditor/Private/StaticMeshEditor.h b/Engine/Source/Editor/StaticMeshEditor/Private/StaticMeshEditor.h
--- a/Engine/Source/Editor/StaticMeshEditor/Private/StaticMeshEditor.h
+++ b/Engine/Source/Editor/StaticMeshEditor/Private/StaticMeshEditor.h
@@ -162,7 +162,7 @@ private:
 	/** Event for importing the light map channel to a static mesh from an intermediate file from Max/Maya */
 	void OnImportLightmapMesh( bool IsFBX );
 
-	/*
+	/**
 	* Quick and dirty way of creating box vertices from a box collision representation
 	* Grossly inefficient, but not time critical
 	* @param BoxElem - Box collision to get the vertices for
@@ -307,4 +307,4 @@ private:
 	static const FName SocketManagerTabId;
 	static const FName CollisionTabId;
 	static const FName GenerateUniqueUVsTabId;
-};
\ No newline at end of file
+};
diff --git a/Engine/Source/Editor/UnrealEd/Private/FbxImporter.h b/Engine/Source/Editor/UnrealEd/Private/FbxImporter.h
--- a/Engine/Source/Editor/UnrealEd/Private/FbxImporter.h
+++ b/Engine/Source/Editor/UnrealEd/Private/FbxImporter.h
@@ -1001,7 +1001,7 @@ private:
 	bool ImportCurve(const FbxAnimCurve* FbxCurve, FFloatCurve * Curve, const FbxTimeSpan &AnimTimeSpan, const float ValueScale = 1.f) const;
 
 
-	/*
+	/**
 	 * Import FbxCurve to anim sequence
 	 */
 	bool ImportCurveToAnimSequence(class UAnimSequence * TargetSequence, const FString & CurveName, const FbxAnimCurve * FbxCurve, int32 CurveFlags,const FbxTimeSpan AnimTimeSpan, const float ValueScale = 1.f) const;
diff --git a/Engine/Source/Editor/UnrealEd/Private/PlayLevel.cpp b/Engine/Source/Editor/UnrealEd/Private/PlayLevel.cpp
--- a/Engine/Source/Editor/UnrealEd/Private/PlayLevel.cpp
+++ b/Engine/Source/Editor/UnrealEd/Private/PlayLevel.cpp
@@ -679,7 +679,7 @@ void GetWindowSizeForInstanceType(FIntPoint &WindowSize, const ULevelEditorPlayS
 	}
 }
 
-/* 
+/** 
  * Generate the command line for pie instance. Window position, size etc. 
  *
  * @param	WinX			Window X position. This will contain the X position to use for the next window. (Not changed for dedicated server window).
diff --git a/Engine/Source/Editor/UnrealEd/Private/Tests/BlueprintAutomationTests.cpp b/Engine/Source/Editor/UnrealEd/Private/Tests/BlueprintAutomationTests.cpp
--- a/Engine/Source/Editor/UnrealEd/Private/Tests/BlueprintAutomationTests.cpp
+++ b/Engine/Source/Editor/UnrealEd/Private/Tests/BlueprintAutomationTests.cpp
@@ -1036,7 +1036,7 @@ bool FDelayLatentCommand::Update()
 }
 
 /************************************************************************/
-/*
+/**
  * Uses test maps in Engine and/or game content folder which are populated with a few blueprint instances
  * See InstanceTestMaps entries in the [Automation.Blueprint] config sections
  * For all blueprint instances in the map:
diff --git a/Engine/Source/Editor/UnrealEd/Public/EditorAnalytics.h b/Engine/Source/Editor/UnrealEd/Public/EditorAnalytics.h
--- a/Engine/Source/Editor/UnrealEd/Public/EditorAnalytics.h
+++ b/Engine/Source/Editor/UnrealEd/Public/EditorAnalytics.h
@@ -55,7 +55,7 @@ namespace EAnalyticsErrorCodes
 class FEditorAnalytics
 {
 public:
-	/*
+	/**
 	* Reports an event to the analytics system if it is enabled
 	*
 	* @param EventName - name of the event
@@ -63,7 +63,7 @@ public:
 	*/
 	UNREALED_API static void ReportEvent(FString EventName, FString PlatformName, bool bHasCode);
 
-	/*
+	/**
 	* Reports an event to the analytics system if it is enabled with some extra parameters
 	*
 	* @param EventName - name of the event
@@ -72,7 +72,7 @@ public:
 	*/
 	UNREALED_API static void ReportEvent(FString EventName, FString PlatformName, bool bHasCode, TArray<FAnalyticsEventAttribute>& ExtraParams);
 
-	/*
+	/**
 	* Reports an event to the analytics system if it is enabled with some extra parameters with an error code
 	*
 	* @param EventName - name of the event
@@ -82,7 +82,7 @@ public:
 	*/
 	UNREALED_API static void ReportEvent(FString EventName, FString PlatformName, bool bHasCode, int32 ErrorCode, TArray<FAnalyticsEventAttribute>& ExtraParams);
 
-	/*
+	/**
 	* Reports an event to the analytics system of a build requirements failure
 	*
 	* @param EventName - name of the event
@@ -94,4 +94,4 @@ public:
 private:
 	/* Translates the error code to something human readable */
 	static FString TranslateErrorCode(int32 ErrorCode);
-};
\ No newline at end of file
+};
diff --git a/Engine/Source/Editor/UnrealEd/Public/Kismet2/BlueprintEditorUtils.h b/Engine/Source/Editor/UnrealEd/Public/Kismet2/BlueprintEditorUtils.h
--- a/Engine/Source/Editor/UnrealEd/Public/Kismet2/BlueprintEditorUtils.h
+++ b/Engine/Source/Editor/UnrealEd/Public/Kismet2/BlueprintEditorUtils.h
@@ -238,7 +238,7 @@ public:
 	 */
 	static void RemoveGraph( UBlueprint* Blueprint, class UEdGraph* GraphToRemove, EGraphRemoveFlags::Type Flags = EGraphRemoveFlags::Default );
 
-	/*
+	/**
 	 * Tries to rename the supplied graph.
 	 * Cleans up function entry node if one exists and marks objects for modification
 	 *
@@ -247,7 +247,7 @@ public:
 	 */
 	static void RenameGraph(class UEdGraph* Graph, const FString& NewName );
 
-	/*
+	/**
 	 * Renames the graph of the supplied node with a valid name based off of the suggestion.
 	 *
 	 * @param GraphNode			The node of the graph to rename.
diff --git a/Engine/Source/Editor/UnrealEd/Public/SComponentClassCombo.h b/Engine/Source/Editor/UnrealEd/Public/SComponentClassCombo.h
--- a/Engine/Source/Editor/UnrealEd/Public/SComponentClassCombo.h
+++ b/Engine/Source/Editor/UnrealEd/Public/SComponentClassCombo.h
@@ -67,7 +67,7 @@ public:
 	/** Clear the current combo list selection */
 	void ClearSelection();
 
-	/*
+	/**
 	 * Updates the filtered list of component names.
 	 * @param InSearchText The search text from the search control.
 	 */
@@ -75,7 +75,7 @@ public:
 
 	FText GetCurrentSearchString() const;
 
-	/*
+	/**
 	 * Called when the user changes the text in the search box.
 	 * @param InSearchText The new search string.
 	 */
@@ -115,4 +115,4 @@ private:
 
 	/** Cached selection index used to skip over unselectable items */
 	int32 PrevSelectedIndex;
-};
\ No newline at end of file
+};
diff --git a/Engine/Source/Editor/UnrealEd/Public/SEditorViewport.h b/Engine/Source/Editor/UnrealEd/Public/SEditorViewport.h
--- a/Engine/Source/Editor/UnrealEd/Public/SEditorViewport.h
+++ b/Engine/Source/Editor/UnrealEd/Public/SEditorViewport.h
@@ -33,7 +33,7 @@ public:
 
 	TSharedPtr<FEditorViewportClient> GetViewportClient() const { return Client; }
 
-	/*
+	/**
 	 * Controls the visibility of the widget transform toolbar, if there is an associated toolbar
 	 */
 	virtual EVisibility GetTransformToolbarVisibility() const;
diff --git a/Engine/Source/Runtime/Core/Private/HAL/ConsoleManager.cpp b/Engine/Source/Runtime/Core/Private/HAL/ConsoleManager.cpp
--- a/Engine/Source/Runtime/Core/Private/HAL/ConsoleManager.cpp
+++ b/Engine/Source/Runtime/Core/Private/HAL/ConsoleManager.cpp
@@ -16,7 +16,7 @@ static inline bool IsWhiteSpace(TCHAR Value) { return Value == TCHAR(' '); }
 class FConsoleVariableBase : public IConsoleVariable
 {
 public:
-	/*
+	/**
 	 * Constructor
 	 * @param InHelp must not be 0, must not be empty
 	 */
@@ -114,7 +114,7 @@ protected: // -----------------------------------------
 class FConsoleCommandBase : public IConsoleCommand
 {
 public:
-	/*
+	/**
 	 * Constructor
 	 * @param InHelp must not be 0, must not be empty
 	 */
@@ -1873,4 +1873,4 @@ static TAutoConsoleVariable<int32> CVarGBuffer(
 	TEXT("r.GBuffer"),
 	1,
 	TEXT("0=Do not use GBuffer (fast minimal rendering path), 1=Use GBuffer [default]."),
-	ECVF_RenderThreadSafe);
\ No newline at end of file
+	ECVF_RenderThreadSafe);
diff --git a/Engine/Source/Runtime/Core/Public/Async/AsyncWork.h b/Engine/Source/Runtime/Core/Public/Async/AsyncWork.h
--- a/Engine/Source/Runtime/Core/Public/Async/AsyncWork.h
+++ b/Engine/Source/Runtime/Core/Public/Async/AsyncWork.h
@@ -70,7 +70,7 @@ class FAutoDeleteAsyncTask : private FQueuedWork
 		}
 	}
 
-	/* 
+	/** 
 	* Tells the user job to do the work, sometimes called synchronously, sometimes from the thread pool. Calls the event tracker.
 	**/
 	void DoWork()
@@ -79,7 +79,7 @@ class FAutoDeleteAsyncTask : private FQueuedWork
 		delete this;
 	}
 
-	/* 
+	/** 
 	* Always called from the thread pool. Just passes off to DoWork
 	**/
 	virtual void DoThreadedWork()
@@ -165,7 +165,7 @@ public:
 	}
 
 
-	/* 
+	/** 
 	* Run this task on this thread, now. Will end up destroying myself, so it is not safe to use this object after this call.
 	**/
 	void StartSynchronousTask()
@@ -173,7 +173,7 @@ public:
 		Start(true);
 	}
 
-	/* 
+	/** 
 	* Run this task on the lo priority thread pool. It is not safe to use this object after this call.
 	**/
 	void StartBackgroundTask()
@@ -287,7 +287,7 @@ class FAsyncTask : private FQueuedWork
 		}
 	}
 
-	/* 
+	/** 
 	* Tells the user job to do the work, sometimes called synchronously, sometimes from the thread pool. Calls the event tracker.
 	**/
 	void DoWork()
@@ -297,7 +297,7 @@ class FAsyncTask : private FQueuedWork
 		WorkNotFinishedCounter.Decrement();
 	}
 
-	/* 
+	/** 
 	* Triggers the work completion event, only called from a pool thread
 	**/
 	void FinishThreadedWork()
@@ -309,7 +309,7 @@ class FAsyncTask : private FQueuedWork
 		}
 	}
 
-	/* 
+	/** 
 	* Performs the work, this is only called from a pool thread.
 	**/
 	virtual void DoThreadedWork()
@@ -337,7 +337,7 @@ class FAsyncTask : private FQueuedWork
 		FinishThreadedWork();
 	}
 
-	/* 
+	/** 
 	* Internal call to assert that we are idle
 	**/
 	void CheckIdle() const
@@ -346,7 +346,7 @@ class FAsyncTask : private FQueuedWork
 		check(!QueuedPool);
 	}
 
-	/* 
+	/** 
 	* Internal call to synchronize completion between threads, never called from a pool thread
 	**/
 	void SyncCompletion()
@@ -361,7 +361,7 @@ class FAsyncTask : private FQueuedWork
 		CheckIdle();
 	}
 
-	/* 
+	/** 
 	* Internal call to intialize internal variables
 	**/
 	void Init()
@@ -466,7 +466,7 @@ public:
 		return Task;
 	}
 
-	/* 
+	/** 
 	* Run this task on this thread
 	* @param bDoNow if true then do the job now instead of at EnsureCompletion
 	**/
@@ -475,7 +475,7 @@ public:
 		Start(true);
 	}
 
-	/* 
+	/** 
 	* Queue this task for processing by the background thread pool
 	**/
 	void StartBackgroundTask()
@@ -483,7 +483,7 @@ public:
 		Start(false);
 	}
 
-	/* 
+	/** 
 	* Wait until the job is complete
 	* @param bDoWorkOnThisThreadIfNotStarted if true and the work has not been started, retract the async task and do it now on this thread
 	**/
diff --git a/Engine/Source/Runtime/Core/Public/HAL/ThreadingBase.h b/Engine/Source/Runtime/Core/Public/HAL/ThreadingBase.h
--- a/Engine/Source/Runtime/Core/Public/HAL/ThreadingBase.h
+++ b/Engine/Source/Runtime/Core/Public/HAL/ThreadingBase.h
@@ -629,7 +629,7 @@ public:
 };
 
 
-/*
+/**
  *  Global thread pool for shared async operations
  */
 extern CORE_API FQueuedThreadPool* GThreadPool;
@@ -964,4 +964,4 @@ public:
 };
 
 #define DECLARE_THREAD_SINGLETON(ClassType) \
-	static auto& GForceInitAtBoot_ThreadSingletonInitializer_##ClassType = ClassType::Get();
\ No newline at end of file
+	static auto& GForceInitAtBoot_ThreadSingletonInitializer_##ClassType = ClassType::Get();
diff --git a/Engine/Source/Runtime/Core/Public/HTML5/HTML5DebugLogging.h b/Engine/Source/Runtime/Core/Public/HTML5/HTML5DebugLogging.h
--- a/Engine/Source/Runtime/Core/Public/HTML5/HTML5DebugLogging.h
+++ b/Engine/Source/Runtime/Core/Public/HTML5/HTML5DebugLogging.h
@@ -29,7 +29,7 @@ extern "C" {
  * stack will be omitted. */
 #define EM_LOG_NO_PATHS 64
 
-/*
+/**
  * Prints out a message to the console, optionally with the 
  * callstack information.
  * @param flags A binary OR of items from the list of EM_LOG_xxx 
@@ -39,7 +39,7 @@ extern "C" {
  */
 void emscripten_log(int flags, ...);
 
-/*
+/**
  * Programmatically obtains the current callstack.
  * @param flags    A binary OR of items from the list of EM_LOG_xxx 
  *                    flags that specify printing options. The 
diff --git a/Engine/Source/Runtime/Core/Public/Math/Quat.h b/Engine/Source/Runtime/Core/Public/Math/Quat.h
--- a/Engine/Source/Runtime/Core/Public/Math/Quat.h
+++ b/Engine/Source/Runtime/Core/Public/Math/Quat.h
@@ -35,7 +35,7 @@ public:
 
 public:
 
-	/*
+	/**
 	 * Default constructor (no initialization).
 	 */
 	FORCEINLINE FQuat( ) { }
@@ -154,14 +154,14 @@ public:
 	 */
 	FORCEINLINE FQuat operator*=(const FQuat& Q);
 
-	/*
+	/**
 	 * This transforms vector, not dot product
 	 *
 	 * @param V
 	 */
 	FVector operator*( const FVector& V ) const;
 
-	/* 
+	/** 
 	 * This matrix conversion came from
 	 * http://www.m-hikari.com/ija/ija-password-2008/ija-password17-20-2008/aristidouIJA17-20-2008.pdf
 	 * used for non-uniform scaling transform.
diff --git a/Engine/Source/Runtime/Core/Public/Math/RandomStream.h b/Engine/Source/Runtime/Core/Public/Math/RandomStream.h
--- a/Engine/Source/Runtime/Core/Public/Math/RandomStream.h
+++ b/Engine/Source/Runtime/Core/Public/Math/RandomStream.h
@@ -187,7 +187,7 @@ public:
 		return GetUnitVector();
 	}
 
-	/*
+	/**
 	 * Returns a random unit vector, uniformly distributed, within the specified cone.
 	 *
 	 * @param Dir The center direction of the cone
@@ -231,7 +231,7 @@ public:
 		}
 	}
 
-	/*
+	/**
 	 * Returns a random unit vector, uniformly distributed, within the specified cone.
 	 *
 	 * @param Dir The center direction of the cone
diff --git a/Engine/Source/Runtime/Core/Public/Math/UnrealMathUtility.h b/Engine/Source/Runtime/Core/Public/Math/UnrealMathUtility.h
--- a/Engine/Source/Runtime/Core/Public/Math/UnrealMathUtility.h
+++ b/Engine/Source/Runtime/Core/Public/Math/UnrealMathUtility.h
@@ -112,7 +112,7 @@ struct FMath : public FPlatformMath
 	/** Return a uniformly distributed random unit length vector = point on the unit sphere surface. */
 	static FVector VRand();
 	
-	/*
+	/**
 	 * Returns a random unit vector, uniformly distributed, within the specified cone
 	 * ConeHalfAngleRad is the half-angle of cone, in radians.  Returns a normalized vector. 
 	 */
@@ -870,7 +870,7 @@ struct FMath : public FPlatformMath
 	 */
 	static CORE_API bool Eval( FString Str, float& OutValue );
 
-	/*
+	/**
 	 * Computes the barycentric coordinates for a given point in a triangle - simpler version
 	 *
 	 * @param	Point			point to convert to barycentric coordinates (in plane of ABC)
@@ -881,7 +881,7 @@ struct FMath : public FPlatformMath
 	 */
 	static CORE_API FVector GetBaryCentric2D(const FVector& Point, const FVector& A, const FVector& B, const FVector& C);
 
-	/*
+	/**
 	 * Computes the barycentric coordinates for a given point in a triangle
 	 *
 	 * @param	Point			point to convert to barycentric coordinates (in plane of ABC)
@@ -892,7 +892,7 @@ struct FMath : public FPlatformMath
 	 */
 	static CORE_API FVector ComputeBaryCentric2D(const FVector& Point, const FVector& A, const FVector& B, const FVector& C);
 
-	/*
+	/**
 	 * Computes the barycentric coordinates for a given point on a tetrahedron (3D)
 	 *
 	 * @param	Point			point to convert to barycentric coordinates
diff --git a/Engine/Source/Runtime/Core/Public/Misc/RemoteConfigIni.h b/Engine/Source/Runtime/Core/Public/Misc/RemoteConfigIni.h
--- a/Engine/Source/Runtime/Core/Public/Misc/RemoteConfigIni.h
+++ b/Engine/Source/Runtime/Core/Public/Misc/RemoteConfigIni.h
@@ -2,7 +2,7 @@
 
 #pragma once
 
-/*
+/**
  * Stores info relating to remote config files
  */
 struct FRemoteConfigAsyncIOInfo
@@ -135,7 +135,7 @@ private:
 };
 
 
-/*
+/**
  * Manages remote config files
  */
 class FRemoteConfig
diff --git a/Engine/Source/Runtime/Core/Public/Templates/UnrealTemplate.h b/Engine/Source/Runtime/Core/Public/Templates/UnrealTemplate.h
--- a/Engine/Source/Runtime/Core/Public/Templates/UnrealTemplate.h
+++ b/Engine/Source/Runtime/Core/Public/Templates/UnrealTemplate.h
@@ -374,14 +374,14 @@ struct TAreTypesEqual<A,A>
 #define TEMPLATE_PARAMETERS2(X,Y) X,Y
 
 
-/*
+/**
  * TRemoveReference<type> will remove any references from a type.
  */
 template <typename T> struct TRemoveReference      { typedef T Type; };
 template <typename T> struct TRemoveReference<T& > { typedef T Type; };
 template <typename T> struct TRemoveReference<T&&> { typedef T Type; };
 
-/*
+/**
  * MoveTemp will cast a reference to an rvalue reference.
  * This is UE's equivalent of std::move.
  */
@@ -432,7 +432,7 @@ FORCEINLINE T StaticCast(ArgType&& Arg)
 	return static_cast<T>(Arg);
 }
 
-/*
+/**
  * TRemoveCV<type> will remove any const/volatile qualifiers from a type.
  * (based on std::remove_cv<>
  * note: won't remove the const from "const int*", as the pointer is not const
diff --git a/Engine/Source/Runtime/CoreUObject/Private/UObject/Class.cpp b/Engine/Source/Runtime/CoreUObject/Private/UObject/Class.cpp
--- a/Engine/Source/Runtime/CoreUObject/Private/UObject/Class.cpp
+++ b/Engine/Source/Runtime/CoreUObject/Private/UObject/Class.cpp
@@ -1,4 +1,4 @@
-﻿// Copyright 1998-2014 Epic Games, Inc. All Rights Reserved.
+// Copyright 1998-2014 Epic Games, Inc. All Rights Reserved.
 
 /*=============================================================================
 	UnClass.cpp: Object class implementation.
@@ -95,7 +95,7 @@ FString FPropertySpecifier::ConvertToString() const
 
 //////////////////////////////////////////////////////////////////////////
 
-/*
+/**
  * Shared function called from the various InitializePrivateStaticClass functions generated my the IMPLEMENT_CLASS macro.
  */
 COREUOBJECT_API void InitializePrivateStaticClass(
diff --git a/Engine/Source/Runtime/CoreUObject/Public/UObject/ScriptInterface.h b/Engine/Source/Runtime/CoreUObject/Public/UObject/ScriptInterface.h
--- a/Engine/Source/Runtime/CoreUObject/Public/UObject/ScriptInterface.h
+++ b/Engine/Source/Runtime/CoreUObject/Public/UObject/ScriptInterface.h
@@ -25,7 +25,7 @@ private:
 	 */
 	void*		InterfacePointer;
 
-	/*
+	/**
 	 * Serialize ScriptInterface
 	 */
 	FArchive& Serialize(FArchive& Ar, class UClass* InterfaceType);
diff --git a/Engine/Source/Runtime/CoreUObject/Public/UObject/UObjectGlobals.h b/Engine/Source/Runtime/CoreUObject/Public/UObject/UObjectGlobals.h
--- a/Engine/Source/Runtime/CoreUObject/Public/UObject/UObjectGlobals.h
+++ b/Engine/Source/Runtime/CoreUObject/Public/UObject/UObjectGlobals.h
@@ -230,7 +230,7 @@ COREUOBJECT_API UObject* StaticDuplicateObjectEx( struct FObjectDuplicationParam
  * Performs UObject system pre-initialization. Depracated, do not use.
  */
 COREUOBJECT_API void PreInitUObject();
-/* 
+/** 
  *   Iterate over all objects considered part of the root to setup GC optimizations
  */
 COREUOBJECT_API void MarkObjectsToDisregardForGC();
diff --git a/Engine/Source/Runtime/Engine/Classes/Animation/AnimSequence.h b/Engine/Source/Runtime/Engine/Classes/Animation/AnimSequence.h
--- a/Engine/Source/Runtime/Engine/Classes/Animation/AnimSequence.h
+++ b/Engine/Source/Runtime/Engine/Classes/Animation/AnimSequence.h
@@ -355,7 +355,7 @@ class UAnimSequence : public UAnimSequenceBase
 	TArray<struct FScaleTrack> ScaleData;
 
 
-	/*
+	/**
 	 * Curve data - no compression yet                                                                       
 	 */
 	UPROPERTY()
@@ -726,7 +726,7 @@ private:
 	void RemoveNaNTracks();
 #endif
 
-	/*
+	/**
 	 * Utility function that helps to remove track, you can't just remove RawAnimationData
 	 */
 	void RemoveTrack(int32 TrackIndex);
diff --git a/Engine/Source/Runtime/Engine/Classes/Animation/AnimSequenceBase.h b/Engine/Source/Runtime/Engine/Classes/Animation/AnimSequenceBase.h
--- a/Engine/Source/Runtime/Engine/Classes/Animation/AnimSequenceBase.h
+++ b/Engine/Source/Runtime/Engine/Classes/Animation/AnimSequenceBase.h
@@ -28,7 +28,7 @@ namespace EAnimEventTriggerOffsets
 
 ENGINE_API float GetTriggerTimeOffsetForType(EAnimEventTriggerOffsets::Type OffsetType);
 
-/*
+/**
  * Triggers an animation notify.  Each AnimNotifyEvent contains an AnimNotify object
  * which has its Notify method called and passed to the animation.
  */
@@ -289,7 +289,7 @@ struct FRawCurveTracks
 	 * bVectorInterpCurve == true, then it will create FVectorCuve, otherwise, FFloatCurve
 	 */
 	ENGINE_API bool AddCurveData(USkeleton::AnimCurveUID Uid, int32 CurveFlags = ACF_DefaultCurve);
-	/*
+	/**
 	 * Delete curve data 
 	 */
 	ENGINE_API bool DeleteCurveData(USkeleton::AnimCurveUID Uid);
diff --git a/Engine/Source/Runtime/Engine/Classes/Animation/BlendSpaceBase.h b/Engine/Source/Runtime/Engine/Classes/Animation/BlendSpaceBase.h
--- a/Engine/Source/Runtime/Engine/Classes/Animation/BlendSpaceBase.h
+++ b/Engine/Source/Runtime/Engine/Classes/Animation/BlendSpaceBase.h
@@ -209,7 +209,7 @@ protected:
 	UPROPERTY(EditAnywhere, Category=InputInterpolation)
 	FInterpolationParameter	InterpolationParam[3];
 
-	/* 
+	/** 
 	 * Target weight interpolation. When target samples are set, how fast you'd like to get to target. Improve target blending. 
 	 * i.e. for locomotion, if you interpolate input, when you move from left to right rapidly, you'll interpolate through forward, but if you use target weight interpolation, 
 	 * you'll skip forward, but interpolate between left to right 
@@ -227,7 +227,7 @@ protected:
 	
 public:
 
-	/* 
+	/** 
 	 * When you use blend per bone, allows rotation to blend in mesh space. This only works if this does not contain additive animation samples
 	 * This is more performance intensive
 	 */
@@ -248,7 +248,7 @@ public:
 	UPROPERTY(transient)
 	float AnimLength;
 
-	/* 
+	/** 
 	 * Define target weight interpolation per bone. This will blend in different speed per each bone setting 
 	 */
 	UPROPERTY(EditAnywhere, Category=SampleInterpolation)
@@ -415,4 +415,4 @@ public:
 	/** Grid samples, indexing scheme imposed by subclass **/
 	UPROPERTY()
 	TArray<struct FEditorElement> GridSamples;
-};
\ No newline at end of file
+};
diff --git a/Engine/Source/Runtime/Engine/Classes/EdGraph/EdGraphPin.h b/Engine/Source/Runtime/Engine/Classes/EdGraph/EdGraphPin.h
--- a/Engine/Source/Runtime/Engine/Classes/EdGraph/EdGraphPin.h
+++ b/Engine/Source/Runtime/Engine/Classes/EdGraph/EdGraphPin.h
@@ -239,7 +239,7 @@ public:
 	/** Break a link to the specified pin (if present) */
 	ENGINE_API void BreakLinkTo(UEdGraphPin* ToPin);
 
-	/*
+	/**
 	 * Generate a string detailing the link this pin has to another pin.
 	 * 
 	 * @Param	InFunctionName	String with function name requesting the info
diff --git a/Engine/Source/Runtime/Engine/Classes/EditorFramework/AssetImportData.h b/Engine/Source/Runtime/Engine/Classes/EditorFramework/AssetImportData.h
--- a/Engine/Source/Runtime/Engine/Classes/EditorFramework/AssetImportData.h
+++ b/Engine/Source/Runtime/Engine/Classes/EditorFramework/AssetImportData.h
@@ -3,7 +3,7 @@
 #pragma once
 #include "AssetImportData.generated.h"
 
-/*
+/**
  * A class to hold import options and data for assets. This class has many children in editor code.
  */
 UCLASS(MinimalAPI)
diff --git a/Engine/Source/Runtime/Engine/Classes/EditorFramework/ThumbnailInfo.h b/Engine/Source/Runtime/Engine/Classes/EditorFramework/ThumbnailInfo.h
--- a/Engine/Source/Runtime/Engine/Classes/EditorFramework/ThumbnailInfo.h
+++ b/Engine/Source/Runtime/Engine/Classes/EditorFramework/ThumbnailInfo.h
@@ -1,6 +1,6 @@
 // Copyright 1998-2014 Epic Games, Inc. All Rights Reserved.
 
-/*
+/**
  * A base class for the helper object that holds thumbnail information an asset.
  */
 
diff --git a/Engine/Source/Runtime/Engine/Classes/Engine/Canvas.h b/Engine/Source/Runtime/Engine/Classes/Engine/Canvas.h
--- a/Engine/Source/Runtime/Engine/Classes/Engine/Canvas.h
+++ b/Engine/Source/Runtime/Engine/Classes/Engine/Canvas.h
@@ -423,14 +423,14 @@ public:
 	 */
 	virtual void DrawDebugGraph(const FString& Title, float ValueX, float ValueY, float UL_X, float UL_Y, float W, float H, FVector2D RangeX, FVector2D RangeY);
 
-	/* 
+	/** 
 	 * Draw a CanvasItem
 	 *
 	 * @param Item			Item to draw
 	 */
 	void DrawItem( FCanvasItem& Item );
 
-	/* 
+	/** 
 	 * Draw a CanvasItem at the given coordinates
 	 *
 	 * @param Item			Item to draw
@@ -438,7 +438,7 @@ public:
 	 */
 	void DrawItem( FCanvasItem& Item, const FVector2D& InPosition );
 	
-	/* 
+	/** 
 	 * Draw a CanvasItem at the given coordinates
 	 *
 	 * @param Item			Item to draw
@@ -602,4 +602,4 @@ public:
 	 */
 	UFUNCTION(BlueprintCallable, Category=Canvas, meta=(FriendlyName="Clipped Text Size"))
 	FVector2D K2_TextSize(UFont* RenderFont, const FString& RenderText, FVector2D Scale=FVector2D::UnitVector);
-};
\ No newline at end of file
+};
diff --git a/Engine/Source/Runtime/Engine/Classes/Engine/Engine.h b/Engine/Source/Runtime/Engine/Classes/Engine/Engine.h
--- a/Engine/Source/Runtime/Engine/Classes/Engine/Engine.h
+++ b/Engine/Source/Runtime/Engine/Classes/Engine/Engine.h
@@ -1423,7 +1423,7 @@ public:
 	void RegisterEndStreamingPauseRenderingDelegate( FEndStreamingPauseDelegate* InDelegate );
 	FEndStreamingPauseDelegate* EndStreamingPauseDelegate;
 
-	/* 
+	/** 
 	 * Error message event relating to server travel failures 
 	 * 
 	 * @param Type type of travel failure
@@ -1432,7 +1432,7 @@ public:
 	DECLARE_EVENT_ThreeParams(UEngine, FOnTravelFailure, UWorld*, ETravelFailure::Type, const FString&);
 	FOnTravelFailure TravelFailureEvent;
 
-	/* 
+	/** 
 	 * Error message event relating to network failures 
 	 * 
 	 * @param Type type of network failure
@@ -1959,7 +1959,7 @@ private:
 
 protected:
 
-	/*
+	/**
 	 * Handles freezing/unfreezing of rendering 
 	 * 
 	 * @param InWorld	World context
@@ -2002,7 +2002,7 @@ public:
 	 */
 	void AddTextureStreamingSlaveLoc(FVector InLoc, float BoostFactor, bool bOverrideLocation, float OverrideDuration);
 
-	/* 
+	/** 
 	 * Obtain a world object pointer from an object with has a world context.
 	 * This should be be overridden to cater for game specific object types that do not derive from the Actor class.
 	 *
diff --git a/Engine/Source/Runtime/Engine/Classes/Engine/GameViewportClient.h b/Engine/Source/Runtime/Engine/Classes/Engine/GameViewportClient.h
--- a/Engine/Source/Runtime/Engine/Classes/Engine/GameViewportClient.h
+++ b/Engine/Source/Runtime/Engine/Classes/Engine/GameViewportClient.h
@@ -569,7 +569,7 @@ public:
 	/** The platform-specific viewport frame which this viewport is contained by. */
 	FViewportFrame* ViewportFrame;
 
-	/*
+	/**
  	 * Controls supression of the blue transition text messages 
 	 * 
 	 * @param bSuppress	Pass true to supress messages
diff --git a/Engine/Source/Runtime/Engine/Classes/Engine/Level.h b/Engine/Source/Runtime/Engine/Classes/Engine/Level.h
--- a/Engine/Source/Runtime/Engine/Classes/Engine/Level.h
+++ b/Engine/Source/Runtime/Engine/Classes/Engine/Level.h
@@ -647,17 +647,17 @@ public:
 	/** @todo document */
 	ENGINE_API TArray<FVector> const* GetStaticNavigableGeometry() const { return &StaticNavigableGeometry;}
 
-	/* 
+	/** 
 	* Is this the persistent level 
 	*/
 	ENGINE_API bool IsPersistentLevel() const;
 
-	/* 
+	/** 
 	* Is this the current level in the world it is owned by
 	*/
 	ENGINE_API bool IsCurrentLevel() const;
 	
-	/* 
+	/** 
 	 * Shift level actors by specified offset
 	 * The offset vector will get subtracted from all actors positions and corresponding data structures
 	 *
diff --git a/Engine/Source/Runtime/Engine/Classes/Engine/UserDefinedEnum.h b/Engine/Source/Runtime/Engine/Classes/Engine/UserDefinedEnum.h
--- a/Engine/Source/Runtime/Engine/Classes/Engine/UserDefinedEnum.h
+++ b/Engine/Source/Runtime/Engine/Classes/Engine/UserDefinedEnum.h
@@ -24,7 +24,7 @@ class ENGINE_API UUserDefinedEnum : public UEnum
 	uint32 UniqueNameIndex;
 #endif //WITH_EDITORONLY_DATA
 
-	/*
+	/**
 	 * Names stored in "DisplayName" meta data. They are duplicated here, 
 	 * so functions like UKismetNodeHelperLibrary::GetEnumeratorUserFriendlyName can use them
 	 * outside the editor. (When meta data are not loaded).
diff --git a/Engine/Source/Runtime/Engine/Classes/Engine/World.h b/Engine/Source/Runtime/Engine/Classes/Engine/World.h
--- a/Engine/Source/Runtime/Engine/Classes/Engine/World.h
+++ b/Engine/Source/Runtime/Engine/Classes/Engine/World.h
@@ -1447,7 +1447,7 @@ public:
 
 	DEPRECATED(4.3, "GetBrush is deprecated use GetDefaultBrush instead.")
 	ABrush* GetBrush() const;
-	/* 
+	/** 
 	 * Returns the default brush for the persistent level.
 	 * This is usually the 'builder brush' for editor builds, undefined for non editor instances and may be NULL.
 	 */
diff --git a/Engine/Source/Runtime/Engine/Classes/Foliage/InstancedFoliageActor.h b/Engine/Source/Runtime/Engine/Classes/Foliage/InstancedFoliageActor.h
--- a/Engine/Source/Runtime/Engine/Classes/Foliage/InstancedFoliageActor.h
+++ b/Engine/Source/Runtime/Engine/Classes/Foliage/InstancedFoliageActor.h
@@ -93,7 +93,7 @@ public:
 	// Transforms Editor specific data which is stored in world space
 	ENGINE_API void ApplyLevelTransform(const FTransform& LevelTransform);
 
-	/*
+	/**
 	* Get the instanced foliage actor for the current streaming level.
 	*
 	* @param InCreationWorldIfNone			World to create the foliage instance in
diff --git a/Engine/Source/Runtime/Engine/Classes/GameFramework/Actor.h b/Engine/Source/Runtime/Engine/Classes/GameFramework/Actor.h
--- a/Engine/Source/Runtime/Engine/Classes/GameFramework/Actor.h
+++ b/Engine/Source/Runtime/Engine/Classes/GameFramework/Actor.h
@@ -938,7 +938,7 @@ public:
 	UFUNCTION(BlueprintImplementableEvent, meta=(BlueprintInternalUseOnly = "true", FriendlyName = "Construction Script"))
 	virtual void UserConstructionScript();
 
-	/*
+	/**
 	 * Destroy this actor. Returns true if destroyed, false if indestructible.
 	 * Destruction is latent. It occurs at the end of the tick.
 	 * @param	bNetForce				[opt] Ignored unless called during play.  Default is false.
@@ -1672,13 +1672,13 @@ public:
 	/** Rerun construction scripts, destroying all autogenerated components; will attempt to preserve the root component location. */
 	virtual void RerunConstructionScripts();
 
-	/* 
+	/** 
 	 * Debug helper to show the component hierarchy of this actor.
 	 * @param Info	Optional String to display at top of info
 	 */
 	void DebugShowComponentHierarchy( const TCHAR* Info, bool bShowPosition  = true);
 	
-	/* 
+	/** 
 	 * Debug helper for showing the component hierarchy of one component
 	 * @param Info	Optional String to display at top of info
 	 */
diff --git a/Engine/Source/Runtime/Engine/Classes/GameFramework/HUD.h b/Engine/Source/Runtime/Engine/Classes/GameFramework/HUD.h
--- a/Engine/Source/Runtime/Engine/Classes/GameFramework/HUD.h
+++ b/Engine/Source/Runtime/Engine/Classes/GameFramework/HUD.h
@@ -407,7 +407,7 @@ public:
 	void Deproject(float ScreenX, float ScreenY, FVector& WorldPosition, FVector& WorldDirection);
 
 
-	/*
+	/**
 	 * Returns the array of actors inside a selection rectangle, with a class filter.
 	 *
 	 * Sample usage:
diff --git a/Engine/Source/Runtime/Engine/Classes/Kismet/KismetMathLibrary.h b/Engine/Source/Runtime/Engine/Classes/Kismet/KismetMathLibrary.h
--- a/Engine/Source/Runtime/Engine/Classes/Kismet/KismetMathLibrary.h
+++ b/Engine/Source/Runtime/Engine/Classes/Kismet/KismetMathLibrary.h
@@ -510,7 +510,7 @@ class UKismetMathLibrary : public UBlueprintFunctionLibrary
 	UFUNCTION(BlueprintPure, Category = "Math|Random")
 	static FVector RandomPointInBoundingBox(const FVector& Origin, const FVector& BoxExtent);
 
-	/* 
+	/** 
 	 * Returns a random vector with length of 1, within the specified cone, with uniform random distribution. 
 	 * @param ConeDir	The base "center" direction of the cone.
 	 * @param ConeHalfAngle		The half-angle of the cone (from ConeDir to edge), in degrees.
@@ -1048,7 +1048,7 @@ class UKismetMathLibrary : public UBlueprintFunctionLibrary
 	// Geometry
 	//
 
-	/*  
+	/**  
 	 * Finds the minimum area rectangle that encloses all of the points in InVerts
 	 * Uses algorithm found in http://www.geometrictools.com/Documentation/MinimumAreaRectangle.pdf
 	 *	
diff --git a/Engine/Source/Runtime/Engine/Classes/Particles/ParticleEmitter.h b/Engine/Source/Runtime/Engine/Classes/Particles/ParticleEmitter.h
--- a/Engine/Source/Runtime/Engine/Classes/Particles/ParticleEmitter.h
+++ b/Engine/Source/Runtime/Engine/Classes/Particles/ParticleEmitter.h
@@ -304,7 +304,7 @@ class UParticleEmitter : public UObject
 	 */
 	void Build();
 
-	/*
+	/**
 	 *   Calculate spawn rate multiplier based on global effects quality level and emitter's quality scale
  	 */
 	FORCEINLINE float GetQualityLevelSpawnRateMult()
diff --git a/Engine/Source/Runtime/Engine/Classes/Particles/ParticleSystem.h b/Engine/Source/Runtime/Engine/Classes/Particles/ParticleSystem.h
--- a/Engine/Source/Runtime/Engine/Classes/Particles/ParticleSystem.h
+++ b/Engine/Source/Runtime/Engine/Classes/Particles/ParticleSystem.h
@@ -436,7 +436,7 @@ public:
 	 */
 	virtual bool SetLODDistance(int32 LODLevelIndex, float InDistance);
 
-	/*
+	/**
 	 * Builds all emitters in the particle system.
 	 */
 	ENGINE_API void BuildEmitters();
diff --git a/Engine/Source/Runtime/Engine/Classes/Sound/SoundGroups.h b/Engine/Source/Runtime/Engine/Classes/Sound/SoundGroups.h
--- a/Engine/Source/Runtime/Engine/Classes/Sound/SoundGroups.h
+++ b/Engine/Source/Runtime/Engine/Classes/Sound/SoundGroups.h
@@ -51,7 +51,7 @@ struct FSoundGroup
 	UPROPERTY(config)
 	uint32 bAlwaysDecompressOnLoad:1;
 
-	/*
+	/**
 	 * Sound duration in seconds below which sounds are entirely expanded to PCM at load time
 	 * Disregarded if bAlwaysDecompressOnLoad is true
 	*/
@@ -86,4 +86,4 @@ private:
 	// Easy lookup mechanism for sound group settings
 	mutable TMap<ESoundGroup, FSoundGroup> SoundGroupMap;
 
-};
\ No newline at end of file
+};
diff --git a/Engine/Source/Runtime/Engine/Private/BatchedElements.cpp b/Engine/Source/Runtime/Engine/Private/BatchedElements.cpp
--- a/Engine/Source/Runtime/Engine/Private/BatchedElements.cpp
+++ b/Engine/Source/Runtime/Engine/Private/BatchedElements.cpp
@@ -367,7 +367,7 @@ float GBatchedElementAlphaRefVal = 128.f;
 /** Global smoothing width for rendering batched elements with distance field blend modes */
 float GBatchedElementSmoothWidth = 12;
 
-/*
+/**
  * Sets the appropriate vertex and pixel shader.
  */
 void FBatchedElements::PrepareShaders(
diff --git a/Engine/Source/Runtime/Engine/Private/Collision/PhysXCollision.cpp b/Engine/Source/Runtime/Engine/Private/Collision/PhysXCollision.cpp
--- a/Engine/Source/Runtime/Engine/Private/Collision/PhysXCollision.cpp
+++ b/Engine/Source/Runtime/Engine/Private/Collision/PhysXCollision.cpp
@@ -477,7 +477,7 @@ private:
 
 static float DebugLineLifetime = 2.f;
 
-/* 
+/** 
  * Type of query for object type or trace type
  * Trace queries correspond to trace functions with TravelChannel/ResponseParams
  * Object queries correspond to trace functions with Object types 
diff --git a/Engine/Source/Runtime/Engine/Private/ShadowMap.cpp b/Engine/Source/Runtime/Engine/Private/ShadowMap.cpp
--- a/Engine/Source/Runtime/Engine/Private/ShadowMap.cpp
+++ b/Engine/Source/Runtime/Engine/Private/ShadowMap.cpp
@@ -110,7 +110,7 @@ struct FShadowMapPendingTexture : FTextureLayout
 
 	bool AddElement(FShadowMapAllocation& Allocation, const bool bForceIntoThisTexture = false);
 
-	/*
+	/**
 	 * Begin encoding the textures
 	 *
 	 * @param	InWorld	World in which the textures exist
diff --git a/Engine/Source/Runtime/Engine/Private/UnrealEngine.cpp b/Engine/Source/Runtime/Engine/Private/UnrealEngine.cpp
--- a/Engine/Source/Runtime/Engine/Private/UnrealEngine.cpp
+++ b/Engine/Source/Runtime/Engine/Private/UnrealEngine.cpp
@@ -329,7 +329,7 @@ void SystemResolutionSinkCallback()
 	}
 }
 
-/*
+/**
  * if we need to update the sample states
 */
 void RefreshSamplerStatesCallback()
diff --git a/Engine/Source/Runtime/Engine/Public/Audio.h b/Engine/Source/Runtime/Engine/Public/Audio.h
--- a/Engine/Source/Runtime/Engine/Public/Audio.h
+++ b/Engine/Source/Runtime/Engine/Public/Audio.h
@@ -24,7 +24,7 @@ ENGINE_API DECLARE_LOG_CATEGORY_EXTERN(LogAudio, Warning, All);
  */
 #define INDEFINITELY_LOOPING_DURATION	10000.0f
 
-/*
+/**
  * Some defaults to help cross platform consistency
  */
 #define SPEAKER_COUNT					6
diff --git a/Engine/Source/Runtime/Engine/Public/BatchedElements.h b/Engine/Source/Runtime/Engine/Public/BatchedElements.h
--- a/Engine/Source/Runtime/Engine/Public/BatchedElements.h
+++ b/Engine/Source/Runtime/Engine/Public/BatchedElements.h
@@ -308,7 +308,7 @@ private:
 	/** bound shader state for color masked elements */
 	static FGlobalBoundShaderState ColorChannelMaskShaderState;
 
-	/*
+	/**
 	 * Sets the appropriate vertex and pixel shader.
 	 */
 	void PrepareShaders(
diff --git a/Engine/Source/Runtime/Engine/Public/CanvasItem.h b/Engine/Source/Runtime/Engine/Public/CanvasItem.h
--- a/Engine/Source/Runtime/Engine/Public/CanvasItem.h
+++ b/Engine/Source/Runtime/Engine/Public/CanvasItem.h
@@ -18,7 +18,7 @@ class FTexture;
 class FCanvasItem
 {
 public:
-	/* 
+	/** 
 	 * Basic render item.
 	 *
 	 * @param	InPosition		Draw position
@@ -32,7 +32,7 @@ public:
 
 	virtual void Draw( FCanvas* InCanvas ) = 0;
 
-	/* 
+	/** 
 	 * Draw this item (this will affect the items position for future draw calls that do no specify a position)
 	 *
 	 * @param	InCanvas		Canvas on which to draw
@@ -44,7 +44,7 @@ public:
 		Draw( InCanvas );
 	}
 
-	/* 
+	/** 
 	 * Draw this item (this will affect the items position for future draw calls that do no specify a position)
 	 *
 	 * @param	InCanvas		Canvas on which to draw
@@ -84,7 +84,7 @@ protected:
 class ENGINE_API FCanvasTileItem : public FCanvasItem
 {
 public:
-	/* 
+	/** 
 	 * Tile item using size from texture.
 	 *
 	 * @param	InPosition		Draw position
@@ -92,7 +92,7 @@ public:
 	 */
 	FCanvasTileItem( const FVector2D& InPosition, const FTexture* InTexture, const FLinearColor& InColor );
 
-	/* 
+	/** 
 	 * Tile item with texture using given size. 
 	 *
 	 * @param	InPosition		Draw position
@@ -101,7 +101,7 @@ public:
 	 */
 	FCanvasTileItem( const FVector2D& InPosition, const FTexture* InTexture, const FVector2D& InSize, const FLinearColor& InColor );
 
-	/* 
+	/** 
 	 * Tile item which uses the default white texture using given size. 
 	 *
 	 * @param	InPosition		Draw position
@@ -109,7 +109,7 @@ public:
 	 */
 	FCanvasTileItem( const FVector2D& InPosition, const FVector2D& InSize, const FLinearColor& InColor );
 
-	/* 
+	/** 
 	 * Tile item with texture using size from texture specific UVs. 
 	 *
 	 * @param	InPosition		Draw position
@@ -119,7 +119,7 @@ public:
 	 */
 	FCanvasTileItem( const FVector2D& InPosition, const FTexture* InTexture, const FVector2D& InUV0, const FVector2D& InUV1, const FLinearColor& InColor );
 
-	/* 
+	/** 
 	 * Tile item with texture using given size and specific UVs.
 	 *
 	 * @param	InPosition		Draw position
@@ -130,7 +130,7 @@ public:
 	 */
 	FCanvasTileItem( const FVector2D& InPosition, const FTexture* InTexture, const FVector2D& InSize, const FVector2D& InUV0, const FVector2D& InUV1, const FLinearColor& InColor );
 
-	/* 
+	/** 
 	 * Tile item with FMaterialRenderProxy using given size. 
 	 *
 	 * @param	InPosition				Draw position
@@ -139,7 +139,7 @@ public:
 	 */
 	FCanvasTileItem( const FVector2D& InPosition, const FMaterialRenderProxy* InMaterialRenderProxy, const FVector2D& InSize );
 
-	/* 
+	/** 
 	 * Tile item with FMaterialRenderProxy using given size and UVs.
 	 *
 	 * @param	InPosition				Draw position
@@ -150,7 +150,7 @@ public:
 	 */
 	FCanvasTileItem( const FVector2D& InPosition, const FMaterialRenderProxy* InMaterialRenderProxy, const FVector2D& InSize, const FVector2D& InUV0, const FVector2D& InUV1 );
 
-	/* 
+	/** 
 	 * Draw the item at the given coordinates.
 	 *
 	 * @param	InPosition		Draw position.
@@ -194,7 +194,7 @@ class ENGINE_API FCanvasBorderItem : public FCanvasItem
 {
 public:
 
-	/* 
+	/** 
 	 * 3x3 grid border with tiled frame and tiled interior. 
 	 *
 	 * @param	InPosition		    Draw position
@@ -228,7 +228,7 @@ public:
 		check( InBackgroundTexture );
 	}
 
-	/* 
+	/** 
 	 * Draw the item at the given coordinates.
 	 *
 	 * @param	InPosition		Draw position.
@@ -288,7 +288,7 @@ public:
 class ENGINE_API FCanvasTextItem : public FCanvasItem
 {
 public:
-	/* 	 
+	/** 	 
 	 * Text item
 	 *
 	 * @param	InPosition		Draw position
@@ -317,7 +317,7 @@ public:
 		BlendMode = SE_BLEND_Translucent;
 	};
 		
-	/* 
+	/** 
 	 * Set the shadow offset and color. 
 	 *
 	 * @param	InColor			Shadow color
@@ -329,7 +329,7 @@ public:
 		ShadowColor = InColor;		
 	}
 	
-	/* 
+	/** 
 	 * Disable the shadow
 	 */
 	void DisableShadow() 
@@ -337,7 +337,7 @@ public:
 		ShadowOffset = FVector2D::ZeroVector;
 	}
 
-	/* 
+	/** 
 	 * Draw the item at the given coordinates.
 	 *
 	 * @param	InCanvas		Canvas on which to draw
@@ -395,14 +395,14 @@ protected:
 	/* Background tile used to fixup 3d text issues. */
 	FCanvasTileItem	TileItem;
 
-	/* 
+	/** 
 	 * Internal string draw
 	 *
 	 * In a method to make it simpler to do effects like shadow, outline
 	 */
 	void DrawStringInternal( FCanvas* InCanvas, const FVector2D& DrawPos, const FLinearColor& DrawColor );
 
-	/* 
+	/** 
 	 * These are used bye the DrawStringInternal function. 
 	 */
 	/* String char length. */
@@ -442,7 +442,7 @@ public:
 		EndPos.Y = 0.0f;
 		EndPos.Z = 0.0f;
 	}
-	/* 
+	/** 
 	 * A Line. 
 	 * 
 	 * @param	InPosition		Start position
@@ -460,7 +460,7 @@ public:
 		EndPos.Z = 0.0f;
 	}
 	
-	/* 
+	/** 
 	 * A Line. 
 	 * 
 	 * @param	InPosition		Start position
@@ -474,14 +474,14 @@ public:
 		EndPos = InEndPos;
 	}
 
-	/* 
+	/** 
 	 * Draw line at the given coordinates.
 	 *
 	 * @param	InCanvas		Canvas on which to draw
 	 */
 	virtual void Draw( FCanvas* InCanvas ) override;
 	
-	/* 
+	/** 
 	 * Draw line at the given coordinates.
 	 *
 	 * @param	InCanvas		Canvas on which to draw
@@ -493,7 +493,7 @@ public:
 		Origin.Y = InPosition.Y;
 		Draw( InCanvas );
 	}
-	/* 
+	/** 
 	 * Draw line using the given coordinates.
 	 *
 	 * @param	InCanvas		Canvas on which to draw
@@ -508,7 +508,7 @@ public:
 		EndPos.Y = InEndPos.Y;
 		Draw( InCanvas );
 	}
-	/* 
+	/** 
 	 * Draw line at the given coordinates.
 	 *
 	 * @param	InCanvas		Canvas on which to draw
@@ -522,7 +522,7 @@ public:
 		Draw( InCanvas );
 	}
 		
-	/* 
+	/** 
 	 * Draw line at the given coordinates.
 	 *
 	 * @param	InCanvas		Canvas on which to draw
@@ -534,7 +534,7 @@ public:
 		Draw( InCanvas );
 	}
 
-	/* 
+	/** 
 	 * Draw line at the given coordinates.
 	 *
 	 * @param	InCanvas		Canvas on which to draw
@@ -550,7 +550,7 @@ public:
 		Draw( InCanvas );
 	}
 
-	/* 
+	/** 
 	 * Set the line end position.
 	 *
 	 * @param	InEndPos		End position of the line
@@ -598,7 +598,7 @@ private:
 class ENGINE_API FCanvasTriangleItem : public FCanvasItem
 {
 public:	
-	/* 	 
+	/** 	 
 	 * Triangle item (no texture)
 	 *
 	 * @param	InPointA		Point A of triangle
@@ -626,7 +626,7 @@ public:
 		TriangleList.Add( SingleTri );
 	};
 
-	/* 	 
+	/** 	 
 	 * Triangle item
 	 *
 	 * @param	InPointA		Point A of triangle
@@ -657,7 +657,7 @@ public:
 		TriangleList.Add( SingleTri );
 	};
 
-	/* 	 
+	/** 	 
 	 * Triangle item
 	 *
 	 * @param	InSingleTri		Triangle struct
@@ -671,7 +671,7 @@ public:
 		TriangleList.Add( InSingleTri );
 	};
 
-	/* 	 
+	/** 	 
 	 * Triangle item
 	 *
 	 * @param	InTriangleList	List of triangles
@@ -720,7 +720,7 @@ public:
 class ENGINE_API FCanvasNGonItem :  public FCanvasItem
 {
 public:	
-	/* 	 
+	/** 	 
 	 * NGon item Several texture tris with a common central point with a fixed radius. 
 	 *
 	 * @param	InPosition	List of triangles
@@ -740,7 +740,7 @@ public:
 		 SetupPosition( InPosition, InRadius );
 	 }
 
-	 /* 	 
+	 /** 	 
 	 * NGon item Several tris with a common central point with a fixed radius. 
 	 *
 	 * @param	InPosition	List of triangles
@@ -756,7 +756,7 @@ public:
 	 }
 	 virtual void Draw( FCanvas* InCanvas ) override;
 
-	 /* 	 
+	 /** 	 
 	 * Regenerates the tri list for the object with a new central point and radius
 	 *
 	 * @param	InPosition	List of triangles
diff --git a/Engine/Source/Runtime/Engine/Public/CanvasTypes.h b/Engine/Source/Runtime/Engine/Public/CanvasTypes.h
--- a/Engine/Source/Runtime/Engine/Public/CanvasTypes.h
+++ b/Engine/Source/Runtime/Engine/Public/CanvasTypes.h
@@ -394,22 +394,22 @@ private:
 
 public:	
 
-	/*
+	/**
 	 * Access current real time 
 	 */
 	float GetCurrentRealTime() const { return CurrentRealTime; }
 
-	/*
+	/**
 	 * Access current world time 
 	 */
 	float GetCurrentWorldTime() const { return CurrentWorldTime; }
 
-	/*
+	/**
 	 * Access current delta time 
 	 */
 	float GetCurrentDeltaWorldTime() const { return CurrentDeltaWorldTime; }
 
-	/* 
+	/** 
 	 * Draw a CanvasItem
 	 *
 	 * @param Item			Item to draw
@@ -418,7 +418,7 @@ public:
 	{
 		Item.Draw( this );
 	}
-	/* 
+	/** 
 	 * Draw a CanvasItem at the given coordinates
 	 *
 	 * @param Item			Item to draw
@@ -429,7 +429,7 @@ public:
 		Item.Draw( this, InPosition );
 	}
 	
-	/* 
+	/** 
 	 * Draw a CanvasItem at the given coordinates
 	 *
 	 * @param Item			Item to draw
diff --git a/Engine/Source/Runtime/Engine/Public/WorldCollision.h b/Engine/Source/Runtime/Engine/Public/WorldCollision.h
--- a/Engine/Source/Runtime/Engine/Public/WorldCollision.h
+++ b/Engine/Source/Runtime/Engine/Public/WorldCollision.h
@@ -250,13 +250,13 @@ struct FBaseTraceDatum
 struct FTraceDatum;
 struct FOverlapDatum;
 
-/*
+/**
  * Delegate for Trace 
  * @param	FTraceHandle	TraceHandle that is returned when requested
  * @param	FTraceDatum		TraceDatum that includes input/output
  */
 DECLARE_DELEGATE_TwoParams( FTraceDelegate, const FTraceHandle&, FTraceDatum &);
-/*
+/**
  * Delegate for Overlap
  * @param	FTHandle		TraceHandle that is returned when requestsed
  * @param	FOverlapDatum	OverlapDatum that includes input/output
diff --git a/Engine/Source/Runtime/GameMenuBuilder/Public/GameMenuPage.h b/Engine/Source/Runtime/GameMenuBuilder/Public/GameMenuPage.h
--- a/Engine/Source/Runtime/GameMenuBuilder/Public/GameMenuPage.h
+++ b/Engine/Source/Runtime/GameMenuBuilder/Public/GameMenuPage.h
@@ -16,7 +16,7 @@ public:
 	/** Delegate for when menu is about to open. */
 	DECLARE_DELEGATE(FOnMenuOpening);
 
-	/* 
+	/** 
 	 * Initialize the menu page widget and set this menu as root. 
 	 * @param	InPCOwner		Player controller that owns this menu.
 	 * @param	InStyleName		Name of the menu style to use.
@@ -48,7 +48,7 @@ public:
 	/** The menu title. */
 	FText MenuTitle;
 
-	/* 
+	/** 
 	 * Add a menu item.
 	 *
 	 * @param	Text		The string for the item (EG START GAME)
@@ -75,7 +75,7 @@ public:
 		return Item.ToSharedRef();
 	}
 
-	/*
+	/**
 	 * Add a menu entry with a variable number of selectable options
 	 *
 	 * @param	Text		Label of the menu item
@@ -93,7 +93,7 @@ public:
 		return Item.ToSharedRef();
 	}
 
-	/*
+	/**
 	* Add a handler for the menu being canceled.
 	*
 	* @param	InObj		Menu page object
@@ -105,7 +105,7 @@ public:
 		OnGoBackCancel.BindSP(InObj, InMethod);
 	}
 	
-	/*
+	/**
 	* Add a handler for the menu being canceled.
 	*
 	* @param	InObj		Menu page object
@@ -117,7 +117,7 @@ public:
 		OnMenuHidden.BindSP(InObj, InMethod);
 	}
 
-	/*
+	/**
 	* Add a handler for the menu being accepted.
 	*
 	* @param	InObj		Menu page object
@@ -129,7 +129,7 @@ public:
 		OnGoBack.BindSP(InObj, InMethod);
 	}
 
-	/*
+	/**
 	* Add a handler for the menu being opened.
 	*
 	* @param	InObj		Menu page object
@@ -141,7 +141,7 @@ public:
 		OnMenuOpening.BindSP(InObj, InMethod);
 	}
 
-	/*
+	/**
 	 * Lock/Unlock the menu and prevent user interaction.
 	 *
 	 * @param		bLockState	true to lock
diff --git a/Engine/Source/Runtime/GameMenuBuilder/Public/SGameMenuPageWidget.h b/Engine/Source/Runtime/GameMenuBuilder/Public/SGameMenuPageWidget.h
--- a/Engine/Source/Runtime/GameMenuBuilder/Public/SGameMenuPageWidget.h
+++ b/Engine/Source/Runtime/GameMenuBuilder/Public/SGameMenuPageWidget.h
@@ -132,7 +132,7 @@ public:
 	/** Make the currently selected menu sub menu new main menu if valid. */
 	void EnterSubMenu(TSharedPtr<class FGameMenuPage> InSubMenu);
 
-	/*
+	/**
 	 * Go back the the previous menu.
 	 * 
 	 * @param	bIsCancel	if true will be treated as CANCEL (IE escape pressed).
@@ -143,7 +143,7 @@ public:
 	/** Confirms current menu item and performs an action. Will also play selection sound. */
 	void ConfirmMenuItem();
 
-	/*
+	/**
 	 * Show the given menu and make it the current menu. 
 	 *
 	 * @param InMenu	The menu to show and set as the current menu. If this is NULL Current menu will be used if valid.
@@ -154,7 +154,7 @@ public:
 	/** Hide the menu. */
 	void HideMenu();
 
-	/*
+	/**
 	 * Updates arrows visibility for multi-choice menu item.
 	 *
 	 * @param	InMenuItem The item to update.
@@ -162,7 +162,7 @@ public:
 	 */
 	void UpdateArrows(TSharedPtr<class FGameMenuItem> InMenuItem);
 
-	/*
+	/**
 	 * Change the currently selection option of the currently selected menu item.
 	 *
 	 * @param	InMoveBy		Number to change the option by.
@@ -173,7 +173,7 @@ public:
 	/** disable/enable moving around menu. */
 	void LockControls(bool bEnable);
 
-	/* 
+	/** 
 	 * Set the current menu and 'open' it as the main panel. 
 	 *
 	 * @param InMenu	The menu to open and set as current.
@@ -181,7 +181,7 @@ public:
 	 */
 	void OpenMainPanel(TSharedPtr< class FGameMenuPage > InMenu);
 		
-	/*
+	/**
 	* Select a given item from an index
 	*
 	* @param InSelection	The index of the item to select
@@ -208,7 +208,7 @@ protected:
 	/** Gets current scale for drawing menu. */
 	float GetUIScale() const;
 
-	/* 
+	/** 
 	 * Rebuild the widgets in the main menu widget.
 	 * Selected item will be set from InPreviousIndex unless it is invalid.
 	 *
@@ -264,7 +264,7 @@ protected:
 	/** Start the main menu open fade/anim and set keyboard focu.s */
 	void FadeIn();
 
-	/* 
+	/** 
 	 * Animates the title widget.
 	 *
 	 * @param	bShowTitle	true to show the title, false to hide it.
diff --git a/Engine/Source/Runtime/LauncherAutomatedService/Private/AutomatedTestManager.h b/Engine/Source/Runtime/LauncherAutomatedService/Private/AutomatedTestManager.h
--- a/Engine/Source/Runtime/LauncherAutomatedService/Private/AutomatedTestManager.h
+++ b/Engine/Source/Runtime/LauncherAutomatedService/Private/AutomatedTestManager.h
@@ -48,26 +48,26 @@ public:
 
 protected:
 
-	/*
+	/**
 	 * Finds some automation workers on the discovered game instances.
 	 *
 	 * @param DeltaTime Time since last update. Used to add a delay before we send the request.
 	 */
 	void FindWorkers( float DeltaTime );
 
-	/*
+	/**
 	 * Generates the automation report, and finish the testing process.
 	 */
 	void GenerateReport();
 
-	/*
+	/**
 	 * Monitors the testing process to see if we have finished.
 	 */
 	void MonitorTests();
 
 private:
 
-	/*
+	/**
 	 * Callback function from the automation controller to let us know there are tests available.
 	 */
 	void HandleRefreshTestCallback();
diff --git a/Engine/Source/Runtime/LauncherAutomatedService/Private/LauncherAutomatedServiceProvider.h b/Engine/Source/Runtime/LauncherAutomatedService/Private/LauncherAutomatedServiceProvider.h
--- a/Engine/Source/Runtime/LauncherAutomatedService/Private/LauncherAutomatedServiceProvider.h
+++ b/Engine/Source/Runtime/LauncherAutomatedService/Private/LauncherAutomatedServiceProvider.h
@@ -36,21 +36,21 @@ public:
 
 protected:
 
-	/*
+	/**
 	 * Sets up the profile manager responsible for deploying the session.
 	 *
 	 * @param Params The settings required to determine the profile/device group to use.
 	 */
 	void SetupProfileAndGroupSettings( const TCHAR* Params );
 
-	/*
+	/**
 	 * Start the automation tests
 	 */
 	void StartAutomationTests( );
 
 private:
 
-	/*
+	/**
 	 * Callback for devices being added to the proxy manager.
 	 *
 	 * @param AddedProxy The newly added device proxy to the device proxy manager.
diff --git a/Engine/Source/Runtime/LauncherAutomatedService/Public/Interfaces/IAutomatedTestManager.h b/Engine/Source/Runtime/LauncherAutomatedService/Public/Interfaces/IAutomatedTestManager.h
--- a/Engine/Source/Runtime/LauncherAutomatedService/Public/Interfaces/IAutomatedTestManager.h
+++ b/Engine/Source/Runtime/LauncherAutomatedService/Public/Interfaces/IAutomatedTestManager.h
@@ -18,19 +18,19 @@ class IAutomatedTestManager
 {
 public:
 
-	/*
+	/**
 	 * Checks whether we have finished running the automation tests.
 	 *
 	 * @return True if we have finished
 	 */
 	virtual bool IsTestingComplete() = 0;
 
-	/*
+	/**
 	 * Starts the testing process. Called after we have discovered some devices.
 	 */
 	virtual void RunTests() = 0;
 
-	/*
+	/**
 	 * Ticks the automation tests.
 	 *
 	 * @param DeltaTime Time since last update.
diff --git a/Engine/Source/Runtime/LauncherAutomatedService/Public/Interfaces/ILauncherAutomatedServiceProvider.h b/Engine/Source/Runtime/LauncherAutomatedService/Public/Interfaces/ILauncherAutomatedServiceProvider.h
--- a/Engine/Source/Runtime/LauncherAutomatedService/Public/Interfaces/ILauncherAutomatedServiceProvider.h
+++ b/Engine/Source/Runtime/LauncherAutomatedService/Public/Interfaces/ILauncherAutomatedServiceProvider.h
@@ -17,7 +17,7 @@ class ILauncherAutomatedServiceProvider
 {
 public:
 
-	/*
+	/**
 	 * Get the exit code of the service
 	 *
 	 * @return The exit Code of the service
diff --git a/Engine/Source/Runtime/Messaging/Public/Common/MessageBridgeBuilder.h b/Engine/Source/Runtime/Messaging/Public/Common/MessageBridgeBuilder.h
--- a/Engine/Source/Runtime/Messaging/Public/Common/MessageBridgeBuilder.h
+++ b/Engine/Source/Runtime/Messaging/Public/Common/MessageBridgeBuilder.h
@@ -18,7 +18,7 @@ public:
 		, Transport(nullptr)
 	{ }
 
-	/*
+	/**
 	 * Creates and initializes a new instance.
 	 *
 	 * @param InBus The message bus to attach the bridge to.
diff --git a/Engine/Source/Runtime/MovieSceneCore/Classes/MovieSceneTrack.h b/Engine/Source/Runtime/MovieSceneCore/Classes/MovieSceneTrack.h
--- a/Engine/Source/Runtime/MovieSceneCore/Classes/MovieSceneTrack.h
+++ b/Engine/Source/Runtime/MovieSceneCore/Classes/MovieSceneTrack.h
@@ -6,7 +6,7 @@
 
 class UMovieSceneSection;
 
-/*
+/**
  * Base class for a track in a Movie Scene
  */
 UCLASS( abstract, MinimalAPI )
diff --git a/Engine/Source/Runtime/MovieSceneCore/Public/IMovieSceneTrackInstance.h b/Engine/Source/Runtime/MovieSceneCore/Public/IMovieSceneTrackInstance.h
--- a/Engine/Source/Runtime/MovieSceneCore/Public/IMovieSceneTrackInstance.h
+++ b/Engine/Source/Runtime/MovieSceneCore/Public/IMovieSceneTrackInstance.h
@@ -2,7 +2,7 @@
 
 #pragma once
 
-/*
+/**
  * Base class for an animated type instance in a Movie Scene
  */
 class IMovieSceneTrackInstance
diff --git a/Engine/Source/Runtime/Online/OnlineSubsystemNull/Private/OnlineSessionInterfaceNull.h b/Engine/Source/Runtime/Online/OnlineSubsystemNull/Private/OnlineSessionInterfaceNull.h
--- a/Engine/Source/Runtime/Online/OnlineSubsystemNull/Private/OnlineSessionInterfaceNull.h
+++ b/Engine/Source/Runtime/Online/OnlineSubsystemNull/Private/OnlineSessionInterfaceNull.h
@@ -117,7 +117,7 @@ private:
 	 */
 	void ReadSettingsFromPacket(class FNboSerializeFromBufferNull& Packet, FOnlineSessionSettings& SessionSettings);
 
-	/*
+	/**
 	 * Delegate triggered when the LAN beacon has detected a valid client request has been received
 	 *
 	 * @param PacketData packet data sent by the requesting client with header information removed
@@ -126,7 +126,7 @@ private:
 	 */
 	void OnValidQueryPacketReceived(uint8* PacketData, int32 PacketLength, uint64 ClientNonce);
 
-	/*
+	/**
 	 * Delegate triggered when the LAN beacon has detected a valid host response to a client request has been received
 	 *
 	 * @param PacketData packet data sent by the requesting client with header information removed
@@ -134,7 +134,7 @@ private:
 	 */
 	void OnValidResponsePacketReceived(uint8* PacketData, int32 PacketLength);
 
-	/*
+	/**
 	 * Delegate triggered when the LAN beacon has finished searching (some time after last received host packet)
 	 */
 	void OnLANSearchTimeout();
diff --git a/Engine/Source/Runtime/Online/OnlineSubsystemSteam/Private/OnlineSessionInterfaceSteam.h b/Engine/Source/Runtime/Online/OnlineSubsystemSteam/Private/OnlineSessionInterfaceSteam.h
--- a/Engine/Source/Runtime/Online/OnlineSubsystemSteam/Private/OnlineSessionInterfaceSteam.h
+++ b/Engine/Source/Runtime/Online/OnlineSubsystemSteam/Private/OnlineSessionInterfaceSteam.h
@@ -192,7 +192,7 @@ private:
 	 */
 	void ReadSettingsFromPacket(class FNboSerializeFromBufferSteam& Packet, FOnlineSessionSettings& SessionSettings);
 
-	/*
+	/**
 	 * Delegate triggered when the LAN beacon has detected a valid client request has been received
 	 *
 	 * @param PacketData packet data sent by the requesting client with header information removed
@@ -201,7 +201,7 @@ private:
 	 */
 	void OnValidQueryPacketReceived(uint8* PacketData, int32 PacketLength, uint64 ClientNonce);
 
-	/*
+	/**
 	 * Delegate triggered when the LAN beacon has detected a valid host response to a client request has been received
 	 *
 	 * @param PacketData packet data sent by the requesting client with header information removed
@@ -209,7 +209,7 @@ private:
 	 */
 	void OnValidResponsePacketReceived(uint8* PacketData, int32 PacketLength);
 
-	/*
+	/**
 	 * Delegate triggered when the LAN beacon has finished searching (some time after last received host packet)
 	 */
 	void OnLANSearchTimeout();
diff --git a/Engine/Source/Runtime/OpenGLDrv/Private/OpenGLCommands.cpp b/Engine/Source/Runtime/OpenGLDrv/Private/OpenGLCommands.cpp
--- a/Engine/Source/Runtime/OpenGLDrv/Private/OpenGLCommands.cpp
+++ b/Engine/Source/Runtime/OpenGLDrv/Private/OpenGLCommands.cpp
@@ -3101,7 +3101,7 @@ void FOpenGLDynamicRHI::RHIBlockUntilGPUIdle()
 	// Not really supported
 }
 
-/*
+/**
  * Returns the total GPU time taken to render the last frame. Same metric as FPlatformTime::Cycles().
  */
 uint32 FOpenGLDynamicRHI::RHIGetGPUFrameCycles()
diff --git a/Engine/Source/Runtime/OpenGLDrv/Public/OpenGLResources.h b/Engine/Source/Runtime/OpenGLDrv/Public/OpenGLResources.h
--- a/Engine/Source/Runtime/OpenGLDrv/Public/OpenGLResources.h
+++ b/Engine/Source/Runtime/OpenGLDrv/Public/OpenGLResources.h
@@ -896,7 +896,7 @@ public:
 	 */
 	void CloneViaPBO( TOpenGLTexture* Src, uint32 NumMips, int32 SrcOffset, int32 DstOffset);
 
-	/*
+	/**
 	 * Resolved the specified face for a read Lock, for non-renderable, CPU readable surfaces this eliminates the readback inside Lock itself.
 	 */
 	void Resolve(uint32 MipIndex,uint32 ArrayIndex);
diff --git a/Engine/Source/Runtime/Slate/Private/Widgets/Notifications/SNotificationList.cpp b/Engine/Source/Runtime/Slate/Private/Widgets/Notifications/SNotificationList.cpp
--- a/Engine/Source/Runtime/Slate/Private/Widgets/Notifications/SNotificationList.cpp
+++ b/Engine/Source/Runtime/Slate/Private/Widgets/Notifications/SNotificationList.cpp
@@ -321,7 +321,7 @@ public:
 		];
 	}
 	
-	/*
+	/**
 	 * Returns the internals of the notification
 	 */
 	TSharedRef<SHorizontalBox> ConstructInternals( const FArguments& InArgs ) 
diff --git a/Engine/Source/Runtime/Slate/Public/Widgets/Input/SSpinBox.h b/Engine/Source/Runtime/Slate/Public/Widgets/Input/SSpinBox.h
--- a/Engine/Source/Runtime/Slate/Public/Widgets/Input/SSpinBox.h
+++ b/Engine/Source/Runtime/Slate/Public/Widgets/Input/SSpinBox.h
@@ -608,7 +608,7 @@ protected:
 		return FMath::Clamp<double>(Result, TNumericLimits<NumericType>::Lowest(), TNumericLimits<NumericType>::Max());
 	}
 	
-	/*
+	/**
 	 * Helper structure to check if a given character is supported by underlying numeric type
 	 */
 	template<typename T, typename U = void>
diff --git a/Engine/Source/Runtime/SlateCore/Public/Fonts/FontCache.h b/Engine/Source/Runtime/SlateCore/Public/Fonts/FontCache.h
--- a/Engine/Source/Runtime/SlateCore/Public/Fonts/FontCache.h
+++ b/Engine/Source/Runtime/SlateCore/Public/Fonts/FontCache.h
@@ -217,7 +217,7 @@ private:
 	mutable uint16 Baseline;
 };
 
-/*
+/**
  * Font caching implementation
  * Caches characters into textures as needed
  */
diff --git a/Engine/Source/Runtime/SlateCore/Public/Styling/SlateBrush.h b/Engine/Source/Runtime/SlateCore/Public/Styling/SlateBrush.h
--- a/Engine/Source/Runtime/SlateCore/Public/Styling/SlateBrush.h
+++ b/Engine/Source/Runtime/SlateCore/Public/Styling/SlateBrush.h
@@ -77,7 +77,7 @@ namespace ESlateBrushMirrorType
 }
 
 
-/*
+/**
  * Enumerates brush image types.
  */
 UENUM()
diff --git a/Engine/Source/Runtime/Windows/D3D11RHI/Private/D3D11Commands.cpp b/Engine/Source/Runtime/Windows/D3D11RHI/Private/D3D11Commands.cpp
--- a/Engine/Source/Runtime/Windows/D3D11RHI/Private/D3D11Commands.cpp
+++ b/Engine/Source/Runtime/Windows/D3D11RHI/Private/D3D11Commands.cpp
@@ -2074,7 +2074,7 @@ void FD3D11DynamicRHI::RHIBlockUntilGPUIdle()
 	// Not really supported
 }
 
-/*
+/**
  * Returns the total GPU time taken to render the last frame. Same metric as FPlatformTime::Cycles().
  */
 uint32 FD3D11DynamicRHI::RHIGetGPUFrameCycles()
