CrashReport - Added debug threadensure and debug threadfatal

diff --git a/Engine/Source/Runtime/Engine/Private/UnrealEngine.cpp b/Engine/Source/Runtime/Engine/Private/UnrealEngine.cpp
--- a/Engine/Source/Runtime/Engine/Private/UnrealEngine.cpp
+++ b/Engine/Source/Runtime/Engine/Private/UnrealEngine.cpp
@@ -3972,15 +3972,15 @@ bool UEngine::HandleDebugCommand( const TCHAR* Cmd, FOutputDevice& Ar )
 	}
 	if( FParse::Command(&Cmd,TEXT("RENDERCHECK")) )
 	{
-		struct FFatal
+		struct FRender
 		{
-			static void Crash()
+			static void Check()
 			{
 				UE_LOG(LogEngine, Warning, TEXT("Printed warning to log.") );
 				check(!"Crashing the renderthread via check(0) at your request");
 			}
 		};			
-		ENQUEUE_UNIQUE_RENDER_COMMAND( CauseRenderThreadCrash, { FFatal::Crash();});
+		ENQUEUE_UNIQUE_RENDER_COMMAND( CauseRenderThreadCrash, { FRender::Check();});
 		return true;
 	}
 	if( FParse::Command(&Cmd,TEXT("RENDERGPF")) )
@@ -3988,9 +3988,18 @@ bool UEngine::HandleDebugCommand( const TCHAR* Cmd, FOutputDevice& Ar )
 		ENQUEUE_UNIQUE_RENDER_COMMAND( CauseRenderThreadCrash, {UE_LOG(LogEngine, Warning, TEXT("Printed warning to log.") ); *(int32 *)3 = 123;} );
 		return true;
 	}
+	if( FParse::Command( &Cmd, TEXT( "RENDERFATAL" ) ) )
+	{
+		ENQUEUE_UNIQUE_RENDER_COMMAND( CauseRenderThreadCrash, 
+		{
+			UE_LOG( LogEngine, Warning, TEXT( "Printed warning to log." ) );
+			LowLevelFatalError( TEXT( "FError::LowLevelFatal test" ) ); 
+		} );
+		return true;
+	}
 	if( FParse::Command(&Cmd,TEXT("THREADCRASH")) )
 	{
-		struct FFatal
+		struct FThread
 		{
 			static void Crash(ENamedThreads::Type, const  FGraphEventRef&)
 			{
@@ -3998,33 +4007,59 @@ bool UEngine::HandleDebugCommand( const TCHAR* Cmd, FOutputDevice& Ar )
 				UE_LOG(LogEngine, Fatal, TEXT("Crashing the worker thread at your request") );
 			}
 		};
-		FTaskGraphInterface::Get().WaitUntilTaskCompletes(FDelegateGraphTask::CreateAndDispatchWhenReady(FDelegateGraphTask::FDelegate::CreateStatic(FFatal::Crash), TEXT("Crash")), ENamedThreads::GameThread);
+		FTaskGraphInterface::Get().WaitUntilTaskCompletes(FDelegateGraphTask::CreateAndDispatchWhenReady(FDelegateGraphTask::FDelegate::CreateStatic(FThread::Crash), TEXT("FThread::Crash")), ENamedThreads::GameThread);
 		return true;
 	}
 	if( FParse::Command(&Cmd,TEXT("THREADCHECK")) )
 	{
-		struct FFatal
+		struct FThread
 		{
-			static void Crash(ENamedThreads::Type, const  FGraphEventRef&)
+			static void Check(ENamedThreads::Type, const FGraphEventRef&)
 			{
 				UE_LOG(LogEngine, Warning, TEXT("Printed warning to log.") );
 				check(!"Crashing a worker thread via check(0) at your request");
 			}
 		};
-		FTaskGraphInterface::Get().WaitUntilTaskCompletes(FDelegateGraphTask::CreateAndDispatchWhenReady(FDelegateGraphTask::FDelegate::CreateStatic(FFatal::Crash), TEXT("Crash")), ENamedThreads::GameThread);
+		FTaskGraphInterface::Get().WaitUntilTaskCompletes(FDelegateGraphTask::CreateAndDispatchWhenReady(FDelegateGraphTask::FDelegate::CreateStatic(FThread::Check), TEXT("FThread::Check")), ENamedThreads::GameThread);
 		return true;
 	}
 	if( FParse::Command(&Cmd,TEXT("THREADGPF")) )
 	{
-		struct FFatal
+		struct FThread
 		{
-			static void Crash(ENamedThreads::Type, const  FGraphEventRef&)
+			static void GPF(ENamedThreads::Type, const FGraphEventRef&)
 			{
 				UE_LOG(LogEngine, Warning, TEXT("Printed warning to log.") );
 				*(int32 *)3 = 123;
 			}
 		};
-		FTaskGraphInterface::Get().WaitUntilTaskCompletes(FDelegateGraphTask::CreateAndDispatchWhenReady(FDelegateGraphTask::FDelegate::CreateStatic(FFatal::Crash), TEXT("Crash")), ENamedThreads::GameThread);
+		FTaskGraphInterface::Get().WaitUntilTaskCompletes(FDelegateGraphTask::CreateAndDispatchWhenReady(FDelegateGraphTask::FDelegate::CreateStatic(FThread::GPF), TEXT("FThread::GPF")), ENamedThreads::GameThread);
+		return true;
+	}
+	else if( FParse::Command( &Cmd, TEXT( "THREADENSURE" ) ) )
+	{
+		struct FThread
+		{
+			static void Ensure( ENamedThreads::Type, const FGraphEventRef& )
+			{
+				UE_LOG( LogEngine, Warning, TEXT( "Printed warning to log." ) );
+				ensure( 0 );
+			}
+		};
+		FTaskGraphInterface::Get().WaitUntilTaskCompletes( FDelegateGraphTask::CreateAndDispatchWhenReady( FDelegateGraphTask::FDelegate::CreateStatic( FThread::Ensure ), TEXT( "FThread::Ensure" ) ), ENamedThreads::GameThread );
+		return true;
+	}
+	else if( FParse::Command( &Cmd, TEXT( "THREADFATAL" ) ) )
+	{
+		struct FThread
+		{
+			static void Fatal( ENamedThreads::Type, const FGraphEventRef& )
+			{
+				UE_LOG( LogEngine, Warning, TEXT( "Printed warning to log." ) );
+				LowLevelFatalError( TEXT( "FError::LowLevelFatal test" ) );
+			}
+		};
+		FTaskGraphInterface::Get().WaitUntilTaskCompletes( FDelegateGraphTask::CreateAndDispatchWhenReady( FDelegateGraphTask::FDelegate::CreateStatic( FThread::Fatal ), TEXT( "FThread::Fatal" ) ), ENamedThreads::GameThread );
 		return true;
 	}
 	else if( FParse::Command(&Cmd,TEXT("CRASH")) )
@@ -4047,12 +4082,6 @@ bool UEngine::HandleDebugCommand( const TCHAR* Cmd, FOutputDevice& Ar )
 		*(int32 *)3 = 123;
 		return true;
 	}
-	else if( FParse::Command( &Cmd, TEXT("ASSERT") ) )
-	{
-		UE_LOG(LogEngine, Warning, TEXT("Printed warning to log.") );
-		check(0);
-		return true;
-	}
 	else if( FParse::Command( &Cmd, TEXT("ENSURE") ) )
 	{
 		UE_LOG(LogEngine, Warning, TEXT("Printed warning to log.") );
@@ -4061,6 +4090,12 @@ bool UEngine::HandleDebugCommand( const TCHAR* Cmd, FOutputDevice& Ar )
 			return true;
 		}
 	}
+	else if( FParse::Command( &Cmd, TEXT( "FATAL" ) ) )
+	{
+		UE_LOG( LogEngine, Warning, TEXT( "Printed warning to log." ) );
+		LowLevelFatalError( TEXT( "FError::LowLevelFatal test" ) );
+		return true;
+	}
 	else if( FParse::Command( &Cmd, TEXT("RESETLOADERS") ) )
 	{
 		ResetLoaders( NULL );
