Compartamentalize finish spawning aspects of Actor in to a function and track if it is being done.
Ensure that it is only called once.

diff --git a/Engine/Source/Runtime/Engine/Classes/GameFramework/Actor.h b/Engine/Source/Runtime/Engine/Classes/GameFramework/Actor.h
--- a/Engine/Source/Runtime/Engine/Classes/GameFramework/Actor.h
+++ b/Engine/Source/Runtime/Engine/Classes/GameFramework/Actor.h
@@ -160,6 +160,9 @@ private:
 	/** True if this actor is currently running user construction script (used to defer component registration) */
 	uint32 bRunningUserConstructionScript:1;
 
+	/** Whether FinishSpawning has been called for this Actor.  If it has not, the Actor is in a mal-formed state */
+	uint32 bHasFinishedSpawning:1;
+
 protected:
 	// This actor will replicate to remote machines
 	UPROPERTY(Category = Replication, EditDefaultsOnly, BlueprintReadOnly)
@@ -1470,9 +1473,14 @@ public:
 	/** Called after the actor is spawned in the world.  Responsible for setting up actor for play. */
 	void PostSpawnInitialize(FVector const& SpawnLocation, FRotator const& SpawnRotation, AActor* InOwner, APawn* InInstigator, bool bRemoteOwned, bool bNoFail, bool bDeferConstruction);
 
+        /** Called to finish the spawning process, generally in the case of deferred spawning */
+	void FinishSpawning(const FTransform& Transform);
+
+private:
 	/** Called after the actor has run its construction. Responsible for finishing the actor spawn process. */
 	void PostActorConstruction();
 
+public:
 	/** Called immediately before gameplay begins. */
 	virtual void PreInitializeComponents();
 
diff --git a/Engine/Source/Runtime/Engine/Private/ChildActorComponent.cpp b/Engine/Source/Runtime/Engine/Private/ChildActorComponent.cpp
--- a/Engine/Source/Runtime/Engine/Private/ChildActorComponent.cpp
+++ b/Engine/Source/Runtime/Engine/Private/ChildActorComponent.cpp
@@ -170,8 +170,7 @@ void UChildActorComponent::CreateChildActor()
 					ChildActor->ParentComponentActor = GetOwner();
 
 					// Parts that we deferred from SpawnActor
-					ChildActor->ExecuteConstruction(ComponentToWorld, NULL);
-					ChildActor->PostActorConstruction();
+					ChildActor->FinishSpawning(ComponentToWorld);
 
 					// attach new actor to this component
 					ChildActor->AttachRootComponentTo(this, NAME_None, EAttachLocation::SnapToTarget);
diff --git a/Engine/Source/Runtime/Engine/Private/GameplayStatics.cpp b/Engine/Source/Runtime/Engine/Private/GameplayStatics.cpp
--- a/Engine/Source/Runtime/Engine/Private/GameplayStatics.cpp
+++ b/Engine/Source/Runtime/Engine/Private/GameplayStatics.cpp
@@ -313,8 +313,7 @@ AActor* UGameplayStatics::FinishSpawningActor(AActor* Actor, const FTransform& S
 {
 	if (Actor)
 	{
-		Actor->ExecuteConstruction(SpawnTransform, NULL);
-		Actor->PostActorConstruction();
+		Actor->FinishSpawning(SpawnTransform);
 	}
 
 	return Actor;
diff --git a/Engine/Source/Runtime/GameplayAbilities/Private/Abilities/Tasks/AbilityTask_SpawnActor.cpp b/Engine/Source/Runtime/GameplayAbilities/Private/Abilities/Tasks/AbilityTask_SpawnActor.cpp
--- a/Engine/Source/Runtime/GameplayAbilities/Private/Abilities/Tasks/AbilityTask_SpawnActor.cpp
+++ b/Engine/Source/Runtime/GameplayAbilities/Private/Abilities/Tasks/AbilityTask_SpawnActor.cpp
@@ -33,10 +33,9 @@ void UAbilityTask_SpawnActor::FinishSpawningActor(UObject* WorldContextObject, A
 {
 	if (SpawnedActor)
 	{
-		FTransform SpawnTransform = AbilitySystemComponent->GetOwner()->GetTransform();
+		const FTransform SpawnTransform = AbilitySystemComponent->GetOwner()->GetTransform();
 
-		SpawnedActor->ExecuteConstruction(SpawnTransform, NULL);
-		SpawnedActor->PostActorConstruction();
+		SpawnedActor->FinishSpawning(SpawnTransform);
 
 		Success.Broadcast(SpawnedActor);
 	}
diff --git a/Engine/Source/Runtime/GameplayAbilities/Private/Abilities/Tasks/AbilityTask_WaitTargetData.cpp b/Engine/Source/Runtime/GameplayAbilities/Private/Abilities/Tasks/AbilityTask_WaitTargetData.cpp
--- a/Engine/Source/Runtime/GameplayAbilities/Private/Abilities/Tasks/AbilityTask_WaitTargetData.cpp
+++ b/Engine/Source/Runtime/GameplayAbilities/Private/Abilities/Tasks/AbilityTask_WaitTargetData.cpp
@@ -98,10 +98,9 @@ void UAbilityTask_WaitTargetData::FinishSpawningActor(UObject* WorldContextObjec
 	{
 		check(MyTargetActor == SpawnedActor);
 
-		FTransform SpawnTransform = AbilitySystemComponent->GetOwner()->GetTransform();
+		const FTransform SpawnTransform = AbilitySystemComponent->GetOwner()->GetTransform();
 
-		SpawnedActor->ExecuteConstruction(SpawnTransform, NULL);
-		SpawnedActor->PostActorConstruction();
+		SpawnedActor->FinishSpawning(SpawnTransform);
 
 		// User ability activation is inhibited while this is active
 		AbilitySystemComponent->SetUserAbilityActivationInhibited(true);
