Missed file for CL# 2266098

diff --git a/Engine/Source/Runtime/Engine/Private/Actor.cpp b/Engine/Source/Runtime/Engine/Private/Actor.cpp
--- a/Engine/Source/Runtime/Engine/Private/Actor.cpp
+++ b/Engine/Source/Runtime/Engine/Private/Actor.cpp
@@ -2198,7 +2198,17 @@ void AActor::PostSpawnInitialize(FVector const& SpawnLocation, FRotator const& S
 	{
 		// Preserve original root component scale
 		const FVector SpawnScale = GetRootComponent() ? GetRootComponent()->RelativeScale3D : FVector::ZeroVector;
-		ExecuteConstruction( FTransform(SpawnRotation, SpawnLocation, SpawnScale), NULL );
+		FinishSpawning(FTransform(SpawnRotation, SpawnLocation, SpawnScale));
+	}
+}
+
+void AActor::FinishSpawning(const FTransform& Transform)
+{
+	if (ensure(!bHasFinishedSpawning))
+	{
+		bHasFinishedSpawning = true;
+
+		ExecuteConstruction(Transform, nullptr);
 		PostActorConstruction();
 	}
 }
