Spawn rate now works properly across LODs for GPU particles; other parameters still don't, but at least the spawn rate can now be scaled with LODs

diff --git a/Engine/Source/Runtime/Engine/Private/Particles/ParticleGpuSimulation.cpp b/Engine/Source/Runtime/Engine/Private/Particles/ParticleGpuSimulation.cpp
--- a/Engine/Source/Runtime/Engine/Private/Particles/ParticleGpuSimulation.cpp
+++ b/Engine/Source/Runtime/Engine/Private/Particles/ParticleGpuSimulation.cpp
@@ -3638,8 +3638,8 @@ private:
 
 		// Determine spawn rate.
 		check(SpawnModule && RequiredModule);
-		const float RateScale = SpawnModule->RateScale.GetValue(EmitterTime, Component);
-		float SpawnRate = SpawnModule->Rate.GetValue(EmitterTime, Component) * RateScale;
+		const float RateScale = CurrentLODLevel->SpawnModule->RateScale.GetValue(EmitterTime, Component);
+		float SpawnRate = CurrentLODLevel->SpawnModule->Rate.GetValue(EmitterTime, Component) * RateScale;
 		SpawnRate = FMath::Max<float>(0.0f, SpawnRate);
 
 		if (EmitterInfo.SpawnPerUnitModule)
@@ -3830,6 +3830,8 @@ private:
 
 	virtual void SetCurrentLODIndex(int32 InLODIndex, bool bInFullyProcess)
 	{
+		bool bDifferent = (InLODIndex != CurrentLODLevelIndex);
+		FParticleEmitterInstance::SetCurrentLODIndex(InLODIndex, bInFullyProcess);
 	}
 
 	virtual uint32 RequiredBytes()
