#ue4 - Fix UPrimitiveComponent MoveComponent() and UpdateOverlaps() not using the same criteria for what is a possible overlap. Attached components are now allowed to generate overlaps.

- Previously, attached objects would get continuous Begin/End overlap events when the parent moved. Now those events only fire when the overlaps actually change.
- Added new functions for managing the MoveIgnoreActors list in PrimitiveComponent in addition to IgnoreActorWhenMoving: GetMoveIgnoreActors, ClearMoveIgnoreActors.
- Attaching and detaching components now updates overlap state immediately.

reported at: https://udn.unrealengine.com/questions/217218/actors-attached-to-other-actors-constantly-generat.html

diff --git a/Engine/Source/Runtime/Engine/Classes/Components/PrimitiveComponent.h b/Engine/Source/Runtime/Engine/Classes/Components/PrimitiveComponent.h
--- a/Engine/Source/Runtime/Engine/Classes/Components/PrimitiveComponent.h
+++ b/Engine/Source/Runtime/Engine/Classes/Components/PrimitiveComponent.h
@@ -443,9 +443,34 @@ public:
 	UPROPERTY(EditDefaultsOnly, BlueprintReadOnly, Category=Base)
 	TEnumAsByte<enum ECanBeCharacterBase> CanCharacterStepUpOn;
 
-	/** Set of actors to ignore during component sweeps in MoveComponent() */
+	/**
+	 * Set of actors to ignore during component sweeps in MoveComponent().
+	 * All components owned by these actors will be ignored when this component moves or updates overlaps.
+	 * Components on the other Actor may also need to be told to do the same when they move.
+	 * @see IgnoreActorWhenMoving
+	 */
 	TArray<TWeakObjectPtr<AActor> > MoveIgnoreActors;
 
+	/**
+	 * Tells this component whether to ignore collision with all components of a specific Actor when this component is moved.
+	 * Components on the other Actor may also need to be told to do the same when they move.
+	 */
+	UFUNCTION(BlueprintCallable, Category = "Collision")
+	void IgnoreActorWhenMoving(AActor* Actor, bool bShouldIgnore);
+
+	/**
+	 * Returns the list of actors we currently ignore when moving.
+	 */
+	UFUNCTION(BlueprintCallable, Category = "Collision")
+	TArray<TWeakObjectPtr<AActor> > & GetMoveIgnoreActors();
+
+	/**
+	 * Clear the list of actors we ignore when moving.
+	 */
+	UFUNCTION(BlueprintCallable, Category = "Collision")
+	void ClearMoveIgnoreActors();
+
+
 #if WITH_EDITOR
 	/** Override delegate used for checking the selection state of a component */
 	DECLARE_DELEGATE_RetVal_OneParam( bool, FSelectionOverride, const UPrimitiveComponent* );
@@ -509,10 +534,6 @@ public:
 	 */
 	virtual void UpdateOverlaps(TArray<FOverlapInfo> const* PendingOverlaps=NULL, bool bDoNotifies=true, const TArray<FOverlapInfo>* OverlapsAtEndLocation=NULL) override;
 
-	/** Tells this component to ignore collision with the specified actor when being moved. */
-	UFUNCTION(BlueprintCallable, Category = "Collision")
-	void IgnoreActorWhenMoving(AActor* Actor, bool bShouldIgnore);
-
 	/** Overridden to use the overlaps to find the physics volume. */
 	virtual void UpdatePhysicsVolume( bool bTriggerNotifiers ) override;
 
diff --git a/Engine/Source/Runtime/Engine/Private/Pawn.cpp b/Engine/Source/Runtime/Engine/Private/Pawn.cpp
--- a/Engine/Source/Runtime/Engine/Private/Pawn.cpp
+++ b/Engine/Source/Runtime/Engine/Private/Pawn.cpp
@@ -1037,12 +1037,7 @@ void APawn::MoveIgnoreActorAdd(AActor * ActorToIgnore)
 	UPrimitiveComponent * RootPrimitiveComponent = Cast<UPrimitiveComponent>(GetRootComponent());
 	if( RootPrimitiveComponent )
 	{
-		// Remove dead references first
-		RootPrimitiveComponent->MoveIgnoreActors.Remove(NULL);
-		if( ActorToIgnore )
-		{
-			RootPrimitiveComponent->MoveIgnoreActors.AddUnique(ActorToIgnore);
-		}
+		RootPrimitiveComponent->IgnoreActorWhenMoving(ActorToIgnore, true);
 	}
 }
 
@@ -1051,12 +1046,7 @@ void APawn::MoveIgnoreActorRemove(AActor * ActorToIgnore)
 	UPrimitiveComponent * RootPrimitiveComponent = Cast<UPrimitiveComponent>(GetRootComponent());
 	if( RootPrimitiveComponent )
 	{
-		// Remove dead references first
-		RootPrimitiveComponent->MoveIgnoreActors.Remove(NULL);
-		if( ActorToIgnore )
-		{
-			RootPrimitiveComponent->MoveIgnoreActors.Remove(ActorToIgnore);
-		}
+		RootPrimitiveComponent->IgnoreActorWhenMoving(ActorToIgnore, false);
 	}
 }
 
diff --git a/Engine/Source/Runtime/Engine/Private/PrimitiveComponent.cpp b/Engine/Source/Runtime/Engine/Private/PrimitiveComponent.cpp
--- a/Engine/Source/Runtime/Engine/Private/PrimitiveComponent.cpp
+++ b/Engine/Source/Runtime/Engine/Private/PrimitiveComponent.cpp
@@ -1212,58 +1212,88 @@ static void PullBackHit(FHitResult& Hit, const FVector& Start, const FVector& En
 const static float PERF_SHOW_MOVECOMPONENT_TAKING_LONG_TIME_AMOUNT = 2.0f; // modify this value to look at larger or smaller sets of "bad" actors
 
 
-static bool ShouldIgnoreHitResult(const UWorld* InWorld, FHitResult const& TestHit, FVector const& MovementDirDenormalized, AActor* MovingActor, EMoveComponentFlags MoveFlags)
+static bool ShouldIgnoreHitResult(const UWorld* InWorld, FHitResult const& TestHit, FVector const& MovementDirDenormalized, const AActor* MovingActor, EMoveComponentFlags MoveFlags)
 {
-	// check "ignore bases" functionality
-	if ( (MoveFlags & MOVECOMP_IgnoreBases) && MovingActor && TestHit.bBlockingHit )	//we let overlap components go through because their overlap is still needed and will cause beginOverlap/endOverlap events
+	if (TestHit.bBlockingHit)
 	{
-		// ignore if there's a base relationship between moving actor and hit actor
-		AActor const* const HitActor = TestHit.GetActor();
-		if (HitActor)
+		// check "ignore bases" functionality
+		if ( (MoveFlags & MOVECOMP_IgnoreBases) && MovingActor )	//we let overlap components go through because their overlap is still needed and will cause beginOverlap/endOverlap events
 		{
-			if (MovingActor->IsBasedOnActor(HitActor) || HitActor->IsBasedOnActor(MovingActor))
+			// ignore if there's a base relationship between moving actor and hit actor
+			AActor const* const HitActor = TestHit.GetActor();
+			if (HitActor)
 			{
-				return true;
+				if (MovingActor->IsBasedOnActor(HitActor) || HitActor->IsBasedOnActor(MovingActor))
+				{
+					return true;
+				}
 			}
 		}
-	}
 	
-	// If we started penetrating, we may want to ignore it if we are moving out of penetration.
-	// This helps prevent getting stuck in walls.
-	if ( TestHit.bStartPenetrating && TestHit.bBlockingHit && !(MoveFlags & MOVECOMP_NeverIgnoreBlockingOverlaps) )
-	{
-		const float DotTolerance = CVarInitialOverlapTolerance.GetValueOnGameThread();
+		// If we started penetrating, we may want to ignore it if we are moving out of penetration.
+		// This helps prevent getting stuck in walls.
+		if ( TestHit.bStartPenetrating && !(MoveFlags & MOVECOMP_NeverIgnoreBlockingOverlaps) )
+		{
+			const float DotTolerance = CVarInitialOverlapTolerance.GetValueOnGameThread();
 
-		// Dot product of movement direction against 'exit' direction
-		const FVector MovementDir = MovementDirDenormalized.SafeNormal();
-		const float MoveDot = (TestHit.ImpactNormal | MovementDir);
+			// Dot product of movement direction against 'exit' direction
+			const FVector MovementDir = MovementDirDenormalized.SafeNormal();
+			const float MoveDot = (TestHit.ImpactNormal | MovementDir);
 
-		const bool bMovingOut = MoveDot > DotTolerance;
+			const bool bMovingOut = MoveDot > DotTolerance;
 
-#if !(UE_BUILD_SHIPPING || UE_BUILD_TEST)
-		{
-			if(CVarShowInitialOverlaps.GetValueOnGameThread() != 0)
+	#if !(UE_BUILD_SHIPPING || UE_BUILD_TEST)
 			{
-				UE_LOG(LogTemp, Log, TEXT("Overlapping %s Dir %s Dot %f Normal %s Depth %f"), *GetNameSafe(TestHit.Component.Get()), *MovementDir.ToString(), MoveDot, *TestHit.ImpactNormal.ToString(), TestHit.PenetrationDepth);
-				DrawDebugDirectionalArrow(InWorld, TestHit.TraceStart, TestHit.TraceStart + 30.f * TestHit.ImpactNormal, 5.f, bMovingOut ? FColor(64,128,255) : FColor(255,64,64), true, 4.f);
-				if (TestHit.PenetrationDepth > KINDA_SMALL_NUMBER)
+				if(CVarShowInitialOverlaps.GetValueOnGameThread() != 0)
 				{
-					DrawDebugDirectionalArrow(InWorld, TestHit.TraceStart, TestHit.TraceStart + TestHit.PenetrationDepth * TestHit.Normal, 5.f, FColor(64,255,64), true, 4.f);
+					UE_LOG(LogTemp, Log, TEXT("Overlapping %s Dir %s Dot %f Normal %s Depth %f"), *GetNameSafe(TestHit.Component.Get()), *MovementDir.ToString(), MoveDot, *TestHit.ImpactNormal.ToString(), TestHit.PenetrationDepth);
+					DrawDebugDirectionalArrow(InWorld, TestHit.TraceStart, TestHit.TraceStart + 30.f * TestHit.ImpactNormal, 5.f, bMovingOut ? FColor(64,128,255) : FColor(255,64,64), true, 4.f);
+					if (TestHit.PenetrationDepth > KINDA_SMALL_NUMBER)
+					{
+						DrawDebugDirectionalArrow(InWorld, TestHit.TraceStart, TestHit.TraceStart + TestHit.PenetrationDepth * TestHit.Normal, 5.f, FColor(64,255,64), true, 4.f);
+					}
 				}
 			}
-		}
-#endif
+	#endif
 
-		// If we are moving out, ignore this result!
-		if (bMovingOut)
-		{
-			return true;
+			// If we are moving out, ignore this result!
+			if (bMovingOut)
+			{
+				return true;
+			}
 		}
 	}
 
 	return false;
 }
 
+static bool ShouldIgnoreOverlapResult(const UWorld* World, const AActor* ThisActor, const UPrimitiveComponent& ThisComponent, const AActor* OtherActor, const UPrimitiveComponent& OtherComponent)
+{
+	// Don't overlap with self
+	if (&ThisComponent == &OtherComponent)
+	{
+		return true;
+	}
+
+	// Both components must set bGenerateOverlapEvents
+	if (!ThisComponent.bGenerateOverlapEvents || !OtherComponent.bGenerateOverlapEvents)
+	{
+		return true;
+	}
+
+	if (!ThisActor || !OtherActor)
+	{
+		return true;
+	}
+
+	if (!World || OtherActor == World->GetWorldSettings() || !OtherActor->bActorInitialized)
+	{
+		return true;
+	}
+
+	return false;
+}
+
 
 void UPrimitiveComponent::InitSweepCollisionParams(FCollisionQueryParams &OutParams, FCollisionResponseParams& OutResponseParam) const
 {
@@ -1441,22 +1471,25 @@ bool UPrimitiveComponent::MoveComponent( const FVector& Delta, const FRotator& N
 						else if (bGenerateOverlapEvents)
 						{
 							UPrimitiveComponent * OverlapComponent = TestHit.Component.Get();
-							if (OverlapComponent && OverlapComponent->bGenerateOverlapEvents)	//only register overlap if both components have set bGenerateOverlapEvents
+							if (OverlapComponent && OverlapComponent->bGenerateOverlapEvents)
 							{
-								// don't process touch events after initial blocking hits
-								if (BlockingHitIndex >= 0 && TestHit.Time > Hits[BlockingHitIndex].Time)
+								if (!ShouldIgnoreOverlapResult(GetWorld(), Actor, *this, TestHit.GetActor(), *OverlapComponent))
 								{
-									break;
+									// don't process touch events after initial blocking hits
+									if (BlockingHitIndex >= 0 && TestHit.Time > Hits[BlockingHitIndex].Time)
+									{
+										break;
+									}
+
+									if (FirstNonInitialOverlapIdx == INDEX_NONE && TestHit.Time > 0.f)
+									{
+										// We are about to add the first non-initial overlap.
+										FirstNonInitialOverlapIdx = PendingOverlaps.Num();
+									}
+
+									// cache touches
+									PendingOverlaps.AddUnique(FOverlapInfo(TestHit));
 								}
-
-								if (FirstNonInitialOverlapIdx == INDEX_NONE && TestHit.Time > 0.f)
-								{
-									// We are about to add the first non-initial overlap.
-									FirstNonInitialOverlapIdx = PendingOverlaps.Num();
-								}
-
-								// cache touches
-								PendingOverlaps.AddUnique(FOverlapInfo(TestHit));
 							}
 						}
 					}
@@ -1498,7 +1531,7 @@ bool UPrimitiveComponent::MoveComponent( const FVector& Delta, const FRotator& N
 			// However any rotation was set at the end and not swept, so we can't assume the overlaps at the end location are correct if rotation changed.
 			if (PendingOverlaps.Num() == 0 && CVarAllowCachedOverlaps->GetInt())
 			{
-				if (Actor->GetRootComponent() == this && AreSymmetricRotations(NewRotation.Quaternion(), ComponentToWorld.GetRotation(), ComponentToWorld.GetScale3D()))
+				if (Actor && Actor->GetRootComponent() == this && AreSymmetricRotations(NewRotation.Quaternion(), ComponentToWorld.GetRotation(), ComponentToWorld.GetScale3D()))
 				{
 					// We know we are not overlapping any new components at the end location.
 					// Keep known overlapping child components, as long as we know their overlap status could not have changed (ie they are positioned relative to us).
@@ -1514,7 +1547,7 @@ bool UPrimitiveComponent::MoveComponent( const FVector& Delta, const FRotator& N
 			if ( (BlockingHit.Time < 1.f) && !IsZeroExtent() )
 			{
 				// this is sole debug purpose to find how capsule trace information was when hit 
-				// to resolve stuck or improve our movement system - To turn this on, use DebugCapsuleTracePawn
+				// to resolve stuck or improve our movement system - To turn this on, use DebugCapsuleSweepPawn
 				APawn const* const ActorPawn = (Actor ? Cast<APawn>(Actor) : NULL);
 				if (ActorPawn && ActorPawn->Controller && ActorPawn->Controller->IsLocalPlayerController())
 				{
@@ -1538,7 +1571,7 @@ bool UPrimitiveComponent::MoveComponent( const FVector& Delta, const FRotator& N
 		{
 			// We didn't move, and any rotation that doesn't change our overlap bounds means we already know what we are overlapping at this point.
 			// Can only do this if current known overlaps are valid (not deferring updates)
-			if (CVarAllowCachedOverlaps->GetInt() && Actor->GetRootComponent() == this && !IsDeferringMovementUpdates())
+			if (CVarAllowCachedOverlaps->GetInt() && Actor && Actor->GetRootComponent() == this && !IsDeferringMovementUpdates())
 			{
 				// Only if we know that we won't change our overlap status with children by moving.
 				if (AreSymmetricRotations(NewRotation.Quaternion(), ComponentToWorld.GetRotation(), ComponentToWorld.GetScale3D()) &&
@@ -1598,7 +1631,7 @@ bool UPrimitiveComponent::MoveComponent( const FVector& Delta, const FRotator& N
 		}
 		else
 		{
-			UE_LOG(LogPrimitiveComponent, Log, TEXT("%10f executing MoveComponent for %s"), MSec, *GetFullName(), *Actor->GetFullName() );
+			UE_LOG(LogPrimitiveComponent, Log, TEXT("%10f executing MoveComponent for %s"), MSec, *GetFullName() );
 		}
 	}
 #endif
@@ -2095,16 +2128,35 @@ const TArray<FOverlapInfo>& UPrimitiveComponent::GetOverlapInfos() const
 
 void UPrimitiveComponent::IgnoreActorWhenMoving(AActor* Actor, bool bShouldIgnore)
 {
-	if(bShouldIgnore)
-	{
-		MoveIgnoreActors.AddUnique(Actor);
-	}
-	else
+	// Clean up stale references
+	MoveIgnoreActors.RemoveSwap(NULL);
+
+	// Add/Remove the actor from the list
+	if (Actor)
 	{
-		MoveIgnoreActors.Remove(Actor);
+		if (bShouldIgnore)
+		{
+			MoveIgnoreActors.AddUnique(Actor);
+		}
+		else
+		{
+			MoveIgnoreActors.RemoveSingleSwap(Actor);
+		}
 	}
 }
 
+TArray<TWeakObjectPtr<AActor> > & UPrimitiveComponent::GetMoveIgnoreActors()
+{
+	// Clean up stale references
+	MoveIgnoreActors.RemoveSwap(NULL);
+	return MoveIgnoreActors;
+}
+
+void UPrimitiveComponent::ClearMoveIgnoreActors()
+{
+	MoveIgnoreActors.Empty();
+}
+
 
 void UPrimitiveComponent::UpdateOverlaps(TArray<FOverlapInfo> const* PendingOverlaps, bool bDoNotifies, const TArray<FOverlapInfo>* OverlapsAtEndLocation)
 {
@@ -2160,14 +2212,13 @@ void UPrimitiveComponent::UpdateOverlaps(TArray<FOverlapInfo> const* PendingOver
 					{
 						const FOverlapResult& Result = Overlaps[ResultIdx];
 
-						UPrimitiveComponent const* const HitComp = Result.Component.Get();
+						UPrimitiveComponent* const HitComp = Result.Component.Get();
 						if (!Result.bBlockingHit && HitComp && (HitComp != this) && HitComp->bGenerateOverlapEvents)
 						{
-							AActor* const ResultActor = Result.GetActor();
-							if ( (ResultActor != NULL) && !ResultActor->IsAttachedTo(MyActor) && !MyActor->IsAttachedTo(ResultActor) && ResultActor != MyWorld->GetWorldSettings() && ResultActor->bActorInitialized)
+							if (!ShouldIgnoreOverlapResult(MyWorld, MyActor, *this, Result.GetActor(), *HitComp))
 							{
-								NewOverlappingComponents.Add(FOverlapInfo(Result.Component.Get(), Result.ItemIndex));		// don't need to add unique unless the overlap check can return dupes
-							}						
+								NewOverlappingComponents.Add(FOverlapInfo(HitComp, Result.ItemIndex));		// don't need to add unique unless the overlap check can return dupes
+							}
 						}
 					}
 				}
diff --git a/Engine/Source/Runtime/Engine/Private/SceneComponent.cpp b/Engine/Source/Runtime/Engine/Private/SceneComponent.cpp
--- a/Engine/Source/Runtime/Engine/Private/SceneComponent.cpp
+++ b/Engine/Source/Runtime/Engine/Private/SceneComponent.cpp
@@ -828,6 +828,12 @@ void USceneComponent::AttachTo(class USceneComponent* Parent, FName InSocketName
 
 		// calculate transform with new attachment condition
 		UpdateComponentToWorld();
+
+		// Update overlaps, in case location changed or overlap state depends on attachment.
+		if (IsRegistered())
+		{
+			UpdateOverlaps();
+		}
 	}
 }
 
@@ -875,6 +881,12 @@ void USceneComponent::DetachFromParent(bool bMaintainWorldPosition)
 
 		// calculate transform with new attachment condition
 		UpdateComponentToWorld();
+
+		// Update overlaps, in case location changed or overlap state depends on attachment.
+		if (IsRegistered())
+		{
+			UpdateOverlaps();
+		}
 	}
 }
 
