Fixed Character and PaperCharacter not properly setting Mesh/Sprite tick pre-requisite to happen after MovementComponent tick.
Removed SkeletalMeshComponent dependency on Character.

diff --git a/Engine/Plugins/Experimental/Paper2D/Source/Paper2D/Classes/PaperCharacter.h b/Engine/Plugins/Experimental/Paper2D/Source/Paper2D/Classes/PaperCharacter.h
--- a/Engine/Plugins/Experimental/Paper2D/Source/Paper2D/Classes/PaperCharacter.h
+++ b/Engine/Plugins/Experimental/Paper2D/Source/Paper2D/Classes/PaperCharacter.h
@@ -17,4 +17,6 @@ class PAPER2D_API APaperCharacter : public ACharacter
 	/** The main skeletal mesh associated with this Character (optional sub-object). */
 	UPROPERTY(Category = Character, VisibleAnywhere, BlueprintReadOnly)
 	TSubobjectPtr<class UPaperFlipbookComponent> Sprite;
+
+	virtual void PostInitializeComponents() override;
 };
diff --git a/Engine/Plugins/Experimental/Paper2D/Source/Paper2D/Private/PaperCharacter.cpp b/Engine/Plugins/Experimental/Paper2D/Source/Paper2D/Private/PaperCharacter.cpp
--- a/Engine/Plugins/Experimental/Paper2D/Source/Paper2D/Private/PaperCharacter.cpp
+++ b/Engine/Plugins/Experimental/Paper2D/Source/Paper2D/Private/PaperCharacter.cpp
@@ -19,16 +19,26 @@ APaperCharacter::APaperCharacter(const class FPostConstructInitializeProperties&
 		Sprite->bOwnerNoSee = false;
 		Sprite->bAffectDynamicIndirectLighting = true;
 		Sprite->PrimaryComponentTick.TickGroup = TG_PrePhysics;
-
-		// force tick after movement component updates
-		if (CharacterMovement)
-		{
-			Sprite->PrimaryComponentTick.AddPrerequisite(this, CharacterMovement->PrimaryComponentTick);
-		}
-
 		Sprite->AttachParent = CapsuleComponent;
 		static FName CollisionProfileName(TEXT("CharacterMesh"));
 		Sprite->SetCollisionProfileName(CollisionProfileName);
 		Sprite->bGenerateOverlapEvents = false;
 	}
 }
+
+void APaperCharacter::PostInitializeComponents()
+{
+	Super::PostInitializeComponents();
+
+	if (!IsPendingKill())
+	{
+		if (Sprite)
+		{
+			// force animation tick after movement component updates
+			if (Sprite->PrimaryComponentTick.bCanEverTick && CharacterMovement.IsValid())
+			{
+				Sprite->PrimaryComponentTick.AddPrerequisite(CharacterMovement, CharacterMovement->PrimaryComponentTick);
+			}
+		}
+	}
+}
\ No newline at end of file
diff --git a/Engine/Source/Runtime/Engine/Classes/Components/SkeletalMeshComponent.h b/Engine/Source/Runtime/Engine/Classes/Components/SkeletalMeshComponent.h
--- a/Engine/Source/Runtime/Engine/Classes/Components/SkeletalMeshComponent.h
+++ b/Engine/Source/Runtime/Engine/Classes/Components/SkeletalMeshComponent.h
@@ -474,6 +474,11 @@ public:
 	 * Misc 
 	 */
 	
+	/** If true TickPose() will not be called from the Component's TickComponent function.
+	* It will instead be called from Autonomous networking updates. See ACharacter. */
+	UPROPERTY(Transient)
+	uint32 bAutonomousTickPose : 1;
+
 	/** If true, force the mesh into the reference pose - is an optimization. */
 	UPROPERTY()
 	uint32 bForceRefpose:1;
diff --git a/Engine/Source/Runtime/Engine/Classes/GameFramework/Character.h b/Engine/Source/Runtime/Engine/Classes/GameFramework/Character.h
--- a/Engine/Source/Runtime/Engine/Classes/GameFramework/Character.h
+++ b/Engine/Source/Runtime/Engine/Classes/GameFramework/Character.h
@@ -348,6 +348,7 @@ public:
 	virtual void TurnOff() override;
 	virtual void Restart() override;
 	virtual void PawnClientRestart() override;
+	virtual void PossessedBy(AController* NewController) override;
 	virtual void UnPossessed() override;
 	virtual void SetupPlayerInputComponent(class UInputComponent* InputComponent) override;
 	virtual void DisplayDebug(class UCanvas* Canvas, const FDebugDisplayInfo& DebugDisplay, float& YL, float& YPos) override;
diff --git a/Engine/Source/Runtime/Engine/Private/Character.cpp b/Engine/Source/Runtime/Engine/Private/Character.cpp
--- a/Engine/Source/Runtime/Engine/Private/Character.cpp
+++ b/Engine/Source/Runtime/Engine/Private/Character.cpp
@@ -89,11 +89,6 @@ ACharacter::ACharacter(const class FPostConstructInitializeProperties& PCIP)
 		Mesh->bCastDynamicShadow = true;
 		Mesh->bAffectDynamicIndirectLighting = true;
 		Mesh->PrimaryComponentTick.TickGroup = TG_PrePhysics;
-		// force tick after movement component updates
-		if( CharacterMovement )
-		{
-			Mesh->PrimaryComponentTick.AddPrerequisite(this, CharacterMovement->PrimaryComponentTick); 
-		}
 		Mesh->bChartDistanceFactor = true;
 		Mesh->AttachParent = CapsuleComponent;
 		static FName CollisionProfileName(TEXT("CharacterMesh"));
@@ -112,6 +107,12 @@ void ACharacter::PostInitializeComponents()
 		if (Mesh)
 		{
 			BaseTranslationOffset = Mesh->RelativeLocation;
+
+			// force animation tick after movement component updates
+			if (Mesh->PrimaryComponentTick.bCanEverTick && CharacterMovement.IsValid())
+			{
+				Mesh->PrimaryComponentTick.AddPrerequisite(CharacterMovement, CharacterMovement->PrimaryComponentTick);
+			}
 		}
 
 		if (CharacterMovement.IsValid() && CapsuleComponent.IsValid())
@@ -678,6 +679,17 @@ void ACharacter::PawnClientRestart()
 	Super::PawnClientRestart();
 }
 
+void ACharacter::PossessedBy(AController* NewController)
+{
+	Super::PossessedBy(NewController);
+
+	// If we are controlled remotely, set animation timing to be driven by client's network updates. So timing and events remain in sync.
+	if (Mesh && (GetRemoteRole() == ROLE_AutonomousProxy))
+	{
+		Mesh->bAutonomousTickPose = true;
+	}
+}
+
 void ACharacter::UnPossessed()
 {
 	Super::UnPossessed();
@@ -687,6 +699,12 @@ void ACharacter::UnPossessed()
 		CharacterMovement->ResetPredictionData_Client();
 		CharacterMovement->ResetPredictionData_Server();
 	}
+
+	// We're no longer controlled remotely, resume regular ticking of animations.
+	if (Mesh)
+	{
+		Mesh->bAutonomousTickPose = false;
+	}
 }
 
 
@@ -699,6 +717,12 @@ void ACharacter::TornOff()
 		CharacterMovement->ResetPredictionData_Client();
 		CharacterMovement->ResetPredictionData_Server();
 	}
+
+	// We're no longer controlled remotely, resume regular ticking of animations.
+	if (Mesh)
+	{
+		Mesh->bAutonomousTickPose = false;
+	}
 }
 
 
diff --git a/Engine/Source/Runtime/Engine/Private/CharacterMovementComponent.cpp b/Engine/Source/Runtime/Engine/Private/CharacterMovementComponent.cpp
--- a/Engine/Source/Runtime/Engine/Private/CharacterMovementComponent.cpp
+++ b/Engine/Source/Runtime/Engine/Private/CharacterMovementComponent.cpp
@@ -6159,14 +6159,17 @@ void UCharacterMovementComponent::RegisterComponentTickFunctions(bool bRegister)
 
 void UCharacterMovementComponent::TickCharacterPose(float DeltaTime)
 {
-	check(CharacterOwner && CharacterOwner->Mesh)
+	check(CharacterOwner && CharacterOwner->Mesh);
 
 	CharacterOwner->Mesh->TickPose(DeltaTime);
 
 	// Grab root motion now that we have ticked the pose
-	if (UAnimInstance* AnimInstance = CharacterOwner->Mesh->GetAnimInstance())
+	if (CharacterOwner->IsPlayingRootMotion())
 	{
-		RootMotionParams.Accumulate(AnimInstance->ConsumeExtractedRootMotion());
+		if (UAnimInstance * AnimInstance = CharacterOwner->Mesh->GetAnimInstance())
+		{
+			RootMotionParams.Accumulate(AnimInstance->ConsumeExtractedRootMotion());
+		}
 	}
 }
 
diff --git a/Engine/Source/Runtime/Engine/Private/SkeletalMeshComponent.cpp b/Engine/Source/Runtime/Engine/Private/SkeletalMeshComponent.cpp
--- a/Engine/Source/Runtime/Engine/Private/SkeletalMeshComponent.cpp
+++ b/Engine/Source/Runtime/Engine/Private/SkeletalMeshComponent.cpp
@@ -16,7 +16,6 @@
 #include "AnimTree.h"
 #include "Animation/AnimNodeBase.h"
 #include "Animation/VertexAnim/VertexAnimation.h"
-#include "GameFramework/Character.h"
 #include "Particles/ParticleSystemComponent.h"
 #include "Animation/AnimSingleNodeInstance.h"
 
@@ -510,18 +509,10 @@ bool USkeletalMeshComponent::ShouldUpdateTransform(bool bLODHasChanged) const
 
 bool USkeletalMeshComponent::ShouldTickPose() const
 {
-	// Characters playing Root Motion will tick pose themselves before physics.
-	ACharacter * CharacterOwner = Cast<ACharacter>(GetOwner());
-	const bool bSkipBecauseOfRootMotion = CharacterOwner && CharacterOwner->IsPlayingRootMotion();
-
 	// When we stop root motion we go back to ticking after CharacterMovement. Unfortunately that means that we could tick twice that frame.
 	// So only enforce a single tick per frame.
 	const bool bAlreadyTickedThisFrame = (LastTickTime == GetWorld()->TimeSeconds);
-
-	// Remote Clients on the Server will always tick animations as updates from the client comes in. To use Client's delta time.
-	const bool bRemoteClientOnServer = CharacterOwner && (CharacterOwner->Role == ROLE_Authority) && CharacterOwner->Controller && !CharacterOwner->Controller->IsLocalController();
-
-	return (Super::ShouldTickPose() && IsRegistered() && AnimScriptInstance && !bPauseAnims && GetWorld()->AreActorsInitialized() && !bNoSkeletonUpdate && !bSkipBecauseOfRootMotion && !bRemoteClientOnServer && !bAlreadyTickedThisFrame);
+	return (Super::ShouldTickPose() && IsRegistered() && AnimScriptInstance && !bAutonomousTickPose && !bPauseAnims && GetWorld()->AreActorsInitialized() && !bNoSkeletonUpdate && !bAlreadyTickedThisFrame);
 }
 
 void USkeletalMeshComponent::TickPose(float DeltaTime)
