Don't rename a ChildActor if the desired name is already taken (thus crashing)
When destroying a ChildActor, "reserve" the fname number so that normal SpawnActors won't inadvertently take over the name

diff --git a/Engine/Source/Runtime/Engine/Private/ChildActorComponent.cpp b/Engine/Source/Runtime/Engine/Private/ChildActorComponent.cpp
--- a/Engine/Source/Runtime/Engine/Private/ChildActorComponent.cpp
+++ b/Engine/Source/Runtime/Engine/Private/ChildActorComponent.cpp
@@ -95,7 +95,12 @@ void UChildActorComponent::ApplyComponentInstanceData(TSharedPtr<FComponentInsta
 	ChildActorName = ChildActorInstanceData->ChildActorName;
 	if (ChildActor)
 	{
-		ChildActor->Rename(*ChildActorName.ToString(), NULL, REN_DoNotDirty);
+		// Only rename if it is safe to
+		const FString ChildActorNameString = ChildActorName.ToString();
+		if (ChildActor->Rename(*ChildActorName.ToString(), NULL, REN_Test))
+		{
+			ChildActor->Rename(*ChildActorName.ToString(), NULL, REN_DoNotDirty);
+		}
 
 		USceneComponent* ChildActorRoot = ChildActor->GetRootComponent();
 		if (ChildActorRoot)
@@ -195,6 +200,11 @@ void UChildActorComponent::DestroyChildActor()
 			if(World != NULL)
 			{
 				UClass* ChildClass = ChildActor->GetClass();
+
+				// We would like to make certain that our name is not going to accidentally get taken from us while we're destroyed
+				// so we increment ClassUnique beyond our index to be certain of it.  This is ... a bit hacky.
+				ChildClass->ClassUnique = FMath::Max(ChildClass->ClassUnique, ChildActor->GetFName().GetNumber());
+
 				const FString ObjectBaseName = FString::Printf(TEXT("DESTROYED_%s_CHILDACTOR"), *ChildClass->GetName());
 				ChildActor->Rename(*MakeUniqueObjectName(GetOuter(), ChildClass, *ObjectBaseName).ToString(), NULL, REN_DoNotDirty);
 				World->DestroyActor(ChildActor);
