#ue4 - Optimize USceneComponent::CalcNewComponentToWorld() when there is no attach parent.

- Removes useless multiply by FTransform::Identity, branches, and potential copies when bAbsoluteTranslation etc are set.

diff --git a/Engine/Source/Runtime/Engine/Private/SceneComponent.cpp b/Engine/Source/Runtime/Engine/Private/SceneComponent.cpp
--- a/Engine/Source/Runtime/Engine/Private/SceneComponent.cpp
+++ b/Engine/Source/Runtime/Engine/Private/SceneComponent.cpp
@@ -37,34 +37,33 @@ USceneComponent::USceneComponent(const class FPostConstructInitializeProperties&
 
 FTransform USceneComponent::CalcNewComponentToWorld(const FTransform& NewRelativeTransform, const USceneComponent * Parent) const
 {
-	// With no attachment
-	FTransform ParentToWorld = FTransform::Identity;
-
 	Parent = Parent ? Parent : AttachParent;
-
 	if (Parent != NULL)
 	{
-		ParentToWorld = Parent->GetSocketTransform(AttachSocketName);
-	}
+		const FTransform ParentToWorld = Parent->GetSocketTransform(AttachSocketName);
+		FTransform NewCompToWorld = NewRelativeTransform * ParentToWorld;
 
-	FTransform NewCompToWorld = NewRelativeTransform * ParentToWorld;
+		if(bAbsoluteLocation)
+		{
+			NewCompToWorld.SetTranslation(NewRelativeTransform.GetTranslation());
+		}
 
-	if(bAbsoluteLocation)
-	{
-		NewCompToWorld.SetTranslation(NewRelativeTransform.GetTranslation());
-	}
+		if(bAbsoluteRotation)
+		{
+			NewCompToWorld.SetRotation(NewRelativeTransform.GetRotation());
+		}
 
-	if(bAbsoluteRotation)
-	{
-		NewCompToWorld.SetRotation(NewRelativeTransform.GetRotation());
-	}
+		if(bAbsoluteScale)
+		{
+			NewCompToWorld.SetScale3D(NewRelativeTransform.GetScale3D());
+		}
 
-	if(bAbsoluteScale)
+		return NewCompToWorld;
+	}
+	else
 	{
-		NewCompToWorld.SetScale3D(NewRelativeTransform.GetScale3D());
+		return NewRelativeTransform;
 	}
-
-	return NewCompToWorld;
 }
 
 void USceneComponent::OnUpdateTransform(bool bSkipPhysicsMove)
