Fixed TTP #344520 : Play From Here couldn't spawn Pawn problem

diff --git a/Engine/Source/Runtime/Engine/Private/Landscape/LandscapeCollision.cpp b/Engine/Source/Runtime/Engine/Private/Landscape/LandscapeCollision.cpp
--- a/Engine/Source/Runtime/Engine/Private/Landscape/LandscapeCollision.cpp
+++ b/Engine/Source/Runtime/Engine/Private/Landscape/LandscapeCollision.cpp
@@ -196,6 +196,7 @@ void ULandscapeHeightfieldCollisionComponent::CreatePhysicsState()
 				}
 
 				// Set body instance data
+				BodyInstance.PhysxUserData = FPhysxUserData(&BodyInstance);
 				BodyInstance.OwnerComponent = this;
 				BodyInstance.SceneIndexSync = PhysScene->PhysXSceneIndex[PST_Sync];
 				BodyInstance.SceneIndexAsync = PhysScene->HasAsyncScene() ? PhysScene->PhysXSceneIndex[PST_Async] : 0;
@@ -730,6 +731,7 @@ void ULandscapeMeshCollisionComponent::CreatePhysicsState()
 #endif// WITH_EDITOR
 
 				// Set body instance data
+				BodyInstance.PhysxUserData = FPhysxUserData(&BodyInstance);
 				BodyInstance.OwnerComponent = this;
 				BodyInstance.SceneIndexSync = PhysScene->PhysXSceneIndex[PST_Sync];
 				BodyInstance.SceneIndexAsync = PhysScene->HasAsyncScene() ? PhysScene->PhysXSceneIndex[PST_Async] : 0;
