Landscape : Fixed TTP #344509 : Removed Physics log Unknown geom type

diff --git a/Engine/Source/Runtime/Engine/Private/PhysicsEngine/BodyInstance.cpp b/Engine/Source/Runtime/Engine/Private/PhysicsEngine/BodyInstance.cpp
--- a/Engine/Source/Runtime/Engine/Private/PhysicsEngine/BodyInstance.cpp
+++ b/Engine/Source/Runtime/Engine/Private/PhysicsEngine/BodyInstance.cpp
@@ -1909,6 +1909,12 @@ bool FBodyInstance::UpdateBodyScale(const FVector& InScale3D)
 				}
 				break;
 			}
+			case PxGeometryType::eHEIGHTFIELD:
+			{
+				// HeightField is only used by Landscape, which does different code path from other primitives
+				SCENE_UNLOCK_READ(PScene);
+				break;
+			}
 			default:
 			{
 					   SCENE_UNLOCK_READ(PScene);
