navigation octree refactor
- common interface for all objects that can export navigation data
- simplified registering/unregistering
- replaced navigation proxy objects with custom node parenting
- fixed corrupted entries from blueprint based components during editor move

diff --git a/Engine/Source/Editor/DetailCustomizations/Private/ActorMaterialCategory.cpp b/Engine/Source/Editor/DetailCustomizations/Private/ActorMaterialCategory.cpp
--- a/Engine/Source/Editor/DetailCustomizations/Private/ActorMaterialCategory.cpp
+++ b/Engine/Source/Editor/DetailCustomizations/Private/ActorMaterialCategory.cpp
@@ -341,7 +341,7 @@ void FActorMaterialCategory::OnMaterialChanged( UMaterialInterface* NewMaterial,
 					EditChangeObject = CastChecked<ULandscapeComponent>(CurrentComponent)->GetLandscapeProxy();
 				}
 
-				FNavigationLockContext NavUpdateLock(Actor->GetWorld());
+				FNavigationLockContext NavUpdateLock(Actor->GetWorld(), ENavigationLockReason::MaterialUpdate);
 
 				EditChangeObject->PreEditChange( MaterialProperty );
 
diff --git a/Engine/Source/Editor/LandscapeEditor/Private/LandscapeEdModeTools.h b/Engine/Source/Editor/LandscapeEditor/Private/LandscapeEdModeTools.h
--- a/Engine/Source/Editor/LandscapeEditor/Private/LandscapeEdModeTools.h
+++ b/Engine/Source/Editor/LandscapeEditor/Private/LandscapeEdModeTools.h
@@ -659,11 +659,7 @@ struct FHeightmapAccessor
 			if (CollisionComponent)
 			{
 				CollisionComponent->RecreateCollision(true);
-				UNavigationSystem* NavSys = UNavigationSystem::GetCurrent(*It);
-				if (NavSys)
-				{
-					NavSys->UpdateNavOctree(CollisionComponent);
-				}
+				UNavigationSystem::UpdateNavOctree(CollisionComponent);
 			}
 		}
 	}
diff --git a/Engine/Source/Editor/MaterialEditor/Private/MaterialEditor.cpp b/Engine/Source/Editor/MaterialEditor/Private/MaterialEditor.cpp
--- a/Engine/Source/Editor/MaterialEditor/Private/MaterialEditor.cpp
+++ b/Engine/Source/Editor/MaterialEditor/Private/MaterialEditor.cpp
@@ -1386,7 +1386,7 @@ void FMaterialEditor::UpdateOriginalMaterial()
 	// Handle propagation of the material being edited
 	else
 	{
-		FNavigationLockContext NavUpdateLock;
+		FNavigationLockContext NavUpdateLock(ENavigationLockReason::MaterialUpdate);
 
 		// Create a material update context so we can safely update materials.
 		{
diff --git a/Engine/Source/Editor/UnrealEd/Private/Editor.cpp b/Engine/Source/Editor/UnrealEd/Private/Editor.cpp
--- a/Engine/Source/Editor/UnrealEd/Private/Editor.cpp
+++ b/Engine/Source/Editor/UnrealEd/Private/Editor.cpp
@@ -2117,7 +2117,7 @@ void UEditorEngine::ApplyDeltaToActor(AActor* InActor,
 		}
 	}
 
-	FNavigationLockContext LockNavigationUpdates;
+	FNavigationLockContext LockNavigationUpdates(InActor->GetWorld(), ENavigationLockReason::ContinuousEditorMove);
 
 	bool bTranslationOnly = true;
 
diff --git a/Engine/Source/Editor/UnrealEd/Private/PreviewMaterial.cpp b/Engine/Source/Editor/UnrealEd/Private/PreviewMaterial.cpp
--- a/Engine/Source/Editor/UnrealEd/Private/PreviewMaterial.cpp
+++ b/Engine/Source/Editor/UnrealEd/Private/PreviewMaterial.cpp
@@ -171,7 +171,7 @@ void UMaterialEditorInstanceConstant::PostEditChangeProperty(FPropertyChangedEve
 	{
 		UProperty* PropertyThatChanged = PropertyChangedEvent.Property;
 
-		FNavigationLockContext NavUpdateLock;
+		FNavigationLockContext NavUpdateLock(ENavigationLockReason::MaterialUpdate);
 
 		if(PropertyThatChanged && PropertyThatChanged->GetName()==TEXT("Parent") )
 		{
diff --git a/Engine/Source/Editor/UnrealEd/Private/StaticLightingSystem/StaticLightingSystem.cpp b/Engine/Source/Editor/UnrealEd/Private/StaticLightingSystem/StaticLightingSystem.cpp
--- a/Engine/Source/Editor/UnrealEd/Private/StaticLightingSystem/StaticLightingSystem.cpp
+++ b/Engine/Source/Editor/UnrealEd/Private/StaticLightingSystem/StaticLightingSystem.cpp
@@ -108,7 +108,7 @@ void FStaticLightingManager::ProcessLightingData(bool bDiscardResults)
 
 	check(StaticLightingSystem);
 
-	FNavigationLockContext NavUpdateLock(StaticLightingSystem->GetWorld());
+	FNavigationLockContext NavUpdateLock(StaticLightingSystem->GetWorld(), ENavigationLockReason::LightingUpdate);
 
 	if (!bDiscardResults)
 	{
diff --git a/Engine/Source/Runtime/Engine/Classes/AI/Navigation/NavLinkCustomComponent.h b/Engine/Source/Runtime/Engine/Classes/AI/Navigation/NavLinkCustomComponent.h
--- a/Engine/Source/Runtime/Engine/Classes/AI/Navigation/NavLinkCustomComponent.h
+++ b/Engine/Source/Runtime/Engine/Classes/AI/Navigation/NavLinkCustomComponent.h
@@ -33,12 +33,15 @@ class ENGINE_API UNavLinkCustomComponent : public UNavRelevantComponent, public
 	virtual void OnLinkMoveFinished(class UPathFollowingComponent* PathComp) override;
 	// END INavLinkCustomInterface
 
-	// BEGIN UNavRelevantComponent
+	// BEGIN UNavRelevantComponent Interface
+	virtual void GetNavigationData(struct FNavigationRelevantData& Data) const;
+	virtual void CalcBounds();
+	// END UNavRelevantComponent Interface
+
+	// BEGIN UActorComponent Interface
 	virtual void OnRegister() override;
-	virtual void OnOwnerRegistered() override;
-	virtual void OnOwnerUnregistered() override;
-	virtual void OnApplyModifiers(struct FCompositeNavModifier& Modifiers) override;
-	// END UNavRelevantComponent
+	virtual void OnUnregister() override;
+	// END UActorComponent Interface
 
 	// BEGIN UObject Interface
 	virtual void PostLoad() override;
diff --git a/Engine/Source/Runtime/Engine/Classes/AI/Navigation/NavLinkHostInterface.h b/Engine/Source/Runtime/Engine/Classes/AI/Navigation/NavLinkHostInterface.h
--- a/Engine/Source/Runtime/Engine/Classes/AI/Navigation/NavLinkHostInterface.h
+++ b/Engine/Source/Runtime/Engine/Classes/AI/Navigation/NavLinkHostInterface.h
@@ -13,7 +13,6 @@ class UNavLinkHostInterface : public UInterface
 	GENERATED_UINTERFACE_BODY()
 };
 
-/** classes implementing INavLinkHostInterface should also implement INavRelevantActorInterface */
 class ENGINE_API INavLinkHostInterface
 {
 	GENERATED_IINTERFACE_BODY()
diff --git a/Engine/Source/Runtime/Engine/Classes/AI/Navigation/NavLinkProxy.h b/Engine/Source/Runtime/Engine/Classes/AI/Navigation/NavLinkProxy.h
--- a/Engine/Source/Runtime/Engine/Classes/AI/Navigation/NavLinkProxy.h
+++ b/Engine/Source/Runtime/Engine/Classes/AI/Navigation/NavLinkProxy.h
@@ -3,14 +3,14 @@
 #pragma once
 #include "AI/Navigation/NavLinkDefinition.h"
 #include "AI/Navigation/NavLinkHostInterface.h"
-#include "AI/Navigation/NavRelevantActorInterface.h"
+#include "AI/Navigation/NavRelevantInterface.h"
 #include "GameFramework/Actor.h"
 #include "NavLinkProxy.generated.h"
 
 DECLARE_DYNAMIC_MULTICAST_DELEGATE_TwoParams( FSmartLinkReachedSignature, class AActor*, MovingActor, const FVector&, DestinationPoint );
 
 UCLASS(Blueprintable)
-class ENGINE_API ANavLinkProxy : public AActor, public INavLinkHostInterface, public INavRelevantActorInterface
+class ENGINE_API ANavLinkProxy : public AActor, public INavLinkHostInterface, public INavRelevantInterface
 {
 	GENERATED_UCLASS_BODY()
 
@@ -40,10 +40,11 @@ class ENGINE_API ANavLinkProxy : public AActor, public INavLinkHostInterface, pu
 	TSubobjectPtr<class UBillboardComponent> SpriteComponent;
 #endif // WITH_EDITORONLY_DATA
 
-	// BEGIN INavRelevantActorInterface
-	virtual bool UpdateNavigationRelevancy() override;
-	virtual bool GetNavigationRelevantData(struct FNavigationRelevantData& Data) const override;
-	// END INavRelevantActorInterface
+	// BEGIN INavRelevantInterface
+	virtual void GetNavigationData(struct FNavigationRelevantData& Data) const;
+	virtual FBox GetNavigationBounds() const;
+	virtual bool IsNavigationRelevant() const;
+	// END INavRelevantInterface
 
 	// BEGIN INavLinkHostInterface
 	virtual bool GetNavigationLinksClasses(TArray<TSubclassOf<class UNavLinkDefinition> >& OutClasses) const override;
@@ -54,7 +55,6 @@ class ENGINE_API ANavLinkProxy : public AActor, public INavLinkHostInterface, pu
 	virtual void PostEditChangeProperty(FPropertyChangedEvent& PropertyChangedEvent) override;
 #endif // WITH_EDITOR
 
-	virtual void PostInitProperties() override;
 	virtual FBox GetComponentsBoundingBox(bool bNonColliding = false) const override;
 
 	//////////////////////////////////////////////////////////////////////////
diff --git a/Engine/Source/Runtime/Engine/Classes/AI/Navigation/NavModifierComponent.h b/Engine/Source/Runtime/Engine/Classes/AI/Navigation/NavModifierComponent.h
--- a/Engine/Source/Runtime/Engine/Classes/AI/Navigation/NavModifierComponent.h
+++ b/Engine/Source/Runtime/Engine/Classes/AI/Navigation/NavModifierComponent.h
@@ -16,10 +16,8 @@ class UNavModifierComponent : public UNavRelevantComponent
 	UPROPERTY(EditAnywhere, Category = Navigation)
 	TSubclassOf<UNavArea> AreaClass;
 
-	virtual void OnRegister() override;
-	virtual void OnApplyModifiers(FCompositeNavModifier& Modifiers) override;
-	virtual void OnOwnerRegistered() override;
-	virtual void OnOwnerUnregistered() override;
+	virtual void CalcBounds() override;
+	virtual void GetNavigationData(struct FNavigationRelevantData& Data) const override;
 
 protected:
 	FVector ObstacleExtent;
diff --git a/Engine/Source/Runtime/Engine/Classes/AI/Navigation/NavModifierVolume.h b/Engine/Source/Runtime/Engine/Classes/AI/Navigation/NavModifierVolume.h
--- a/Engine/Source/Runtime/Engine/Classes/AI/Navigation/NavModifierVolume.h
+++ b/Engine/Source/Runtime/Engine/Classes/AI/Navigation/NavModifierVolume.h
@@ -2,19 +2,20 @@
 
 #pragma once 
 #include "GameFramework/Volume.h"
+#include "NavRelevantInterface.h"
 #include "NavModifierVolume.generated.h"
 
 /** 
- *	allows applying selected NavAreaDefinition to navmesh, using Volume's shape
+ *	Allows applying selected AreaClass to navmesh, using Volume's shape
  */
 UCLASS(hidecategories=(Navigation, "AI|Navigation"))
-class ANavModifierVolume : public AVolume, public INavRelevantActorInterface
+class ANavModifierVolume : public AVolume, public INavRelevantInterface
 {
 	GENERATED_UCLASS_BODY()
 
 	UPROPERTY(EditAnywhere, BlueprintReadWrite, Category=Default)
 	TSubclassOf<class UNavArea> AreaClass;
 
-	virtual bool GetNavigationRelevantData(struct FNavigationRelevantData& Data) const override;
-	virtual bool UpdateNavigationRelevancy() override { SetNavigationRelevancy(true); return true; }
+	virtual void GetNavigationData(struct FNavigationRelevantData& Data) const override;
+	virtual FBox GetNavigationBounds() const;
 };
diff --git a/Engine/Source/Runtime/Engine/Classes/AI/Navigation/NavRelevantActorInterface.h b/Engine/Source/Runtime/Engine/Classes/AI/Navigation/NavRelevantActorInterface.h
deleted file mode 100644
diff --git a/Engine/Source/Runtime/Engine/Classes/AI/Navigation/NavRelevantComponent.h b/Engine/Source/Runtime/Engine/Classes/AI/Navigation/NavRelevantComponent.h
--- a/Engine/Source/Runtime/Engine/Classes/AI/Navigation/NavRelevantComponent.h
+++ b/Engine/Source/Runtime/Engine/Classes/AI/Navigation/NavRelevantComponent.h
@@ -2,30 +2,33 @@
 
 #pragma once
 #include "Components/ActorComponent.h"
+#include "NavRelevantInterface.h"
 #include "NavRelevantComponent.generated.h"
 
 UCLASS()
-class ENGINE_API UNavRelevantComponent : public UActorComponent
+class ENGINE_API UNavRelevantComponent : public UActorComponent, public INavRelevantInterface
 {
 	GENERATED_UCLASS_BODY()
 
-	virtual void OnOwnerRegistered();
-	virtual void OnOwnerUnregistered();
-
+	// Begin UActorComponent Interface
 	virtual void OnRegister() override;
 	virtual void OnUnregister() override;
+	// End UActorComponent Interface
 
-	virtual void OnApplyModifiers(struct FCompositeNavModifier& Modifiers);
+	// Begin INavRelevantInterface Interface
+	virtual FBox GetNavigationBounds() const override;
+	virtual bool IsNavigationRelevant() const override;
+	// End INavRelevantInterface Interface
 
-	UFUNCTION(BlueprintCallable, Category = "Navigation")
-	void SetNavigationRelevancy(bool bRelevant);
+	virtual void CalcBounds();
 
-	bool IsNavigationRelevant() const { return bNavigationRelevant; }
+	UFUNCTION(BlueprintCallable, Category="Navigation")
+	void SetNavigationRelevancy(bool bRelevant);
 
 	/** force refresh in navigation octree */
 	void RefreshNavigationModifiers();
 
-	/** bounds when owner is providing per component data */
+	/** bounds for navigation octree */
 	FBox Bounds;
 
 protected:
diff --git a/Engine/Source/Runtime/Engine/Classes/AI/Navigation/NavRelevantInterface.h b/Engine/Source/Runtime/Engine/Classes/AI/Navigation/NavRelevantInterface.h
new file mode 100644
--- /dev/null
+++ b/Engine/Source/Runtime/Engine/Classes/AI/Navigation/NavRelevantInterface.h
@@ -0,0 +1,30 @@
+// Copyright 1998-2014 Epic Games, Inc. All Rights Reserved.
+
+#pragma once
+#include "Interface.h"
+#include "NavRelevantInterface.generated.h"
+
+UINTERFACE(MinimalAPI, meta=(CannotImplementInterfaceInBlueprint))
+class UNavRelevantInterface : public UInterface
+{
+	GENERATED_UINTERFACE_BODY()
+};
+
+class INavRelevantInterface
+{
+	GENERATED_IINTERFACE_BODY()
+
+	/** Prepare navigation modifiers */
+	virtual void GetNavigationData(struct FNavigationRelevantData& Data) const {}
+
+	/** Get bounds for navigation octree */
+	virtual FBox GetNavigationBounds() const { return FBox(0); }
+
+	/** Are modifiers active? */
+	virtual bool IsNavigationRelevant() const { return true; }
+
+	/** Get navigation parent
+	 *  Adds modifiers to existing octree node, GetNavigationBounds and IsNavigationRelevant won't be checked
+	 */
+	virtual UObject* GetNavigationParent() const { return NULL; }
+};
diff --git a/Engine/Source/Runtime/Engine/Classes/AI/Navigation/NavigationProxy.h b/Engine/Source/Runtime/Engine/Classes/AI/Navigation/NavigationProxy.h
deleted file mode 100644
diff --git a/Engine/Source/Runtime/Engine/Classes/AI/Navigation/NavigationSystem.h b/Engine/Source/Runtime/Engine/Classes/AI/Navigation/NavigationSystem.h
--- a/Engine/Source/Runtime/Engine/Classes/AI/Navigation/NavigationSystem.h
+++ b/Engine/Source/Runtime/Engine/Classes/AI/Navigation/NavigationSystem.h
@@ -30,6 +30,7 @@ class AController;
 class UCrowdManager;
 class UNavArea;
 class INavLinkCustomInterface;
+class INavRelevantInterface;
 class FNavigationOctree;
 class ANavMeshBoundsVolume;
 class FNavDataGenerator;
@@ -72,25 +73,48 @@ namespace FNavigationSystem
 
 struct FNavigationSystemExec: public FSelfRegisteringExec
 {
-	FNavigationSystemExec()
-	{
-	}
-
 	// Begin FExec Interface
 	virtual bool Exec(UWorld* Inworld, const TCHAR* Cmd, FOutputDevice& Ar) override;
 	// End FExec Interface
 };
 
+namespace ENavigationLockReason
+{
+	enum Type
+	{
+		Unknown					= 1 << 0,
+		AllowUnregister			= 1 << 1,
+
+		MaterialUpdate			= 1 << 2,
+		LightingUpdate			= 1 << 3,
+		ContinuousEditorMove	= 1 << 4,
+	};
+}
+
 class ENGINE_API FNavigationLockContext
 {
 public:
-	FNavigationLockContext() : MyWorld(NULL), bSingleWorld(false) { LockUpdates(); }
-	FNavigationLockContext(UWorld* InWorld) : MyWorld(InWorld), bSingleWorld(true) { LockUpdates(); }
-	~FNavigationLockContext() { UnlockUpdates(); }
+	FNavigationLockContext(ENavigationLockReason::Type Reason = ENavigationLockReason::Unknown) 
+		: MyWorld(NULL), LockReason(Reason), bSingleWorld(false)
+	{
+		LockUpdates(); 
+	}
+
+	FNavigationLockContext(UWorld* InWorld, ENavigationLockReason::Type Reason = ENavigationLockReason::Unknown) 
+		: MyWorld(InWorld), LockReason(Reason), bSingleWorld(true)
+	{
+		LockUpdates(); 
+	}
+
+	~FNavigationLockContext()
+	{
+		UnlockUpdates(); 
+	}
 
 private:
 	UWorld* MyWorld;
-	bool bSingleWorld;
+	uint8 LockReason;
+	uint8 bSingleWorld : 1;
 
 	void LockUpdates();
 	void UnlockUpdates();
@@ -240,6 +264,7 @@ public:
 		OctreeUpdate_Geometry = 1,						// full update, mark dirty areas for geometry rebuild
 		OctreeUpdate_Modifiers = 2,						// quick update, mark dirty areas for modifier rebuild
 		OctreeUpdate_Refresh = 4,						// update is used for refresh, don't invalidate pending queue
+		OctreeUpdate_ParentChain = 8,					// update child nodes, don't remove anything
 	};
 
 	enum ECleanupMode
@@ -369,6 +394,8 @@ public:
 	
 	FBox GetLevelBounds(ULevel* InLevel) const;
 
+	bool IsNavigationRelevant(const AActor* TestActor) const;
+
 	/** @return default walkable area class */
 	FORCEINLINE static TSubclassOf<UNavArea> GetDefaultWalkableArea() { return DefaultWalkableArea; }
 
@@ -419,8 +446,18 @@ public:
 	//----------------------------------------------------------------------//
 	// navigation octree related functions
 	//----------------------------------------------------------------------//
-	FSetElementId RegisterNavigationRelevantActor(AActor* Actor, int32 UpdateFlags = OctreeUpdate_Default);
-	void UnregisterNavigationRelevantActor(AActor* Actor, int32 UpdateFlags = OctreeUpdate_Default);
+	static void OnComponentRegistered(UActorComponent* Comp);
+	static void OnComponentUnregistered(UActorComponent* Comp);
+	static void OnActorRegistered(AActor* Actor);
+	static void OnActorUnregistered(AActor* Actor);
+
+	/** update navoctree entry for specified actor/component */
+	static void UpdateNavOctree(AActor* Actor);
+	static void UpdateNavOctree(UActorComponent* Comp);
+	/** update all navoctree entries for actor and its components */
+	static void UpdateNavOctreeAll(AActor* Actor);
+	/** removes all navoctree entries for actor and its components */
+	static void ClearNavOctreeAll(AActor* Actor);
 
 #if WITH_NAVIGATION_GENERATOR
 	void AddDirtyArea(const FBox& NewArea, int32 Flags);
@@ -429,10 +466,6 @@ public:
 
 	const FNavigationOctree* GetNavOctree() const { return NavOctree; }
 
-	void UpdateNavOctree(AActor* Actor, int32 UpdateFlags = OctreeUpdate_Default);
-	void UpdateNavOctree(UActorComponent* ActorComp, int32 UpdateFlags = OctreeUpdate_Default);
-	void UpdateNavOctreeElement(UObject* ElementOwner, int32 UpdateFlags, const FBox& Bounds);
-
 	/** called to gather navigation relevant actors that have been created while
 	 *	Navigation System was not present */
 	void PopulateNavOctree();
@@ -445,6 +478,12 @@ public:
 	/** find all elements in navigation octree within given box (intersection) */
 	void FindElementsInNavOctree(const FBox& QueryBox, const struct FNavigationOctreeFilter& Filter, TArray<struct FNavigationOctreeElement>& Elements);
 
+	/** update single element in navoctree */
+	void UpdateNavOctreeElement(UObject* ElementOwner, INavRelevantInterface* ElementInterface, int32 UpdateFlags);
+
+	/** force updating parent node and all its children */
+	void UpdateNavOctreeParentChain(UObject* ElementOwner);
+
 	//----------------------------------------------------------------------//
 	// Custom navigation links
 	//----------------------------------------------------------------------//
@@ -530,9 +569,20 @@ public:
 	/** check whether seamless navigation building is enabled*/
 	FORCEINLINE static bool GetIsNavigationAutoUpdateEnabled() { return bNavigationAutoUpdateEnabled; }
 
-	bool AreFakeComponentChangesBeingApplied() { return bFakeComponentChangesBeingApplied; }
-	void BeginFakeComponentChanges() { bFakeComponentChangesBeingApplied = true; }
-	void EndFakeComponentChanges() { bFakeComponentChangesBeingApplied = false; }
+	FORCEINLINE bool IsNavigationRegisterLocked() const { return NavUpdateLockFlags != 0; }
+	FORCEINLINE bool IsNavigationUnregisterLocked() const { return NavUpdateLockFlags && !(NavUpdateLockFlags & ENavigationLockReason::AllowUnregister); }
+	FORCEINLINE bool IsNavigationUpdateLocked() const { return IsNavigationRegisterLocked(); }
+	FORCEINLINE void AddNavigationUpdateLock(uint8 Flags) { NavUpdateLockFlags |= Flags; }
+	FORCEINLINE void RemoveNavigationUpdateLock(uint8 Flags) { NavUpdateLockFlags &= ~Flags; }
+
+	DEPRECATED(4.5, "AreFakeComponentChangesBeingApplied is deprecated. Use IsNavigationUpdateLocked instead.")
+	bool AreFakeComponentChangesBeingApplied() { return IsNavigationUpdateLocked(); }
+	
+	DEPRECATED(4.5, "BeginFakeComponentChanges is deprecated. Use AddNavigationUpdateLock instead.")
+	void BeginFakeComponentChanges() { AddNavigationUpdateLock(ENavigationLockReason::Unknown); }
+	
+	DEPRECATED(4.5, "EndFakeComponentChanges is deprecated. Use RemoveNavigationUpdateLock instead.")
+	void EndFakeComponentChanges() { RemoveNavigationUpdateLock(ENavigationLockReason::Unknown); }
 
 	void UpdateLevelCollision(ULevel* InLevel);
 
@@ -573,7 +623,7 @@ public:
 
 protected:
 #if WITH_EDITOR
-	bool bFakeComponentChangesBeingApplied;
+	uint8 NavUpdateLockFlags;
 #endif
 
 	FNavigationSystem::EMode OperationMode;
@@ -588,6 +638,9 @@ protected:
 	
 	TMap<const UObject*, FOctreeElementId> ObjectToOctreeId;
 
+	/** Map of all objects that are tied to indexed navigation parent */
+	TMultiMap<UObject*, UObject*> OctreeChildNodesMap;
+
 	/** Map of all custom navigation links, that are relevant for path following */
 	TMap<uint32, INavLinkCustomInterface*> CustomLinksMap;
 
@@ -654,8 +707,8 @@ protected:
 
 	void OnNavigationAreaEvent(UClass* AreaClass, ENavAreaEvent::Type Event);
 	
- 	FSetElementId RegisterNavOctreeElement(UObject* ElementOwner, int32 UpdateFlags, const FBox& Bounds);
-	void UnregisterNavOctreeElement(UObject* ElementOwner, int32 UpdateFlags, const FBox& Bounds);
+ 	FSetElementId RegisterNavOctreeElement(UObject* ElementOwner, INavRelevantInterface* ElementInterface, int32 UpdateFlags);
+	void UnregisterNavOctreeElement(UObject* ElementOwner, INavRelevantInterface* ElementInterface, int32 UpdateFlags);
 	
 	/** read element data from navigation octree */
 	bool GetNavOctreeElementData(UObject* NodeOwner, int32& DirtyFlags, FBox& DirtyBounds);
@@ -665,9 +718,6 @@ protected:
 	 *	in this regard **/
 	void AddElementToNavOctree(const FNavigationDirtyElement& DirtyElement);
 
-	void AddActorElementToNavOctree(AActor& Actor, const FNavigationDirtyElement& DirtyElement);
-	void AddComponentElementToNavOctree(UActorComponent& ActorComp, const FNavigationDirtyElement& DirtyElement, const FBox& Bounds);
-
 	void SetCrowdManager(UCrowdManager* NewCrowdManager); 
 
 #if WITH_NAVIGATION_GENERATOR
diff --git a/Engine/Source/Runtime/Engine/Classes/AI/Navigation/NavigationTestingActor.h b/Engine/Source/Runtime/Engine/Classes/AI/Navigation/NavigationTestingActor.h
--- a/Engine/Source/Runtime/Engine/Classes/AI/Navigation/NavigationTestingActor.h
+++ b/Engine/Source/Runtime/Engine/Classes/AI/Navigation/NavigationTestingActor.h
@@ -144,7 +144,6 @@ class ENGINE_API ANavigationTestingActor : public AActor, public INavAgentInterf
 	void TickMe();
 #endif // WITH_EDITOR
 
-	virtual bool UpdateNavigationRelevancy() override { SetNavigationRelevancy(false); return false; }
 	// Begin INavAgentInterface Interface
 	virtual const struct FNavAgentProperties* GetNavAgentProperties() const override { return &NavAgentProps; }
 	virtual FVector GetNavAgentLocation() const override;
diff --git a/Engine/Source/Runtime/Engine/Classes/AI/Navigation/NavigationTypes.h b/Engine/Source/Runtime/Engine/Classes/AI/Navigation/NavigationTypes.h
--- a/Engine/Source/Runtime/Engine/Classes/AI/Navigation/NavigationTypes.h
+++ b/Engine/Source/Runtime/Engine/Classes/AI/Navigation/NavigationTypes.h
@@ -21,6 +21,7 @@ struct FNavigationQueryFilter;
 class AActor;
 class ANavigationData;
 class INavAgentInterface;
+class INavRelevantInterface;
 
 namespace FNavigationSystem
 {
@@ -90,21 +91,18 @@ struct FNavigationDirtyElement
 	/** object owning this element */
 	FWeakObjectPtr Owner;
 	
+	/** cached interface pointer */
+	INavRelevantInterface* NavInterface;
+
 	/** override for update flags */
 	int32 FlagsOverride;
 	
 	/** flags of already existing entry for this actor */
 	int32 PrevFlags;
 	
-	/** override to actor bounds */
-	FBox BoundsOverride;
-	
 	/** bounds of already existing entry for this actor */
 	FBox PrevBounds;
 
-	/** override bounds are set */
-	uint8 bHasBoundsOverride : 1;
-
 	/** prev flags & bounds data are set */
 	uint8 bHasPrevData : 1;
 
@@ -112,20 +110,18 @@ struct FNavigationDirtyElement
 	uint8 bInvalidRequest : 1;
 
 	FNavigationDirtyElement()
-		: FlagsOverride(0), bHasBoundsOverride(false), bHasPrevData(false), bInvalidRequest(false)
+		: NavInterface(0), FlagsOverride(0), bHasPrevData(false), bInvalidRequest(false)
 	{
-
 	}
 
-	FNavigationDirtyElement(UObject* InOwner, int32 InFlagsOverride = 0)
-		: Owner(InOwner), FlagsOverride(InFlagsOverride), bHasBoundsOverride(false), bHasPrevData(false), bInvalidRequest(false)
+	FNavigationDirtyElement(UObject* InOwner)
+		: Owner(InOwner), NavInterface(0), FlagsOverride(0), bHasPrevData(false), bInvalidRequest(false)
 	{
 	}
 
-	void SetBounds(const FBox& InBounds)
+	FNavigationDirtyElement(UObject* InOwner, INavRelevantInterface* InNavInterface, int32 InFlagsOverride = 0)
+		: Owner(InOwner), NavInterface(InNavInterface),	FlagsOverride(InFlagsOverride), bHasPrevData(false), bInvalidRequest(false)
 	{
-		BoundsOverride = InBounds;
-		bHasBoundsOverride = true;
 	}
 
 	bool operator==(const FNavigationDirtyElement& Other) const 
diff --git a/Engine/Source/Runtime/Engine/Classes/Components/InstancedStaticMeshComponent.h b/Engine/Source/Runtime/Engine/Classes/Components/InstancedStaticMeshComponent.h
--- a/Engine/Source/Runtime/Engine/Classes/Components/InstancedStaticMeshComponent.h
+++ b/Engine/Source/Runtime/Engine/Classes/Components/InstancedStaticMeshComponent.h
@@ -154,6 +154,7 @@ public:
 	virtual SIZE_T GetResourceSize(EResourceSizeMode::Type Mode) override;
 #if WITH_EDITOR
 	virtual void PostEditChangeChainProperty(FPropertyChangedChainEvent& PropertyChangedEvent) override;
+	virtual void PostEditUndo() override;
 #endif
 	//End UObject Interface
 
diff --git a/Engine/Source/Runtime/Engine/Classes/Components/PrimitiveComponent.h b/Engine/Source/Runtime/Engine/Classes/Components/PrimitiveComponent.h
--- a/Engine/Source/Runtime/Engine/Classes/Components/PrimitiveComponent.h
+++ b/Engine/Source/Runtime/Engine/Classes/Components/PrimitiveComponent.h
@@ -7,6 +7,7 @@
 #include "Components/SceneComponent.h"
 #include "SceneTypes.h"
 #include "Engine/EngineTypes.h"
+#include "AI/Navigation/NavRelevantInterface.h"
 
 #include "PrimitiveComponent.generated.h"
 
@@ -88,7 +89,7 @@ DECLARE_DYNAMIC_MULTICAST_DELEGATE_TwoParams( FComponentBeginTouchOverSignature,
 DECLARE_DYNAMIC_MULTICAST_DELEGATE_TwoParams( FComponentEndTouchOverSignature, ETouchIndex::Type, FingerIndex, UPrimitiveComponent*, TouchedComponent );
 
 UCLASS(abstract, HideCategories=(Mobility), ShowCategories=(PhysicsVolume))
-class ENGINE_API UPrimitiveComponent : public USceneComponent
+class ENGINE_API UPrimitiveComponent : public USceneComponent, public INavRelevantInterface
 {
 	GENERATED_UCLASS_BODY()
 
@@ -389,6 +390,9 @@ public:
 	UPROPERTY(EditDefaultsOnly, Category=Collision, AdvancedDisplay)
 	uint32 bCanEverAffectNavigation:1;
 
+	/** cached navigation relevancy flag for collision updates */
+	uint32 bNavigationRelevant : 1;
+
 protected:
 
 	/** Returns true if all descendant components that we can possibly collide with use relative location and rotation. */
@@ -401,9 +405,7 @@ protected:
 	float LastCheckedAllCollideableDescendantsTime;
 
 	/** if set to true then DoCustomNavigableGeometryExport will be called to collect navigable 
-	 *	geometry of this component.
-	 *	@NOTE that owner Actor needs to be "navigation relevant" (Actor->IsNavigationRelevant() == true)
-	 *	to have this mechanism triggered at all */
+	 *	geometry of this component. */
 	TEnumAsByte<EHasCustomNavigableGeometry::Type> bHasCustomNavigableGeometry;
 
 	/** Next id to be used by a component. */
@@ -1119,7 +1121,6 @@ public:
 	virtual ECollisionChannel GetCollisionObjectType() const override;
 	virtual const FCollisionResponseContainer & GetCollisionResponseToChannels() const override;
 	virtual FVector GetComponentVelocity() const override;
-	virtual bool IsNavigationRelevant(bool bSkipCollisionEnabledCheck = false) const override;
 	//End USceneComponent Interface
 
 	/**
@@ -1411,8 +1412,19 @@ public:
 		return bCanEverAffectNavigation;
 	}
 
-	/** turn off navigation relevance, must be called before component is registered! */
-	void DisableNavigationRelevance();
+	/** set value of bCanEverAffectNavigation flag and update navigation octree if needed */
+	void SetCanEverAffectNavigation(bool bRelevant);
+
+	DEPRECATED(4.5, "UPrimitiveComponent::DisableNavigationRelevance() is deprecated, use SetCanEverAffectNavigation() instead.")
+	void DisableNavigationRelevance()
+	{
+		SetCanEverAffectNavigation(false);
+	}
+
+	// Begin INavRelevantInterface Interface
+	virtual FBox GetNavigationBounds() const;
+	virtual bool IsNavigationRelevant() const;
+	// End INavRelevantInterface Interface
 
 	FORCEINLINE EHasCustomNavigableGeometry::Type HasCustomNavigableGeometry() const { return bHasCustomNavigableGeometry; }
 
diff --git a/Engine/Source/Runtime/Engine/Classes/Components/SceneComponent.h b/Engine/Source/Runtime/Engine/Classes/Components/SceneComponent.h
--- a/Engine/Source/Runtime/Engine/Classes/Components/SceneComponent.h
+++ b/Engine/Source/Runtime/Engine/Classes/Components/SceneComponent.h
@@ -707,9 +707,6 @@ public:
 	/** Get the extent used when placing this component in the editor, used for 'pulling back' hit. */
 	virtual FBoxSphereBounds GetPlacementExtent() const;
 
-	/** Is component (not including attachments) relevant for navigation updates? */
-	virtual bool IsNavigationRelevant(bool bSkipCollisionEnabledCheck = false) const;
-
 #if WITH_BODY_WELDING
 
 	/**
@@ -740,9 +737,6 @@ protected:
 
 	/** Called after changing transform, tries to update navigation octree */
 	void UpdateNavigationData();
-
-	/** Check if given component or any of its attachments are relevant for navigation updates */
-	static bool CheckNavigationRelevancy(USceneComponent* TestComponent);
 };
 
 
diff --git a/Engine/Source/Runtime/Engine/Classes/Engine/BlockingVolume.h b/Engine/Source/Runtime/Engine/Classes/Engine/BlockingVolume.h
--- a/Engine/Source/Runtime/Engine/Classes/Engine/BlockingVolume.h
+++ b/Engine/Source/Runtime/Engine/Classes/Engine/BlockingVolume.h
@@ -14,8 +14,6 @@ class ABlockingVolume : public AVolume
 {
 	GENERATED_UCLASS_BODY()
 
-	virtual bool UpdateNavigationRelevancy() override;
-
 	// Begin UObject interface.
 #if WITH_EDITOR
 	virtual void LoadedFromAnotherClass(const FName& OldClassName) override;
diff --git a/Engine/Source/Runtime/Engine/Classes/Engine/Brush.h b/Engine/Source/Runtime/Engine/Classes/Engine/Brush.h
--- a/Engine/Source/Runtime/Engine/Classes/Engine/Brush.h
+++ b/Engine/Source/Runtime/Engine/Classes/Engine/Brush.h
@@ -161,7 +161,6 @@ class ABrush : public AActor
 	ENGINE_API virtual void Destroyed() override;
 	ENGINE_API virtual void PostRegisterAllComponents() override;
 	
-	virtual bool UpdateNavigationRelevancy() override { SetNavigationRelevancy(false); return false; }
 	virtual bool IsLevelBoundsRelevant() const override;
 
 #if WITH_EDITOR
diff --git a/Engine/Source/Runtime/Engine/Classes/Engine/StaticMeshActor.h b/Engine/Source/Runtime/Engine/Classes/Engine/StaticMeshActor.h
--- a/Engine/Source/Runtime/Engine/Classes/Engine/StaticMeshActor.h
+++ b/Engine/Source/Runtime/Engine/Classes/Engine/StaticMeshActor.h
@@ -1,7 +1,6 @@
 // Copyright 1998-2014 Epic Games, Inc. All Rights Reserved.
 
 #pragma once
-#include "AI/Navigation/NavRelevantActorInterface.h"
 #include "StaticMeshActor.generated.h"
 
 /**
@@ -9,7 +8,7 @@
  * Note that PostInitializeComponents() is not called for StaticMeshActors
  */
 UCLASS(hidecategories=(Input), showcategories=("Input|MouseInput", "Input|TouchInput"), ConversionRoot, meta=(ChildCanTick))
-class ENGINE_API AStaticMeshActor : public AActor, public INavRelevantActorInterface
+class ENGINE_API AStaticMeshActor : public AActor
 {
 	GENERATED_UCLASS_BODY()
 
diff --git a/Engine/Source/Runtime/Engine/Classes/Foliage/InstancedFoliageActor.h b/Engine/Source/Runtime/Engine/Classes/Foliage/InstancedFoliageActor.h
--- a/Engine/Source/Runtime/Engine/Classes/Foliage/InstancedFoliageActor.h
+++ b/Engine/Source/Runtime/Engine/Classes/Foliage/InstancedFoliageActor.h
@@ -13,7 +13,7 @@ struct FFoliageInstancePlacementInfo;
 struct FFoliageMeshInfo;
 
 UCLASS(notplaceable, hidecategories = Object, MinimalAPI, NotBlueprintable)
-class AInstancedFoliageActor : public AActor, public INavRelevantActorInterface
+class AInstancedFoliageActor : public AActor
 {
 	GENERATED_UCLASS_BODY()
 
@@ -106,8 +106,4 @@ public:
 	// Get the instanced foliage actor for the specified streaming level. Never creates a new IFA.
 	static ENGINE_API AInstancedFoliageActor* GetInstancedFoliageActorForLevel(ULevel* Level);
 #endif	//WITH_EDITOR
-
-	// Begin INavRelevantActorInterface Interface
-	virtual bool DoesSupplyPerComponentNavigationCollision() const override { return true; }
-	// End INavRelevantActorInterface Interface
 };
diff --git a/Engine/Source/Runtime/Engine/Classes/GameFramework/Actor.h b/Engine/Source/Runtime/Engine/Classes/GameFramework/Actor.h
--- a/Engine/Source/Runtime/Engine/Classes/GameFramework/Actor.h
+++ b/Engine/Source/Runtime/Engine/Classes/GameFramework/Actor.h
@@ -154,9 +154,6 @@ public:
 	uint32 bBlockInput:1;
 
 private:
-	/** Cached value whether this actor is relevant for navigation generation. Usually will be calculated in actor's constructor, but can be changed at runtime as well */
-	uint32 bNavigationRelevant:1;
-
 	/** True if this actor is currently running user construction script (used to defer component registration) */
 	uint32 bRunningUserConstructionScript:1;
 
@@ -2007,14 +2004,11 @@ public:
 	//=============================================================================
 	// Navigation related functions
 	// 
-	/** Updates bNavigationRelevant. Override for custom navigation behaviors	 */
-	virtual bool UpdateNavigationRelevancy();
-
-	/** Indicates whether this actor is to be considered by navigation system a valid actor	 */
-	FORCEINLINE bool IsNavigationRelevant() const { return bNavigationRelevant; }
 
-	/**	Sets bNavigationRelevant, and notifies navigation system if value has changed	 */
-	void SetNavigationRelevancy(const bool bNewRelevancy);
+	/** Check if owned component should be relevant for navigation
+	 *  Allows implementing master switch to disable e.g. collision export in projectiles
+	 */
+	virtual bool IsComponentRelevantForNavigation(UActorComponent* Component) const { return true; }
 
 	//=============================================================================
 	// Debugging functions
diff --git a/Engine/Source/Runtime/Engine/Classes/GameFramework/Character.h b/Engine/Source/Runtime/Engine/Classes/GameFramework/Character.h
--- a/Engine/Source/Runtime/Engine/Classes/GameFramework/Character.h
+++ b/Engine/Source/Runtime/Engine/Classes/GameFramework/Character.h
@@ -353,6 +353,7 @@ public:
 	virtual void SetupPlayerInputComponent(class UInputComponent* InputComponent) override;
 	virtual void DisplayDebug(class UCanvas* Canvas, const FDebugDisplayInfo& DebugDisplay, float& YL, float& YPos) override;
 	virtual void RecalculateBaseEyeHeight() override;
+	virtual void UpdateNavigationRelevance() override;
 	// End APawn Interface
 
 	/** Apply momentum caused by damage. */
diff --git a/Engine/Source/Runtime/Engine/Classes/GameFramework/DefaultPawn.h b/Engine/Source/Runtime/Engine/Classes/GameFramework/DefaultPawn.h
--- a/Engine/Source/Runtime/Engine/Classes/GameFramework/DefaultPawn.h
+++ b/Engine/Source/Runtime/Engine/Classes/GameFramework/DefaultPawn.h
@@ -21,6 +21,7 @@ class ENGINE_API ADefaultPawn : public APawn
 	// Begin Pawn overrides
 	virtual UPawnMovementComponent* GetMovementComponent() const override;
 	virtual void SetupPlayerInputComponent(UInputComponent* InputComponent) override;
+	virtual void UpdateNavigationRelevance() override;
 	// End Pawn overrides
 
 	/** Input callback to move forward. */
diff --git a/Engine/Source/Runtime/Engine/Classes/GameFramework/Info.h b/Engine/Source/Runtime/Engine/Classes/GameFramework/Info.h
--- a/Engine/Source/Runtime/Engine/Classes/GameFramework/Info.h
+++ b/Engine/Source/Runtime/Engine/Classes/GameFramework/Info.h
@@ -20,8 +20,6 @@ class AInfo : public AActor
 	TSubobjectPtr<class UBillboardComponent> SpriteComponent;
 #endif
 
-	virtual bool UpdateNavigationRelevancy() override { SetNavigationRelevancy(false); return false; }
-	
 	/** Indicates whether this actor should participate in level bounds calculations. */
 	virtual bool IsLevelBoundsRelevant() const override { return false; }
 };
diff --git a/Engine/Source/Runtime/Engine/Classes/GameFramework/Pawn.h b/Engine/Source/Runtime/Engine/Classes/GameFramework/Pawn.h
--- a/Engine/Source/Runtime/Engine/Classes/GameFramework/Pawn.h
+++ b/Engine/Source/Runtime/Engine/Classes/GameFramework/Pawn.h
@@ -207,7 +207,6 @@ public:
 	virtual void PostRegisterAllComponents() override;
 	virtual float TakeDamage(float Damage, struct FDamageEvent const& DamageEvent, AController* EventInstigator, AActor* DamageCauser) override;
 	virtual void BecomeViewTarget(APlayerController* PC) override;
-	virtual bool UpdateNavigationRelevancy() override; 
 	virtual void EnableInput(APlayerController* PlayerController) override;
 	virtual void DisableInput(APlayerController* PlayerController) override;
 
@@ -220,6 +219,9 @@ public:
 	UFUNCTION(BlueprintCallable, Category="Navigation")
 	void SetCanAffectNavigationGeneration(bool bNewValue);
 
+	/** update all components relevant for navigation generators to match bCanAffectNavigationGeneration flag */
+	virtual void UpdateNavigationRelevance() {}
+
 	// Begin INavAgentInterface Interface
 	virtual const struct FNavAgentProperties* GetNavAgentProperties() const override;
 	/** Basically retrieved pawn's location on navmesh */
diff --git a/Engine/Source/Runtime/Engine/Classes/Landscape/LandscapeComponent.h b/Engine/Source/Runtime/Engine/Classes/Landscape/LandscapeComponent.h
--- a/Engine/Source/Runtime/Engine/Classes/Landscape/LandscapeComponent.h
+++ b/Engine/Source/Runtime/Engine/Classes/Landscape/LandscapeComponent.h
@@ -461,6 +461,9 @@ public:
 
 	/** Extends passed region with this component section size */
 	ENGINE_API void GetComponentExtent(int32& MinX, int32& MinY, int32& MaxX, int32& MaxY) const;
+
+	/** Updates navigation properties to match landscape's master switch */
+	void UpdateNavigationRelevance();
 #endif
 
 	friend class FLandscapeComponentSceneProxy;
diff --git a/Engine/Source/Runtime/Engine/Classes/Landscape/LandscapeProxy.h b/Engine/Source/Runtime/Engine/Classes/Landscape/LandscapeProxy.h
--- a/Engine/Source/Runtime/Engine/Classes/Landscape/LandscapeProxy.h
+++ b/Engine/Source/Runtime/Engine/Classes/Landscape/LandscapeProxy.h
@@ -1,7 +1,6 @@
 // Copyright 1998-2014 Epic Games, Inc. All Rights Reserved.
 
 #pragma once
-#include "AI/Navigation/NavRelevantActorInterface.h"
 #include "Engine/EngineTypes.h"
 #include "GameFramework/Actor.h"
 #include "PhysicsEngine/BodyInstance.h"
@@ -185,7 +184,7 @@ namespace ELandscapeLODFalloff
 }
 
 UCLASS(NotPlaceable, hidecategories=(Display, Attachment, Physics, Debug, Lighting, LOD), showcategories=(Rendering, "Utilities|Transformation"), MinimalAPI)
-class ALandscapeProxy : public AActor, public INavRelevantActorInterface
+class ALandscapeProxy : public AActor
 {
 	GENERATED_UCLASS_BODY()
 
@@ -342,8 +341,6 @@ public:
 	virtual void RegisterAllComponents() override;
 	virtual void RerunConstructionScripts() override {}
 	virtual bool IsLevelBoundsRelevant() const override { return true; }
-	virtual bool UpdateNavigationRelevancy() override;
-	virtual FBox GetComponentsBoundingBox(bool bNonColliding = false) const override;
 #if WITH_EDITOR
 	virtual void Destroyed() override;
 	virtual void EditorApplyScale(const FVector& DeltaScale, const FVector* PivotLocation, bool bAltDown, bool bShiftDown, bool bCtrlDown) override;
@@ -357,10 +354,6 @@ public:
 	FGuid GetLandscapeGuid() const { return LandscapeGuid; }
 	virtual ALandscape* GetLandscapeActor();
 
-	// Begin INavRelevantActorInterface Interface
-	virtual bool DoesSupplyPerComponentNavigationCollision() const override{ return true; }
-	// End INavRelevantActorInterface Interface
-
 	// Begin UObject interface.
 	virtual void Serialize(FArchive& Ar) override;
 	static void AddReferencedObjects(UObject* InThis, FReferenceCollector& Collector);
diff --git a/Engine/Source/Runtime/Engine/Classes/PhysicsEngine/DestructibleActor.h b/Engine/Source/Runtime/Engine/Classes/PhysicsEngine/DestructibleActor.h
--- a/Engine/Source/Runtime/Engine/Classes/PhysicsEngine/DestructibleActor.h
+++ b/Engine/Source/Runtime/Engine/Classes/PhysicsEngine/DestructibleActor.h
@@ -29,7 +29,6 @@ class ADestructibleActor : public AActor
 	FActorFractureSignature OnActorFracture;
 
 	// Begin AActor interface.
-	virtual bool UpdateNavigationRelevancy() override { SetNavigationRelevancy(!!bAffectNavigation); return !!bAffectNavigation; }
 #if WITH_EDITOR
 	ENGINE_API virtual bool GetReferencedContentObjects( TArray<UObject*>& Objects ) const override;
 #endif // WITH_EDITOR
diff --git a/Engine/Source/Runtime/Engine/Private/AI/Navigation/NavLinkCustomComponent.cpp b/Engine/Source/Runtime/Engine/Private/AI/Navigation/NavLinkCustomComponent.cpp
--- a/Engine/Source/Runtime/Engine/Private/AI/Navigation/NavLinkCustomComponent.cpp
+++ b/Engine/Source/Runtime/Engine/Private/AI/Navigation/NavLinkCustomComponent.cpp
@@ -88,8 +88,34 @@ void UNavLinkCustomComponent::OnLinkMoveFinished(class UPathFollowingComponent*
 	MovingAgents.Remove(WeakPathComp);
 }
 
-void UNavLinkCustomComponent::OnOwnerRegistered()
+void UNavLinkCustomComponent::GetNavigationData(struct FNavigationRelevantData& Data) const
 {
+	FNavigationLink LinkMod = GetLinkModifier();
+	Data.Modifiers.Add(FSimpleLinkNavModifier(LinkMod, GetOwner()->GetTransform()));
+
+	if (bCreateBoxObstacle)
+	{
+		Data.Modifiers.Add(FAreaNavModifier(FBox::BuildAABB(ObstacleOffset, ObstacleExtent), GetOwner()->GetTransform(), ObstacleAreaClass));
+	}
+}
+
+void UNavLinkCustomComponent::CalcBounds()
+{
+	Bounds = FBox(0);
+	Bounds += GetStartPoint();
+	Bounds += GetEndPoint();
+
+	if (bCreateBoxObstacle)
+	{
+		FBox ObstacleBounds = FBox::BuildAABB(ObstacleOffset, ObstacleExtent);
+		Bounds += ObstacleBounds.TransformBy(GetOwner()->GetTransform());
+	}
+}
+
+void UNavLinkCustomComponent::OnRegister()
+{
+	Super::OnRegister();
+
 	UNavigationSystem* NavSys = UNavigationSystem::GetCurrent(GetWorld());
 	if (NavSys == NULL)
 	{
@@ -104,8 +130,10 @@ void UNavLinkCustomComponent::OnOwnerRegistered()
 	NavSys->RegisterCustomLink(this);
 }
 
-void UNavLinkCustomComponent::OnOwnerUnregistered()
+void UNavLinkCustomComponent::OnUnregister()
 {
+	Super::OnUnregister();
+
 	UNavigationSystem* NavSys = UNavigationSystem::GetCurrent(GetWorld());
 	if (NavSys == NULL)
 	{
@@ -116,17 +144,6 @@ void UNavLinkCustomComponent::OnOwnerUnregistered()
 	NavSys->UnregisterCustomLink(this);
 }
 
-void UNavLinkCustomComponent::OnApplyModifiers(struct FCompositeNavModifier& Modifiers)
-{
-	FNavigationLink LinkMod = GetLinkModifier();
-	Modifiers.Add(FSimpleLinkNavModifier(LinkMod, GetOwner()->GetTransform()));
-
-	if (bCreateBoxObstacle)
-	{
-		Modifiers.Add(FAreaNavModifier(FBox::BuildAABB(ObstacleOffset, ObstacleExtent), GetOwner()->GetTransform(), ObstacleAreaClass));
-	}
-}
-
 void UNavLinkCustomComponent::SetLinkData(const FVector& RelativeStart, const FVector& RelativeEnd, ENavLinkDirection::Type Direction)
 {
 	LinkRelativeStart = RelativeStart;
@@ -179,27 +196,6 @@ void UNavLinkCustomComponent::ClearNavigationObstacle()
 	RefreshNavigationModifiers();
 }
 
-void UNavLinkCustomComponent::OnRegister()
-{
-	Super::OnRegister();
-
-	const AActor* MyOwner = GetOwner();
-	if (MyOwner)
-	{
-		FBox NewBounds(0);
-		NewBounds += MyOwner->GetActorLocation();
-		NewBounds += MyOwner->GetActorLocation() + LinkRelativeStart;
-		NewBounds += MyOwner->GetActorLocation() + LinkRelativeEnd;
-		
-		if (bCreateBoxObstacle)
-		{
-			NewBounds += FBox::BuildAABB(MyOwner->GetActorLocation() + ObstacleOffset, ObstacleExtent);
-		}
-
-		Bounds = NewBounds;
-	}
-}
-
 void UNavLinkCustomComponent::SetEnabled(bool bNewEnabled)
 {
 	if (bLinkEnabled != bNewEnabled)
diff --git a/Engine/Source/Runtime/Engine/Private/AI/Navigation/NavLinkProxy.cpp b/Engine/Source/Runtime/Engine/Private/AI/Navigation/NavLinkProxy.cpp
--- a/Engine/Source/Runtime/Engine/Private/AI/Navigation/NavLinkProxy.cpp
+++ b/Engine/Source/Runtime/Engine/Private/AI/Navigation/NavLinkProxy.cpp
@@ -60,43 +60,32 @@ void ANavLinkProxy::PostEditChangeProperty(FPropertyChangedEvent& PropertyChange
 		SmartLinkComp->SetNavigationRelevancy(bSmartLinkIsRelevant);
 	}
 
-	UpdateNavigationRelevancy();
-
-	if (PropertyChangedEvent.Property && 
-		IsNavigationRelevant() && GetWorld() && GetWorld()->GetNavigationSystem())
+	const FName CategoryName = FObjectEditorUtils::GetCategoryFName(PropertyChangedEvent.Property);
+	if (CategoryName == TEXT("SimpleLink"))
 	{
-		const FName CategoryName = FObjectEditorUtils::GetCategoryFName(PropertyChangedEvent.Property);
-		if (CategoryName == TEXT("SimpleLink") ||
-			CategoryName == TEXT("SmartLink") ||
-			CategoryName == TEXT("Obstacle"))
+		UNavigationSystem* NavSys = UNavigationSystem::GetCurrent(GetWorld());
+		if (NavSys)
 		{
-			GetWorld()->GetNavigationSystem()->UpdateNavOctree(this);
+			NavSys->UpdateNavOctree(this);
 		}
 	}
 }
 #endif // WITH_EDITOR
 
-void ANavLinkProxy::PostInitProperties()
+void ANavLinkProxy::GetNavigationData(struct FNavigationRelevantData& Data) const
 {
-	Super::PostInitProperties();
-
-	// NavLinkProxy has no components so we need to update navigation relevancy manually
-	UpdateNavigationRelevancy();
+	NavigationHelper::ProcessNavLinkAndAppend(&Data.Modifiers, this, PointLinks);
+	NavigationHelper::ProcessNavLinkSegmentAndAppend(&Data.Modifiers, this, SegmentLinks);
 }
 
-bool ANavLinkProxy::UpdateNavigationRelevancy()
-{ 
-	const bool bNewRelevancy = (PointLinks.Num() > 0) || (SegmentLinks.Num() > 0) || SmartLinkComp->IsNavigationRelevant();
-	SetNavigationRelevancy(bNewRelevancy);
-	return bNewRelevancy; 
+FBox ANavLinkProxy::GetNavigationBounds() const
+{
+	return GetComponentsBoundingBox();
 }
 
-bool ANavLinkProxy::GetNavigationRelevantData(struct FNavigationRelevantData& Data) const
+bool ANavLinkProxy::IsNavigationRelevant() const
 {
-	NavigationHelper::ProcessNavLinkAndAppend(&Data.Modifiers, this, PointLinks);
-	NavigationHelper::ProcessNavLinkSegmentAndAppend(&Data.Modifiers, this, SegmentLinks);
-
-	return false;
+	return (PointLinks.Num() > 0) || (SegmentLinks.Num() > 0);
 }
 
 bool ANavLinkProxy::GetNavigationLinksClasses(TArray<TSubclassOf<class UNavLinkDefinition> >& OutClasses) const
@@ -109,12 +98,7 @@ bool ANavLinkProxy::GetNavigationLinksArray(TArray<FNavigationLink>& OutLink, TA
 	OutLink.Append(PointLinks);
 	OutSegments.Append(SegmentLinks);
 
-	if (SmartLinkComp->IsNavigationRelevant())
-	{
-		OutLink.Add(SmartLinkComp->GetLinkModifier());
-	}
-
-	return (PointLinks.Num() > 0) || (SegmentLinks.Num() > 0) || SmartLinkComp->IsNavigationRelevant();
+	return (PointLinks.Num() > 0) || (SegmentLinks.Num() > 0);
 }
 
 FBox ANavLinkProxy::GetComponentsBoundingBox(bool bNonColliding) const
diff --git a/Engine/Source/Runtime/Engine/Private/AI/Navigation/NavModifierComponent.cpp b/Engine/Source/Runtime/Engine/Private/AI/Navigation/NavModifierComponent.cpp
--- a/Engine/Source/Runtime/Engine/Private/AI/Navigation/NavModifierComponent.cpp
+++ b/Engine/Source/Runtime/Engine/Private/AI/Navigation/NavModifierComponent.cpp
@@ -5,52 +5,24 @@
 #include "AI/Navigation/NavModifierComponent.h"
 #include "AI/Navigation/NavAreas/NavArea_Null.h"
 
-UNavModifierComponent::UNavModifierComponent(const FPostConstructInitializeProperties& PCIP)
-: Super(PCIP)
+UNavModifierComponent::UNavModifierComponent(const FPostConstructInitializeProperties& PCIP) : Super(PCIP)
 {
 	AreaClass = UNavArea_Null::StaticClass();
-	Bounds = FNavigationSystem::InvalidBoundingBox;
 }
 
-void UNavModifierComponent::OnOwnerRegistered()
+void UNavModifierComponent::CalcBounds()
 {
-	if (AActor* MyOwner = GetOwner())
-	{
-		MyOwner->UpdateNavigationRelevancy();
-	}
-}
-
-void UNavModifierComponent::OnOwnerUnregistered()
-{
-	if (AActor* MyOwner = GetOwner())
-	{
-		MyOwner->UpdateNavigationRelevancy();
-	}
-}
-
-void UNavModifierComponent::OnRegister()
-{
-	Super::OnRegister();
-
 	const AActor* MyOwner = GetOwner();
-
 	if (MyOwner)
 	{
 		const float Radius = MyOwner->GetSimpleCollisionRadius();
 		const float HalfHeght = MyOwner->GetSimpleCollisionHalfHeight();
-		const FVector Loc = MyOwner->GetActorLocation();
 		ObstacleExtent = FVector(Radius, Radius, HalfHeght);
-		Bounds = FBox(Loc + ObstacleExtent, Loc - ObstacleExtent);
+		Bounds = FBox::BuildAABB(MyOwner->GetActorLocation(), ObstacleExtent);
 	}
 }
 
-void UNavModifierComponent::OnApplyModifiers(FCompositeNavModifier& Modifiers)
+void UNavModifierComponent::GetNavigationData(struct FNavigationRelevantData& Data) const
 {
-	const AActor* MyOwner = GetOwner();
-	if (MyOwner)
-	{
-		const FVector Loc = MyOwner->GetActorLocation();
-		Bounds = FBox(Loc + ObstacleExtent, Loc - ObstacleExtent);
-		Modifiers.Add(FAreaNavModifier(Bounds, FTransform::Identity, AreaClass));
-	}
+	Data.Modifiers.Add(FAreaNavModifier(Bounds, FTransform::Identity, AreaClass));
 }
diff --git a/Engine/Source/Runtime/Engine/Private/AI/Navigation/NavModifierVolume.cpp b/Engine/Source/Runtime/Engine/Private/AI/Navigation/NavModifierVolume.cpp
--- a/Engine/Source/Runtime/Engine/Private/AI/Navigation/NavModifierVolume.cpp
+++ b/Engine/Source/Runtime/Engine/Private/AI/Navigation/NavModifierVolume.cpp
@@ -12,13 +12,16 @@ ANavModifierVolume::ANavModifierVolume(const class FPostConstructInitializePrope
 {
 }
 
-bool ANavModifierVolume::GetNavigationRelevantData(struct FNavigationRelevantData& Data) const
+void ANavModifierVolume::GetNavigationData(struct FNavigationRelevantData& Data) const
 {
 	if (Brush && AreaClass && AreaClass != UNavigationSystem::GetDefaultWalkableArea())
 	{
 		FAreaNavModifier AreaMod(BrushComponent, AreaClass);
 		Data.Modifiers.Add(AreaMod);
 	}
-	
-	return false;
+}
+
+FBox ANavModifierVolume::GetNavigationBounds() const
+{
+	return GetComponentsBoundingBox();
 }
diff --git a/Engine/Source/Runtime/Engine/Private/AI/Navigation/NavRelevantComponent.cpp b/Engine/Source/Runtime/Engine/Private/AI/Navigation/NavRelevantComponent.cpp
--- a/Engine/Source/Runtime/Engine/Private/AI/Navigation/NavRelevantComponent.cpp
+++ b/Engine/Source/Runtime/Engine/Private/AI/Navigation/NavRelevantComponent.cpp
@@ -6,45 +6,40 @@
 UNavRelevantComponent::UNavRelevantComponent(const class FPostConstructInitializeProperties& PCIP) : Super(PCIP)
 {
 	bNavigationRelevant = true;
-	Bounds = FBox::BuildAABB(FVector::ZeroVector, FVector(100.0f, 100.0f, 100.0f));
 }
 
 void UNavRelevantComponent::OnRegister()
 {
 	Super::OnRegister();
 
-	AActor* MyOwner = GetOwner();
-	if (MyOwner)
-	{
-		Bounds.ShiftBy(MyOwner->GetActorLocation());
-		MyOwner->UpdateNavigationRelevancy();
-	}
+	CalcBounds();
+
+	UNavigationSystem::OnComponentRegistered(this);
 }
 
 void UNavRelevantComponent::OnUnregister()
 {
 	Super::OnUnregister();
 
-	if (GetOwner())
-	{
-		GetOwner()->UpdateNavigationRelevancy();
-	}
+	UNavigationSystem::OnComponentUnregistered(this);
 }
 
-
-void UNavRelevantComponent::OnOwnerRegistered()
+FBox UNavRelevantComponent::GetNavigationBounds() const
 {
-
+	return Bounds;
 }
 
-void UNavRelevantComponent::OnOwnerUnregistered()
+bool UNavRelevantComponent::IsNavigationRelevant() const
 {
-
+	return bNavigationRelevant;
 }
 
-void UNavRelevantComponent::OnApplyModifiers(FCompositeNavModifier& Modifiers)
+void UNavRelevantComponent::CalcBounds()
 {
+	AActor* OwnerActor = GetOwner();
+	const FVector MyLocation = OwnerActor ? OwnerActor->GetActorLocation() : FVector::ZeroVector;
 
+	Bounds = FBox::BuildAABB(MyLocation, FVector(100.0f, 100.0f, 100.0f));
 }
 
 void UNavRelevantComponent::SetNavigationRelevancy(bool bRelevant)
@@ -52,33 +47,11 @@ void UNavRelevantComponent::SetNavigationRelevancy(bool bRelevant)
 	if (bNavigationRelevant != bRelevant)
 	{
 		bNavigationRelevant = bRelevant;
-
-		if (bRelevant)
-		{
-			OnOwnerRegistered();
-		}
-		else
-		{
-			OnOwnerUnregistered();
-		}
-
 		RefreshNavigationModifiers();
 	}
 }
 
 void UNavRelevantComponent::RefreshNavigationModifiers()
 {
-	UNavigationSystem* NavSys = UNavigationSystem::GetCurrent(GetWorld());
-	if (NavSys && GetOwner())
-	{
-		INavRelevantActorInterface* NavRelevantOwner = InterfaceCast<INavRelevantActorInterface>(GetOwner());
-		if (NavRelevantOwner && NavRelevantOwner->DoesSupplyPerComponentNavigationCollision())
-		{
-			NavSys->UpdateNavOctree(this, UNavigationSystem::OctreeUpdate_Modifiers);
-		}
-		else
-		{
-			NavSys->UpdateNavOctree(GetOwner(), UNavigationSystem::OctreeUpdate_Modifiers);
-		}
-	}
+	UNavigationSystem::UpdateNavOctree(this);
 }
diff --git a/Engine/Source/Runtime/Engine/Private/AI/Navigation/NavigationInterfaces.cpp b/Engine/Source/Runtime/Engine/Private/AI/Navigation/NavigationInterfaces.cpp
--- a/Engine/Source/Runtime/Engine/Private/AI/Navigation/NavigationInterfaces.cpp
+++ b/Engine/Source/Runtime/Engine/Private/AI/Navigation/NavigationInterfaces.cpp
@@ -6,7 +6,7 @@
 #include "AI/Navigation/NavNodeInterface.h"
 #include "AI/Navigation/NavLinkHostInterface.h"
 #include "AI/Navigation/NavPathObserverInterface.h"
-#include "AI/Navigation/NavRelevantActorInterface.h"
+#include "AI/Navigation/NavRelevantInterface.h"
 #include "AI/Navigation/NavLinkCustomInterface.h"
 
 uint32 INavLinkCustomInterface::NextUniqueId = 1;
@@ -64,7 +64,7 @@ UNavPathObserverInterface::UNavPathObserverInterface(const class FPostConstructI
 {
 }
 
-UNavRelevantActorInterface::UNavRelevantActorInterface(const class FPostConstructInitializeProperties& PCIP)
+UNavRelevantInterface::UNavRelevantInterface(const class FPostConstructInitializeProperties& PCIP)
 	: Super(PCIP)
 {
 }
diff --git a/Engine/Source/Runtime/Engine/Private/AI/Navigation/NavigationModifier.cpp b/Engine/Source/Runtime/Engine/Private/AI/Navigation/NavigationModifier.cpp
--- a/Engine/Source/Runtime/Engine/Private/AI/Navigation/NavigationModifier.cpp
+++ b/Engine/Source/Runtime/Engine/Private/AI/Navigation/NavigationModifier.cpp
@@ -11,7 +11,6 @@
 #include "AI/Navigation/NavLinkDefinition.h"
 #include "AI/Navigation/NavLinkTrivial.h"
 #include "AI/Navigation/NavAreas/NavAreaMeta.h"
-#include "AI/Navigation/NavigationProxy.h"
 
 // if square distance between two points is less than this the those points
 // will be considered identical when calculating convex hull
@@ -434,6 +433,13 @@ void FSimpleLinkNavModifier::AddSegmentLink(const FNavigationSegmentLink& InLink
 //----------------------------------------------------------------------//
 // FCompositeNavMeshModifier
 //----------------------------------------------------------------------//
+void FCompositeNavModifier::Shrink()
+{
+	Areas.Shrink();
+	SimpleLinks.Shrink();
+	CustomLinks.Shrink();
+}
+
 void FCompositeNavModifier::Reset()
 {
 	Areas.Reset();
@@ -464,9 +470,7 @@ FCompositeNavModifier FCompositeNavModifier::GetInstantiatedMetaModifier(const F
 		return Result;
 	}
 
-	const AActor* ActorOwner = Cast<UNavigationProxy>(ObjectOwner) ? ((UNavigationProxy*)ObjectOwner)->MyOwner :
-		Cast<AActor>(ObjectOwner) ? (AActor*)ObjectOwner : Cast<AActor>(ObjectOwner->GetOuter());
-
+	const AActor* ActorOwner = Cast<AActor>(ObjectOwner) ? (AActor*)ObjectOwner : Cast<AActor>(ObjectOwner->GetOuter());
 	if (ActorOwner == NULL)
 	{
 		return Result;
diff --git a/Engine/Source/Runtime/Engine/Private/AI/Navigation/NavigationOctree.cpp b/Engine/Source/Runtime/Engine/Private/AI/Navigation/NavigationOctree.cpp
--- a/Engine/Source/Runtime/Engine/Private/AI/Navigation/NavigationOctree.cpp
+++ b/Engine/Source/Runtime/Engine/Private/AI/Navigation/NavigationOctree.cpp
@@ -4,7 +4,6 @@
 #include "NavigationOctree.h"
 #include "RecastHelpers.h"
 #include "AI/Navigation/NavRelevantComponent.h"
-#include "AI/Navigation/NavigationProxy.h"
 
 #if NAVOCTREE_CONTAINS_COLLISION_DATA
 #include "RecastNavMeshGenerator.h"
@@ -27,85 +26,68 @@ FNavigationOctree::~FNavigationOctree()
 	DEC_MEMORY_STAT_BY(STAT_Navigation_CollisionTreeMemory, NodesMemory);
 }
 
-void FNavigationOctree::AddActor(AActor& Actor, FNavigationOctreeElement& Data)
+void FNavigationOctree::SetNavigableGeometryStoringMode(ENavGeometryStoringMode NavGeometryMode)
 {
-	if (Actor.IsPendingKill())
-	{
-		return;
-	}
-
-	INavRelevantActorInterface* NavRelevantActor = InterfaceCast<INavRelevantActorInterface>(&Actor);
-	UNavigationProxy* ProxyOb = NavRelevantActor ? NavRelevantActor->GetNavigationProxy() : NULL;
-	UObject* NodeOwner = ProxyOb ? (UObject*)ProxyOb : &Actor;
-
-	Data.Owner = NodeOwner;
-
-	const FBox BBox = Actor.GetComponentsBoundingBox();
-	if (BBox.IsValid && !BBox.GetSize().IsNearlyZero())
-	{
-		Data.Bounds = BBox;
-	
-#if NAVOCTREE_CONTAINS_COLLISION_DATA
-		bool bExportGeometry = true;
-
-		{	
-			SCOPE_CYCLE_COUNTER(STAT_Navigation_GatheringNavigationModifiersSync);
-			if (NavRelevantActor)
-			{
-				bExportGeometry = NavRelevantActor->GetNavigationRelevantData(Data.Data);
-			}
-		}
-
-		if (bExportGeometry)
-		{
-			SCOPE_CYCLE_COUNTER(STAT_Navigation_ActorsGeometryExportSync);
-			NavigableGeometryExportDelegate.ExecuteIfBound(Actor, Data.Data);
-		}
-#endif // NAVOCTREE_CONTAINS_COLLISION_DATA
-
-		AddNode(Data);
-	}
+	bGatherGeometry = (NavGeometryMode == FNavigationOctree::StoreNavGeometry);
 }
 
-void FNavigationOctree::AddComponent(UActorComponent& ActorComp, const FBox& Bounds, FNavigationOctreeElement& Data)
+void FNavigationOctree::AddNode(UObject* ElementOb, INavRelevantInterface* NavElement, const FBox& Bounds, FNavigationOctreeElement& Data)
 {
-	if (ActorComp.IsPendingKill())
-	{
-		return;
-	}
-
-	Data.Owner = &ActorComp;
+	Data.Owner = ElementOb;
 	Data.Bounds = Bounds;
 
 #if NAVOCTREE_CONTAINS_COLLISION_DATA
-	if (bGatherGeometry)
+	UActorComponent* ActorComp = Cast<UActorComponent>(ElementOb);
+	if (bGatherGeometry && ActorComp)
 	{
-		NavigableGeometryComponentExportDelegate.ExecuteIfBound(ActorComp, Data.Data);
+		ComponentExportDelegate.ExecuteIfBound(ActorComp, Data.Data);
 	}
 
-	UNavRelevantComponent* NavRelevantComponent = Cast<UNavRelevantComponent>(&ActorComp);
-	if (NavRelevantComponent)
+	if (NavElement)
 	{
-		NavRelevantComponent->OnApplyModifiers(Data.Data.Modifiers);
+		SCOPE_CYCLE_COUNTER(STAT_Navigation_GatheringNavigationModifiersSync);
+		NavElement->GetNavigationData(Data.Data);
 	}
-#endif // NAVOCTREE_CONTAINS_COLLISION_DATA
 
-	AddNode(Data);
-}
+	// shrink arrays before counting memory
+	// it will be reallocated when adding to octree and RemoveNode will have different value returned by GetAllocatedSize()
+	Data.Shrink();
 
-void FNavigationOctree::SetNavigableGeometryStoringMode(ENavGeometryStoringMode NavGeometryMode)
-{
-	bGatherGeometry = (NavGeometryMode == FNavigationOctree::StoreNavGeometry);
+	const int32 ElementMemory = Data.GetAllocatedSize();
+	NodesMemory += ElementMemory;
+	INC_MEMORY_STAT_BY(STAT_Navigation_CollisionTreeMemory, ElementMemory);
+#endif
+
+	AddElement(Data);
 }
 
-void FNavigationOctree::AddNode(FNavigationOctreeElement& Data)
+void FNavigationOctree::AppendToNode(const FOctreeElementId& Id, INavRelevantInterface* NavElement, const FBox& Bounds, FNavigationOctreeElement& Data)
 {
+	FNavigationOctreeElement OrgData = GetElementById(Id);
+
+	Data = OrgData;
+	Data.Bounds = Bounds + OrgData.Bounds.GetBox();
+
 #if NAVOCTREE_CONTAINS_COLLISION_DATA
-	const int32 ElementMemory = Data.GetAllocatedSize();
-	NodesMemory += ElementMemory;
-	INC_MEMORY_STAT_BY(STAT_Navigation_CollisionTreeMemory, ElementMemory);
+	if (NavElement)
+	{
+		SCOPE_CYCLE_COUNTER(STAT_Navigation_GatheringNavigationModifiersSync);
+		NavElement->GetNavigationData(Data.Data);
+	}
+
+	// shrink arrays before counting memory
+	// it will be reallocated when adding to octree and RemoveNode will have different value returned by GetAllocatedSize()
+	Data.Shrink();
+
+	const int32 OrgElementMemory = OrgData.GetAllocatedSize();
+	const int32 NewElementMemory = Data.GetAllocatedSize();
+	const int32 MemoryDelta = NewElementMemory - OrgElementMemory;
+
+	NodesMemory += MemoryDelta;
+	INC_MEMORY_STAT_BY(STAT_Navigation_CollisionTreeMemory, MemoryDelta);
 #endif
 
+	RemoveElement(Id);
 	AddElement(Data);
 }
 
@@ -129,6 +111,13 @@ bool FNavigationRelevantData::IsMatchingFilter(const FNavigationOctreeFilter& Fi
 		(Filter.bIncludeMetaAreas && Modifiers.HasMetaAreas());
 }
 
+void FNavigationRelevantData::Shrink()
+{
+	CollisionData.Shrink();
+	VoxelData.Shrink();
+	Modifiers.Shrink();
+}
+
 #if NAVSYS_DEBUG
 FORCENOINLINE
 #endif // NAVSYS_DEBUG
@@ -141,10 +130,6 @@ void FNavigationOctreeSemantics::SetElementId(const FNavigationOctreeElement& El
 	{
 		World = Actor->GetWorld();
 	}
-	else if (UNavigationProxy* Proxy = Cast<UNavigationProxy>(ElementOwner))
-	{
-		World = Proxy->MyOwner != NULL ? Proxy->MyOwner->GetWorld() : NULL;
-	}
 	else if (UActorComponent* AC = Cast<UActorComponent>(ElementOwner))
 	{
 		World = AC->GetWorld();
diff --git a/Engine/Source/Runtime/Engine/Private/AI/Navigation/NavigationProxy.cpp b/Engine/Source/Runtime/Engine/Private/AI/Navigation/NavigationProxy.cpp
deleted file mode 100644
diff --git a/Engine/Source/Runtime/Engine/Private/AI/Navigation/NavigationSystem.cpp b/Engine/Source/Runtime/Engine/Private/AI/Navigation/NavigationSystem.cpp
--- a/Engine/Source/Runtime/Engine/Private/AI/Navigation/NavigationSystem.cpp
+++ b/Engine/Source/Runtime/Engine/Private/AI/Navigation/NavigationSystem.cpp
@@ -25,7 +25,6 @@
 #include "AI/Navigation/NavigationSystem.h"
 #include "AI/Navigation/NavRelevantComponent.h"
 #include "AI/Navigation/NavigationPath.h"
-#include "AI/Navigation/NavigationProxy.h"
 
 static const uint32 INITIAL_ASYNC_QUERIES_SIZE = 32;
 static const uint32 REGISTRATION_QUEUE_SIZE = 16;	// and we'll not reallocate
@@ -36,11 +35,13 @@ static const uint32 MAX_NAV_SEARCH_NODES = 2048;
 #endif // WITH_RECAST
 
 DEFINE_LOG_CATEGORY(LogNavigation);
+DEFINE_LOG_CATEGORY_STATIC(LogNavOctree, Warning, All);
 
 DECLARE_CYCLE_STAT(TEXT("Rasterize triangles"), STAT_Navigation_RasterizeTriangles,STATGROUP_Navigation);
 DECLARE_CYCLE_STAT(TEXT("Nav Tick: area register"), STAT_Navigation_TickNavAreaRegister, STATGROUP_Navigation);
 DECLARE_CYCLE_STAT(TEXT("Nav Tick: rebuild"), STAT_Navigation_TickRebuild, STATGROUP_Navigation);
 DECLARE_CYCLE_STAT(TEXT("Nav Tick: async pathfinding"), STAT_Navigation_TickAsyncPathfinding, STATGROUP_Navigation);
+DECLARE_CYCLE_STAT(TEXT("Debug NavOctree Time"), STAT_DebugNavOctree, STATGROUP_Navigation);
 
 //----------------------------------------------------------------------//
 // Stats
@@ -116,7 +117,7 @@ void FNavigationLockContext::LockUpdates()
 		UNavigationSystem* NavSys = UNavigationSystem::GetCurrent(MyWorld);
 		if (NavSys)
 		{
-			NavSys->BeginFakeComponentChanges();
+			NavSys->AddNavigationUpdateLock(LockReason);
 		}
 	}
 	else
@@ -126,7 +127,7 @@ void FNavigationLockContext::LockUpdates()
 			UNavigationSystem* NavSys = UNavigationSystem::GetCurrent(Context.World());
 			if (NavSys)
 			{
-				NavSys->BeginFakeComponentChanges();
+				NavSys->AddNavigationUpdateLock(LockReason);
 			}
 		}
 	}
@@ -141,7 +142,7 @@ void FNavigationLockContext::UnlockUpdates()
 		UNavigationSystem* NavSys = UNavigationSystem::GetCurrent(MyWorld);
 		if (NavSys)
 		{
-			NavSys->EndFakeComponentChanges();
+			NavSys->RemoveNavigationUpdateLock(LockReason);
 		}
 	}
 	else
@@ -151,7 +152,7 @@ void FNavigationLockContext::UnlockUpdates()
 			UNavigationSystem* NavSys = UNavigationSystem::GetCurrent(Context.World());
 			if (NavSys)
 			{
-				NavSys->EndFakeComponentChanges();
+				NavSys->RemoveNavigationUpdateLock(LockReason);
 			}
 		}
 	}
@@ -194,7 +195,7 @@ UNavigationSystem::UNavigationSystem(const FPostConstructInitializeProperties& P
 	, DirtyAreasUpdateTime(0)
 {
 #if WITH_EDITOR
-	bFakeComponentChangesBeingApplied = false;
+	NavUpdateLockFlags = 0;
 #endif
 
 	if (HasAnyFlags(RF_ClassDefaultObject) == false)
@@ -205,10 +206,9 @@ UNavigationSystem::UNavigationSystem(const FPostConstructInitializeProperties& P
 	
 		NavOctree = new FNavigationOctree(FVector(0,0,0), 64000);
 #if WITH_NAVIGATION_GENERATOR
-#if WITH_RECAST
-		NavOctree->NavigableGeometryExportDelegate = FNavigationOctree::FNavigableGeometryExportDelegate::CreateStatic(&FRecastNavMeshGenerator::ExportActorGeometry);
-		NavOctree->NavigableGeometryComponentExportDelegate = FNavigationOctree::FNavigableGeometryComponentExportDelegate::CreateStatic(&FRecastNavMeshGenerator::ExportComponentGeometry);
-#endif // WITH_RECAST
+	#if WITH_RECAST
+		NavOctree->ComponentExportDelegate = FNavigationOctree::FNavigableGeometryComponentExportDelegate::CreateStatic(&FRecastNavMeshGenerator::ExportComponentGeometry);
+	#endif // WITH_RECAST
 		FWorldDelegates::LevelAddedToWorld.AddUObject(this, &UNavigationSystem::OnLevelAddedToWorld);
 		FWorldDelegates::LevelRemovedFromWorld.AddUObject(this, &UNavigationSystem::OnLevelRemovedFromWorld);
 #endif // WITH_NAVIGATION_GENERATOR	
@@ -1187,6 +1187,32 @@ bool UNavigationSystem::IsThereAnywhereToBuildNavigation() const
 	return bCreateNavigation;
 }
 
+bool UNavigationSystem::IsNavigationRelevant(const AActor* TestActor) const
+{
+	const INavRelevantInterface* NavInterface = InterfaceCast<const INavRelevantInterface>(TestActor);
+	if (NavInterface && NavInterface->IsNavigationRelevant())
+	{
+		return true;
+	}
+
+	TArray<UActorComponent*> Components;
+	if (TestActor)
+	{
+		TestActor->GetComponents(Components);
+	}
+
+	for (int32 Idx = 0; Idx < Components.Num(); Idx++)
+	{
+		NavInterface = InterfaceCast<const INavRelevantInterface>(Components[Idx]);
+		if (NavInterface && NavInterface->IsNavigationRelevant())
+		{
+			return true;
+		}
+	}
+
+	return false;
+}
+
 FBox UNavigationSystem::GetWorldBounds() const
 {
 	checkSlow(IsInGameThread() == true);
@@ -1198,7 +1224,7 @@ FBox UNavigationSystem::GetWorldBounds() const
 		// @TODO - super slow! Need to ask tech guys where I can get this from
 		for( FActorIterator It(GetWorld()); It; ++It )
 		{
-			if ((*It)->IsNavigationRelevant() == true)
+			if (IsNavigationRelevant(*It))
 			{
 				NavigableWorldBounds += (*It)->GetComponentsBoundingBox();
 			}
@@ -1218,7 +1244,7 @@ FBox UNavigationSystem::GetLevelBounds(ULevel* InLevel) const
 		const int32 ActorCount = InLevel->Actors.Num();
 		for (int32 ActorIndex = 0; ActorIndex < ActorCount; ++ActorIndex, ++Actor)
 		{
-			if (*Actor != NULL && (*Actor)->IsNavigationRelevant() == true)
+			if (IsNavigationRelevant(*Actor))
 			{
 				NavigableLevelBounds += (*Actor)->GetComponentsBoundingBox();
 			}
@@ -1786,108 +1812,60 @@ void UNavigationSystem::AddDirtyAreas(const TArray<FBox>& NewAreas, int32 Flags)
 
 int32 GetDirtyFlagHelper(int32 UpdateFlags, int32 DefaultValue)
 {
-	return ((UpdateFlags & UNavigationSystem::OctreeUpdate_Modifiers) != 0) ? ENavigationDirtyFlag::DynamicModifier :
-		((UpdateFlags & UNavigationSystem::OctreeUpdate_Geometry) != 0) ? ENavigationDirtyFlag::All :
+	return ((UpdateFlags & UNavigationSystem::OctreeUpdate_Geometry) != 0) ? ENavigationDirtyFlag::All :
+		((UpdateFlags & UNavigationSystem::OctreeUpdate_Modifiers) != 0) ? ENavigationDirtyFlag::DynamicModifier :		
 		DefaultValue;
 }
 
-FSetElementId UNavigationSystem::RegisterNavigationRelevantActor(AActor* Actor, int32 UpdateFlags) 
-{ 
-	FSetElementId RequestId;
-
-#if WITH_EDITOR
-	if (AreFakeComponentChangesBeingApplied())
-	{
-		// SMALL HACK: don't have to unregister actor, Static lighting system is refreshing components in world
-		return RequestId;
-	}
-#endif
-	
-	if (Actor->IsNavigationRelevant())
-	{
-		INavRelevantActorInterface* NavRelevantActor = InterfaceCast<INavRelevantActorInterface>(Actor);
-		UNavigationProxy* ProxyOb = NavRelevantActor ? NavRelevantActor->GetNavigationProxy() : NULL;
-		const bool bPerComponentNavigation = NavRelevantActor ? NavRelevantActor->DoesSupplyPerComponentNavigationCollision() : false;
-		
-		// actor's shouldn't use proxy and per component navigation at the same time (conflict ownership of octree element)
-		if (bPerComponentNavigation && ProxyOb)
-		{
-			UE_LOG(LogNavigation, Warning, TEXT("Clearing navigation proxy of %s (provides per component data!)"), *GetNameSafe(Actor));
-			ProxyOb = NULL;
-		}
-
-		// dirty proxy will force update
-		if (ProxyOb && ProxyOb->bDirty)
-		{
-			ProxyOb->bDirty = false;
-
-			UnregisterNavigationRelevantActor(Actor, UpdateFlags | OctreeUpdate_Refresh);
-		}
-
-		UObject* ElementOwner = ProxyOb ? (UObject*)ProxyOb : Actor;
-		RequestId = RegisterNavOctreeElement(ElementOwner, UpdateFlags, FNavigationSystem::InvalidBoundingBox);
-	}
-
-	return RequestId;
-}
-
-FSetElementId UNavigationSystem::RegisterNavOctreeElement(UObject* ElementOwner, int32 UpdateFlags, const FBox& Bounds)
+FSetElementId UNavigationSystem::RegisterNavOctreeElement(UObject* ElementOwner, class INavRelevantInterface* ElementInterface, int32 UpdateFlags)
 {
 	FSetElementId SetId;
 
 #if WITH_EDITOR
-	if (AreFakeComponentChangesBeingApplied())
+	if (IsNavigationRegisterLocked())
 	{
-		// SMALL HACK: don't have to unregister actor, Static lighting system is refreshing components in world
 		return SetId;
 	}
 #endif
 
-	if (NavOctree == NULL || ElementOwner == NULL)
+	if (NavOctree == NULL || ElementOwner == NULL || ElementInterface == NULL)
 	{
 		return SetId;
 	}
 
-	//
-	// HACK: This is necessary to work around the case where a PerComponentNavigation actor has its components registered,
-	// which registers the actor, and also UpdateNavOctree(Component) is called in the same frame.
-	//
-	// A better solution would be to make the RegisterNavigationRelevantActor function add only individual components if
-	// the actor is using PerComponentNavigation.
-	//
-	UPrimitiveComponent* Component = Cast<UPrimitiveComponent>(ElementOwner);
-	if (Component && Component->GetOwner())
-	{
-		INavRelevantActorInterface* NavRelevantActor = InterfaceCast<INavRelevantActorInterface>(Component->GetOwner());
-		const bool bPerComponentNavigation = NavRelevantActor ? NavRelevantActor->DoesSupplyPerComponentNavigationCollision() : false;
-		if (bPerComponentNavigation)
-		{
-			FSetElementId ActorSetId = PendingOctreeUpdates.FindId(Component->GetOwner());
-			if (ActorSetId.IsValidId())
-			{
-				return ActorSetId;
-			}
-		}
-	}
+	const bool bIsRelevant = ElementInterface->IsNavigationRelevant();
+	UE_LOG(LogNavOctree, Log, TEXT("REG %s %s"), *GetNameSafe(ElementOwner), bIsRelevant ? TEXT("[relevant]") : TEXT(""));
 
-	const FOctreeElementId* ElementId = GetObjectsNavOctreeId(ElementOwner);
-	if (ElementId == NULL)
+	if (bIsRelevant)
 	{
-		FNavigationDirtyElement UpdateInfo(ElementOwner, GetDirtyFlagHelper(UpdateFlags, 0));
-		if (!(Bounds == FNavigationSystem::InvalidBoundingBox))
-		{
-			UpdateInfo.SetBounds(Bounds);
-		}
+		bool bCanAdd = false;
 
-		SetId = PendingOctreeUpdates.FindId(UpdateInfo);
-		if (SetId.IsValidId())
+		UObject* ParentNode = ElementInterface->GetNavigationParent();
+		if (ParentNode)
 		{
-			// make sure this request stays, in case it has been invalidated already
-			PendingOctreeUpdates[SetId] = UpdateInfo;
+			OctreeChildNodesMap.AddUnique(ParentNode, ElementOwner);
+			bCanAdd = true;
 		}
 		else
 		{
-			SetId = PendingOctreeUpdates.Add(UpdateInfo);
+			const FOctreeElementId* ElementId = GetObjectsNavOctreeId(ElementOwner);
+			bCanAdd = (ElementId == NULL);
+		}
+
+		if (bCanAdd)
+		{
+			FNavigationDirtyElement UpdateInfo(ElementOwner, ElementInterface, GetDirtyFlagHelper(UpdateFlags, 0));
+
+			SetId = PendingOctreeUpdates.FindId(UpdateInfo);
+			if (SetId.IsValidId())
+			{
+				// make sure this request stays, in case it has been invalidated already
+				PendingOctreeUpdates[SetId] = UpdateInfo;
+			}
+			else
+			{
+				SetId = PendingOctreeUpdates.Add(UpdateInfo);
+			}
 		}
 	}
 
@@ -1909,119 +1887,46 @@ void UNavigationSystem::AddElementToNavOctree(const FNavigationDirtyElement& Dir
 		return;
 	}
 
-	// no check for owner's validity are performed, nor its presence
-	// in NavOctree - function assumes callee responsibility in this 
-	// regard
-
 	UObject* ElementOwner = DirtyElement.Owner.Get();
-	AActor* ActorOwner = Cast<AActor>(ElementOwner);
-	UActorComponent* CompOwner = Cast<UActorComponent>(ElementOwner);
-	
-	UNavigationProxy* NavProxy = Cast<UNavigationProxy>(ElementOwner);
-	if (NavProxy)
-	{
-		ActorOwner = NavProxy->MyOwner;
-	}
-
-	if (ElementOwner == NULL ||
-		(ActorOwner && ActorOwner->GetWorld() == NULL) ||
-		(CompOwner && CompOwner->GetWorld() == NULL))
+	if (ElementOwner == NULL || ElementOwner->IsPendingKill())
 	{
 		return;
 	}
 
-	if (ActorOwner)
-	{
-		INavRelevantActorInterface* NavRelevantActor = InterfaceCast<INavRelevantActorInterface>(ActorOwner);
-		const bool bPerComponentNavigation = NavRelevantActor ? NavRelevantActor->DoesSupplyPerComponentNavigationCollision() : false;
+	FNavigationOctreeElement GeneratedData;
+	const FBox ElementBounds = DirtyElement.NavInterface->GetNavigationBounds();
 
-		if (bPerComponentNavigation)
-		{
-			TArray<UActorComponent*> Components;
-			ActorOwner->GetComponents(Components);
-			for (int32 ComponentIndex = 0; ComponentIndex < Components.Num(); ComponentIndex++)
-			{
-				UPrimitiveComponent* PrimComp = Cast<UPrimitiveComponent>(Components[ComponentIndex]);
-				if (PrimComp && PrimComp->CanEverAffectNavigation())
-				{
-					AddComponentElementToNavOctree(*PrimComp, DirtyElement, PrimComp->Bounds.GetBox());
-				}
-				else
-				{
-					UNavRelevantComponent* NavRelevantComponent = Cast<UNavRelevantComponent>(Components[ComponentIndex]);
-					if (NavRelevantComponent && NavRelevantComponent->IsNavigationRelevant())
-					{
-						NavRelevantComponent->OnOwnerRegistered();
-						AddComponentElementToNavOctree(*NavRelevantComponent, DirtyElement, NavRelevantComponent->Bounds);
-					}
-				}
-			}
-		}
-		else
+	UObject* ParentNode = DirtyElement.NavInterface->GetNavigationParent();
+	if (ParentNode)
+	{
+		// check if parent node is waiting in queue
+		const FSetElementId ParentRequestId = PendingOctreeUpdates.FindId(FNavigationDirtyElement(ParentNode));
+		if (ParentRequestId.IsValidId())
 		{
-			TArray<UNavRelevantComponent*> Components;
-			ActorOwner->GetComponents(Components);
-			for (int32 i = 0; i < Components.Num(); i++)
-			{
-				UNavRelevantComponent* NavRelevantComponent = Components[i];
-				if (NavRelevantComponent && NavRelevantComponent->IsNavigationRelevant())
-				{
-					NavRelevantComponent->OnOwnerRegistered();
-					AddComponentElementToNavOctree(*NavRelevantComponent, DirtyElement, NavRelevantComponent->Bounds);
-				}
-			}
+			FNavigationDirtyElement& ParentNode = PendingOctreeUpdates[ParentRequestId];
+			AddElementToNavOctree(ParentNode);
 
-			AddActorElementToNavOctree(*ActorOwner, DirtyElement);
+			// mark as invalid so it won't be processed twice
+			ParentNode.bInvalidRequest = true;
 		}
-	}
-	else if (CompOwner)
-	{
-		AddComponentElementToNavOctree(*CompOwner, DirtyElement, DirtyElement.BoundsOverride);
-	}
-}
 
-void UNavigationSystem::AddActorElementToNavOctree(AActor& Actor, const FNavigationDirtyElement& DirtyElement)
-{
-	ensure(GetObjectsNavOctreeId(&Actor) == NULL);
-
-	// notify relevant components about adding owner to octree
-	TArray<UNavRelevantComponent*> Components;
-	Actor.GetComponents(Components);
-	for (int32 ComponentIndex = 0; ComponentIndex < Components.Num(); ComponentIndex++)
-	{
-		UNavRelevantComponent* NavRelevantComponent = Components[ComponentIndex];
-		if (NavRelevantComponent->IsNavigationRelevant())
+		const FOctreeElementId* ParentId = GetObjectsNavOctreeId(ParentNode);
+		if (ParentId)
 		{
-			NavRelevantComponent->OnOwnerRegistered();
+			UE_LOG(LogNavOctree, Log, TEXT("ADD %s to %s"), *GetNameSafe(ElementOwner), *GetNameSafe(ParentNode));
+			NavOctree->AppendToNode(*ParentId, DirtyElement.NavInterface, ElementBounds, GeneratedData);
+		}
+		else 
+		{
+			UE_LOG(LogNavOctree, Warning, TEXT("Can't add node [%s] - parent [%s] not found in octree!"), *GetNameSafe(ElementOwner), *GetNameSafe(ParentNode));
 		}
 	}
-
-	FNavigationOctreeElement GeneratedData;
-	NavOctree->AddActor(Actor, GeneratedData);
-
-#if WITH_NAVIGATION_GENERATOR && NAVOCTREE_CONTAINS_COLLISION_DATA
-	const FBox BBox = GeneratedData.Bounds.GetBox();
-	const bool bValidBBox = BBox.IsValid && !BBox.GetSize().IsNearlyZero();
-	if (bValidBBox && !GeneratedData.IsEmpty())
-	{
-		const int32 DirtyFlag = DirtyElement.FlagsOverride ? DirtyElement.FlagsOverride : GeneratedData.Data.GetDirtyFlag();
-		AddDirtyArea(BBox, DirtyFlag);
-	}
-#endif // WITH_NAVIGATION_GENERATOR
-}
-
-void UNavigationSystem::AddComponentElementToNavOctree(UActorComponent& ActorComp, const FNavigationDirtyElement& DirtyElement, const FBox& Bounds)
-{
-	const FOctreeElementId* ElementId = GetObjectsNavOctreeId(&ActorComp);
-	if (ElementId != NULL)
+	else
 	{
-		RemoveNavOctreeElementId(*ElementId, OctreeUpdate_Default);
-		RemoveObjectsNavOctreeId(&ActorComp);
+		UE_LOG(LogNavOctree, Log, TEXT("ADD %s"), *GetNameSafe(ElementOwner));
+		NavOctree->AddNode(ElementOwner, DirtyElement.NavInterface, ElementBounds, GeneratedData);
 	}
 
-	FNavigationOctreeElement GeneratedData;
-	NavOctree->AddComponent(ActorComp, Bounds, GeneratedData);
-
 #if WITH_NAVIGATION_GENERATOR && NAVOCTREE_CONTAINS_COLLISION_DATA
 	const FBox BBox = GeneratedData.Bounds.GetBox();
 	const bool bValidBBox = BBox.IsValid && !BBox.GetSize().IsNearlyZero();
@@ -2054,100 +1959,48 @@ bool UNavigationSystem::GetNavOctreeElementData(UObject* NodeOwner, int32& Dirty
 	return false;
 }
 
-void UNavigationSystem::UnregisterNavigationRelevantActor(AActor* Actor, int32 UpdateFlags)
+void UNavigationSystem::UnregisterNavOctreeElement(UObject* ElementOwner, class INavRelevantInterface* ElementInterface, int32 UpdateFlags)
 {
 #if WITH_EDITOR
-	if (AreFakeComponentChangesBeingApplied())
+	if (IsNavigationUnregisterLocked())
 	{
-		// SMALL HACK: don't have to unregister actor, Static lighting system is refreshing components in world
 		return;
 	}
 #endif
 
-	INavRelevantActorInterface* NavRelevantActor = InterfaceCast<INavRelevantActorInterface>(Actor);
-	UNavigationProxy* ProxyOb = NavRelevantActor ? NavRelevantActor->GetNavigationProxy() : NULL;
-	const bool bPerComponentNavigation = NavRelevantActor ? NavRelevantActor->DoesSupplyPerComponentNavigationCollision() : false;
-
-	// actor's shouldn't use proxy and per component navigation at the same time (conflict ownership of octree element)
-	if (bPerComponentNavigation && ProxyOb)
+	if (NavOctree == NULL || ElementOwner == NULL || ElementInterface == NULL)
 	{
-		UE_LOG(LogNavigation, Warning, TEXT("Clearing navigation proxy of %s (provides per component data!)"), *GetNameSafe(Actor));
-		ProxyOb = NULL;
+		return;
 	}
 
-	if (bPerComponentNavigation)
-	{
-		TArray<UPrimitiveComponent*> Components;
-		Actor->GetComponents(Components);
-		for (int32 ComponentIndex = 0; ComponentIndex < Components.Num(); ComponentIndex++)
-		{
-			// always unregister, otherwise it will skip no longer relevant components
-
-			UPrimitiveComponent* PrimComp = Components[ComponentIndex];
-			if (PrimComp)
-			{
-				UnregisterNavOctreeElement(PrimComp, UpdateFlags, PrimComp->Bounds.GetBox());
-			}
-			else
-			{
-				UNavRelevantComponent* NavRelevantComponent = Cast<UNavRelevantComponent>(Components[ComponentIndex]);
-				if (NavRelevantComponent)
-				{
-					NavRelevantComponent->OnOwnerUnregistered();
+	const FOctreeElementId* ElementId = GetObjectsNavOctreeId(ElementOwner);
+	UE_LOG(LogNavOctree, Log, TEXT("UNREG %s %s"), *GetNameSafe(ElementOwner), ElementId ? TEXT("[exists]") : TEXT(""));
 
-					UnregisterNavOctreeElement(NavRelevantComponent, UpdateFlags, NavRelevantComponent->Bounds);
-				}
-			}
-		}
+	if (ElementId != NULL)
+	{
+		RemoveNavOctreeElementId(*ElementId, UpdateFlags);
+		RemoveObjectsNavOctreeId(ElementOwner);
 	}
 	else
 	{
-		TArray<UNavRelevantComponent*> Components;
-		Actor->GetComponents(Components);
-		// notify relevant components about removing owner from octree
-		for (int32 ComponentIndex = 0; ComponentIndex < Components.Num(); ComponentIndex++)
+		const bool bCanRemoveChildNode = (UpdateFlags & OctreeUpdate_ParentChain) == 0;
+		UObject* ParentNode = ElementInterface->GetNavigationParent();
+		if (ParentNode && bCanRemoveChildNode)
 		{
-			// always unregister, otherwise it will skip no longer relevant components
+			// if node has navigation parent (= doesn't exists in octree on its own)
+			// and it's not part of parent chain update
+			// remove it from map and force update on parent to rebuild octree element
 
-			UNavRelevantComponent* NavRelevantComponent = Components[ComponentIndex];
-			if (NavRelevantComponent)
-			{
-				NavRelevantComponent->OnOwnerUnregistered();
-				UnregisterNavOctreeElement(NavRelevantComponent, UpdateFlags, NavRelevantComponent->Bounds);
-			}
+			OctreeChildNodesMap.RemoveSingle(ParentNode, ElementOwner);
+			UpdateNavOctreeParentChain(ParentNode);
 		}
-
-		UObject* ElementOwner = ProxyOb ? (UObject*)ProxyOb : Actor;
-		UnregisterNavOctreeElement(ElementOwner, UpdateFlags, FNavigationSystem::InvalidBoundingBox);
-	}
-}
-
-void UNavigationSystem::UnregisterNavOctreeElement(UObject* ElementOwner, int32 UpdateFlags, const FBox& Bounds)
-{
-#if WITH_EDITOR
-	if (AreFakeComponentChangesBeingApplied())
-	{
-		// SMALL HACK: don't have to unregister actor, Static lighting system is refreshing components in world
-		return;
-	}
-#endif
-	if (NavOctree == NULL || ElementOwner == NULL)
-	{
-		return;
-	}
-
-	const FOctreeElementId* ElementId = GetObjectsNavOctreeId(ElementOwner);
-	if (ElementId != NULL)
-	{
-		RemoveNavOctreeElementId(*ElementId, UpdateFlags);
-		RemoveObjectsNavOctreeId(ElementOwner);
 	}
 
 	// mark pending update as invalid, it will be dirtied according to currently active settings
 	const bool bCanInvalidateQueue = (UpdateFlags & OctreeUpdate_Refresh) == 0;
 	if (bCanInvalidateQueue)
 	{
-		const FSetElementId RequestId = PendingOctreeUpdates.FindId(ElementOwner);
+		const FSetElementId RequestId = PendingOctreeUpdates.FindId(FNavigationDirtyElement(ElementOwner));
 		if (RequestId.IsValidId())
 		{
 			PendingOctreeUpdates[RequestId].bInvalidRequest = true;
@@ -2169,90 +2022,102 @@ void UNavigationSystem::RemoveNavOctreeElementId(const FOctreeElementId& Element
 	}
 }
 
-void UNavigationSystem::UpdateNavOctree(AActor* Actor, int32 UpdateFlags)
+void UNavigationSystem::UpdateNavOctree(class AActor* Actor)
 {
-	INC_DWORD_STAT(STAT_Navigation_UpdateNavOctree);
+	SCOPE_CYCLE_COUNTER(STAT_DebugNavOctree);
+
+	INavRelevantInterface* NavElement = InterfaceCast<INavRelevantInterface>(Actor);
+	if (NavElement)
+	{
+		UNavigationSystem* NavSys = Actor ? UNavigationSystem::GetCurrent(Actor->GetWorld()) : NULL;
+		if (NavSys)
+		{
+			NavSys->UpdateNavOctreeElement(Actor, NavElement, OctreeUpdate_Modifiers);
+		}
+	}
+}
+
+void UNavigationSystem::UpdateNavOctree(class UActorComponent* Comp)
+{
+	SCOPE_CYCLE_COUNTER(STAT_DebugNavOctree);
+
+	// special case for early out: use cached nav relevancy
+	UPrimitiveComponent* PrimComp = Cast<UPrimitiveComponent>(Comp);
+	if (PrimComp && !PrimComp->bNavigationRelevant)
+	{
+		return;
+	}
 
-	if (Actor->IsNavigationRelevant())
+	INavRelevantInterface* NavElement = InterfaceCast<INavRelevantInterface>(Comp);
+	if (NavElement)
 	{
-		const INavRelevantActorInterface* NavRelevantActor = InterfaceCast<INavRelevantActorInterface>(Actor);
-		const bool bPerComponentNavigation = NavRelevantActor ? NavRelevantActor->DoesSupplyPerComponentNavigationCollision() : false;
-		if (bPerComponentNavigation)
+		AActor* OwnerActor = Comp ? Comp->GetOwner() : NULL;
+		if (OwnerActor)
 		{
-			TArray<UActorComponent*> Components;
-			Actor->GetComponents(Components);
-			// make sure it's unregistered correctly
-			for (int32 ComponentIndex = 0; ComponentIndex < Components.Num(); ComponentIndex++)
+			UNavigationSystem* NavSys = UNavigationSystem::GetCurrent(OwnerActor->GetWorld());
+			if (NavSys)
 			{
-				UPrimitiveComponent* PrimComp = Cast<UPrimitiveComponent>(Components[ComponentIndex]);
-				if (PrimComp && PrimComp->CanEverAffectNavigation())
+				if (OwnerActor->IsComponentRelevantForNavigation(Comp))
 				{
-					UnregisterNavOctreeElement(PrimComp, UpdateFlags | OctreeUpdate_Refresh, PrimComp->Bounds.GetBox());
+					NavSys->UpdateNavOctreeElement(Comp, NavElement, OctreeUpdate_Default);
 				}
 				else
 				{
-					UNavRelevantComponent* NavRelevantComponent = Cast<UNavRelevantComponent>(Components[ComponentIndex]);
-					if (NavRelevantComponent && NavRelevantComponent->IsNavigationRelevant())
-					{
-						UnregisterNavOctreeElement(NavRelevantComponent, UpdateFlags | OctreeUpdate_Refresh, NavRelevantComponent->Bounds);
-					}
-				}
-			}
-		}
-		else
-		{
-			TArray<UNavRelevantComponent*> Components;
-			Actor->GetComponents(Components);
-			for (int32 i = 0; i < Components.Num(); i++)
-			{
-				UNavRelevantComponent* NavRelevantComponent = Components[i];
-				if (NavRelevantComponent && NavRelevantComponent->IsNavigationRelevant())
-				{
-					UnregisterNavOctreeElement(NavRelevantComponent, UpdateFlags | OctreeUpdate_Refresh, NavRelevantComponent->Bounds);
+					NavSys->UnregisterNavOctreeElement(Comp, NavElement, OctreeUpdate_Default);
 				}
 			}
 		}
+	}
+}
 
-		UNavigationProxy* ProxyOb = NavRelevantActor ? NavRelevantActor->GetNavigationProxy() : NULL;
-		UObject* NodeOwner = (ProxyOb && !bPerComponentNavigation) ? (UObject*)ProxyOb : Actor;
+void UNavigationSystem::UpdateNavOctreeAll(class AActor* Actor)
+{
+	if (Actor)
+	{
+		UpdateNavOctree(Actor);
 
-		UpdateNavOctreeElement(NodeOwner, UpdateFlags, FNavigationSystem::InvalidBoundingBox);
+		TArray<UActorComponent*> Components;
+		Actor->GetComponents(Components);
+
+		for (int32 Idx = 0; Idx < Components.Num(); Idx++)
+		{
+			UpdateNavOctree(Components[Idx]);
+		}
 	}
 }
 
-void UNavigationSystem::UpdateNavOctree(UActorComponent* ActorComp, int32 UpdateFlags)
+void UNavigationSystem::ClearNavOctreeAll(class AActor* Actor)
 {
-	INC_DWORD_STAT(STAT_Navigation_UpdateNavOctree);
+	if (Actor)
+	{
+		OnActorUnregistered(Actor);
 
-	UPrimitiveComponent* PrimComp = Cast<UPrimitiveComponent>(ActorComp);
-	UNavRelevantComponent* NavComp = Cast<UNavRelevantComponent>(ActorComp);
+		TArray<UActorComponent*> Components;
+		Actor->GetComponents(Components);
 
-	if (PrimComp)
-	{
-		UpdateNavOctreeElement(ActorComp, UpdateFlags, PrimComp->Bounds.GetBox());
+		for (int32 Idx = 0; Idx < Components.Num(); Idx++)
+		{
+			OnComponentUnregistered(Components[Idx]);
+		}
 	}
-	else if (NavComp)
-	{
-		UpdateNavOctreeElement(ActorComp, UpdateFlags, NavComp->Bounds);
-	}	
 }
 
-void UNavigationSystem::UpdateNavOctreeElement(UObject* ElementOwner, int32 UpdateFlags, const FBox& Bounds)
+void UNavigationSystem::UpdateNavOctreeElement(UObject* ElementOwner, class INavRelevantInterface* ElementInterface, int32 UpdateFlags)
 {
+	INC_DWORD_STAT(STAT_Navigation_UpdateNavOctree);
+
 	// grab existing octree data
 	FBox CurrentBounds;
 	int32 CurrentFlags;
 	const bool bAlreadyExists = GetNavOctreeElementData(ElementOwner, CurrentFlags, CurrentBounds);
 
-	if (bAlreadyExists)
-	{
-		// don't invalidate pending requests
-		UpdateFlags |= OctreeUpdate_Refresh;
+	// don't invalidate pending requests
+	UpdateFlags |= OctreeUpdate_Refresh;
 
-		UnregisterNavOctreeElement(ElementOwner, UpdateFlags, CurrentBounds);
-	}
+	// always try to unregister, even if element owner doesn't exists in octree (parent nodes)
+	UnregisterNavOctreeElement(ElementOwner, ElementInterface, UpdateFlags);
 
-	const FSetElementId RequestId = RegisterNavOctreeElement(ElementOwner, UpdateFlags, Bounds);
+	const FSetElementId RequestId = RegisterNavOctreeElement(ElementOwner, ElementInterface, UpdateFlags);
 
 	// add original data to pending registration request
 	// so it could be dirtied properly when system receive unregister request while actor is still queued
@@ -2265,6 +2130,102 @@ void UNavigationSystem::UpdateNavOctreeElement(UObject* ElementOwner, int32 Upda
 	}
 }
 
+void UNavigationSystem::UpdateNavOctreeParentChain(UObject* ElementOwner)
+{
+	TArray<UObject*> ChildNodes;
+
+	OctreeChildNodesMap.MultiFind(ElementOwner, ChildNodes);
+	if (ChildNodes.Num() == 0)
+	{
+		return;
+	}
+
+	INavRelevantInterface* ElementInterface = InterfaceCast<INavRelevantInterface>(ElementOwner);
+	const int32 UpdateFlags = OctreeUpdate_ParentChain | OctreeUpdate_Refresh;
+	TArray<INavRelevantInterface*> ChildNavInterfaces;
+	ChildNavInterfaces.AddZeroed(ChildNodes.Num());
+	
+	for (int32 Idx = 0; Idx < ChildNodes.Num(); Idx++)
+	{
+		ChildNavInterfaces[Idx] = InterfaceCast<INavRelevantInterface>(ChildNodes[Idx]);
+		UnregisterNavOctreeElement(ChildNodes[Idx], ChildNavInterfaces[Idx], UpdateFlags);
+	}
+
+	UnregisterNavOctreeElement(ElementOwner, ElementInterface, UpdateFlags);
+	RegisterNavOctreeElement(ElementOwner, ElementInterface, UpdateFlags);
+
+	for (int32 Idx = 0; Idx < ChildNodes.Num(); Idx++)
+	{
+		RegisterNavOctreeElement(ChildNodes[Idx], ChildNavInterfaces[Idx], UpdateFlags);
+	}
+}
+
+void UNavigationSystem::OnComponentRegistered(UActorComponent* Comp)
+{
+	SCOPE_CYCLE_COUNTER(STAT_DebugNavOctree);
+	INavRelevantInterface* NavInterface = InterfaceCast<INavRelevantInterface>(Comp);
+	if (NavInterface)
+	{
+		AActor* OwnerActor = Comp ? Comp->GetOwner() : NULL;
+		if (OwnerActor && OwnerActor->IsComponentRelevantForNavigation(Comp))
+		{
+			UNavigationSystem* NavSys = UNavigationSystem::GetCurrent(OwnerActor->GetWorld());
+			if (NavSys)
+			{
+				NavSys->RegisterNavOctreeElement(Comp, NavInterface, OctreeUpdate_Default);
+			}
+		}
+	}
+}
+
+void UNavigationSystem::OnComponentUnregistered(UActorComponent* Comp)
+{
+	SCOPE_CYCLE_COUNTER(STAT_DebugNavOctree);
+	INavRelevantInterface* NavInterface = InterfaceCast<INavRelevantInterface>(Comp);
+	if (NavInterface)
+	{
+		AActor* OwnerActor = Comp ? Comp->GetOwner() : NULL;
+		if (OwnerActor)
+		{
+			// skip IsComponentRelevantForNavigation check, it's only for adding new stuff
+
+			UNavigationSystem* NavSys = UNavigationSystem::GetCurrent(OwnerActor->GetWorld());
+			if (NavSys)
+			{
+				NavSys->UnregisterNavOctreeElement(Comp, NavInterface, OctreeUpdate_Default);
+			}
+		}
+	}
+}
+
+void UNavigationSystem::OnActorRegistered(class AActor* Actor)
+{
+	SCOPE_CYCLE_COUNTER(STAT_DebugNavOctree);
+	INavRelevantInterface* NavInterface = InterfaceCast<INavRelevantInterface>(Actor);
+	if (NavInterface)
+	{
+		UNavigationSystem* NavSys = UNavigationSystem::GetCurrent(Actor->GetWorld());
+		if (NavSys)
+		{
+			NavSys->RegisterNavOctreeElement(Actor, NavInterface, OctreeUpdate_Modifiers);
+		}
+	}
+}
+
+void UNavigationSystem::OnActorUnregistered(class AActor* Actor)
+{
+	SCOPE_CYCLE_COUNTER(STAT_DebugNavOctree);
+	INavRelevantInterface* NavInterface = InterfaceCast<INavRelevantInterface>(Actor);
+	if (NavInterface)
+	{
+		UNavigationSystem* NavSys = UNavigationSystem::GetCurrent(Actor->GetWorld());
+		if (NavSys)
+		{
+			NavSys->UnregisterNavOctreeElement(Actor, NavInterface, OctreeUpdate_Modifiers);
+		}
+	}
+}
+
 void UNavigationSystem::PopulateNavOctree()
 {
 	UWorld* World = GetWorld();
@@ -2280,10 +2241,10 @@ void UNavigationSystem::PopulateNavOctree()
 		{
 			AActor* Actor = Level->Actors[ActorIndex];
 
-			const bool bLegalActor = Actor && Actor->IsNavigationRelevant() && !Actor->IsPendingKill();
+			const bool bLegalActor = Actor && !Actor->IsPendingKill();
 			if (bLegalActor)
 			{
-				UpdateNavOctree(Actor);
+				UpdateNavOctreeAll(Actor);
 			}
 		}
 	}
@@ -2733,10 +2694,7 @@ void UNavigationSystem::AddLevelCollisionToOctree(ULevel* Level)
 		FNavigationOctreeElement BSPElem;
 		FRecastNavMeshGenerator::ExportVertexSoupGeometry(*LevelGeom, BSPElem.Data);
 
-		BSPElem.Bounds = BSPElem.Data.Bounds;
-		BSPElem.Owner = Level;
-		
-		NavOctree->AddNode(BSPElem);
+		NavOctree->AddNode(Level, NULL, BSPElem.Data.Bounds, BSPElem);
 
 		AddDirtyArea(BSPElem.Bounds.GetBox(), ENavigationDirtyFlag::All);
 	}
diff --git a/Engine/Source/Runtime/Engine/Private/AI/Navigation/NavigationTestingActor.cpp b/Engine/Source/Runtime/Engine/Private/AI/Navigation/NavigationTestingActor.cpp
--- a/Engine/Source/Runtime/Engine/Private/AI/Navigation/NavigationTestingActor.cpp
+++ b/Engine/Source/Runtime/Engine/Private/AI/Navigation/NavigationTestingActor.cpp
@@ -53,6 +53,7 @@ ANavigationTestingActor::ANavigationTestingActor(const class FPostConstructIniti
 	CapsuleComponent->SetCollisionProfileName(CollisionProfileName);
 	CapsuleComponent->CanCharacterStepUpOn = ECB_No;
 	CapsuleComponent->bShouldUpdatePhysicsVolume = true;
+	CapsuleComponent->bCanEverAffectNavigation = false;
 
 	RootComponent = CapsuleComponent;
 }
diff --git a/Engine/Source/Runtime/Engine/Private/AI/Navigation/RecastNavMeshGenerator.cpp b/Engine/Source/Runtime/Engine/Private/AI/Navigation/RecastNavMeshGenerator.cpp
--- a/Engine/Source/Runtime/Engine/Private/AI/Navigation/RecastNavMeshGenerator.cpp
+++ b/Engine/Source/Runtime/Engine/Private/AI/Navigation/RecastNavMeshGenerator.cpp
@@ -400,7 +400,10 @@ static void StoreCollisionCache(FRecastGeometryExport* GeomExport)
 	const int32 CoordsSize = sizeof(float) * 3 * NumVerts;
 	const int32 IndicesSize = sizeof(int32) * 3 * NumFaces;
 	const int32 CacheSize = HeaderSize + CoordsSize + IndicesSize;
-	GeomExport->Data->CollisionData.SetNum(CacheSize);
+	
+	// reserve + add combo to allocate exact amount (without any overhead/slack)
+	GeomExport->Data->CollisionData.Reserve(CacheSize);
+	GeomExport->Data->CollisionData.AddUninitialized(CacheSize);
 
 	// store collisions
 	uint8* RawMemory = GeomExport->Data->CollisionData.GetTypedData();
@@ -952,12 +955,12 @@ FORCEINLINE_DEBUGGABLE void ExportRigidBodySetup(UBodySetup& BodySetup, TNavStat
 	TemporaryShapeBuffer.Reset();
 }
 
-FORCEINLINE_DEBUGGABLE void ExportComponent(UActorComponent& Component, FRecastGeometryExport* GeomExport, const FBox* ClipBounds=NULL)
+FORCEINLINE_DEBUGGABLE void ExportComponent(UActorComponent* Component, FRecastGeometryExport* GeomExport, const FBox* ClipBounds=NULL)
 {
 #if WITH_PHYSX
 	bool bHasData = false;
 
-	UPrimitiveComponent* PrimComp = Cast<UPrimitiveComponent>(&Component);
+	UPrimitiveComponent* PrimComp = Cast<UPrimitiveComponent>(Component);
 	if (PrimComp && PrimComp->IsNavigationRelevant())
 	{
 		if (PrimComp->HasCustomNavigableGeometry() && !PrimComp->DoCustomNavigableGeometryExport(GeomExport)) 
@@ -980,23 +983,6 @@ FORCEINLINE_DEBUGGABLE void ExportComponent(UActorComponent& Component, FRecastG
 #endif // WITH_PHYSX
 }
 
-FORCEINLINE_DEBUGGABLE void ExportActor(AActor& Actor, FRecastGeometryExport* GeomExport, const FBox* ClipBounds=NULL)
-{
-#if WITH_PHYSX
-	TArray<UPrimitiveComponent*> Components;
-	Actor.GetComponents(Components);
-
-	for (int32 i = 0; i < Components.Num(); ++i)
-	{
-		UPrimitiveComponent* PrimComp = Components[i];
-		if (PrimComp->CanEverAffectNavigation() == true)
-		{
-			ExportComponent(*PrimComp, GeomExport, ClipBounds);
-		}
-	}
-#endif // WITH_PHYSX
-}
-
 FORCEINLINE void TransformVertexSoupToRecast(const TArray<FVector>& VertexSoup, TNavStatArray<FVector>& Verts, TNavStatArray<int32>& Faces)
 {
 	if (VertexSoup.Num() == 0)
@@ -4328,15 +4314,7 @@ void FRecastNavMeshGenerator::UpdateTileGenerationWorkers(int32 TileId)
 	}
 }
 
-void FRecastNavMeshGenerator::ExportActorGeometry(AActor& Actor, FNavigationRelevantData& Data)
-{
-	FRecastGeometryExport GeomExport(Data);
-	RecastGeometryExport::ExportActor(Actor, &GeomExport);
-	RecastGeometryExport::CovertCoordDataToRecast(GeomExport.VertexBuffer);
-	RecastGeometryExport::StoreCollisionCache(&GeomExport);
-}
-
-void FRecastNavMeshGenerator::ExportComponentGeometry(UActorComponent& Component, FNavigationRelevantData& Data)
+void FRecastNavMeshGenerator::ExportComponentGeometry(UActorComponent* Component, FNavigationRelevantData& Data)
 {
 	FRecastGeometryExport GeomExport(Data);
 	RecastGeometryExport::ExportComponent(Component, &GeomExport);
diff --git a/Engine/Source/Runtime/Engine/Private/Actor.cpp b/Engine/Source/Runtime/Engine/Private/Actor.cpp
--- a/Engine/Source/Runtime/Engine/Private/Actor.cpp
+++ b/Engine/Source/Runtime/Engine/Private/Actor.cpp
@@ -14,7 +14,6 @@
 #include "Matinee/InterpGroup.h"
 #include "Matinee/InterpGroupInst.h"
 #include "Particles/ParticleSystemComponent.h"
-#include "AI/Navigation/NavRelevantComponent.h"
 #include "VisualLog.h"
 #include "Animation/AnimInstance.h"
 
@@ -70,7 +69,6 @@ AActor::AActor(const class FPostConstructInitializeProperties& PCIP)
 	NetDriverName = NAME_GameNetDriver;
 	NetDormancy = DORM_Awake;
 	// will be updated in PostInitProperties
-	bNavigationRelevant = false;
 	bActorEnableCollision = true;
 	bActorSeamlessTraveled = false;
 	InputConsumeOption_DEPRECATED = ICO_ConsumeBoundKeys;
@@ -1498,9 +1496,10 @@ void AActor::EndPlay(const EEndPlayReason::Type EndPlayReason)
 	if (EndPlayReason == EEndPlayReason::RemovedFromWorld)
 	{
 		bActorInitialized = false;
-		SetNavigationRelevancy(false);
 		GetWorld()->RemoveNetworkActor(this);
 	}
+
+	UNavigationSystem::OnActorUnregistered(this);
 }
 
 FVector AActor::GetPlacementExtent() const
@@ -2042,72 +2041,6 @@ void AActor::DisableComponentsSimulatePhysics()
 
 void AActor::PostRegisterAllComponents() 
 {
-	if (bNavigationRelevant == true && GetWorld()->GetNavigationSystem() != NULL)
-	{
-		// force re-adding to navoctree if already added, in case it has already 
-		// been added (possibly with not all components that could affect navigation
-		// not registered yet
-		GetWorld()->GetNavigationSystem()->UpdateNavOctree(this);
-	}
-}
-
-bool AActor::UpdateNavigationRelevancy()
-{
-	// iterate through all scene components. It's enough to have at least one navigation relevant component
-	bool bNewRelevancy = false;
-
-	TArray<UActorComponent*> Components;
-	GetComponents(Components);
-
-	for (int32 CompIdx = 0; CompIdx < Components.Num() && bNewRelevancy == false; ++CompIdx)
-	{
-		USceneComponent* Comp = Cast<USceneComponent>(Components[CompIdx]);
-		if (Comp)
-		{
-			bNewRelevancy = Comp->IsRegistered() && Comp->IsNavigationRelevant();
-		}
-		else
-		{
-			UNavRelevantComponent* NavComp = Cast<UNavRelevantComponent>(Components[CompIdx]);
-			if (NavComp)
-			{
-				bNewRelevancy = NavComp->IsRegistered() && NavComp->IsNavigationRelevant();
-			}
-		}
-	}
-
-	SetNavigationRelevancy(bNewRelevancy);
-
-	return bNewRelevancy;
-}
-
-void AActor::SetNavigationRelevancy(const bool bNewRelevancy)
-{
-	if (bNavigationRelevant != bNewRelevancy)
-	{
-		// @TODO if this starts failing I'll remodel the way things get registered 
-		// with NavOctree.
-		if(IsInGameThread() == false)
-		{
-			UE_LOG(LogNavigation, Fatal, TEXT("AActor::SetNavigationRelevancy called outside of GameThread for %s"), *GetFullName());
-		}
-		
-		bNavigationRelevant = bNewRelevancy;
-		UWorld* MyWorld = GetWorld();
-		if (MyWorld != NULL && MyWorld->GetNavigationSystem() != NULL)
-		{
-			if (bNewRelevancy == true)
-			{
-				UE_LOG(LogNavigation, Verbose, TEXT("AActor::SetNavigationRelevancy registering new relevant actor %s"), *GetFullName());
-
-				MyWorld->GetNavigationSystem()->RegisterNavigationRelevantActor(this);
-			}
-			else
-			{
-				MyWorld->GetNavigationSystem()->UnregisterNavigationRelevantActor(this);
-			}
-		}
-	}
 }
 
 /** Util to call OnComponentCreated on components */
@@ -2573,7 +2506,6 @@ void AActor::SetActorHiddenInGame( bool bNewHidden )
 	if (bHidden != bNewHidden)
 	{
 		bHidden = bNewHidden;
-		UpdateNavigationRelevancy();
 		MarkComponentsRenderStateDirty();
 	}
 }
@@ -2592,8 +2524,6 @@ void AActor::SetActorEnableCollision(bool bNewActorEnableCollision)
 		{
 			Components[CompIdx]->OnActorEnableCollisionChanged();
 		}
-
-		UpdateNavigationRelevancy();
 	}
 }
 
@@ -3338,7 +3268,7 @@ void AActor::PostInitializeComponents()
 	{
 		bActorInitialized = true;
 
-		UpdateNavigationRelevancy();
+		UNavigationSystem::OnActorRegistered(this);
 	}
 }
 
diff --git a/Engine/Source/Runtime/Engine/Private/ActorEditor.cpp b/Engine/Source/Runtime/Engine/Private/ActorEditor.cpp
--- a/Engine/Source/Runtime/Engine/Private/ActorEditor.cpp
+++ b/Engine/Source/Runtime/Engine/Private/ActorEditor.cpp
@@ -79,6 +79,7 @@ void AActor::PostEditMove(bool bFinished)
 		UBlueprint* Blueprint = Cast<UBlueprint>(GetClass()->ClassGeneratedBy);
 		if(Blueprint && (Blueprint->bRunConstructionScriptOnDrag || bFinished) && !FLevelUtils::IsMovingLevel() )
 		{
+			FNavigationLockContext NavLock(GetWorld(), ENavigationLockReason::AllowUnregister);
 			RerunConstructionScripts();
 		}
 	}
@@ -120,12 +121,9 @@ void AActor::PostEditMove(bool bFinished)
 
 	if (bFinished)
 	{
-		// @todo we could diverse between dynamic and not-that-dynamic actors
-		// and handle updating NavOctree differently
-		if (IsNavigationRelevant() == true && GetWorld() != NULL && GetWorld()->GetNavigationSystem() != NULL)
-		{
-			GetWorld()->GetNavigationSystem()->UpdateNavOctree(this);
-		}
+		// update actor and all its components in navigation system after finishing move
+		// USceneComponent::UpdateNavigationData works only in game world
+		UNavigationSystem::UpdateNavOctreeAll(this);
 	}
 }
 
diff --git a/Engine/Source/Runtime/Engine/Private/BlockingVolume.cpp b/Engine/Source/Runtime/Engine/Private/BlockingVolume.cpp
--- a/Engine/Source/Runtime/Engine/Private/BlockingVolume.cpp
+++ b/Engine/Source/Runtime/Engine/Private/BlockingVolume.cpp
@@ -8,16 +8,11 @@ static FName InvisibleWall_NAME(TEXT("InvisibleWall"));
 ABlockingVolume::ABlockingVolume(const class FPostConstructInitializeProperties& PCIP)
 	: Super(PCIP)
 {
+	BrushComponent->bCanEverAffectNavigation = true;
 	BrushComponent->BodyInstance.bEnableCollision_DEPRECATED = true;
 	BrushComponent->SetCollisionProfileName(InvisibleWall_NAME);
 }
 
-bool ABlockingVolume::UpdateNavigationRelevancy() 
-{ 
-	SetNavigationRelevancy(GetActorEnableCollision() && RootComponent->IsRegistered()); 
-	return IsNavigationRelevant(); 
-}
-
 #if WITH_EDITOR
 
 static FName InvisibleWallDynamic_NAME(TEXT("InvisibleWallDynamic"));
diff --git a/Engine/Source/Runtime/Engine/Private/Brush.cpp b/Engine/Source/Runtime/Engine/Private/Brush.cpp
--- a/Engine/Source/Runtime/Engine/Private/Brush.cpp
+++ b/Engine/Source/Runtime/Engine/Private/Brush.cpp
@@ -24,6 +24,7 @@ ABrush::ABrush(const class FPostConstructInitializeProperties& PCIP)
 	BrushComponent = PCIP.CreateDefaultSubobject<UBrushComponent>(this, TEXT("BrushComponent0"));
 	BrushComponent->Mobility = EComponentMobility::Static;
 	BrushComponent->bGenerateOverlapEvents = false;
+	BrushComponent->bCanEverAffectNavigation = false;
 
 	RootComponent = BrushComponent;
 	
diff --git a/Engine/Source/Runtime/Engine/Private/Character.cpp b/Engine/Source/Runtime/Engine/Private/Character.cpp
--- a/Engine/Source/Runtime/Engine/Private/Character.cpp
+++ b/Engine/Source/Runtime/Engine/Private/Character.cpp
@@ -54,6 +54,7 @@ ACharacter::ACharacter(const class FPostConstructInitializeProperties& PCIP)
 	CapsuleComponent->CanCharacterStepUpOn = ECB_No;
 	CapsuleComponent->bShouldUpdatePhysicsVolume = true;
 	CapsuleComponent->bCheckAsyncSceneOnMove = false;	
+	CapsuleComponent->bCanEverAffectNavigation = false;
 	RootComponent = CapsuleComponent;
 
 	JumpKeyHoldTime = 0.0f;
@@ -94,6 +95,7 @@ ACharacter::ACharacter(const class FPostConstructInitializeProperties& PCIP)
 		static FName CollisionProfileName(TEXT("CharacterMesh"));
 		Mesh->SetCollisionProfileName(CollisionProfileName);
 		Mesh->bGenerateOverlapEvents = false;
+		Mesh->bCanEverAffectNavigation = false;
 	}
 }
 
@@ -155,6 +157,11 @@ void ACharacter::GetSimpleCollisionCylinder(float& CollisionRadius, float& Colli
 	}
 }
 
+void ACharacter::UpdateNavigationRelevance()
+{
+	CapsuleComponent->SetCanEverAffectNavigation(bCanAffectNavigationGeneration);
+}
+
 void ACharacter::ApplyWorldOffset(const FVector& InOffset, bool bWorldShift)
 {
 	Super::ApplyWorldOffset(InOffset, bWorldShift);
diff --git a/Engine/Source/Runtime/Engine/Private/CheatManager.cpp b/Engine/Source/Runtime/Engine/Private/CheatManager.cpp
--- a/Engine/Source/Runtime/Engine/Private/CheatManager.cpp
+++ b/Engine/Source/Runtime/Engine/Private/CheatManager.cpp
@@ -905,9 +905,10 @@ void UCheatManager::WidgetReflector()
 void UCheatManager::RebuildNavigation()
 {
 #if WITH_NAVIGATION_GENERATOR
-	if (GetWorld() && GetWorld()->GetNavigationSystem())
+	UNavigationSystem* NavSys = UNavigationSystem::GetCurrent(GetWorld());
+	if (NavSys)
 	{
-		GetWorld()->GetNavigationSystem()->Build();
+		NavSys->Build();
 	}
 #endif // WITH_NAVIGATION_GENERATOR
 }
diff --git a/Engine/Source/Runtime/Engine/Private/DefaultPawn.cpp b/Engine/Source/Runtime/Engine/Private/DefaultPawn.cpp
--- a/Engine/Source/Runtime/Engine/Private/DefaultPawn.cpp
+++ b/Engine/Source/Runtime/Engine/Private/DefaultPawn.cpp
@@ -31,6 +31,7 @@ ADefaultPawn::ADefaultPawn(const class FPostConstructInitializeProperties& PCIP)
 
 	CollisionComponent->CanCharacterStepUpOn = ECB_No;
 	CollisionComponent->bShouldUpdatePhysicsVolume = true;
+	CollisionComponent->bCanEverAffectNavigation = false;
 
 	RootComponent = CollisionComponent;
 
@@ -62,6 +63,7 @@ ADefaultPawn::ADefaultPawn(const class FPostConstructInitializeProperties& PCIP)
 		const float Scale = CollisionComponent->GetUnscaledSphereRadius() / 160.f; // @TODO: hardcoding known size of EngineMeshes.Sphere. Should use a unit sphere instead.
 		MeshComponent->SetRelativeScale3D(FVector(Scale));
 		MeshComponent->bGenerateOverlapEvents = false;
+		MeshComponent->bCanEverAffectNavigation = false;
 	}
 
 	// This is the default pawn class, we want to have it be able to move out of the box.
@@ -125,6 +127,10 @@ void ADefaultPawn::SetupPlayerInputComponent(UInputComponent* InputComponent)
 	}
 }
 
+void ADefaultPawn::UpdateNavigationRelevance()
+{
+	CollisionComponent->SetCanEverAffectNavigation(bCanAffectNavigationGeneration);
+}
 
 void ADefaultPawn::MoveRight(float Val)
 {
diff --git a/Engine/Source/Runtime/Engine/Private/DestructibleActor.cpp b/Engine/Source/Runtime/Engine/Private/DestructibleActor.cpp
--- a/Engine/Source/Runtime/Engine/Private/DestructibleActor.cpp
+++ b/Engine/Source/Runtime/Engine/Private/DestructibleActor.cpp
@@ -13,6 +13,7 @@ ADestructibleActor::ADestructibleActor(const class FPostConstructInitializePrope
 	: Super(PCIP)
 {
 	DestructibleComponent = PCIP.CreateDefaultSubobject<UDestructibleComponent>(this, TEXT("DestructibleComponent0"));
+	DestructibleComponent->bCanEverAffectNavigation = bAffectNavigation;
 	RootComponent = DestructibleComponent;
 }
 
diff --git a/Engine/Source/Runtime/Engine/Private/InstancedFoliage.cpp b/Engine/Source/Runtime/Engine/Private/InstancedFoliage.cpp
--- a/Engine/Source/Runtime/Engine/Private/InstancedFoliage.cpp
+++ b/Engine/Source/Runtime/Engine/Private/InstancedFoliage.cpp
@@ -422,11 +422,6 @@ void FFoliageMeshInfo::RemoveInstances(AInstancedFoliageActor* InIFA, const TArr
 		{
 			//!! need to update bounds for this cluster
 			ClusterComponent->RegisterComponent();
-
-			if (ClusterComponent->IsCollisionEnabled() && ClusterComponent->GetWorld()->GetNavigationSystem() != nullptr)
-			{
-				ClusterComponent->GetWorld()->GetNavigationSystem()->UpdateNavOctree(ClusterComponent);
-			}
 			bRemoved = true;
 		}
 	}
diff --git a/Engine/Source/Runtime/Engine/Private/InstancedStaticMesh.cpp b/Engine/Source/Runtime/Engine/Private/InstancedStaticMesh.cpp
--- a/Engine/Source/Runtime/Engine/Private/InstancedStaticMesh.cpp
+++ b/Engine/Source/Runtime/Engine/Private/InstancedStaticMesh.cpp
@@ -1584,10 +1584,7 @@ void UInstancedStaticMeshComponent::AddInstance(const FTransform& InstanceTransf
 
 	MarkRenderStateDirty();
 
-	if (IsCollisionEnabled() && GetWorld()->GetNavigationSystem() != nullptr)
-	{
-		GetWorld()->GetNavigationSystem()->UpdateNavOctree(this);
-	}
+	UNavigationSystem::UpdateNavOctree(this);
 }
 
 void UInstancedStaticMeshComponent::AddInstanceWorldSpace(const FTransform& WorldTransform)
@@ -1678,10 +1675,7 @@ bool UInstancedStaticMeshComponent::RemoveInstance(int32 InstanceIndex)
 	// Indicate we need to update render state to reflect changes
 	MarkRenderStateDirty();
 
-	if (IsCollisionEnabled() && GetWorld()->GetNavigationSystem() != nullptr)
-	{
-		GetWorld()->GetNavigationSystem()->UpdateNavOctree(this);
-	}
+	UNavigationSystem::UpdateNavOctree(this);
 
 	return true;
 }
@@ -1696,10 +1690,7 @@ void UInstancedStaticMeshComponent::ClearInstances()
 	// Indicate we need to update render state to reflect changes
 	MarkRenderStateDirty();
 
-	if (IsCollisionEnabled() && GetWorld()->GetNavigationSystem() != nullptr)
-	{
-		GetWorld()->GetNavigationSystem()->UpdateNavOctree(this);
-	}
+	UNavigationSystem::UpdateNavOctree(this);
 }
 
 int32 UInstancedStaticMeshComponent::GetInstanceCount() const
@@ -1735,7 +1726,21 @@ bool UInstancedStaticMeshComponent::DoCustomNavigableGeometryExport(struct FNavi
 	if (StaticMesh != NULL)
 	{
 		UNavCollision* NavCollision = StaticMesh->NavCollision;
-		if (NavCollision != NULL && NavCollision->bHasConvexGeometry)
+		if (NavCollision && NavCollision->bIsDynamicObstacle)
+		{
+			for (const FInstancedStaticMeshInstanceData& InstanceData : PerInstanceSMData)
+			{
+				const FVector Scale3D = InstanceData.Transform.GetScaleVector();
+				// if any of scales is 0 there's no point in exporting it
+				if (!Scale3D.IsZero())
+				{
+					FCompositeNavModifier Modifiers;
+					NavCollision->GetNavigationModifier(Modifiers, FTransform(InstanceData.Transform) * ComponentToWorld);
+					GeomExport->AddNavModifiers(Modifiers);
+				}
+			}
+		}
+		else if (NavCollision && NavCollision->bHasConvexGeometry)
 		{
 			for (const FInstancedStaticMeshInstanceData& InstanceData : PerInstanceSMData)
 			{
@@ -1810,6 +1815,13 @@ void UInstancedStaticMeshComponent::PostEditChangeChainProperty(FPropertyChanged
 
 	Super::PostEditChangeChainProperty(PropertyChangedEvent);
 }
+
+void UInstancedStaticMeshComponent::PostEditUndo()
+{
+	Super::PostEditUndo();
+
+	UNavigationSystem::UpdateNavOctree(this);
+}
 #endif
 
 bool UInstancedStaticMeshComponent::IsInstanceSelected(int32 InInstanceIndex) const
diff --git a/Engine/Source/Runtime/Engine/Private/Landscape/Landscape.cpp b/Engine/Source/Runtime/Engine/Private/Landscape/Landscape.cpp
--- a/Engine/Source/Runtime/Engine/Private/Landscape/Landscape.cpp
+++ b/Engine/Source/Runtime/Engine/Private/Landscape/Landscape.cpp
@@ -796,38 +796,6 @@ void ALandscapeProxy::UnregisterAllComponents()
 	}
 }
 
-bool ALandscapeProxy::UpdateNavigationRelevancy()
-{
-	SetNavigationRelevancy(bUsedForNavigation);
-
-	return bUsedForNavigation;
-}
-
-FBox ALandscapeProxy::GetComponentsBoundingBox(bool bNonColliding) const
-{
-	FBox Bounds = Super::GetComponentsBoundingBox(bNonColliding);
-
-	USceneComponent* LandscapeRoot = GetRootComponent();
-
-	// can't really tell if it breaks anything, but is needed for navmesh generation
-	if (bUsedForNavigation)
-	{
-		for (auto CollisionComponentIt(CollisionComponents.CreateConstIterator()); CollisionComponentIt; ++CollisionComponentIt)
-		{
-			const ULandscapeHeightfieldCollisionComponent* Component = *CollisionComponentIt;
-			if (Component != NULL)
-			{
-				// Component may not be attached, so we'll calculate bounds manually
-				FTransform ComponentLtWTransform = FTransform(Component->RelativeRotation, Component->RelativeLocation, Component->RelativeScale3D) *
-					FTransform(LandscapeRoot->RelativeRotation, LandscapeRoot->RelativeLocation, LandscapeRoot->RelativeScale3D);
-
-				Bounds += Component->CalcBounds(ComponentLtWTransform).GetBox();
-			}
-		}
-	}
-	return Bounds;
-}
-
 // FLandscapeWeightmapUsage serializer
 FArchive& operator<<(FArchive& Ar, FLandscapeWeightmapUsage& U)
 {
@@ -946,10 +914,7 @@ void ALandscapeProxy::PreEditUndo()
 	{
 		GEngine->DeferredCommands.AddUnique(TEXT("UpdateLandscapeEditorData"));
 
-		if (GetWorld()->GetNavigationSystem())
-		{
-			GetWorld()->GetNavigationSystem()->UnregisterNavigationRelevantActor(this);
-		}
+		UNavigationSystem::ClearNavOctreeAll(this);
 	}
 }
 
@@ -961,10 +926,7 @@ void ALandscapeProxy::PostEditUndo()
 	{
 		GEngine->DeferredCommands.AddUnique(TEXT("UpdateLandscapeEditorData"));
 
-		if (GetWorld()->GetNavigationSystem())
-		{
-			GetWorld()->GetNavigationSystem()->RegisterNavigationRelevantActor(this);
-		}
+		UNavigationSystem::UpdateNavOctreeAll(this);
 	}
 }
 
diff --git a/Engine/Source/Runtime/Engine/Private/Landscape/LandscapeCollision.cpp b/Engine/Source/Runtime/Engine/Private/Landscape/LandscapeCollision.cpp
--- a/Engine/Source/Runtime/Engine/Private/Landscape/LandscapeCollision.cpp
+++ b/Engine/Source/Runtime/Engine/Private/Landscape/LandscapeCollision.cpp
@@ -1022,11 +1022,14 @@ void ULandscapeHeightfieldCollisionComponent::PostEditImport()
 void ULandscapeHeightfieldCollisionComponent::PostEditUndo()
 {
 	Super::PostEditUndo();
+
 	// Reinitialize physics after undo
 	if (CollisionSizeQuads > 0)
 	{
 		RecreateCollision(false);
 	}
+
+	UNavigationSystem::UpdateNavOctree(this);
 }
 #endif
 
diff --git a/Engine/Source/Runtime/Engine/Private/Landscape/LandscapeEdit.cpp b/Engine/Source/Runtime/Engine/Private/Landscape/LandscapeEdit.cpp
--- a/Engine/Source/Runtime/Engine/Private/Landscape/LandscapeEdit.cpp
+++ b/Engine/Source/Runtime/Engine/Private/Landscape/LandscapeEdit.cpp
@@ -78,6 +78,17 @@ void ULandscapeComponent::UpdateCachedBounds()
 	}
 }
 
+void ULandscapeComponent::UpdateNavigationRelevance()
+{
+	ALandscapeProxy* Proxy = GetLandscapeProxy();
+	if (CollisionComponent && Proxy)
+	{
+		CollisionComponent->bCanEverAffectNavigation = Proxy->bUsedForNavigation;
+
+		UNavigationSystem::UpdateNavOctree(CollisionComponent.Get());
+	}
+}
+
 ULandscapeMaterialInstanceConstant* ALandscapeProxy::GetLayerThumbnailMIC(UMaterialInterface* LandscapeMaterial, FName LayerName, UTexture2D* ThumbnailWeightmap, UTexture2D* ThumbnailHeightmap, ALandscapeProxy* Proxy)
 {
 	if (!LandscapeMaterial)
@@ -701,6 +712,12 @@ void ULandscapeComponent::UpdateCollisionHeightData(const FColor* HeightmapTextu
 		}
 	}
 
+	{
+		// set relevancy for navigation system
+		ALandscapeProxy* LandscapeProxy = CollisionComp->GetLandscapeProxy();
+		CollisionComp->bCanEverAffectNavigation = LandscapeProxy ? LandscapeProxy->bUsedForNavigation : false;
+	}
+
 	if (bRebuild && CollisionProxy)
 	{
 		CollisionProxy->RegisterAllComponents();
@@ -3640,8 +3657,6 @@ void ALandscapeProxy::PostEditChangeChainProperty(FPropertyChangedChainEvent& Pr
 
 void ALandscape::PostEditChangeProperty(FPropertyChangedEvent& PropertyChangedEvent)
 {
-	static const FName NAME_UsedForNavigation = FName(TEXT("bUsedForNavigation"));
-
 	Super::PostEditChangeProperty(PropertyChangedEvent);
 
 	const FName PropertyName = PropertyChangedEvent.Property ? PropertyChangedEvent.Property->GetFName() : NAME_None;
@@ -3649,6 +3664,7 @@ void ALandscape::PostEditChangeProperty(FPropertyChangedEvent& PropertyChangedEv
 	bool ChangedMaterial = false;
 	bool bNeedsRecalcBoundingBox = false;
 	bool bChangedLighting = false;
+	bool bChangedNavRelevance = false;
 	bool bPropagateToProxies = false;
 
 	ULandscapeInfo* Info = GetLandscapeInfo();
@@ -3711,9 +3727,9 @@ void ALandscape::PostEditChangeProperty(FPropertyChangedEvent& PropertyChangedEv
 	{
 		ExportLOD = FMath::Clamp<int32>(ExportLOD, 0, FMath::CeilLogTwo(SubsectionSizeQuads + 1) - 1);
 	}
-	else if (GIsEditor && PropertyName == NAME_UsedForNavigation)
+	else if (GIsEditor && PropertyName == GET_MEMBER_NAME_CHECKED(ALandscapeProxy, bUsedForNavigation))
 	{
-		UpdateNavigationRelevancy();
+		bChangedNavRelevance = true;
 	}
 
 	bPropagateToProxies = bPropagateToProxies || bNeedsRecalcBoundingBox || bChangedLighting;
@@ -3763,6 +3779,11 @@ void ALandscape::PostEditChangeProperty(FPropertyChangedEvent& PropertyChangedEv
 					Comp->InvalidateLightingCache();
 				}
 
+				if (bChangedNavRelevance)
+				{
+					Comp->UpdateNavigationRelevance();
+				}
+
 				// Reattach all components
 				FComponentReregisterContext ReregisterContext(Comp);
 			}
diff --git a/Engine/Source/Runtime/Engine/Private/Landscape/LandscapeSplines.cpp b/Engine/Source/Runtime/Engine/Private/Landscape/LandscapeSplines.cpp
--- a/Engine/Source/Runtime/Engine/Private/Landscape/LandscapeSplines.cpp
+++ b/Engine/Source/Runtime/Engine/Private/Landscape/LandscapeSplines.cpp
@@ -8,7 +8,6 @@
 #include "Components/SplineMeshComponent.h"
 #include "ShaderParameterUtils.h"
 #include "StaticMeshResources.h"
-#include "AI/Navigation/RecastHelpers.h"
 #include "Landscape/LandscapeSplineProxies.h"
 #include "Landscape/LandscapeSplinesComponent.h"
 #include "Landscape/LandscapeSplineControlPoint.h"
@@ -897,7 +896,6 @@ void ULandscapeSplineControlPoint::UpdateSplinePoints(bool bUpdateCollision, boo
 		)
 	{
 		OuterSplines->Modify();
-		bNavDirty = true;
 
 		if (MeshComponent)
 		{
@@ -932,12 +930,6 @@ void ULandscapeSplineControlPoint::UpdateSplinePoints(bool bUpdateCollision, boo
 		}
 		else
 		{
-			// UpdateNavOctree won't work if we destroy any components, ghosts get left in the navmesh, so we have to unregister
-			if (OuterSplines->GetWorld()->GetNavigationSystem() != nullptr)
-			{
-				OuterSplines->GetWorld()->GetNavigationSystem()->UnregisterNavigationRelevantActor(OuterSplines->GetOwner());
-			}
-
 			Modify();
 			MeshComponent->GetOuter()->Modify();
 			MeshComponent->DestroyComponent();
@@ -979,7 +971,6 @@ void ULandscapeSplineControlPoint::UpdateSplinePoints(bool bUpdateCollision, boo
 		{
 			MeshComponent->Modify();
 			MeshComponent->BodyInstance.SetCollisionProfileName(CollisionProfile);
-			bNavDirty = true;
 		}
 	}
 
@@ -1034,12 +1025,6 @@ void ULandscapeSplineControlPoint::UpdateSplinePoints(bool bUpdateCollision, boo
 
 	OuterSplines->MarkRenderStateDirty();
 
-	if (bNavDirty && bUpdateCollision && OuterSplines->GetWorld()->GetNavigationSystem() != nullptr)
-	{
-		OuterSplines->GetWorld()->GetNavigationSystem()->UpdateNavOctree(OuterSplines->GetOwner());
-		bNavDirty = false;
-	}
-
 	if (bUpdateAttachedSegments)
 	{
 		for (const FLandscapeSplineConnection& Connection : ConnectedSegments)
@@ -1060,23 +1045,12 @@ void ULandscapeSplineControlPoint::DeleteSplinePoints()
 
 	OuterSplines->MarkRenderStateDirty();
 
-	// UpdateNavOctree won't work if we destroy any components, ghosts get left in the navmesh, so we have to unregister
-	if (OuterSplines->GetWorld()->GetNavigationSystem() != nullptr)
-	{
-		OuterSplines->GetWorld()->GetNavigationSystem()->UnregisterNavigationRelevantActor(OuterSplines->GetOwner());
-	}
-
 	if (MeshComponent != NULL)
 	{
 		MeshComponent->GetOuter()->Modify();
 		MeshComponent->DestroyComponent();
 		MeshComponent = NULL;
 	}
-
-	if (OuterSplines->GetWorld()->GetNavigationSystem() != nullptr)
-	{
-		OuterSplines->GetWorld()->GetNavigationSystem()->UpdateNavOctree(OuterSplines->GetOwner());
-	}
 }
 
 void ULandscapeSplineControlPoint::PostEditUndo()
@@ -1545,16 +1519,6 @@ void ULandscapeSplineSegment::UpdateSplinePoints(bool bUpdateCollision)
 				MeshComponents.Add(MeshComponent);
 			}
 
-			// if the spline mesh isn't the editor mesh or wasn't the editor mesh
-			if (!bUsingEditorMesh || (MeshComponent->StaticMesh != NULL && !MeshComponent->bHiddenInGame))
-			{
-				// either collision is enabled or the mesh component had collision enabled before
-				if (bEnableCollision || MeshComponent->BodyInstance.GetCollisionEnabled() != ECollisionEnabled::NoCollision)
-				{
-					bNavDirty = true;
-				}
-			}
-
 			MeshComponent->SetStaticMesh(Mesh);
 
 			MeshComponent->Materials = MeshEntry->MaterialOverrides;
@@ -1575,12 +1539,6 @@ void ULandscapeSplineSegment::UpdateSplinePoints(bool bUpdateCollision)
 
 		if (iMesh < MeshComponents.Num())
 		{
-			// UpdateNavOctree won't work if we destroy any components, ghosts get left in the navmesh, so we have to unregister
-			if (OuterSplines->GetWorld()->GetNavigationSystem() != nullptr)
-			{
-				OuterSplines->GetWorld()->GetNavigationSystem()->UnregisterNavigationRelevantActor(OuterSplines->GetOwner());
-			}
-
 			for (int32 i = iMesh; i < MeshComponents.Num(); i++)
 			{
 				MeshComponents[i]->GetOuter()->Modify();
@@ -1588,8 +1546,6 @@ void ULandscapeSplineSegment::UpdateSplinePoints(bool bUpdateCollision)
 			}
 
 			MeshComponents.RemoveAt(iMesh, MeshComponents.Num() - iMesh);
-
-			bNavDirty = true;
 		}
 
 		// Second pass:
@@ -1766,27 +1722,14 @@ void ULandscapeSplineSegment::UpdateSplinePoints(bool bUpdateCollision)
 	{
 		if (MeshComponents.Num() > 0)
 		{
-			// UpdateNavOctree won't work if we destroy any components, ghosts get left in the navmesh, so we have to unregister
-			if (OuterSplines->GetWorld()->GetNavigationSystem() != nullptr)
-			{
-				OuterSplines->GetWorld()->GetNavigationSystem()->UnregisterNavigationRelevantActor(OuterSplines->GetOwner());
-			}
-
 			for (auto* MeshComponent : MeshComponents)
 			{
 				MeshComponent->GetOuter()->Modify();
 				MeshComponent->DestroyComponent();
 			}
 			MeshComponents.Empty();
-			bNavDirty = true;
 		}
 	}
-
-	if (bNavDirty && bUpdateCollision && OuterSplines->GetWorld()->GetNavigationSystem() != nullptr)
-	{
-		OuterSplines->GetWorld()->GetNavigationSystem()->UpdateNavOctree(OuterSplines->GetOwner());
-		bNavDirty = false;
-	}
 }
 
 void ULandscapeSplineSegment::UpdateSplineEditorMesh()
@@ -1814,23 +1757,12 @@ void ULandscapeSplineSegment::DeleteSplinePoints()
 
 	OuterSplines->MarkRenderStateDirty();
 
-	// UpdateNavOctree won't work if we destroy any components, ghosts get left in the navmesh, so we have to unregister
-	if (OuterSplines->GetWorld()->GetNavigationSystem() != nullptr)
-	{
-		OuterSplines->GetWorld()->GetNavigationSystem()->UnregisterNavigationRelevantActor(OuterSplines->GetOwner());
-	}
-
 	for (auto* MeshComponent : MeshComponents)
 	{
 		MeshComponent->GetOuter()->Modify();
 		MeshComponent->DestroyComponent();
 	}
 	MeshComponents.Empty();
-
-	if (OuterSplines->GetWorld()->GetNavigationSystem() != nullptr)
-	{
-		OuterSplines->GetWorld()->GetNavigationSystem()->UpdateNavOctree(OuterSplines->GetOwner());
-	}
 }
 
 static bool LineIntersect( const FVector2D& L1Start, const FVector2D& L1End, const FVector2D& L2Start, const FVector2D& L2End, FVector2D& Intersect, float Tolerance = KINDA_SMALL_NUMBER)
diff --git a/Engine/Source/Runtime/Engine/Private/LevelActor.cpp b/Engine/Source/Runtime/Engine/Private/LevelActor.cpp
--- a/Engine/Source/Runtime/Engine/Private/LevelActor.cpp
+++ b/Engine/Source/Runtime/Engine/Private/LevelActor.cpp
@@ -360,14 +360,8 @@ ABrush* UWorld::SpawnBrush()
  */
 bool UWorld::EditorDestroyActor( AActor* ThisActor, bool bShouldModifyLevel )
 {
-#if !WITH_RUNTIME_NAVIGATION_BUILDING
-	// @NOTE: this piece needs to be here to have in-editor navmesh react to 
-	// actors being deleted, while not doing so in pure game mode.
-	if (ThisActor->GetWorld() != NULL && ThisActor->GetWorld()->GetNavigationSystem() != NULL)
-	{
-		ThisActor->GetWorld()->GetNavigationSystem()->UnregisterNavigationRelevantActor(ThisActor);
-	}
-#endif // !WITH_RUNTIME_NAVIGATION_BUILDING
+	UNavigationSystem::OnActorUnregistered(ThisActor);
+
 	bool bReturnValue = DestroyActor( ThisActor, false, bShouldModifyLevel );
 	ThisActor->GetWorld()->BroadcastLevelsChanged();
 	return bReturnValue;
@@ -390,12 +384,6 @@ bool UWorld::DestroyActor( AActor* ThisActor, bool bNetForce, bool bShouldModify
 	check(ThisActor->IsValidLowLevel());
 	//UE_LOG(LogSpawn, Log,  "Destroy %s", *ThisActor->GetClass()->GetName() );
 
-#if WITH_RUNTIME_NAVIGATION_BUILDING
-	if (ThisActor->GetWorld() != NULL && ThisActor->GetWorld()->GetNavigationSystem() != NULL)
-	{
-		ThisActor->GetWorld()->GetNavigationSystem()->UnregisterNavigationRelevantActor(ThisActor);
-	}
-#endif // WITH_RUNTIME_NAVIGATION_BUILDING
 	if (ThisActor->GetWorld() == NULL)
 	{
 		UE_LOG(LogSpawn, Warning, TEXT("Destroying %s, which doesn't have a valid world pointer"), *ThisActor->GetPathName());
diff --git a/Engine/Source/Runtime/Engine/Private/Pawn.cpp b/Engine/Source/Runtime/Engine/Private/Pawn.cpp
--- a/Engine/Source/Runtime/Engine/Private/Pawn.cpp
+++ b/Engine/Source/Runtime/Engine/Private/Pawn.cpp
@@ -71,6 +71,8 @@ void APawn::PreInitializeComponents()
 			GetWorld()->PersistentLevel->RegisterActorForAutoReceiveInput(this, PlayerIndex);
 		}
 	}
+
+	UpdateNavigationRelevance();
 }
 
 void APawn::PostInitializeComponents()
@@ -124,25 +126,17 @@ void APawn::UpdateNavAgent()
 	}
 }
 
-bool APawn::UpdateNavigationRelevancy()
-{
-	if (bCanAffectNavigationGeneration == false)
-	{
-		SetNavigationRelevancy(false);
-		return false;
-	}
-	else
-	{
-		return Super::UpdateNavigationRelevancy();
-	}
-}
-
 void APawn::SetCanAffectNavigationGeneration(bool bNewValue)
 {
 	if (bCanAffectNavigationGeneration != bNewValue)
 	{
 		bCanAffectNavigationGeneration = bNewValue;
-		UpdateNavigationRelevancy();
+
+		// update components 
+		UpdateNavigationRelevance();
+
+		// update entries in navigation octree 
+		UNavigationSystem::UpdateNavOctreeAll(this);
 	}
 }
 
diff --git a/Engine/Source/Runtime/Engine/Private/PrimitiveComponent.cpp b/Engine/Source/Runtime/Engine/Private/PrimitiveComponent.cpp
--- a/Engine/Source/Runtime/Engine/Private/PrimitiveComponent.cpp
+++ b/Engine/Source/Runtime/Engine/Private/PrimitiveComponent.cpp
@@ -101,6 +101,7 @@ UPrimitiveComponent::UPrimitiveComponent(const class FPostConstructInitializePro
 	bCheckAsyncSceneOnMove = false;
 	bReturnMaterialOnMove = false;
 	bCanEverAffectNavigation = false;
+	bNavigationRelevant = false;
 
 	bCachedAllCollideableDescendantsRelative = false;
 	LastCheckedAllCollideableDescendantsTime = 0.f;
@@ -292,10 +293,10 @@ void UPrimitiveComponent::OnRegister()
 	// Notify the streaming system. Will only update the component data if it's already tracked.
 	IStreamingManager::Get().NotifyPrimitiveUpdated(this);
 
-	AActor* Owner = GetOwner();
-	if (bCanEverAffectNavigation && Owner != NULL)
+	if (bCanEverAffectNavigation)
 	{
-		Owner->UpdateNavigationRelevancy();
+		UNavigationSystem::OnComponentRegistered(this);
+		bNavigationRelevant = IsNavigationRelevant();
 	}
 }
 
@@ -309,10 +310,9 @@ void UPrimitiveComponent::OnUnregister()
 
 	Super::OnUnregister();
 
-	AActor* Owner = GetOwner();
-	if (bCanEverAffectNavigation && Owner != NULL)
+	if (bCanEverAffectNavigation)
 	{
-		Owner->UpdateNavigationRelevancy();
+		UNavigationSystem::OnComponentUnregistered(this);
 	}
 }
 
@@ -552,7 +552,6 @@ void UPrimitiveComponent::Serialize(FArchive& Ar)
 #if WITH_EDITOR
 void UPrimitiveComponent::PostEditChangeProperty(FPropertyChangedEvent& PropertyChangedEvent)
 {
-	static const FName NAME_CollisionEnabled = TEXT("CollisionEnabled");
 	static const FName NAME_SimulatePhysics = TEXT("bSimulatePhysics");
 	// Keep track of old cached cull distance to see whether we need to re-attach component.
 	const float OldCachedMaxDrawDistance = CachedMaxDrawDistance;
@@ -593,15 +592,6 @@ void UPrimitiveComponent::PostEditChangeProperty(FPropertyChangedEvent& Property
 
 	Super::PostEditChangeProperty(PropertyChangedEvent);
 
-	if(PropertyThatChanged != NULL && PropertyThatChanged->GetFName() == NAME_CollisionEnabled)
-	{
-		AActor* Owner = GetOwner();
-		if ( Owner )
-		{
-			Owner->UpdateNavigationRelevancy();
-		}
-	}
-
 	// Make sure cached cull distance is up-to-date.
 	if( LDMaxDrawDistance > 0 )
 	{
@@ -1659,7 +1649,7 @@ void UPrimitiveComponent::DispatchBlockingHit(AActor& Owner, FHitResult const& B
 }
 
 
-bool UPrimitiveComponent::IsNavigationRelevant(bool bSkipCollisionEnabledCheck) const 
+bool UPrimitiveComponent::IsNavigationRelevant() const 
 { 
 	if (!CanEverAffectNavigation())
 	{
@@ -1671,19 +1661,37 @@ bool UPrimitiveComponent::IsNavigationRelevant(bool bSkipCollisionEnabledCheck)
 		return true;
 	}
 
-	const FCollisionResponseContainer & ResponseToChannels = GetCollisionResponseToChannels();
-
-	return (IsCollisionEnabled() || bSkipCollisionEnabledCheck == true)
-		&& (ResponseToChannels.GetResponse(ECC_Pawn) == ECR_Block 
-			|| ResponseToChannels.GetResponse(ECC_Vehicle) == ECR_Block);
+	const FCollisionResponseContainer& ResponseToChannels = GetCollisionResponseToChannels();
+	return IsCollisionEnabled() &&
+		(ResponseToChannels.GetResponse(ECC_Pawn) == ECR_Block || ResponseToChannels.GetResponse(ECC_Vehicle) == ECR_Block);
 }
 
-void UPrimitiveComponent::DisableNavigationRelevance()
+FBox UPrimitiveComponent::GetNavigationBounds() const
 {
-	check(!bRegistered);
-	bCanEverAffectNavigation = false;
+	return Bounds.GetBox();
 }
 
+void UPrimitiveComponent::SetCanEverAffectNavigation(bool bRelevant)
+{
+	if (bCanEverAffectNavigation != bRelevant)
+	{
+		bCanEverAffectNavigation = bRelevant;
+
+		// update octree if already registered
+		if (bRegistered)
+		{
+			if (bRelevant)
+			{
+				UNavigationSystem::OnComponentRegistered(this);
+				bNavigationRelevant = IsNavigationRelevant();
+			}
+			else
+			{
+				UNavigationSystem::OnComponentUnregistered(this);
+			}
+		}
+	}
+}
 
 //////////////////////////////////////////////////////////////////////////
 // COLLISION
diff --git a/Engine/Source/Runtime/Engine/Private/PrimitiveComponentPhysics.cpp b/Engine/Source/Runtime/Engine/Private/PrimitiveComponentPhysics.cpp
--- a/Engine/Source/Runtime/Engine/Private/PrimitiveComponentPhysics.cpp
+++ b/Engine/Source/Runtime/Engine/Private/PrimitiveComponentPhysics.cpp
@@ -507,8 +507,6 @@ void UPrimitiveComponent::SetCollisionResponseToChannels(const FCollisionRespons
 
 void UPrimitiveComponent::SetCollisionEnabled(ECollisionEnabled::Type NewType)
 {
-	const bool bOldNavRelevant = IsRegistered() && IsNavigationRelevant(true);
-
 	if (BodyInstance.GetCollisionEnabled() != NewType)
 	{
 		BodyInstance.SetCollisionEnabled(NewType);
@@ -516,25 +514,11 @@ void UPrimitiveComponent::SetCollisionEnabled(ECollisionEnabled::Type NewType)
 
 		EnsurePhysicsStateCreated();
 	}
-
-	AActor* Owner = GetOwner();
-	if (Owner && IsRegistered())
-	{
-		const bool bNewNavRelevant = IsNavigationRelevant(true);
-		const bool bOwnerNavRelevant = Owner->IsNavigationRelevant();
-
-		if (bNewNavRelevant != bOldNavRelevant || bNewNavRelevant != bOwnerNavRelevant)
-		{
-			Owner->UpdateNavigationRelevancy();
-		}
-	}
 }
 
 // @todo : implement skeletalmeshcomponent version
 void UPrimitiveComponent::SetCollisionProfileName(FName InCollisionProfileName)
 {
-	const bool bOldNavRelevant = IsRegistered() && IsNavigationRelevant(true);
-
 	ECollisionEnabled::Type OldCollisionEnabled = BodyInstance.GetCollisionEnabled();
 	BodyInstance.SetCollisionProfileName(InCollisionProfileName);
 	OnComponentCollisionSettingsChanged();
@@ -545,18 +529,6 @@ void UPrimitiveComponent::SetCollisionProfileName(FName InCollisionProfileName)
 	{
 		EnsurePhysicsStateCreated();
 	}
-
-	AActor* Owner = GetOwner();
-	if (Owner && IsRegistered())
-	{
-		const bool bNewNavRelevant = IsNavigationRelevant(true);
-		const bool bOwnerNavRelevant = Owner->IsNavigationRelevant();
-
-		if (bNewNavRelevant != bOldNavRelevant || bNewNavRelevant != bOwnerNavRelevant)
-		{
-			Owner->UpdateNavigationRelevancy();
-		}
-	}
 }
 
 FName UPrimitiveComponent::GetCollisionProfileName()
@@ -576,6 +548,14 @@ void UPrimitiveComponent::OnComponentCollisionSettingsChanged()
 	{
 		// changing collision settings could affect touching status, need to update
 		UpdateOverlaps();
+
+		// update navigation data if needed
+		const bool bNewNavRelevant = IsNavigationRelevant();
+		if (bNavigationRelevant != bNewNavRelevant)
+		{
+			bNavigationRelevant = bNewNavRelevant;
+			UNavigationSystem::UpdateNavOctree(this);
+		}
 	}
 }
 
diff --git a/Engine/Source/Runtime/Engine/Private/SceneComponent.cpp b/Engine/Source/Runtime/Engine/Private/SceneComponent.cpp
--- a/Engine/Source/Runtime/Engine/Private/SceneComponent.cpp
+++ b/Engine/Source/Runtime/Engine/Private/SceneComponent.cpp
@@ -1726,42 +1726,12 @@ const int32 USceneComponent::GetNumUncachedStaticLightingInteractions() const
 
 void USceneComponent::UpdateNavigationData()
 {
-	AActor* MyActor = GetOwner();
-
 	if (UNavigationSystem::ShouldUpdateNavOctreeOnComponentChange() &&
-		IsRegistered() && MyActor && MyActor->IsNavigationRelevant() &&
-		World && World->IsGameWorld() && World->GetNetMode() < ENetMode::NM_Client)
-	{
-		const bool bShouldUpdate = CheckNavigationRelevancy(this);
-		UNavigationSystem* NavSys = UNavigationSystem::GetCurrent(World);
-		
-		if (bShouldUpdate && NavSys)
-		{
-			NavSys->UpdateNavOctree(MyActor);
-		}
-	}
-}
-
-bool USceneComponent::CheckNavigationRelevancy(USceneComponent* TestComponent)
-{
-	check(TestComponent);
-
-	bool bResult = TestComponent->IsNavigationRelevant();
-	for (int32 Idx = 0; Idx < TestComponent->AttachChildren.Num() && !bResult; Idx++)
+		IsRegistered() && World && World->IsGameWorld() &&
+		World->GetNetMode() < ENetMode::NM_Client)
 	{
-		USceneComponent* Child = TestComponent->AttachChildren[Idx];
-		if (Child)
-		{
-			bResult = CheckNavigationRelevancy(Child);
-		}
+		UNavigationSystem::UpdateNavOctree(this);
 	}
-
-	return bResult;
-}
-
-bool USceneComponent::IsNavigationRelevant(bool bSkipCollisionEnabledCheck) const
-{
-	return false;
 }
 
 #undef LOCTEXT_NAMESPACE
diff --git a/Engine/Source/Runtime/Engine/Private/WheeledVehicle.cpp b/Engine/Source/Runtime/Engine/Private/WheeledVehicle.cpp
--- a/Engine/Source/Runtime/Engine/Private/WheeledVehicle.cpp
+++ b/Engine/Source/Runtime/Engine/Private/WheeledVehicle.cpp
@@ -23,6 +23,7 @@ AWheeledVehicle::AWheeledVehicle(const class FPostConstructInitializeProperties&
 	Mesh->BodyInstance.bUseCCD = true;
 	Mesh->bBlendPhysics = true;
 	Mesh->bGenerateOverlapEvents = true;
+	Mesh->bCanEverAffectNavigation = false;
 	RootComponent = Mesh;
 
 	VehicleMovement = PCIP.CreateDefaultSubobject<UWheeledVehicleMovementComponent, UWheeledVehicleMovementComponent4W>(this, VehicleMovementComponentName);
diff --git a/Engine/Source/Runtime/Engine/Public/AI/Navigation/RecastNavMeshGenerator.h b/Engine/Source/Runtime/Engine/Public/AI/Navigation/RecastNavMeshGenerator.h
--- a/Engine/Source/Runtime/Engine/Public/AI/Navigation/RecastNavMeshGenerator.h
+++ b/Engine/Source/Runtime/Engine/Public/AI/Navigation/RecastNavMeshGenerator.h
@@ -377,8 +377,7 @@ public:
 	/** 
 	 *	@param Actor is a reference to make callee responsible for assuring it's valid
 	 */
-	static void ExportActorGeometry(AActor& Actor, FNavigationRelevantData& Data);
-	static void ExportComponentGeometry(UActorComponent& Component, FNavigationRelevantData& Data);
+	static void ExportComponentGeometry(UActorComponent* Component, FNavigationRelevantData& Data);
 	static void ExportVertexSoupGeometry(const TArray<FVector>& Verts, FNavigationRelevantData& Data);
 
 	static void ExportRigidBodyGeometry(UBodySetup& BodySetup, TNavStatArray<FVector>& OutVertexBuffer, TNavStatArray<int32>& OutIndexBuffer, const FTransform& LocalToWorld = FTransform::Identity);
diff --git a/Engine/Source/Runtime/Engine/Public/AI/NavigationModifier.h b/Engine/Source/Runtime/Engine/Public/AI/NavigationModifier.h
--- a/Engine/Source/Runtime/Engine/Public/AI/NavigationModifier.h
+++ b/Engine/Source/Runtime/Engine/Public/AI/NavigationModifier.h
@@ -193,6 +193,7 @@ struct ENGINE_API FCompositeNavModifier : public FNavigationModifier
 {
 	FCompositeNavModifier() : bHasPotentialLinks(false) {}
 
+	void Shrink();
 	void Reset();
 	void Empty();
 
diff --git a/Engine/Source/Runtime/Engine/Public/AI/NavigationOctree.h b/Engine/Source/Runtime/Engine/Public/AI/NavigationOctree.h
--- a/Engine/Source/Runtime/Engine/Public/AI/NavigationOctree.h
+++ b/Engine/Source/Runtime/Engine/Public/AI/NavigationOctree.h
@@ -56,6 +56,7 @@ struct ENGINE_API FNavigationRelevantData
 	}
 
 	bool IsMatchingFilter(const FNavigationOctreeFilter& Filter) const;
+	void Shrink();
 };
 
 struct ENGINE_API FNavigationOctreeElement
@@ -105,6 +106,11 @@ struct ENGINE_API FNavigationOctreeElement
 	{
 		return Data.GetAllocatedSize();
 	}
+
+	FORCEINLINE void Shrink()
+	{
+		Data.Shrink();
+	}
 #endif // NAVOCTREE_CONTAINS_COLLISION_DATA
 };
 
@@ -138,11 +144,8 @@ class FNavigationOctree : public TOctree<FNavigationOctreeElement, FNavigationOc
 	typedef TOctree<FNavigationOctreeElement, FNavigationOctreeSemantics> Super;
 
 public:
-	DECLARE_DELEGATE_TwoParams(FNavigableGeometryExportDelegate, AActor&, FNavigationRelevantData&);
-	FNavigableGeometryExportDelegate NavigableGeometryExportDelegate;
-
-	DECLARE_DELEGATE_TwoParams(FNavigableGeometryComponentExportDelegate, UActorComponent&, FNavigationRelevantData&);
-	FNavigableGeometryComponentExportDelegate NavigableGeometryComponentExportDelegate;
+	DECLARE_DELEGATE_TwoParams(FNavigableGeometryComponentExportDelegate, UActorComponent*, FNavigationRelevantData&);
+	FNavigableGeometryComponentExportDelegate ComponentExportDelegate;
 
 	enum ENavGeometryStoringMode {
 		SkipNavGeometry,
@@ -152,15 +155,13 @@ public:
 	FNavigationOctree(FVector Origin, float Radius);
 	~FNavigationOctree();
 
-	/** @param Data allow you to access data put into navoctree while not 
-	 *	having to query navoctree about it just after adding */
-	void AddActor(AActor& Actor, FNavigationOctreeElement& Data);
-	void AddComponent(UActorComponent& ActorComp, const FBox& Bounds, FNavigationOctreeElement& Data);
+	/** Add new node and fill it with navigation export data */
+	void AddNode(UObject* ElementOb, INavRelevantInterface* NavElement, const FBox& Bounds, FNavigationOctreeElement& Data);
 
-	/** Add node and update memory stats */
-	void AddNode(FNavigationOctreeElement& Data);
+	/** Append new data to existing node */
+	void AppendToNode(const FOctreeElementId& Id, INavRelevantInterface* NavElement, const FBox& Bounds, FNavigationOctreeElement& Data);
 
-	/** Remove node and update memory stats */
+	/** Remove node */
 	void RemoveNode(const FOctreeElementId& Id);
 
 	void SetNavigableGeometryStoringMode(ENavGeometryStoringMode NavGeometryMode);
diff --git a/Engine/Source/Runtime/Engine/Public/Engine.h b/Engine/Source/Runtime/Engine/Public/Engine.h
--- a/Engine/Source/Runtime/Engine/Public/Engine.h
+++ b/Engine/Source/Runtime/Engine/Public/Engine.h
@@ -32,7 +32,7 @@
 #include "Engine/CullDistanceVolume.h"
 #include "Engine/LevelStreamingVolume.h"
 #include "AI/Navigation/NavMeshBoundsVolume.h"
-#include "AI/Navigation/NavRelevantActorInterface.h"
+#include "AI/Navigation/NavRelevantInterface.h"
 #include "AI/Navigation/NavModifierVolume.h"
 #include "GameFramework/PhysicsVolume.h"
 #include "GameFramework/DefaultPhysicsVolume.h"
