Changes to visual logger. Logs can be logged not only by Actors but by any UObjects now. There is a catch - object has to redirect his logs to actor and this actor will own logs. It should be good enough to log data from that classes. #ue4

diff --git a/Engine/Source/Developer/LogVisualizer/Private/SLogVisualizer.cpp b/Engine/Source/Developer/LogVisualizer/Private/SLogVisualizer.cpp
--- a/Engine/Source/Developer/LogVisualizer/Private/SLogVisualizer.cpp
+++ b/Engine/Source/Developer/LogVisualizer/Private/SLogVisualizer.cpp
@@ -2529,7 +2529,7 @@ void SLogVisualizer::CameraActorSelected(AActor* SelectedActor)
 
 void SLogVisualizer::SelectActor(AActor* SelectedActor)
 {
-	const AActor* LogOwner = SelectedActor->GetVisualLogRedirection();
+	const AActor* LogOwner = FVisualLog::Get().GetVisualLogRedirection(SelectedActor);
 	const int32 LogIndex = LogVisualizer->GetLogIndexForActor(LogOwner);
 	if (LogVisualizer->Logs.IsValidIndex(LogIndex))
 	{
diff --git a/Engine/Source/Runtime/AIModule/Classes/EnvironmentQuery/EnvQueryDebugHelpers.h b/Engine/Source/Runtime/AIModule/Classes/EnvironmentQuery/EnvQueryDebugHelpers.h
--- a/Engine/Source/Runtime/AIModule/Classes/EnvironmentQuery/EnvQueryDebugHelpers.h
+++ b/Engine/Source/Runtime/AIModule/Classes/EnvironmentQuery/EnvQueryDebugHelpers.h
@@ -109,10 +109,15 @@ FArchive& operator<<(FArchive& Ar, EQSDebug::FQueryData& Data)
 { \
 	SCOPE_CYCLE_COUNTER(STAT_VisualLog); \
 	static_assert((ELogVerbosity::Verbosity & ELogVerbosity::VerbosityMask) < ELogVerbosity::NumVerbosity && ELogVerbosity::Verbosity > 0, "Verbosity must be constant and in range."); \
-	if (FVisualLog::Get().IsRecording() && (!FVisualLog::Get().IsAllBlocked() || FVisualLog::Get().InWhitelist(CategoryName.GetCategoryName())) && Cast<AActor>(Query->Owner.Get())) \
+	if (FVisualLog::Get().IsRecording() && (!FVisualLog::Get().IsAllBlocked() || FVisualLog::Get().InWhitelist(CategoryName.GetCategoryName()))) \
 	{ \
-		UE_VLOG(Cast<AActor>(Query->Owner.Get()), CategoryName, Verbosity, TEXT("Executed EQS: \n - Name: '%s' (id=%d, option=%d),\n - All Items: %d,\n - ValidItems: %d"), *Query->QueryName, Query->QueryID, Query->OptionIndex, Query->ItemDetails.Num(), Query->NumValidItems); \
-		UEnvQueryDebugHelpers::LogQuery(Query, CategoryName.GetCategoryName(), ELogVerbosity::Verbosity); \
+		const AActor* OwnerActor = FVisualLog::Get().GetVisualLogRedirection(Query->Owner.Get()); \
+		ensure(OwnerActor != NULL); \
+		if (OwnerActor) \
+		{  \
+			UE_VLOG(OwnerActor, CategoryName, Verbosity, TEXT("Executed EQS: \n - Name: '%s' (id=%d, option=%d),\n - All Items: %d,\n - ValidItems: %d"), *Query->QueryName, Query->QueryID, Query->OptionIndex, Query->ItemDetails.Num(), Query->NumValidItems); \
+			UEnvQueryDebugHelpers::LogQuery(OwnerActor, Query, CategoryName.GetCategoryName(), ELogVerbosity::Verbosity); \
+		} \
 	} \
 }
 
@@ -130,7 +135,7 @@ class AIMODULE_API UEnvQueryDebugHelpers : public UObject
 #if USE_EQS_DEBUGGER || ENABLE_VISUAL_LOG
 	static void QueryToDebugData(struct FEnvQueryInstance* Query, EQSDebug::FQueryData& EQSLocalData);
 	static void QueryToBlobArray(struct FEnvQueryInstance* Query, TArray<uint8>& BlobArray, bool bUseCompression = false);
-	static void  BlobArrayToDebugData(const TArray<uint8>& BlobArray, EQSDebug::FQueryData& EQSLocalData, bool bUseCompression = false);
-	static void LogQuery(struct FEnvQueryInstance* Query, const FName& CategoryName, ELogVerbosity::Type);
+	static void BlobArrayToDebugData(const TArray<uint8>& BlobArray, EQSDebug::FQueryData& EQSLocalData, bool bUseCompression = false);
+	static void LogQuery(const class AActor* LogOwnerActor, struct FEnvQueryInstance* Query, const FName& CategoryName, ELogVerbosity::Type);
 #endif
 };
diff --git a/Engine/Source/Runtime/AIModule/Private/EnvironmentQuery/EnvQueryDebugHelpers.cpp b/Engine/Source/Runtime/AIModule/Private/EnvironmentQuery/EnvQueryDebugHelpers.cpp
--- a/Engine/Source/Runtime/AIModule/Private/EnvironmentQuery/EnvQueryDebugHelpers.cpp
+++ b/Engine/Source/Runtime/AIModule/Private/EnvironmentQuery/EnvQueryDebugHelpers.cpp
@@ -120,11 +120,11 @@ void  UEnvQueryDebugHelpers::BlobArrayToDebugData(const TArray<uint8>& BlobArray
 	}
 }
 
-void UEnvQueryDebugHelpers::LogQuery(struct FEnvQueryInstance* Query, const FName& CategoryName, ELogVerbosity::Type Type)
+void UEnvQueryDebugHelpers::LogQuery(const class AActor* LogOwnerActor, struct FEnvQueryInstance* Query, const FName& CategoryName, ELogVerbosity::Type Type)
 {
 	TArray<uint8> BlobArray;
 	UEnvQueryDebugHelpers::QueryToBlobArray(Query, BlobArray);
-	FVisualLog::Get().GetEntryToWrite(Cast<AActor>(Query->Owner.Get()))->AddDataBlock(EVisLogTags::TAG_EQS, BlobArray, CategoryName);
+	FVisualLog::Get().GetEntryToWrite(LogOwnerActor)->AddDataBlock(EVisLogTags::TAG_EQS, BlobArray, CategoryName);
 }
 
 #endif //ENABLE_VISUAL_LOG
diff --git a/Engine/Source/Runtime/Engine/Classes/GameFramework/Actor.h b/Engine/Source/Runtime/Engine/Classes/GameFramework/Actor.h
--- a/Engine/Source/Runtime/Engine/Classes/GameFramework/Actor.h
+++ b/Engine/Source/Runtime/Engine/Classes/GameFramework/Actor.h
@@ -386,11 +386,6 @@ public:
 	UPROPERTY()
 	TWeakObjectPtr<AActor> ParentComponentActor;	
 
-#if ENABLE_VISUAL_LOG
-	/** indicates where given actor will be writing its visual log data. Defaults to 'this' */
-	const AActor* VLogRedirection;
-#endif // ENABLE_VISUAL_LOG
-
 public:
 
 	/** 
@@ -2033,12 +2028,8 @@ public:
 	 */
 	virtual void GrabDebugSnapshot(struct FVisLogEntry* Snapshot) const {}
 
-	FORCEINLINE const AActor* GetVisualLogRedirection() const { return VLogRedirection; }
-	void RedirectToVisualLog(const AActor* NewRedir);
-
 private:
 	friend class FVisualLog;
-	void SetVisualLogRedirection(const AActor* NewRedir) { VLogRedirection = NewRedir; }
 #endif // ENABLE_VISUAL_LOG
 
 	//* Sets the friendly actor label and name */
diff --git a/Engine/Source/Runtime/Engine/Private/Actor.cpp b/Engine/Source/Runtime/Engine/Private/Actor.cpp
--- a/Engine/Source/Runtime/Engine/Private/Actor.cpp
+++ b/Engine/Source/Runtime/Engine/Private/Actor.cpp
@@ -39,10 +39,6 @@ FOnProcessEvent AActor::ProcessEventDelegate;
 AActor::AActor(const class FPostConstructInitializeProperties& PCIP)
 	: Super(PCIP)
 {
-#if ENABLE_VISUAL_LOG
-	VLogRedirection = this;
-#endif // ENABLE_VISUAL_LOG
-
 	PrimaryActorTick.TickGroup = TG_PrePhysics;
 	// Default to no tick function, but if we set 'never ticks' to false (so there is a tick function) it is enabled by default
 	PrimaryActorTick.bCanEverTick = false;
@@ -2266,13 +2262,6 @@ void AActor::BeginPlay()
 	ReceiveBeginPlay();
 }
 
-#if ENABLE_VISUAL_LOG
-void AActor::RedirectToVisualLog(const AActor* NewRedir)
-{ 
-	FVisualLog::Get().Redirect(this, NewRedir);
-}
-#endif // ENABLE_VISUAL_LOG
-
 void AActor::EnableInput(APlayerController* PlayerController)
 {
 	if (PlayerController)
diff --git a/Engine/Source/Runtime/Engine/Private/Controller.cpp b/Engine/Source/Runtime/Engine/Private/Controller.cpp
--- a/Engine/Source/Runtime/Engine/Private/Controller.cpp
+++ b/Engine/Source/Runtime/Engine/Private/Controller.cpp
@@ -215,7 +215,7 @@ void AController::PostInitializeComponents()
 
 void AController::Possess(APawn* InPawn)
 {
-	REDIRECT_ACTOR_TO_VLOG(InPawn, this);
+	REDIRECT_OBJECT_TO_VLOG(InPawn, this);
 
 	if (InPawn != NULL)
 	{
diff --git a/Engine/Source/Runtime/Engine/Private/VisualLog.cpp b/Engine/Source/Runtime/Engine/Private/VisualLog.cpp
--- a/Engine/Source/Runtime/Engine/Private/VisualLog.cpp
+++ b/Engine/Source/Runtime/Engine/Private/VisualLog.cpp
@@ -169,7 +169,7 @@ bool DecodeFromBase64(const FString& Source, TArray<uint8>& Dest)
 //----------------------------------------------------------------------//
 // FVisLogEntry
 //----------------------------------------------------------------------//
-FVisLogEntry::FVisLogEntry(const class AActor* InActor, TArray<TWeakObjectPtr<AActor> >* Children)
+FVisLogEntry::FVisLogEntry(const class AActor* InActor, TArray<TWeakObjectPtr<UObject> >* Children)
 {
 	if (InActor->IsPendingKill() == false)
 	{
@@ -178,13 +178,16 @@ FVisLogEntry::FVisLogEntry(const class AActor* InActor, TArray<TWeakObjectPtr<AA
 		InActor->GrabDebugSnapshot(this);
 		if (Children != NULL)
 		{
-			TWeakObjectPtr<AActor>* WeakActorPtr = Children->GetTypedData();
+			TWeakObjectPtr<UObject>* WeakActorPtr = Children->GetTypedData();
 			for (int32 Index = 0; Index < Children->Num(); ++Index, ++WeakActorPtr)
 			{
-				if (WeakActorPtr->IsValid())
+				if (WeakActorPtr->IsValid() )
 				{
-					const AActor* ChildActor = WeakActorPtr->Get();
-					ChildActor->GrabDebugSnapshot(this);
+					const AActor* ChildActor = Cast<AActor>(WeakActorPtr->Get());
+					if (ChildActor)
+					{
+						ChildActor->GrabDebugSnapshot(this);
+					}
 				}
 			}
 		}
@@ -521,7 +524,7 @@ void FVisLogEntry::AddDataBlock(const FString& TagName, const TArray<uint8>& Blo
 //----------------------------------------------------------------------//
 // FActorsVisLog
 //----------------------------------------------------------------------//
-FActorsVisLog::FActorsVisLog(const class AActor* Actor, TArray<TWeakObjectPtr<AActor> >* Children)
+FActorsVisLog::FActorsVisLog(const class AActor* Actor, TArray<TWeakObjectPtr<UObject> >* Children)
 	: Name(Actor->GetFName())
 	, FullName(Actor->GetFullName())
 {
@@ -702,8 +705,8 @@ void FVisualLog::SetIsRecording(bool NewRecording, bool bRecordToFile)
 
 FVisLogEntry*  FVisualLog::GetEntryToWrite(const class AActor* Actor)
 {
-	check(Actor && Actor->GetWorld() && Actor->GetVisualLogRedirection());
-	const class AActor* LogOwner = Actor->GetVisualLogRedirection();
+	const class AActor* LogOwner = GetVisualLogRedirection(Actor);
+	check(Actor && Actor->GetWorld() && LogOwner);
 	const float TimeStamp = Actor->GetWorld()->TimeSeconds;
 	TSharedPtr<FActorsVisLog> Log = GetLog(LogOwner);
 	const int32 LastIndex = Log->Entries.Num() - 1;
@@ -733,51 +736,76 @@ void FVisualLog::Cleanup(bool bReleaseMemory)
 	}
 }
 
-void FVisualLog::Redirect(AActor* Actor, const AActor* NewRedirection)
+const class AActor* FVisualLog::GetVisualLogRedirection(const class UObject* Source)
+{
+	for (auto Iterator = RedirectsMap.CreateConstIterator(); Iterator; ++Iterator)
+	{
+		const auto& Children = (*Iterator).Value;
+		if (Children.Find(Source) != INDEX_NONE)
+		{
+			return (*Iterator).Key;
+		}
+	}
+
+	return Cast<AActor>(Source);
+}
+
+void FVisualLog::RedirectToVisualLog(const class UObject* Src, const class AActor* Dest)
+{
+	Redirect( const_cast<class UObject*>(Src), Dest);
+}
+
+void FVisualLog::Redirect(UObject* Source, const AActor* NewRedirection)
 {
 	// sanity check
-	if (Actor == NULL)
+	if (Source == NULL)
 	{ 
 		return;
 	}
 
 	if (NewRedirection != NULL)
 	{
-		NewRedirection = NewRedirection->GetVisualLogRedirection();
+		NewRedirection = GetVisualLogRedirection(NewRedirection);
 	}
-	const AActor* OldRedirect = Actor->GetVisualLogRedirection();
+	const AActor* OldRedirect = GetVisualLogRedirection(Source);
 
 	if (NewRedirection == OldRedirect)
 	{
 		return;
 	}
-	if (NewRedirection == NULL)
+	if (!NewRedirection)
+	{
+		NewRedirection = Cast<AActor>(Source);
+	}
+	if (!NewRedirection)
 	{
-		NewRedirection = Actor;
+		return;
 	}
 
 	// this should log to OldRedirect
-	UE_VLOG(Actor, LogVisual, Display, TEXT("Binding %s to log %s"), *Actor->GetName(), *NewRedirection->GetName());
+	UE_VLOG(Source, LogVisual, Display, TEXT("Binding %s to log %s"), *Source->GetName(), *NewRedirection->GetName());
 
-	TArray<TWeakObjectPtr<AActor> >& NewTargetChildren = RedirectsMap.FindOrAdd(NewRedirection);
+	TArray<TWeakObjectPtr<UObject> >& NewTargetChildren = RedirectsMap.FindOrAdd(NewRedirection);
 
-	Actor->SetVisualLogRedirection(NewRedirection);
-	NewTargetChildren.AddUnique(Actor);
+	NewTargetChildren.AddUnique(Source);
 
 	// now update all actors that have Actor as their VLog redirection
-	TArray<TWeakObjectPtr<AActor> >* Children = RedirectsMap.Find(Actor);
-	if (Children != NULL)
+	AActor* SourceAsActor = Cast<AActor>(Source);
+	if (SourceAsActor)
 	{
-		TWeakObjectPtr<AActor>* WeakActorPtr = Children->GetTypedData();
-		for (int32 Index = 0; Index < Children->Num(); ++Index)
+		TArray<TWeakObjectPtr<UObject> >* Children = RedirectsMap.Find(SourceAsActor);
+		if (Children != NULL)
 		{
-			if (WeakActorPtr->IsValid())
+			TWeakObjectPtr<UObject>* WeakActorPtr = Children->GetTypedData();
+			for (int32 Index = 0; Index < Children->Num(); ++Index)
 			{
-				WeakActorPtr->Get()->SetVisualLogRedirection(NewRedirection);
-				NewTargetChildren.AddUnique(*WeakActorPtr);
+				if (WeakActorPtr->IsValid())
+				{
+					NewTargetChildren.AddUnique(*WeakActorPtr);
+				}
 			}
+			RedirectsMap.Remove(SourceAsActor);
 		}
-		RedirectsMap.Remove(Actor);
 	}
 }
 
diff --git a/Engine/Source/Runtime/Engine/Public/VisualLog.h b/Engine/Source/Runtime/Engine/Public/VisualLog.h
--- a/Engine/Source/Runtime/Engine/Public/VisualLog.h
+++ b/Engine/Source/Runtime/Engine/Public/VisualLog.h
@@ -160,7 +160,7 @@ struct ENGINE_API FVisLogEntry
 	TArray<FDataBlock>	DataBlocks;
 
 
-	FVisLogEntry(const class AActor* Actor, TArray<TWeakObjectPtr<AActor> >* Children);
+	FVisLogEntry(const class AActor* Actor, TArray<TWeakObjectPtr<UObject> >* Children);
 	FVisLogEntry(TSharedPtr<FJsonValue> FromJson);
 	
 	TSharedPtr<FJsonValue> ToJson() const;
@@ -204,70 +204,96 @@ struct ENGINE_API FVisLogEntry
 #define TO_TEXT_LOG(CategoryName, Verbosity, Format, ...) 
 #endif // LOG_TO_UELOG_AS_WELL
 
-#define REDIRECT_TO_VLOG(Dest) RedirectToVisualLog(Dest)
-#define REDIRECT_ACTOR_TO_VLOG(Src, Dest) Src->RedirectToVisualLog(Dest)
+#define REDIRECT_TO_VLOG(Dest) FVisualLog::Get().RedirectToVisualLog(this, Dest)
+#define REDIRECT_OBJECT_TO_VLOG(Src, Dest) FVisualLog::Get().RedirectToVisualLog(Src, Dest)
 
-#define UE_VLOG(Actor, CategoryName, Verbosity, Format, ...) \
+#define UE_VLOG(Object, CategoryName, Verbosity, Format, ...) \
 { \
 	SCOPE_CYCLE_COUNTER(STAT_VisualLog); \
 	static_assert((ELogVerbosity::Verbosity & ELogVerbosity::VerbosityMask) < ELogVerbosity::NumVerbosity && ELogVerbosity::Verbosity > 0, "Verbosity must be constant and in range."); \
 	if (FVisualLog::Get().IsRecording() && (!FVisualLog::Get().IsAllBlocked() || FVisualLog::Get().InWhitelist(CategoryName.GetCategoryName()))) \
 	{ \
-		FVisualLog::Get().LogLine(Actor, CategoryName.GetCategoryName(), ELogVerbosity::Verbosity, FString::Printf(Format, ##__VA_ARGS__)); \
+		const AActor* MyOwnerActor = FVisualLog::Get().GetVisualLogRedirection(Object); \
+		ensure(MyOwnerActor != NULL); \
+		if (MyOwnerActor) \
+		{  \
+			FVisualLog::Get().LogLine(MyOwnerActor, CategoryName.GetCategoryName(), ELogVerbosity::Verbosity, FString::Printf(Format, ##__VA_ARGS__)); \
+		} \
 	} \
 	if (UE_LOG_CHECK_COMPILEDIN_VERBOSITY(CategoryName, Verbosity)) \
 	{ \
-		if (!CategoryName.IsSuppressed(ELogVerbosity::Verbosity) && (Actor) != NULL) \
+		if (!CategoryName.IsSuppressed(ELogVerbosity::Verbosity) && (Object) != NULL) \
 		{ \
 			TO_TEXT_LOG(CategoryName, Verbosity, Format, ##__VA_ARGS__); \
 		} \
 	} \
 }
 
-#define UE_CVLOG(Condition, Actor, CategoryName, Verbosity, Format, ...) if(Condition) \
+#define UE_CVLOG(Condition, Object, CategoryName, Verbosity, Format, ...) if(Condition) \
 { \
-	UE_VLOG(Actor, CategoryName, Verbosity, Format, ##__VA_ARGS__); \
+	UE_VLOG(Object, CategoryName, Verbosity, Format, ##__VA_ARGS__); \
 }
 
-#define UE_VLOG_SEGMENT_THICK(Actor, CategoryName, Verbosity, SegmentStart, SegmentEnd, Color, Thickness, DescriptionFormat, ...) \
+#define UE_VLOG_SEGMENT_THICK(Object, CategoryName, Verbosity, SegmentStart, SegmentEnd, Color, Thickness, DescriptionFormat, ...) \
 { \
 	SCOPE_CYCLE_COUNTER(STAT_VisualLog); \
 	static_assert((ELogVerbosity::Verbosity & ELogVerbosity::VerbosityMask) < ELogVerbosity::NumVerbosity && ELogVerbosity::Verbosity > 0, "Verbosity must be constant and in range."); \
 	if (FVisualLog::Get().IsRecording() && (!FVisualLog::Get().IsAllBlocked() || FVisualLog::Get().InWhitelist(CategoryName.GetCategoryName()))) \
 	{ \
-		FVisualLog::Get().GetEntryToWrite(Actor)->AddElement(SegmentStart, SegmentEnd, CategoryName.GetCategoryName(), Color, FString::Printf(DescriptionFormat, ##__VA_ARGS__), Thickness); \
+		const AActor* OwnerActor = FVisualLog::Get().GetVisualLogRedirection(Object); \
+		ensure(OwnerActor != NULL); \
+		if (OwnerActor) \
+		{  \
+			FVisualLog::Get().GetEntryToWrite(OwnerActor)->AddElement(SegmentStart, SegmentEnd, CategoryName.GetCategoryName(), Color, FString::Printf(DescriptionFormat, ##__VA_ARGS__), Thickness); \
+		} \
 	} \
 }
 
-#define UE_VLOG_SEGMENT(Actor, CategoryName, Verbosity, SegmentStart, SegmentEnd, Color, DescriptionFormat, ...) UE_VLOG_SEGMENT_THICK(Actor, CategoryName, Verbosity, SegmentStart, SegmentEnd, Color, 0, DescriptionFormat, ##__VA_ARGS__)
+#define UE_VLOG_SEGMENT(Object, CategoryName, Verbosity, SegmentStart, SegmentEnd, Color, DescriptionFormat, ...) \
+	UE_VLOG_SEGMENT_THICK(Object, CategoryName, Verbosity, SegmentStart, SegmentEnd, Color, 0, DescriptionFormat, ##__VA_ARGS__)
 
-#define UE_VLOG_LOCATION(Actor, CategoryName, Verbosity, Location, Radius, Color, DescriptionFormat, ...) \
+#define UE_VLOG_LOCATION(Object, CategoryName, Verbosity, Location, Radius, Color, DescriptionFormat, ...) \
 { \
 	SCOPE_CYCLE_COUNTER(STAT_VisualLog); \
 	static_assert((ELogVerbosity::Verbosity & ELogVerbosity::VerbosityMask) < ELogVerbosity::NumVerbosity && ELogVerbosity::Verbosity > 0, "Verbosity must be constant and in range."); \
 	if (FVisualLog::Get().IsRecording() && (!FVisualLog::Get().IsAllBlocked() || FVisualLog::Get().InWhitelist(CategoryName.GetCategoryName()))) \
 	{ \
-		FVisualLog::Get().GetEntryToWrite(Actor)->AddElement(Location, CategoryName.GetCategoryName(), Color, FString::Printf(DescriptionFormat, ##__VA_ARGS__), Radius); \
+		const AActor* OwnerActor = FVisualLog::Get().GetVisualLogRedirection(Object); \
+		ensure(OwnerActor != NULL); \
+		if (OwnerActor) \
+		{  \
+			FVisualLog::Get().GetEntryToWrite(OwnerActor)->AddElement(Location, CategoryName.GetCategoryName(), Color, FString::Printf(DescriptionFormat, ##__VA_ARGS__), Radius); \
+		} \
 	} \
 }
 
-#define UE_VLOG_BOX(Actor, CategoryName, Verbosity, Box, Color, DescriptionFormat, ...) \
+#define UE_VLOG_BOX(Object, CategoryName, Verbosity, Box, Color, DescriptionFormat, ...) \
 { \
 	SCOPE_CYCLE_COUNTER(STAT_VisualLog); \
 	static_assert((ELogVerbosity::Verbosity & ELogVerbosity::VerbosityMask) < ELogVerbosity::NumVerbosity && ELogVerbosity::Verbosity > 0, "Verbosity must be constant and in range."); \
 	if (FVisualLog::Get().IsRecording() && (!FVisualLog::Get().IsAllBlocked() || FVisualLog::Get().InWhitelist(CategoryName.GetCategoryName()))) \
 	{ \
-		FVisualLog::Get().GetEntryToWrite(Actor)->AddElement(Box, CategoryName.GetCategoryName(), Color, FString::Printf(DescriptionFormat, ##__VA_ARGS__)); \
+		const AActor* OwnerActor = FVisualLog::Get().GetVisualLogRedirection(Object); \
+		ensure(OwnerActor != NULL); \
+		if (OwnerActor) \
+		{  \
+			FVisualLog::Get().GetEntryToWrite(OwnerActor)->AddElement(Box, CategoryName.GetCategoryName(), Color, FString::Printf(DescriptionFormat, ##__VA_ARGS__)); \
+		} \
 	} \
 }
 
-#define UE_VLOG_HISTOGRAM(Actor, CategoryName, Verbosity, GraphName, DataName, Data) \
+#define UE_VLOG_HISTOGRAM(Object, CategoryName, Verbosity, GraphName, DataName, Data) \
 { \
 	SCOPE_CYCLE_COUNTER(STAT_VisualLog); \
 	static_assert((ELogVerbosity::Verbosity & ELogVerbosity::VerbosityMask) < ELogVerbosity::NumVerbosity && ELogVerbosity::Verbosity > 0, "Verbosity must be constant and in range."); \
 	if (FVisualLog::Get().IsRecording() && (!FVisualLog::Get().IsAllBlocked() || FVisualLog::Get().InWhitelist(CategoryName.GetCategoryName()))) \
 	{ \
-		FVisualLog::Get().GetEntryToWrite(Actor)->AddHistogramData(Data, CategoryName.GetCategoryName(), GraphName, DataName); \
+		const AActor* OwnerActor = FVisualLog::Get().GetVisualLogRedirection(Object); \
+		ensure(OwnerActor != NULL); \
+		if (OwnerActor) \
+		{  \
+			FVisualLog::Get().GetEntryToWrite(OwnerActor)->AddHistogramData(Data, CategoryName.GetCategoryName(), GraphName, DataName); \
+		} \
 	} \
 }
 
@@ -309,7 +335,7 @@ struct ENGINE_API FActorsVisLog
 	FString FullName;
 	TArray<TSharedPtr<FVisLogEntry> > Entries;
 
-	FActorsVisLog(const class AActor* Actor, TArray<TWeakObjectPtr<AActor> >* Children);
+	FActorsVisLog(const class AActor* Actor, TArray<TWeakObjectPtr<UObject> >* Children);
 	FActorsVisLog(TSharedPtr<FJsonValue> FromJson);
 
 	TSharedPtr<FJsonValue> ToJson() const;
@@ -330,7 +356,7 @@ class ENGINE_API FVisualLog : public FOutputDevice
 {
 public:
 	typedef TMap<const class AActor*, TSharedPtr<FActorsVisLog> > FLogsMap;
-	typedef TMap<const class AActor*, TArray<TWeakObjectPtr<AActor> > > FLogRedirectsMap;
+	typedef TMap<const class AActor*, TArray<TWeakObjectPtr<UObject> > > FLogRedirectsMap;
 
 	DECLARE_DELEGATE_TwoParams(FOnNewLogCreatedDelegate, const AActor*, TSharedPtr<FActorsVisLog>);
 
@@ -345,8 +371,12 @@ public:
 
 	void Cleanup(bool bReleaseMemory = false);
 
-	void Redirect(class AActor* Actor, const class AActor* NewRedirection);
+	void Redirect(class UObject* Source, const class AActor* NewRedirection);
 	
+	const class AActor* GetVisualLogRedirection(const class UObject* Source);
+
+	void RedirectToVisualLog(const class UObject* Src, const class AActor* Dest);
+
 	void LogLine(const class AActor* Actor, const FName& CategoryName, ELogVerbosity::Type Verbosity, const FString& Line);
 
 	const FLogsMap* GetLogs() const { return &LogsMap; }
@@ -426,5 +456,5 @@ private:
 	#define UE_VLOG_BOX(Actor, CategoryName, Verbosity, Box, Color, DescriptionFormat, ...) 
 	#define UE_VLOG_HISTOGRAM(Actor, CategoryName, Verbosity, GraphName, DataName, Data)
 	#define REDIRECT_TO_VLOG(Dest)
-	#define REDIRECT_ACTOR_TO_VLOG(Src, Destination) 
+	#define REDIRECT_OBJECT_TO_VLOG(Src, Destination) 
 #endif //ENABLE_VISUAL_LOG
