- Added Android HardwareSurvey
- Unified some of the hardware survey code to use Set/GetStoredValue() for last survey time
- Added Get/SetStoredValue for IOS (to use NSUserDefaults)
- NOTE: Mobile devices will report very often - settings are per app, and when they are deleted and re-installed, they will report again.

diff --git a/Engine/Build/Android/Java/src/com/epicgames/ue4/GameActivity.java b/Engine/Build/Android/Java/src/com/epicgames/ue4/GameActivity.java
--- a/Engine/Build/Android/Java/src/com/epicgames/ue4/GameActivity.java
+++ b/Engine/Build/Android/Java/src/com/epicgames/ue4/GameActivity.java
@@ -286,12 +286,16 @@ public class GameActivity extends NativeActivity implements GoogleApiClient.Conn
 		nativeSetGlobalActivity();
 		nativeSetWindowInfo(getResources().getConfiguration().orientation == Configuration.ORIENTATION_PORTRAIT, DepthBufferPreference);
 
+		// get the full language code, like en-US
+		// note: this may need to be Locale.getDefault().getLanguage()
+		String Language = java.util.Locale.getDefault().toString();
 
 		Log.debug( "Android version is " + android.os.Build.VERSION.RELEASE );
 		Log.debug( "Android manufacturer is " + android.os.Build.MANUFACTURER );
 		Log.debug( "Android model is " + android.os.Build.MODEL );
+		Log.debug( "OS language is set to " + Language );
 
-		nativeSetAndroidVersionInformation( android.os.Build.VERSION.RELEASE, android.os.Build.MANUFACTURER, android.os.Build.MODEL );
+		nativeSetAndroidVersionInformation( android.os.Build.VERSION.RELEASE, android.os.Build.MANUFACTURER, android.os.Build.MODEL, Language );
 
 		try
 		{
@@ -1054,7 +1058,7 @@ public class GameActivity extends NativeActivity implements GoogleApiClient.Conn
 	public native void nativeSetObbInfo(String PackageName, int Version, int PatchVersion);
 	public native void nativeUpdateAchievements(JavaAchievement[] Achievements);
 	public native void nativeFailedUpdateAchievements();
-	public native void nativeSetAndroidVersionInformation( String AndroidVersion, String PhoneMake, String PhoneModel );
+	public native void nativeSetAndroidVersionInformation( String AndroidVersion, String PhoneMake, String PhoneModel, String OSLanguage );
 
 	public native void nativeConsoleCommand(String commandString);
 	public native void nativeVirtualKeyboardResult(boolean update, String contents);
diff --git a/Engine/Source/Runtime/Core/Private/Android/AndroidMisc.cpp b/Engine/Source/Runtime/Core/Private/Android/AndroidMisc.cpp
--- a/Engine/Source/Runtime/Core/Private/Android/AndroidMisc.cpp
+++ b/Engine/Source/Runtime/Core/Private/Android/AndroidMisc.cpp
@@ -16,6 +16,7 @@ void* FAndroidMisc::NativeWindow = NULL;
 FString FAndroidMisc::AndroidVersion; // version of android we are running eg "4.0.4"
 FString FAndroidMisc::DeviceMake; // make of the device we are running on eg. "samsung"
 FString FAndroidMisc::DeviceModel; // model of the device we are running on eg "SAMSUNG-SGH-I437"
+FString FAndroidMisc::OSLanguage; // language code the device is set to eg "deu"
 
 GenericApplication* FAndroidMisc::CreateApplication()
 {
@@ -242,13 +243,14 @@ bool FAndroidMisc::GetUseVirtualJoysticks()
 
 
 
-void FAndroidMisc::SetVersionInfo( FString InAndroidVersion, FString InDeviceMake, FString InDeviceModel )
+void FAndroidMisc::SetVersionInfo( FString InAndroidVersion, FString InDeviceMake, FString InDeviceModel, FString InOSLanguage )
 {
 	AndroidVersion = InAndroidVersion;
 	DeviceMake = InDeviceMake;
 	DeviceModel = InDeviceModel;
+	OSLanguage = InOSLanguage;
 
-	UE_LOG(LogEngine, Display, TEXT("Android Version Make Model: %s %s %s"), *AndroidVersion, *DeviceMake, *DeviceModel);
+	UE_LOG(LogEngine, Display, TEXT("Android Version Make Model Language: %s %s %s %s"), *AndroidVersion, *DeviceMake, *DeviceModel, *OSLanguage);
 }
 
 const FString FAndroidMisc::GetAndroidVersion()
@@ -266,6 +268,11 @@ const FString FAndroidMisc::GetDeviceModel()
 	return DeviceModel;
 }
 
+const FString FAndroidMisc::GetOSLanguage()
+{
+	return OSLanguage;
+}
+
 uint32 FAndroidMisc::GetKeyMap( uint16* KeyCodes, FString* KeyNames, uint32 MaxMappings )
 {
 #define ADDKEYMAP(KeyCode, KeyName)		if (NumMappings<MaxMappings) { KeyCodes[NumMappings]=KeyCode; if(KeyNames) { KeyNames[NumMappings]=KeyName; } ++NumMappings; };
diff --git a/Engine/Source/Runtime/Core/Private/Android/AndroidSurvey.cpp b/Engine/Source/Runtime/Core/Private/Android/AndroidSurvey.cpp
new file mode 100644
--- /dev/null
+++ b/Engine/Source/Runtime/Core/Private/Android/AndroidSurvey.cpp
@@ -0,0 +1,36 @@
+// Copyright 1998-2014 Epic Games, Inc. All Rights Reserved.
+
+/*=============================================================================
+	AndroidSurvey.cpp: HardwareSurvey implementation
+=============================================================================*/
+#include "Core.h"
+#include "AndroidWindow.h"
+
+bool FAndroidPlatformSurvey::GetSurveyResults(FHardwareSurveyResults& OutResults, bool bWait)
+{
+	FCString::Strcpy(OutResults.Platform, *(FAndroidMisc::GetDeviceMake() + TEXT("-") + FAndroidMisc::GetDeviceModel()));
+	FCString::Strcpy(OutResults.OSVersion, *FAndroidMisc::GetAndroidVersion());
+#if PLATFORM_64BITS
+	OutResults.OSBits = 64;
+#else
+	OutResults.OSBits = 32;
+#endif
+	FCString::Strcpy(OutResults.OSLanguage, *FPlatformMisc::GetDefaultLocale());
+	FCString::Strcpy(OutResults.MultimediaAPI, *FAndroidMisc::GetGLVersion());
+	OutResults.CPUCount = FPlatformMisc::NumberOfCores();
+
+
+	// display 0 is max size
+	int32 ScreenWidth, ScreenHeight;
+	FAndroidWindow::CalculateSurfaceSize(FPlatformMisc::GetHardwareWindow(), ScreenWidth, ScreenHeight);
+
+	OutResults.Displays[0].CurrentModeWidth = ScreenWidth;
+	OutResults.Displays[0].CurrentModeHeight = ScreenHeight;
+	FCString::Strcpy(OutResults.Displays[0].GPUCardName, *FAndroidMisc::GetGPUFamily());
+
+	// display 1 is current size
+	FPlatformRect ViewSize = FAndroidWindow::GetScreenRect();
+	OutResults.Displays[1].CurrentModeWidth = ViewSize.Right - ViewSize.Left;
+	OutResults.Displays[1].CurrentModeHeight = ViewSize.Bottom - ViewSize.Top;
+	return true;
+}
diff --git a/Engine/Source/Runtime/Core/Private/IOS/IOSPlatformMisc.cpp b/Engine/Source/Runtime/Core/Private/IOS/IOSPlatformMisc.cpp
--- a/Engine/Source/Runtime/Core/Private/IOS/IOSPlatformMisc.cpp
+++ b/Engine/Source/Runtime/Core/Private/IOS/IOSPlatformMisc.cpp
@@ -81,6 +81,42 @@ void FIOSPlatformMisc::GetEnvironmentVariable(const TCHAR* VariableName, TCHAR*
 	}
 }
 
+// Make sure that SetStoredValue and GetStoredValue generate the same key
+static NSString* MakeStoredValueKeyName(const FString& SectionName, const FString& KeyName)
+{
+	return [NSString stringWithFString:(SectionName + "/" + KeyName)];
+}
+
+bool FIOSPlatformMisc::SetStoredValue(const FString& InStoreId, const FString& InSectionName, const FString& InKeyName, const FString& InValue)
+{
+	NSUserDefaults* UserSettings = [NSUserDefaults standardUserDefaults];
+
+	// convert input to an NSString
+	NSString* StoredValue = [NSString stringWithFString:InValue];
+
+	// store it
+	[UserSettings setObject:StoredValue forKey:MakeStoredValueKeyName(InSectionName, InKeyName)];
+
+	return true;
+}
+
+bool FIOSPlatformMisc::GetStoredValue(const FString& InStoreId, const FString& InSectionName, const FString& InKeyName, FString& OutValue)
+{
+	NSUserDefaults* UserSettings = [NSUserDefaults standardUserDefaults];
+	
+	// get the stored NSString
+	NSString* StoredValue = [UserSettings objectForKey:MakeStoredValueKeyName(InSectionName, InKeyName)];
+
+	// if it was there, convert back to FString
+	if (StoredValue != nil)
+	{
+		OutValue = StoredValue;
+		return true;
+	}
+
+	return false;
+}
+
 //void FIOSPlatformMisc::LowLevelOutputDebugStringf(const TCHAR *Fmt, ... )
 //{
 //
diff --git a/Engine/Source/Runtime/Core/Public/Android/AndroidMisc.h b/Engine/Source/Runtime/Core/Public/Android/AndroidMisc.h
--- a/Engine/Source/Runtime/Core/Public/Android/AndroidMisc.h
+++ b/Engine/Source/Runtime/Core/Public/Android/AndroidMisc.h
@@ -78,10 +78,11 @@ struct CORE_API FAndroidMisc : public FGenericPlatformMisc
 	static const TCHAR* GetDefaultDeviceProfileName() { return TEXT("Android"); }
 
 	// ANDROID ONLY:
-	static void SetVersionInfo( FString AndroidVersion, FString DeviceMake, FString DeviceModel );
+	static void SetVersionInfo( FString AndroidVersion, FString DeviceMake, FString DeviceModel, FString OSLanguage );
 	static const FString GetAndroidVersion();
 	static const FString GetDeviceMake();
 	static const FString GetDeviceModel();
+	static const FString GetOSLanguage();
 	static FString GetGPUFamily();
 	static FString GetGLVersion();
 	static bool SupportsFloatingPointRenderTargets();
@@ -120,6 +121,7 @@ struct CORE_API FAndroidMisc : public FGenericPlatformMisc
 	static FString AndroidVersion; // version of android we are running eg "4.0.4"
 	static FString DeviceMake; // make of the device we are running on eg. "samsung"
 	static FString DeviceModel; // model of the device we are running on eg "SAMSUNG-SGH-I437"
+	static FString OSLanguage; // language code the device is set to
 
 };
 
diff --git a/Engine/Source/Runtime/Core/Public/Android/AndroidSurvey.h b/Engine/Source/Runtime/Core/Public/Android/AndroidSurvey.h
--- a/Engine/Source/Runtime/Core/Public/Android/AndroidSurvey.h
+++ b/Engine/Source/Runtime/Core/Public/Android/AndroidSurvey.h
@@ -8,11 +8,11 @@
 #pragma once
 
 /**
-* Android implementation of FGenericPlatformSurvey
-**/
+ * Android implementation of FGenericPlatformSurvey
+ **/
 struct FAndroidPlatformSurvey : public FGenericPlatformSurvey
 {
-	// default implementation for now
+	static bool GetSurveyResults(FHardwareSurveyResults& OutResults, bool bWait=false);
 };
 
 typedef FAndroidPlatformSurvey FPlatformSurvey;
\ No newline at end of file
diff --git a/Engine/Source/Runtime/Core/Public/IOS/IOSPlatformMisc.h b/Engine/Source/Runtime/Core/Public/IOS/IOSPlatformMisc.h
--- a/Engine/Source/Runtime/Core/Public/IOS/IOSPlatformMisc.h
+++ b/Engine/Source/Runtime/Core/Public/IOS/IOSPlatformMisc.h
@@ -65,7 +65,9 @@ struct CORE_API FIOSPlatformMisc : public FGenericPlatformMisc
 	static void SetMemoryWarningHandler(void (* Handler)(const FGenericMemoryWarningContext & Context));
 	static bool HasPlatformFeature(const TCHAR* FeatureName);
 	static FString GetDefaultLocale();
-    
+	static bool SetStoredValue(const FString& InStoreId, const FString& InSectionName, const FString& InKeyName, const FString& InValue);
+	static bool GetStoredValue(const FString& InStoreId, const FString& InSectionName, const FString& InKeyName, FString& OutValue);
+
 	//////// Platform specific
 	static void* CreateAutoreleasePool();
 	static void ReleaseAutoreleasePool(void *Pool);
diff --git a/Engine/Source/Runtime/Engine/Private/UnrealEngine.cpp b/Engine/Source/Runtime/Engine/Private/UnrealEngine.cpp
--- a/Engine/Source/Runtime/Engine/Private/UnrealEngine.cpp
+++ b/Engine/Source/Runtime/Engine/Private/UnrealEngine.cpp
@@ -5357,42 +5357,26 @@ bool UEngine::IsHardwareSurveyRequired()
 		return false;
 	}
 
-#if PLATFORM_IOS
-	// look up to see if how long ago the last survey was, if ever
-	NSUserDefaults* UserSettings = [NSUserDefaults standardUserDefaults];
-	NSDate* SurveyDateTime = [UserSettings objectForKey:@"HardwareSurveyDateTime"];
-
-	// if we never did one, do it now
-	if (SurveyDateTime == nil)
-	{
-		return true;
-	}
-
-	// how long has it been (negate since the past is negative)
-	NSTimeInterval SecondsBeforeNow = -[SurveyDateTime timeIntervalSinceNow];
-	
-	// survey again after a 30 days (TimeDiff is in seconds)
-	return SecondsBeforeNow > (30.0 * 24.0 * 60.0 * 60.0);
-
-#elif PLATFORM_DESKTOP
-
+#if PLATFORM_IOS || PLATFORM_ANDROID || PLATFORM_DESKTOP
 	bool bSurveyDone = false;
-	GConfig->GetBool(TEXT("Engine.HardwareSurvey"), TEXT("bHardwareSurveyDone"), bSurveyDone, GEditorGameAgnosticIni);
-
 	bool bSurveyExpired = false;
-	if (bSurveyDone)
+
+	// platform agnostic code to get the last time we did a survey
+	FString LastRecordedTimeString;
+	if (FPlatformMisc::GetStoredValue(TEXT("Epic Games"), TEXT("Unreal Engine/Hardware Survey"), TEXT("HardwareSurveyDateTime"), LastRecordedTimeString))
 	{
-		bSurveyExpired = true;
-		FString SurveyDateTimeString;
-		if (GConfig->GetString(TEXT("Engine.HardwareSurvey"), TEXT("HardwareSurveyDateTime"), SurveyDateTimeString, GEditorGameAgnosticIni))
+		// attempt to convert to FDateTime
+		FDateTime LastRecordedTime;
+		if (FDateTime::Parse(LastRecordedTimeString, LastRecordedTime))
 		{
-			FDateTime SurveyDateTime;
-			if (FDateTime::Parse(SurveyDateTimeString, SurveyDateTime))
-			{
-				FDateTime Now = FDateTime::UtcNow();
-				int MonthsDelta = 12 * (Now.GetYear() - SurveyDateTime.GetYear()) + Now.GetMonth() - SurveyDateTime.GetMonth();
+			bSurveyDone = true;
 
-				bSurveyExpired = MonthsDelta > 1 || (MonthsDelta == 1 && Now.GetDay() >= SurveyDateTime.GetDay());
+			// make sure it was a month ago
+			FTimespan Diff = FDateTime::UtcNow() - LastRecordedTime;
+
+			if (Diff.GetTotalDays() > 30)
+			{
+				bSurveyExpired = true;
 			}
 		}
 	}
@@ -5400,7 +5384,7 @@ bool UEngine::IsHardwareSurveyRequired()
 	return !bSurveyDone || bSurveyExpired;
 #else
 	return false;
-#endif		// PLATFORM_DESKTOP
+#endif
 }
 
 FString UEngine::HardwareSurveyBucketRAM(uint32 MemoryMB)
@@ -5515,12 +5499,15 @@ FString UEngine::HardwareSurveyGetResolutionClass(uint32 LargestDisplayHeight)
 
 void UEngine::OnHardwareSurveyComplete(const FHardwareSurveyResults& SurveyResults)
 {
-#if PLATFORM_IOS
+#if PLATFORM_IOS || PLATFORM_ANDROID || PLATFORM_DESKTOP
 	if (FEngineAnalytics::IsAvailable())
 	{
-		// mark now as last survey time
-		NSUserDefaults* UserSettings = [NSUserDefaults standardUserDefaults];
-		[UserSettings setObject:[NSDate date] forKey:@"HardwareSurveyDateTime"];
+		IAnalyticsProvider& Analytics = FEngineAnalytics::GetProvider();
+
+		// remember last time we did a survey
+		FPlatformMisc::SetStoredValue(TEXT("Epic Games"), TEXT("Unreal Engine/Hardware Survey"), TEXT("HardwareSurveyDateTime"), FDateTime::UtcNow().ToString());
+
+#if PLATFORM_IOS || PLATFORM_ANDROID
 
 		TArray<FAnalyticsEventAttribute> HardwareStatsAttribs;
 		// copy from what IOSPlatformSurvey has filled out
@@ -5535,19 +5522,13 @@ void UEngine::OnHardwareSurveyComplete(const FHardwareSurveyResults& SurveyResul
 		HardwareStatsAttribs.Add(FAnalyticsEventAttribute(TEXT("DisplayResolution"), DisplayResolution));
 		HardwareStatsAttribs.Add(FAnalyticsEventAttribute(TEXT("ViewResolution"), ViewResolution));
 	
-		FEngineAnalytics::GetProvider().RecordEvent(TEXT("IOSHardwareStats"), HardwareStatsAttribs);
-	}
+#if PLATFORM_ANDROID
+		HardwareStatsAttribs.Add(FAnalyticsEventAttribute(TEXT("GPUModel"), SurveyResults.Displays[0].GPUCardName));
+#endif
 
-#elif PLATFORM_DESKTOP
-	if (GConfig)
-	{
-		GConfig->SetBool(TEXT("Engine.HardwareSurvey"), TEXT("bHardwareSurveyDone"), true, GEditorGameAgnosticIni);
-		GConfig->SetString(TEXT("Engine.HardwareSurvey"), TEXT("HardwareSurveyDateTime"), *FDateTime::UtcNow().ToString(), GEditorGameAgnosticIni);
-	}
+		Analytics.RecordEvent(*FString::Printf(TEXT("%sHardwareStats"), FPlatformProperties::IniPlatformName()), HardwareStatsAttribs);
 
-	if (FEngineAnalytics::IsAvailable())
-	{
-		IAnalyticsProvider& Analytics = FEngineAnalytics::GetProvider();
+#elif PLATFORM_DESKTOP
 
 		TArray<FAnalyticsEventAttribute> HardwareWEIAttribs;
 		HardwareWEIAttribs.Add(FAnalyticsEventAttribute(TEXT( "CPU.WEI" ), FString::Printf( TEXT( "%.1f" ), SurveyResults.CPUPerformanceIndex )));
@@ -5648,8 +5629,9 @@ void UEngine::OnHardwareSurveyComplete(const FHardwareSurveyResults& SurveyResul
 		HardwareStatErrorsAttribs.Add(FAnalyticsEventAttribute(TEXT( "LastError.WEI.Detail" ), SurveyResults.LastPerformanceIndexErrorDetail));
 
 		Analytics.RecordEvent(TEXT( "HardwareStatErrors.1" ), HardwareStatErrorsAttribs);
+#endif // PLATFORM_DESKTOP
 	}
-#endif		// PLATFORM_DESKTOP
+#endif
 }
 
 static TAutoConsoleVariable<float> CVarMaxFPS(
diff --git a/Engine/Source/Runtime/Launch/Private/Android/LaunchAndroid.cpp b/Engine/Source/Runtime/Launch/Private/Android/LaunchAndroid.cpp
--- a/Engine/Source/Runtime/Launch/Private/Android/LaunchAndroid.cpp
+++ b/Engine/Source/Runtime/Launch/Private/Android/LaunchAndroid.cpp
@@ -920,7 +920,7 @@ extern "C" jboolean Java_com_epicgames_ue4_GameActivity_nativeIsGooglePlayEnable
 }
 
 
-extern "C" void Java_com_epicgames_ue4_GameActivity_nativeSetAndroidVersionInformation(JNIEnv* jenv, jobject thiz, jstring androidVersion, jstring phoneMake, jstring phoneModel )
+extern "C" void Java_com_epicgames_ue4_GameActivity_nativeSetAndroidVersionInformation(JNIEnv* jenv, jobject thiz, jstring androidVersion, jstring phoneMake, jstring phoneModel, jstring osLanguage )
 {
 	const char *javaAndroidVersion = jenv->GetStringUTFChars(androidVersion, 0 );
 	FString UEAndroidVersion = FString(UTF8_TO_TCHAR( javaAndroidVersion ));
@@ -931,5 +931,8 @@ extern "C" void Java_com_epicgames_ue4_GameActivity_nativeSetAndroidVersionInfor
 	const char *javaPhoneModel = jenv->GetStringUTFChars(phoneModel, 0 );
 	FString UEPhoneModel = FString(UTF8_TO_TCHAR( javaPhoneModel ));
 
-	FAndroidMisc::SetVersionInfo( UEAndroidVersion, UEPhoneMake, UEPhoneModel );
+	const char *javaOSLanguage = jenv->GetStringUTFChars(osLanguage, 0);
+	FString UEOSLanguage = FString(UTF8_TO_TCHAR(javaOSLanguage));
+
+	FAndroidMisc::SetVersionInfo( UEAndroidVersion, UEPhoneMake, UEPhoneModel, UEOSLanguage );
 }
\ No newline at end of file
