UE4 - Fix canvas flipping on Metal

diff --git a/Engine/Source/Runtime/Engine/Private/TileRendering.cpp b/Engine/Source/Runtime/Engine/Private/TileRendering.cpp
--- a/Engine/Source/Runtime/Engine/Private/TileRendering.cpp
+++ b/Engine/Source/Runtime/Engine/Private/TileRendering.cpp
@@ -54,7 +54,7 @@ public:
 		VertexBufferRHI = RHICreateVertexBuffer(Size,BUF_Static,CreateInfo);
 		// lock it
 		void* Buffer = RHILockVertexBuffer(VertexBufferRHI,0,Size,RLM_WriteOnly);
-        	// first vertex element
+		// first vertex element
 		FMaterialTileVertex* DestVertex = (FMaterialTileVertex*)Buffer;
 
 		// fill out the verts
@@ -139,7 +139,7 @@ void FTileRenderer::DrawTile(FRHICommandListImmediate& RHICmdList, const class F
 	FMaterialTileVertex DestVertex[4];
 
 	// create verts
-	if ((IsES2Platform(GRHIShaderPlatform) && !IsPCPlatform(GRHIShaderPlatform)))
+	if (RHINeedsToSwitchVerticalAxis(GRHIShaderPlatform))
 	{
 		DestVertex[0].Initialize(X + SizeX, View.ViewRect.Height() - (Y + SizeY), U + SizeU, V + SizeV);
 		DestVertex[1].Initialize(X, View.ViewRect.Height() - (Y + SizeY), U, V + SizeV);
diff --git a/Engine/Source/Runtime/Engine/Private/UserInterface/Canvas.cpp b/Engine/Source/Runtime/Engine/Private/UserInterface/Canvas.cpp
--- a/Engine/Source/Runtime/Engine/Private/UserInterface/Canvas.cpp
+++ b/Engine/Source/Runtime/Engine/Private/UserInterface/Canvas.cpp
@@ -173,7 +173,7 @@ bool FCanvasBatchedElementRenderItem::Render_RenderThread(FRHICommandListImmedia
 			Gamma = 1.0f;
 		}
 
-		const bool bNeedsToSwitchVerticalAxis = IsES2Platform(GRHIShaderPlatform) && !IsPCPlatform(GRHIShaderPlatform);
+		const bool bNeedsToSwitchVerticalAxis = RHINeedsToSwitchVerticalAxis(GRHIShaderPlatform);
 
 
 		// draw batched items
@@ -245,7 +245,7 @@ bool FCanvasBatchedElementRenderItem::Render_GameThread(const FCanvas* Canvas)
 			BatchedDrawCommand,
 			FBatchedDrawParameters,Parameters,DrawParameters,
 		{
-			const bool bNeedsToSwitchVerticalAxis = IsES2Platform(GRHIShaderPlatform) && !IsPCPlatform(GRHIShaderPlatform);
+			const bool bNeedsToSwitchVerticalAxis = RHINeedsToSwitchVerticalAxis(GRHIShaderPlatform);
 				
 			// draw batched items
 			Parameters.RenderData->BatchedElements.Draw(
diff --git a/Engine/Source/Runtime/RHI/Public/RHIDefinitions.h b/Engine/Source/Runtime/RHI/Public/RHIDefinitions.h
--- a/Engine/Source/Runtime/RHI/Public/RHIDefinitions.h
+++ b/Engine/Source/Runtime/RHI/Public/RHIDefinitions.h
@@ -612,7 +612,7 @@ inline bool RHISupportsTessellation(const EShaderPlatform Platform)
 inline bool RHINeedsToSwitchVerticalAxis(EShaderPlatform Platform)
 {
 	// ES2 needs to flip when rendering to an RT that will be post processed
-	return IsES2Platform(Platform) && !IsPCPlatform(Platform);
+	return IsES2Platform(Platform) && !IsPCPlatform(Platform) && Platform != SP_METAL;
 }
 
 inline bool RHISupportsInstancing(const EShaderPlatform Platform)
