UParticleModuleLocationPrimitiveCylinder was missing its DisplayName metadata, and not showing up in the list in Cascade.

diff --git a/Engine/Source/Runtime/Engine/Classes/Particles/Location/ParticleModuleLocationPrimitiveCylinder.h b/Engine/Source/Runtime/Engine/Classes/Particles/Location/ParticleModuleLocationPrimitiveCylinder.h
--- a/Engine/Source/Runtime/Engine/Classes/Particles/Location/ParticleModuleLocationPrimitiveCylinder.h
+++ b/Engine/Source/Runtime/Engine/Classes/Particles/Location/ParticleModuleLocationPrimitiveCylinder.h
@@ -18,7 +18,7 @@ enum CylinderHeightAxis
 	PMLPC_HEIGHTAXIS_MAX,
 };
 
-UCLASS(editinlinenew, hidecategories=Object, MinimalAPI)
+UCLASS(editinlinenew, hidecategories = Object, MinimalAPI, meta = (DisplayName = "Cylinder"))
 class UParticleModuleLocationPrimitiveCylinder : public UParticleModuleLocationPrimitiveBase
 {
 	GENERATED_UCLASS_BODY()
