revised fix for applying default transform scale at Actor spawn time

- replaced somewhat unsafe/hackish solution (setting default transform scale to zero) with a flag param to indicate whether or not the given transform is "default" (and thus spawn scale can be overridden with the default value from the template)

diff --git a/Engine/Source/Runtime/Engine/Classes/Engine/SCS_Node.h b/Engine/Source/Runtime/Engine/Classes/Engine/SCS_Node.h
--- a/Engine/Source/Runtime/Engine/Classes/Engine/SCS_Node.h
+++ b/Engine/Source/Runtime/Engine/Classes/Engine/SCS_Node.h
@@ -87,7 +87,7 @@ class USCS_Node : public UObject
 #endif
 
 	/** Create the specified component on the actor, then call action on children */
-	void ExecuteNodeOnActor(AActor* Actor, USceneComponent* ParentComponent, const FTransform* RootTransform);
+	void ExecuteNodeOnActor(AActor* Actor, USceneComponent* ParentComponent, const FTransform* RootTransform, bool bIsDefaultTransform);
 
 	/** Returns an array containing this node and all children below it */
 	TArray<USCS_Node*> GetAllNodes();
diff --git a/Engine/Source/Runtime/Engine/Classes/Engine/SimpleConstructionScript.h b/Engine/Source/Runtime/Engine/Classes/Engine/SimpleConstructionScript.h
--- a/Engine/Source/Runtime/Engine/Classes/Engine/SimpleConstructionScript.h
+++ b/Engine/Source/Runtime/Engine/Classes/Engine/SimpleConstructionScript.h
@@ -18,7 +18,7 @@ class USimpleConstructionScript : public UObject
 	ENGINE_API void FixupRootNodeParentReferences();
 
 	/** Execute this script on the supplied actor, creating components */
-	void ExecuteScriptOnActor(AActor* Actor, const FTransform& RootTransform);
+	void ExecuteScriptOnActor(AActor* Actor, const FTransform& RootTransform, bool bIsDefaultTransform);
 #if WITH_EDITOR
 	/** Return the Blueprint associated with this SCS instance */
 	ENGINE_API class UBlueprint* GetBlueprint() const;
diff --git a/Engine/Source/Runtime/Engine/Classes/GameFramework/Actor.h b/Engine/Source/Runtime/Engine/Classes/GameFramework/Actor.h
--- a/Engine/Source/Runtime/Engine/Classes/GameFramework/Actor.h
+++ b/Engine/Source/Runtime/Engine/Classes/GameFramework/Actor.h
@@ -1465,8 +1465,8 @@ public:
 	/** Called after the actor is spawned in the world.  Responsible for setting up actor for play. */
 	void PostSpawnInitialize(FVector const& SpawnLocation, FRotator const& SpawnRotation, AActor* InOwner, APawn* InInstigator, bool bRemoteOwned, bool bNoFail, bool bDeferConstruction);
 
-        /** Called to finish the spawning process, generally in the case of deferred spawning */
-	void FinishSpawning(const FTransform& Transform);
+    /** Called to finish the spawning process, generally in the case of deferred spawning */
+	void FinishSpawning(const FTransform& Transform, bool bIsDefaultTransform = false);
 
 private:
 	/** Called after the actor has run its construction. Responsible for finishing the actor spawn process. */
@@ -1688,8 +1688,9 @@ public:
 	 * Run any construction script for this Actor. Will call OnConstruction.
 	 * @param	Transform			The transform to construct the actor at.
 	 * @param	InstanceDataCache	Optional cache of state to apply to newly created components (e.g. precomputed lighting)
+	 * @param	bIsDefaultTransform	Whether or not the given transform is a "default" transform, in which case it can be overridden by template defaults
 	 */
-	void ExecuteConstruction(const FTransform& Transform, const class FComponentInstanceDataCache* InstanceDataCache);
+	void ExecuteConstruction(const FTransform& Transform, const class FComponentInstanceDataCache* InstanceDataCache, bool bIsDefaultTransform = false);
 
 	/**
 	 * Called when an instance of this class is placed (in editor) or spawned.
diff --git a/Engine/Source/Runtime/Engine/Private/Actor.cpp b/Engine/Source/Runtime/Engine/Private/Actor.cpp
--- a/Engine/Source/Runtime/Engine/Private/Actor.cpp
+++ b/Engine/Source/Runtime/Engine/Private/Actor.cpp
@@ -2116,18 +2116,18 @@ void AActor::PostSpawnInitialize(FVector const& SpawnLocation, FRotator const& S
 	if (!bDeferConstruction)
 	{
 		// Preserve original root component scale
-		const FVector SpawnScale = GetRootComponent() ? GetRootComponent()->RelativeScale3D : FVector::ZeroVector;
-		FinishSpawning(FTransform(SpawnRotation, SpawnLocation, SpawnScale));
+		const FVector SpawnScale = GetRootComponent() ? GetRootComponent()->RelativeScale3D : FVector(1.0f, 1.0f, 1.0f);
+		FinishSpawning(FTransform(SpawnRotation, SpawnLocation, SpawnScale), true);
 	}
 }
 
-void AActor::FinishSpawning(const FTransform& Transform)
+void AActor::FinishSpawning(const FTransform& Transform, bool bIsDefaultTransform)
 {
 	if (ensure(!bHasFinishedSpawning))
 	{
 		bHasFinishedSpawning = true;
 
-		ExecuteConstruction(Transform, nullptr);
+		ExecuteConstruction(Transform, nullptr, bIsDefaultTransform);
 		PostActorConstruction();
 	}
 }
diff --git a/Engine/Source/Runtime/Engine/Private/ActorConstruction.cpp b/Engine/Source/Runtime/Engine/Private/ActorConstruction.cpp
--- a/Engine/Source/Runtime/Engine/Private/ActorConstruction.cpp
+++ b/Engine/Source/Runtime/Engine/Private/ActorConstruction.cpp
@@ -311,15 +311,16 @@ void AActor::RerunConstructionScripts()
 	}
 }
 
-void AActor::ExecuteConstruction(const FTransform& Transform, const FComponentInstanceDataCache* InstanceDataCache)
+void AActor::ExecuteConstruction(const FTransform& Transform, const FComponentInstanceDataCache* InstanceDataCache, bool bIsDefaultTransform)
 {
 	check(!IsPendingKill());
 	check(!HasAnyFlags(RF_BeginDestroyed|RF_FinishDestroyed));
 
 	// ensure that any existing native root component gets this new transform
-	if (GetRootComponent())
+	// we can skip this in the default case as the given transform will be the root component's transform
+	if (RootComponent && !bIsDefaultTransform)
 	{
-		GetRootComponent()->SetWorldTransform(Transform);
+		RootComponent->SetWorldTransform(Transform);
 	}
 
 	// Generate the parent blueprint hierarchy for this actor, so we can run all the construction scripts sequentially
@@ -339,7 +340,7 @@ void AActor::ExecuteConstruction(const FTransform& Transform, const FComponentIn
 				check(CurrentBPGClass);
 				if(CurrentBPGClass->SimpleConstructionScript)
 				{
-					CurrentBPGClass->SimpleConstructionScript->ExecuteScriptOnActor(this, Transform);
+					CurrentBPGClass->SimpleConstructionScript->ExecuteScriptOnActor(this, Transform, bIsDefaultTransform);
 				}
 				// Now that the construction scripts have been run, we can create timelines and hook them up
 				CurrentBPGClass->CreateComponentsForActor(this);
@@ -383,13 +384,7 @@ void AActor::ExecuteConstruction(const FTransform& Transform, const FComponentIn
 #if WITH_EDITOR
 				BillboardComponent->Sprite = (UTexture2D*)(StaticLoadObject(UTexture2D::StaticClass(), NULL, TEXT("/Engine/EditorResources/BadBlueprintSprite.BadBlueprintSprite"), NULL, LOAD_None, NULL));
 #endif
-				FTransform RelativeTransform = Transform;
-				if(RelativeTransform.GetScale3D().IsZero())
-				{
-					RelativeTransform.SetScale3D(FVector(1.0f, 1.0f, 1.0f));
-				}
-
-				BillboardComponent->SetRelativeTransform(RelativeTransform);
+				BillboardComponent->SetRelativeTransform(Transform);
 
 				SetRootComponent(BillboardComponent);
 				FinishAndRegisterComponent(BillboardComponent);
diff --git a/Engine/Source/Runtime/Engine/Private/SCS_Node.cpp b/Engine/Source/Runtime/Engine/Private/SCS_Node.cpp
--- a/Engine/Source/Runtime/Engine/Private/SCS_Node.cpp
+++ b/Engine/Source/Runtime/Engine/Private/SCS_Node.cpp
@@ -24,7 +24,7 @@ USCS_Node::USCS_Node(const class FPostConstructInitializeProperties& PCIP)
 #endif
 }
 
-void USCS_Node::ExecuteNodeOnActor(AActor* Actor, USceneComponent* ParentComponent, const FTransform* RootTransform)
+void USCS_Node::ExecuteNodeOnActor(AActor* Actor, USceneComponent* ParentComponent, const FTransform* RootTransform, bool bIsDefaultTransform)
 {
 	check(Actor != NULL);
 	check((ParentComponent != NULL && !ParentComponent->IsPendingKill()) || (RootTransform != NULL)); // must specify either a parent component or a world transform
@@ -55,13 +55,13 @@ void USCS_Node::ExecuteNodeOnActor(AActor* Actor, USceneComponent* ParentCompone
 			// If NULL is passed in, we are the root, so set transform and assign as RootComponent on Actor
 			if (ParentComponent == NULL || (ParentComponent && ParentComponent->IsPendingKill()))
 			{
-				FTransform SpawnTransform = *RootTransform;
-				if(SpawnTransform.GetScale3D().IsZero())
+				FTransform WorldTransform = *RootTransform;
+				if(bIsDefaultTransform)
 				{
 					// Note: We use the scale vector from the component template when spawning (to match what happens with a native root)
-					SpawnTransform.SetScale3D(NewSceneComp->RelativeScale3D);
+					WorldTransform.SetScale3D(NewSceneComp->RelativeScale3D);
 				}
-				NewSceneComp->SetWorldTransform(SpawnTransform);
+				NewSceneComp->SetWorldTransform(WorldTransform);
 				Actor->SetRootComponent(NewSceneComp);
 			}
 			// Otherwise, attach to parent component passed in
@@ -122,7 +122,7 @@ void USCS_Node::ExecuteNodeOnActor(AActor* Actor, USceneComponent* ParentCompone
 		{
 			USCS_Node* Node = ChildNodes[NodeIdx];
 			check(Node != NULL);
-			Node->ExecuteNodeOnActor(Actor, ParentSceneComponentOfChildren, NULL);
+			Node->ExecuteNodeOnActor(Actor, ParentSceneComponentOfChildren, NULL, false);
 		}
 	}
 }
diff --git a/Engine/Source/Runtime/Engine/Private/SimpleConstructionScript.cpp b/Engine/Source/Runtime/Engine/Private/SimpleConstructionScript.cpp
--- a/Engine/Source/Runtime/Engine/Private/SimpleConstructionScript.cpp
+++ b/Engine/Source/Runtime/Engine/Private/SimpleConstructionScript.cpp
@@ -288,7 +288,7 @@ void USimpleConstructionScript::FixupRootNodeParentReferences()
 	}
 }
 
-void USimpleConstructionScript::ExecuteScriptOnActor(AActor* Actor, const FTransform& RootTransform)
+void USimpleConstructionScript::ExecuteScriptOnActor(AActor* Actor, const FTransform& RootTransform, bool bIsDefaultTransform)
 {
 	if(RootNodes.Num() > 0)
 	{
@@ -333,22 +333,16 @@ void USimpleConstructionScript::ExecuteScriptOnActor(AActor* Actor, const FTrans
 					}
 				}
 
-				RootNode->ExecuteNodeOnActor(Actor, ParentComponent != NULL ? ParentComponent : Actor->GetRootComponent(), &RootTransform);
+				RootNode->ExecuteNodeOnActor(Actor, ParentComponent != NULL ? ParentComponent : Actor->GetRootComponent(), &RootTransform, bIsDefaultTransform);
 			}
 		}
 	}
 	else if(Actor->GetRootComponent() == NULL) // Must have a root component at the end of SCS, so if we don't have one already (from base class), create a SceneComponent now
 	{
-		FTransform SpawnTransform = RootTransform;
-		if(SpawnTransform.GetScale3D().IsZero())
-		{
-			SpawnTransform.SetScale3D(FVector(1.0f, 1.0f, 1.0f));
-		}
-
 		USceneComponent* SceneComp = NewObject<USceneComponent>(Actor);
 		SceneComp->SetFlags(RF_Transactional);
 		SceneComp->bCreatedByConstructionScript = true;
-		SceneComp->SetWorldTransform(SpawnTransform);
+		SceneComp->SetWorldTransform(RootTransform);
 		Actor->SetRootComponent(SceneComp);
 		SceneComp->RegisterComponent();
 	}
