Stats - Added command 'stat display' which may be used to change various stats dispay settings, temporary solution
Stats - Better comment regarding stats commands exec

diff --git a/Engine/Source/Runtime/Core/Private/Stats/StatsCommand.cpp b/Engine/Source/Runtime/Core/Private/Stats/StatsCommand.cpp
--- a/Engine/Source/Runtime/Core/Private/Stats/StatsCommand.cpp
+++ b/Engine/Source/Runtime/Core/Private/Stats/StatsCommand.cpp
@@ -33,53 +33,6 @@ struct TTypeFromString<EStatCompareBy::Type>
 	}
 };
 
-template<>
-struct TTypeFromString<FName>
-{
-	static void FromString( FName& OutValue, const TCHAR* Buffer )
-	{
-		OutValue = FName(Buffer);
-	}
-};
-
-template <typename T>
-void ParseTypedValue( const TCHAR* Stream, const TCHAR* Match, T& Out )
-{
-	TCHAR Temp[64]=TEXT("");
-	if( FParse::Value( Stream, Match, Temp, ARRAY_COUNT(Temp) ) )
-	{
-		TTypeFromString<T>::FromString( Out, Temp );
-	}
-}
-
-/**
- *	Helper class to parse a value from the stream.
- *	If value is not present, provided default value will be used.
- */
-template<typename T>
-struct TParsedValueWithDefault
-{
-public:
-	TParsedValueWithDefault( const TCHAR* Stream, const TCHAR* Match, const T& Default )
-		: Value( Default )
-	{
-		ParseTypedValue<T>( Stream, Match, Value );
-	}
-
-	const T& Get() const
-	{
-		return Value;
-	}
-
-	void Set( const T& NewValue )
-	{
-		Value = NewValue;
-	}
-
-private:
-	T Value;
-};
-
 struct FGroupFilter : public IItemFiler
 {
 	TSet<FName> const& EnabledItems;
@@ -121,6 +74,7 @@ struct FStatGroupParams
 	 *	Maximum number of frames to be included in the history. 
 	 *	-maxhistoryframes=[20:20-120]
 	 */
+	// @TODO yrx 2014-08-21 Replace with TParsedValueWithDefaultAndRange
 	TParsedValueWithDefault<int32> MaxHistoryFrames;
 
 	/** Whether to reset all collected data. */
@@ -1107,6 +1061,12 @@ static void PrintStatsHelpToOutputDevice( FOutputDevice& Ar )
 
 	Ar.Log( TEXT("stat group list|listall|enable name|disable name|none|all|default - manages stats groups"));
 
+#if WITH_ENGINE
+	// Engine @see FStatCmdEngine
+	Ar.Log( TEXT( "stat display -font=small[tiny]" ) );
+	Ar.Log( TEXT( "    Changes stats rendering display options" ) );
+#endif // WITH_ENGINE
+
 	Ar.Log( TEXT("stat startfile - starts dumping a capture"));
 	Ar.Log( TEXT("stat stopfile - stops dumping a capture"));
 }
@@ -1301,7 +1261,8 @@ static void StatCmd(FString InCmd)
 	}
 }
 
-static class FStatCmd : private FSelfRegisteringExec
+/** Exec used to execute core stats commands on the stats thread. */
+static class FStatCmdCore : private FSelfRegisteringExec
 {
 public:
 	/** Console commands, see embeded usage statement **/
@@ -1310,7 +1271,8 @@ public:
 		// Block the thread as this affects external stat states now
 		return DirectStatsCommand(Cmd,true,&Ar);
 	}
-} StatCmdExec;
+}
+StatCmdCoreExec;
 
 bool DirectStatsCommand(const TCHAR* Cmd, bool bBlockForCompletion /*= false*/, FOutputDevice* Ar /*= nullptr*/)
 {
diff --git a/Engine/Source/Runtime/Core/Public/Stats/StatsData.h b/Engine/Source/Runtime/Core/Public/Stats/StatsData.h
--- a/Engine/Source/Runtime/Core/Public/Stats/StatsData.h
+++ b/Engine/Source/Runtime/Core/Public/Stats/StatsData.h
@@ -39,6 +39,53 @@ struct CORE_API FStatConstants
 	static const FString ThreadNameMarker;
 };
 
+template<>
+struct TTypeFromString<FName>
+{
+	static void FromString( FName& OutValue, const TCHAR* Buffer )
+	{
+		OutValue = FName( Buffer );
+	}
+};
+
+template <typename T>
+void ParseTypedValue( const TCHAR* Stream, const TCHAR* Match, T& Out )
+{
+	TCHAR Temp[64] = TEXT( "" );
+	if( FParse::Value( Stream, Match, Temp, ARRAY_COUNT( Temp ) ) )
+	{
+		TTypeFromString<T>::FromString( Out, Temp );
+	}
+}
+
+/**
+*	Helper class to parse a value from the stream.
+*	If value is not present, provided default value will be used.
+*/
+template<typename T>
+struct TParsedValueWithDefault
+{
+public:
+	TParsedValueWithDefault( const TCHAR* Stream, const TCHAR* Match, const T& Default )
+		: Value( Default )
+	{
+		ParseTypedValue<T>( Stream, Match, Value );
+	}
+
+	const T& Get() const
+	{
+		return Value;
+	}
+
+	void Set( const T& NewValue )
+	{
+		Value = NewValue;
+	}
+
+private:
+	T Value;
+};
+
 /** Enumerates stat compare types. */
 struct EStatCompareBy
 {
diff --git a/Engine/Source/Runtime/Engine/Private/StatsRender2.cpp b/Engine/Source/Runtime/Engine/Private/StatsRender2.cpp
--- a/Engine/Source/Runtime/Engine/Private/StatsRender2.cpp
+++ b/Engine/Source/Runtime/Engine/Private/StatsRender2.cpp
@@ -11,112 +11,244 @@
 
 #include "StatsData.h"
 
+/** Stats rendering constants. */
+enum class EStatRenderConsts
+{
+	NUM_COLUMNS = 5,
+};
+
+/** Enumerates stat font types and maximum length of the stat names. */
+enum class EStatFontTypes : int32
+{
+	/** Tiny font, used when ViewRectX < 1280. */
+	Tiny = 0,
+	/** Small font. */
+	Small = 1,
+	/** Number of fonts. */
+	NumFonts,
+};
+
+/** Contains misc stat font properties. */
+struct FStatFont
+{
+	/** Default constructor. */
+	FStatFont( int32 InMaxDisplayedChars, int32 InFontHeight, int32 InFontHeightOffsets ) :
+		MaxDisplayedChars( InMaxDisplayedChars ),
+		FontHeight( InFontHeight ),
+		FontHeightOffset( InFontHeightOffsets )
+	{}
 
-#define STAT_FONT			GEngine->GetSmallFont()
-#define FONT_HEIGHT			12
-// The default small font (which the stat rendering code was written around) is
-// actually 16 in height... So we will use this to calculate the scale factor
-// of the actual font being used to render.
-#define IDEAL_FONT_HEIGHT	16
+	/** Maximum length of the displayed stat names. */
+	const int32 MaxDisplayedChars;
 
-// This is the number of W characters to leave spacing for in the stat name column (32 chars)
-static const FString STAT_COLUMN_WIDTH = FString::ChrN( 32, 'W' );
+	/** Font height, manually selected to fix more stats on screen. */
+	const int32 FontHeight;
 
-// This is the number of W characters to leave spacing for in the other stat columns (8 chars)
-static const FString TIME_COLUMN_WIDTH = FString::ChrN( 8, 'W' );
+	/** Y offset of the background tile, manually selected to fix more stats on screen. */
+	const int32 FontHeightOffset;
+};
 
+static FStatFont GStatFonts[EStatFontTypes::NumFonts] =
+{
+	/** Tiny. */
+	FStatFont( 24, 10, 1 ),
+	/** Small. */
+	FStatFont( 48, 12, 2 ),
+};
+
+template<>
+struct TTypeFromString<EStatFontTypes>
+{
+	static void FromString( EStatFontTypes& OutValue, const TCHAR* Buffer )
+	{
+		OutValue = EStatFontTypes::Small;
+
+		if( FCString::Stricmp( Buffer, TEXT( "Tiny" ) ) == 0 )
+		{
+			OutValue = EStatFontTypes::Tiny;
+		}
+	}
+};
+
+/** Holds various parameters used for rendering stats. */
 struct FStatRenderGlobals
 {
-	/**
-	 * Rendering offset for first column from stat label
-	 */
+	/** Rendering offset for first column from stat label. */
 	int32 AfterNameColumnOffset;
-	/**
-	 * Rendering offsets for additional columns
-	 */
+
+	/** Rendering offsets for additional columns. */
 	int32 InterColumnOffset;
-	/**
-	 * Color for rendering stats
-	 */
+
+	/** Color for rendering stats. */
 	FLinearColor StatColor;
-	/**
-	 * Color to use when rendering headings
-	 */
+
+	/** Color to use when rendering headings. */
 	FLinearColor HeadingColor;
-	/**
-	 * Color to use when rendering group names
-	 */
+
+	/** Color to use when rendering group names. */
 	FLinearColor GroupColor;
-	/**
-	 * Color used as the background for every other stat item to make it easier to read across lines 
-	 */
+
+	/** Color used as the background for every other stat item to make it easier to read across lines. */
 	FLinearColor BackgroundColor;
-	/** The scale to apply to the font used for rendering stats. */
-	float StatFontScale;
 
+	/** The font used for rendering stats. */
+	UFont* StatFont;
 
-	FStatRenderGlobals(void) :
+	/** Current size of the viewport, used to detect resolution changes. */
+	FIntPoint SizeXY;
+
+	/** Current stat font type. */
+	EStatFontTypes StatFontType;
+
+	/** Whether we need update internals. */
+	bool bNeedRefresh;
+
+	FStatRenderGlobals() :
 		AfterNameColumnOffset(0),
 		InterColumnOffset(0),
 		StatColor(0.f,1.f,0.f),
 		HeadingColor(1.f,0.2f,0.f),
 		GroupColor(FLinearColor::White),
 		BackgroundColor(0.05f, 0.05f, 0.05f, 0.90f),	// dark gray mostly occluding the background
-		StatFontScale(0.0f)
+		StatFontType(EStatFontTypes::NumFonts),
+		StatFont(nullptr),
+		bNeedRefresh(true)
 	{
+		SetNewFont( EStatFontTypes::Small );
 	}
 
 	/**
 	 * Initializes stat render globals.
-	 *
-	 * @param bForce - forces to initialize globals again, used after font has been changed
-	 *
 	 */
-	void Initialize( bool bForce = false )
+	void Initialize( const FIntPoint NewSizeXY )
 	{
-		static bool bIsInitialized = false;
-		if( !bIsInitialized || bForce )
+		if( SizeXY != NewSizeXY )
+		{
+			SizeXY = NewSizeXY;
+			bNeedRefresh = true;
+		}
+
+		if( SizeXY.X < 1280 )
+		{
+			SetNewFont( EStatFontTypes::Tiny );
+		}
+
+		if( bNeedRefresh )
 		{
+			// This is the number of W characters to leave spacing for in the stat name column.
+			const FString STATNAME_COLUMN_WIDTH = FString::ChrN( GetNumCharsForStatName(), 'M' );
+
+			// This is the number of W characters to leave spacing for in the other stat columns.
+			const FString STATDATA_COLUMN_WIDTH = FString::ChrN( StatFontType == EStatFontTypes::Small ? 8 : 7, 'W' );
+
 			// Get the size of the spaces, since we can't render the width calculation strings.
-			int32 StatColumnSpaceSizeY;
-			int32 TimeColumnSpaceSizeY;
+			int32 StatColumnSpaceSizeY = 0;
+			int32 TimeColumnSpaceSizeY = 0;
 
 			// Determine where the first column goes.
-			StringSize(STAT_FONT, AfterNameColumnOffset, StatColumnSpaceSizeY, *STAT_COLUMN_WIDTH);
+			StringSize(StatFont, AfterNameColumnOffset, StatColumnSpaceSizeY, *STATNAME_COLUMN_WIDTH);
 
 			// Determine the width of subsequent columns.
-			StringSize(STAT_FONT, InterColumnOffset, TimeColumnSpaceSizeY, *TIME_COLUMN_WIDTH);
+			StringSize(StatFont, InterColumnOffset, TimeColumnSpaceSizeY, *STATDATA_COLUMN_WIDTH);
 
-			bIsInitialized = true;
+			bNeedRefresh = false;
 		}
 	}
-};
-
-static FStatRenderGlobals StatGlobals;
+	
+	/**
+	 * @return number of characters use to render the stat name
+	 */
+	int32 GetNumCharsForStatName() const
+	{
+		return GStatFonts[(int32)StatFontType].MaxDisplayedChars;
+	}
+	
+	/**
+	 * @return number of characters use to render the stat name
+	 */
+	int32 GetFontHeight() const
+	{
+		return GStatFonts[( int32 )StatFontType].FontHeight;
+	}
 
+	/**
+	 * @return Y offset of the background tile, so this will align with the text
+	 */
+	int32 GetYOffset() const
+	{
+		return GStatFonts[(int32)StatFontType].FontHeightOffset;
+	}
 
+	/** Sets a new font. */
+	void SetNewFont( EStatFontTypes NewFontType )
+	{
+		if( StatFontType != NewFontType )
+		{
+			StatFontType = NewFontType;
+			if( StatFontType == EStatFontTypes::Tiny )
+			{
+				StatFont = GEngine->GetTinyFont();
+			}
+			else if( StatFontType == EStatFontTypes::Small )
+			{
+				StatFont = GEngine->GetSmallFont();
+			}
+			bNeedRefresh = true;
+		}
+	}
+};
 
-static void RightJustify(FCanvas* Canvas, int32 X, int32 Y, TCHAR const* Text, FLinearColor const& Color)
+FStatRenderGlobals& GetStatRenderGlobals()
 {
-	int32 StatColumnSpaceSizeX, StatColumnSpaceSizeY;
-	StringSize(STAT_FONT, StatColumnSpaceSizeX, StatColumnSpaceSizeY, Text);
-
-	Canvas->DrawShadowedString(X + StatGlobals.InterColumnOffset - StatColumnSpaceSizeX,Y,Text,STAT_FONT,Color);
+	// Sanity checks.
+	check( IsInGameThread() );
+	check( GEngine );
+	static FStatRenderGlobals Singleton;
+	return Singleton;
 }
 
-
-// shorten a name for stats display
-static FString ShortenName(TCHAR const* LongName)
+/** Shorten a name for stats display. */
+static FString ShortenName( TCHAR const* LongName )
 {
-	FString Result(LongName);
-	const int32 Limit = 50;
-	if (Result.Len() > Limit)
+	FString Result( LongName );
+	const int32 Limit = GetStatRenderGlobals().GetNumCharsForStatName();
+	if( Result.Len() > Limit )
 	{
-		Result = FString(TEXT("...")) + Result.Right(Limit);
+		Result = FString( TEXT( "..." ) ) + Result.Right( Limit );
 	}
 	return Result;
 }
 
+/** Exec used to execute engine stats command on the game thread. */
+static class FStatCmdEngine : private FSelfRegisteringExec
+{
+public:
+	/** Console commands. */
+	virtual bool Exec( UWorld*, const TCHAR* Cmd, FOutputDevice& Ar ) override
+	{
+		if( FParse::Command( &Cmd, TEXT( "stat" ) ) )
+		{
+			if( FParse::Command( &Cmd, TEXT( "display" ) ) )
+			{
+				TParsedValueWithDefault<EStatFontTypes> Font( Cmd, TEXT( "font=" ), EStatFontTypes::Small );
+				GetStatRenderGlobals().SetNewFont( Font.Get() );
+				return true;
+			}
+		}
+		return false;
+	}
+}
+StatCmdEngineExec;
+
+static void RightJustify(FCanvas* Canvas, int32 X, int32 Y, TCHAR const* Text, FLinearColor const& Color)
+{
+	int32 StatColumnSpaceSizeX, StatColumnSpaceSizeY;
+	StringSize(GetStatRenderGlobals().StatFont, StatColumnSpaceSizeX, StatColumnSpaceSizeY, Text);
+
+	Canvas->DrawShadowedString(X + GetStatRenderGlobals().InterColumnOffset - StatColumnSpaceSizeX,Y,Text,GetStatRenderGlobals().StatFont,Color);
+}
+
+
 /**
  *
  * @param Item the stat to render
@@ -128,7 +260,7 @@ static FString ShortenName(TCHAR const* LongName)
  */
 static int32 RenderCycle( const FComplexStatMessage& Item, class FCanvas* Canvas, int32 X, int32 Y, const int32 Indent, const bool bStackStat )
 {
-	FColor Color = StatGlobals.StatColor;	
+	FColor Color = GetStatRenderGlobals().StatColor;	
 
 	check(Item.NameAndInfo.GetFlag(EStatMetaFlags::IsCycle));
 
@@ -138,15 +270,16 @@ static int32 RenderCycle( const FComplexStatMessage& Item, class FCanvas* Canvas
 
 	if( bIsInitialized )
 	{
-		float InMs = FPlatformTime::ToMilliseconds(Item.GetValue_Duration(EComplexStatField::IncAve));
+		const float InMs = FPlatformTime::ToMilliseconds(Item.GetValue_Duration(EComplexStatField::IncAve));
 		// Color will be determined by the average value of history
 		// If show inclusive and and show exclusive is on, then it will choose color based on inclusive average
+		// @TODO yrx 2014-08-21 This is slow, fix this.
 		FString CounterName = Item.GetShortName().ToString();
 		CounterName.RemoveFromStart(TEXT("STAT_"));
 		GEngine->GetStatValueColoration(CounterName, InMs, Color);
 
 		const float MaxMeter = 33.3f; // the time of a "full bar" in ms
-		const int32 MeterWidth = StatGlobals.AfterNameColumnOffset;
+		const int32 MeterWidth = GetStatRenderGlobals().AfterNameColumnOffset;
 
 		int32 BarWidth = int32((InMs / MaxMeter) * MeterWidth);
 		if (BarWidth > 2)
@@ -156,17 +289,15 @@ static int32 RenderCycle( const FComplexStatMessage& Item, class FCanvas* Canvas
 				BarWidth = MeterWidth;
 			}
 
-			FCanvasBoxItem BoxItem( FVector2D(X + MeterWidth - BarWidth, Y + .4f * FONT_HEIGHT), FVector2D(BarWidth, 0.2f * FONT_HEIGHT) );
+			FCanvasBoxItem BoxItem( FVector2D(X + MeterWidth - BarWidth, Y + .4f * GetStatRenderGlobals().GetFontHeight()), FVector2D(BarWidth, 0.2f * GetStatRenderGlobals().GetFontHeight()) );
 			BoxItem.SetColor( FLinearColor::Red );
 			BoxItem.Draw( Canvas );		
 		}
 	}
 
+	Canvas->DrawShadowedString(X + IndentWidth, Y, *ShortenName(*Item.GetDescription()), GetStatRenderGlobals().StatFont, Color);
 
-
-	Canvas->DrawShadowedString(X + IndentWidth, Y, *ShortenName(*Item.GetDescription()), STAT_FONT, Color);
-
-	int32 CurrX = X + StatGlobals.AfterNameColumnOffset;
+	int32 CurrX = X + GetStatRenderGlobals().AfterNameColumnOffset;
 	// Now append the call count
 	if( bStackStat )
 	{
@@ -174,7 +305,7 @@ static int32 RenderCycle( const FComplexStatMessage& Item, class FCanvas* Canvas
 		{
 			RightJustify(Canvas,CurrX,Y,*FString::Printf(TEXT("%u"), Item.GetValue_CallCount(EComplexStatField::IncAve)),Color);
 		}
-		CurrX += StatGlobals.InterColumnOffset;
+		CurrX += GetStatRenderGlobals().InterColumnOffset;
 	}
 
 	// Add the two inclusive columns if asked
@@ -182,14 +313,14 @@ static int32 RenderCycle( const FComplexStatMessage& Item, class FCanvas* Canvas
 	{
 		RightJustify(Canvas,CurrX,Y,*FString::Printf(TEXT("%1.2f ms"),FPlatformTime::ToMilliseconds(Item.GetValue_Duration(EComplexStatField::IncAve))),Color); 
 	}
-	CurrX += StatGlobals.InterColumnOffset;
+	CurrX += GetStatRenderGlobals().InterColumnOffset;
 
 	if( bIsInitialized )
 	{
 		RightJustify(Canvas,CurrX,Y,*FString::Printf(TEXT("%1.2f ms"),FPlatformTime::ToMilliseconds(Item.GetValue_Duration(EComplexStatField::IncMax))),Color);
 		
 	}
-	CurrX += StatGlobals.InterColumnOffset;
+	CurrX += GetStatRenderGlobals().InterColumnOffset;
 
 	if( bStackStat )
 	{
@@ -198,15 +329,15 @@ static int32 RenderCycle( const FComplexStatMessage& Item, class FCanvas* Canvas
 		{
 			RightJustify(Canvas,CurrX,Y,*FString::Printf(TEXT("%1.2f ms"),FPlatformTime::ToMilliseconds(Item.GetValue_Duration(EComplexStatField::ExcAve))),Color);
 		}
-		CurrX += StatGlobals.InterColumnOffset;
+		CurrX += GetStatRenderGlobals().InterColumnOffset;
 
 		if( bIsInitialized )
 		{
 			RightJustify(Canvas,CurrX,Y,*FString::Printf(TEXT("%1.2f ms"),FPlatformTime::ToMilliseconds(Item.GetValue_Duration(EComplexStatField::ExcMax))),Color);
 		}
-		CurrX += StatGlobals.InterColumnOffset;
+		CurrX += GetStatRenderGlobals().InterColumnOffset;
 	}
-	return FONT_HEIGHT;
+	return GetStatRenderGlobals().GetFontHeight();
 }
 
 /**
@@ -226,23 +357,23 @@ static int32 RenderGroupedHeadings(class FCanvas* Canvas,int X,int32 Y,const boo
 	static const TCHAR* CaptionFlat = TEXT("Cycle counters (flat)");
 	static const TCHAR* CaptionHier = TEXT("Cycle counters (hierarchy)");
 
-	Canvas->DrawShadowedString(X,Y,bIsHierarchy?CaptionHier:CaptionFlat,STAT_FONT,StatGlobals.HeadingColor);
+	Canvas->DrawShadowedString(X,Y,bIsHierarchy?CaptionHier:CaptionFlat,GetStatRenderGlobals().StatFont,GetStatRenderGlobals().HeadingColor);
 
-	int32 CurrX = X + StatGlobals.AfterNameColumnOffset;
-	RightJustify(Canvas,CurrX,Y,TEXT("CallCount"),StatGlobals.HeadingColor);
-	CurrX += StatGlobals.InterColumnOffset;
+	int32 CurrX = X + GetStatRenderGlobals().AfterNameColumnOffset;
+	RightJustify(Canvas,CurrX,Y,TEXT("CallCount"),GetStatRenderGlobals().HeadingColor);
+	CurrX += GetStatRenderGlobals().InterColumnOffset;
 
-	RightJustify(Canvas,CurrX,Y,TEXT("InclusiveAvg"),StatGlobals.HeadingColor);
-	CurrX += StatGlobals.InterColumnOffset;
-	RightJustify(Canvas,CurrX,Y,TEXT("InclusiveMax"),StatGlobals.HeadingColor);
-	CurrX += StatGlobals.InterColumnOffset;
+	RightJustify(Canvas,CurrX,Y,TEXT("InclusiveAvg"),GetStatRenderGlobals().HeadingColor);
+	CurrX += GetStatRenderGlobals().InterColumnOffset;
+	RightJustify(Canvas,CurrX,Y,TEXT("InclusiveMax"),GetStatRenderGlobals().HeadingColor);
+	CurrX += GetStatRenderGlobals().InterColumnOffset;
 
-	RightJustify(Canvas,CurrX,Y,TEXT("ExclusiveAvg"),StatGlobals.HeadingColor);
-	CurrX += StatGlobals.InterColumnOffset;
-	RightJustify(Canvas,CurrX,Y,TEXT("ExclusiveMax"),StatGlobals.HeadingColor);
-	CurrX += StatGlobals.InterColumnOffset;
+	RightJustify(Canvas,CurrX,Y,TEXT("ExclusiveAvg"),GetStatRenderGlobals().HeadingColor);
+	CurrX += GetStatRenderGlobals().InterColumnOffset;
+	RightJustify(Canvas,CurrX,Y,TEXT("ExclusiveMax"),GetStatRenderGlobals().HeadingColor);
+	CurrX += GetStatRenderGlobals().InterColumnOffset;
 
-	return FONT_HEIGHT + (FONT_HEIGHT / 3);
+	return GetStatRenderGlobals().GetFontHeight() + (GetStatRenderGlobals().GetFontHeight() / 3);
 }
 
 /**
@@ -259,18 +390,18 @@ static int32 RenderCounterHeadings(class FCanvas* Canvas,int32 X,int32 Y)
 	// The heading looks like:
 	// Stat [32chars]    Value [8chars]    Average [8chars]
 
-	Canvas->DrawShadowedString(X,Y,TEXT("Counters"),STAT_FONT,StatGlobals.HeadingColor);
+	Canvas->DrawShadowedString(X,Y,TEXT("Counters"),GetStatRenderGlobals().StatFont,GetStatRenderGlobals().HeadingColor);
 
 	// Determine where the first column goes
-	int32 CurrX = X + StatGlobals.AfterNameColumnOffset;
+	int32 CurrX = X + GetStatRenderGlobals().AfterNameColumnOffset;
 
 	// Draw the average column label.
-	RightJustify(Canvas,CurrX,Y,TEXT("Average"),StatGlobals.HeadingColor);
-	CurrX += StatGlobals.InterColumnOffset;
+	RightJustify(Canvas,CurrX,Y,TEXT("Average"),GetStatRenderGlobals().HeadingColor);
+	CurrX += GetStatRenderGlobals().InterColumnOffset;
 
 	// Draw the max column label.
-	RightJustify(Canvas,CurrX,Y,TEXT("Max"),StatGlobals.HeadingColor);
-	return FONT_HEIGHT + (FONT_HEIGHT / 3);
+	RightJustify(Canvas,CurrX,Y,TEXT("Max"),GetStatRenderGlobals().HeadingColor);
+	return GetStatRenderGlobals().GetFontHeight() + (GetStatRenderGlobals().GetFontHeight() / 3);
 }
 
 /**
@@ -287,34 +418,22 @@ static int32 RenderMemoryHeadings(class FCanvas* Canvas,int32 X,int32 Y)
 	// The heading looks like:
 	// Stat [32chars]    MemUsed [8chars]    PhysMem [8chars]
 
-	// get the size of the spaces, since we can't render the width calculation strings
-	int32 StatColumnSpaceSizeX, StatColumnSpaceSizeY;
-	int32 TimeColumnSpaceSizeX, TimeColumnSpaceSizeY;
-	StringSize(STAT_FONT, StatColumnSpaceSizeX, StatColumnSpaceSizeY, *STAT_COLUMN_WIDTH);
-	StringSize(STAT_FONT, TimeColumnSpaceSizeX, TimeColumnSpaceSizeY, *TIME_COLUMN_WIDTH);
-
-	Canvas->DrawShadowedString(X,Y,TEXT("Memory Counters"),STAT_FONT,StatGlobals.HeadingColor);
+	Canvas->DrawShadowedString(X,Y,TEXT("Memory Counters"),GetStatRenderGlobals().StatFont,GetStatRenderGlobals().HeadingColor);
 
 	// Determine where the first column goes
-	int32 CurrX = X + StatGlobals.AfterNameColumnOffset;
-	RightJustify(Canvas,CurrX,Y,TEXT("MemUsedAvg"),StatGlobals.HeadingColor);
-
-	CurrX += StatGlobals.InterColumnOffset;
-	RightJustify(Canvas,CurrX,Y,TEXT("MemUsedAvg%"),StatGlobals.HeadingColor);
+	int32 CurrX = X + GetStatRenderGlobals().AfterNameColumnOffset;
+	RightJustify(Canvas,CurrX,Y,TEXT("MemUsedMax"),GetStatRenderGlobals().HeadingColor);
 
-	CurrX += StatGlobals.InterColumnOffset;
-	RightJustify(Canvas,CurrX,Y,TEXT("MemUsedMax"),StatGlobals.HeadingColor);
+	CurrX += GetStatRenderGlobals().InterColumnOffset;
+	RightJustify(Canvas,CurrX,Y,TEXT("MemUsedMax%"),GetStatRenderGlobals().HeadingColor);
 
-	CurrX += StatGlobals.InterColumnOffset;
-	RightJustify(Canvas,CurrX,Y,TEXT("MemUsedMax%"),StatGlobals.HeadingColor);
+	CurrX += GetStatRenderGlobals().InterColumnOffset;
+	RightJustify(Canvas,CurrX,Y,TEXT("MemPool"),GetStatRenderGlobals().HeadingColor);
 
-	CurrX += StatGlobals.InterColumnOffset;
-	RightJustify(Canvas,CurrX,Y,TEXT("MemPool"),StatGlobals.HeadingColor);
+	CurrX += GetStatRenderGlobals().InterColumnOffset;
+	RightJustify(Canvas,CurrX,Y,TEXT("Pool Capacity"),GetStatRenderGlobals().HeadingColor);
 
-	CurrX += StatGlobals.InterColumnOffset;
-	RightJustify(Canvas,CurrX,Y,TEXT("Pool Capacity"),StatGlobals.HeadingColor);
-
-	return FONT_HEIGHT + (FONT_HEIGHT / 3);
+	return GetStatRenderGlobals().GetFontHeight() + (GetStatRenderGlobals().GetFontHeight() / 3);
 }
 
 static FString GetMemoryString(double Value)
@@ -337,52 +456,35 @@ static FString GetMemoryString(double Value)
 static int32 RenderMemoryCounter(const FGameThreadHudData& ViewData, const FComplexStatMessage& All,class FCanvas* Canvas,int32 X,int32 Y)
 {
 	FPlatformMemory::EMemoryCounterRegion Region = FPlatformMemory::EMemoryCounterRegion(All.NameAndInfo.GetField<EMemoryRegion>());
-	const bool bDisplayAll = All.NameAndInfo.GetFlag(EStatMetaFlags::ShouldClearEveryFrame);
+	// At this moment we only have memory stats that are marked as non frame stats, so can't be cleared every frame.
+	//const bool bDisplayAll = All.NameAndInfo.GetFlag(EStatMetaFlags::ShouldClearEveryFrame);
 	
-	const float MemUsed = All.GetValue_double(EComplexStatField::IncAve);
 	const float MaxMemUsed = All.GetValue_double(EComplexStatField::IncMax);
 
 	// Draw the label
-	Canvas->DrawShadowedString(X,Y,*ShortenName(*All.GetDescription()),STAT_FONT,StatGlobals.StatColor);
-	int32 CurrX = X + StatGlobals.AfterNameColumnOffset;
+	Canvas->DrawShadowedString(X,Y,*ShortenName(*All.GetDescription()),GetStatRenderGlobals().StatFont,GetStatRenderGlobals().StatColor);
+	int32 CurrX = X + GetStatRenderGlobals().AfterNameColumnOffset;
 
-	if( bDisplayAll )
-	{
-		// Now append the value of the stat
-		RightJustify(Canvas,CurrX,Y,*GetMemoryString(MemUsed),StatGlobals.StatColor);
-		CurrX += StatGlobals.InterColumnOffset;
-		if (ViewData.PoolCapacity.Contains(Region))
-		{
-			RightJustify(Canvas,CurrX,Y,*FString::Printf(TEXT("%.0f%%"), float(100.0 * MemUsed / double(ViewData.PoolCapacity[Region]))),StatGlobals.StatColor);
-		}
-		CurrX += StatGlobals.InterColumnOffset;
-	}
-	else
-	{
-		CurrX += StatGlobals.InterColumnOffset;
-	CurrX += StatGlobals.InterColumnOffset;
-	}
-	
 	// Now append the max value of the stat
-	RightJustify(Canvas,CurrX,Y,*GetMemoryString(MaxMemUsed),StatGlobals.StatColor);
-	CurrX += StatGlobals.InterColumnOffset;
+	RightJustify(Canvas,CurrX,Y,*GetMemoryString(MaxMemUsed),GetStatRenderGlobals().StatColor);
+	CurrX += GetStatRenderGlobals().InterColumnOffset;
 	if (ViewData.PoolCapacity.Contains(Region))
 	{
-		RightJustify(Canvas,CurrX,Y,*FString::Printf(TEXT("%.0f%%"), float(100.0 * MaxMemUsed / double(ViewData.PoolCapacity[Region]))),StatGlobals.StatColor);
+		RightJustify(Canvas,CurrX,Y,*FString::Printf(TEXT("%.0f%%"), float(100.0 * MaxMemUsed / double(ViewData.PoolCapacity[Region]))),GetStatRenderGlobals().StatColor);
 	}
-	CurrX += StatGlobals.InterColumnOffset;
+	CurrX += GetStatRenderGlobals().InterColumnOffset;
 	if (ViewData.PoolAbbreviation.Contains(Region))
 	{
-		RightJustify(Canvas,CurrX,Y,*ViewData.PoolAbbreviation[Region],StatGlobals.StatColor);
+		RightJustify(Canvas,CurrX,Y,*ViewData.PoolAbbreviation[Region],GetStatRenderGlobals().StatColor);
 	}
-	CurrX += StatGlobals.InterColumnOffset;
+	CurrX += GetStatRenderGlobals().InterColumnOffset;
 	if (ViewData.PoolCapacity.Contains(Region))
 	{
-		RightJustify(Canvas,CurrX,Y,*GetMemoryString(double(ViewData.PoolCapacity[Region])),StatGlobals.StatColor);
+		RightJustify(Canvas,CurrX,Y,*GetMemoryString(double(ViewData.PoolCapacity[Region])),GetStatRenderGlobals().StatColor);
 	}
-	CurrX += StatGlobals.InterColumnOffset;
+	CurrX += GetStatRenderGlobals().InterColumnOffset;
 
-	return FONT_HEIGHT;
+	return GetStatRenderGlobals().GetFontHeight();
 }
 
 static int32 RenderCounter(const FGameThreadHudData& ViewData, const FComplexStatMessage& All,class FCanvas* Canvas,int32 X,int32 Y)
@@ -397,33 +499,33 @@ static int32 RenderCounter(const FGameThreadHudData& ViewData, const FComplexSta
 	const bool bDisplayAll = All.NameAndInfo.GetFlag(EStatMetaFlags::ShouldClearEveryFrame);
 
 	// Draw the label
-	Canvas->DrawShadowedString(X,Y,*ShortenName(*All.GetDescription()),STAT_FONT,StatGlobals.StatColor);
-	int32 CurrX = X + StatGlobals.AfterNameColumnOffset;
+	Canvas->DrawShadowedString(X,Y,*ShortenName(*All.GetDescription()),GetStatRenderGlobals().StatFont,GetStatRenderGlobals().StatColor);
+	int32 CurrX = X + GetStatRenderGlobals().AfterNameColumnOffset;
 
 	if( bDisplayAll )
 	{
 		// Append the average.
 		if (All.NameAndInfo.GetField<EStatDataType>() == EStatDataType::ST_double)
 		{
-			RightJustify(Canvas,CurrX,Y,*FString::Printf(TEXT("%.2f"), All.GetValue_double(EComplexStatField::IncAve)),StatGlobals.StatColor);
+			RightJustify(Canvas,CurrX,Y,*FString::Printf(TEXT("%.2f"), All.GetValue_double(EComplexStatField::IncAve)),GetStatRenderGlobals().StatColor);
 		}
 		else if (All.NameAndInfo.GetField<EStatDataType>() == EStatDataType::ST_int64)
 		{
-			RightJustify(Canvas,CurrX,Y,*FString::Printf(TEXT("%lld"), All.GetValue_int64(EComplexStatField::IncAve)),StatGlobals.StatColor);
+			RightJustify(Canvas,CurrX,Y,*FString::Printf(TEXT("%lld"), All.GetValue_int64(EComplexStatField::IncAve)),GetStatRenderGlobals().StatColor);
 		}
 	}
-	CurrX += StatGlobals.InterColumnOffset;
+	CurrX += GetStatRenderGlobals().InterColumnOffset;
 
 	// Append the maximum.
 	if (All.NameAndInfo.GetField<EStatDataType>() == EStatDataType::ST_double)
 	{
-		RightJustify(Canvas,CurrX,Y,*FString::Printf(TEXT("%.2f"), All.GetValue_double(EComplexStatField::IncMax)),StatGlobals.StatColor);
+		RightJustify(Canvas,CurrX,Y,*FString::Printf(TEXT("%.2f"), All.GetValue_double(EComplexStatField::IncMax)),GetStatRenderGlobals().StatColor);
 	}
 	else if (All.NameAndInfo.GetField<EStatDataType>() == EStatDataType::ST_int64)
 	{
-		RightJustify(Canvas,CurrX,Y,*FString::Printf(TEXT("%lld"), All.GetValue_int64(EComplexStatField::IncMax)),StatGlobals.StatColor);
+		RightJustify(Canvas,CurrX,Y,*FString::Printf(TEXT("%lld"), All.GetValue_int64(EComplexStatField::IncMax)),GetStatRenderGlobals().StatColor);
 	}
-	return FONT_HEIGHT;
+	return GetStatRenderGlobals().GetFontHeight();
 }
 
 void RenderHierCycles( FCanvas* Canvas, int32 X, int32& Y, const FHudGroup& HudGroup )
@@ -438,9 +540,9 @@ void RenderHierCycles( FCanvas* Canvas, int32 X, int32& Y, const FHudGroup& HudG
 
 		if(BackgroundTexture != nullptr)
 		{
-			Canvas->DrawTile( X, Y + 2, StatGlobals.AfterNameColumnOffset + StatGlobals.InterColumnOffset * 6, FONT_HEIGHT,
+			Canvas->DrawTile( X, Y + GetStatRenderGlobals().GetYOffset(), GetStatRenderGlobals().AfterNameColumnOffset + GetStatRenderGlobals().InterColumnOffset * (int32)EStatRenderConsts::NUM_COLUMNS, GetStatRenderGlobals().GetFontHeight(),
 				0, 0, 1, 1,
-				StatGlobals.BackgroundColor, BackgroundTexture->Resource, true );
+				GetStatRenderGlobals().BackgroundColor, BackgroundTexture->Resource, true );
 		}
 
 		Y += RenderCycle( ComplexStat, Canvas, X, Y, Indent, true );
@@ -459,9 +561,9 @@ void RenderArrayOfStats( FCanvas* Canvas, int32 X, int32& Y, const TArray<FCompl
 
 		if(BackgroundTexture != nullptr)
 		{
-			Canvas->DrawTile( X, Y + 2, StatGlobals.AfterNameColumnOffset + StatGlobals.InterColumnOffset * 6, FONT_HEIGHT,
+			Canvas->DrawTile( X, Y + GetStatRenderGlobals().GetYOffset(), GetStatRenderGlobals().AfterNameColumnOffset + GetStatRenderGlobals().InterColumnOffset * (int32)EStatRenderConsts::NUM_COLUMNS, GetStatRenderGlobals().GetFontHeight(),
 				0, 0, 1, 1,
-				StatGlobals.BackgroundColor, BackgroundTexture->Resource, true );
+				GetStatRenderGlobals().BackgroundColor, BackgroundTexture->Resource, true );
 		}
 
 		Y += FunctionToCall( ViewData, ComplexStat, Canvas, X, Y );
@@ -490,19 +592,19 @@ static void RenderGroupedWithHierarchy(const FGameThreadHudData& ViewData, FView
 	for( int32 GroupIndex = 0; GroupIndex < ViewData.HudGroups.Num(); ++GroupIndex )
 	{
 		// If the stat isn't enabled for this particular viewport, skip
-		FString StatName = ViewData.GroupNames[GroupIndex].ToString();
-		StatName.RemoveFromStart(TEXT("STATGROUP_"));
-		if (!Viewport->GetClient() || !Viewport->GetClient()->IsStatEnabled(*StatName))
+		FString StatGroupName = ViewData.GroupNames[GroupIndex].ToString();
+		StatGroupName.RemoveFromStart(TEXT("STATGROUP_"));
+		if (!Viewport->GetClient() || !Viewport->GetClient()->IsStatEnabled(*StatGroupName))
 		{
 			continue;
 		}
 
 		// Render header.
 		const FName& GroupName = ViewData.GroupNames[GroupIndex];
-		const FString GroupDesc = ViewData.GroupDescriptions[GroupIndex];
+		const FString& GroupDesc = ViewData.GroupDescriptions[GroupIndex];
 		const FString GroupLongName = FString::Printf( TEXT("%s [%s]"), *GroupDesc, *GroupName.GetPlainNameString() );
-		Canvas->DrawShadowedString( X, Y, *GroupLongName, STAT_FONT, StatGlobals.GroupColor );
-		Y += FONT_HEIGHT;
+		Canvas->DrawShadowedString( X, Y, *GroupLongName, GetStatRenderGlobals().StatFont, GetStatRenderGlobals().GroupColor );
+		Y += GetStatRenderGlobals().GetFontHeight();
 
 		const FHudGroup& HudGroup = ViewData.HudGroups[GroupIndex];
 		const bool bHasHierarchy = !!HudGroup.HierAggregate.Num();
@@ -510,22 +612,22 @@ static void RenderGroupedWithHierarchy(const FGameThreadHudData& ViewData, FView
 
 		if (bHasHierarchy || bHasFlat)
 		{
-		// Render grouped headings.
-		Y += RenderGroupedHeadings( Canvas, X, Y, bHasHierarchy );
+			// Render grouped headings.
+			Y += RenderGroupedHeadings( Canvas, X, Y, bHasHierarchy );
 		}
 
 		// Render hierarchy.
 		if( bHasHierarchy )
 		{
 			RenderHierCycles( Canvas, X, Y, HudGroup );
-			Y += FONT_HEIGHT;
+			Y += GetStatRenderGlobals().GetFontHeight();
 		}
 
 		// Render flat.
 		if( bHasFlat )
 		{
 			RenderArrayOfStats(Canvas,X,Y,HudGroup.FlatAggregate, ViewData, RenderFlatCycle);
-			Y += FONT_HEIGHT;
+			Y += GetStatRenderGlobals().GetFontHeight();
 		}
 
 		// Render memory counters.
@@ -533,7 +635,7 @@ static void RenderGroupedWithHierarchy(const FGameThreadHudData& ViewData, FView
 		{
 			Y += RenderMemoryHeadings(Canvas,X,Y);
 			RenderArrayOfStats(Canvas,X,Y,HudGroup.MemoryAggregate, ViewData, RenderMemoryCounter);
-			Y += FONT_HEIGHT;
+			Y += GetStatRenderGlobals().GetFontHeight();
 		}
 
 		// Render remaining counters.
@@ -541,7 +643,7 @@ static void RenderGroupedWithHierarchy(const FGameThreadHudData& ViewData, FView
 		{
 			Y += RenderCounterHeadings(Canvas,X,Y);
 			RenderArrayOfStats(Canvas,X,Y,HudGroup.CountersAggregate, ViewData, RenderCounter);
-			Y += FONT_HEIGHT;
+			Y += GetStatRenderGlobals().GetFontHeight();
 		}
 	}
 }
@@ -561,25 +663,8 @@ void RenderStats(FViewport* Viewport, class FCanvas* Canvas, int32 X, int32 Y)
 	{
 		return;
 	}
-
-	if (FMath::IsNearlyZero(StatGlobals.StatFontScale))
-	{
-		StatGlobals.StatFontScale = 1.0f;
-		int32 StatFontHeight = STAT_FONT->GetMaxCharHeight();
-		float CheckFontHeight = float(IDEAL_FONT_HEIGHT);
-		if (StatFontHeight > CheckFontHeight)
-		{
-			StatGlobals.StatFontScale = CheckFontHeight / StatFontHeight;
-		}
-	}
- 	
-	float SavedFontScale = STAT_FONT->GetFontScalingFactor();
- 	STAT_FONT->SetFontScalingFactor(StatGlobals.StatFontScale);
-
-	StatGlobals.Initialize();
+	GetStatRenderGlobals().Initialize( Viewport->GetSizeXY() );
 	RenderGroupedWithHierarchy(*ViewData, Viewport, Canvas, X, Y);
-
-	STAT_FONT->SetFontScalingFactor(SavedFontScale);
 }
 
 #endif
