Compilation fix

diff --git a/Engine/Source/Runtime/Engine/Private/StatsRender2.cpp b/Engine/Source/Runtime/Engine/Private/StatsRender2.cpp
--- a/Engine/Source/Runtime/Engine/Private/StatsRender2.cpp
+++ b/Engine/Source/Runtime/Engine/Private/StatsRender2.cpp
@@ -48,7 +48,7 @@ struct FStatFont
 	const int32 FontHeightOffset;
 };
 
-static FStatFont GStatFonts[EStatFontTypes::NumFonts] =
+static FStatFont GStatFonts[(int32)EStatFontTypes::NumFonts] =
 {
 	/** Tiny. */
 	FStatFont( 24, 10, 1 ),
