Refactor body welding. Some stuff is showing through even without define because of how metasystem works

diff --git a/Engine/Source/Editor/DetailCustomizations/Private/PrimitiveComponentDetails.cpp b/Engine/Source/Editor/DetailCustomizations/Private/PrimitiveComponentDetails.cpp
--- a/Engine/Source/Editor/DetailCustomizations/Private/PrimitiveComponentDetails.cpp
+++ b/Engine/Source/Editor/DetailCustomizations/Private/PrimitiveComponentDetails.cpp
@@ -153,28 +153,35 @@ void FPrimitiveComponentDetails::CustomizeDetails( IDetailLayoutBuilder& DetailB
 			{
 				TSharedPtr<IPropertyHandle> ChildProperty = BodyInstanceHandler->GetChildHandle(ChildIndex);
 				FString Category = FObjectEditorUtils::GetCategory(ChildProperty->GetProperty());
+				FString PropName = ChildProperty->GetProperty()->GetName();
 				if (Category == TEXT("Physics"))
 				{
 					// Only permit modifying bSimulatePhysics when the body has some geometry.
-					if (ChildProperty->GetProperty()->GetName() == TEXT("bSimulatePhysics"))
+					if (PropName == TEXT("bSimulatePhysics"))
 					{
 						PhysicsCategory.AddProperty(ChildProperty).EditCondition(TAttribute<bool>(this, &FPrimitiveComponentDetails::IsSimulatePhysicsEditable), NULL);
 					}
-					else if (ChildProperty->GetProperty()->GetName() == TEXT("bUseAsyncScene"))
+					else if (PropName == TEXT("bUseAsyncScene"))
 					{
 						//we only enable bUseAsyncScene if the project uses an AsyncScene
 						PhysicsCategory.AddProperty(ChildProperty).EditCondition(TAttribute<bool>(this, &FPrimitiveComponentDetails::IsUseAsyncEditable), NULL);
 					}
-					else if (ChildProperty->GetProperty()->GetName() == TEXT("LockedAxisMode"))
+					else if (PropName == TEXT("LockedAxisMode"))
 					{
 						LockedAxisProperty = ChildProperty;
 						PhysicsCategory.AddProperty(ChildProperty);
 					}
-					else if (ChildProperty->GetProperty()->GetName() == TEXT("CustomLockedAxis"))
+					else if (PropName == TEXT("CustomLockedAxis"))
 					{
 						//we only enable bUseAsyncScene if the project uses an AsyncScene
 						PhysicsCategory.AddProperty(ChildProperty).Visibility(TAttribute<EVisibility>(this, &FPrimitiveComponentDetails::IsCustomLockedAxisSelected));
 					}
+					else if (PropName == TEXT("AutoWeld"))
+					{
+#if WITH_BODY_WELDING
+						PhysicsCategory.AddProperty(ChildProperty);	//can't ifdef the property out so for now we just customize it to be hidden in UI
+#endif
+					}
 					else
 					{
 						PhysicsCategory.AddProperty(ChildProperty);
diff --git a/Engine/Source/Runtime/Engine/Classes/Components/PrimitiveComponent.h b/Engine/Source/Runtime/Engine/Classes/Components/PrimitiveComponent.h
--- a/Engine/Source/Runtime/Engine/Classes/Components/PrimitiveComponent.h
+++ b/Engine/Source/Runtime/Engine/Classes/Components/PrimitiveComponent.h
@@ -1025,8 +1025,32 @@ public:
 	virtual void ReceiveComponentDamage(float DamageAmount, FDamageEvent const& DamageEvent, AController* EventInstigator, AActor* DamageCauser);
 
 #if WITH_BODY_WELDING
-	/** Used to physically weld objects at runtime. See decleration in USceneComponent for full description */
-	virtual bool WeldPhysicsBody(USceneComponent * ChildComponent, bool bWeld = true, FName ParentBoneName = NAME_None, FName ChildBoneName = NAME_None) override;
+
+	/**
+	*   Welds this component to another scene component, optionally at a named socket. Component is automatically attached if not already
+	*	Welding allows the child physics object to become physically connected to its parent. This is useful for creating compound rigid bodies with correct mass distribution.
+	*   @param InParent the component to be physically attached to
+	*   @param InSocketName optional socket to attach component to
+	*/
+	UFUNCTION(BlueprintCallable, Category = "Physics")
+	virtual void WeldTo(class USceneComponent* InParent, FName InSocketName = NAME_None);
+
+	/**
+	*   UnWelds this component from its parent component. Attachment is maintained (DetachFromParent automatically unwelds)
+	*/
+	UFUNCTION(BlueprintCallable, Category = "Physics")
+	virtual void UnWeldFromParent();
+
+	/**
+	*	Does the actual work for welding.
+	*	@return true if did a true weld of shapes, meaning body initialization is not needed
+	*/
+	bool WeldToInternal(USceneComponent * InParent, FName ParentSocketName = NAME_None);
+
+	/**
+	*	Adds the bodies that are currently welded to the OutWeldedBodies array 
+	*/
+	virtual void GetWeldedBodies(TArray<FBodyInstance*> & OutWeldedBodies, TArray<FName> & OutLabels);
 #endif
 
 protected:
diff --git a/Engine/Source/Runtime/Engine/Classes/Components/SceneComponent.h b/Engine/Source/Runtime/Engine/Classes/Components/SceneComponent.h
--- a/Engine/Source/Runtime/Engine/Classes/Components/SceneComponent.h
+++ b/Engine/Source/Runtime/Engine/Classes/Components/SceneComponent.h
@@ -336,12 +336,43 @@ public:
 	void SnapTo(class USceneComponent* InParent, FName InSocketName = NAME_None);
 
 	/** 
-	 *	Detach this component from whatever it is attached to 
+	 *	Detach this component from whatever it is attached to. Automatically unwelds components that are welded together (See WeldTo)
 	 *   @param bMaintainWorldTransform	If true, update the relative location/rotation of this component to keep its world position the same *	
 	 */
 	UFUNCTION(BlueprintCallable, Category="Utilities|Transformation")
 	virtual void DetachFromParent(bool bMaintainWorldPosition = false);
 
+#if WITH_BODY_WELDING
+	/**
+	*   Welds this component to another scene component, optionally at a named socket. Component is automatically attached if not already
+	*	Welding allows the child physics object to become physically connected to its parent. This is useful for creating compound rigid bodies with correct mass distribution.
+	*   @param InParent the component to be physically attached to
+	*   @param InSocketName optional socket to attach component to
+	*/
+	UFUNCTION(BlueprintCallable, Category = "Physics")
+	virtual void WeldTo(class USceneComponent* InParent, FName InSocketName = NAME_None){}
+
+	/**
+	*   UnWelds this component from its parent component. Attachment is maintained (DetachFromParent automatically unwelds)
+	*/
+	UFUNCTION(BlueprintCallable, Category = "Physics")
+	virtual void UnWeldFromParent(){}
+
+
+	/**
+	* Attempts to weld/unweld the child component to this component in a physically correct way. This re-calculates mass distribution and geometry to achieve correct physics.
+	* This function assumes ChildComponent has already been attached to this component and is now doing the physics fix-up.
+	* The function uses the attachment hierarchy to easily attach/detach multiple components together
+	* @param ChildComponent the component to be physically attached - assumed to already be attached in the scene hierarchy
+	* @param bWeld determines whether to weld or unweld. This flag is used for creating or breaking welded pieces at runtime
+	* @param ParentBoneName potentially used by skeletal mesh to weld child to specific bone on this component
+	* @param ChildBoneName potentially used by skeletal mesh to weld this component to specific bone in child
+	* @return whether or not the weld/unweld was succesfull
+	*/
+	//virtual bool WeldPhysicsBody(USceneComponent * ChildComponent, bool bWeld = true, FName ParentBoneName = NAME_None, FName ChildBoneName = NAME_None) { return false; }
+
+#endif
+
 	/** 
 	 * Gets the names of all the sockets on the component.
 	 * @return Get the names of all the sockets on the component.
@@ -707,21 +738,6 @@ public:
 	/** Get the extent used when placing this component in the editor, used for 'pulling back' hit. */
 	virtual FBoxSphereBounds GetPlacementExtent() const;
 
-#if WITH_BODY_WELDING
-
-	/**
-	* Attempts to weld/unweld the child component to this component in a physically correct way. This re-calculates mass distribution and geometry to achieve correct physics.
-	* This function assumes ChildComponent has already been attached to this component and is now doing the physics fix-up.
-	* The function uses the attachment hierarchy to easily attach/detach multiple components together
-	* @param ChildComponent the component to be physically attached - assumed to already be attached in the scene hierarchy
-	* @param bWeld determines whether to weld or unweld. This flag is used for creating or breaking welded pieces at runtime
-	* @param ParentBoneName potentially used by skeletal mesh to weld child to specific bone on this component
-	* @param ChildBoneName potentially used by skeletal mesh to weld this component to specific bone in child
-	* @return whether or not the weld/unweld was succesfull
-	*/
-	virtual bool WeldPhysicsBody(USceneComponent * ChildComponent, bool bWeld = true, FName ParentBoneName = NAME_None, FName ChildBoneName = NAME_None) { return false;  }
-#endif
-
 protected:
 	/**
 	 * Called after a child scene component is attached to this component.
diff --git a/Engine/Source/Runtime/Engine/Classes/Components/SkeletalMeshComponent.h b/Engine/Source/Runtime/Engine/Classes/Components/SkeletalMeshComponent.h
--- a/Engine/Source/Runtime/Engine/Classes/Components/SkeletalMeshComponent.h
+++ b/Engine/Source/Runtime/Engine/Classes/Components/SkeletalMeshComponent.h
@@ -954,6 +954,10 @@ public:
 	/** Enable or Disable AngularVelocityDrive based on a list of bone names */
 	void SetNamedMotorsAngularVelocityDrive(bool bEnableSwingDrive, bool bEnableTwistDrive, const TArray<FName>& BoneNames, bool bSetOtherBodiesToComplement = false);
 
+#if WITH_BODY_WELDING
+	void GetWeldedBodies(TArray<FBodyInstance*> & OutWeldedBodies, TArray<FName> & OutChildrenLabels) override;
+	bool virtual WeldToInternal(USceneComponent * InParent, FName ParentSocketName = NAME_None){ return false; }	//not currently supported for skeletal mesh
+#endif
 
 	/**
 	 * Return Transform Matrix for SkeletalMeshComponent considering root motion setups
diff --git a/Engine/Source/Runtime/Engine/Classes/PhysicsEngine/BodyInstance.h b/Engine/Source/Runtime/Engine/Classes/PhysicsEngine/BodyInstance.h
--- a/Engine/Source/Runtime/Engine/Classes/PhysicsEngine/BodyInstance.h
+++ b/Engine/Source/Runtime/Engine/Classes/PhysicsEngine/BodyInstance.h
@@ -168,6 +168,15 @@ public:
 	UPROPERTY(EditAnywhere, BlueprintReadWrite, Category=Physics)
 	uint32 bSimulatePhysics:1;
 
+//#if WITH_BODY_WELDING
+	/** If true and is attached to a parent, the two bodies will be joined into a single rigid body. Physical settings like collision profile and body settings are determined by the root */
+	UPROPERTY(EditAnywhere, BlueprintReadWrite, Category = Physics)
+	uint32 bAutoWeld : 1;
+
+	/** determines if the body is currently welded */
+	uint32 bWelded : 1;
+//#endif
+
 	/** If object should start awake, or if it should initially be sleeping */
 	UPROPERTY(EditAnywhere, BlueprintReadOnly, Category=Physics, meta=(editcondition = "bSimulatePhysics"))
 	uint32 bStartAwake:1;
@@ -315,7 +324,7 @@ public:
 #if WITH_BODY_WELDING
 
 #if WITH_PHYSX
-	TMap<physx::PxShape*, class UPrimitiveComponent*> ShapeToComponentMap;
+	TMap<physx::PxShape*, FBodyInstance*> ShapeToBodyMap;
 #endif
 
 	/** 
@@ -325,7 +334,7 @@ public:
 
 	/** 
 	 * Takes a welded body and unwelds it. This function does not create the new body, it only removes the old one */
-	void UnWeld(class UPrimitiveComponent * Body);
+	void UnWeld(FBodyInstance * Body);
 
 	/**
 	 * After adding/removing shapes call this function to update mass distribution etc... */
diff --git a/Engine/Source/Runtime/Engine/Private/Actor.cpp b/Engine/Source/Runtime/Engine/Private/Actor.cpp
--- a/Engine/Source/Runtime/Engine/Private/Actor.cpp
+++ b/Engine/Source/Runtime/Engine/Private/Actor.cpp
@@ -1103,9 +1103,7 @@ void AActor::AttachRootComponentTo(USceneComponent* InParent, FName InSocketName
 	if(RootComponent && InParent)
 	{
 		RootComponent->AttachTo(InParent, InSocketName, AttachLocationType );
-#if WITH_BODY_WELDING
-		InParent->WeldPhysicsBody(RootComponent);
-#endif
+
 
 		AttachmentReplication.AttachParent = InParent->GetAttachmentRootActor();
 		AttachmentReplication.LocationOffset = RootComponent->RelativeLocation;
@@ -1137,9 +1135,7 @@ void AActor::OnRep_AttachmentReplication()
 				RootComponent->RelativeScale3D = AttachmentReplication.RelativeScale3D;
 
 				RootComponent->UpdateComponentToWorld();
-#if WITH_BODY_WELDING
-				ParentComponent->WeldPhysicsBody(RootComponent, true, AttachmentReplication.AttachSocket);
-#endif
+
 			}
 		}
 	}
@@ -1157,9 +1153,7 @@ void AActor::AttachRootComponentToActor(AActor* InParentActor, FName InSocketNam
 		if (ParentRootComponent)
 		{
 			RootComponent->AttachTo(ParentRootComponent, InSocketName, AttachLocationType );
-#if WITH_BODY_WELDING
-			ParentRootComponent->WeldPhysicsBody(RootComponent, true, InSocketName);
-#endif
+
 
 			AttachmentReplication.AttachParent = InParentActor;
 			AttachmentReplication.LocationOffset = RootComponent->RelativeLocation;
@@ -1188,12 +1182,7 @@ void AActor::DetachRootComponentFromParent(bool bMaintainWorldPosition)
 	if(RootComponent)
 	{
 		USceneComponent * RootComponent = GetRootComponent();
-#if WITH_BODY_WELDING
-		if (RootComponent->AttachParent)
-		{
-			RootComponent->AttachParent->WeldPhysicsBody(RootComponent, false);
-		}
-#endif
+
 
 		GetRootComponent()->DetachFromParent(bMaintainWorldPosition);
 
diff --git a/Engine/Source/Runtime/Engine/Private/PhysicsEngine/BodyInstance.cpp b/Engine/Source/Runtime/Engine/Private/PhysicsEngine/BodyInstance.cpp
--- a/Engine/Source/Runtime/Engine/Private/PhysicsEngine/BodyInstance.cpp
+++ b/Engine/Source/Runtime/Engine/Private/PhysicsEngine/BodyInstance.cpp
@@ -197,6 +197,10 @@ FBodyInstance::FBodyInstance()
 , bUseCCD(false)
 , bNotifyRigidBodyCollision(false)
 , bSimulatePhysics(false)
+#if WITH_BODY_WELDING
+, bAutoWeld(false)
+, bWelded(false)
+#endif
 , bStartAwake(true)
 , bEnableGravity(true)
 , bUseAsyncScene(false)
@@ -899,21 +903,17 @@ void FBodyInstance::InitBody(UBodySetup* Setup, const FTransform& Transform, UPr
 	check(BodySetup.IsValid());
 
 #if WITH_BODY_WELDING
+	if(bAutoWeld)
 	{
 		UPrimitiveComponent * ParentPrimComponent = PrimComp ? Cast<UPrimitiveComponent>(PrimComp->AttachParent) : NULL;
 		
 		//if we have a parent we will now do the weld and exit any further initialization
 		if (ParentPrimComponent && PrimComp->GetWorld() && PrimComp->GetWorld()->IsGameWorld())
 		{
-			/*FTransform RelativeTM = FTransform(PrimComp->RelativeRotation, PrimComp->RelativeLocation, PrimComp->RelativeScale3D);
-			FBodyInstance * OwnerBody = ParentPrimComponent->GetBodyInstance();
-			OwnerBody->Weld(this, RelativeTM);*/
-			if (ParentPrimComponent->WeldPhysicsBody(PrimComp))
+			if (PrimComp->WeldToInternal(ParentPrimComponent, PrimComp->AttachSocketName))	//welded new simulated body so initialization is done
 			{
-				//welded so stop any further initialization
 				return;
 			}
-			
 		}
 	}
 #endif
@@ -1451,7 +1451,7 @@ void FBodyInstance::TermBody()
 }
 
 #if WITH_BODY_WELDING
-bool FBodyInstance::Weld(FBodyInstance* TheirBody, const FTransform& RelativeTM)
+bool FBodyInstance::Weld(FBodyInstance* TheirBody, const FTransform& TheirTM)
 {
 	check(TheirBody);
 	if (TheirBody->BodySetup.IsValid() == false)	//attach actor can be called before body has been initialized. In this case just return false
@@ -1464,7 +1464,7 @@ bool FBodyInstance::Weld(FBodyInstance* TheirBody, const FTransform& RelativeTM)
 #if WITH_PHYSX
 	
 	//first time we're welding with root so we need to grab the shapes
-	if (ShapeToComponentMap.Num() == 0)
+	if (ShapeToBodyMap.Num() == 0)
 	{
 		int32 NumSyncShapes = 0;
 		TArray<PxShape *> PShapes = GetAllShapes(NumSyncShapes);
@@ -1472,13 +1472,17 @@ bool FBodyInstance::Weld(FBodyInstance* TheirBody, const FTransform& RelativeTM)
 		for (int32 ShapeIdx = 0; ShapeIdx < PShapes.Num(); ++ShapeIdx)
 		{
 			PxShape* PShape = PShapes[ShapeIdx];
-			UPrimitiveComponent *& PrimComp = ShapeToComponentMap.FindOrAdd(PShape);
-			PrimComp = OwnerComponent.Get();
+			FBodyInstance *& BI = ShapeToBodyMap.FindOrAdd(PShape);
+			BI = this;
 		}
 	}
 
 	TArray<PxShape *> PNewShapes;
 
+	FTransform MyTM = GetUnrealWorldTransform();
+	MyTM.SetScale3D(Scale3D);	//physx doesn't store 3d so set it here
+
+	FTransform RelativeTM = TheirTM.GetRelativeTransform(MyTM);
 
 	//child body gets placed into the same scenes as parent body
 	if (PxRigidActor* MyBody = RigidActorSync)
@@ -1495,8 +1499,8 @@ bool FBodyInstance::Weld(FBodyInstance* TheirBody, const FTransform& RelativeTM)
 	for (int32 ShapeIdx = 0; ShapeIdx < PNewShapes.Num(); ++ShapeIdx)
 	{
 		PxShape* PShape = PNewShapes[ShapeIdx];
-		UPrimitiveComponent *& PrimComp = ShapeToComponentMap.FindOrAdd(PShape);
-		PrimComp = TheirBody->OwnerComponent.Get();
+		FBodyInstance *& BI = ShapeToBodyMap.FindOrAdd(PShape);
+		BI = TheirBody;
 	}
 
 	PostShapeChange();
@@ -1511,10 +1515,8 @@ bool FBodyInstance::Weld(FBodyInstance* TheirBody, const FTransform& RelativeTM)
 	return false;
 }
 
-void FBodyInstance::UnWeld(UPrimitiveComponent* ChildComponent)
+void FBodyInstance::UnWeld(FBodyInstance* BI)
 {
-	check(ChildComponent);
-
 	//@TODO: BOX2D: Implement Weld
 
 #if WITH_PHYSX
@@ -1525,11 +1527,12 @@ void FBodyInstance::UnWeld(UPrimitiveComponent* ChildComponent)
 	for (int32 ShapeIdx = 0; ShapeIdx < NumSyncShapes; ++ShapeIdx)
 	{
 		PxShape* PShape = PShapes[ShapeIdx];
-		if (UPrimitiveComponent ** PrimCompPtrPtr = ShapeToComponentMap.Find(PShape))
+		if (FBodyInstance ** BIPtrPtr = ShapeToBodyMap.Find(PShape))
 		{
-			if (*PrimCompPtrPtr == ChildComponent)
+			if (*BIPtrPtr == BI)
 			{
 				RigidActorSync->detachShape(*PShape);
+				ShapeToBodyMap.Remove(PShape);
 			}
 		}
 	}
@@ -1537,11 +1540,12 @@ void FBodyInstance::UnWeld(UPrimitiveComponent* ChildComponent)
 	for (int32 ShapeIdx = NumSyncShapes; ShapeIdx < PShapes.Num(); ++ShapeIdx)
 	{
 		PxShape* PShape = PShapes[ShapeIdx];
-		if (UPrimitiveComponent ** PrimCompPtrPtr = ShapeToComponentMap.Find(PShape))
+		if (FBodyInstance ** BIPtrPtr = ShapeToBodyMap.Find(PShape))
 		{
-			if (*PrimCompPtrPtr == ChildComponent)
+			if (*BIPtrPtr == BI)
 			{
 				RigidActorAsync->detachShape(*PShape);
+				ShapeToBodyMap.Remove(PShape);
 			}
 		}
 	}
diff --git a/Engine/Source/Runtime/Engine/Private/PrimitiveComponent.cpp b/Engine/Source/Runtime/Engine/Private/PrimitiveComponent.cpp
--- a/Engine/Source/Runtime/Engine/Private/PrimitiveComponent.cpp
+++ b/Engine/Source/Runtime/Engine/Private/PrimitiveComponent.cpp
@@ -348,84 +348,162 @@ void UPrimitiveComponent::OnChildAttached(USceneComponent* ChildComponent)
 }
 
 #if WITH_BODY_WELDING
-bool UPrimitiveComponent::WeldPhysicsBody(USceneComponent * ChildComponent, bool bWeld, FName ParentBoneName, FName ChildBoneName)
+
+UPrimitiveComponent * GetRootWelded(UPrimitiveComponent * PrimComponent, FName ParentSocketName = NAME_None, FName * OutSocketName = NULL)
 {
-	//assumes that child is already attached
-	if (ChildComponent->AttachParent != this)
+	UPrimitiveComponent * Result = NULL;
+	UPrimitiveComponent * RootComponent = Cast<UPrimitiveComponent>(PrimComponent->AttachParent);	//we must find the root component along hierarchy that has bWelded set to true
+
+	//check that body itself is welded
+	if (FBodyInstance * BI = PrimComponent->GetBodyInstance(ParentSocketName))
 	{
-		return false;
+		if (BI->bWelded == false && BI->bAutoWeld == false)	//we're not welded and we aren't trying to become welded
+		{
+			return NULL;
+		}
 	}
 
-	if (UPrimitiveComponent * ChildPrim = Cast<UPrimitiveComponent>(ChildComponent))
+	FName SocketName = ParentSocketName; //because of skeletal mesh it's important that we check along the bones that we attached
+	FBodyInstance * RootBI = NULL;
+	for (; RootComponent; RootComponent = Cast<UPrimitiveComponent>(RootComponent->AttachParent))
 	{
-		//traverse the hierarchy to weld/unweld from root
-		USceneComponent * RootComponent = this;
-		FTransform RelativeTM = FTransform(ChildComponent->RelativeRotation, ChildComponent->RelativeLocation, ChildComponent->RelativeScale3D);
+		Result = RootComponent;
+		SocketName = RootComponent->AttachSocketName;
 
-		for (; RootComponent->AttachParent; RootComponent = RootComponent->AttachParent)
+		RootBI = RootComponent->GetBodyInstance(SocketName);
+		if (RootBI && RootBI->bWelded)
 		{
-			RelativeTM = RelativeTM * FTransform(RootComponent->RelativeRotation, RootComponent->RelativeLocation, RootComponent->RelativeScale3D);
+			continue;
 		}
 
-		if (UPrimitiveComponent * RootPrim = Cast<UPrimitiveComponent>(RootComponent))
+		break;
+	}
+
+	if (OutSocketName)
+	{
+		*OutSocketName = SocketName;
+	}
+
+	return Result;
+}
+
+void UPrimitiveComponent::GetWeldedBodies(TArray<FBodyInstance*> & OutWeldedBodies, TArray<FName> & OutLabels)
+{
+	if (BodyInstance.bWelded)
+	{
+		OutWeldedBodies.Add(&BodyInstance);
+		OutLabels.Add(NAME_None);
+			
+		for (USceneComponent * Child : AttachChildren)
 		{
-			//we support welding of kinematic actors to simulated parents.
-			//welding simulated to anything doesn't make sense and we'd rather use a physical constraint
-			//welding kinematic to kinematic is not needed as the attachment hierarchy already handles moving the object around
-			if (RootPrim->IsSimulatingPhysics())
+			if (UPrimitiveComponent * PrimChild = Cast<UPrimitiveComponent>(Child))
 			{
-				FBodyInstance * RootBody = RootPrim->GetBodyInstance(ParentBoneName);
-				FBodyInstance * ChildBody = ChildPrim->GetBodyInstance(ChildBoneName);
+				PrimChild->GetWeldedBodies(OutWeldedBodies, OutLabels);
+			}
+		}
+	}
+}
 
-				if (RootBody && ChildBody)
-				{
-					if (bWeld)
-					{
-						return RootBody->Weld(ChildBody, RelativeTM);
-					}
-					else 
-					{
-						//we first init the new root which is ChildBody
-						USceneComponent * OldParent = ChildPrim->AttachParent;	//we need to temporarily NULL it out to create a new body then put it back so that we can properly NULL it out later
-						ChildPrim->AttachParent = NULL;
-						ChildBody->InitBody(ChildPrim->GetBodySetup(), ChildPrim->GetComponentToWorld(), ChildPrim, ChildPrim->GetWorld()->GetPhysicsScene());
-						ChildPrim->AttachParent = OldParent;
-
-						//we need to find all of ChildComponent's children
-						TArray<USceneComponent*> Children;
-						ChildComponent->GetChildrenComponents(true, Children);
-
-						//then remove child and all its children from original root
-						RootBody->UnWeld(ChildPrim);
-						for (USceneComponent * Child : Children)
-						{
-							if (UPrimitiveComponent * Prim = Cast<UPrimitiveComponent>(Child))
-							{
-								RootBody->UnWeld(Prim);
-							}
-						}
-						
-						bool bSuccess = true;
-						//then weld all of child's children to child
-						for (USceneComponent * Child : Children)
-						{
-							if (UPrimitiveComponent * Prim = Cast<UPrimitiveComponent>(Child))
-							{
-								FTransform RelativeTM = Prim->GetComponentToWorld().GetRelativeTransform(ChildPrim->GetComponentToWorld());
-								bSuccess &= ChildBody->Weld(Prim->GetBodyInstance(), RelativeTM);
-							}
-						}
+bool UPrimitiveComponent::WeldToInternal(USceneComponent * InParent, FName ParentSocketName /* = Name_None */)
+{
+	//WeldToInternal assumes attachment is already done
+	if (AttachParent != InParent && AttachSocketName != ParentSocketName)
+	{
+		return false;
+	}
 
-						return bSuccess;
-					}
-				}
+	//Check that we can actually our own socket name
+	FBodyInstance * BI = GetBodyInstance();
+	if (BI == NULL)
+	{
+		return false;
+	}
+
+	FName SocketName;
+	UPrimitiveComponent * RootComponent = GetRootWelded(this, ParentSocketName, &SocketName);
+	
+	if (RootComponent)
+	{
+		if (FBodyInstance * RootBI = RootComponent->GetBodyInstance(SocketName))
+		{
+			BI->bWelded = true;
+			//There are multiple cases to handle:
+			//Root is kinematic, simulated
+			//Child is kinematic, simulated
+			//Child always inherits from root
+
+			//if root is kinematic simply set child to be kinematic and we're done
+			if (RootComponent->IsSimulatingPhysics(SocketName) == false)
+			{
+				SetSimulatePhysics(false);
+				return false;	//return false because we need to continue with regular body initialization
 			}
+
+			//root is simulated so we actually weld the body
+			FTransform RelativeTM = RootComponent == AttachParent ? GetRelativeTransform() : GetComponentToWorld().GetRelativeTransform(RootComponent->GetComponentToWorld());	//if direct parent we already have relative. Otherwise compute it
+			RootBI->Weld(BI, GetComponentToWorld());
+
+			return true;
 		}
 	}
 
-	
 	return false;
 }
+
+void UPrimitiveComponent::WeldTo(USceneComponent* InParent, FName InSocketName /* = NAME_None */)
+{
+	//automatically attach if needed
+	if (AttachParent != InParent && AttachSocketName != InSocketName)
+	{
+		AttachTo(InParent, InSocketName);
+	}
+
+	WeldToInternal(InParent, InSocketName);
+}
+
+void UPrimitiveComponent::UnWeldFromParent()
+{
+	FBodyInstance * NewRootBI = GetBodyInstance();
+	if (NewRootBI == NULL || NewRootBI->bWelded == false)
+	{
+		return;
+	}
+
+	FName SocketName;
+	UPrimitiveComponent * RootComponent = GetRootWelded(this, AttachSocketName, &SocketName);
+
+	if (RootComponent)
+	{
+		if (FBodyInstance * RootBI = RootComponent->GetBodyInstance(SocketName))
+		{
+			//create new root
+			RootBI->UnWeld(NewRootBI);
+			NewRootBI->bWelded = false;
+			//we need to temporarily turn off AutoWeld
+			bool bPrevAutoWeld = NewRootBI->bAutoWeld;
+			NewRootBI->bAutoWeld = false;
+			NewRootBI->InitBody(GetBodySetup(), GetComponentToWorld(), this, GetWorld()->GetPhysicsScene());
+			NewRootBI->bAutoWeld = bPrevAutoWeld;
+
+			//now weld its children to it
+			TArray<FBodyInstance*> ChildrenBodies;
+			TArray<FName> ChildrenLabels;
+			GetWeldedBodies(ChildrenBodies, ChildrenLabels);
+
+			for (int32 ChildIdx = 0; ChildIdx < ChildrenBodies.Num(); ++ChildIdx)
+			{
+				FBodyInstance * ChildBI = ChildrenBodies[ChildIdx];
+				if (ChildBI != NewRootBI)
+				{
+					RootBI->UnWeld(NewRootBI);
+					NewRootBI->Weld(ChildBI, ChildBI->OwnerComponent->GetSocketTransform(ChildrenLabels[ChildIdx]));
+				}
+			}
+		}
+	}
+}
+
+
 #endif
 
 void UPrimitiveComponent::DestroyRenderState_Concurrent()
diff --git a/Engine/Source/Runtime/Engine/Private/SkeletalMeshComponentPhysics.cpp b/Engine/Source/Runtime/Engine/Private/SkeletalMeshComponentPhysics.cpp
--- a/Engine/Source/Runtime/Engine/Private/SkeletalMeshComponentPhysics.cpp
+++ b/Engine/Source/Runtime/Engine/Private/SkeletalMeshComponentPhysics.cpp
@@ -1509,6 +1509,45 @@ FBodyInstance* USkeletalMeshComponent::GetBodyInstance(FName BoneName) const
 	return BodyInst;
 }
 
+#if WITH_BODY_WELDING
+void USkeletalMeshComponent::GetWeldedBodies(TArray<FBodyInstance*> & OutWeldedBodies, TArray<FName> & OutLabels)
+{
+	UPhysicsAsset* PhysicsAsset = GetPhysicsAsset();
+
+	for (int32 BodyIdx = 0; BodyIdx < Bodies.Num(); ++BodyIdx)
+	{
+		FBodyInstance * BI = Bodies[BodyIdx];
+		if (BI && BI->bWelded)
+		{
+			OutWeldedBodies.Add(&BodyInstance);
+			if (PhysicsAsset)
+			{
+				if (UBodySetup * BodySetup = PhysicsAsset->BodySetup[BodyIdx])
+				{
+					OutLabels.Add(BodySetup->BoneName);
+				}
+				else
+				{
+					OutLabels.Add(NAME_None);
+				}
+			}
+			else
+			{
+				OutLabels.Add(NAME_None);
+			}
+
+			for (USceneComponent * Child : AttachChildren)
+			{
+				if (UPrimitiveComponent * PrimChild = Cast<UPrimitiveComponent>(Child))
+				{
+					PrimChild->GetWeldedBodies(OutWeldedBodies, OutLabels);
+				}
+			}
+		}
+	}
+}
+#endif
+
 
 void USkeletalMeshComponent::BreakConstraint(FVector Impulse, FVector HitLocation, FName InBoneName)
 {
diff --git a/Engine/Source/Runtime/Engine/Private/StaticMeshComponent.cpp b/Engine/Source/Runtime/Engine/Private/StaticMeshComponent.cpp
--- a/Engine/Source/Runtime/Engine/Private/StaticMeshComponent.cpp
+++ b/Engine/Source/Runtime/Engine/Private/StaticMeshComponent.cpp
@@ -136,6 +136,10 @@ UStaticMeshComponent::UStaticMeshComponent(const class FPostConstructInitializeP
 	bBoundsChangeTriggersStreamingDataRebuild = true;
 	bHasCustomNavigableGeometry = EHasCustomNavigableGeometry::Yes;
 
+#if WITH_BODY_WELDING
+	GetBodyInstance()->bAutoWeld = true;	//static mesh by default has auto welding
+#endif
+
 #if WITH_EDITORONLY_DATA
 	SelectedEditorSection = INDEX_NONE;
 #endif
