Fixed crash when a null WorldContextObject is passed to UGameplayStatics::GetPlayerController

diff --git a/Engine/Source/Runtime/Engine/Private/GameplayStatics.cpp b/Engine/Source/Runtime/Engine/Private/GameplayStatics.cpp
--- a/Engine/Source/Runtime/Engine/Private/GameplayStatics.cpp
+++ b/Engine/Source/Runtime/Engine/Private/GameplayStatics.cpp
@@ -29,26 +29,29 @@ UGameInstance* UGameplayStatics::GetGameInstance(UObject* WorldContextObject)
 
 APlayerController* UGameplayStatics::GetPlayerController(UObject* WorldContextObject, int32 PlayerIndex ) 
 {
-	UWorld* World = GEngine->GetWorldFromContextObject( WorldContextObject );
-	uint32 Index = 0;
-	for( FConstPlayerControllerIterator Iterator = World->GetPlayerControllerIterator(); Iterator; ++Iterator )
+	if (WorldContextObject)
 	{
-		APlayerController* PlayerController = *Iterator;
-		if( Index == PlayerIndex )
+		UWorld* World = GEngine->GetWorldFromContextObject(WorldContextObject);
+		uint32 Index = 0;
+		for (FConstPlayerControllerIterator Iterator = World->GetPlayerControllerIterator(); Iterator; ++Iterator)
 		{
-			return PlayerController;			
+			APlayerController* PlayerController = *Iterator;
+			if (Index == PlayerIndex)
+			{
+				return PlayerController;
+			}
+			Index++;
 		}
-		Index++;
 	}
 
-	return NULL;
+	return nullptr;
 }
 
 ACharacter* UGameplayStatics::GetPlayerCharacter(UObject* WorldContextObject, int32 PlayerIndex)
 {
-	ACharacter* Character = NULL;
+	ACharacter* Character = nullptr;
 	APlayerController* PC = GetPlayerController(WorldContextObject, PlayerIndex );
-	if (PC != NULL)
+	if (PC != nullptr)
 	{
 		Character = Cast<ACharacter>(PC->GetPawn());
 	}
