* fixed TTP 343328 RENDERING: ScreenSpaceSubsurface Scattering: Unexpected behavior when altering Subsurface Color
* fixed TTP 343332 RENDERING: ScreenSpaceSubsurface Scattering: Adjusting Subsurface Color above 1 creates glow effect

diff --git a/Engine/Source/Runtime/Core/Public/Math/Color.h b/Engine/Source/Runtime/Core/Public/Math/Color.h
--- a/Engine/Source/Runtime/Core/Public/Math/Color.h
+++ b/Engine/Source/Runtime/Core/Public/Math/Color.h
@@ -161,6 +161,19 @@ struct FLinearColor
 		return *this;
 	}
 
+	// clamped in 0..1 range
+	FORCEINLINE FLinearColor GetClamped(float InMin = 0.0f, float InMax = 1.0f) const
+	{
+		FLinearColor Ret;
+
+		Ret.R = FMath::Clamp(R, InMin, InMax);
+		Ret.G = FMath::Clamp(G, InMin, InMax);
+		Ret.B = FMath::Clamp(B, InMin, InMax);
+		Ret.A = FMath::Clamp(A, InMin, InMax);
+
+		return Ret;
+	}
+
 	/** Comparison operators */
 	FORCEINLINE bool operator==(const FLinearColor& ColorB) const
 	{
diff --git a/Engine/Source/Runtime/Engine/Private/Rendering/SubsurfaceProfile.cpp b/Engine/Source/Runtime/Engine/Private/Rendering/SubsurfaceProfile.cpp
--- a/Engine/Source/Runtime/Engine/Private/Rendering/SubsurfaceProfile.cpp
+++ b/Engine/Source/Runtime/Engine/Private/Rendering/SubsurfaceProfile.cpp
@@ -151,14 +151,14 @@ void FSubsurfaceProfileTexture::CreateTexture(FRHICommandListImmediate& RHICmdLi
 
 	for (uint32 y = 0; y < Height; ++y)
 	{
-		FSubsurfaceProfileStruct& Data = SubsurfaceProfileEntries[y].Settings;
+		FSubsurfaceProfileStruct Data = SubsurfaceProfileEntries[y].Settings;
 
-		if(Data.SubsurfaceColor.IsAlmostBlack())
-		{
-			// to avoid div by 0 and a jump to a different value
-			// this basically means we don't want subsurface scattering
-			Data.SubsurfaceColor = FLinearColor(0.0001f, 0.0001f, 0.0001f);
-		}
+		// bias to avoid div by 0 and a jump to a different value
+		// this basically means we don't want subsurface scattering
+		const float Bias = 0.001f;
+
+		Data.SubsurfaceColor = Data.SubsurfaceColor.GetClamped(Bias);
+		Data.FalloffColor = Data.FalloffColor.GetClamped(Bias);
 
 		ComputeMirroredSSSKernel(kernel, Data.SubsurfaceColor, Data.FalloffColor);
 
