Paper2D: Fix contact filtering not being applied in Box2D prototype (also fixes a crash in debug when using Box2D)

diff --git a/Engine/Source/Runtime/Engine/Private/PhysicsEngine2D/Box2DIntegration.cpp b/Engine/Source/Runtime/Engine/Private/PhysicsEngine2D/Box2DIntegration.cpp
--- a/Engine/Source/Runtime/Engine/Private/PhysicsEngine2D/Box2DIntegration.cpp
+++ b/Engine/Source/Runtime/Engine/Private/PhysicsEngine2D/Box2DIntegration.cpp
@@ -108,6 +108,8 @@ FPhysicsScene2D::FPhysicsScene2D(UWorld* AssociatedWorld)
 	World = new b2World(DefaultGravityBox);
 	World->SetDebugDraw(&FBox2DVisualizer::GlobalVisualizer);
 
+	World->SetContactFilter(this);
+
 	StartPhysicsTickFunction.bCanEverTick = true;
 	StartPhysicsTickFunction.Target = this;
 	StartPhysicsTickFunction.TickGroup = TG_StartPhysics;
