* added comment

diff --git a/Engine/Source/Runtime/Engine/Private/Rendering/SubsurfaceProfile.cpp b/Engine/Source/Runtime/Engine/Private/Rendering/SubsurfaceProfile.cpp
--- a/Engine/Source/Runtime/Engine/Private/Rendering/SubsurfaceProfile.cpp
+++ b/Engine/Source/Runtime/Engine/Private/Rendering/SubsurfaceProfile.cpp
@@ -155,6 +155,7 @@ void FSubsurfaceProfileTexture::CreateTexture(FRHICommandListImmediate& RHICmdLi
 
 		// bias to avoid div by 0 and a jump to a different value
 		// this basically means we don't want subsurface scattering
+		// 0.0001f turned out to be too small to fix the issue
 		const float Bias = 0.001f;
 
 		Data.SubsurfaceColor = Data.SubsurfaceColor.GetClamped(Bias);
