Added editor option to draw NavOctree #ue4

diff --git a/Engine/Source/Runtime/Engine/Classes/AI/Navigation/RecastNavMesh.h b/Engine/Source/Runtime/Engine/Classes/AI/Navigation/RecastNavMesh.h
--- a/Engine/Source/Runtime/Engine/Classes/AI/Navigation/RecastNavMesh.h
+++ b/Engine/Source/Runtime/Engine/Classes/AI/Navigation/RecastNavMesh.h
@@ -460,6 +460,10 @@ class ENGINE_API ARecastNavMesh : public ANavigationData
 	UPROPERTY(EditAnywhere, Category=Display)
 	uint32 bDrawClusters:1;
 
+	/** should we draw edges of every navmesh's triangle */
+	UPROPERTY(EditAnywhere, Category = Display)
+	uint32 bDrawOctree : 1;
+
 	UPROPERTY(config)
 	uint32 bDistinctlyDrawTilesBeingBuilt:1;
 
diff --git a/Engine/Source/Runtime/Engine/Private/NavMeshRenderingComponent.cpp b/Engine/Source/Runtime/Engine/Private/NavMeshRenderingComponent.cpp
--- a/Engine/Source/Runtime/Engine/Private/NavMeshRenderingComponent.cpp
+++ b/Engine/Source/Runtime/Engine/Private/NavMeshRenderingComponent.cpp
@@ -9,6 +9,7 @@
 #include "NavigationOctree.h"
 #include "AI/Navigation/RecastHelpers.h"
 #include "AI/Navigation/RecastNavMeshGenerator.h"
+#include "AI/Navigation/NavigationSystem.h"
 
 #include "NavMeshRenderingHelpers.h"
 
@@ -109,6 +110,31 @@ void UNavMeshRenderingComponent::GatherData(struct FNavMeshSceneProxyData* Curre
 		CurrentData->NavMeshDrawOffset.Z += NavMesh->GetConfig().AgentRadius / 10.f;
 
 		NavMesh->BeginBatchQuery();
+
+		if (NavMesh->bDrawOctree)
+		{
+			const UNavigationSystem* NavSys = UNavigationSystem::GetCurrent(GetWorld());
+			const FNavigationOctree* NavOctree = NavSys ? NavSys->GetNavOctree() : NULL;
+
+			for (FNavigationOctree::TConstIterator<> NodeIt(*NavOctree); NodeIt.HasPendingNodes(); NodeIt.Advance())
+			{
+				const FNavigationOctree::FNode& CurrentNode = NodeIt.GetCurrentNode();
+				const int32 CurrentNodeElementCount = CurrentNode.GetElementCount();
+				const FOctreeNodeContext& CorrentContext = NodeIt.GetCurrentContext();
+				CurrentData->OctreeBounds.Add(CorrentContext.Bounds);
+
+				FOREACH_OCTREE_CHILD_NODE(ChildRef)
+				{
+					auto Context = CorrentContext.GetChildContext(ChildRef);
+					auto Child = CurrentNode.GetChild(ChildRef);
+					if (CurrentNode.HasChild(ChildRef))
+					{
+						NodeIt.PushChild(ChildRef);
+					}
+				}
+			}
+		}
+
 		NavMesh->GetDebugGeometry(CurrentData->NavMeshGeometry);
 		const TArray<FVector>& MeshVerts = CurrentData->NavMeshGeometry.MeshVerts;
 
diff --git a/Engine/Source/Runtime/Engine/Private/NavMeshRenderingHelpers.h b/Engine/Source/Runtime/Engine/Private/NavMeshRenderingHelpers.h
--- a/Engine/Source/Runtime/Engine/Private/NavMeshRenderingHelpers.h
+++ b/Engine/Source/Runtime/Engine/Private/NavMeshRenderingHelpers.h
@@ -11,7 +11,7 @@
 #include "AI/Navigation/NavMeshRenderingComponent.h"
 #include "Materials/Material.h"
 #include "MaterialShared.h"
-
+#include "GenericOctree.h"
 
 static const FColor NavMeshRenderColor_Recast_TriangleEdges(255,255,255);
 static const FColor NavMeshRenderColor_Recast_TileEdges(16,16,16,32);
@@ -63,6 +63,7 @@ struct FNavMeshSceneProxyData : public TSharedFromThis<FNavMeshSceneProxyData, E
 		ClusterLinkLines.Reset();
 		PathCollidingGeomIndices.Reset();
 		PathCollidingGeomVerts.Reset();
+		OctreeBounds.Reset();
 		MeshBuilders.Reset();
 		DebugLabels.Reset();
 		Bounds.Init();
@@ -86,6 +87,8 @@ struct FNavMeshSceneProxyData : public TSharedFromThis<FNavMeshSceneProxyData, E
 	TArray<int32> PathCollidingGeomIndices;
 	TNavStatArray<FVector> PathCollidingGeomVerts;
 
+	TArray<FBoxCenterAndExtent>	OctreeBounds;
+		
 	FColor NavMeshColors[RECAST_MAX_AREAS];
 
 	struct FDebugMeshData
@@ -355,6 +358,28 @@ public:
 		return FColor(r*63, g*63, b*63, 164);
 	}
 
+	void DrawDebugBox(FPrimitiveDrawInterface* PDI, FVector const& Center, FVector const& Box, FColor const& Color) const
+	{
+		// no debug line drawing on dedicated server
+		if (PDI != NULL)
+		{
+			PDI->DrawLine(Center + FVector(Box.X, Box.Y, Box.Z), Center + FVector(Box.X, -Box.Y, Box.Z), Color, SDPG_World);
+			PDI->DrawLine(Center + FVector(Box.X, -Box.Y, Box.Z), Center + FVector(-Box.X, -Box.Y, Box.Z), Color, SDPG_World);
+			PDI->DrawLine(Center + FVector(-Box.X, -Box.Y, Box.Z), Center + FVector(-Box.X, Box.Y, Box.Z), Color, SDPG_World);
+			PDI->DrawLine(Center + FVector(-Box.X, Box.Y, Box.Z), Center + FVector(Box.X, Box.Y, Box.Z), Color, SDPG_World);
+
+			PDI->DrawLine(Center + FVector(Box.X, Box.Y, -Box.Z), Center + FVector(Box.X, -Box.Y, -Box.Z), Color, SDPG_World);
+			PDI->DrawLine(Center + FVector(Box.X, -Box.Y, -Box.Z), Center + FVector(-Box.X, -Box.Y, -Box.Z), Color, SDPG_World);
+			PDI->DrawLine(Center + FVector(-Box.X, -Box.Y, -Box.Z), Center + FVector(-Box.X, Box.Y, -Box.Z), Color, SDPG_World);
+			PDI->DrawLine(Center + FVector(-Box.X, Box.Y, -Box.Z), Center + FVector(Box.X, Box.Y, -Box.Z), Color, SDPG_World);
+
+			PDI->DrawLine(Center + FVector(Box.X, Box.Y, Box.Z), Center + FVector(Box.X, Box.Y, -Box.Z), Color, SDPG_World);
+			PDI->DrawLine(Center + FVector(Box.X, -Box.Y, Box.Z), Center + FVector(Box.X, -Box.Y, -Box.Z), Color, SDPG_World);
+			PDI->DrawLine(Center + FVector(-Box.X, -Box.Y, Box.Z), Center + FVector(-Box.X, -Box.Y, -Box.Z), Color, SDPG_World);
+			PDI->DrawLine(Center + FVector(-Box.X, Box.Y, Box.Z), Center + FVector(-Box.X, Box.Y, -Box.Z), Color, SDPG_World);
+		}
+	}
+
 	virtual void DrawDynamicElements(FPrimitiveDrawInterface* PDI,const FSceneView* View) override
 	{
 		QUICK_SCOPE_CYCLE_COUNTER( STAT_RecastRenderingSceneProxy_DrawDynamicElements );
@@ -381,7 +406,13 @@ public:
 			MeshBuilder.AddTriangles( ProxyData.MeshBuilders[Index].Indices );
 			MeshBuilder.Draw(PDI, FMatrix::Identity, MeshColorInstance, GetDepthPriorityGroup(View));
 		}
-
+		
+		for (int32 Index = 0; Index < ProxyData.OctreeBounds.Num(); ++Index)
+		{
+			FBoxCenterAndExtent& Bounds = ProxyData.OctreeBounds[Index];
+			DrawDebugBox(PDI, Bounds.Center, Bounds.Extent, FColor::White);
+		}
+		
 		FHitProxyId HitProxyId;
 		FBatchedElements BatchedElements;
 		int32 Num = ProxyData.NavMeshEdgeLines.Num();
@@ -502,6 +533,12 @@ public:
 
 				const bool bSkipDistanceCheck = GIsEditor && (GEngine->GetDebugLocalPlayer() == NULL);
 
+				for (int32 Index = 0; Index < ProxyData.OctreeBounds.Num(); ++Index)
+				{
+					const FBoxCenterAndExtent& Bounds = ProxyData.OctreeBounds[Index];
+					DrawDebugBox(PDI, Bounds.Center, Bounds.Extent, FColor::White);
+				}
+
 				// Draw Mesh
 				for (int32 Index = 0; Index < MeshBatchElements.Num(); ++Index)
 				{
