Addecd check for level bounds while adding to NavOctree. #ue4

diff --git a/Engine/Source/Runtime/Engine/Private/AI/Navigation/NavigationSystem.cpp b/Engine/Source/Runtime/Engine/Private/AI/Navigation/NavigationSystem.cpp
--- a/Engine/Source/Runtime/Engine/Private/AI/Navigation/NavigationSystem.cpp
+++ b/Engine/Source/Runtime/Engine/Private/AI/Navigation/NavigationSystem.cpp
@@ -2694,9 +2694,12 @@ void UNavigationSystem::AddLevelCollisionToOctree(ULevel* Level)
 		FNavigationOctreeElement BSPElem;
 		FRecastNavMeshGenerator::ExportVertexSoupGeometry(*LevelGeom, BSPElem.Data);
 
-		NavOctree->AddNode(Level, NULL, BSPElem.Data.Bounds, BSPElem);
-
-		AddDirtyArea(BSPElem.Bounds.GetBox(), ENavigationDirtyFlag::All);
+		const auto& Bounds = BSPElem.Data.Bounds;
+		if (!Bounds.GetExtent().IsNearlyZero())
+		{
+			NavOctree->AddNode(Level, NULL, Bounds, BSPElem);
+			AddDirtyArea(Bounds, ENavigationDirtyFlag::All);
+		}
 	}
 #endif // NAVOCTREE_CONTAINS_COLLISION_DATA
 }
