If startAwake is set to true then SetSimulatePhysics will automatically wake body up

diff --git a/Engine/Source/Runtime/Engine/Private/PhysicsEngine/BodyInstance.cpp b/Engine/Source/Runtime/Engine/Private/PhysicsEngine/BodyInstance.cpp
--- a/Engine/Source/Runtime/Engine/Private/PhysicsEngine/BodyInstance.cpp
+++ b/Engine/Source/Runtime/Engine/Private/PhysicsEngine/BodyInstance.cpp
@@ -1976,6 +1976,12 @@ void FBodyInstance::UpdateInstanceSimulatePhysics(bool bIgnoreParent)
 		{
 			SCOPED_SCENE_WRITE_LOCK(PRigidDynamic->getScene());
 			PRigidDynamic->setRigidDynamicFlag(PxRigidDynamicFlag::eKINEMATIC, bNewKinematic);
+
+			//if wake when level starts is true, calling this function automatically wakes body up
+			if (bSimulatePhysics && bStartAwake)
+			{
+				PRigidDynamic->wakeUp();
+			}
 		}
 	}
 #endif
