Remove autowelding ifdef

diff --git a/Engine/Source/Editor/DetailCustomizations/Private/PrimitiveComponentDetails.cpp b/Engine/Source/Editor/DetailCustomizations/Private/PrimitiveComponentDetails.cpp
--- a/Engine/Source/Editor/DetailCustomizations/Private/PrimitiveComponentDetails.cpp
+++ b/Engine/Source/Editor/DetailCustomizations/Private/PrimitiveComponentDetails.cpp
@@ -178,9 +178,7 @@ void FPrimitiveComponentDetails::CustomizeDetails( IDetailLayoutBuilder& DetailB
 					}
 					else if (PropName == TEXT("AutoWeld"))
 					{
-#if WITH_BODY_WELDING
 						PhysicsCategory.AddProperty(ChildProperty);	//can't ifdef the property out so for now we just customize it to be hidden in UI
-#endif
 					}
 					else
 					{
diff --git a/Engine/Source/Runtime/Core/Public/UObject/ObjectVersion.h b/Engine/Source/Runtime/Core/Public/UObject/ObjectVersion.h
--- a/Engine/Source/Runtime/Core/Public/UObject/ObjectVersion.h
+++ b/Engine/Source/Runtime/Core/Public/UObject/ObjectVersion.h
@@ -598,6 +598,8 @@ enum EUnrealEngineObjectUE4Version
 	VER_UE4_FIXUP_WIDGET_ANIMATION_CLASS,
 	// USoundWave objects now contain details about compression scheme used.
 	VER_UE4_SOUND_COMPRESSION_TYPE_ADDED,
+	// Bodies will automatically weld when attached
+	VER_UE4_AUTO_WELDING,
 
 	// -----<new versions can be added before this line>-------------------------------------------------
 	// - this needs to be the last line (see note below)
diff --git a/Engine/Source/Runtime/Engine/Classes/Components/PrimitiveComponent.h b/Engine/Source/Runtime/Engine/Classes/Components/PrimitiveComponent.h
--- a/Engine/Source/Runtime/Engine/Classes/Components/PrimitiveComponent.h
+++ b/Engine/Source/Runtime/Engine/Classes/Components/PrimitiveComponent.h
@@ -1024,22 +1024,18 @@ public:
 	/** Called when a component is 'damaged', allowing for component class specific behaviour */
 	virtual void ReceiveComponentDamage(float DamageAmount, FDamageEvent const& DamageEvent, AController* EventInstigator, AActor* DamageCauser);
 
-#if WITH_BODY_WELDING
-
 	/**
 	*   Welds this component to another scene component, optionally at a named socket. Component is automatically attached if not already
 	*	Welding allows the child physics object to become physically connected to its parent. This is useful for creating compound rigid bodies with correct mass distribution.
 	*   @param InParent the component to be physically attached to
 	*   @param InSocketName optional socket to attach component to
 	*/
-	UFUNCTION(BlueprintCallable, Category = "Physics")
-	virtual void WeldTo(class USceneComponent* InParent, FName InSocketName = NAME_None);
+	virtual void WeldTo(class USceneComponent* InParent, FName InSocketName = NAME_None) override;
 
 	/**
 	*   UnWelds this component from its parent component. Attachment is maintained (DetachFromParent automatically unwelds)
 	*/
-	UFUNCTION(BlueprintCallable, Category = "Physics")
-	virtual void UnWeldFromParent();
+	virtual void UnWeldFromParent() override;
 
 	/**
 	*	Does the actual work for welding.
@@ -1051,7 +1047,6 @@ public:
 	*	Adds the bodies that are currently welded to the OutWeldedBodies array 
 	*/
 	virtual void GetWeldedBodies(TArray<FBodyInstance*> & OutWeldedBodies, TArray<FName> & OutLabels);
-#endif
 
 protected:
 
diff --git a/Engine/Source/Runtime/Engine/Classes/Components/SceneComponent.h b/Engine/Source/Runtime/Engine/Classes/Components/SceneComponent.h
--- a/Engine/Source/Runtime/Engine/Classes/Components/SceneComponent.h
+++ b/Engine/Source/Runtime/Engine/Classes/Components/SceneComponent.h
@@ -342,7 +342,6 @@ public:
 	UFUNCTION(BlueprintCallable, Category="Utilities|Transformation")
 	virtual void DetachFromParent(bool bMaintainWorldPosition = false);
 
-#if WITH_BODY_WELDING
 	/**
 	*   Welds this component to another scene component, optionally at a named socket. Component is automatically attached if not already
 	*	Welding allows the child physics object to become physically connected to its parent. This is useful for creating compound rigid bodies with correct mass distribution.
@@ -350,13 +349,13 @@ public:
 	*   @param InSocketName optional socket to attach component to
 	*/
 	UFUNCTION(BlueprintCallable, Category = "Physics")
-	virtual void WeldTo(class USceneComponent* InParent, FName InSocketName = NAME_None){}
+	virtual void WeldTo(class USceneComponent* InParent, FName InSocketName = NAME_None);
 
 	/**
 	*   UnWelds this component from its parent component. Attachment is maintained (DetachFromParent automatically unwelds)
 	*/
 	UFUNCTION(BlueprintCallable, Category = "Physics")
-	virtual void UnWeldFromParent(){}
+	virtual void UnWeldFromParent();
 
 
 	/**
@@ -371,7 +370,6 @@ public:
 	*/
 	//virtual bool WeldPhysicsBody(USceneComponent * ChildComponent, bool bWeld = true, FName ParentBoneName = NAME_None, FName ChildBoneName = NAME_None) { return false; }
 
-#endif
 
 	/** 
 	 * Gets the names of all the sockets on the component.
diff --git a/Engine/Source/Runtime/Engine/Classes/Components/SkeletalMeshComponent.h b/Engine/Source/Runtime/Engine/Classes/Components/SkeletalMeshComponent.h
--- a/Engine/Source/Runtime/Engine/Classes/Components/SkeletalMeshComponent.h
+++ b/Engine/Source/Runtime/Engine/Classes/Components/SkeletalMeshComponent.h
@@ -954,10 +954,8 @@ public:
 	/** Enable or Disable AngularVelocityDrive based on a list of bone names */
 	void SetNamedMotorsAngularVelocityDrive(bool bEnableSwingDrive, bool bEnableTwistDrive, const TArray<FName>& BoneNames, bool bSetOtherBodiesToComplement = false);
 
-#if WITH_BODY_WELDING
 	void GetWeldedBodies(TArray<FBodyInstance*> & OutWeldedBodies, TArray<FName> & OutChildrenLabels) override;
 	bool virtual WeldToInternal(USceneComponent * InParent, FName ParentSocketName = NAME_None){ return false; }	//not currently supported for skeletal mesh
-#endif
 
 	/**
 	 * Return Transform Matrix for SkeletalMeshComponent considering root motion setups
diff --git a/Engine/Source/Runtime/Engine/Classes/PhysicsEngine/BodyInstance.h b/Engine/Source/Runtime/Engine/Classes/PhysicsEngine/BodyInstance.h
--- a/Engine/Source/Runtime/Engine/Classes/PhysicsEngine/BodyInstance.h
+++ b/Engine/Source/Runtime/Engine/Classes/PhysicsEngine/BodyInstance.h
@@ -168,14 +168,12 @@ public:
 	UPROPERTY(EditAnywhere, BlueprintReadWrite, Category=Physics)
 	uint32 bSimulatePhysics:1;
 
-//#if WITH_BODY_WELDING
 	/** If true and is attached to a parent, the two bodies will be joined into a single rigid body. Physical settings like collision profile and body settings are determined by the root */
 	UPROPERTY(EditAnywhere, BlueprintReadWrite, Category = Physics)
 	uint32 bAutoWeld : 1;
 
 	/** determines if the body is currently welded */
 	uint32 bWelded : 1;
-//#endif
 
 	/** If object should start awake, or if it should initially be sleeping */
 	UPROPERTY(EditAnywhere, BlueprintReadOnly, Category=Physics, meta=(editcondition = "bSimulatePhysics"))
@@ -321,7 +319,6 @@ public:
 	void TermBody();
 
 
-#if WITH_BODY_WELDING
 
 #if WITH_PHYSX
 	TMap<physx::PxShape*, FBodyInstance*> ShapeToBodyMap;
@@ -339,7 +336,6 @@ public:
 	/**
 	 * After adding/removing shapes call this function to update mass distribution etc... */
 	void PostShapeChange();
-#endif
 
 	/**
 	 * Update Body Scale
diff --git a/Engine/Source/Runtime/Engine/Classes/PhysicsEngine/BodySetup.h b/Engine/Source/Runtime/Engine/Classes/PhysicsEngine/BodySetup.h
--- a/Engine/Source/Runtime/Engine/Classes/PhysicsEngine/BodySetup.h
+++ b/Engine/Source/Runtime/Engine/Classes/PhysicsEngine/BodySetup.h
@@ -587,11 +587,7 @@ public:
 	/** 
 	 *   Add the shapes defined by this body setup to the supplied PxRigidBody. 
 	 */
-#if WITH_BODY_WELDING
 	void AddShapesToRigidActor(physx::PxRigidActor* PDestActor, FVector& Scale3D, const FTransform* RelativeTM = NULL, TArray<physx::PxShape*> * NewShapes = NULL);
-#else
-	void AddShapesToRigidActor(physx::PxRigidActor* PDestActor, FVector& Scale3D);
-#endif
 #endif // WITH_PHYSX
 
 };
diff --git a/Engine/Source/Runtime/Engine/Private/PhysicsEngine/BodyInstance.cpp b/Engine/Source/Runtime/Engine/Private/PhysicsEngine/BodyInstance.cpp
--- a/Engine/Source/Runtime/Engine/Private/PhysicsEngine/BodyInstance.cpp
+++ b/Engine/Source/Runtime/Engine/Private/PhysicsEngine/BodyInstance.cpp
@@ -197,10 +197,8 @@ FBodyInstance::FBodyInstance()
 , bUseCCD(false)
 , bNotifyRigidBodyCollision(false)
 , bSimulatePhysics(false)
-#if WITH_BODY_WELDING
 , bAutoWeld(false)
 , bWelded(false)
-#endif
 , bStartAwake(true)
 , bEnableGravity(true)
 , bUseAsyncScene(false)
@@ -902,7 +900,6 @@ void FBodyInstance::InitBody(UBodySetup* Setup, const FTransform& Transform, UPr
 	// whenever update filter, this check will trigger, it seems it's best to know when we initialize body instead. 
 	check(BodySetup.IsValid());
 
-#if WITH_BODY_WELDING
 	if(bAutoWeld)
 	{
 		UPrimitiveComponent * ParentPrimComponent = PrimComp ? Cast<UPrimitiveComponent>(PrimComp->AttachParent) : NULL;
@@ -916,7 +913,6 @@ void FBodyInstance::InitBody(UBodySetup* Setup, const FTransform& Transform, UPr
 			}
 		}
 	}
-#endif
 
 	USkeletalMeshComponent* SkelMeshComp = Cast<USkeletalMeshComponent>(OwnerComponent.Get());
 	if (SkelMeshComp != nullptr)
@@ -1450,7 +1446,6 @@ void FBodyInstance::TermBody()
 	
 }
 
-#if WITH_BODY_WELDING
 bool FBodyInstance::Weld(FBodyInstance* TheirBody, const FTransform& TheirTM)
 {
 	check(TheirBody);
@@ -1566,7 +1561,6 @@ void FBodyInstance::PostShapeChange()
 	// Update damping
 	UpdateDampingProperties();
 }
-#endif
 
 float AdjustForSmallThreshold(float NewVal, float OldVal)
 {
diff --git a/Engine/Source/Runtime/Engine/Private/PhysicsEngine/BodySetup.cpp b/Engine/Source/Runtime/Engine/Private/PhysicsEngine/BodySetup.cpp
--- a/Engine/Source/Runtime/Engine/Private/PhysicsEngine/BodySetup.cpp
+++ b/Engine/Source/Runtime/Engine/Private/PhysicsEngine/BodySetup.cpp
@@ -275,22 +275,13 @@ void SetupNonUniformHelper(FVector & Scale3D, float & MinScale, float & MinScale
 	}
 }
 
-#if WITH_BODY_WELDING
 void UBodySetup::AddShapesToRigidActor(PxRigidActor* PDestActor, FVector& Scale3D, const FTransform * RelativeTM /* = NULL */, TArray<physx::PxShape*> * NewShapes /* = NULL */ )
-#else
-void UBodySetup::AddShapesToRigidActor(PxRigidActor* PDestActor, FVector& Scale3D)
-#endif
 {
 #if WITH_EDITOR
 	// in editor, there are a lot of things relying on body setup to create physics meshes
 	CreatePhysicsMeshes();
 #endif
 
-#if WITH_BODY_WELDING
-#else
-	FTransform * RelativeTM = NULL;	//code below already handles NULL case, just using ifdef to make it clear that the non NULL case isn't really supported yet
-#endif
-
 	float MinScale;
 	float MinScaleAbs;
 	FVector Scale3DAbs;
@@ -340,12 +331,11 @@ void UBodySetup::AddShapesToRigidActor(PxRigidActor* PDestActor, FVector& Scale3
 
 				ensure(PLocalPose.isValid());
 				PxShape* NewShape = PDestActor->createShape(PSphereGeom, *PDefaultMat, PLocalPose);
-#if WITH_BODY_WELDING
+				
 				if (NewShapes)
 				{
 					NewShapes->Add(NewShape);
 				}
-#endif
 
 				const float ContactOffset = FMath::Min(MaxContactOffset, ContactOffsetFactor * PSphereGeom.radius);
 				NewShape->setContactOffset(ContactOffset);
@@ -376,12 +366,11 @@ void UBodySetup::AddShapesToRigidActor(PxRigidActor* PDestActor, FVector& Scale3
 
 				ensure(PLocalPose.isValid());
 				PxShape* NewShape = PDestActor->createShape(PBoxGeom, *PDefaultMat, PLocalPose);
-#if WITH_BODY_WELDING
+	
 				if (NewShapes)
 				{
 					NewShapes->Add(NewShape);
 				}
-#endif
 
 				const float ContactOffset = FMath::Min(MaxContactOffset, ContactOffsetFactor * PBoxGeom.halfExtents.minElement());
 				NewShape->setContactOffset(ContactOffset);
@@ -424,12 +413,11 @@ void UBodySetup::AddShapesToRigidActor(PxRigidActor* PDestActor, FVector& Scale3
 
 					ensure(PLocalPose.isValid());
 					PxShape* NewShape = PDestActor->createShape(PCapsuleGeom, *PDefaultMat, PLocalPose);
-#if WITH_BODY_WELDING
+	
 					if (NewShapes)
 					{
 						NewShapes->Add(NewShape);
 					}
-#endif
 
 					const float ContactOffset = FMath::Min(MaxContactOffset, ContactOffsetFactor * PCapsuleGeom.radius);
 					NewShape->setContactOffset(ContactOffset);
@@ -471,12 +459,11 @@ void UBodySetup::AddShapesToRigidActor(PxRigidActor* PDestActor, FVector& Scale3
 
 						ensure(PLocalPose.isValid());
 						PxShape* NewShape = PDestActor->createShape(PConvexGeom, *PDefaultMat, PLocalPose);
-#if WITH_BODY_WELDING
+	
 						if (NewShapes)
 						{
 							NewShapes->Add(NewShape);
 						}
-#endif
 
 						const float ContactOffset = FMath::Min(MaxContactOffset, ContactOffsetFactor * PBoundsExtents.minElement());
 						NewShape->setContactOffset(ContactOffset);
diff --git a/Engine/Source/Runtime/Engine/Private/SceneComponent.cpp b/Engine/Source/Runtime/Engine/Private/SceneComponent.cpp
--- a/Engine/Source/Runtime/Engine/Private/SceneComponent.cpp
+++ b/Engine/Source/Runtime/Engine/Private/SceneComponent.cpp
@@ -70,6 +70,15 @@ void USceneComponent::OnUpdateTransform(bool bSkipPhysicsMove)
 {
 }
 
+void USceneComponent::WeldTo(class USceneComponent* InParent, FName InSocketName)
+{
+}
+
+void USceneComponent::UnWeldFromParent()
+{
+
+}
+
 void USceneComponent::UpdateComponentToWorldWithParent(USceneComponent * Parent, bool bSkipPhysicsMove)
 {
 	// If our parent hasn't been updated before, we'll need walk up our parent attach hierarchy
@@ -725,11 +734,12 @@ void USceneComponent::AttachTo(class USceneComponent* Parent, FName InSocketName
 		{
 			//This code requires some explaining. Inside the editor we allow user to attach physically simulated objects to other objects. This is done for convenience so that users can group things together in hierarchy.
 			//At runtime we must not attach physically simulated objects as it will cause double transform updates, and you should just use a physical constraint if attachment is the desired behavior.
+			//Note that if bAutoWeld is true then it means that they actually do want the bodies to be simulated together, in which case we should not break the attachment
 			//We must fixup the relative location,rotation,scale as the attachment is no longer valid. Blueprint uses simple construction to try and attach before ComponentToWorld has ever been updated, so we cannot rely on it.
 			//As such we must calculate the proper Relative information
 			//Also physics state may not be created yet so we use bSimulatePhysics to determine if the object has any intention of being physically simulated
 			UPrimitiveComponent * PrimitiveComponent = Cast<UPrimitiveComponent>(this);
-			if (PrimitiveComponent && PrimitiveComponent->BodyInstance.bSimulatePhysics && GetWorld() && GetWorld()->IsGameWorld())
+			if (PrimitiveComponent && PrimitiveComponent->BodyInstance.bSimulatePhysics && !PrimitiveComponent->BodyInstance.bAutoWeld && GetWorld() && GetWorld()->IsGameWorld())
 			{
 				//Since the object is physically simulated it can't be the case that it's a child of object A and being attached to object B (at runtime)
 				UpdateComponentToWorldWithParent(Parent, false);
@@ -828,6 +838,17 @@ void USceneComponent::AttachTo(class USceneComponent* Parent, FName InSocketName
 		// calculate transform with new attachment condition
 		UpdateComponentToWorld();
 
+		if (UPrimitiveComponent * PrimitiveComponent = Cast<UPrimitiveComponent>(this))
+		{
+			if (FBodyInstance * BI = PrimitiveComponent->GetBodyInstance())
+			{
+				if (BI->bAutoWeld)
+				{
+					PrimitiveComponent->WeldToInternal(AttachParent, AttachSocketName);
+				}
+			}
+		}
+
 		// Update overlaps, in case location changed or overlap state depends on attachment.
 		if (IsRegistered())
 		{
@@ -845,6 +866,8 @@ void USceneComponent::DetachFromParent(bool bMaintainWorldPosition)
 {
 	if(AttachParent != NULL)
 	{
+		UnWeldFromParent();
+
 		// Make sure parent points to us if we're registered
 		checkf(!bRegistered || AttachParent->AttachChildren.Contains(this), TEXT("Attempt to detach SceneComponent '%s' owned by '%s' from AttachParent '%s' while not attached."), *GetName(), (GetOwner() ? *GetOwner()->GetName() : TEXT("Unowned")), *AttachParent->GetName());
 
diff --git a/Engine/Source/Runtime/Engine/Private/SkeletalMeshComponentPhysics.cpp b/Engine/Source/Runtime/Engine/Private/SkeletalMeshComponentPhysics.cpp
--- a/Engine/Source/Runtime/Engine/Private/SkeletalMeshComponentPhysics.cpp
+++ b/Engine/Source/Runtime/Engine/Private/SkeletalMeshComponentPhysics.cpp
@@ -1509,7 +1509,6 @@ FBodyInstance* USkeletalMeshComponent::GetBodyInstance(FName BoneName) const
 	return BodyInst;
 }
 
-#if WITH_BODY_WELDING
 void USkeletalMeshComponent::GetWeldedBodies(TArray<FBodyInstance*> & OutWeldedBodies, TArray<FName> & OutLabels)
 {
 	UPhysicsAsset* PhysicsAsset = GetPhysicsAsset();
@@ -1546,7 +1545,6 @@ void USkeletalMeshComponent::GetWeldedBodies(TArray<FBodyInstance*> & OutWeldedB
 		}
 	}
 }
-#endif
 
 
 void USkeletalMeshComponent::BreakConstraint(FVector Impulse, FVector HitLocation, FName InBoneName)
diff --git a/Engine/Source/Runtime/Engine/Private/StaticMeshComponent.cpp b/Engine/Source/Runtime/Engine/Private/StaticMeshComponent.cpp
--- a/Engine/Source/Runtime/Engine/Private/StaticMeshComponent.cpp
+++ b/Engine/Source/Runtime/Engine/Private/StaticMeshComponent.cpp
@@ -136,9 +136,7 @@ UStaticMeshComponent::UStaticMeshComponent(const class FPostConstructInitializeP
 	bBoundsChangeTriggersStreamingDataRebuild = true;
 	bHasCustomNavigableGeometry = EHasCustomNavigableGeometry::Yes;
 
-#if WITH_BODY_WELDING
 	GetBodyInstance()->bAutoWeld = true;	//static mesh by default has auto welding
-#endif
 
 #if WITH_EDITORONLY_DATA
 	SelectedEditorSection = INDEX_NONE;
@@ -261,6 +259,11 @@ void UStaticMeshComponent::Serialize(FArchive& Ar)
 		// Irrelevant lights were incorrect before VER_UE4_TOSS_IRRELEVANT_LIGHTS
 		IrrelevantLights.Empty();
 	}
+
+	if (Ar.UE4Ver() < VER_UE4_AUTO_WELDING)
+	{
+		GetBodyInstance()->bAutoWeld = false;	//existing content may rely on no auto welding
+	}
 }
 
 
