Forgot this file. Fix auto weld linker error

diff --git a/Engine/Source/Runtime/Engine/Private/PrimitiveComponent.cpp b/Engine/Source/Runtime/Engine/Private/PrimitiveComponent.cpp
--- a/Engine/Source/Runtime/Engine/Private/PrimitiveComponent.cpp
+++ b/Engine/Source/Runtime/Engine/Private/PrimitiveComponent.cpp
@@ -327,28 +327,8 @@ void UPrimitiveComponent::OnAttachmentChanged()
 void UPrimitiveComponent::OnChildAttached(USceneComponent* ChildComponent)
 {
 	Super::OnChildAttached(ChildComponent);
-	/*
-#if WITH_BODY_WELDING
-	//We want to support attaching physically simulated rigid bodies to create a single rigid body.
-	if (UPrimitiveComponent * PrimitiveChild = Cast<UPrimitiveComponent>(ChildComponent))
-	{
-		//Note that the skeletal mesh case is kind of odd because we're attaching to the root bone
-		FBodyInstance * MyBody = GetBodyInstance();
-		FBodyInstance * TheirBody = PrimitiveChild->GetBodyInstance();
-
-		//We only do the physical weld when actually running the game
-		if (MyBody && TheirBody && IsSimulatingPhysics() && GetWorld() && GetWorld()->IsGameWorld())
-		{
-			FTransform RelativeTM = FTransform(ChildComponent->RelativeRotation, ChildComponent->RelativeLocation, ChildComponent->RelativeScale3D);
-			MyBody->Weld(TheirBody, RelativeTM);
-		}
-	}
-#endif
-	*/
 }
 
-#if WITH_BODY_WELDING
-
 UPrimitiveComponent * GetRootWelded(UPrimitiveComponent * PrimComponent, FName ParentSocketName = NAME_None, FName * OutSocketName = NULL)
 {
 	UPrimitiveComponent * Result = NULL;
@@ -407,7 +387,7 @@ void UPrimitiveComponent::GetWeldedBodies(TArray<FBodyInstance*> & OutWeldedBodi
 bool UPrimitiveComponent::WeldToInternal(USceneComponent * InParent, FName ParentSocketName /* = Name_None */)
 {
 	//WeldToInternal assumes attachment is already done
-	if (AttachParent != InParent && AttachSocketName != ParentSocketName)
+	if (AttachParent != InParent || AttachSocketName != ParentSocketName)
 	{
 		return false;
 	}
@@ -453,9 +433,9 @@ bool UPrimitiveComponent::WeldToInternal(USceneComponent * InParent, FName Paren
 void UPrimitiveComponent::WeldTo(USceneComponent* InParent, FName InSocketName /* = NAME_None */)
 {
 	//automatically attach if needed
-	if (AttachParent != InParent && AttachSocketName != InSocketName)
+	if (AttachParent != InParent || AttachSocketName != InSocketName)
 	{
-		AttachTo(InParent, InSocketName);
+		AttachTo(InParent, InSocketName, EAttachLocation::KeepWorldPosition);
 	}
 
 	WeldToInternal(InParent, InSocketName);
@@ -504,8 +484,6 @@ void UPrimitiveComponent::UnWeldFromParent()
 }
 
 
-#endif
-
 void UPrimitiveComponent::DestroyRenderState_Concurrent()
 {
 	// Remove the primitive from the scene.
