Fixed CIS error with Monolithic builds

diff --git a/Engine/Source/Runtime/Engine/Private/NavMeshRenderingHelpers.h b/Engine/Source/Runtime/Engine/Private/NavMeshRenderingHelpers.h
--- a/Engine/Source/Runtime/Engine/Private/NavMeshRenderingHelpers.h
+++ b/Engine/Source/Runtime/Engine/Private/NavMeshRenderingHelpers.h
@@ -11,6 +11,7 @@
 #include "AI/Navigation/NavMeshRenderingComponent.h"
 #include "Materials/Material.h"
 #include "MaterialShared.h"
+#include "GenericOctreePublic.h"
 #include "GenericOctree.h"
 
 static const FColor NavMeshRenderColor_Recast_TriangleEdges(255,255,255);
