Temp fix for particle crash in -emulatestereo

diff --git a/Engine/Source/Runtime/Engine/Private/Particles/ParticleSystemRender.cpp b/Engine/Source/Runtime/Engine/Private/Particles/ParticleSystemRender.cpp
--- a/Engine/Source/Runtime/Engine/Private/Particles/ParticleSystemRender.cpp
+++ b/Engine/Source/Runtime/Engine/Private/Particles/ParticleSystemRender.cpp
@@ -1193,6 +1193,11 @@ void GatherParticleLightData(const FDynamicSpriteEmitterReplayDataBase& Source,
 				FVector LightPosition = Source.bUseLocalSpace ? FVector(InLocalToWorld.TransformPosition(ParticlePosition)) : ParticlePosition;
 				
 				// Calculate light positions per-view if we are using a camera offset.
+				// OutParticleLights is accumulated over MULTIPLE emitters.  Some may use camera offset, some may not.
+				// if some do and some don't, then OutParticleLights.PerViewData will cease to have NumLights*NumViews entries and 
+				// GetViewDependentData on it will become bogus and in the worst case, crash from reading out of bounds.
+				// Hack fix is to just add all views always.  Real fix would be detecting whether any emitter was using cameraoffset and 
+				// add for all emitters in that case.
 				if (Source.CameraPayloadOffset != 0)
 				{
 					for(int32 ViewIndex = 0; ViewIndex < InViewFamily.Views.Num(); ++ViewIndex)
@@ -1210,7 +1215,10 @@ void GatherParticleLightData(const FDynamicSpriteEmitterReplayDataBase& Source,
 					// When not using camera-offset, output one position for all views to share. 
 					FSimpleLightPerViewEntry PerViewData;
 					PerViewData.Position = LightPosition;
-					OutParticleLights.PerViewData.Add(PerViewData);
+					for (int32 ViewIndex = 0; ViewIndex < InViewFamily.Views.Num(); ++ViewIndex)
+					{
+						OutParticleLights.PerViewData.Add(PerViewData);
+					}									
 				}
 
 				// Add an entry for the light instance.
