Allow moving simulated skeletal mesh by calling SetLocation etc... #codereview Lina.Halper

diff --git a/Engine/Source/Runtime/Engine/Classes/Components/SkeletalMeshComponent.h b/Engine/Source/Runtime/Engine/Classes/Components/SkeletalMeshComponent.h
--- a/Engine/Source/Runtime/Engine/Classes/Components/SkeletalMeshComponent.h
+++ b/Engine/Source/Runtime/Engine/Classes/Components/SkeletalMeshComponent.h
@@ -999,7 +999,7 @@ public:
 	 *	Iterate over each physics body in the physics for this mesh, and for each 'kinematic' (ie fixed or default if owner isn't simulating) one, update its
 	 *	transform based on the animated transform.
 	 */
-	void UpdateKinematicBonesToPhysics(bool bTeleport, bool bNeedsSkinning);
+	void UpdateKinematicBonesToPhysics(bool bTeleport, bool bNeedsSkinning, bool bForceUpdate = false);
 	
 	/**
 	 * Look up all bodies for broken constraints.
diff --git a/Engine/Source/Runtime/Engine/Private/PhysicsEngine/PhysAnim.cpp b/Engine/Source/Runtime/Engine/Private/PhysicsEngine/PhysAnim.cpp
--- a/Engine/Source/Runtime/Engine/Private/PhysicsEngine/PhysAnim.cpp
+++ b/Engine/Source/Runtime/Engine/Private/PhysicsEngine/PhysAnim.cpp
@@ -242,7 +242,7 @@ void USkeletalMeshComponent::BlendInPhysics()
 
 
 
-void USkeletalMeshComponent::UpdateKinematicBonesToPhysics(bool bTeleport, bool bNeedsSkinning)
+void USkeletalMeshComponent::UpdateKinematicBonesToPhysics(bool bTeleport, bool bNeedsSkinning, bool bForceUpdate)
 {
 	SCOPE_CYCLE_COUNTER(STAT_UpdateRBBones);
 
@@ -254,7 +254,7 @@ void USkeletalMeshComponent::UpdateKinematicBonesToPhysics(bool bTeleport, bool
 	const bool bUpdateKinematics = (KinematicBonesUpdateType != EKinematicBonesUpdateToPhysics::SkipAllBones);
 
 	// If desired, update physics bodies associated with skeletal mesh component to match.
-	if( !bUpdateKinematics )
+	if( !bUpdateKinematics && !bForceUpdate)
 	{
 		// nothing to do 
 		return;
@@ -312,7 +312,7 @@ void USkeletalMeshComponent::UpdateKinematicBonesToPhysics(bool bTeleport, bool
 				FBodyInstance* BodyInst = Bodies[i];
 				check(BodyInst);
 
-				if (BodyInst->IsValidBodyInstance() && !BodyInst->IsInstanceSimulatingPhysics(true))
+				if (bForceUpdate || (BodyInst->IsValidBodyInstance() && !BodyInst->IsInstanceSimulatingPhysics(true)))
 				{
 					// Find the graphics bone index that corresponds to this physics body.
 					FName const BodyName = PhysicsAsset->BodySetup[i]->BoneName;
diff --git a/Engine/Source/Runtime/Engine/Private/SkeletalMeshComponentPhysics.cpp b/Engine/Source/Runtime/Engine/Private/SkeletalMeshComponentPhysics.cpp
--- a/Engine/Source/Runtime/Engine/Private/SkeletalMeshComponentPhysics.cpp
+++ b/Engine/Source/Runtime/Engine/Private/SkeletalMeshComponentPhysics.cpp
@@ -1330,7 +1330,7 @@ void USkeletalMeshComponent::OnUpdateTransform(bool bSkipPhysicsMove)
 	// Always send new transform to physics
 	if(bPhysicsStateCreated && !bSkipPhysicsMove )
 	{
-		UpdateKinematicBonesToPhysics(false, false);
+		UpdateKinematicBonesToPhysics(false, false, true);
 	}
 
 #if WITH_APEX_CLOTHING
