Support for "EditInline" properties in UEngine::CopyPropertiesForUnrelatedObjects
http://udn.unrealengine.com/questions/218011/edit-inline-tarray-all-instance-data-lost-on-bluep.html

diff --git a/Engine/Source/Runtime/Engine/Private/UnrealEngine.cpp b/Engine/Source/Runtime/Engine/Private/UnrealEngine.cpp
--- a/Engine/Source/Runtime/Engine/Private/UnrealEngine.cpp
+++ b/Engine/Source/Runtime/Engine/Private/UnrealEngine.cpp
@@ -9580,6 +9580,46 @@ static TAutoConsoleVariable<int32> CVarDumpCopyPropertiesForUnrelatedObjects(
 	TEXT("Dump the objects that are cross class copied")
 	);
 
+struct FFindEditInlineSubobjectHelper
+{
+	template<typename T>
+	static void Get(const UObject* InObject, T& OutObjects)
+	{
+		const UClass* Class = InObject ? InObject->GetClass() : NULL;
+		for (UProperty* Prop = (Class ? Class->RefLink : NULL); Prop; Prop = Prop->NextRef)
+		{
+			if (Prop->HasAnyPropertyFlags(CPF_EditInline))
+			{
+				auto ObjectProperty = Cast<const UObjectProperty>(Prop);
+				if (ObjectProperty)
+				{
+					auto ObjectValue = ObjectProperty->GetObjectPropertyValue_InContainer(InObject);
+					if (ObjectValue)
+					{
+						OutObjects.Add(ObjectValue);
+					}
+					continue;
+				}
+
+				auto ArrayProperty = Cast<const UArrayProperty>(Prop);
+				auto InnerObjectProperty = ArrayProperty ? Cast<const UObjectProperty>(ArrayProperty->Inner) : NULL;
+				if (InnerObjectProperty)
+				{
+					FScriptArrayHelper_InContainer ArrayHelper(ArrayProperty, InObject);
+					for (int32 ElementIndex = 0; ElementIndex < ArrayHelper.Num(); ++ElementIndex)
+					{
+						UObject** ObjectValuePtr = (UObject**)(ArrayHelper.GetRawPtr(ElementIndex));
+						if (ObjectValuePtr && *ObjectValuePtr)
+						{
+							OutObjects.Add(*ObjectValuePtr);
+						}
+					}
+				}
+			}
+		}
+	}
+};
+
 void UEngine::CopyPropertiesForUnrelatedObjects(UObject* OldObject, UObject* NewObject, FCopyPropertiesForUnrelatedObjectsParams Params)
 {
 	check(OldObject && NewObject);
@@ -9716,6 +9756,32 @@ void UEngine::CopyPropertiesForUnrelatedObjects(UObject* OldObject, UObject* New
 		}
 	}
 
+	{
+		TSet<UObject*> OldEditInlineObjects;
+		FFindEditInlineSubobjectHelper::Get(OldObject, OldEditInlineObjects);
+
+		TArray<UObject*> NewEditInlineObjects;
+		FFindEditInlineSubobjectHelper::Get(NewObject, NewEditInlineObjects);
+		// if object exist in both sets, it's outer is old object, and its class is "editinline"
+		for (auto Obj : NewEditInlineObjects)
+		{
+			const bool bProperOuter = (Obj->GetOuter() == OldObject);
+			const bool bEditInlineNew = Obj->GetClass()->HasAllClassFlags(CLASS_EditInlineNew);
+			if (bProperOuter && bEditInlineNew)
+			{
+				const bool bKeptByOld = OldEditInlineObjects.Contains(Obj);
+				const bool bNotHandledYet = !ReferenceReplacementMap.Contains(Obj);
+				if (bKeptByOld && bNotHandledYet)
+				{
+					UObject* NewEditInlineSubobject = StaticDuplicateObject(Obj, NewObject, NULL);
+					ReferenceReplacementMap.Add(Obj, NewEditInlineSubobject);
+				}
+			}
+		}
+	}
+	
+
+
 	// Replace anything with an outer of the old object with NULL, unless it already has a replacement
 	TArray<UObject*> ObjectsInOuter;
 	GetObjectsWithOuter(OldObject, ObjectsInOuter, true);
