Support for "EditInline" properties in UEngine::CopyPropertiesForUnrelatedObjects, update

diff --git a/Engine/Source/Runtime/Engine/Private/UnrealEngine.cpp b/Engine/Source/Runtime/Engine/Private/UnrealEngine.cpp
--- a/Engine/Source/Runtime/Engine/Private/UnrealEngine.cpp
+++ b/Engine/Source/Runtime/Engine/Private/UnrealEngine.cpp
@@ -9759,22 +9759,24 @@ void UEngine::CopyPropertiesForUnrelatedObjects(UObject* OldObject, UObject* New
 	{
 		TSet<UObject*> OldEditInlineObjects;
 		FFindEditInlineSubobjectHelper::Get(OldObject, OldEditInlineObjects);
-
-		TArray<UObject*> NewEditInlineObjects;
-		FFindEditInlineSubobjectHelper::Get(NewObject, NewEditInlineObjects);
-		// if object exist in both sets, it's outer is old object, and its class is "editinline"
-		for (auto Obj : NewEditInlineObjects)
+		if (OldEditInlineObjects.Num())
 		{
-			const bool bProperOuter = (Obj->GetOuter() == OldObject);
-			const bool bEditInlineNew = Obj->GetClass()->HasAllClassFlags(CLASS_EditInlineNew);
-			if (bProperOuter && bEditInlineNew)
+			TArray<UObject*> NewEditInlineObjects;
+			FFindEditInlineSubobjectHelper::Get(NewObject, NewEditInlineObjects);
+			// if object exist in both sets, it's outer is old object, and its class is "editinline"
+			for (auto Obj : NewEditInlineObjects)
 			{
-				const bool bKeptByOld = OldEditInlineObjects.Contains(Obj);
-				const bool bNotHandledYet = !ReferenceReplacementMap.Contains(Obj);
-				if (bKeptByOld && bNotHandledYet)
+				const bool bProperOuter = (Obj->GetOuter() == OldObject);
+				const bool bEditInlineNew = Obj->GetClass()->HasAllClassFlags(CLASS_EditInlineNew);
+				if (bProperOuter && bEditInlineNew)
 				{
-					UObject* NewEditInlineSubobject = StaticDuplicateObject(Obj, NewObject, NULL);
-					ReferenceReplacementMap.Add(Obj, NewEditInlineSubobject);
+					const bool bKeptByOld = OldEditInlineObjects.Contains(Obj);
+					const bool bNotHandledYet = !ReferenceReplacementMap.Contains(Obj);
+					if (bKeptByOld && bNotHandledYet)
+					{
+						UObject* NewEditInlineSubobject = StaticDuplicateObject(Obj, NewObject, NULL);
+						ReferenceReplacementMap.Add(Obj, NewEditInlineSubobject);
+					}
 				}
 			}
 		}
