Fixed SVirtualJoystick shadowing SWidget SetVisibility

diff --git a/Engine/Source/Runtime/Engine/Private/PlayerController.cpp b/Engine/Source/Runtime/Engine/Private/PlayerController.cpp
--- a/Engine/Source/Runtime/Engine/Private/PlayerController.cpp
+++ b/Engine/Source/Runtime/Engine/Private/PlayerController.cpp
@@ -2886,7 +2886,7 @@ void APlayerController::ClientSetCinematicMode_Implementation(bool bInCinematicM
 		ULocalPlayer* LocPlayer = Cast<ULocalPlayer>(Player);
 		if (VirtualJoystick.IsValid())
 		{
-			VirtualJoystick->SetVisibility(MyHUD->bShowHUD, true);
+			VirtualJoystick->SetJoystickVisibility(MyHUD->bShowHUD, true);
 		}
 	}
 
@@ -4239,7 +4239,7 @@ void APlayerController::SetVirtualJoystickVisibility(bool bVisible)
 {
 	if (VirtualJoystick.IsValid())
 	{
-		VirtualJoystick->SetVisibility(bVisible, false);
+		VirtualJoystick->SetJoystickVisibility(bVisible, false);
 	}
 }
 
diff --git a/Engine/Source/Runtime/Slate/Private/Widgets/Input/SVirtualJoystick.cpp b/Engine/Source/Runtime/Slate/Private/Widgets/Input/SVirtualJoystick.cpp
--- a/Engine/Source/Runtime/Slate/Private/Widgets/Input/SVirtualJoystick.cpp
+++ b/Engine/Source/Runtime/Slate/Private/Widgets/Input/SVirtualJoystick.cpp
@@ -444,7 +444,7 @@ void SVirtualJoystick::Tick( const FGeometry& AllottedGeometry, const double InC
 	}
 }
 
-void SVirtualJoystick::SetVisibility(const bool bInVisible, const bool bInFade)
+void SVirtualJoystick::SetJoystickVisibility(const bool bInVisible, const bool bInFade)
 {
 	// if we aren't fading, then just set the current opacity to desired
 	if (!bInFade)
diff --git a/Engine/Source/Runtime/Slate/Public/Widgets/Input/SVirtualJoystick.h b/Engine/Source/Runtime/Slate/Public/Widgets/Input/SVirtualJoystick.h
--- a/Engine/Source/Runtime/Slate/Public/Widgets/Input/SVirtualJoystick.h
+++ b/Engine/Source/Runtime/Slate/Public/Widgets/Input/SVirtualJoystick.h
@@ -100,7 +100,7 @@ public:
 	/** 
 	 * Shows or hides the controls (for instance during cinematics
 	 */
-	void SetVisibility(const bool bVisible, const bool bFade);
+	void SetJoystickVisibility(const bool bVisible, const bool bFade);
 
 	void AddControl(const FControlInfo& Control)
 	{
