Only cache last location for an active sound when updating for actual playing update (as opposed to display for debug commands)

diff --git a/Engine/Source/Runtime/Engine/Private/ActiveSound.cpp b/Engine/Source/Runtime/Engine/Private/ActiveSound.cpp
--- a/Engine/Source/Runtime/Engine/Private/ActiveSound.cpp
+++ b/Engine/Source/Runtime/Engine/Private/ActiveSound.cpp
@@ -190,8 +190,11 @@ void FActiveSound::UpdateWaveInstances( FAudioDevice* AudioDevice, TArray<FWaveI
 	}
 
 	// for velocity-based effects like doppler
-	ParseParams.Velocity = (DeltaTime <= 0.0f) ? FVector::ZeroVector : ((ParseParams.Transform.GetTranslation() - LastLocation) / DeltaTime);
-	LastLocation = ParseParams.Transform.GetTranslation();
+	if (DeltaTime > 0.f)
+	{
+		ParseParams.Velocity = (ParseParams.Transform.GetTranslation() - LastLocation) / DeltaTime;
+		LastLocation = ParseParams.Transform.GetTranslation();
+	}
 
 	// if the closest listener is not the primary one, transform CurrentLocation
 	if( ClosestListenerIndex != 0 )
