* game console history works now like auto completion (displays available values), up to 50 elements are stored, no elements are removed if there is a duplicated, serialization for editor is not done yet.

diff --git a/Engine/Source/Runtime/Engine/Classes/Engine/Console.h b/Engine/Source/Runtime/Engine/Classes/Engine/Console.h
--- a/Engine/Source/Runtime/Engine/Classes/Engine/Console.h
+++ b/Engine/Source/Runtime/Engine/Classes/Engine/Console.h
@@ -74,24 +74,12 @@ class ENGINE_API UConsole
 
 	int32 SBPos;
 
-	/** index into the History array for the latest command that was entered */
-	UPROPERTY(config)
-	int32 HistoryTop;
-
-	/** index into the History array for the earliest command that was entered */
-	UPROPERTY(config)
-	int32 HistoryBot;
-
-	/** the index of the current position in the History array */
-	UPROPERTY(config)
-	int32 HistoryCur;
+	/** Max number of command history entries */
+	static const int32 MAX_HISTORY_ENTRIES = 50;
 
-	/** tracks previously entered console commands */
+	/** Holds the history buffer, order is old to new */
 	UPROPERTY(config)
-	FString History[16];    /*MAX_HISTORY_ENTRIES @fixmeconst*/
-
-	/** tracks whether the user is using arrows keys to navigate the history, so that auto-complete doesn't override */
-	uint32 bNavigatingHistory:1;
+	TArray<FString> HistoryBuffer;
 
 	/** The command the user is currently typing. */
 	FString TypedStr;
@@ -122,15 +110,13 @@ class ENGINE_API UConsole
 	/** Do we need to rebuild auto complete? */
 	uint32 bIsRuntimeAutoCompleteUpToDate:1;
 
+	// NAME_Typing, NAME_Open or NAME_None
 	FName ConsoleState;
 
 	FAutoCompleteNode AutoCompleteTree;
 
 	/** Current list of matching commands for auto-complete, @see UpdateCompleteIndices() */
-	TArray<int32> AutoCompleteIndices;
-
-	/** Max number of command history entries */
-	static const int32 MAX_HISTORY_ENTRIES = 16;
+	TArray<struct FAutoCompleteCommand> AutoComplete;
 
 	~UConsole();
 
@@ -262,9 +248,9 @@ private:
 	
 	void PostRender_InputLine(class UCanvas* Canvas, FIntPoint UserInputLinePos);
 
-	/**
-	* Searches console command history and removes any entries matching the specified command.
-	* @param Command - The command to search for and purge from the history.
-	*/
-	virtual void PurgeCommandFromHistory(const FString& Command);
+	void SetAutoCompleteFromHistory();
+
+	void SetInputLineFromAutoComplete();
+
+	void NormalizeHistoryBuffer();
 };
diff --git a/Engine/Source/Runtime/Engine/Private/UserInterface/Console.cpp b/Engine/Source/Runtime/Engine/Private/UserInterface/Console.cpp
--- a/Engine/Source/Runtime/Engine/Private/UserInterface/Console.cpp
+++ b/Engine/Source/Runtime/Engine/Private/UserInterface/Console.cpp
@@ -44,7 +44,8 @@ public:
 
 UConsole::UConsole(const class FPostConstructInitializeProperties& PCIP)
 	: Super(PCIP)
-	, AutoCompleteCursor(-1)
+//	, HistoryBufferCur(0)
+//	, AutoCompleteCursor(-1)
 {
 	// Structure to hold one-time initialization
 	struct FConstructorStatics
@@ -256,10 +257,9 @@ void UConsole::UpdateCompleteIndices()
 	{
 		BuildRuntimeAutoCompleteList(true);
 	}
-	bNavigatingHistory = false;
 	AutoCompleteIndex = 0;
 	AutoCompleteCursor = -1;
-	AutoCompleteIndices.Empty();
+	AutoComplete.Empty();
 	FAutoCompleteNode *Node = &AutoCompleteTree;
 	FString LowerTypedStr = TypedStr.ToLower();
 	for (int32 Idx = 0; Idx < TypedStr.Len(); Idx++)
@@ -288,7 +288,7 @@ void UConsole::UpdateCompleteIndices()
 			{
 				if(Idx < TypedStr.Len())
 				{
-					// thre is more text behind the auto completed text, we don't need auto completion
+					// there is more text behind the auto completed text, we don't need auto completion
 					return;
 				}
 				else
@@ -300,8 +300,29 @@ void UConsole::UpdateCompleteIndices()
 	}
 	if (Node != &AutoCompleteTree)
 	{
-		AutoCompleteIndices = Node->AutoCompleteListIndices;
-		bNavigatingHistory = false;
+		const TArray<int32>& Leaf = Node->AutoCompleteListIndices;
+
+		for(uint32 i = 0, Num = (uint32)Leaf.Num(); i < Num; ++i)
+		{
+			AutoComplete.Add(AutoCompleteList[Leaf[i]]);
+		}
+	}
+}
+
+void UConsole::SetAutoCompleteFromHistory()
+{
+	AutoCompleteIndex = 0;
+	AutoCompleteCursor = 0;
+	AutoComplete.Empty();
+
+	for(int32 i = HistoryBuffer.Num() - 1; i >= 0 ; --i)
+	{
+		FAutoCompleteCommand Cmd;
+
+		Cmd.Command = HistoryBuffer[i]; 
+		Cmd.Desc = FString::Printf(TEXT("> %s"), *Cmd.Command);
+
+		AutoComplete.Add(Cmd);
 	}
 }
 
@@ -315,49 +336,14 @@ void UConsole::SetCursorPos( int32 Position )
 	TypedStrPos = Position;
 }
 
-
-void UConsole::PurgeCommandFromHistory(const FString& Command)
-{
-	if ( (HistoryTop >= 0) && (HistoryTop < MAX_HISTORY_ENTRIES) )
-	{
-		for (int32 HistoryIdx=0; HistoryIdx<MAX_HISTORY_ENTRIES; ++HistoryIdx)
-		{
-			if (FCString::Stricmp(*History[HistoryIdx], *Command) == 0)
-			{
-				// from here to the top idx, shift everything back one
-				int32 Idx = HistoryIdx;
-				int32 NextIdx = (HistoryIdx + 1) % MAX_HISTORY_ENTRIES;
-
-				while (Idx != HistoryTop)
-				{
-					History[Idx] = History[NextIdx];
-					Idx = NextIdx;
-					NextIdx = (NextIdx + 1) % MAX_HISTORY_ENTRIES;
-				}
-
-				// top index backs up one as well
-				HistoryTop = (HistoryTop == 0) ? (MAX_HISTORY_ENTRIES - 1) : (HistoryTop - 1);
-			}
-		}
-	}
-}
-
-
 void UConsole::ConsoleCommand(const FString& Command)
 {
-	// Store the command in the console history.
-	if ((HistoryTop == 0) ? (FCString::Stricmp(*History[MAX_HISTORY_ENTRIES - 1], *Command) != 0) : (FCString::Stricmp(*History[HistoryTop - 1], *Command) != 0))
+	// insert into history buffer
 	{
-		// ensure uniqueness
-		PurgeCommandFromHistory(Command);
+		HistoryBuffer.Add(Command);
 
-		History[HistoryTop] = Command;
-		HistoryTop = (HistoryTop+1) % MAX_HISTORY_ENTRIES;
-
-		if ( ( HistoryBot == -1) || ( HistoryBot == HistoryTop ) )
-			HistoryBot = (HistoryBot+1) % MAX_HISTORY_ENTRIES;
+		NormalizeHistoryBuffer();
 	}
-	HistoryCur = HistoryTop;
 
 	// Save the command history to the INI.
 	SaveConfig();
@@ -602,11 +588,10 @@ bool UConsole::InputKey_InputLine( int32 ControllerId, FKey Key, EInputEvent Eve
 			{
 				SetInputText("");
 				SetCursorPos(0);
-				HistoryCur = HistoryTop;
 
 				AutoCompleteIndex = 0;
 				AutoCompleteCursor = -1;
-				AutoCompleteIndices.Empty();
+				AutoComplete.Empty();
 				bAutoCompleteLocked = false;
 
 				return true;
@@ -620,15 +605,11 @@ bool UConsole::InputKey_InputLine( int32 ControllerId, FKey Key, EInputEvent Eve
 		}
 		else if (Key == EKeys::Tab && Event == IE_Pressed)
 		{
-			if (AutoCompleteIndices.Num() > 0 && !bAutoCompleteLocked)
+			if (!bAutoCompleteLocked)
 			{
-				const int32 Index = AutoCompleteIndex + (AutoCompleteCursor >= 0 ? AutoCompleteCursor : 0);
-				const FAutoCompleteCommand& Cmd = AutoCompleteList[AutoCompleteIndices[Index]];
-
-				TypedStr = Cmd.Command;
-				SetCursorPos(TypedStr.Len());
-				bAutoCompleteLocked = true;
+				SetInputLineFromAutoComplete();
 			}
+
 			return true;
 		}
 		else if( Key == EKeys::Enter && Event == IE_Released )
@@ -674,52 +655,39 @@ bool UConsole::InputKey_InputLine( int32 ControllerId, FKey Key, EInputEvent Eve
 		}
 		else if( Key == EKeys::Up )
 		{
-			if (!bNavigatingHistory && !bCtrl && AutoCompleteIndices.Num())
+			if (!bCtrl)
 			{
-				if(AutoCompleteCursor + 1 < FMath::Min((int32)MAX_AUTOCOMPLETION_LINES, AutoCompleteIndices.Num()))
-				{
-					// move cursor within displayed region
-					++AutoCompleteCursor;
-				}
-				else 
+				if(AutoComplete.Num())
 				{
-					// can be negative
-					int32 ScrollRegionSize = AutoCompleteIndices.Num() - (int32)MAX_AUTOCOMPLETION_LINES;
-
-					// can we scroll?
-					if(AutoCompleteIndex + 1 < ScrollRegionSize)
+					if(AutoCompleteCursor + 1 < FMath::Min((int32)MAX_AUTOCOMPLETION_LINES, AutoComplete.Num()))
 					{
-						++AutoCompleteIndex;
+						// move cursor within displayed region
+						++AutoCompleteCursor;
 					}
-				}
-			}
-			else
-				if ( HistoryBot >= 0 )
-				{
-					if (HistoryCur == HistoryBot)
-						HistoryCur = HistoryTop;
-					else
+					else 
 					{
-						HistoryCur--;
-						if (HistoryCur<0)
-							HistoryCur = MAX_HISTORY_ENTRIES-1;
+						// can be negative
+						int32 ScrollRegionSize = AutoComplete.Num() - (int32)MAX_AUTOCOMPLETION_LINES;
+
+						// can we scroll?
+						if(AutoCompleteIndex + 1 < ScrollRegionSize)
+						{
+							++AutoCompleteIndex;
+						}
 					}
-
-					SetInputText(History[HistoryCur]);
-					SetCursorPos(History[HistoryCur].Len());
-					//testUpdateCompleteIndices();
-					AutoCompleteIndex = 0;
-					AutoCompleteCursor = -1;
-					AutoCompleteIndices.Empty();
-
-
-					bNavigatingHistory = true;
+					SetInputLineFromAutoComplete();
 				}
-				return true;
+				else
+				{
+					SetAutoCompleteFromHistory();
+					SetInputLineFromAutoComplete();
+				}
+			}
+			return true;
 		}
 		else if( Key == EKeys::Down )
 		{
-			if (!bNavigatingHistory && !bCtrl && AutoCompleteIndices.Num() > 1)
+			if (!bCtrl && AutoComplete.Num() > 1)
 			{
 				if(AutoCompleteCursor > 0)
 				{
@@ -735,25 +703,8 @@ bool UConsole::InputKey_InputLine( int32 ControllerId, FKey Key, EInputEvent Eve
 					}
 					bAutoCompleteLocked = false;
 				}
+				SetInputLineFromAutoComplete();
 			}
-			else
-				if ( HistoryBot >= 0 )
-				{
-					if (HistoryCur == HistoryTop)
-						HistoryCur = HistoryBot;
-					else
-						HistoryCur = (HistoryCur+1) % MAX_HISTORY_ENTRIES;
-
-					SetInputText(History[HistoryCur]);
-					SetCursorPos(History[HistoryCur].Len());
-					//testUpdateCompleteIndices();
-					AutoCompleteIndex = 0;
-					AutoCompleteCursor = -1;
-					AutoCompleteIndices.Empty();
-
-					bNavigatingHistory = true;
-				}
-
 		}
 		else if( Key==EKeys::BackSpace )
 		{
@@ -824,9 +775,9 @@ bool UConsole::InputKey_InputLine( int32 ControllerId, FKey Key, EInputEvent Eve
 		else
 		{
 			// auto complete is open and the user adds a key
-			if (AutoCompleteIndices.Num() > 0 && !bAutoCompleteLocked && AutoCompleteCursor >= 0)
+			if (AutoComplete.Num() > 0 && !bAutoCompleteLocked && AutoCompleteCursor >= 0)
 			{
-				const FAutoCompleteCommand& Cmd = AutoCompleteList[AutoCompleteIndices[AutoCompleteIndex + AutoCompleteCursor]];
+				const FAutoCompleteCommand& Cmd = AutoComplete[AutoCompleteIndex + AutoCompleteCursor];
 
 				TypedStr = Cmd.Command;
 				SetCursorPos(TypedStr.Len());
@@ -855,6 +806,31 @@ namespace ConsoleDefs
 	static const FString LeadingInputText( TEXT( " > " ) );
 }
 
+void UConsole::SetInputLineFromAutoComplete()
+{
+	if (AutoComplete.Num() > 0)
+	{
+		const int32 Index = AutoCompleteIndex + (AutoCompleteCursor >= 0 ? AutoCompleteCursor : 0);
+		const FAutoCompleteCommand& Cmd = AutoComplete[Index];
+
+		TypedStr = Cmd.Command;
+		SetCursorPos(TypedStr.Len());
+		bAutoCompleteLocked = true;
+	}
+}
+
+void UConsole::NormalizeHistoryBuffer()
+{
+	const uint32 Count = MAX_HISTORY_ENTRIES;
+
+	check(Count > 0);
+
+	if ((uint32)HistoryBuffer.Num() > Count)
+	{
+		uint32 ShrinkCount = (uint32)(HistoryBuffer.Num() - Count);
+		HistoryBuffer.RemoveAt(0, ShrinkCount);
+	}
+}
 
 void UConsole::PostRender_Console_Typing(UCanvas* Canvas)
 {
@@ -896,7 +872,7 @@ void UConsole::BeginState_Typing(FName PreviousStateName)
 		FlushPlayerInput();
 	}
 	bCaptureKeyInput = true;
-	HistoryCur = HistoryTop;
+//	HistoryCur = HistoryTop;
 }
 
 void UConsole::EndState_Typing( FName NextStateName )
@@ -1034,7 +1010,7 @@ void UConsole::PostRender_Console_Open(UCanvas* Canvas)
 void UConsole::BeginState_Open(FName PreviousStateName)
 {
 	bCaptureKeyInput = true;
-	HistoryCur = HistoryTop;
+//	HistoryCur = HistoryTop;
 
 	SBPos = 0;
 	bCtrl = false;
@@ -1135,21 +1111,16 @@ void UConsole::PostRender_InputLine(UCanvas* Canvas, FIntPoint UserInputLinePos)
 	Canvas->DrawItem( ConsoleText );
 
 	// draw the remaining text for matching auto-complete
-	if (AutoCompleteIndices.Num() > 0)
+	if (AutoComplete.Num() > 0)
 	{
-		int32 Idx = AutoCompleteIndices[AutoCompleteIndex];
-		ConsoleText.SetColor( ConsoleDefs::AutocompletePartialSuggestionColor );
-		ConsoleText.Text = FText::FromString( AutoCompleteList[Idx].Command.Right(AutoCompleteList[Idx].Command.Len() - TypedStr.Len()) );
-		Canvas->DrawItem( ConsoleText, UserInputLinePos.X+xl, UserInputLinePos.Y-3.0f-yl );
-		Canvas->StrLen(Font, *ConsoleDefs::LeadingInputText, xl, yl );
-
 		int32 StartIdx = AutoCompleteIndex;
-
 		if (StartIdx < 0)
 		{
-			StartIdx = FMath::Max(0,AutoCompleteIndices.Num() + StartIdx);
+			StartIdx = FMath::Max(0, AutoComplete.Num() + StartIdx);
 		}
-		Idx = StartIdx;
+
+		Canvas->StrLen(Font, *ConsoleDefs::LeadingInputText, xl, yl);
+		
 		float y = UserInputLinePos.Y - 6.0f - (yl * 2.0f);
 
 		// Set the background color/texture of the auto-complete section
@@ -1161,9 +1132,9 @@ void UConsole::PostRender_InputLine(UCanvas* Canvas, FIntPoint UserInputLinePos)
 		// to avoid memory many allocations
 		AutoCompleteElements.Empty(MAX_AUTOCOMPLETION_LINES + 1);
 
-		for (int32 MatchIdx = 0; MatchIdx < MAX_AUTOCOMPLETION_LINES && MatchIdx < AutoCompleteIndices.Num(); MatchIdx++)
+		for (int32 MatchIdx = 0; MatchIdx < MAX_AUTOCOMPLETION_LINES && MatchIdx < AutoComplete.Num(); MatchIdx++)
 		{
-			const FAutoCompleteCommand &Cmd = AutoCompleteList[AutoCompleteIndices[Idx + MatchIdx]];
+			const FAutoCompleteCommand &Cmd = AutoComplete[StartIdx + MatchIdx];
 			OutStr = Cmd.Desc;
 
 			if(OutStr.IsEmpty())
@@ -1176,9 +1147,9 @@ void UConsole::PostRender_InputLine(UCanvas* Canvas, FIntPoint UserInputLinePos)
 		}
 
 		// Display a message if there were more matches
-		if (AutoCompleteIndices.Num() >= MAX_AUTOCOMPLETION_LINES)
+		if (AutoComplete.Num() >= MAX_AUTOCOMPLETION_LINES)
 		{
-			OutStr = FString::Printf(TEXT("[%i more matches]"), (AutoCompleteIndices.Num() - MAX_AUTOCOMPLETION_LINES + 1));
+			OutStr = FString::Printf(TEXT("[%i more matches]"), (AutoComplete.Num() - MAX_AUTOCOMPLETION_LINES + 1));
 			AutoCompleteElements.Add(OutStr);
 		}
 
