* improved error/log printout when ShouldCache() is wrong (or a shader is using a permutation that should not be used)

diff --git a/Engine/Source/Runtime/Engine/Private/Materials/MaterialShared.cpp b/Engine/Source/Runtime/Engine/Private/Materials/MaterialShared.cpp
--- a/Engine/Source/Runtime/Engine/Private/Materials/MaterialShared.cpp
+++ b/Engine/Source/Runtime/Engine/Private/Materials/MaterialShared.cpp
@@ -1404,7 +1404,7 @@ FShader* FMaterial::GetShader(FMeshMaterialShaderType* ShaderType, FVertexFactor
 			TEXT("Couldn't find Shader %s for Material Resource %s!\n")
 			TEXT("		With VF=%s, Platform=%s \n")
 			TEXT("		ShouldCache: Mat=%u, VF=%u, Shader=%u \n")
-			TEXT("		Material Usage = %s"),
+			TEXT("		MaterialUsageDesc: %s"),
 			ShaderType->GetName(), 
 			*GetFriendlyName(),
 			VertexFactoryType->GetName(),
@@ -1719,44 +1719,32 @@ int32 FMaterialResource::GetSamplerUsage() const
 	return -1;
 }
 
-/** Returns a string that describes the material's usage for debugging purposes. */
 FString FMaterialResource::GetMaterialUsageDescription() const
 {
 	check(Material);
-	FString BaseDescription = GetShadingModelString(GetShadingModel()) + TEXT(", ") + GetBlendModeString(GetBlendMode());
+	FString BaseDescription = FString::Printf(
+		TEXT("LightingModel=%s, BlendMode=%s, "),
+		*GetShadingModelString(GetShadingModel()), *GetBlendModeString(GetBlendMode()));
 
-	if (IsSpecialEngineMaterial())
-	{
-		BaseDescription += TEXT(", SpecialEngine");
-	}
-	if (IsTwoSided())
-	{
-		BaseDescription += TEXT(", TwoSided");
-	}
-	if (IsTangentSpaceNormal())
-	{
-		BaseDescription += TEXT(", TSNormal");
-	}
-	if (ShouldInjectEmissiveIntoLPV())
-	{
-		BaseDescription += TEXT(", InjectEmissiveIntoLPV");
-	}
-	if (IsMasked())
-	{
-		BaseDescription += TEXT(", Masked");
-	}
-	if (IsDistorted())
-	{
-		BaseDescription += TEXT(", Distorted");
-	}
+	// this changed from ",SpecialEngine, TwoSided" to ",SpecialEngine=1, TwoSided=1, TSNormal=0, ..." to be more readable
+	BaseDescription += FString::Printf(
+		TEXT("SpecialEngine=%d, TwoSided=%d, TSNormal=%d, InjectEmissiveIntoLPV=%d, Masked=%d, Distorted=%d, Usage={"),
+		(int32)IsSpecialEngineMaterial(), (int32)IsTwoSided(), (int32)IsTangentSpaceNormal(), (int32)ShouldInjectEmissiveIntoLPV(), (int32)IsMasked(), (int32)IsDistorted());
 
+	bool bFirst = true;
 	for (int32 MaterialUsageIndex = 0; MaterialUsageIndex < MATUSAGE_MAX; MaterialUsageIndex++)
 	{
 		if (Material->GetUsageByFlag((EMaterialUsage)MaterialUsageIndex))
 		{
-			BaseDescription += FString(TEXT(", ")) + Material->GetUsageName((EMaterialUsage)MaterialUsageIndex);
+			if (!bFirst)
+			{
+				BaseDescription += FString(TEXT(","));
+			}
+			BaseDescription += Material->GetUsageName((EMaterialUsage)MaterialUsageIndex);
+			bFirst = false;
 		}
 	}
+	BaseDescription += FString(TEXT("}"));
 
 	return BaseDescription;
 }
