#ue4 - AActor::PostEditChangeProperty() now uses static FNames that use GET_MEMBER_NAME_CHECKED, for speed and safety.

diff --git a/Engine/Source/Runtime/Engine/Private/ActorEditor.cpp b/Engine/Source/Runtime/Engine/Private/ActorEditor.cpp
--- a/Engine/Source/Runtime/Engine/Private/ActorEditor.cpp
+++ b/Engine/Source/Runtime/Engine/Private/ActorEditor.cpp
@@ -23,6 +23,10 @@ void AActor::PreEditChange(UProperty* PropertyThatWillChange)
 	UnregisterAllComponents();
 }
 
+static FName Name_RelativeLocation = GET_MEMBER_NAME_CHECKED(USceneComponent, RelativeLocation);
+static FName Name_RelativeRotation = GET_MEMBER_NAME_CHECKED(USceneComponent, RelativeRotation);
+static FName Name_RelativeScale3D = GET_MEMBER_NAME_CHECKED(USceneComponent, RelativeScale3D);
+
 void AActor::PostEditChangeProperty(FPropertyChangedEvent& PropertyChangedEvent)
 {
 	UProperty* PropertyThatChanged = PropertyChangedEvent.Property;
@@ -30,7 +34,7 @@ void AActor::PostEditChangeProperty(FPropertyChangedEvent& PropertyChangedEvent)
 	
 	bool bTransformationChanged = false;
 
-	if ( PropertyName==FName(TEXT("RelativeLocation")) || PropertyName==FName(TEXT("RelativeRotation")) || PropertyName==FName(TEXT("RelativeScale3D")) )
+	if ( PropertyName==Name_RelativeLocation || PropertyName==Name_RelativeRotation || PropertyName==Name_RelativeScale3D )
 	{
 		bTransformationChanged = true;
 
