#ue4 - PlayerCameraManager traces should check for any blocking hit, not just Hit.GetActor(), so it handles BSP.

diff --git a/Engine/Source/Runtime/Engine/Private/PlayerCameraManager.cpp b/Engine/Source/Runtime/Engine/Private/PlayerCameraManager.cpp
--- a/Engine/Source/Runtime/Engine/Private/PlayerCameraManager.cpp
+++ b/Engine/Source/Runtime/Engine/Private/PlayerCameraManager.cpp
@@ -536,7 +536,7 @@ void APlayerCameraManager::UpdateViewTarget(FTViewTarget& OutVT, float DeltaTime
 			FHitResult Result;
 
 			GetWorld()->SweepSingle(Result, Loc, Pos, FQuat::Identity, ECC_Camera, FCollisionShape::MakeBox(FVector(12.f)), BoxParams);
-			OutVT.POV.Location = (Result.GetActor() == NULL) ? Pos : Result.Location;
+			OutVT.POV.Location = !Result.bBlockingHit ? Pos : Result.Location;
 			OutVT.POV.Rotation = Rotator;
 
 			// don't apply modifiers when using this debug camera mode
