#ue4 - Fix 'WalkableFloorAngle' changes not propagating resultant 'WalkableFloorZ' changes from blueprint to instances placed in a level.

TTP 344444
udn: https://udn.unrealengine.com/questions/218018/bug-walkable-floor-z-doesnt-get-updated-for-instan.html

diff --git a/Engine/Source/Runtime/Engine/Classes/GameFramework/CharacterMovementComponent.h b/Engine/Source/Runtime/Engine/Classes/GameFramework/CharacterMovementComponent.h
--- a/Engine/Source/Runtime/Engine/Classes/GameFramework/CharacterMovementComponent.h
+++ b/Engine/Source/Runtime/Engine/Classes/GameFramework/CharacterMovementComponent.h
@@ -661,6 +661,7 @@ public:
 
 	//Begin UActorComponent Interface
 	virtual void TickComponent(float DeltaTime, enum ELevelTick TickType, FActorComponentTickFunction *ThisTickFunction) override;
+	virtual void OnRegister() override;
 	virtual void OnUnregister() override;
 	virtual void PostLoad() override;
 	virtual void RegisterComponentTickFunctions(bool bRegister) override;
diff --git a/Engine/Source/Runtime/Engine/Private/CharacterMovementComponent.cpp b/Engine/Source/Runtime/Engine/Private/CharacterMovementComponent.cpp
--- a/Engine/Source/Runtime/Engine/Private/CharacterMovementComponent.cpp
+++ b/Engine/Source/Runtime/Engine/Private/CharacterMovementComponent.cpp
@@ -303,6 +303,18 @@ void UCharacterMovementComponent::PostEditChangeProperty(FPropertyChangedEvent&
 #endif // WITH_EDITOR
 
 
+void UCharacterMovementComponent::OnRegister()
+{
+	Super::OnRegister();
+
+#if WITH_EDITOR
+	// Compute WalkableFloorZ from the WalkableFloorAngle.
+	// This is only to respond to changes propagated by PostEditChangeProperty, so it's only done in the editor.
+	SetWalkableFloorAngle(WalkableFloorAngle);
+#endif
+}
+
+
 void UCharacterMovementComponent::OnUnregister()
 {
 	delete ClientPredictionData;
