* PostProcessVolume polish (moved setting in and out of Advanced, grey items depending on Enabled/Unbound)

diff --git a/Engine/Source/Runtime/Engine/Classes/Engine/PostProcessVolume.h b/Engine/Source/Runtime/Engine/Classes/Engine/PostProcessVolume.h
--- a/Engine/Source/Runtime/Engine/Classes/Engine/PostProcessVolume.h
+++ b/Engine/Source/Runtime/Engine/Classes/Engine/PostProcessVolume.h
@@ -73,6 +73,7 @@ public:
 	// Begin UObject interface
 #if WITH_EDITOR
 	virtual void PostEditChangeProperty(FPropertyChangedEvent& PropertyChangedEvent) override;
+	virtual bool CanEditChange(const UProperty* InProperty) const override;
 #endif // WITH_EDITOR
 	// End UObject interface
 };
diff --git a/Engine/Source/Runtime/Engine/Classes/Engine/Scene.h b/Engine/Source/Runtime/Engine/Classes/Engine/Scene.h
--- a/Engine/Source/Runtime/Engine/Classes/Engine/Scene.h
+++ b/Engine/Source/Runtime/Engine/Classes/Engine/Scene.h
@@ -340,13 +340,13 @@ struct FPostProcessSettings
 	// -----------------------------------------------------------------------
 
 
-	UPROPERTY(interp, Category=Film, AdvancedDisplay, meta=(editcondition = "bOverride_FilmWhitePoint", DisplayName = "Tint"))
+	UPROPERTY(interp, Category=Film, meta=(editcondition = "bOverride_FilmWhitePoint", DisplayName = "Tint"))
 	FLinearColor FilmWhitePoint;
-	UPROPERTY(interp, Category=Film, meta=(editcondition = "bOverride_FilmShadowTint", DisplayName = "Tint Shadow"))
+	UPROPERTY(interp, Category=Film, AdvancedDisplay, meta=(editcondition = "bOverride_FilmShadowTint", DisplayName = "Tint Shadow"))
 	FLinearColor FilmShadowTint;
-	UPROPERTY(interp, Category=Film, meta=(UIMin = "0.0", UIMax = "1.0", editcondition = "bOverride_FilmShadowTintBlend", DisplayName = "Tint Shadow Blend"))
+	UPROPERTY(interp, Category=Film, AdvancedDisplay, meta=(UIMin = "0.0", UIMax = "1.0", editcondition = "bOverride_FilmShadowTintBlend", DisplayName = "Tint Shadow Blend"))
 	float FilmShadowTintBlend;
-	UPROPERTY(interp, Category=Film, meta=(UIMin = "0.0", UIMax = "1.0", editcondition = "bOverride_FilmShadowTintAmount", DisplayName = "Tint Shadow Amount"))
+	UPROPERTY(interp, Category=Film, AdvancedDisplay, meta=(UIMin = "0.0", UIMax = "1.0", editcondition = "bOverride_FilmShadowTintAmount", DisplayName = "Tint Shadow Amount"))
 	float FilmShadowTintAmount;
 
 	UPROPERTY(interp, Category=Film, meta=(UIMin = "0.0", UIMax = "2.0", editcondition = "bOverride_FilmSaturation", DisplayName = "Saturation"))
@@ -463,7 +463,7 @@ struct FPostProcessSettings
 	float LPVIntensity;
 
 	/** CURRENTLY DISABLED - The strength of the warp offset for reducing light bleeding. 0.0 is off, 1.0 is the "normal" value, but higher values can be used to boost the effect*/
-	UPROPERTY(interp, Category=LightPropagationVolume, meta=(editcondition = "bOverride_LPVWarpIntensity", UIMin = "0", UIMax = "1", DisplayName = "(DISABLED) Grid Warp Intensity") )
+	UPROPERTY(interp, Category=LightPropagationVolume, AdvancedDisplay, meta=(editcondition = "bOverride_LPVWarpIntensity", UIMin = "0", UIMax = "1", DisplayName = "(DISABLED) Grid Warp Intensity"))
 	float LPVWarpIntensity;
 
 	/** Bias applied to light injected into the LPV in cell units. Increase to reduce bleeding through thin walls*/
@@ -483,7 +483,7 @@ struct FPostProcessSettings
 	float LPVSecondaryBounceIntensity;
 
 	/** Bias applied to the geometry volume in cell units. Increase to reduce darkening due to secondary occlusion */
-	UPROPERTY(interp, Category=LightPropagationVolume, meta=(editcondition = "bOverride_LPVGeometryVolumeBias", UIMin = "0", UIMax = "2", DisplayName = "Geometry Volume Bias") )
+	UPROPERTY(interp, Category=LightPropagationVolume, AdvancedDisplay, meta=(editcondition = "bOverride_LPVGeometryVolumeBias", UIMin = "0", UIMax = "2", DisplayName = "Geometry Volume Bias"))
 	float LPVGeometryVolumeBias;
 
 	UPROPERTY(interp, Category=LightPropagationVolume, meta=(editcondition = "bOverride_LPVEmissiveInjectionIntensity", UIMin = "0", UIMax = "20", DisplayName = "Emissive Injection Intensity") )
@@ -598,11 +598,11 @@ struct FPostProcessSettings
 	float VignetteIntensity;
 
 	/** Vignette color. */
-	UPROPERTY(interp, Category=SceneColor, meta=(editcondition = "bOverride_VignetteColor", DisplayName = "Vignette Color"))
+	UPROPERTY(interp, Category=SceneColor, AdvancedDisplay, meta=(editcondition = "bOverride_VignetteColor", DisplayName = "Vignette Color"))
 	FLinearColor VignetteColor;
 
 	/** 0..1 grain jitter */
-	UPROPERTY(interp, Category=SceneColor, meta=(UIMin = "0.0", UIMax = "1.0", editcondition = "bOverride_GrainJitter"))
+	UPROPERTY(interp, Category = SceneColor, AdvancedDisplay, meta=(UIMin = "0.0", UIMax = "1.0", editcondition = "bOverride_GrainJitter"))
 	float GrainJitter;
 
 	/** 0..1 grain intensity */
@@ -662,11 +662,11 @@ struct FPostProcessSettings
 	float AmbientOcclusionMipThreshold;
 
 	/** Adjusts indirect lighting color. (1,1,1) is default. (0,0,0) to disable GI. The show flag 'Global Illumination' must be enabled to use this property. */
-	UPROPERTY(interp, Category=GlobalIllumination, AdvancedDisplay, meta=(editcondition = "bOverride_IndirectLightingColor", DisplayName = "Indirect Lighting Color"))
+	UPROPERTY(interp, Category=GlobalIllumination, meta=(editcondition = "bOverride_IndirectLightingColor", DisplayName = "Indirect Lighting Color"))
 	FLinearColor IndirectLightingColor;
 
 	/** Scales the indirect lighting contribution. A value of 0 disables GI. Default is 1. The show flag 'Global Illumination' must be enabled to use this property. */
-	UPROPERTY(interp, Category=GlobalIllumination, AdvancedDisplay, meta=(ClampMin = "0", UIMax = "4.0", editcondition = "bOverride_IndirectLightingIntensity", DisplayName = "Indirect Lighting Intensity"))
+	UPROPERTY(interp, Category=GlobalIllumination, meta=(ClampMin = "0", UIMax = "4.0", editcondition = "bOverride_IndirectLightingIntensity", DisplayName = "Indirect Lighting Intensity"))
 	float IndirectLightingIntensity;
 
 	/** 0..1=full intensity */
@@ -744,7 +744,7 @@ struct FPostProcessSettings
 	float MotionBlurPerObjectSize;
 
 	/** to render with lower resolution and upscale, controlled by console variable, 100:off, needs to be <99 to see effect, only applied in game  */
-	UPROPERTY(interp, Category=Misc, meta=(ClampMin = "0.0", ClampMax = "400.0", editcondition = "bOverride_ScreenPercentage"))
+	UPROPERTY(interp, Category=Misc, AdvancedDisplay, meta=(ClampMin = "0.0", ClampMax = "400.0", editcondition = "bOverride_ScreenPercentage"))
 	float ScreenPercentage;
 
 	/** TemporalAA, FXAA, ... */
diff --git a/Engine/Source/Runtime/Engine/Private/PostProcessVolume.cpp b/Engine/Source/Runtime/Engine/Private/PostProcessVolume.cpp
--- a/Engine/Source/Runtime/Engine/Private/PostProcessVolume.cpp
+++ b/Engine/Source/Runtime/Engine/Private/PostProcessVolume.cpp
@@ -46,4 +46,33 @@ void APostProcessVolume::PostEditChangeProperty(FPropertyChangedEvent& PropertyC
 		}
 	}
 }
+
+bool APostProcessVolume::CanEditChange(const UProperty* InProperty) const
+{
+	if (InProperty)
+	{
+		FString PropertyName = InProperty->GetName();
+
+		if (PropertyName == GET_MEMBER_NAME_STRING_CHECKED(APostProcessVolume, bEnabled))
+		{
+			return true;
+		}
+
+		if (!bEnabled)
+		{
+			return false;
+		}
+
+		if (PropertyName == GET_MEMBER_NAME_STRING_CHECKED(APostProcessVolume, BlendRadius))
+		{
+			if (bUnbound)
+			{
+				return false;
+			}
+		}
+	}
+
+	return Super::CanEditChange(InProperty);
+}
+
 #endif // WITH_EDITOR
\ No newline at end of file
