comments added to a few methods for maintenance

- additional documentation of the 'bIsDefaultTransform' param in a couple of headers (per code review)

diff --git a/Engine/Source/Runtime/Engine/Classes/Engine/SCS_Node.h b/Engine/Source/Runtime/Engine/Classes/Engine/SCS_Node.h
--- a/Engine/Source/Runtime/Engine/Classes/Engine/SCS_Node.h
+++ b/Engine/Source/Runtime/Engine/Classes/Engine/SCS_Node.h
@@ -86,7 +86,14 @@ class USCS_Node : public UObject
 	class USceneComponent* EditorComponentInstance;
 #endif
 
-	/** Create the specified component on the actor, then call action on children */
+	/**
+	 * Create the specified component on the actor, then call action on children
+	 *
+	 * @param Actor					The actor instance for which to create a new component based on the template encapsulated by this node.
+	 * @param ParentComponent		If non-NULL, the component to which the new component should be attached as a child. If NULL, the new component will not be attached to anything.
+	 * @param RootTransform			The transform to apply if this node turns out to be the root component of the actor instance.
+	 * @param bIsDefaultTransform	Indicates whether or not the given transform is a "default" transform, in which case it can be overridden by template defaults.
+	 */
 	void ExecuteNodeOnActor(AActor* Actor, USceneComponent* ParentComponent, const FTransform* RootTransform, bool bIsDefaultTransform);
 
 	/** Returns an array containing this node and all children below it */
diff --git a/Engine/Source/Runtime/Engine/Classes/Engine/SimpleConstructionScript.h b/Engine/Source/Runtime/Engine/Classes/Engine/SimpleConstructionScript.h
--- a/Engine/Source/Runtime/Engine/Classes/Engine/SimpleConstructionScript.h
+++ b/Engine/Source/Runtime/Engine/Classes/Engine/SimpleConstructionScript.h
@@ -17,7 +17,13 @@ class USimpleConstructionScript : public UObject
 	/** Ensures that all root node parent references are still valid and clears the reference if not */
 	ENGINE_API void FixupRootNodeParentReferences();
 
-	/** Execute this script on the supplied actor, creating components */
+	/**
+	 * Execute this script on the supplied actor, creating components 
+	 *
+	 * @param Actor					The actor instance to execute the script on.
+	 * @param RootTransform			The transform to apply to the root scene component of the actor instance if defined in this script.
+	 * @param bIsDefaultTransform	Indicates whether or not the given transform is a "default" transform, in which case it can be overridden by template defaults.
+	 */
 	void ExecuteScriptOnActor(AActor* Actor, const FTransform& RootTransform, bool bIsDefaultTransform);
 #if WITH_EDITOR
 	/** Return the Blueprint associated with this SCS instance */
