BP: Exposed FDateTime and FTimespan types to Blueprints

diff --git a/Engine/Source/Runtime/CoreUObject/Classes/Object.h b/Engine/Source/Runtime/CoreUObject/Classes/Object.h
--- a/Engine/Source/Runtime/CoreUObject/Classes/Object.h
+++ b/Engine/Source/Runtime/CoreUObject/Classes/Object.h
@@ -193,7 +193,6 @@ struct FVector
 
 	UPROPERTY(EditAnywhere, BlueprintReadWrite, Category=Vector, SaveGame)
 	float Z;
-
 };
 
 USTRUCT(immutable, noexport)
@@ -653,7 +652,7 @@ public:
 
 // A date/time value.
 
-USTRUCT(immutable, noexport)
+USTRUCT(immutable, noexport, BlueprintType)
 struct FDateTime
 {
 	UPROPERTY()
@@ -663,7 +662,7 @@ struct FDateTime
 
 // A time span value.
 
-USTRUCT(immutable, noexport)
+USTRUCT(immutable, noexport, BlueprintType)
 struct FTimespan
 {
 	UPROPERTY()
diff --git a/Engine/Source/Runtime/Engine/Classes/Kismet/KismetMathLibrary.h b/Engine/Source/Runtime/Engine/Classes/Kismet/KismetMathLibrary.h
--- a/Engine/Source/Runtime/Engine/Classes/Kismet/KismetMathLibrary.h
+++ b/Engine/Source/Runtime/Engine/Classes/Kismet/KismetMathLibrary.h
@@ -610,6 +610,246 @@ class UKismetMathLibrary : public UBlueprintFunctionLibrary
 	UFUNCTION(BlueprintPure, meta=(FriendlyName = "LinearColor * Float", CompactNodeTitle = "*", Keywords = "* multiply"), Category="Math|Color")
 	static FLinearColor Multiply_LinearColorFloat(FLinearColor A, float B);
 
+	//
+	// DateTime functions.
+	//
+
+	/* Addition (A + B) */
+	UFUNCTION(BlueprintPure, meta=(FriendlyName="DateTime + Timespan", CompactNodeTitle="+", Keywords="+ add plus"), Category="Math|DateTime")
+	static FDateTime Add_DateTimeTimespan( FDateTime A, FTimespan B );
+
+	/* Subtraction (A - B) */
+	UFUNCTION(BlueprintPure, meta=(FriendlyName="DateTime - Timespan", CompactNodeTitle="-", Keywords="- subtract minus"), Category="Math|DateTime")
+	static FDateTime Subtract_DateTimeTimespan( FDateTime A, FTimespan B );
+
+	/* Returns true if the values are equal (A == B) */
+	UFUNCTION(BlueprintPure, meta=(FriendlyName="Equal (DateTime)", CompactNodeTitle="==", Keywords="== equal"), Category="Math|DateTime")
+	static bool EqualEqual_DateTimeDateTime( FDateTime A, FDateTime B );
+
+	/* Returns true if the values are not equal (A != B) */
+	UFUNCTION(BlueprintPure, meta=(FriendlyName="NotEqual (DateTime)", CompactNodeTitle="!=", Keywords="!= not equal"), Category="Math|DateTime")
+	static bool NotEqual_DateTimeDateTime( FDateTime A, FDateTime B );
+
+	/* Returns true if A is Greater than B (A > B) */
+	UFUNCTION(BlueprintPure, meta=(FriendlyName="DateTime > DateTime", CompactNodeTitle=">", Keywords="> greater"), Category="Math|DateTime")
+	static bool Greater_DateTimeDateTime( FDateTime A, FDateTime B );
+
+	/* Returns true if A is Greater than B (A >= B) */
+	UFUNCTION(BlueprintPure, meta=(FriendlyName="DateTime >= DateTime", CompactNodeTitle=">=", Keywords=">= greater"), Category="Math|DateTime")
+	static bool GreaterEqual_DateTimeDateTime( FDateTime A, FDateTime B );
+
+	/* Returns true if A is Greater than B (A < B) */
+	UFUNCTION(BlueprintPure, meta=(FriendlyName="DateTime < DateTime", CompactNodeTitle="<", Keywords="< less"), Category="Math|DateTime")
+	static bool Less_DateTimeDateTime( FDateTime A, FDateTime B );
+
+	/* Returns true if A is Greater than B (A <= B) */
+	UFUNCTION(BlueprintPure, meta=(FriendlyName="DateTime <= DateTime", CompactNodeTitle="<=", Keywords="<= less"), Category="Math|DateTime")
+	static bool LessEqual_DateTimeDateTime( FDateTime A, FDateTime B );
+
+	/* Returns the date component of A */
+	UFUNCTION(BlueprintPure, meta=(FriendlyName="GetDate"), Category="Math|DateTime")
+	static FDateTime GetDate( FDateTime A );
+
+	/* Returns the day component of A (1 to 31) */
+	UFUNCTION(BlueprintPure, meta=(FriendlyName="GetDay"), Category="Math|DateTime")
+	static int32 GetDay( FDateTime A );
+
+	/* Returns the day of year of A */
+	UFUNCTION(BlueprintPure, meta=(FriendlyName="GetDayOfYear"), Category="Math|DateTime")
+	static int32 GetDayOfYear( FDateTime A );
+
+	/* Returns the hour component of A (24h format) */
+	UFUNCTION(BlueprintPure, meta=(FriendlyName="GetHour"), Category="Math|DateTime")
+	static int32 GetHour( FDateTime A );
+
+	/* Returns the hour component of A (12h format) */
+	UFUNCTION(BlueprintPure, meta=(FriendlyName="GetHour12"), Category="Math|DateTime")
+	static int32 GetHour12( FDateTime A );
+
+	/* Returns the millisecond component of A */
+	UFUNCTION(BlueprintPure, meta=(FriendlyName="GetMillisecond"), Category="Math|DateTime")
+	static int32 GetMillisecond( FDateTime A );
+
+	/* Returns the minute component of A */
+	UFUNCTION(BlueprintPure, meta=(FriendlyName="GetMinute"), Category="Math|DateTime")
+	static int32 GetMinute( FDateTime A );
+
+	/* Returns the month component of A */
+	UFUNCTION(BlueprintPure, meta=(FriendlyName="GetMonth"), Category="Math|DateTime")
+	static int32 GetMonth( FDateTime A );
+
+	/* Returns the second component of A */
+	UFUNCTION(BlueprintPure, meta=(FriendlyName="GetSecond"), Category="Math|DateTime")
+	static int32 GetSecond( FDateTime A );
+
+	/* Returns the time elapsed since midnight of A */
+	UFUNCTION(BlueprintPure, meta=(FriendlyName="GetTimeOfDay"), Category="Math|DateTime")
+	static FTimespan GetTimeOfDay( FDateTime A );
+
+	/* Returns the year component of A */
+	UFUNCTION(BlueprintPure, meta=(FriendlyName="GetYear"), Category="Math|DateTime")
+	static int32 GetYear( FDateTime A );
+
+	/* Returns whether A's time is in the afternoon */
+	UFUNCTION(BlueprintPure, meta=(FriendlyName="IsAfternoon"), Category="Math|DateTime")
+	static bool IsAfternoon( FDateTime A );
+
+	/* Returns whether A's time is in the morning */
+	UFUNCTION(BlueprintPure, meta=(FriendlyName="IsMorning"), Category="Math|DateTime")
+	static bool IsMorning( FDateTime A );
+
+	/* Returns the number of days in the given year and month */
+	UFUNCTION(BlueprintPure, meta=(FriendlyName="DaysInMonth"), Category="Math|DateTime")
+	static int32 DaysInMonth( int32 Year, int32 Month );
+
+	/* Returns the number of days in the given year */
+	UFUNCTION(BlueprintPure, meta=(FriendlyName="DaysInYear"), Category="Math|DateTime")
+	static int32 DaysInYear( int32 Year );
+
+	/* Returns whether given year is a leap year */
+	UFUNCTION(BlueprintPure, meta=(FriendlyName="IsLeapYear"), Category="Math|DateTime")
+	static bool IsLeapYear( int32 Year );
+
+	/* Returns the maximum date and time value */
+	UFUNCTION(BlueprintPure, meta=(FriendlyName="MaxValue"), Category="Math|DateTime")
+	static FDateTime DateTimeMaxValue( );
+
+	/* Returns the minimum date and time value */
+	UFUNCTION(BlueprintPure, meta=(FriendlyName="MinValue"), Category="Math|DateTime")
+	static FDateTime DateTimeMinValue( );
+
+	/* Returns the local date and time on this computer */
+	UFUNCTION(BlueprintPure, meta=(FriendlyName="Now"), Category="Math|DateTime")
+	static FDateTime Now( );
+
+	/* Returns the local date on this computer */
+	UFUNCTION(BlueprintPure, meta=(FriendlyName="Today"), Category="Math|DateTime")
+	static FDateTime Today( );
+
+	/* Returns the UTC date and time on this computer */
+	UFUNCTION(BlueprintPure, meta=(FriendlyName="UtcNow"), Category="Math|DateTime")
+	static FDateTime UtcNow( );
+
+	//
+	// Timespan functions.
+	//
+
+	/* Addition (A + B) */
+	UFUNCTION(BlueprintPure, meta=(FriendlyName="Timespan + Timespan", CompactNodeTitle="+", Keywords="+ add plus"), Category="Math|Timespan")
+	static FTimespan Add_TimespanTimespan( FTimespan A, FTimespan B );
+
+	/* Subtraction (A - B) */
+	UFUNCTION(BlueprintPure, meta=(FriendlyName="Timespan - Timespan", CompactNodeTitle="-", Keywords="- subtract minus"), Category="Math|Timespan")
+	static FTimespan Subtract_TimespanTimespan( FTimespan A, FTimespan B );
+
+	/* Scalar multiplication (A * s) */
+	UFUNCTION(BlueprintPure, meta=(FriendlyName="Timespan * float", CompactNodeTitle="*", Keywords="* multiply"), Category="Math|Timespan")
+	static FTimespan Multiply_TimespanFloat( FTimespan A, float Scalar );
+
+	/* Returns true if the values are equal (A == B) */
+	UFUNCTION(BlueprintPure, meta=(FriendlyName="Equal (Timespan)", CompactNodeTitle="==", Keywords="== equal"), Category="Math|Timespan")
+	static bool EqualEqual_TimespanTimespan( FTimespan A, FTimespan B );
+
+	/* Returns true if the values are not equal (A != B) */
+	UFUNCTION(BlueprintPure, meta=(FriendlyName="NotEqual (Timespan)", CompactNodeTitle="!=", Keywords="!= not equal"), Category="Math|Timespan")
+	static bool NotEqual_TimespanTimespan( FTimespan A, FTimespan B );
+
+	/* Returns true if A is Greater than B (A > B) */
+	UFUNCTION(BlueprintPure, meta=(FriendlyName="Timespan > Timespan", CompactNodeTitle=">", Keywords="> greater"), Category="Math|Timespan")
+	static bool Greater_TimespanTimespan( FTimespan A, FTimespan B );
+
+	/* Returns true if A is Greater than B (A >= B) */
+	UFUNCTION(BlueprintPure, meta=(FriendlyName="Timespan >= Timespan", CompactNodeTitle=">=", Keywords=">= greater"), Category="Math|Timespan")
+	static bool GreaterEqual_TimespanTimespan( FTimespan A, FTimespan B );
+
+	/* Returns true if A is Greater than B (A < B) */
+	UFUNCTION(BlueprintPure, meta=(FriendlyName="Timespan < Timespan", CompactNodeTitle="<", Keywords="< less"), Category="Math|Timespan")
+	static bool Less_TimespanTimespan( FTimespan A, FTimespan B );
+
+	/* Returns true if A is Greater than B (A <= B) */
+	UFUNCTION(BlueprintPure, meta=(FriendlyName="Timespan <= Timespan", CompactNodeTitle="<=", Keywords="<= less"), Category="Math|Timespan")
+	static bool LessEqual_TimespanTimespan( FTimespan A, FTimespan B );
+
+	/* Returns the days component of A */
+	UFUNCTION(BlueprintPure, meta=(FriendlyName="GetDays"), Category="Math|Timespan")
+	static int32 GetDays( FTimespan A );
+
+	/* Returns the absolute value of A */
+	UFUNCTION(BlueprintPure, meta=(FriendlyName="GetDuration"), Category="Math|Timespan")
+	static FTimespan GetDuration( FTimespan A );
+
+	/* Returns the hours component of A */
+	UFUNCTION(BlueprintPure, meta=(FriendlyName="GetHours"), Category="Math|Timespan")
+	static int32 GetHours( FTimespan A );
+
+	/* Returns the milliseconds component of A */
+	UFUNCTION(BlueprintPure, meta=(FriendlyName="GetMilliseconds"), Category="Math|Timespan")
+	static int32 GetMilliseconds( FTimespan A );
+
+	/* Returns the minutes component of A */
+	UFUNCTION(BlueprintPure, meta=(FriendlyName="GetMinutes"), Category="Math|Timespan")
+	static int32 GetMinutes( FTimespan A );
+
+	/* Returns the seconds component of A */
+	UFUNCTION(BlueprintPure, meta=(FriendlyName="GetSeconds"), Category="Math|Timespan")
+	static int32 GetSeconds( FTimespan A );
+
+	/* Returns the total number of days in A */
+	UFUNCTION(BlueprintPure, meta=(FriendlyName="GetTotalDays"), Category="Math|Timespan")
+	static float GetTotalDays( FTimespan A );
+
+	/* Returns the total number of hours in A */
+	UFUNCTION(BlueprintPure, meta=(FriendlyName="GetTotalHours"), Category="Math|Timespan")
+	static float GetTotalHours( FTimespan A );
+
+	/* Returns the total number of milliseconds in A */
+	UFUNCTION(BlueprintPure, meta=(FriendlyName="GetTotalMilliseconds"), Category="Math|Timespan")
+	static float GetTotalMilliseconds( FTimespan A );
+
+	/* Returns the total number of minutes in A */
+	UFUNCTION(BlueprintPure, meta=(FriendlyName="GetTotalMinutes"), Category="Math|Timespan")
+	static float GetTotalMinutes( FTimespan A );
+
+	/* Returns the total number of seconds in A */
+	UFUNCTION(BlueprintPure, meta=(FriendlyName="GetTotalSeconds"), Category="Math|Timespan")
+	static float GetTotalSeconds( FTimespan A );
+
+	/* Returns a time span that represents the specified number of days */
+	UFUNCTION(BlueprintPure, meta=(FriendlyName="FromDays"), Category="Math|Timespan")
+	static FTimespan FromDays( float Days );
+
+	/* Returns a time span that represents the specified number of hours */
+	UFUNCTION(BlueprintPure, meta=(FriendlyName="FromHours"), Category="Math|Timespan")
+	static FTimespan FromHours( float Hours );
+
+	/* Returns a time span that represents the specified number of milliseconds */
+	UFUNCTION(BlueprintPure, meta=(FriendlyName="FromMilliseconds"), Category="Math|Timespan")
+	static FTimespan FromMilliseconds( float Milliseconds );
+
+	/* Returns a time span that represents the specified number of minutes */
+	UFUNCTION(BlueprintPure, meta=(FriendlyName="FromMinutes"), Category="Math|Timespan")
+	static FTimespan FromMinutes( float Minutes );
+
+	/* Returns a time span that represents the specified number of seconds */
+	UFUNCTION(BlueprintPure, meta=(FriendlyName="FromSeconds"), Category="Math|Timespan")
+	static FTimespan FromSeconds( float Seconds );
+
+	/* Returns the maximum time span value */
+	UFUNCTION(BlueprintPure, meta=(FriendlyName="MaxValue"), Category="Math|Timespan")
+	static FTimespan TimespanMaxValue( );
+
+	/* Returns the minimum time span value */
+	UFUNCTION(BlueprintPure, meta=(FriendlyName="MinValue"), Category="Math|Timespan")
+	static FTimespan TimespanMinValue( );
+
+	/* Returns the ratio between two time spans (A / B), handles zero values */
+	UFUNCTION(BlueprintPure, meta=(FriendlyName="TimespanRatio"), Category="Math|Timespan")
+	static float TimespanRatio( FTimespan A, FTimespan B );
+
+	/* Returns a zero time span value */
+	UFUNCTION(BlueprintPure, meta=(FriendlyName="ZeroValue"), Category="Math|Timespan")
+	static FTimespan TimespanZeroValue( );
+
 	// -- Begin K2 utilities
 
 	/** Converts a byte to a float */
diff --git a/Engine/Source/Runtime/Engine/Classes/Kismet/KismetTextLibrary.h b/Engine/Source/Runtime/Engine/Classes/Kismet/KismetTextLibrary.h
--- a/Engine/Source/Runtime/Engine/Classes/Kismet/KismetTextLibrary.h
+++ b/Engine/Source/Runtime/Engine/Classes/Kismet/KismetTextLibrary.h
@@ -138,6 +138,22 @@ class UKismetTextLibrary : public UBlueprintFunctionLibrary
 	UFUNCTION(BlueprintPure, meta=(FriendlyName = "AsPercent", AdvancedDisplay = "1"), Category="Utilities|Text")
 	static FText AsPercent_Float(float Value, TEnumAsByte<ERoundingMode> RoundingMode, bool bUseGrouping = true, int32 MinimumIntegralDigits = 1, int32 MaximumIntegralDigits = 324, int32 MinimumFractionalDigits = 0, int32 MaximumFractionalDigits = 3);
 
+	/* Converts a passed in date & time to a text, formatted as a date */
+	UFUNCTION(BlueprintPure, meta=(FriendlyName = "AsDate", AdvancedDisplay = "1"), Category="Utilities|Text")
+	static FText AsDate_DateTime(const FDateTime& InDateTime);
+
+	/* Converts a passed in date & time to a text, formatted as a date & time */
+	UFUNCTION(BlueprintPure, meta=(FriendlyName = "AsDateTime", AdvancedDisplay = "1"), Category="Utilities|Text")
+	static FText AsDateTime_DateTime(const FDateTime& In);
+
+	/* Converts a passed in date & time to a text, formatted as a time */
+	UFUNCTION(BlueprintPure, meta=(FriendlyName = "AsTime", AdvancedDisplay = "1"), Category="Utilities|Text")
+	static FText AsTime_DateTime(const FDateTime& In);
+
+	/* Converts a passed in time span to a text, formatted as a time span */
+	UFUNCTION(BlueprintPure, meta=(FriendlyName = "AsTimespan", AdvancedDisplay = "1"), Category="Utilities|Text")
+	static FText AsTimespan_Timespan(const FTimespan& InTimespan);
+
 	/* Used for formatting text using the FText::Format function and utilized by the UK2Node_FormatText */
 	UFUNCTION(BlueprintPure, meta=(BlueprintInternalUseOnly = "true"))
 	static FText Format(FText InPattern, TArray<FFormatTextArgument> InArgs);
diff --git a/Engine/Source/Runtime/Engine/Private/KismetMathLibrary.cpp b/Engine/Source/Runtime/Engine/Private/KismetMathLibrary.cpp
--- a/Engine/Source/Runtime/Engine/Private/KismetMathLibrary.cpp
+++ b/Engine/Source/Runtime/Engine/Private/KismetMathLibrary.cpp
@@ -1005,6 +1005,368 @@ bool UKismetMathLibrary::ClassIsChildOf(TSubclassOf<class UObject> TestClass, TS
 	return ((*ParentClass != NULL) && (*TestClass != NULL)) ? (*TestClass)->IsChildOf(*ParentClass) : false;
 }
 
+
+/* DateTime functions
+ *****************************************************************************/
+
+FDateTime UKismetMathLibrary::Add_DateTimeTimespan( FDateTime A, FTimespan B )
+{
+	return A + B;
+}
+
+
+FDateTime UKismetMathLibrary::Subtract_DateTimeTimespan( FDateTime A, FTimespan B )
+{
+	return A - B;
+}
+
+
+bool UKismetMathLibrary::EqualEqual_DateTimeDateTime( FDateTime A, FDateTime B )
+{
+	return A == B;
+}
+
+
+bool UKismetMathLibrary::NotEqual_DateTimeDateTime( FDateTime A, FDateTime B )
+{
+	return A != B;
+}
+
+
+bool UKismetMathLibrary::Greater_DateTimeDateTime( FDateTime A, FDateTime B )
+{
+	return A > B;
+}
+
+
+bool UKismetMathLibrary::GreaterEqual_DateTimeDateTime( FDateTime A, FDateTime B )
+{
+	return A >= B;
+}
+
+
+bool UKismetMathLibrary::Less_DateTimeDateTime( FDateTime A, FDateTime B )
+{
+	return A < B;
+}
+
+
+bool UKismetMathLibrary::LessEqual_DateTimeDateTime( FDateTime A, FDateTime B )
+{
+	return A <= B;
+}
+
+
+FDateTime UKismetMathLibrary::GetDate( FDateTime A )
+{
+	return A.GetDate();
+}
+
+
+int32 UKismetMathLibrary::GetDay( FDateTime A )
+{
+	return A.GetDay();
+}
+
+int32 UKismetMathLibrary::GetDayOfYear( FDateTime A )
+{
+	return A.GetDayOfYear();
+}
+
+
+int32 UKismetMathLibrary::GetHour( FDateTime A )
+{
+	return A.GetHour();
+}
+
+
+int32 UKismetMathLibrary::GetHour12( FDateTime A )
+{
+	return A.GetHour12();
+}
+
+
+int32 UKismetMathLibrary::GetMillisecond( FDateTime A )
+{
+	return A.GetMillisecond();
+}
+
+
+int32 UKismetMathLibrary::GetMinute( FDateTime A )
+{
+	return A.GetMinute();
+}
+
+
+int32 UKismetMathLibrary::GetMonth( FDateTime A )
+{
+	return A.GetMonth();
+}
+
+
+int32 UKismetMathLibrary::GetSecond( FDateTime A )
+{
+	return A.GetSecond();
+}
+
+
+FTimespan UKismetMathLibrary::GetTimeOfDay( FDateTime A )
+{
+	return A.GetTimeOfDay();
+}
+
+
+int32 UKismetMathLibrary::GetYear( FDateTime A )
+{
+	return A.GetYear();
+}
+
+
+bool UKismetMathLibrary::IsAfternoon( FDateTime A )
+{
+	return A.IsAfternoon();
+}
+
+
+bool UKismetMathLibrary::IsMorning( FDateTime A )
+{
+	return A.IsMorning();
+}
+
+
+int32 UKismetMathLibrary::DaysInMonth( int32 Year, int32 Month )
+{
+	return FDateTime::DaysInMonth(Year, Month);
+}
+
+
+int32 UKismetMathLibrary::DaysInYear( int32 Year )
+{
+	return FDateTime::DaysInYear(Year);
+}
+
+
+bool UKismetMathLibrary::IsLeapYear( int32 Year )
+{
+	return FDateTime::IsLeapYear(Year);
+}
+
+
+FDateTime UKismetMathLibrary::DateTimeMaxValue( )
+{
+	return FDateTime::MaxValue();
+}
+
+
+FDateTime UKismetMathLibrary::DateTimeMinValue( )
+{
+	return FDateTime::MinValue();
+}
+
+
+FDateTime UKismetMathLibrary::Now( )
+{
+	return FDateTime::Now();
+}
+
+
+FDateTime UKismetMathLibrary::Today( )
+{
+	return FDateTime::Today();
+}
+
+
+FDateTime UKismetMathLibrary::UtcNow( )
+{
+	return FDateTime::UtcNow();
+}
+
+
+/* Timespan functions
+ *****************************************************************************/
+
+FTimespan UKismetMathLibrary::Add_TimespanTimespan( FTimespan A, FTimespan B )
+{
+	return A + B;
+}
+
+
+FTimespan UKismetMathLibrary::Subtract_TimespanTimespan( FTimespan A, FTimespan B )
+{
+	return A - B;
+}
+
+
+FTimespan UKismetMathLibrary::Multiply_TimespanFloat( FTimespan A, float Scalar )
+{
+	return A * Scalar;
+}
+
+
+bool UKismetMathLibrary::EqualEqual_TimespanTimespan( FTimespan A, FTimespan B )
+{
+	return A == B;
+}
+
+
+bool UKismetMathLibrary::NotEqual_TimespanTimespan( FTimespan A, FTimespan B )
+{
+	return A != B;
+}
+
+
+bool UKismetMathLibrary::Greater_TimespanTimespan( FTimespan A, FTimespan B )
+{
+	return A > B;
+}
+
+
+bool UKismetMathLibrary::GreaterEqual_TimespanTimespan( FTimespan A, FTimespan B )
+{
+	return A >= B;
+}
+
+
+bool UKismetMathLibrary::Less_TimespanTimespan( FTimespan A, FTimespan B )
+{
+	return A < B;
+}
+
+
+bool UKismetMathLibrary::LessEqual_TimespanTimespan( FTimespan A, FTimespan B )
+{
+	return A <= B;
+}
+
+
+int32 UKismetMathLibrary::GetDays( FTimespan A )
+{
+	return A.GetDays();
+}
+
+
+FTimespan UKismetMathLibrary::GetDuration( FTimespan A )
+{
+	return A.GetDuration();
+}
+
+
+int32 UKismetMathLibrary::GetHours( FTimespan A )
+{
+	return A.GetHours();
+}
+
+
+int32 UKismetMathLibrary::GetMilliseconds( FTimespan A )
+{
+	return A.GetMilliseconds();
+}
+
+
+int32 UKismetMathLibrary::GetMinutes( FTimespan A )
+{
+	return A.GetMinutes();
+}
+
+
+int32 UKismetMathLibrary::GetSeconds( FTimespan A )
+{
+	return A.GetSeconds();
+}
+
+
+float UKismetMathLibrary::GetTotalDays( FTimespan A )
+{
+	return A.GetTotalDays();
+}
+
+
+float UKismetMathLibrary::GetTotalHours( FTimespan A )
+{
+	return A.GetTotalHours();
+}
+
+
+float UKismetMathLibrary::GetTotalMilliseconds( FTimespan A )
+{
+	return A.GetTotalMilliseconds();
+}
+
+
+float UKismetMathLibrary::GetTotalMinutes( FTimespan A )
+{
+	return A.GetTotalMinutes();
+}
+
+
+float UKismetMathLibrary::GetTotalSeconds( FTimespan A )
+{
+	return A.GetTotalSeconds();
+}
+
+
+FTimespan UKismetMathLibrary::FromDays( float Days )
+{
+	return FTimespan::FromDays(Days);
+}
+
+
+FTimespan UKismetMathLibrary::FromHours( float Hours )
+{
+	return FTimespan::FromHours(Hours);
+}
+
+
+FTimespan UKismetMathLibrary::FromMilliseconds( float Milliseconds )
+{
+	return FTimespan::FromMilliseconds(Milliseconds);
+}
+
+
+FTimespan UKismetMathLibrary::FromMinutes( float Minutes )
+{
+	return FTimespan::FromMinutes(Minutes);
+}
+
+
+FTimespan UKismetMathLibrary::FromSeconds( float Seconds )
+{
+	return FTimespan::FromSeconds(Seconds);
+}
+
+
+FTimespan UKismetMathLibrary::TimespanMaxValue( )
+{
+	return FTimespan::MaxValue();
+}
+
+
+FTimespan UKismetMathLibrary::TimespanMinValue( )
+{
+	return FTimespan::MinValue();
+}
+
+
+float UKismetMathLibrary::TimespanRatio( FTimespan A, FTimespan B )
+{
+	if (B != FTimespan::Zero())
+	{
+		return (float)A.GetTicks() / (float)B.GetTicks();
+	}
+
+	return 0.0;
+}
+
+
+FTimespan UKismetMathLibrary::TimespanZeroValue( )
+{
+	return FTimespan::Zero();
+}
+
+
+/* K2 Utilities
+ *****************************************************************************/
+
 float UKismetMathLibrary::Conv_ByteToFloat(uint8 InByte)
 {
 	return (float)InByte;
diff --git a/Engine/Source/Runtime/Engine/Private/KismetTextLibrary.cpp b/Engine/Source/Runtime/Engine/Private/KismetTextLibrary.cpp
--- a/Engine/Source/Runtime/Engine/Private/KismetTextLibrary.cpp
+++ b/Engine/Source/Runtime/Engine/Private/KismetTextLibrary.cpp
@@ -4,94 +4,114 @@
 
 #define LOCTEXT_NAMESPACE "Kismet"
 
-//////////////////////////////////////////////////////////////////////////
-// UKismetTextLibrary
+
+/* UKismetTextLibrary structors
+ *****************************************************************************/
 
 UKismetTextLibrary::UKismetTextLibrary(const class FPostConstructInitializeProperties& PCIP)
 	: Super(PCIP)
-{
-}
+{ }
+
+
+/* UKismetTextLibrary static functions
+ *****************************************************************************/
 
 FString UKismetTextLibrary::Conv_TextToString(const FText& InText)
 {
 	return InText.ToString();
 }
 
+
 FText UKismetTextLibrary::Conv_StringToText(const FString& InString)
 {	
 	return FText::FromString(InString);
 }
 
+
 FText UKismetTextLibrary::Conv_NameToText(FName InName)
 {
 	return FText::FromName(InName);
 }
 
+
 bool UKismetTextLibrary::TextIsEmpty(const FText& InText)
 {
 	return InText.IsEmpty();
 }
 
+
 bool UKismetTextLibrary::TextIsTransient(const FText& InText)
 {
 	return InText.IsTransient();
 }
 
+
 bool UKismetTextLibrary::TextIsCultureInvariant(const FText& InText)
 {
 	return InText.IsCultureInvariant();
 }
 
+
 FText UKismetTextLibrary::TextTrimPreceding(const FText& InText)
 {
 	return FText::TrimPreceding(InText);
 }
 
+
 FText UKismetTextLibrary::TextTrimTrailing(const FText& InText)
 {
 	return FText::TrimTrailing(InText);
 }
 
+
 FText UKismetTextLibrary::TextTrimPrecedingAndTrailing(const FText& InText)
 {
 	return FText::TrimPrecedingAndTrailing(InText);
 }
 
+
 FText UKismetTextLibrary::GetEmptyText()
 {
 	return FText::GetEmpty();
 }
 
+
 bool UKismetTextLibrary::FindTextInLocalizationTable(const FString& Namespace, const FString& Key, FText& OutText)
 {
 	return FText::FindText(Namespace, Key, OutText);
 }
 
+
 bool UKismetTextLibrary::EqualEqual_IgnoreCase_TextText(const FText& A, const FText& B)
 {
 	return A.EqualToCaseIgnored( B );
 }
 
+
 bool UKismetTextLibrary::EqualEqual_TextText(const FText& A, const FText& B)
 {
 	return A.EqualTo( B );
 }
 
+
 bool UKismetTextLibrary::NotEqual_IgnoreCase_TextText(const FText& A, const FText& B)
 {
 	return !A.EqualToCaseIgnored( B );
 }
 
+
 bool UKismetTextLibrary::NotEqual_TextText(const FText& A, const FText& B)
 {
 	return !A.EqualTo( B );
 }
 
+
 FText UKismetTextLibrary::Conv_BoolToText(bool InBool)
 {
 	return InBool ? LOCTEXT("True", "true") : LOCTEXT("False", "false");
 }
 
+
 FText UKismetTextLibrary::Conv_ByteToText(uint8 Value)
 {
 	FNumberFormattingOptions NumberFormatOptions;
@@ -105,6 +125,7 @@ FText UKismetTextLibrary::Conv_ByteToText(uint8 Value)
 	return FText::AsNumber(Value, &NumberFormatOptions);
 }
 
+
 FText UKismetTextLibrary::Conv_IntToText(int32 Value, bool bUseGrouping/* = true*/, int32 MinimumIntegralDigits/* = 1*/, int32 MaximumIntegralDigits/* = 324*/)
 {
 	FNumberFormattingOptions NumberFormatOptions;
@@ -118,6 +139,7 @@ FText UKismetTextLibrary::Conv_IntToText(int32 Value, bool bUseGrouping/* = true
 	return FText::AsNumber(Value, &NumberFormatOptions);
 }
 
+
 FText UKismetTextLibrary::Conv_FloatToText(float Value, TEnumAsByte<ERoundingMode> RoundingMode, bool bUseGrouping/* = true*/, int32 MinimumIntegralDigits/* = 1*/, int32 MaximumIntegralDigits/* = 324*/, int32 MinimumFractionalDigits/* = 0*/, int32 MaximumFractionalDigits/* = 3*/)
 {
 	FNumberFormattingOptions NumberFormatOptions;
@@ -131,6 +153,7 @@ FText UKismetTextLibrary::Conv_FloatToText(float Value, TEnumAsByte<ERoundingMod
 	return FText::AsNumber(Value, &NumberFormatOptions);
 }
 
+
 FText UKismetTextLibrary::AsCurrency_Integer(int32 Value, TEnumAsByte<ERoundingMode> RoundingMode, bool bUseGrouping/* = true*/, int32 MinimumIntegralDigits/* = 1*/, int32 MaximumIntegralDigits/* = 324*/, int32 MinimumFractionalDigits/* = 0*/, int32 MaximumFractionalDigits/* = 3*/, const FString& CurrencyCode)
 {
 	FNumberFormattingOptions NumberFormatOptions;
@@ -144,6 +167,7 @@ FText UKismetTextLibrary::AsCurrency_Integer(int32 Value, TEnumAsByte<ERoundingM
 	return FText::AsCurrency(Value, CurrencyCode, &NumberFormatOptions);
 }
 
+
 FText UKismetTextLibrary::AsCurrency_Float(float Value, TEnumAsByte<ERoundingMode> RoundingMode, bool bUseGrouping/* = true*/, int32 MinimumIntegralDigits/* = 1*/, int32 MaximumIntegralDigits/* = 324*/, int32 MinimumFractionalDigits/* = 0*/, int32 MaximumFractionalDigits/* = 3*/, const FString& CurrencyCode)
 {
 	FNumberFormattingOptions NumberFormatOptions;
@@ -157,6 +181,7 @@ FText UKismetTextLibrary::AsCurrency_Float(float Value, TEnumAsByte<ERoundingMod
 	return FText::AsCurrency(Value, CurrencyCode, &NumberFormatOptions);
 }
 
+
 FText UKismetTextLibrary::AsPercent_Float(float Value, TEnumAsByte<ERoundingMode> RoundingMode, bool bUseGrouping/* = true*/, int32 MinimumIntegralDigits/* = 1*/, int32 MaximumIntegralDigits/* = 324*/, int32 MinimumFractionalDigits/* = 0*/, int32 MaximumFractionalDigits/* = 3*/)
 {
 	FNumberFormattingOptions NumberFormatOptions;
@@ -170,6 +195,31 @@ FText UKismetTextLibrary::AsPercent_Float(float Value, TEnumAsByte<ERoundingMode
 	return FText::AsPercent(Value, &NumberFormatOptions);
 }
 
+
+FText UKismetTextLibrary::AsDate_DateTime(const FDateTime& InDateTime)
+{
+	return FText::AsDate(InDateTime);
+}
+
+
+FText UKismetTextLibrary::AsDateTime_DateTime(const FDateTime& InDateTime)
+{
+	return FText::AsDateTime(InDateTime);
+}
+
+
+FText UKismetTextLibrary::AsTime_DateTime(const FDateTime& InDateTime)
+{
+	return FText::AsTime(InDateTime);
+}
+
+
+FText UKismetTextLibrary::AsTimespan_Timespan(const FTimespan& InTimespan)
+{
+	return FText::AsTimespan(InTimespan);
+}
+
+
 FText UKismetTextLibrary::Format(FText InPattern, TArray<FFormatTextArgument> InArgs)
 {
 	TArray< FFormatArgumentData > Arguments;
@@ -185,4 +235,5 @@ FText UKismetTextLibrary::Format(FText InPattern, TArray<FFormatTextArgument> In
 	return FText::Format(InPattern, Arguments);
 }
 
+
 #undef LOCTEXT_NAMESPACE
