Improve comments for ETickingGroup

diff --git a/Engine/Source/Runtime/Engine/Classes/Engine/EngineBaseTypes.h b/Engine/Source/Runtime/Engine/Classes/Engine/EngineBaseTypes.h
--- a/Engine/Source/Runtime/Engine/Classes/Engine/EngineBaseTypes.h
+++ b/Engine/Source/Runtime/Engine/Classes/Engine/EngineBaseTypes.h
@@ -33,56 +33,40 @@ enum ELevelTick
 	LEVELTICK_PauseTick = 3,	// Delta time is zero, we are paused. Components don't tick.
 };
 
-/**
- * Determines which ticking group an Actor/Component belongs to
- */
+/** Determines which ticking group a tick function belongs to */
 UENUM()
 enum ETickingGroup
 {
-	/**
-	 * Any item that needs to be updated before asynchronous work is done
-	 */
+	/** Any item that needs to be executed before physics simulation starts */
 	TG_PrePhysics,
-	/**
-	 * Temp while we transition away from tick groups, this will only have one task, the one that starts physics
-	 */							
+
+	/** Special tick group that starts physics simulation */							
 	TG_StartPhysics,
-	/**
-	 * Any item that can be run in parallel of our async work
-	 */
+
+	/** Any item that can be run in parallel with our physics simulation work */
 	TG_DuringPhysics,
-	/**
-	 * Temp while we transition away from tick groups, this will only have one task, the one is not completed until physics is done
-	 */
+
+	/** Special tick group that ends physics simulation */
 	TG_EndPhysics,
-	/**
-	 * Any item that needs the async work to be done before being updated
-	 */
 
+	/** Any item that needs physics to be complete before being executed */
 	TG_PreCloth,
-	/**
-	  * Any item that needs to be updated after rigid body simulation is done, but before cloth is simulation is done
-	  */
 
+	/** Any item that needs to be updated after rigid body simulation is done, but before cloth is simulation is done */
 	TG_StartCloth,
-	/**
-	  * Any item that needs to run during cloth simulation
-	  */
-	
+
+	/** Any item that can be run during cloth simulation */	
 	TG_EndCloth,
-	/**
-	* Any item that needs cloth to be done before being updated
-	*/
 
+	/** Any item that needs rigid body and cloth sim to be complete before being executed */
 	TG_PostPhysics,
-	/**
-	 * Any item that needs the update work to be done before being ticked
-	 */
+
+	/** Any item that needs the update work to be done before being ticked */
 	TG_PostUpdateWork,
-	/**
-	 * Special tick group that is not actually a tick group. After every tick group this is repeatedly re-run until there are no more newly spawned items to run
-	 */
+
+	/** Special tick group that is not actually a tick group. After every tick group this is repeatedly re-run until there are no more newly spawned items to run */
 	TG_NewlySpawned,
+
 	TG_MAX,
 };
 
