Mark dynamic data dirty when calling UPoseableMeshComponent::SetBoneTransformByName

diff --git a/Engine/Source/Runtime/Engine/Private/PoseableMeshComponent.cpp b/Engine/Source/Runtime/Engine/Private/PoseableMeshComponent.cpp
--- a/Engine/Source/Runtime/Engine/Private/PoseableMeshComponent.cpp
+++ b/Engine/Source/Runtime/Engine/Private/PoseableMeshComponent.cpp
@@ -146,6 +146,9 @@ void UPoseableMeshComponent::SetBoneTransformByName(FName BoneName, const FTrans
 
 				LocalAtoms[BoneIndex].SetToRelativeTransform(CSPose.GetComponentSpaceTransform(ParentIndex));
 			}
+
+			// Need to send new state to render thread
+			MarkRenderDynamicDataDirty();
 		}
 	}
 }
