Expose FFontRenderInfo and FDepthFieldGlowInfo to BP

diff --git a/Engine/Source/Runtime/Engine/Classes/Engine/EngineTypes.h b/Engine/Source/Runtime/Engine/Classes/Engine/EngineTypes.h
--- a/Engine/Source/Runtime/Engine/Classes/Engine/EngineTypes.h
+++ b/Engine/Source/Runtime/Engine/Classes/Engine/EngineTypes.h
@@ -2735,29 +2735,29 @@ public:
 };
 
 /** info for glow when using depth field rendering */
-USTRUCT()
+USTRUCT(BlueprintType)
 struct FDepthFieldGlowInfo
 {
 	GENERATED_USTRUCT_BODY()
 
 	/** whether to turn on the outline glow (depth field fonts only) */
-	UPROPERTY()
+	UPROPERTY(BlueprintReadWrite, Category="Glow")
 	uint32 bEnableGlow:1;
 
 	/** base color to use for the glow */
-	UPROPERTY()
+	UPROPERTY(BlueprintReadWrite, Category="Glow")
 	FLinearColor GlowColor;
 
 	/** if bEnableGlow, outline glow outer radius (0 to 1, 0.5 is edge of character silhouette)
 	 * glow influence will be 0 at GlowOuterRadius.X and 1 at GlowOuterRadius.Y
 	*/
-	UPROPERTY()
+	UPROPERTY(BlueprintReadWrite, Category="Glow")
 	FVector2D GlowOuterRadius;
 
 	/** if bEnableGlow, outline glow inner radius (0 to 1, 0.5 is edge of character silhouette)
 	 * glow influence will be 1 at GlowInnerRadius.X and 0 at GlowInnerRadius.Y
 	 */
-	UPROPERTY()
+	UPROPERTY(BlueprintReadWrite, Category="Glow")
 	FVector2D GlowInnerRadius;
 
 
@@ -2794,21 +2794,21 @@ struct FDepthFieldGlowInfo
 };
 
 /** information used in font rendering */
-USTRUCT()
+USTRUCT(BlueprintType)
 struct FFontRenderInfo
 {
 	GENERATED_USTRUCT_BODY()
 
 	/** whether to clip text */
-	UPROPERTY()
+	UPROPERTY(BlueprintReadWrite, Category="FontInfo")
 	uint32 bClipText:1;
 
 	/** whether to turn on shadowing */
-	UPROPERTY()
+	UPROPERTY(BlueprintReadWrite, Category="FontInfo")
 	uint32 bEnableShadow:1;
 
 	/** depth field glow parameters (only usable if font was imported with a depth field) */
-	UPROPERTY()
+	UPROPERTY(BlueprintReadWrite, Category="FontInfo")
 	struct FDepthFieldGlowInfo GlowInfo;
 
 	FFontRenderInfo()
