Make UAnimInstance::GetOwningActor and GetSkeletalMeshComponent const.

diff --git a/Engine/Source/Runtime/Engine/Classes/Animation/AnimInstance.h b/Engine/Source/Runtime/Engine/Classes/Animation/AnimInstance.h
--- a/Engine/Source/Runtime/Engine/Classes/Animation/AnimInstance.h
+++ b/Engine/Source/Runtime/Engine/Classes/Animation/AnimInstance.h
@@ -312,12 +312,12 @@ public:
 
 public:
 	/** Returns the owning actor of this AnimInstance */
-	UFUNCTION(BlueprintPure, Category="Animation")
-	AActor* GetOwningActor();
+	UFUNCTION(Category="Animation")
+	AActor* GetOwningActor() const;
 	
 	// Returns the skeletal mesh component that has created this AnimInstance
-	UFUNCTION(BlueprintPure, Category="Animation")
-	USkeletalMeshComponent* GetOwningComponent();
+	UFUNCTION(Category="Animation")
+	USkeletalMeshComponent* GetOwningComponent() const;
 
 public:
 
diff --git a/Engine/Source/Runtime/Engine/Private/Animation/AnimInstance.cpp b/Engine/Source/Runtime/Engine/Private/Animation/AnimInstance.cpp
--- a/Engine/Source/Runtime/Engine/Private/Animation/AnimInstance.cpp
+++ b/Engine/Source/Runtime/Engine/Private/Animation/AnimInstance.cpp
@@ -190,7 +190,7 @@ void UAnimInstance::CopyPose(const FA2Pose& Source, FA2Pose& Destination)
 	}
 }
 
-AActor* UAnimInstance::GetOwningActor()
+AActor* UAnimInstance::GetOwningActor() const
 {
 	USkeletalMeshComponent* OwnerComponent = GetSkelMeshComponent();
 	return OwnerComponent->GetOwner();
@@ -207,7 +207,7 @@ APawn* UAnimInstance::TryGetPawnOwner()
 	return NULL;
 }
 
-USkeletalMeshComponent* UAnimInstance::GetOwningComponent()
+USkeletalMeshComponent* UAnimInstance::GetOwningComponent() const
 {
 	return GetSkelMeshComponent();
 }
