Fix listen server player  not animating

diff --git a/Engine/Source/Runtime/Engine/Private/Character.cpp b/Engine/Source/Runtime/Engine/Private/Character.cpp
--- a/Engine/Source/Runtime/Engine/Private/Character.cpp
+++ b/Engine/Source/Runtime/Engine/Private/Character.cpp
@@ -691,7 +691,7 @@ void ACharacter::PossessedBy(AController* NewController)
 	Super::PossessedBy(NewController);
 
 	// If we are controlled remotely, set animation timing to be driven by client's network updates. So timing and events remain in sync.
-	if (Mesh && (GetRemoteRole() == ROLE_AutonomousProxy))
+	if (Mesh && (GetRemoteRole() == ROLE_AutonomousProxy && GetNetConnection() != nullptr))
 	{
 		Mesh->bAutonomousTickPose = true;
 	}
