Correct [REDACTED_HANDLE] and [REDACTED_HANDLE] comments to ensure accurate pin tooltips

diff --git a/Engine/Source/Runtime/Engine/Classes/Kismet/KismetArrayLibrary.h b/Engine/Source/Runtime/Engine/Classes/Kismet/KismetArrayLibrary.h
--- a/Engine/Source/Runtime/Engine/Classes/Kismet/KismetArrayLibrary.h
+++ b/Engine/Source/Runtime/Engine/Classes/Kismet/KismetArrayLibrary.h
@@ -10,8 +10,8 @@ class UKismetArrayLibrary : public UBlueprintFunctionLibrary
 	/** 
 	 *Add item to array
 	 *
-	 *@param	Array	The array to add item to
-	 *@param	Item	The item to add to the array
+	 *@param	TargetArray		The array to add item to
+	 *@param	NewItem			The item to add to the array
 	*/
 	UFUNCTION(BlueprintCallable, CustomThunk, meta=(FriendlyName = "Add", CompactNodeTitle = "ADD", ArrayParm = "TargetArray|ArrayProperty", ArrayTypeDependentParams = "NewItem", AutoCreateRefTerm = "NewItem"), Category="Utilities|Array")
 	static int32 Array_Add(const TArray<int32>& TargetArray, const UArrayProperty* ArrayProperty, const int32& NewItem);
@@ -19,7 +19,7 @@ class UKismetArrayLibrary : public UBlueprintFunctionLibrary
 	/** 
 	 * Shuffle (randomize) the elements of an array
 	 *
-	 *@param	Array	The array to shuffle
+	 *@param	TargetArray		The array to shuffle
 	*/
 	UFUNCTION(BlueprintCallable, CustomThunk, meta=(FriendlyName = "Shuffle", CompactNodeTitle = "SHUFFLE", ArrayParm = "TargetArray|ArrayProperty"), Category="Utilities|Array")
 	static void Array_Shuffle(const TArray<int32>& TargetArray, const UArrayProperty* ArrayProperty);
@@ -36,9 +36,9 @@ class UKismetArrayLibrary : public UBlueprintFunctionLibrary
 	/* 
 	 *Insert item at the given index into the array.
 	 *	
-	 *@param	Array	The array to insert into
-	 *@param	Item	The item to insert into the array
-	 *@param	Int		The index at which to insert the item into the array
+	 *@param	TargetArray		The array to insert into
+	 *@param	NewItem			The item to insert into the array
+	 *@param	Index			The index at which to insert the item into the array
 	*/
 	UFUNCTION(BlueprintCallable, CustomThunk, meta=(FriendlyName = "Insert", CompactNodeTitle = "INSERT", ArrayParm = "TargetArray|ArrayProperty", ArrayTypeDependentParams = "NewItem", AutoCreateRefTerm = "NewItem"), Category="Utilities|Array")
 	static void Array_Insert(const TArray<int32>& TargetArray, const UArrayProperty* ArrayProperty, const int32& NewItem, int32 Index);
@@ -47,8 +47,8 @@ class UKismetArrayLibrary : public UBlueprintFunctionLibrary
 	/* 
 	 *Remove item at the given index from the array.
 	 *
-	 *@param	Array	The array to remove from
-	 *@param	Item	The index into the array to remove from
+	 *@param	TargetArray		The array to remove from
+	 *@param	IndexToRemove	The index into the array to remove from
 	*/
 	UFUNCTION(BlueprintCallable, CustomThunk, meta=(FriendlyName = "Remove Index", CompactNodeTitle = "REMOVE INDEX", ArrayParm = "TargetArray|ArrayProperty"), Category="Utilities|Array")
 	static void Array_Remove(const TArray<int32>& TargetArray, const UArrayProperty* ArrayProperty, int32 IndexToRemove);
@@ -56,9 +56,9 @@ class UKismetArrayLibrary : public UBlueprintFunctionLibrary
 	/* 
 	 *Remove all instances of item from array.
 	 *
-	 *@param	Array	The array to remove from
-	 *@param	Item	The item to remove from the array
-	 *@return bool	Returns True if one or more items were removed
+	 *@param	TargetArray		The array to remove from
+	 *@param	Item			The item to remove from the array
+	 *@return	True if one or more items were removed
 	*/
 	UFUNCTION(BlueprintCallable, CustomThunk, meta=(FriendlyName = "Remove Item", CompactNodeTitle = "REMOVE", ArrayParm = "TargetArray|ArrayProperty", ArrayTypeDependentParams = "Item", AutoCreateRefTerm = "Item"), Category="Utilities|Array")
 	static bool Array_RemoveItem(const TArray<int32>& TargetArray, const UArrayProperty* ArrayProperty, const int32 &Item);
@@ -66,7 +66,7 @@ class UKismetArrayLibrary : public UBlueprintFunctionLibrary
 	/* 
 	 *Clear an array, removes all content
 	 *
-	 *@param	Pin	The array to clear
+	 *@param	TargetArray		The array to clear
 	*/
 	UFUNCTION(BlueprintCallable, CustomThunk, meta=(FriendlyName = "Clear", CompactNodeTitle = "CLEAR", ArrayParm = "TargetArray|ArrayProperty"), Category="Utilities|Array")
 	static void Array_Clear(const TArray<int32>& TargetArray, const UArrayProperty* ArrayProperty);
@@ -74,8 +74,8 @@ class UKismetArrayLibrary : public UBlueprintFunctionLibrary
 	/* 
 	 *Resize Array to specified size. 
 	 *	
-	 *@param	Array	The array to resize
-	 *@param	Int		The new size of the array
+	 *@param	TargetArray		The array to resize
+	 *@param	Size			The new size of the array
 	*/
 	UFUNCTION(BlueprintCallable, CustomThunk, meta=(FriendlyName = "Resize", CompactNodeTitle = "RESIZE", ArrayParm = "TargetArray|ArrayProperty"), Category="Utilities|Array")
 	static void Array_Resize(const TArray<int32>& TargetArray, const UArrayProperty* ArrayProperty, int32 Size);
@@ -83,8 +83,8 @@ class UKismetArrayLibrary : public UBlueprintFunctionLibrary
 	/* 
 	 *Get the number of items in an array
 	 *
-	 *@param	Array	The array to get the length of
-	 *@return int		Returns the length of the array
+	 *@param	TargetArray		The array to get the length of
+	 *@return	The length of the array
 	*/
 	UFUNCTION(BlueprintPure, CustomThunk, meta=(FriendlyName = "Length", CompactNodeTitle = "LENGTH", ArrayParm = "TargetArray|ArrayProperty", Keywords = "num size count"), Category="Utilities|Array")
 	static int32 Array_Length(const TArray<int32>& TargetArray, const UArrayProperty* ArrayProperty);
@@ -93,8 +93,8 @@ class UKismetArrayLibrary : public UBlueprintFunctionLibrary
 	/* 
 	 *Get the last valid index into an array
 	 *	
-	 *@param	Array	The array to perform the operation on
-	 *@return int		Returns the last valid index of the array
+	 *@param	TargetArray		The array to perform the operation on
+	 *@return	The last valid index of the array
 	*/
 	UFUNCTION(BlueprintPure, CustomThunk, meta=(FriendlyName = "Last Index", CompactNodeTitle = "LAST INDEX", ArrayParm = "TargetArray|ArrayProperty"), Category="Utilities|Array")
 	static int32 Array_LastIndex(const TArray<int32>& TargetArray, const UArrayProperty* ArrayProperty);
@@ -102,9 +102,9 @@ class UKismetArrayLibrary : public UBlueprintFunctionLibrary
 	/* 
 	 *Given an array and an index, returns the item found at that index
 	 *
-	 *@param	Array	The array 
-	 *@param	int		The index
-	 *@return int		Returns the item stored at the index
+	 *@param	TargetArray		The array to get an item from
+	 *@param	Index			The index in the array to get an item from
+	 *@return	The item stored at the index
 	*/
 	UFUNCTION(BlueprintPure, CustomThunk, meta=(FriendlyName = "Get", CompactNodeTitle = "GET", ArrayParm = "TargetArray|ArrayProperty", ArrayTypeDependentParams = "Item"), Category="Utilities|Array")
 	static void Array_Get(const TArray<int32>& TargetArray, const UArrayProperty* ArrayProperty, int32 Index, int32& Item);
@@ -112,10 +112,10 @@ class UKismetArrayLibrary : public UBlueprintFunctionLibrary
 	/* 
 	 *Given an array and an index, assigns the item to that array element
 	 *
-	 *@param	Array	The array to perform the operation on
-	 *@param	int		The index
-	 *@param	Item	The item to assign to the index of the array
-	 *@param	bSizeToFit	If true, the array will expand if Index is greater than the current size of the array
+	 *@param	TargetArray		The array to perform the operation on
+	 *@param	Index			The index to assing the item to
+	 *@param	Item			The item to assign to the index of the array
+	 *@param	bSizeToFit		If true, the array will expand if Index is greater than the current size of the array
 	*/
 	UFUNCTION(BlueprintCallable, CustomThunk, meta=(FriendlyName = "Set Array Elem", ArrayParm = "TargetArray|ArrayProperty", ArrayTypeDependentParams = "Item", AutoCreateRefTerm = "Item"), Category="Utilities|Array")
 	static void Array_Set(const TArray<int32>& TargetArray, const UArrayProperty* ArrayProperty, int32 Index, const int32& Item, bool bSizeToFit);
@@ -123,9 +123,9 @@ class UKismetArrayLibrary : public UBlueprintFunctionLibrary
 	/*  
 	 *Finds the index of the first instance of the item within the array
 	 *	
-	 *@param	Array	The array to search for the item
-	 *@param	int		The item to look for
-	 *@return int		The index the item was found at, or INDEX_NONE if not found
+	 *@param	TargetArray		The array to search for the item
+	 *@param	ItemToFind		The item to look for
+	 *@return	The index the item was found at, or -1 if not found
 	*/
 	UFUNCTION(BlueprintPure, CustomThunk, meta=(FriendlyName = "Find Item", CompactNodeTitle = "FIND", ArrayParm = "TargetArray|ArrayProperty", ArrayTypeDependentParams = "ItemToFind", AutoCreateRefTerm = "ItemToFind"), Category="Utilities|Array")
 	static int32 Array_Find(const TArray<int32>& TargetArray, const UArrayProperty* ArrayProperty, const int32& ItemToFind);
@@ -133,9 +133,9 @@ class UKismetArrayLibrary : public UBlueprintFunctionLibrary
 	/*  
 	 *Returns true if the array contains the given item
 	 *
-	 *@param	Array	The array to search for the item
-	 *@param	item	The item to look for
-	 *@return bool	Returns True if the item was found within the array
+	 *@param	TargetArray		The array to search for the item
+	 *@param	ItemToFind		The item to look for
+	 *@return	True if the item was found within the array
 	*/
 	UFUNCTION(BlueprintPure, CustomThunk, meta=(FriendlyName = "Contains Item", CompactNodeTitle = "CONTAINS", ArrayParm = "TargetArray|ArrayProperty", ArrayTypeDependentParams = "ItemToFind", AutoCreateRefTerm = "ItemToFind"), Category="Utilities|Array")
 	static bool Array_Contains(const TArray<int32>& TargetArray, const UArrayProperty* ArrayProperty, const int32& ItemToFind);
@@ -143,15 +143,15 @@ class UKismetArrayLibrary : public UBlueprintFunctionLibrary
 	/*  
 	 *Filter an array based on a Class derived from Actor.  
 	 *	
-	 *@param	Array	The array to filter from
-	 *@param	Actor	The Actor sub-class type that acts as the filter, only objects derived from it will be returned.
-	 *@return Array	Returns array containing only those objects which are derived from the class type specified.
+	 *@param	TargetArray		The array to filter from
+	 *@param	FilterClass		The Actor sub-class type that acts as the filter, only objects derived from it will be returned.
+	 *@return	An array containing only those objects which are derived from the class specified.
 	*/
 	UFUNCTION(BlueprintCallable, meta=(FriendlyName = "Filter Array"), Category="Utilities|Array")
 	static void FilterArray(const TArray<AActor*>& TargetArray, TSubclassOf<class AActor> FilterClass, TArray<AActor*>& FilteredArray);
 
 	/** 
-	 *
+	 * Not exposed to users. Supports setting an array property on an object by name.
 	 */
 	UFUNCTION(BlueprintCallable, CustomThunk, meta=(BlueprintInternalUseOnly = "true", ArrayParm = "Value", ArrayTypeDependentParams="Value"))
 	static void SetArrayPropertyByName(UObject* Object, FName PropertyName, const TArray<int32>& Value);
