SuggestProjectile traces along dynamics instead of just overlap test. Fixes overlapping volumes from blocking

diff --git a/Engine/Source/Runtime/Engine/Classes/Kismet/GameplayStatics.h b/Engine/Source/Runtime/Engine/Classes/Kismet/GameplayStatics.h
--- a/Engine/Source/Runtime/Engine/Classes/Kismet/GameplayStatics.h
+++ b/Engine/Source/Runtime/Engine/Classes/Kismet/GameplayStatics.h
@@ -495,7 +495,7 @@ class ENGINE_API UGameplayStatics : public UBlueprintFunctionLibrary
 	static bool BlueprintSuggestProjectileVelocity(UObject* WorldContextObject, FVector& TossVelocity, FVector StartLocation, FVector EndLocation, float LaunchSpeed, float OverrideGravityZ, ESuggestProjVelocityTraceOption::Type TraceOption, float CollisionRadius, bool bFavorHighArc, bool bDrawDebug);
 
 	/** Native version, has more options than the Blueprint version. */
-	static bool SuggestProjectileVelocity(UObject* WorldContextObject, FVector& TossVelocity, FVector StartLocation, FVector EndLocation, float TossSpeed, bool bHighArc=false, float CollisionRadius = 0.f, float OverrideGravityZ = 0, ESuggestProjVelocityTraceOption::Type TraceOption = ESuggestProjVelocityTraceOption::TraceFullPath, bool bDrawDebug=false);
+	static bool SuggestProjectileVelocity(UObject* WorldContextObject, FVector& TossVelocity, FVector StartLocation, FVector EndLocation, float TossSpeed, bool bHighArc = false, float CollisionRadius = 0.f, float OverrideGravityZ = 0, ESuggestProjVelocityTraceOption::Type TraceOption = ESuggestProjVelocityTraceOption::TraceFullPath, const FCollisionResponseParams& ResponseParam = FCollisionResponseParams::DefaultResponseParam, const TArray<AActor*>& ActorsToIgnore = TArray<AActor*>(), bool bDrawDebug = false);
 
 	/** Returns world origin current location */
 	UFUNCTION(BlueprintPure, Category="Game", meta=(HidePin="WorldContextObject", DefaultToSelf="WorldContextObject") )
diff --git a/Engine/Source/Runtime/Engine/Private/GameplayStatics.cpp b/Engine/Source/Runtime/Engine/Private/GameplayStatics.cpp
--- a/Engine/Source/Runtime/Engine/Private/GameplayStatics.cpp
+++ b/Engine/Source/Runtime/Engine/Private/GameplayStatics.cpp
@@ -1109,11 +1109,11 @@ FString UGameplayStatics::GetPlatformName()
 bool UGameplayStatics::BlueprintSuggestProjectileVelocity(UObject* WorldContextObject, FVector& OutTossVelocity, FVector StartLocation, FVector EndLocation, float LaunchSpeed, float OverrideGravityZ, ESuggestProjVelocityTraceOption::Type TraceOption, float CollisionRadius, bool bFavorHighArc, bool bDrawDebug)
 {
 	// simple pass-through to the C++ interface
-	return UGameplayStatics::SuggestProjectileVelocity(WorldContextObject, OutTossVelocity, StartLocation, EndLocation, LaunchSpeed, bFavorHighArc, CollisionRadius, OverrideGravityZ, TraceOption, bDrawDebug);
+	return UGameplayStatics::SuggestProjectileVelocity(WorldContextObject, OutTossVelocity, StartLocation, EndLocation, LaunchSpeed, bFavorHighArc, CollisionRadius, OverrideGravityZ, TraceOption, FCollisionResponseParams::DefaultResponseParam, TArray<AActor*>(), bDrawDebug);
 }
 
 // Based on analytic solution to ballistic angle of launch http://en.wikipedia.org/wiki/Trajectory_of_a_projectile#Angle_required_to_hit_coordinate_.28x.2Cy.29
-bool UGameplayStatics::SuggestProjectileVelocity(UObject* WorldContextObject, FVector& OutTossVelocity, FVector Start, FVector End, float TossSpeed, bool bFavorHighArc, float CollisionRadius, float OverrideGravityZ, ESuggestProjVelocityTraceOption::Type TraceOption, bool bDrawDebug)
+bool UGameplayStatics::SuggestProjectileVelocity(UObject* WorldContextObject, FVector& OutTossVelocity, FVector Start, FVector End, float TossSpeed, bool bFavorHighArc, float CollisionRadius, float OverrideGravityZ, ESuggestProjVelocityTraceOption::Type TraceOption, const FCollisionResponseParams& ResponseParam, const TArray<AActor*>& ActorsToIgnore, bool bDrawDebug)
 {
 	const FVector FlightDelta = End - Start;
 	const FVector DirXY = FlightDelta.SafeNormal2D();
@@ -1238,7 +1238,10 @@ bool UGameplayStatics::SuggestProjectileVelocity(UObject* WorldContextObject, FV
 				{
 					// note: this will automatically fall back to line test if radius is small enough
 					static const FName NAME_SuggestProjVelTrace = FName(TEXT("SuggestProjVelTrace"));
-					if ( World->SweepTest( TraceStart, TraceEnd, FQuat::Identity, FCollisionShape::MakeSphere(CollisionRadius), FCollisionQueryParams(NAME_SuggestProjVelTrace, true), FCollisionObjectQueryParams(ECC_WorldStatic) ) )
+
+					FCollisionQueryParams QueryParams(NAME_SuggestProjVelTrace, true);
+					QueryParams.AddIgnoredActors(ActorsToIgnore);
+					if (World->SweepTest(TraceStart, TraceEnd, FQuat::Identity, ECC_WorldDynamic, FCollisionShape::MakeSphere(CollisionRadius), QueryParams, ResponseParam))
 					{
 						// hit something, failed
 						bFailedTrace = true;
