#ue4 - ThirdPerson and FreeCam cameras ignore the viewtarget when sweeping to the camera location.

- Fix from PeterK from UT.

diff --git a/Engine/Source/Runtime/Engine/Private/PlayerCameraManager.cpp b/Engine/Source/Runtime/Engine/Private/PlayerCameraManager.cpp
--- a/Engine/Source/Runtime/Engine/Private/PlayerCameraManager.cpp
+++ b/Engine/Source/Runtime/Engine/Private/PlayerCameraManager.cpp
@@ -529,10 +529,10 @@ void APlayerCameraManager::UpdateViewTarget(FTViewTarget& OutVT, float DeltaTime
 			{
 				Rotator = PCOwner->GetControlRotation();
 			}
-			Loc += FRotationMatrix(Rotator).TransformVector(FreeCamOffset);
 
-			FVector Pos = Loc - Rotator.Vector() * FreeCamDistance;
+			FVector Pos = Loc + FRotationMatrix(Rotator).TransformVector(FreeCamOffset) - Rotator.Vector() * FreeCamDistance;
 			FCollisionQueryParams BoxParams(NAME_FreeCam, false, this);
+			BoxParams.AddIgnoredActor(OutVT.Target);
 			FHitResult Result;
 
 			GetWorld()->SweepSingle(Result, Loc, Pos, FQuat::Identity, ECC_Camera, FCollisionShape::MakeBox(FVector(12.f)), BoxParams);
