#ue4 - Add functions to Actor and SceneComponent for adding location and rotation in world space.

- Added AActor methods: AddActorWorldRotation, AddActorWorldOffset, AddActorWorldTransform (ignores scale)
- Added USceneComponent methods: AddWorldRotation, AddWorldOffset, AddWorldTransform (ignores scale)

diff --git a/Engine/Source/Runtime/Engine/Classes/Components/SceneComponent.h b/Engine/Source/Runtime/Engine/Classes/Components/SceneComponent.h
--- a/Engine/Source/Runtime/Engine/Classes/Components/SceneComponent.h
+++ b/Engine/Source/Runtime/Engine/Classes/Components/SceneComponent.h
@@ -249,7 +249,7 @@ public:
 	UFUNCTION(BlueprintCallable, Category="Utilities|Transformation")
 	void AddLocalRotation(FRotator DeltaRotation, bool bSweep=false);
 
-	/** Adds a delta to the rotation of this component in its local reference frame */
+	/** Adds a delta to the transform of this component in its local reference frame. Scale is unchanged. */
 	UFUNCTION(BlueprintCallable, Category="Utilities|Transformation")
 	void AddLocalTransform(const FTransform& DeltaTransform, bool bSweep=false);
 
@@ -261,13 +261,26 @@ public:
 	UFUNCTION(BlueprintCallable, Category="Utilities|Transformation")
 	void SetWorldRotation(FRotator NewRotation, bool bSweep=false);
 
+	/** Set the relative scale of this component to put it at the supplied scale in world space. */
 	UFUNCTION(BlueprintCallable, Category="Utilities|Transformation")
 	void SetWorldScale3D(FVector NewScale);
 
-	UFUNCTION(BlueprintCallable, Category="Utilities|Transformation")
 	/** Set the transform of this component in world space. */
+	UFUNCTION(BlueprintCallable, Category="Utilities|Transformation")
 	void SetWorldTransform(const FTransform& NewTransform, bool bSweep=false);
 
+	/** Adds a delta to the location of this component in world space. */
+	UFUNCTION(BlueprintCallable, Category="Utilities|Transformation")
+	void AddWorldOffset(FVector DeltaLocation, bool bSweep=false);
+
+	/** Adds a delta to the rotation of this component in world space. */
+	UFUNCTION(BlueprintCallable, Category="Utilities|Transformation")
+	void AddWorldRotation(FRotator DeltaRotation, bool bSweep=false);
+
+	/** Adds a delta to the transform of this component in world space. Scale is unchanged. */
+	UFUNCTION(BlueprintCallable, Category="Utilities|Transformation")
+	void AddWorldTransform(const FTransform& DeltaTransform, bool bSweep=false);
+
 	/** Return location of the component, in world space */
 	UFUNCTION(BlueprintCallable, meta=(FriendlyName = "GetWorldLocation"), Category="Utilities|Transformation")
 	FVector K2_GetComponentLocation() const;
diff --git a/Engine/Source/Runtime/Engine/Classes/GameFramework/Actor.h b/Engine/Source/Runtime/Engine/Classes/GameFramework/Actor.h
--- a/Engine/Source/Runtime/Engine/Classes/GameFramework/Actor.h
+++ b/Engine/Source/Runtime/Engine/Classes/GameFramework/Actor.h
@@ -625,6 +625,18 @@ public:
 	UFUNCTION(BlueprintCallable, Category = "Utilities|Transformation")
 	float GetHorizontalDotProductTo(AActor* OtherActor);
 
+	/** Adds a delta to the location of this actor in world space. */
+	UFUNCTION(BlueprintCallable, Category="Utilities|Transformation")
+	void AddActorWorldOffset(FVector DeltaLocation, bool bSweep=false);
+
+	/** Adds a delta to the rotation of this actor in world space. */
+	UFUNCTION(BlueprintCallable, Category="Utilities|Transformation")
+	void AddActorWorldRotation(FRotator DeltaRotation, bool bSweep=false);
+
+	/** Adds a delta to the transform of this actor in world space. Scale is unchanged. */
+	UFUNCTION(BlueprintCallable, Category="Utilities|Transformation")
+	void AddActorWorldTransform(const FTransform& DeltaTransform, bool bSweep=false);
+
 	/** 
 	 *	Set the Actors transform to the specified one.
 	 *	@param bSweep		Should we sweep to the destination location. If true, will stop short of the target if blocked by something.
diff --git a/Engine/Source/Runtime/Engine/Private/Actor.cpp b/Engine/Source/Runtime/Engine/Private/Actor.cpp
--- a/Engine/Source/Runtime/Engine/Private/Actor.cpp
+++ b/Engine/Source/Runtime/Engine/Private/Actor.cpp
@@ -2388,6 +2388,32 @@ FVector AActor::GetActorScale3D() const
 	return FVector(1,1,1);
 }
 
+void AActor::AddActorWorldOffset(FVector DeltaLocation, bool bSweep)
+{
+	if (RootComponent)
+	{
+		RootComponent->AddWorldOffset(DeltaLocation, bSweep);
+	}
+}
+
+void AActor::AddActorWorldRotation(FRotator DeltaRotation, bool bSweep)
+{
+	if (RootComponent)
+	{
+		RootComponent->AddWorldRotation(DeltaRotation, bSweep);
+	}	
+}
+
+void AActor::AddActorWorldTransform(const FTransform& DeltaTransform, bool bSweep)
+{
+	if (RootComponent)
+	{
+		RootComponent->AddWorldTransform(DeltaTransform, bSweep);
+	}	
+}
+
+
+
 bool AActor::SetActorTransform(const FTransform& NewTransform, bool bSweep)
 {
 	// we have seen this gets NAN from kismet, and would like to see if this
diff --git a/Engine/Source/Runtime/Engine/Private/SceneComponent.cpp b/Engine/Source/Runtime/Engine/Private/SceneComponent.cpp
--- a/Engine/Source/Runtime/Engine/Private/SceneComponent.cpp
+++ b/Engine/Source/Runtime/Engine/Private/SceneComponent.cpp
@@ -383,6 +383,32 @@ void USceneComponent::AddLocalTransform(const FTransform& DeltaTransform, bool b
 	SetRelativeTransform(NewRelTransform, bSweep);
 }
 
+void USceneComponent::AddWorldOffset(FVector DeltaLocation, bool bSweep)
+{
+	if (!DeltaLocation.IsZero())
+	{
+		const FVector NewWorldLocation = DeltaLocation + ComponentToWorld.GetTranslation();
+		SetWorldLocation(NewWorldLocation, bSweep);
+	}
+}
+
+void USceneComponent::AddWorldRotation(FRotator DeltaRotation, bool bSweep)
+{
+	if (!DeltaRotation.IsZero())
+	{
+		const FQuat NewWorldRotation = DeltaRotation.Quaternion() * ComponentToWorld.GetRotation();
+		SetWorldRotation(NewWorldRotation.Rotator(), bSweep);
+	}
+}
+
+void USceneComponent::AddWorldTransform(const FTransform& DeltaTransform, bool bSweep)
+{
+	const FQuat NewWorldRotation = DeltaTransform.GetRotation() * ComponentToWorld.GetRotation();
+	const FVector NewWorldLocation = DeltaTransform.GetTranslation() + ComponentToWorld.GetTranslation();
+
+	SetWorldTransform(FTransform(NewWorldRotation, NewWorldLocation, FVector(1,1,1)));
+}
+
 void USceneComponent::SetRelativeScale3D(FVector NewScale3D)
 {
 	if( NewScale3D != RelativeScale3D )
