Fix invalid uint32 values in blueprint functions

diff --git a/Engine/Source/Runtime/Engine/Classes/Kismet/KismetInputLibrary.h b/Engine/Source/Runtime/Engine/Classes/Kismet/KismetInputLibrary.h
--- a/Engine/Source/Runtime/Engine/Classes/Kismet/KismetInputLibrary.h
+++ b/Engine/Source/Runtime/Engine/Classes/Kismet/KismetInputLibrary.h
@@ -168,15 +168,15 @@ class UKismetInputLibrary : public UBlueprintFunctionLibrary
 
 	/** @return The unique identifier of the pointer (e.g., finger index) */
 	UFUNCTION(BlueprintPure, meta=( FriendlyName = "GetPointerIndex" ), Category="Utilities|PointerEvent")
-	static uint32 PointerEvent_GetPointerIndex(const FPointerEvent& Input);
+	static int32 PointerEvent_GetPointerIndex(const FPointerEvent& Input);
 
 	/** @return The index of the touch pad that generated this event (for platforms with multiple touch pads per user) */
 	UFUNCTION(BlueprintPure, meta=( FriendlyName = "GetTouchpadIndex" ), Category="Utilities|PointerEvent")
-	static uint32 PointerEvent_GetTouchpadIndex(const FPointerEvent& Input);
+	static int32 PointerEvent_GetTouchpadIndex(const FPointerEvent& Input);
 
 	/** @return Is this event a result from a touch (as opposed to a mouse) */
 	UFUNCTION(BlueprintPure, meta=( FriendlyName = "IsTouchEvent" ), Category="Utilities|PointerEvent")
-	static uint32 PointerEvent_IsTouchEvent(const FPointerEvent& Input);
+	static bool PointerEvent_IsTouchEvent(const FPointerEvent& Input);
 
 	//TODO UMG Support GetGestureType()
 
@@ -195,7 +195,7 @@ class UKismetInputLibrary : public UBlueprintFunctionLibrary
 
 	/** @return The index of the user that caused the event */
 	UFUNCTION(BlueprintPure, meta=( FriendlyName = "GetUserIndex" ), Category="Utilities|ControllerEvent")
-	static uint32 ControllerEvent_GetUserIndex(const FControllerEvent& Input);
+	static int32 ControllerEvent_GetUserIndex(const FControllerEvent& Input);
 
 	/** @return Analog value between 0 and 1.  1 being fully pressed, 0 being not pressed at all */
 	UFUNCTION(BlueprintPure, meta=( FriendlyName = "GetAnalogValue" ), Category="Utilities|ControllerEvent")
diff --git a/Engine/Source/Runtime/Engine/Private/KismetInputLibrary.cpp b/Engine/Source/Runtime/Engine/Private/KismetInputLibrary.cpp
--- a/Engine/Source/Runtime/Engine/Private/KismetInputLibrary.cpp
+++ b/Engine/Source/Runtime/Engine/Private/KismetInputLibrary.cpp
@@ -134,17 +134,17 @@ int32 UKismetInputLibrary::PointerEvent_GetUserIndex(const FPointerEvent& Input)
 	return Input.GetUserIndex();
 }
 
-uint32 UKismetInputLibrary::PointerEvent_GetPointerIndex(const FPointerEvent& Input)
+int32 UKismetInputLibrary::PointerEvent_GetPointerIndex(const FPointerEvent& Input)
 {
 	return Input.GetPointerIndex();
 }
 
-uint32 UKismetInputLibrary::PointerEvent_GetTouchpadIndex(const FPointerEvent& Input)
+int32 UKismetInputLibrary::PointerEvent_GetTouchpadIndex(const FPointerEvent& Input)
 {
 	return Input.GetTouchpadIndex();
 }
 
-uint32 UKismetInputLibrary::PointerEvent_IsTouchEvent(const FPointerEvent& Input)
+bool UKismetInputLibrary::PointerEvent_IsTouchEvent(const FPointerEvent& Input)
 {
 	return Input.IsTouchEvent();
 }
@@ -165,7 +165,7 @@ FKey UKismetInputLibrary::ControllerEvent_GetEffectingButton(const FControllerEv
 	return Input.GetEffectingButton();
 }
 
-uint32 UKismetInputLibrary::ControllerEvent_GetUserIndex(const FControllerEvent& Input)
+int32 UKismetInputLibrary::ControllerEvent_GetUserIndex(const FControllerEvent& Input)
 {
 	return Input.GetUserIndex();
 }
diff --git a/Engine/Source/Runtime/SlateCore/Public/Input/Events.h b/Engine/Source/Runtime/SlateCore/Public/Input/Events.h
--- a/Engine/Source/Runtime/SlateCore/Public/Input/Events.h
+++ b/Engine/Source/Runtime/SlateCore/Public/Input/Events.h
@@ -564,7 +564,7 @@ public:
 	uint32 GetTouchpadIndex() const { return TouchpadIndex; }
 
 	/** @return Is this event a result from a touch (as opposed to a mouse) */
-	uint32 IsTouchEvent() const { return bIsTouchEvent; }
+	bool IsTouchEvent() const { return bIsTouchEvent; }
 
 	/** @return The type of touch gesture */
 	EGestureEvent::Type GetGestureType() const { return GestureType; }
