Engine: Fix camera frustum visualization not being updated when camera component properties are modified

diff --git a/Engine/Source/Runtime/Engine/Private/Camera/CameraComponent.cpp b/Engine/Source/Runtime/Engine/Private/Camera/CameraComponent.cpp
--- a/Engine/Source/Runtime/Engine/Private/Camera/CameraComponent.cpp
+++ b/Engine/Source/Runtime/Engine/Private/Camera/CameraComponent.cpp
@@ -95,6 +95,7 @@ void UCameraComponent::RefreshVisualRepresentation()
 		DrawFrustum->FrustumStartDist = 10.f;
 		DrawFrustum->FrustumEndDist = 1000.f;
 		DrawFrustum->FrustumAspectRatio = AspectRatio;
+		DrawFrustum->MarkRenderStateDirty();
 	}
 }
 
