Fix crash in blueprint when adding static mesh without mesh and trying to weld

diff --git a/Engine/Source/Runtime/Engine/Private/PrimitiveComponent.cpp b/Engine/Source/Runtime/Engine/Private/PrimitiveComponent.cpp
--- a/Engine/Source/Runtime/Engine/Private/PrimitiveComponent.cpp
+++ b/Engine/Source/Runtime/Engine/Private/PrimitiveComponent.cpp
@@ -459,11 +459,17 @@ void UPrimitiveComponent::UnWeldFromParent()
 			//create new root
 			RootBI->UnWeld(NewRootBI);
 			NewRootBI->bWelded = false;
+
+			bool bHasBodySetup = GetBodySetup() != NULL;
+
 			//we need to temporarily turn off AutoWeld
-			bool bPrevAutoWeld = NewRootBI->bAutoWeld;
-			NewRootBI->bAutoWeld = false;
-			NewRootBI->InitBody(GetBodySetup(), GetComponentToWorld(), this, GetWorld()->GetPhysicsScene());
-			NewRootBI->bAutoWeld = bPrevAutoWeld;
+			if (bHasBodySetup)
+			{
+				bool bPrevAutoWeld = NewRootBI->bAutoWeld;
+				NewRootBI->bAutoWeld = false;
+				NewRootBI->InitBody(GetBodySetup(), GetComponentToWorld(), this, GetWorld()->GetPhysicsScene());
+				NewRootBI->bAutoWeld = bPrevAutoWeld;
+			}
 
 			//now weld its children to it
 			TArray<FBodyInstance*> ChildrenBodies;
@@ -476,7 +482,10 @@ void UPrimitiveComponent::UnWeldFromParent()
 				if (ChildBI != NewRootBI)
 				{
 					RootBI->UnWeld(NewRootBI);
-					NewRootBI->Weld(ChildBI, ChildBI->OwnerComponent->GetSocketTransform(ChildrenLabels[ChildIdx]));
+					if (bHasBodySetup)
+					{
+						NewRootBI->Weld(ChildBI, ChildBI->OwnerComponent->GetSocketTransform(ChildrenLabels[ChildIdx]));
+					}
 				}
 			}
 		}
