#ue4 - Rename UCharacterMovementComponent.bCrouchMovesCharacterDown to bCrouchMaintainsBaseLocation.

- Also make UnCrouch() check this flag rather than the current movement mode when deciding how to grow the capsule.

diff --git a/Engine/Config/BaseEngine.ini b/Engine/Config/BaseEngine.ini
--- a/Engine/Config/BaseEngine.ini
+++ b/Engine/Config/BaseEngine.ini
@@ -610,6 +610,7 @@ bUseStreamingPause=false
 +TaggedPropertyRedirects=(ClassName="CharacterMovementComponent",OldPropertyName="WaterSpeed",NewPropertyName="MaxSwimSpeed")
 +TaggedPropertyRedirects=(ClassName="CharacterMovementComponent",OldPropertyName="CrouchedSpeedPercent",NewPropertyName="CrouchedSpeedMultiplier")
 +TaggedPropertyRedirects=(ClassName="CharacterMovementComponent",OldPropertyName="bOrientToMovement",NewPropertyName="bOrientRotationToMovement")
++TaggedPropertyRedirects=(ClassName="CharacterMovementComponent",OldPropertyName="bCrouchMovesCharacterDown",NewPropertyName="bCrouchMaintainsBaseLocation")
 +TaggedPropertyRedirects=(ClassName="PawnSensingComponent",OldPropertyName="SightCounterInterval",NewPropertyName="SensingInterval")
 +TaggedPropertyRedirects=(ClassName="PawnSensingComponent",OldPropertyName="bWantsSeePlayerNotify",NewPropertyName="bSeePawns")
 +TaggedPropertyRedirects=(ClassName="FloatingPawnMovement",OldPropertyName="AccelRate",NewPropertyName="Acceleration")
diff --git a/Engine/Source/Runtime/Core/Public/UObject/ObjectVersion.h b/Engine/Source/Runtime/Core/Public/UObject/ObjectVersion.h
--- a/Engine/Source/Runtime/Core/Public/UObject/ObjectVersion.h
+++ b/Engine/Source/Runtime/Core/Public/UObject/ObjectVersion.h
@@ -600,6 +600,8 @@ enum EUnrealEngineObjectUE4Version
 	VER_UE4_SOUND_COMPRESSION_TYPE_ADDED,
 	// Bodies will automatically weld when attached
 	VER_UE4_AUTO_WELDING,
+	// Rename UCharacterMovementComponent::bCrouchMovesCharacterDown
+	VER_UE4_RENAME_CROUCHMOVESCHARACTERDOWN,
 
 	// -----<new versions can be added before this line>-------------------------------------------------
 	// - this needs to be the last line (see note below)
diff --git a/Engine/Source/Runtime/Engine/Classes/GameFramework/CharacterMovementComponent.h b/Engine/Source/Runtime/Engine/Classes/GameFramework/CharacterMovementComponent.h
--- a/Engine/Source/Runtime/Engine/Classes/GameFramework/CharacterMovementComponent.h
+++ b/Engine/Source/Runtime/Engine/Classes/GameFramework/CharacterMovementComponent.h
@@ -377,6 +377,9 @@ public:
 public:
 	// Deprecated properties
 
+	UPROPERTY()
+	uint32 bCrouchMovesCharacterDown_DEPRECATED:1;
+
 	UPROPERTY()
 	uint32 bForceBraking_DEPRECATED:1;
 
@@ -524,11 +527,12 @@ public:
 	uint32 bWantsToCrouch:1;
 
 	/**
-	 * If true, crouching should lower the center of the player and keep the base of the capsule in place. If false, the base of the capsule moves up and the center stays in place.
-	 * By default, this is set to true when walking and false otherwise when the movement mode changes. Feel free to override the behavior when the movement mode changes.
+	 * If true, crouching should keep the base of the capsule in place by lowering the center of the shrunken capsule. If false, the base of the capsule moves up and the center stays in place.
+	 * The same behavior applies when the character uncrouches: if true, the base is kept in the same location and the center moves up. If false, the capsule grows and only moves up if the base impacts something.
+	 * By default this variable is set when the movement mode changes: set to true when walking and false otherwise. Feel free to override the behavior when the movement mode changes.
 	 */
 	UPROPERTY(Category="Character Movement", VisibleInstanceOnly, BlueprintReadWrite, AdvancedDisplay)
-	uint32 bCrouchMovesCharacterDown:1;
+	uint32 bCrouchMaintainsBaseLocation:1;
 
 	/** If true, the pawn ignores the effects of changes in its base's rotation on its rotation. */
 	UPROPERTY(Category="Character Movement", EditAnywhere, BlueprintReadWrite)
@@ -815,7 +819,7 @@ public:
 
 	/** 
 	 * Updates Velocity and Acceleration based on the current state, applying the effects of friction and acceleration or deceleration. Does *not* apply gravity.
-	 * This is used internally during movement updates. Normally you don't need to call this from outside code, but you might want to use it for custom movement.
+	 * This is used internally during movement updates. Normally you don't need to call this from outside code, but you might want to use it for custom movement modes.
 	 *
 	 * @param	DeltaTime						time elapsed since last frame.
 	 * @param	Friction						coefficient of friction when not accelerating, or in the direction opposite acceleration.
diff --git a/Engine/Source/Runtime/Engine/Private/CharacterMovementComponent.cpp b/Engine/Source/Runtime/Engine/Private/CharacterMovementComponent.cpp
--- a/Engine/Source/Runtime/Engine/Private/CharacterMovementComponent.cpp
+++ b/Engine/Source/Runtime/Engine/Private/CharacterMovementComponent.cpp
@@ -199,6 +199,7 @@ UCharacterMovementComponent::UCharacterMovementComponent(const class FPostConstr
 	DefaultWaterMovementMode = MOVE_Swimming;
 	DefaultLandMovementMode = MOVE_Walking;
 	bForceNextFloorCheck = true;
+	bCrouchMovesCharacterDown_DEPRECATED = false;
 	bForceBraking_DEPRECATED = false;
 	bShrinkProxyCapsule = true;
 	bCanWalkOffLedges = true;
@@ -285,6 +286,11 @@ void UCharacterMovementComponent::PostLoad()
 		MaxWalkSpeedCrouched = MaxWalkSpeed * CrouchedSpeedMultiplier_DEPRECATED;
 		MaxCustomMovementSpeed = MaxWalkSpeed;
 	}
+
+	if (LinkerUE4Ver < VER_UE4_RENAME_CROUCHMOVESCHARACTERDOWN)
+	{
+		bCrouchMaintainsBaseLocation = bCrouchMovesCharacterDown_DEPRECATED;
+	}
 }
 
 
@@ -634,7 +640,7 @@ void UCharacterMovementComponent::OnMovementModeChanged(EMovementMode PreviousMo
 	{	
 		// Walking uses only XY velocity, and must be on a walkable floor, with a Base.
 		Velocity.Z = 0.f;
-		bCrouchMovesCharacterDown = true;
+		bCrouchMaintainsBaseLocation = true;
 
 		// make sure we update our new floor/base on initial entry of the walking physics
 		FindFloor(UpdatedComponent->GetComponentLocation(), CurrentFloor, false);
@@ -644,7 +650,7 @@ void UCharacterMovementComponent::OnMovementModeChanged(EMovementMode PreviousMo
 	else
 	{
 		CurrentFloor.Clear();
-		bCrouchMovesCharacterDown = false;
+		bCrouchMaintainsBaseLocation = false;
 
 		if (MovementMode == MOVE_Falling)
 		{
@@ -1609,7 +1615,7 @@ void UCharacterMovementComponent::Crouch(bool bClientSimulation)
 			}
 		}
 
-		if (bCrouchMovesCharacterDown)
+		if (bCrouchMaintainsBaseLocation)
 		{
 			// Intentionally not using MoveUpdatedComponent, where a horizontal plane constraint would prevent the base of the capsule from staying at the same spot.
 			UpdatedComponent->MoveComponent(FVector(0.f, 0.f, -ScaledHalfHeightAdjust), CharacterOwner->GetActorRotation(), true);
@@ -1671,7 +1677,7 @@ void UCharacterMovementComponent::UnCrouch(bool bClientSimulation)
 		const ECollisionChannel CollisionChannel = UpdatedComponent->GetCollisionObjectType();
 		bool bEncroached = true;
 
-		if (!IsMovingOnGround())
+		if (!bCrouchMaintainsBaseLocation)
 		{
 			// Expand in place
 			bEncroached = GetWorld()->OverlapTest(PawnLocation, FQuat::Identity, CollisionChannel, StandingCapsuleShape, CapsuleParams, ResponseParam);
@@ -1712,20 +1718,22 @@ void UCharacterMovementComponent::UnCrouch(bool bClientSimulation)
 		}
 		else
 		{
-			// Moving on ground
 			// Expand while keeping base location the same.
 			FVector StandingLocation = PawnLocation + FVector(0.f, 0.f, StandingCapsuleShape.GetCapsuleHalfHeight() - CurrentCrouchedHalfHeight);
 			bEncroached = GetWorld()->OverlapTest(StandingLocation, FQuat::Identity, CollisionChannel, StandingCapsuleShape, CapsuleParams, ResponseParam);
 
 			if (bEncroached)
 			{
-				// Something might be just barely overhead, try moving down closer to the floor to avoid it.
-				const float MinFloorDist = KINDA_SMALL_NUMBER * 10.f;
-				if (CurrentFloor.bBlockingHit && CurrentFloor.FloorDist > MinFloorDist)
+				if (IsMovingOnGround())
 				{
-					StandingLocation.Z -= CurrentFloor.FloorDist - MinFloorDist;
-					bEncroached = GetWorld()->OverlapTest(StandingLocation, FQuat::Identity, CollisionChannel, StandingCapsuleShape, CapsuleParams, ResponseParam);
-				}
+					// Something might be just barely overhead, try moving down closer to the floor to avoid it.
+					const float MinFloorDist = KINDA_SMALL_NUMBER * 10.f;
+					if (CurrentFloor.bBlockingHit && CurrentFloor.FloorDist > MinFloorDist)
+					{
+						StandingLocation.Z -= CurrentFloor.FloorDist - MinFloorDist;
+						bEncroached = GetWorld()->OverlapTest(StandingLocation, FQuat::Identity, CollisionChannel, StandingCapsuleShape, CapsuleParams, ResponseParam);
+					}
+				}				
 			}
 
 			if (!bEncroached)
