#ue4 - Remove bCrouchMovesCharacterDown_DEPRECATED, TaggedPropertyRedirects handles this automatically.

diff --git a/Engine/Source/Runtime/Engine/Classes/GameFramework/CharacterMovementComponent.h b/Engine/Source/Runtime/Engine/Classes/GameFramework/CharacterMovementComponent.h
--- a/Engine/Source/Runtime/Engine/Classes/GameFramework/CharacterMovementComponent.h
+++ b/Engine/Source/Runtime/Engine/Classes/GameFramework/CharacterMovementComponent.h
@@ -377,9 +377,6 @@ public:
 public:
 	// Deprecated properties
 
-	UPROPERTY()
-	uint32 bCrouchMovesCharacterDown_DEPRECATED:1;
-
 	UPROPERTY()
 	uint32 bForceBraking_DEPRECATED:1;
 
diff --git a/Engine/Source/Runtime/Engine/Private/CharacterMovementComponent.cpp b/Engine/Source/Runtime/Engine/Private/CharacterMovementComponent.cpp
--- a/Engine/Source/Runtime/Engine/Private/CharacterMovementComponent.cpp
+++ b/Engine/Source/Runtime/Engine/Private/CharacterMovementComponent.cpp
@@ -199,7 +199,6 @@ UCharacterMovementComponent::UCharacterMovementComponent(const class FPostConstr
 	DefaultWaterMovementMode = MOVE_Swimming;
 	DefaultLandMovementMode = MOVE_Walking;
 	bForceNextFloorCheck = true;
-	bCrouchMovesCharacterDown_DEPRECATED = false;
 	bForceBraking_DEPRECATED = false;
 	bShrinkProxyCapsule = true;
 	bCanWalkOffLedges = true;
@@ -286,11 +285,6 @@ void UCharacterMovementComponent::PostLoad()
 		MaxWalkSpeedCrouched = MaxWalkSpeed * CrouchedSpeedMultiplier_DEPRECATED;
 		MaxCustomMovementSpeed = MaxWalkSpeed;
 	}
-
-	if (LinkerUE4Ver < VER_UE4_RENAME_CROUCHMOVESCHARACTERDOWN)
-	{
-		bCrouchMaintainsBaseLocation = bCrouchMovesCharacterDown_DEPRECATED;
-	}
 }
 
 
