Fixing crash with renaming dynamic parameters

diff --git a/Engine/Source/Runtime/Engine/Private/Materials/Material.cpp b/Engine/Source/Runtime/Engine/Private/Materials/Material.cpp
--- a/Engine/Source/Runtime/Engine/Private/Materials/Material.cpp
+++ b/Engine/Source/Runtime/Engine/Private/Materials/Material.cpp
@@ -2514,7 +2514,7 @@ void UMaterial::UpdateExpressionDynamicParameterNames(const UMaterialExpression*
 		for (int32 ExpIndex = 0; ExpIndex < Expressions.Num(); ExpIndex++)
 		{
 			UMaterialExpressionDynamicParameter* CheckParam = Cast<UMaterialExpressionDynamicParameter>(Expressions[ExpIndex]);
-			if (CheckParam->CopyDynamicParameterNames(DynParam))
+			if (CheckParam && CheckParam->CopyDynamicParameterNames(DynParam))
 			{
 #if WITH_EDITOR
 				CheckParam->GraphNode->ReconstructNode();
