- Added back in the fix for GetHMDMonitorInfo() returning false.

diff --git a/Engine/Source/Runtime/Engine/Private/Slate/SceneViewport.cpp b/Engine/Source/Runtime/Engine/Private/Slate/SceneViewport.cpp
--- a/Engine/Source/Runtime/Engine/Private/Slate/SceneViewport.cpp
+++ b/Engine/Source/Runtime/Engine/Private/Slate/SceneViewport.cpp
@@ -846,17 +846,19 @@ void FSceneViewport::ResizeFrame(uint32 NewSizeX, uint32 NewSizeY, EWindowMode::
 						FSlateRect PreFullScreenRect = WindowToResize->GetRectInScreen();
 
 						IHeadMountedDisplay::MonitorInfo MonitorInfo;
-						GEngine->HMDDevice->GetHMDMonitorInfo(MonitorInfo);
-						NewSizeX = MonitorInfo.ResolutionX;
-						NewSizeY = MonitorInfo.ResolutionY;
-						if (GEngine->HMDDevice->IsFullScreenAllowed())
+						if (GEngine->HMDDevice->GetHMDMonitorInfo(MonitorInfo))
 						{
-							WindowToResize->ReshapeWindow(FVector2D(MonitorInfo.DesktopX, MonitorInfo.DesktopY), FVector2D(MonitorInfo.ResolutionX, MonitorInfo.ResolutionY));
-						}
-						else
-						{
-							WindowToResize->Resize(FVector2D(MonitorInfo.ResolutionX, MonitorInfo.ResolutionY));
-							DesiredWindowMode = EWindowMode::WindowedMirror;
+							NewSizeX = MonitorInfo.ResolutionX;
+							NewSizeY = MonitorInfo.ResolutionY;
+							if (GEngine->HMDDevice->IsFullScreenAllowed())
+							{
+								WindowToResize->ReshapeWindow(FVector2D(MonitorInfo.DesktopX, MonitorInfo.DesktopY), FVector2D(MonitorInfo.ResolutionX, MonitorInfo.ResolutionY));
+							}
+							else
+							{
+								WindowToResize->Resize(FVector2D(MonitorInfo.ResolutionX, MonitorInfo.ResolutionY));
+								DesiredWindowMode = EWindowMode::WindowedMirror;
+							}
 						}
 
 						GEngine->HMDDevice->PushPreFullScreenRect(PreFullScreenRect);
