Fixed debug rendering for NavMesh collision geometry #ue4

diff --git a/Engine/Source/Runtime/Engine/Private/NavMeshRenderingComponent.cpp b/Engine/Source/Runtime/Engine/Private/NavMeshRenderingComponent.cpp
--- a/Engine/Source/Runtime/Engine/Private/NavMeshRenderingComponent.cpp
+++ b/Engine/Source/Runtime/Engine/Private/NavMeshRenderingComponent.cpp
@@ -73,6 +73,19 @@ FBoxSphereBounds UNavMeshRenderingComponent::CalcBounds(const FTransform & Local
 	if (NavMesh)
 	{
 		BoundingBox = NavMesh->GetNavMeshBounds();
+		if (NavMesh->bDrawOctree)
+		{
+			const UNavigationSystem* NavSys = UNavigationSystem::GetCurrent(GetWorld());
+			const FNavigationOctree* NavOctree = NavSys ? NavSys->GetNavOctree() : NULL;
+			if (NavOctree)
+			{
+				for (FNavigationOctree::TConstIterator<> NodeIt(*NavOctree); NodeIt.HasPendingNodes(); NodeIt.Advance())
+				{
+					const FOctreeNodeContext& CorrentContext = NodeIt.GetCurrentContext();
+					BoundingBox += CorrentContext.Bounds.GetBox();
+				}
+			}
+		}
 	}
 #endif
 	return FBoxSphereBounds(BoundingBox);
@@ -115,21 +128,20 @@ void UNavMeshRenderingComponent::GatherData(struct FNavMeshSceneProxyData* Curre
 		{
 			const UNavigationSystem* NavSys = UNavigationSystem::GetCurrent(GetWorld());
 			const FNavigationOctree* NavOctree = NavSys ? NavSys->GetNavOctree() : NULL;
-
-			for (FNavigationOctree::TConstIterator<> NodeIt(*NavOctree); NodeIt.HasPendingNodes(); NodeIt.Advance())
+			if (NavOctree)
 			{
-				const FNavigationOctree::FNode& CurrentNode = NodeIt.GetCurrentNode();
-				const int32 CurrentNodeElementCount = CurrentNode.GetElementCount();
-				const FOctreeNodeContext& CorrentContext = NodeIt.GetCurrentContext();
-				CurrentData->OctreeBounds.Add(CorrentContext.Bounds);
-
-				FOREACH_OCTREE_CHILD_NODE(ChildRef)
+				for (FNavigationOctree::TConstIterator<> NodeIt(*NavOctree); NodeIt.HasPendingNodes(); NodeIt.Advance())
 				{
-					auto Context = CorrentContext.GetChildContext(ChildRef);
-					auto Child = CurrentNode.GetChild(ChildRef);
-					if (CurrentNode.HasChild(ChildRef))
+					const FNavigationOctree::FNode& CurrentNode = NodeIt.GetCurrentNode();
+					const FOctreeNodeContext& CorrentContext = NodeIt.GetCurrentContext();
+					CurrentData->OctreeBounds.Add(CorrentContext.Bounds);
+
+					FOREACH_OCTREE_CHILD_NODE(ChildRef)
 					{
-						NodeIt.PushChild(ChildRef);
+						if (CurrentNode.HasChild(ChildRef))
+						{
+							NodeIt.PushChild(ChildRef);
+						}
 					}
 				}
 			}
@@ -256,6 +268,14 @@ void UNavMeshRenderingComponent::GatherData(struct FNavMeshSceneProxyData* Curre
 							CachedGeometry.Verts, CachedGeometry.Header.NumVerts, CachedGeometry.Indices, CachedGeometry.Header.NumFaces);
 					}
 				}
+
+				FOREACH_OCTREE_CHILD_NODE(ChildRef)
+				{
+					if (Node.HasChild(ChildRef))
+					{
+						It.PushChild(ChildRef);
+					}
+				}
 			}
 		}
 		
