Another refactor to simplify physics simulate of skeletal mesh

diff --git a/Engine/Source/Runtime/Engine/Classes/PhysicsEngine/BodyInstance.h b/Engine/Source/Runtime/Engine/Classes/PhysicsEngine/BodyInstance.h
--- a/Engine/Source/Runtime/Engine/Classes/PhysicsEngine/BodyInstance.h
+++ b/Engine/Source/Runtime/Engine/Classes/PhysicsEngine/BodyInstance.h
@@ -396,11 +396,11 @@ public:
 	/** Set this body to be fixed (kinematic) or not. */
 	void SetInstanceSimulatePhysics(bool bSimulate, bool bMaintainPhysicsBlending=false);
 	/** Makes sure the current kinematic state matches the simulate flag */
-	void UpdateInstanceSimulatePhysics(bool bIgnoreOwner = false);
+	void UpdateInstanceSimulatePhysics();
 	/** Returns true if this body is simulating, false if it is fixed (kinematic) */
-	bool IsInstanceSimulatingPhysics(bool bIgnoreOwner = false);
+	bool IsInstanceSimulatingPhysics();
 	/** Should Simulate Physics **/
-	bool ShouldInstanceSimulatingPhysics(bool bIgnoreOwner=false);
+	bool ShouldInstanceSimulatingPhysics();
 	/** Returns whether this body is awake */
 	bool IsInstanceAwake() const;
 	/** Wake this body */
diff --git a/Engine/Source/Runtime/Engine/Private/PhysicsEngine/BodyInstance.cpp b/Engine/Source/Runtime/Engine/Private/PhysicsEngine/BodyInstance.cpp
--- a/Engine/Source/Runtime/Engine/Private/PhysicsEngine/BodyInstance.cpp
+++ b/Engine/Source/Runtime/Engine/Private/PhysicsEngine/BodyInstance.cpp
@@ -1955,11 +1955,11 @@ bool FBodyInstance::UpdateBodyScale(const FVector& InScale3D)
 	return bSuccess;
 }
 
-void FBodyInstance::UpdateInstanceSimulatePhysics(bool bIgnoreParent)
+void FBodyInstance::UpdateInstanceSimulatePhysics()
 {
 	// In skeletal case, we need both our bone and skelcomponent flag to be true.
 	// This might be 'and'ing us with ourself, but thats fine.
-	const bool bUseSimulate = IsInstanceSimulatingPhysics(bIgnoreParent);
+	const bool bUseSimulate = IsInstanceSimulatingPhysics();
 	bool bInitialized = false;
 #if WITH_PHYSX
 	if (PxRigidDynamic* PRigidDynamic = GetPxRigidDynamic())
@@ -2061,26 +2061,17 @@ void FBodyInstance::SetInstanceSimulatePhysics(bool bSimulate, bool bMaintainPhy
 		}
 	}
 
-	UpdateInstanceSimulatePhysics(true);
+	UpdateInstanceSimulatePhysics();
 }
 
-bool FBodyInstance::IsInstanceSimulatingPhysics(bool bIgnoreOwner)
+bool FBodyInstance::IsInstanceSimulatingPhysics()
 {
 	// if I'm simulating or owner is simulating
-	return ShouldInstanceSimulatingPhysics(bIgnoreOwner) && IsValidBodyInstance();
+	return ShouldInstanceSimulatingPhysics() && IsValidBodyInstance();
 }
 
-bool FBodyInstance::ShouldInstanceSimulatingPhysics(bool bIgnoreOwner)
+bool FBodyInstance::ShouldInstanceSimulatingPhysics()
 {
-	// if I'm simulating or owner is simulating
-	if (OwnerComponent != NULL && BodySetup.IsValid() && BodySetup.Get()->PhysicsType == PhysType_Default && !bIgnoreOwner)
-	{
-		// if derive from owner, and owner is simulating, this should simulate
-		return OwnerComponent->BodyInstance.bSimulatePhysics;
-
-		// or else, it should look its own setting
-	}
-
 	return bSimulatePhysics;
 }
 
diff --git a/Engine/Source/Runtime/Engine/Private/PhysicsEngine/PhysAnim.cpp b/Engine/Source/Runtime/Engine/Private/PhysicsEngine/PhysAnim.cpp
--- a/Engine/Source/Runtime/Engine/Private/PhysicsEngine/PhysAnim.cpp
+++ b/Engine/Source/Runtime/Engine/Private/PhysicsEngine/PhysAnim.cpp
@@ -109,7 +109,7 @@ void USkeletalMeshComponent::BlendPhysicsBones( TArray<FBoneIndexType>& InRequir
 			BodyInstance = Bodies[BodyIndex];
 
 			//if simulated body copy back and blend with animation
-			if(BodyInstance->IsInstanceSimulatingPhysics(true))
+			if(BodyInstance->IsInstanceSimulatingPhysics())
 			{
 				FTransform PhysTM = BodyInstance->GetUnrealWorldTransform();
 
@@ -312,7 +312,7 @@ void USkeletalMeshComponent::UpdateKinematicBonesToPhysics(bool bTeleport, bool
 				FBodyInstance* BodyInst = Bodies[i];
 				check(BodyInst);
 
-				if (bForceUpdate || (BodyInst->IsValidBodyInstance() && !BodyInst->IsInstanceSimulatingPhysics(true)))
+				if (bForceUpdate || (BodyInst->IsValidBodyInstance() && !BodyInst->IsInstanceSimulatingPhysics()))
 				{
 					// Find the graphics bone index that corresponds to this physics body.
 					FName const BodyName = PhysicsAsset->BodySetup[i]->BoneName;
diff --git a/Engine/Source/Runtime/Engine/Private/SkeletalMeshComponentPhysics.cpp b/Engine/Source/Runtime/Engine/Private/SkeletalMeshComponentPhysics.cpp
--- a/Engine/Source/Runtime/Engine/Private/SkeletalMeshComponentPhysics.cpp
+++ b/Engine/Source/Runtime/Engine/Private/SkeletalMeshComponentPhysics.cpp
@@ -506,17 +506,27 @@ void USkeletalMeshComponent::SetSimulatePhysics(bool bSimulate)
 		return;
 	}
 
-	// skeletalmeshcomponent BodyInstance is data class
-	// we don't instantiate BodyInstance but Bodies
-	// however this information is used for Owner data
 	BodyInstance.bSimulatePhysics = bSimulate;
 
 	// enable blending physics
 	bBlendPhysics = bSimulate;
 
-	for(int32 i=0; i<Bodies.Num(); i++)
+	//Go through body setups and see which bodies should be turned on and off
+	if (UPhysicsAsset * PhysAsset = GetPhysicsAsset())
 	{
-		Bodies[i]->UpdateInstanceSimulatePhysics();
+		for (int32 BodyIdx = 0; BodyIdx < Bodies.Num(); ++BodyIdx)
+		{
+			if (FBodyInstance * BodyInstance = Bodies[BodyIdx])
+			{
+				if (UBodySetup * BodySetup = PhysAsset->BodySetup[BodyIdx])
+				{
+					if (BodySetup->PhysicsType == EPhysicsType::PhysType_Default)
+					{
+						BodyInstance->SetInstanceSimulatePhysics(bSimulate);
+					}
+				}
+			}
+		}
 	}
 }
 
