* exposed project defaults to allow some postprocess settings to be anled/disabled by default

diff --git a/Engine/Source/Runtime/Engine/Classes/Engine/RendererSettings.h b/Engine/Source/Runtime/Engine/Classes/Engine/RendererSettings.h
--- a/Engine/Source/Runtime/Engine/Classes/Engine/RendererSettings.h
+++ b/Engine/Source/Runtime/Engine/Classes/Engine/RendererSettings.h
@@ -152,6 +152,31 @@ class ENGINE_API URendererSettings
 		ToolTip="Whether the custom depth pass for tagging primitives for postprocessing passes is enabled. Enabling it on demand can save memory but may cause a hitch the first time the feature is used."))
 	TEnumAsByte<ECustomDepth::Type> CustomDepth;
 
+	UPROPERTY(config, EditAnywhere, Category = Postprocessing, meta = (
+		ConsoleVariable = "r.DefaultFeature.Bloom", DisplayName = "Bloom is enabled by default",
+		ToolTip = "Whether the default for Bloom is enabled or not (postprocess volume/camera/game setting still can enable)"))
+	uint32 bDefaultFeatureBloom : 1;
+
+	UPROPERTY(config, EditAnywhere, Category = Postprocessing, meta = (
+		ConsoleVariable = "r.DefaultFeature.AmbientOcclusion", DisplayName = "Ambient Occlusion is enabled by default",
+		ToolTip = "Whether the default for AmbientOcclusion is enabled or not (postprocess volume/camera/game setting still can enable)"))
+	uint32 bDefaultFeatureAmbientOcclusion : 1;
+
+	UPROPERTY(config, EditAnywhere, Category = Postprocessing, meta = (
+		ConsoleVariable = "r.DefaultFeature.AutoExposure", DisplayName = "Auto Exposure is enabled by default",
+		ToolTip = "Whether the default for AutoExposure is enabled or not (postprocess volume/camera/game setting still can enable)"))
+	uint32 bDefaultFeatureAutoExposure : 1;
+
+	UPROPERTY(config, EditAnywhere, Category = Postprocessing, meta = (
+		ConsoleVariable = "r.DefaultFeature.MotionBlur", DisplayName = "Motion Blur is enabled by default",
+		ToolTip = "Whether the default for MotionBlur is enabled or not (postprocess volume/camera/game setting still can enable)"))
+	uint32 bDefaultFeatureMotionBlur : 1;
+
+	UPROPERTY(config, EditAnywhere, Category = Postprocessing, meta = (
+		ConsoleVariable = "r.DefaultFeature.LensFlare", DisplayName = "Lens Flares (Image based) are enabled by default",
+		ToolTip = "Whether the default for LensFlare is enabled or not (postprocess volume/camera/game setting still can enable)"))
+	uint32 bDefaultFeatureLensFlare : 1;
+
 	UPROPERTY(config, EditAnywhere, Category=Optimizations, meta=(
 		ConsoleVariable="r.EarlyZPass",DisplayName="Early Z-pass",
 		ToolTip="Whether to use a depth only pass to initialize Z culling for the base pass."))
diff --git a/Engine/Source/Runtime/Engine/Classes/Engine/Scene.h b/Engine/Source/Runtime/Engine/Classes/Engine/Scene.h
--- a/Engine/Source/Runtime/Engine/Classes/Engine/Scene.h
+++ b/Engine/Source/Runtime/Engine/Classes/Engine/Scene.h
@@ -794,6 +794,7 @@ struct FPostProcessSettings
 		SceneColorTint = FLinearColor(1, 1, 1);
 		SceneFringeIntensity = 0.0f;
 		SceneFringeSaturation = 0.5f;
+		// next value might get overwritten by r.DefaultFeature.Bloom
 		BloomIntensity = 1.0f;
 		BloomThreshold = 1.0f;
 		Bloom1Tint = FLinearColor(0.5f, 0.5f, 0.5f);
@@ -821,15 +822,17 @@ struct FPostProcessSettings
 		LPVTransmissionIntensity = 1.0f;
 		AutoExposureLowPercent = 80.0f;
 		AutoExposureHighPercent = 98.3f;
+		// next value might get overwritten by r.DefaultFeature.AutoExposure
 		AutoExposureMinBrightness = 0.03f;
+		// next value might get overwritten by r.DefaultFeature.AutoExposure
 		AutoExposureMaxBrightness = 2.0f;
 		AutoExposureBias = 0.0f;
 		AutoExposureSpeedUp = 3.0f;
 		AutoExposureSpeedDown = 1.0f;
 		HistogramLogMin = -8.0f;
 		HistogramLogMax = 4.0f;
-		// Note: We changed this from 1.0f to 0.0f for a faster default.
-		LensFlareIntensity = 0.0f;
+		// next value might get overwritten by r.DefaultFeature.LensFlare
+		LensFlareIntensity = 1.0f;
 		LensFlareTint = FLinearColor(1.0f, 1.0f, 1.0f);
 		LensFlareBokehSize = 3.0f;
 		LensFlareThreshold = 8.0f;
@@ -837,6 +840,7 @@ struct FPostProcessSettings
 		VignetteColor = FLinearColor(0.0f, 0.0f, 0.0f);
 		GrainIntensity = 0.0f;
 		GrainJitter = 0.0f;
+		// next value might get overwritten by r.DefaultFeature.AmbientOcclusion
 		AmbientOcclusionIntensity = .5f;
 		AmbientOcclusionStaticFraction = 1.0f;
 		AmbientOcclusionRadius = 40.0f;
@@ -873,6 +877,7 @@ struct FPostProcessSettings
 		LensFlareTints[5] = FLinearColor(0.9f, 1.0f, 0.8f, 0.27f);
 		LensFlareTints[6] = FLinearColor(1.0f, 0.8f, 0.4f, 0.22f);
 		LensFlareTints[7] = FLinearColor(0.9f, 0.7f, 0.7f, 0.15f);
+		// next value might get overwritten by r.DefaultFeature.MotionBlur
 		MotionBlurAmount = 0.5f;
 		MotionBlurMax = 5.0f;
 		MotionBlurPerObjectSize = 0.5f;
diff --git a/Engine/Source/Runtime/Engine/Private/SceneView.cpp b/Engine/Source/Runtime/Engine/Private/SceneView.cpp
--- a/Engine/Source/Runtime/Engine/Private/SceneView.cpp
+++ b/Engine/Source/Runtime/Engine/Private/SceneView.cpp
@@ -60,6 +60,44 @@ static TAutoConsoleVariable<float> CVarSSAOFadeRadiusScale(
 	TEXT(" 0.01:smallest .. 1.0:normal (default), <1:smaller, >1:larger"),
 	ECVF_Cheat | ECVF_RenderThreadSafe);
 
+// Engine default (project settings):
+
+static TAutoConsoleVariable<int32> CVarDefaultBloom(
+	TEXT("r.DefaultFeature.Bloom"),
+	1,
+	TEXT("Engine default (project setting) for Bloom is\n")
+	TEXT(" 0: off, set BloomIntensity  to 0 (postprocess volume/camera/game setting still can override)\n")
+	TEXT(" 1: on (default)"));
+
+static TAutoConsoleVariable<int32> CVarDefaultAmbientOcclusion(
+	TEXT("r.DefaultFeature.AmbientOcclusion"),
+	1,
+	TEXT("Engine default (project setting) for AmbientOcclusion is\n")
+	TEXT(" 0: off, sets AmbientOcclusionIntensity to 0 (postprocess volume/camera/game setting still can override)\n")
+	TEXT(" 1: on (default)"));
+
+static TAutoConsoleVariable<int32> CVarDefaultAutoExposure(
+	TEXT("r.DefaultFeature.AutoExposure"),
+	1,
+	TEXT("Engine default (project setting) for AutoExposure is\n")
+	TEXT(" 0: off, sets AutoExposureMinBrightness and AutoExposureMaxBrightness to 1 (postprocess volume/camera/game setting still can override)\n")
+	TEXT(" 1: on (default)"));
+
+static TAutoConsoleVariable<int32> CVarDefaultMotionBlur(
+	TEXT("r.DefaultFeature.MotionBlur"),
+	1,
+	TEXT("Engine default (project setting) for MotionBlur is\n")
+	TEXT(" 0: off, sets MotionBlurAmount to 0 (postprocess volume/camera/game setting still can override)\n")
+	TEXT(" 1: on (default)"));
+
+// off by default for better performance and less distractions
+static TAutoConsoleVariable<int32> CVarDefaultLensFlare(
+	TEXT("r.DefaultFeature.LensFlare"),
+	0,
+	TEXT("Engine default (project setting) for LensFlare is\n")
+	TEXT(" 0: off, sets LensFlareIntensity to 0 (postprocess volume/camera/game setting still can override)\n")
+	TEXT(" 1: on (default)"));
+
 /** Global vertex color view mode setting when SHOW_VertexColors show flag is set */
 EVertexColorViewMode::Type GVertexColorViewMode = EVertexColorViewMode::Color;
 
@@ -809,6 +847,31 @@ void FSceneView::StartFinalPostprocessSettings(FVector InViewLocation)
 	// Set values before any override happens.
 	FinalPostProcessSettings.SetBaseValues();
 
+	// project settings might want to have different defaults
+	{
+		if(!CVarDefaultBloom.GetValueOnGameThread())
+		{
+			FinalPostProcessSettings.BloomIntensity = 0;
+		}
+		if (!CVarDefaultAmbientOcclusion.GetValueOnGameThread())
+		{
+			FinalPostProcessSettings.AmbientOcclusionIntensity = 0;
+		}
+		if (!CVarDefaultAutoExposure.GetValueOnGameThread())
+		{
+			FinalPostProcessSettings.AutoExposureMinBrightness = 1;
+			FinalPostProcessSettings.AutoExposureMaxBrightness = 1;
+		}
+		if (!CVarDefaultMotionBlur.GetValueOnGameThread())
+		{
+			FinalPostProcessSettings.MotionBlurAmount = 0;
+		}
+		if (!CVarDefaultLensFlare.GetValueOnGameThread())
+		{
+			FinalPostProcessSettings.LensFlareIntensity = 0;
+		}
+	}
+
 	if(State)
 	{
 		State->OnStartPostProcessing(*this);
