[343275] fix GC crash caused by invalid ChildActorComponent cached data during reconstruction

- added an FChildActorComponentInstanceData::MatchesComponent() override to include the ChildActorName in the component matching logic for the ChildActorComponent type

diff --git a/Engine/Source/Runtime/Engine/Private/ChildActorComponent.cpp b/Engine/Source/Runtime/Engine/Private/ChildActorComponent.cpp
--- a/Engine/Source/Runtime/Engine/Private/ChildActorComponent.cpp
+++ b/Engine/Source/Runtime/Engine/Private/ChildActorComponent.cpp
@@ -65,6 +65,12 @@ public:
 		}
 	}
 
+	// Overridden to also check the child actor name. Without the additional check, stale cached data may be applied to the child actor instance, which can lead to a crash during GC.
+	virtual bool MatchesComponent(const UActorComponent* Component) const override
+	{
+		return (CastChecked<UChildActorComponent>(Component)->ChildActorName == ChildActorName && FComponentInstanceDataBase::MatchesComponent(Component));
+	}
+
 	FName ChildActorName;
 
 	struct FAttachedActorInfo
