Better support for "EditInline"/"EditinlineNew" feature in UEngine::CopyPropertiesForUnrelatedObjects.
Invalid/useless "EditInline" tags removed. UHT generates error when "EditInline" is used for property other than Object reference (or an array).

diff --git a/Engine/Source/Editor/IntroTutorials/Public/EditorTutorial.h b/Engine/Source/Editor/IntroTutorials/Public/EditorTutorial.h
--- a/Engine/Source/Editor/IntroTutorials/Public/EditorTutorial.h
+++ b/Engine/Source/Editor/IntroTutorials/Public/EditorTutorial.h
@@ -134,7 +134,7 @@ struct INTROTUTORIALS_API FTutorialWidgetContent
 	}
 
 	/** Content to associate with widget */
-	UPROPERTY(EditAnywhere, EditInline, Category="Widget")
+	UPROPERTY(EditAnywhere, Category="Widget")
 	FTutorialContent Content;
 
 	/** Anchor for content widget to highlight */
@@ -165,10 +165,10 @@ struct INTROTUTORIALS_API FTutorialStage
 	UPROPERTY(EditAnywhere, Category="Stage")
 	FName Name;
 
-	UPROPERTY(EditAnywhere, EditInline, Category="Stage")
+	UPROPERTY(EditAnywhere, Category="Stage")
 	FTutorialContent Content;
 
-	UPROPERTY(EditAnywhere, EditInline, Category="Stage")
+	UPROPERTY(EditAnywhere, Category="Stage")
 	TArray<FTutorialWidgetContent> WidgetContent;
 };
 
@@ -195,7 +195,7 @@ class INTROTUTORIALS_API UEditorTutorial : public UObject
 	FTutorialContent SummaryContent;
 
 	/** The various stages of this tutorial */
-	UPROPERTY(EditAnywhere, EditInline, Category="Stages")
+	UPROPERTY(EditAnywhere, Category="Stages")
 	TArray<FTutorialStage> Stages;
 
 	/** Tutorial to optionally chain onto after this tutorial completes */
diff --git a/Engine/Source/Editor/UnrealEd/Classes/MaterialEditor/MaterialEditorInstanceConstant.h b/Engine/Source/Editor/UnrealEd/Classes/MaterialEditor/MaterialEditorInstanceConstant.h
--- a/Engine/Source/Editor/UnrealEd/Classes/MaterialEditor/MaterialEditorInstanceConstant.h
+++ b/Engine/Source/Editor/UnrealEd/Classes/MaterialEditor/MaterialEditorInstanceConstant.h
@@ -212,7 +212,7 @@ class UNREALED_API UMaterialEditorInstanceConstant : public UObject
 	UPROPERTY(EditAnywhere, Category=MaterialEditorInstanceConstant)
 	class UMaterialInterface* Parent;
 
-	UPROPERTY(EditAnywhere, editfixedsize, editinline, Category=MaterialEditorInstanceConstant)
+	UPROPERTY(EditAnywhere, editfixedsize, Category=MaterialEditorInstanceConstant)
 	TArray<struct FEditorParameterGroup> ParameterGroups;
 
 	/** This is the refraction depth bias, larger values offset distortion to prevent closer objects from rendering into the distorted surface at acute viewing angles but increases the disconnect between surface and where the refraction starts. */
diff --git a/Engine/Source/Programs/UnrealHeaderTool/Private/HeaderParser.cpp b/Engine/Source/Programs/UnrealHeaderTool/Private/HeaderParser.cpp
--- a/Engine/Source/Programs/UnrealHeaderTool/Private/HeaderParser.cpp
+++ b/Engine/Source/Programs/UnrealHeaderTool/Private/HeaderParser.cpp
@@ -3268,6 +3268,11 @@ bool FHeaderParser::GetVarType
 	}
 
 	// Perform some more specific validation on the property flags
+	if ((VarProperty.PropertyFlags & CPF_EditInline) && VarProperty.Type != CPT_ObjectReference)
+	{
+		FError::Throwf(TEXT("'EditInline' is only allowed on object property (or array of objects)"));
+	}
+
 	if ( VarProperty.IsObject() && VarProperty.MetaClass == NULL && (VarProperty.PropertyFlags&CPF_Config) != 0 )
 	{
 		FError::Throwf(TEXT("Not allowed to use 'config' with object variables"));
diff --git a/Engine/Source/Runtime/AIModule/Classes/BehaviorTree/BlackboardData.h b/Engine/Source/Runtime/AIModule/Classes/BehaviorTree/BlackboardData.h
--- a/Engine/Source/Runtime/AIModule/Classes/BehaviorTree/BlackboardData.h
+++ b/Engine/Source/Runtime/AIModule/Classes/BehaviorTree/BlackboardData.h
@@ -37,7 +37,7 @@ class AIMODULE_API UBlackboardData : public UDataAsset
 #endif
 
 	/** blackboard keys */
-	UPROPERTY(EditAnywhere, EditInLine, Category=Blackboard)
+	UPROPERTY(EditAnywhere, Category=Blackboard)
 	TArray<FBlackboardEntry> Keys;
 
 	/** @return key ID from name */
diff --git a/Engine/Source/Runtime/Engine/Classes/Components/AudioComponent.h b/Engine/Source/Runtime/Engine/Classes/Components/AudioComponent.h
--- a/Engine/Source/Runtime/Engine/Classes/Components/AudioComponent.h
+++ b/Engine/Source/Runtime/Engine/Classes/Components/AudioComponent.h
@@ -101,7 +101,7 @@ class ENGINE_API UAudioComponent : public USceneComponent
 	class USoundBase* Sound;
 
 	/** Array of per-instance parameters for this AudioComponent. */
-	UPROPERTY(EditAnywhere, BlueprintReadWrite, editinline, Category=Sound, AdvancedDisplay)
+	UPROPERTY(EditAnywhere, BlueprintReadWrite, Category=Sound, AdvancedDisplay)
 	TArray<struct FAudioComponentParam> InstanceParameters;
 
 	/** Optional sound group this AudioComponent belongs to */
diff --git a/Engine/Source/Runtime/Engine/Classes/Engine/DestructibleFractureSettings.h b/Engine/Source/Runtime/Engine/Classes/Engine/DestructibleFractureSettings.h
--- a/Engine/Source/Runtime/Engine/Classes/Engine/DestructibleFractureSettings.h
+++ b/Engine/Source/Runtime/Engine/Classes/Engine/DestructibleFractureSettings.h
@@ -157,7 +157,7 @@ class UDestructibleFractureSettings : public UObject
 	int32									CellSiteCount;
 
 	/** Stored interior material data.  Just need one as we only support Voronoi splitting. */
-	UPROPERTY(EditAnywhere, transient, editinline, Category=General)
+	UPROPERTY(EditAnywhere, transient, Category=General)
 	FFractureMaterial						FractureMaterialDesc;
 
 	/** Random seed for reproducibility */
diff --git a/Engine/Source/Runtime/Engine/Classes/Engine/Font.h b/Engine/Source/Runtime/Engine/Classes/Engine/Font.h
--- a/Engine/Source/Runtime/Engine/Classes/Engine/Font.h
+++ b/Engine/Source/Runtime/Engine/Classes/Engine/Font.h
@@ -81,7 +81,7 @@ class UFont : public UObject
 
 	/** List of characters in the font.  For a MultiFont, this will include all characters in all sub-fonts!  Thus,
 		the number of characters in this array isn't necessary the number of characters available in the font */
-	UPROPERTY(EditAnywhere, editinline, Category=Font)
+	UPROPERTY(EditAnywhere, Category=Font)
 	TArray<struct FFontCharacter> Characters;
 
 	/** Textures that store this font's glyph image data */
diff --git a/Engine/Source/Runtime/Engine/Classes/Matinee/InterpTrackParticleReplay.h b/Engine/Source/Runtime/Engine/Classes/Matinee/InterpTrackParticleReplay.h
--- a/Engine/Source/Runtime/Engine/Classes/Matinee/InterpTrackParticleReplay.h
+++ b/Engine/Source/Runtime/Engine/Classes/Matinee/InterpTrackParticleReplay.h
@@ -42,7 +42,7 @@ class UInterpTrackParticleReplay : public UInterpTrack
 	GENERATED_UCLASS_BODY()
 
 	/** Array of keys */
-	UPROPERTY(editinline)
+	UPROPERTY()
 	TArray<struct FParticleReplayTrackKey> TrackKeys;
 
 #if WITH_EDITORONLY_DATA
diff --git a/Engine/Source/Runtime/Engine/Classes/Particles/ParticleSystemComponent.h b/Engine/Source/Runtime/Engine/Classes/Particles/ParticleSystemComponent.h
--- a/Engine/Source/Runtime/Engine/Classes/Particles/ParticleSystemComponent.h
+++ b/Engine/Source/Runtime/Engine/Classes/Particles/ParticleSystemComponent.h
@@ -295,7 +295,7 @@ public:
 	 *	Array holding name instance parameters for this ParticleSystemComponent.
 	 *	Parameters can be used in Cascade using DistributionFloat/VectorParticleParameters.
 	 */
-	UPROPERTY(EditAnywhere, BlueprintReadWrite, editinline, Category=Particles)
+	UPROPERTY(EditAnywhere, BlueprintReadWrite, Category=Particles)
 	TArray<struct FParticleSysParam> InstanceParameters;
 
 	UPROPERTY(BlueprintAssignable)
diff --git a/Engine/Source/Runtime/Engine/Private/UnrealEngine.cpp b/Engine/Source/Runtime/Engine/Private/UnrealEngine.cpp
--- a/Engine/Source/Runtime/Engine/Private/UnrealEngine.cpp
+++ b/Engine/Source/Runtime/Engine/Private/UnrealEngine.cpp
@@ -9583,35 +9583,72 @@ static TAutoConsoleVariable<int32> CVarDumpCopyPropertiesForUnrelatedObjects(
 struct FFindEditInlineSubobjectHelper
 {
 	template<typename T>
-	static void Get(const UObject* InObject, T& OutObjects)
+	static void Get(const UStruct* Struct, const uint8* ContainerAddress, T& OutObjects)
 	{
-		const UClass* Class = InObject ? InObject->GetClass() : NULL;
-		for (UProperty* Prop = (Class ? Class->RefLink : NULL); Prop; Prop = Prop->NextRef)
+		check(ContainerAddress);
+		for (UProperty* Prop = (Struct ? Struct->RefLink : NULL); Prop; Prop = Prop->NextRef)
 		{
 			if (Prop->HasAnyPropertyFlags(CPF_EditInline))
 			{
+				// Collect Objects referenced in EditInline properties
+
 				auto ObjectProperty = Cast<const UObjectProperty>(Prop);
 				if (ObjectProperty)
 				{
-					auto ObjectValue = ObjectProperty->GetObjectPropertyValue_InContainer(InObject);
-					if (ObjectValue)
+					for (int32 ArrayIdx = 0; ArrayIdx < Prop->ArrayDim; ++ArrayIdx)
 					{
-						OutObjects.Add(ObjectValue);
+						auto ObjectValue = ObjectProperty->GetObjectPropertyValue_InContainer(ContainerAddress, ArrayIdx);
+						if (ObjectValue)
+						{
+							OutObjects.Add(ObjectValue);
+						}
 					}
-					continue;
 				}
+				else
+				{
+					auto ArrayProperty = Cast<const UArrayProperty>(Prop);
+					auto InnerObjectProperty = ArrayProperty ? Cast<const UObjectProperty>(ArrayProperty->Inner) : NULL;
+					if (InnerObjectProperty)
+					{
+						FScriptArrayHelper_InContainer ArrayHelper(ArrayProperty, ContainerAddress);
+						for (int32 ElementIndex = 0; ElementIndex < ArrayHelper.Num(); ++ElementIndex)
+						{
+							UObject** ObjectValuePtr = reinterpret_cast<UObject**>(ArrayHelper.GetRawPtr(ElementIndex));
+							if (ObjectValuePtr && *ObjectValuePtr)
+							{
+								OutObjects.Add(*ObjectValuePtr);
+							}
+						}
+					}
+				}
+			}
+			else
+			{
+				// Look for EditInline properties in structs varaibles
 
-				auto ArrayProperty = Cast<const UArrayProperty>(Prop);
-				auto InnerObjectProperty = ArrayProperty ? Cast<const UObjectProperty>(ArrayProperty->Inner) : NULL;
-				if (InnerObjectProperty)
+				auto StructProperty = Cast<const UStructProperty>(Prop);
+				if (StructProperty && StructProperty->Struct)
 				{
-					FScriptArrayHelper_InContainer ArrayHelper(ArrayProperty, InObject);
-					for (int32 ElementIndex = 0; ElementIndex < ArrayHelper.Num(); ++ElementIndex)
+					for (int32 ArrayIdx = 0; ArrayIdx < Prop->ArrayDim; ++ArrayIdx)
 					{
-						UObject** ObjectValuePtr = (UObject**)(ArrayHelper.GetRawPtr(ElementIndex));
-						if (ObjectValuePtr && *ObjectValuePtr)
+						const uint8* ValueAddress = StructProperty->ContainerPtrToValuePtr<uint8>(ContainerAddress, ArrayIdx);
+						Get(StructProperty->Struct, ValueAddress, OutObjects);
+					}
+				}
+				else
+				{
+					auto ArrayProperty = Cast<const UArrayProperty>(Prop);
+					auto InnerStructProperty = ArrayProperty ? Cast<const UStructProperty>(ArrayProperty->Inner) : NULL;
+					if (InnerStructProperty && InnerStructProperty->Struct)
+					{
+						FScriptArrayHelper_InContainer ArrayHelper(ArrayProperty, ContainerAddress);
+						for (int32 ElementIndex = 0; ElementIndex < ArrayHelper.Num(); ++ElementIndex)
 						{
-							OutObjects.Add(*ObjectValuePtr);
+							const uint8* ValueAddress = ArrayHelper.GetRawPtr(ElementIndex);
+							if (ValueAddress)
+							{
+								Get(InnerStructProperty->Struct, ValueAddress, OutObjects);
+							}
 						}
 					}
 				}
@@ -9758,11 +9795,11 @@ void UEngine::CopyPropertiesForUnrelatedObjects(UObject* OldObject, UObject* New
 
 	{
 		TSet<UObject*> OldEditInlineObjects;
-		FFindEditInlineSubobjectHelper::Get(OldObject, OldEditInlineObjects);
+		FFindEditInlineSubobjectHelper::Get(OldObject->GetClass(), reinterpret_cast<uint8*>(OldObject), OldEditInlineObjects);
 		if (OldEditInlineObjects.Num())
 		{
-			TArray<UObject*> NewEditInlineObjects;
-			FFindEditInlineSubobjectHelper::Get(NewObject, NewEditInlineObjects);
+			TSet<UObject*> NewEditInlineObjects;
+			FFindEditInlineSubobjectHelper::Get(NewObject->GetClass(), reinterpret_cast<uint8*>(NewObject), NewEditInlineObjects);
 			// if object exist in both sets, it's outer is old object, and its class is "editinline"
 			for (auto Obj : NewEditInlineObjects)
 			{
