Don't crash when releasing mouse outside of viewport when handling mouse click events

diff --git a/Engine/Source/Runtime/Engine/Private/PlayerController.cpp b/Engine/Source/Runtime/Engine/Private/PlayerController.cpp
--- a/Engine/Source/Runtime/Engine/Private/PlayerController.cpp
+++ b/Engine/Source/Runtime/Engine/Private/PlayerController.cpp
@@ -2141,45 +2141,47 @@ bool APlayerController::InputKey(FKey Key, EInputEvent EventType, float AmountDe
 		if (bEnableClickEvents && Key == EKeys::LeftMouseButton)
 		{
 			FVector2D MousePosition;
-			verify(CastChecked<ULocalPlayer>(Player)->ViewportClient->GetMousePosition(MousePosition));
-			UPrimitiveComponent* ClickedPrimitive = NULL;
-			if (bEnableMouseOverEvents)
+			if (CastChecked<ULocalPlayer>(Player)->ViewportClient->GetMousePosition(MousePosition))
 			{
-				ClickedPrimitive = CurrentClickablePrimitive.Get();
-			}
-			else
-			{
-				FHitResult HitResult;
-				const bool bHit = GetHitResultAtScreenPosition(MousePosition, CurrentClickTraceChannel, true, HitResult);
-				if (bHit)
+				UPrimitiveComponent* ClickedPrimitive = NULL;
+				if (bEnableMouseOverEvents)
 				{
-					ClickedPrimitive = HitResult.Component.Get();
+					ClickedPrimitive = CurrentClickablePrimitive.Get();
 				}
-			}
-			if( GetHUD() )
-			{
-				if (GetHUD()->UpdateAndDispatchHitBoxClickEvents(MousePosition, EventType, false))
+				else
 				{
-					ClickedPrimitive = NULL;
+					FHitResult HitResult;
+					const bool bHit = GetHitResultAtScreenPosition(MousePosition, CurrentClickTraceChannel, true, HitResult);
+					if (bHit)
+					{
+						ClickedPrimitive = HitResult.Component.Get();
+					}
 				}
-			}
-
-			if (ClickedPrimitive)
-			{
-				switch(EventType)
+				if( GetHUD() )
 				{
-				case IE_Pressed:
-				case IE_DoubleClick:
-					ClickedPrimitive->DispatchOnClicked();
-					break;
+					if (GetHUD()->UpdateAndDispatchHitBoxClickEvents(MousePosition, EventType, false))
+					{
+						ClickedPrimitive = NULL;
+					}
+				}
 
-				case IE_Released:
-					ClickedPrimitive->DispatchOnReleased();
-					break;
+				if (ClickedPrimitive)
+				{
+					switch(EventType)
+					{
+					case IE_Pressed:
+					case IE_DoubleClick:
+						ClickedPrimitive->DispatchOnClicked();
+						break;
+
+					case IE_Released:
+						ClickedPrimitive->DispatchOnReleased();
+						break;
+					}
 				}
-			}
 
-			bResult = true;
+				bResult = true;
+			}
 		}
 	}
 
