Remove the new mid-load GC.  Can cause crashes if loading process killed any temporary render resources and marked the objects as dirty.

diff --git a/Engine/Source/Runtime/Engine/Private/UnrealEngine.cpp b/Engine/Source/Runtime/Engine/Private/UnrealEngine.cpp
--- a/Engine/Source/Runtime/Engine/Private/UnrealEngine.cpp
+++ b/Engine/Source/Runtime/Engine/Private/UnrealEngine.cpp
@@ -8639,10 +8639,7 @@ bool UEngine::LoadMap( FWorldContext& WorldContext, FURL URL, class UPendingNetG
 	WorldContext.World()->CreateAISystem();
 
 	// Initialize gameplay for the level.
-	WorldContext.World()->InitializeActorsForPlay(URL);
-
-	// Reclaim some memory allocated during load/initialize. This is needed on low-memory devices.
-	CollectGarbage( GARBAGE_COLLECTION_KEEPFLAGS );
+	WorldContext.World()->InitializeActorsForPlay(URL);		
 
 	// Remember the URL. Put this before spawning player controllers so that
 	// a player controller can get the map name during initialization and
