Fix categorization of Make Noise

diff --git a/Engine/Source/Runtime/Engine/Classes/GameFramework/Actor.h b/Engine/Source/Runtime/Engine/Classes/GameFramework/Actor.h
--- a/Engine/Source/Runtime/Engine/Classes/GameFramework/Actor.h
+++ b/Engine/Source/Runtime/Engine/Classes/GameFramework/Actor.h
@@ -837,7 +837,7 @@ public:
 	 * @param NoiseInstigator - Pawn responsible for this noise.  Uses the actor's Instigator if NoiseInstigator=NULL
 	 * @param NoiseLocation - Position of noise source.  If zero vector, use the actor's location.
 	*/
-	UFUNCTION(BlueprintCallable, BlueprintAuthorityOnly, Category="Audio", meta=(BlueprintProtected = "true"))
+	UFUNCTION(BlueprintCallable, BlueprintAuthorityOnly, Category="AI", meta=(BlueprintProtected = "true"))
 	void MakeNoise(float Loudness=1.f, APawn* NoiseInstigator=NULL, FVector NoiseLocation=FVector::ZeroVector);
 
 	//=============================================================================
