* added newproject setting: AntiAliasing method

diff --git a/Engine/Source/Runtime/Engine/Classes/Engine/RendererSettings.h b/Engine/Source/Runtime/Engine/Classes/Engine/RendererSettings.h
--- a/Engine/Source/Runtime/Engine/Classes/Engine/RendererSettings.h
+++ b/Engine/Source/Runtime/Engine/Classes/Engine/RendererSettings.h
@@ -67,6 +67,20 @@ namespace EEarlyZPass
 }
 
 
+/** used by FPostProcessSettings Anti-aliasing */
+UENUM()
+namespace EAntiAliasingMethodUI
+{
+	enum Type
+	{
+		AAM_None UMETA(DisplayName = "None"),
+		AAM_FXAA UMETA(DisplayName = "FXAA"),
+		AAM_TemporalAA UMETA(DisplayName = "TemporalAA"),
+		AAM_MAX,
+	};
+}
+
+
 /**
  * Implements project settings for the Rendering sub-system.
  */
@@ -177,7 +191,12 @@ class ENGINE_API URendererSettings
 		ToolTip = "Whether the default for LensFlare is enabled or not (postprocess volume/camera/game setting still can enable)"))
 	uint32 bDefaultFeatureLensFlare : 1;
 
-	UPROPERTY(config, EditAnywhere, Category=Optimizations, meta=(
+	UPROPERTY(config, EditAnywhere, Category = Postprocessing, meta = (
+		ConsoleVariable = "r.DefaultFeature.AntiAliasing", DisplayName = "Anti Aliasing method that is used by default",
+		ToolTip = "What anti-aliasing mode is used by default (postprocess volume/camera/game setting still can enable)"))
+	TEnumAsByte<EAntiAliasingMethodUI::Type> DefaultFeatureAntiAliasing;
+
+	UPROPERTY(config, EditAnywhere, Category = Optimizations, meta = (
 		ConsoleVariable="r.EarlyZPass",DisplayName="Early Z-pass",
 		ToolTip="Whether to use a depth only pass to initialize Z culling for the base pass."))
 	TEnumAsByte<EEarlyZPass::Type> EarlyZPass;
diff --git a/Engine/Source/Runtime/Engine/Private/SceneView.cpp b/Engine/Source/Runtime/Engine/Private/SceneView.cpp
--- a/Engine/Source/Runtime/Engine/Private/SceneView.cpp
+++ b/Engine/Source/Runtime/Engine/Private/SceneView.cpp
@@ -65,39 +65,48 @@ static TAutoConsoleVariable<float> CVarSSAOFadeRadiusScale(
 static TAutoConsoleVariable<int32> CVarDefaultBloom(
 	TEXT("r.DefaultFeature.Bloom"),
 	1,
-	TEXT("Engine default (project setting) for Bloom is\n")
-	TEXT(" 0: off, set BloomIntensity  to 0 (postprocess volume/camera/game setting still can override)\n")
+	TEXT("Engine default (project setting) for Bloom is (postprocess volume/camera/game setting still can override)\n")
+	TEXT(" 0: off, set BloomIntensity to 0\n")
 	TEXT(" 1: on (default)"));
 
 static TAutoConsoleVariable<int32> CVarDefaultAmbientOcclusion(
 	TEXT("r.DefaultFeature.AmbientOcclusion"),
 	1,
-	TEXT("Engine default (project setting) for AmbientOcclusion is\n")
-	TEXT(" 0: off, sets AmbientOcclusionIntensity to 0 (postprocess volume/camera/game setting still can override)\n")
+	TEXT("Engine default (project setting) for AmbientOcclusion is (postprocess volume/camera/game setting still can override)\n")
+	TEXT(" 0: off, sets AmbientOcclusionIntensity to 0\n")
 	TEXT(" 1: on (default)"));
 
 static TAutoConsoleVariable<int32> CVarDefaultAutoExposure(
 	TEXT("r.DefaultFeature.AutoExposure"),
 	1,
-	TEXT("Engine default (project setting) for AutoExposure is\n")
-	TEXT(" 0: off, sets AutoExposureMinBrightness and AutoExposureMaxBrightness to 1 (postprocess volume/camera/game setting still can override)\n")
+	TEXT("Engine default (project setting) for AutoExposure is (postprocess volume/camera/game setting still can override)\n")
+	TEXT(" 0: off, sets AutoExposureMinBrightness and AutoExposureMaxBrightness to 1\n")
 	TEXT(" 1: on (default)"));
 
 static TAutoConsoleVariable<int32> CVarDefaultMotionBlur(
 	TEXT("r.DefaultFeature.MotionBlur"),
 	1,
-	TEXT("Engine default (project setting) for MotionBlur is\n")
-	TEXT(" 0: off, sets MotionBlurAmount to 0 (postprocess volume/camera/game setting still can override)\n")
+	TEXT("Engine default (project setting) for MotionBlur is (postprocess volume/camera/game setting still can override)\n")
+	TEXT(" 0: off, sets MotionBlurAmount to 0\n")
 	TEXT(" 1: on (default)"));
 
 // off by default for better performance and less distractions
 static TAutoConsoleVariable<int32> CVarDefaultLensFlare(
 	TEXT("r.DefaultFeature.LensFlare"),
 	0,
-	TEXT("Engine default (project setting) for LensFlare is\n")
-	TEXT(" 0: off, sets LensFlareIntensity to 0 (postprocess volume/camera/game setting still can override)\n")
+	TEXT("Engine default (project setting) for LensFlare is (postprocess volume/camera/game setting still can override)\n")
+	TEXT(" 0: off, sets LensFlareIntensity to 0\n")
 	TEXT(" 1: on (default)"));
 
+// see EAntiAliasingMethod
+static TAutoConsoleVariable<int32> CVarDefaultAntiAliasing(
+	TEXT("r.DefaultFeature.AntiAliasing"),
+	2,
+	TEXT("Engine default (project setting) for AntiAliasingMethod is (postprocess volume/camera/game setting still can override)\n")
+	TEXT(" 0: off (no anti-aliasing)\n")
+	TEXT(" 1: FXAA (faster than TemporalAA but much more shimmering for non static cases)\n")
+	TEXT(" 2: TemporalAA (default)"));
+
 /** Global vertex color view mode setting when SHOW_VertexColors show flag is set */
 EVertexColorViewMode::Type GVertexColorViewMode = EVertexColorViewMode::Color;
 
@@ -870,6 +879,14 @@ void FSceneView::StartFinalPostprocessSettings(FVector InViewLocation)
 		{
 			FinalPostProcessSettings.LensFlareIntensity = 0;
 		}
+
+		{
+			int32 Value = CVarDefaultAntiAliasing.GetValueOnGameThread();
+			if (Value >= 0 && Value < AAM_MAX)
+			{
+				FinalPostProcessSettings.AntiAliasingMethod = (EAntiAliasingMethod)Value;
+			}
+		}
 	}
 
 	if(State)
