Brush.h requires Level.h be included

diff --git a/Engine/Source/Runtime/Engine/Classes/Engine/Brush.h b/Engine/Source/Runtime/Engine/Classes/Engine/Brush.h
--- a/Engine/Source/Runtime/Engine/Classes/Engine/Brush.h
+++ b/Engine/Source/Runtime/Engine/Classes/Engine/Brush.h
@@ -6,6 +6,7 @@
 
 #pragma once
 #include "Model.h"
+#include "Level.h"
 #include "GameFramework/Actor.h"
 #include "Brush.generated.h"
 
