Restoring ability to call two AnimInstance functions from BP.

diff --git a/Engine/Source/Runtime/Engine/Classes/Animation/AnimInstance.h b/Engine/Source/Runtime/Engine/Classes/Animation/AnimInstance.h
--- a/Engine/Source/Runtime/Engine/Classes/Animation/AnimInstance.h
+++ b/Engine/Source/Runtime/Engine/Classes/Animation/AnimInstance.h
@@ -312,11 +312,11 @@ public:
 
 public:
 	/** Returns the owning actor of this AnimInstance */
-	UFUNCTION(Category="Animation")
+	UFUNCTION(BlueprintCallable, Category = "Animation")
 	AActor* GetOwningActor() const;
 	
 	// Returns the skeletal mesh component that has created this AnimInstance
-	UFUNCTION(Category="Animation")
+	UFUNCTION(BlueprintCallable, Category = "Animation")
 	USkeletalMeshComponent* GetOwningComponent() const;
 
 public:
