#ue4 - Header cleanup: fix GetComponentVelocity() comment.

diff --git a/Engine/Source/Runtime/Engine/Classes/Components/SceneComponent.h b/Engine/Source/Runtime/Engine/Classes/Components/SceneComponent.h
--- a/Engine/Source/Runtime/Engine/Classes/Components/SceneComponent.h
+++ b/Engine/Source/Runtime/Engine/Classes/Components/SceneComponent.h
@@ -442,8 +442,8 @@ public:
 	virtual void QuerySupportedSockets(TArray<FComponentSocketDescription>& OutSockets) const;
 
 	/** 
-	 * Get Component Velocity vector
-	 * @return Socket transform in world space if socket is found. Otherwise it will return component's transform in world space.
+	 * Get velocity of this component.
+	 * @return velocity of this component.
 	 */
 	UFUNCTION(BlueprintCallable, Category="Utilities|Transformation")
 	virtual FVector GetComponentVelocity() const;
