[343275] fix GC crash caused by an uninitialized ChildActor instance name

- we now check for ChildActorName == NAME_None in UChildActorComponent::ApplyComponentInstanceData(), and don't rename the "new" ChildActor instance in that case, in order to ensure that the ChildActor will be reinstanced if the ChildActorClass is regenerated. this can happen if the user drags an Actor type containing a ChildActorComponent that hasn't yet had its ChildActorClass property set into the scene, in which case the "old" ChildActor instance along with its name will not have been initialized yet at the start of reconstruction.

diff --git a/Engine/Source/Runtime/Engine/Private/ChildActorComponent.cpp b/Engine/Source/Runtime/Engine/Private/ChildActorComponent.cpp
--- a/Engine/Source/Runtime/Engine/Private/ChildActorComponent.cpp
+++ b/Engine/Source/Runtime/Engine/Private/ChildActorComponent.cpp
@@ -96,10 +96,13 @@ void UChildActorComponent::ApplyComponentInstanceData(TSharedPtr<FComponentInsta
 	if (ChildActor)
 	{
 		// Only rename if it is safe to
-		const FString ChildActorNameString = ChildActorName.ToString();
-		if (ChildActor->Rename(*ChildActorName.ToString(), NULL, REN_Test))
+		if(ChildActorName != NAME_None)
 		{
-			ChildActor->Rename(*ChildActorName.ToString(), NULL, REN_DoNotDirty);
+			const FString ChildActorNameString = ChildActorName.ToString();
+			if (ChildActor->Rename(*ChildActorNameString, NULL, REN_Test))
+			{
+				ChildActor->Rename(*ChildActorNameString, NULL, REN_DoNotDirty);
+			}
 		}
 
 		USceneComponent* ChildActorRoot = ChildActor->GetRootComponent();
