Fix for UDN question; mesh data orientation has no effect when Camera Facing is enabled.

diff --git a/Engine/Source/Runtime/Engine/Private/Particles/ParticleSystemRender.cpp b/Engine/Source/Runtime/Engine/Private/Particles/ParticleSystemRender.cpp
--- a/Engine/Source/Runtime/Engine/Private/Particles/ParticleSystemRender.cpp
+++ b/Engine/Source/Runtime/Engine/Private/Particles/ParticleSystemRender.cpp
@@ -2514,16 +2514,17 @@ void FDynamicMeshEmitterData::GetParticleTransform(const FBaseParticle& InPartic
 	FRotationMatrix kRotMat(kRotator);
 	if (bApplyPreRotation == true)
 	{
+		const FMeshRotationPayloadData* PayloadData = (const FMeshRotationPayloadData*)((const uint8*)&InParticle + Source.MeshRotationOffset);
+		FRotator MeshOrient = FRotator::MakeFromEuler(PayloadData->InitialOrientation);
+		FRotationMatrix OrientMat(MeshOrient);
+
 		if ((bUseCameraFacing == true) || (bUseMeshLockedAxis == true))
 		{
-			OutTransformMat = kScaleMat * FRotationMatrix(kLockedAxisRotator) * kRotMat * kTransMat;
+			OutTransformMat = (OrientMat * kScaleMat) * FRotationMatrix(kLockedAxisRotator) * kRotMat * kTransMat;
 		}
 		else
 		{
-			FMeshRotationPayloadData* PayloadData = (FMeshRotationPayloadData*)((uint8*)&InParticle + Source.MeshRotationOffset);
-			FRotator MeshOrient = FRotator::MakeFromEuler(PayloadData->InitialOrientation);
 			FRotator OtherRot = FRotator::MakeFromEuler(PayloadData->InitRotation+PayloadData->CurContinuousRotation);
-			FRotationMatrix OrientMat(MeshOrient);
 			FRotationMatrix RotMat(OtherRot);
 			OutTransformMat = (OrientMat*kScaleMat) * RotMat * kTransMat;
 		}
