Paper2D: Fix sprites having swapped vertex color channels on iOS and Android devices
- Convert FPaperSpriteVertex and FMaterialSpriteVertex to use FColor instead of FLinearColor
- Add check() in FLocalVertexFactory to enforce the assumption that the ColorComponent is VET_Color

diff --git a/Engine/Plugins/Experimental/Paper2D/Source/Paper2D/Private/PaperRenderSceneProxy.cpp b/Engine/Plugins/Experimental/Paper2D/Source/Paper2D/Private/PaperRenderSceneProxy.cpp
--- a/Engine/Plugins/Experimental/Paper2D/Source/Paper2D/Private/PaperRenderSceneProxy.cpp
+++ b/Engine/Plugins/Experimental/Paper2D/Source/Paper2D/Private/PaperRenderSceneProxy.cpp
@@ -67,7 +67,7 @@ public:
 			&GDummyMaterialSpriteVertexBuffer,
 			STRUCT_OFFSET(FPaperSpriteVertex,Color),
 			sizeof(FPaperSpriteVertex),
-			VET_Float4
+			VET_Color
 			);
 		Data.TextureCoordinates.Empty();
 		Data.TextureCoordinates.Add(FVertexStreamComponent(
diff --git a/Engine/Plugins/Experimental/Paper2D/Source/Paper2D/Private/PaperRenderSceneProxy.h b/Engine/Plugins/Experimental/Paper2D/Source/Paper2D/Private/PaperRenderSceneProxy.h
--- a/Engine/Plugins/Experimental/Paper2D/Source/Paper2D/Private/PaperRenderSceneProxy.h
+++ b/Engine/Plugins/Experimental/Paper2D/Source/Paper2D/Private/PaperRenderSceneProxy.h
@@ -13,7 +13,7 @@ struct PAPER2D_API FPaperSpriteVertex
 	FVector Position;
 	FPackedNormal TangentX;
 	FPackedNormal TangentZ;
-	FLinearColor Color;
+	FColor Color;
 	FVector2D TexCoords;
 
 	FPaperSpriteVertex() {}
diff --git a/Engine/Source/Runtime/Engine/Private/LocalVertexFactory.cpp b/Engine/Source/Runtime/Engine/Private/LocalVertexFactory.cpp
--- a/Engine/Source/Runtime/Engine/Private/LocalVertexFactory.cpp
+++ b/Engine/Source/Runtime/Engine/Private/LocalVertexFactory.cpp
@@ -48,6 +48,11 @@ bool FLocalVertexFactory::ShouldCache(EShaderPlatform Platform, const class FMat
 void FLocalVertexFactory::SetData(const DataType& InData)
 {
 	check(IsInRenderingThread());
+
+	// The shader code makes assumptions that the color component is a FColor, performing swizzles on ES2 and Metal platforms as necessary
+	// If the color is sent down as anything other than VET_Color then you'll get an undesired swizzle on those platforms
+	check((InData.ColorComponent.Type == VET_None) || (InData.ColorComponent.Type == VET_Color));
+
 	Data = InData;
 	UpdateRHI();
 }
diff --git a/Engine/Source/Runtime/Engine/Private/MaterialBillboardComponent.cpp b/Engine/Source/Runtime/Engine/Private/MaterialBillboardComponent.cpp
--- a/Engine/Source/Runtime/Engine/Private/MaterialBillboardComponent.cpp
+++ b/Engine/Source/Runtime/Engine/Private/MaterialBillboardComponent.cpp
@@ -10,7 +10,7 @@ struct FMaterialSpriteVertex
 	FVector Position;
 	FPackedNormal TangentX;
 	FPackedNormal TangentZ;
-	FLinearColor Color;
+	FColor Color;
 	FVector2D TexCoords;
 };
 
@@ -58,7 +58,7 @@ public:
 			&GDummyMaterialSpriteVertexBuffer,
 			STRUCT_OFFSET(FMaterialSpriteVertex,Color),
 			sizeof(FMaterialSpriteVertex),
-			VET_Float4
+			VET_Color
 			);
 		Data.TextureCoordinates.Empty();
 		Data.TextureCoordinates.Add(FVertexStreamComponent(
