#ue4: Header cleanup: more fixes to ComponentVelocity

diff --git a/Engine/Source/Runtime/Engine/Classes/Components/SceneComponent.h b/Engine/Source/Runtime/Engine/Classes/Components/SceneComponent.h
--- a/Engine/Source/Runtime/Engine/Classes/Components/SceneComponent.h
+++ b/Engine/Source/Runtime/Engine/Classes/Components/SceneComponent.h
@@ -196,7 +196,8 @@ public:
 	FVector RelativeScale3D;
 
 	/**
-	 * Component Velocity that has 
+	 * Velocity of this component.
+	 * @see GetComponentVelocity()
 	 */
 	UPROPERTY()
 	FVector ComponentVelocity;
@@ -442,8 +443,8 @@ public:
 	virtual void QuerySupportedSockets(TArray<FComponentSocketDescription>& OutSockets) const;
 
 	/** 
-	 * Get velocity of this component.
-	 * @return velocity of this component.
+	 * Get velocity of this component: either ComponentVelocity, or the velocity of the physics body if simulating physics.
+	 * @return Velocity of this component
 	 */
 	UFUNCTION(BlueprintCallable, Category="Utilities|Transformation")
 	virtual FVector GetComponentVelocity() const;
