Added several string manipulation functions to blueprint string library.
(Added Mirrored enums for ESearchCase and ESearchDir to object.h)

diff --git a/Engine/Source/Runtime/CoreUObject/Classes/Object.h b/Engine/Source/Runtime/CoreUObject/Classes/Object.h
--- a/Engine/Source/Runtime/CoreUObject/Classes/Object.h
+++ b/Engine/Source/Runtime/CoreUObject/Classes/Object.h
@@ -10,6 +10,28 @@
 #pragma once
 #if !CPP      //noexport class
 
+// String search case used in UnrealString.h
+UENUM()
+namespace ESearchCase
+{
+	enum Type
+	{
+		CaseSensitive,
+		IgnoreCase,
+	};
+}
+
+// String search dir used in UnrealString.h
+UENUM()
+namespace ESearchDir
+{
+	enum Type
+	{
+		FromStart,
+		FromEnd,
+	};
+}
+
 // Generic axis enum (mirrored for native use in Axis.h).
 UENUM()
 namespace EAxis
diff --git a/Engine/Source/Runtime/Engine/Classes/Kismet/KismetStringLibrary.h b/Engine/Source/Runtime/Engine/Classes/Kismet/KismetStringLibrary.h
--- a/Engine/Source/Runtime/Engine/Classes/Kismet/KismetStringLibrary.h
+++ b/Engine/Source/Runtime/Engine/Classes/Kismet/KismetStringLibrary.h
@@ -241,7 +241,7 @@ class UKismetStringLibrary : public UBlueprintFunctionLibrary
 	 * @param StartPosition The position to start the search from
 	 * @return The index (starting from 0 if bSearchFromEnd is false) of the first occurence of the substring
 	 */
-	UFUNCTION(BlueprintPure, Category = "String")
+	UFUNCTION(BlueprintPure, Category = "Utilities|String")
 	static int32 FindSubstring(const FString& SearchIn, const FString& Substring, bool bUseCase = false, bool bSearchFromEnd = false, int32 StartPosition = -1);
 
 	/** 
@@ -264,10 +264,165 @@ class UKismetStringLibrary : public UBlueprintFunctionLibrary
 	static TArray<FString> ParseIntoArray(const FString& SourceString, const FString& Delimiter = FString(TEXT(" ")), const bool CullEmptyStrings = true);
 
 	/**
-	 * Returns an array that contains one entry for each character in SourceString
-	 * @param	SourceString	The string to break apart into characters
-	 * @return	An array containing one entry for each character in SourceString
+	* Returns an array that contains one entry for each character in SourceString
+	* @param	SourceString	The string to break apart into characters
+	* @return	An array containing one entry for each character in SourceString
+	*/
+	UFUNCTION(BlueprintPure, Category = "Utilities|String")
+	static TArray<FString> GetCharacterArrayFromString(const FString& SourceString);
+
+	/**
+	 * Returns a string converted to Upper case
+	 * @param	SourceString	The string to convert
+	 * @return	The string in upper case
 	 */
 	UFUNCTION(BlueprintPure, Category="Utilities|String")
-	static TArray<FString> GetCharacterArrayFromString(const FString& SourceString);
+	static FString ToUpper(const FString& SourceString);
+
+	/**
+	* Returns a string converted to Lower case
+	* @param	SourceString	The string to convert
+	* @return	The string in lower case
+	*/
+	UFUNCTION(BlueprintPure, Category = "Utilities|String")
+	static FString ToLower(const FString& SourceString);
+
+	/*
+	 * Pad the left of this string for a specified number of characters 
+	 * @param	SourceString	The string to pad
+	 * @param	ChCount			Amount of padding required
+	 * @return	The padded string
+	 */
+	UFUNCTION(BlueprintPure, Category = "Utilities|String")
+	static FString LeftPad(const FString& SourceString, int32 ChCount);
+
+	/*
+	 * Pad the right of this string for a specified number of characters
+	 * @param	SourceString	The string to pad
+	 * @param	ChCount			Amount of padding required
+	 * @return	The padded string
+	 */
+	UFUNCTION(BlueprintPure, Category = "Utilities|String")
+	static FString RightPad(const FString& SourceString, int32 ChCount);
+
+	/*
+	 * Checks if a string contains only numeric characters
+	 * @param	SourceString	The string to check
+	 * @return true if the string only contains numeric characters 
+	 */
+	UFUNCTION(BlueprintPure, Category = "Utilities|String")
+	static bool IsNumeric(const FString& SourceString);
+
+	/**
+	 * Test whether this string starts with given string.
+	 *
+	 * @param SearchCase		Indicates whether the search is case sensitive or not ( defaults to ESearchCase::IgnoreCase )
+	 * @return true if this string begins with specified text, false otherwise
+	 */
+	UFUNCTION(BlueprintPure, Category = "Utilities|String")
+	static bool StartsWith(const FString& SourceString, const FString& InPrefix, ESearchCase::Type SearchCase = ESearchCase::IgnoreCase);
+
+	/**
+	 * Test whether this string ends with given string.
+	 *
+	 * @param SearchCase		Indicates whether the search is case sensitive or not ( defaults to ESearchCase::IgnoreCase )
+	 * @return true if this string ends with specified text, false otherwise
+	 */
+	UFUNCTION(BlueprintPure, Category = "Utilities|String")
+	static bool EndsWith(const FString& SourceString, const FString& InSuffix, ESearchCase::Type SearchCase = ESearchCase::IgnoreCase);
+
+	/**
+	 * Searches this string for a given wild card
+	 *
+	 * @param Wildcard		*?-type wildcard
+	 * @param SearchCase	Indicates whether the search is case sensitive or not ( defaults to ESearchCase::IgnoreCase )
+	 * @return true if this string matches the *?-type wildcard given.
+	 * @warning This is a simple, SLOW routine. Use with caution
+	 */
+	UFUNCTION(BlueprintPure, Category = "Utilities|String")
+	static bool MatchesWildcard(const FString& SourceString, const FString& Wildcard, ESearchCase::Type SearchCase = ESearchCase::IgnoreCase);
+
+	/**
+	 * Removes whitespace characters from the front of this string.
+	 */
+	UFUNCTION(BlueprintPure, Category = "Utilities|String")
+	static FString Trim(const FString& SourceString);
+
+	/**
+	 * Removes trailing whitespace characters
+	 */
+	UFUNCTION(BlueprintPure, Category = "Utilities|String")
+	static FString TrimTrailing(const FString& SourceString);
+
+	/**
+	 * Takes an array of strings and removes any zero length entries.
+	 *
+	 * @param	InArray	The array to cull
+	 *
+	 * @return	The number of elements left in InArray
+	 */
+	UFUNCTION(BlueprintPure, Category = "Utilities|String")
+	static int32 CullArray(const FString& SourceString, TArray<FString>& InArray);
+
+	/**
+	* Returns a copy of this string, with the characters in reverse order
+	*/
+	UFUNCTION(BlueprintPure, Category = "Utilities|String")
+	static FString Reverse(const FString& SourceString);
+
+	/**
+	 * Replace all occurrences of a substring in this string
+	 *
+	 * @param From substring to replace
+	 * @param To substring to replace From with
+	 * @param SearchCase	Indicates whether the search is case sensitive or not ( defaults to ESearchCase::IgnoreCase )
+	 * @return a copy of this string with the replacement made
+	 */
+	UFUNCTION(BlueprintPure, Category = "Utilities|String")
+		static FString Replace(const FString& SourceString, const FString& From, const FString& To, ESearchCase::Type SearchCase = ESearchCase::IgnoreCase);
+
+	/**
+	 * Replace all occurrences of SearchText with ReplacementText in this string.
+	 *
+	 * @param	SearchText	the text that should be removed from this string
+	 * @param	ReplacementText		the text to insert in its place
+	 * @param SearchCase	Indicates whether the search is case sensitive or not ( defaults to ESearchCase::IgnoreCase )
+	 *
+	 * @return	the number of occurrences of SearchText that were replaced.
+	 */
+	UFUNCTION(BlueprintPure, Category = "Utilities|String")
+		static int32 ReplaceInline(FString& SourceString, const FString& SearchText, const FString& ReplacementText, ESearchCase::Type SearchCase = ESearchCase::IgnoreCase);
+
+	/**
+	* Splits this string at given string position case sensitive.
+	*
+	* @param InStr The string to search and split at
+	* @param LeftS out the string to the left of InStr, not updated if return is false
+	* @param RightS out the string to the right of InStr, not updated if return is false
+	* @param SearchCase		Indicates whether the search is case sensitive or not ( defaults to ESearchCase::IgnoreCase )
+	* @param SearchDir			Indicates whether the search starts at the begining or at the end ( defaults to ESearchDir::FromStart )
+	* @return true if string is split, otherwise false
+	*/
+	UFUNCTION(BlueprintPure, Category = "Utilities|String")
+	static bool Split(const FString& SourceString, const FString& InStr, FString& LeftS, FString& RightS, ESearchCase::Type SearchCase = ESearchCase::IgnoreCase, ESearchDir::Type SearchDir = ESearchDir::FromStart);
+	
+	/** @return the left most given number of characters */
+	UFUNCTION(BlueprintPure, Category = "Utilities|String")
+	static FString Left(const FString& SourceString, int32 Count);
+	
+	/** @return the left most characters from the string chopping the given number of characters from the end */
+	UFUNCTION(BlueprintPure, Category = "Utilities|String")
+		static FString LeftChop(const FString& SourceString, int32 Count);
+	
+	/** @return the string to the right of the specified location, counting back from the right (end of the word). */
+	UFUNCTION(BlueprintPure, Category = "Utilities|String")
+		static FString Right(const FString& SourceString, int32 Count);
+
+	/** @return the string to the right of the specified location, counting forward from the left (from the beginning of the word). */
+	UFUNCTION(BlueprintPure, Category = "Utilities|String")
+		static FString RightChop(const FString& SourceString, int32 Count);
+
+	/** @return the substring from Start position for Count characters. */
+	UFUNCTION(BlueprintPure, Category = "Utilities|String")
+		static FString Mid(const FString& SourceString, int32 Start, int32 Count);
 };
diff --git a/Engine/Source/Runtime/Engine/Private/KismetStringLibrary.cpp b/Engine/Source/Runtime/Engine/Private/KismetStringLibrary.cpp
--- a/Engine/Source/Runtime/Engine/Private/KismetStringLibrary.cpp
+++ b/Engine/Source/Runtime/Engine/Private/KismetStringLibrary.cpp
@@ -301,4 +301,110 @@ TArray<FString> UKismetStringLibrary::GetCharacterArrayFromString(const FString&
 	}
 
 	return SeparatedChars;
-}
\ No newline at end of file
+}
+
+FString UKismetStringLibrary::ToUpper(const FString& SourceString)
+{
+	return SourceString.ToUpper();
+}
+
+FString UKismetStringLibrary::ToLower(const FString& SourceString)
+{
+	return SourceString.ToLower();
+}
+
+FString UKismetStringLibrary::LeftPad(const FString& SourceString, int32 ChCount)
+{
+	return SourceString.LeftPad(ChCount);
+}
+
+FString UKismetStringLibrary::RightPad(const FString& SourceString, int32 ChCount)
+{
+	return SourceString.RightPad(ChCount);
+}
+
+bool UKismetStringLibrary::IsNumeric(const FString& SourceString)
+{
+	return SourceString.IsNumeric();
+}
+
+bool UKismetStringLibrary::StartsWith(const FString& SourceString, const FString& InPrefix, ESearchCase::Type SearchCase)
+{
+	return SourceString.StartsWith(InPrefix,SearchCase);
+}
+
+bool UKismetStringLibrary::EndsWith(const FString& SourceString, const FString& InSuffix, ESearchCase::Type SearchCase)
+{
+	return SourceString.EndsWith(InSuffix,SearchCase);
+}
+
+bool UKismetStringLibrary::MatchesWildcard(const FString& SourceString, const FString& Wildcard, ESearchCase::Type SearchCase)
+{
+	return SourceString.MatchesWildcard(Wildcard, SearchCase);
+}
+
+FString UKismetStringLibrary::Trim(const FString& SourceString)
+{
+	FString Trimmed = SourceString;
+	return Trimmed.Trim();
+}
+
+FString UKismetStringLibrary::TrimTrailing(const FString& SourceString)
+{
+	FString Trimmed = SourceString;
+	return Trimmed.TrimTrailing();
+}
+
+int32 UKismetStringLibrary::CullArray(const FString& SourceString,TArray<FString>& InArray)
+{
+	return SourceString.CullArray(&InArray);
+}
+
+FString UKismetStringLibrary::Reverse(const FString& SourceString)
+{
+//	return SourceString.Reverse();
+	FString Reversed = SourceString;
+	Reversed.ReverseString();
+	return Reversed;
+}
+
+FString UKismetStringLibrary::Replace(const FString& SourceString, const FString& From, const FString& To, ESearchCase::Type SearchCase)
+{
+	return SourceString.Replace(*From, *To, SearchCase);
+}
+
+int32 UKismetStringLibrary::ReplaceInline(FString& SourceString, const FString& SearchText, const FString& ReplacementText, ESearchCase::Type SearchCase)
+{
+	return SourceString.ReplaceInline(*SearchText, *ReplacementText, SearchCase);
+}
+
+bool UKismetStringLibrary::Split(const FString& SourceString, const FString& InStr, FString& LeftS, FString& RightS, ESearchCase::Type SearchCase, ESearchDir::Type SearchDir)
+{
+	return SourceString.Split(InStr, &LeftS, &RightS, SearchCase, SearchDir);
+}
+
+FString UKismetStringLibrary::Left(const FString& SourceString, int32 Count)
+{
+	return SourceString.Left(Count);
+}
+
+FString UKismetStringLibrary::LeftChop(const FString& SourceString, int32 Count)
+{
+	return SourceString.LeftChop(Count);
+}
+
+FString UKismetStringLibrary::Right(const FString& SourceString, int32 Count)
+{
+	return SourceString.Right(Count);
+}
+
+FString UKismetStringLibrary::RightChop(const FString& SourceString, int32 Count)
+{
+	return SourceString.RightChop(Count);
+}
+
+FString UKismetStringLibrary::Mid(const FString& SourceString, int32 Start, int32 Count = MAX_int32)
+{
+	return SourceString.Mid(Start, Count);
+}
+
