Clear smoothing state when Time Dilation changes to avoid lagging sensitivity artifacts

diff --git a/Engine/Source/Runtime/Engine/Classes/GameFramework/PlayerInput.h b/Engine/Source/Runtime/Engine/Classes/GameFramework/PlayerInput.h
--- a/Engine/Source/Runtime/Engine/Classes/GameFramework/PlayerInput.h
+++ b/Engine/Source/Runtime/Engine/Classes/GameFramework/PlayerInput.h
@@ -542,4 +542,7 @@ private:
 
 	// A counter used to track the order in which events occurred since the last time the input stack was processed
 	uint32 EventCount;
+
+	// Cache the last time dilation so as to be able to clear smoothing when it changes
+	float LastTimeDilation;
 };
diff --git a/Engine/Source/Runtime/Engine/Private/UserInterface/PlayerInput.cpp b/Engine/Source/Runtime/Engine/Private/UserInterface/PlayerInput.cpp
--- a/Engine/Source/Runtime/Engine/Private/UserInterface/PlayerInput.cpp
+++ b/Engine/Source/Runtime/Engine/Private/UserInterface/PlayerInput.cpp
@@ -1072,25 +1072,39 @@ void UPlayerInput::ClearSmoothing()
 {
 	for ( int32 i=0; i<2; i++ )
 	{
-		ZeroTime[i] = 0;
-		SmoothedMouse[i] = 0;
+		ZeroTime[i] = 0.f;
+		SmoothedMouse[i] = 0.f;
 	}
-   	MouseSamplingTotal = GetDefault<UPlayerInput>(GetClass())->MouseSamplingTotal;
-	MouseSamples = GetDefault<UPlayerInput>(GetClass())->MouseSamples;
+
+	const UPlayerInput* DefaultPlayerInput = GetDefault<UPlayerInput>();
+
+   	MouseSamplingTotal = DefaultPlayerInput->MouseSamplingTotal;
+	MouseSamples = DefaultPlayerInput->MouseSamples;
 }
 
 
 float UPlayerInput::SmoothMouse(float aMouse, float DeltaTime, uint8& SampleCount, int32 Index)
 {
-	float DetectedMouseSampleHz = MouseSamples / MouseSamplingTotal;
+	UWorld* World = GetWorld();
+	if (World)
+	{
+		const float EffectiveTimeDilation = World->GetWorldSettings()->GetEffectiveTimeDilation();
+		if (EffectiveTimeDilation != LastTimeDilation)
+		{
+			LastTimeDilation = EffectiveTimeDilation;
+			ClearSmoothing();
+		}
+	}
+
+	const float DetectedMouseSampleHz = MouseSamples / MouseSamplingTotal;
 
 	// actual num samples should be either this or this+1
-	int32 ExpectedNumSamplesThisFrame = FMath::TruncToInt(DeltaTime * DetectedMouseSampleHz);
+	const int32 ExpectedNumSamplesThisFrame = FMath::TruncToInt(DeltaTime * DetectedMouseSampleHz);
 
 	if (DeltaTime < 0.25f)
 	{
 		// this is seconds/sample
-		float MouseSamplingTime = MouseSamplingTotal/MouseSamples;
+		const float MouseSamplingTime = MouseSamplingTotal/MouseSamples;
 
 		if ( aMouse == 0 )
 		{
