Fix behavior of body welding.
Added different behavior when setting up in blueprint vs level editor
Fix detach bug when attaching physically simulated from blueprint
customization work to make auto weld easier to understand

diff --git a/Engine/Config/BaseEngine.ini b/Engine/Config/BaseEngine.ini
--- a/Engine/Config/BaseEngine.ini
+++ b/Engine/Config/BaseEngine.ini
@@ -675,6 +675,10 @@ bUseStreamingPause=false
 
 +TaggedPropertyRedirects=(ClassName="BodyInstance",OldPropertyName="MovementChannel",NewPropertyName="ObjectType")
 +K2FieldRedirects=(OldFieldName="PrimitiveComponent.SetMovementChannel",NewFieldName="PrimitiveComponent.SetCollisionObjectType")
++K2FieldRedirects=(OldFieldName="SceneComponent.AttachTo",NewFieldName="SceneComponent.K2_AttachTo")
++K2FieldRedirects=(OldFieldName="Actor.AttachRootComponentTo",NewFieldName="Actor.K2_AttachRootComponentTo")
++K2FieldRedirects=(OldFieldName="Actor.AttachRootComponentToActor",NewFieldName="Actor.K2_AttachRootComponentToActor")
+
 
 ;Notify Time Change
 +TaggedPropertyRedirects=(ClassName="AnimNotifyEvent",OldPropertyName="Time",NewPropertyName="DisplayTime")
diff --git a/Engine/Source/Editor/DetailCustomizations/Private/BodySetupDetails.cpp b/Engine/Source/Editor/DetailCustomizations/Private/BodySetupDetails.cpp
--- a/Engine/Source/Editor/DetailCustomizations/Private/BodySetupDetails.cpp
+++ b/Engine/Source/Editor/DetailCustomizations/Private/BodySetupDetails.cpp
@@ -57,13 +57,14 @@ void FBodySetupDetails::CustomizeDetails( IDetailLayoutBuilder& DetailBuilder )
 			{
 				TSharedPtr<IPropertyHandle> ChildProperty = BodyInstanceHandler->GetChildHandle(ChildIndex);
 				FString Category = FObjectEditorUtils::GetCategory(ChildProperty->GetProperty());
-				if (ChildProperty->GetProperty()->GetName() == TEXT("bSimulatePhysics"))
+				if (ChildProperty->GetProperty()->GetName() == TEXT("bSimulatePhysics") || ChildProperty->GetProperty()->GetName() == TEXT("bAutoWeld"))
 				{
 					// skip bSimulatePhysics
 					// this is because we don't want bSimulatePhysics to show up 
 					// phat editor 
 					// staitc mesh already hides everything else not interested in
 					// so phat editor just should not show this option
+					//also hide bAutoWeld for phat
 					continue;
 				}
 				if (Category == TEXT("Physics"))
diff --git a/Engine/Source/Editor/DetailCustomizations/Private/PrimitiveComponentDetails.cpp b/Engine/Source/Editor/DetailCustomizations/Private/PrimitiveComponentDetails.cpp
--- a/Engine/Source/Editor/DetailCustomizations/Private/PrimitiveComponentDetails.cpp
+++ b/Engine/Source/Editor/DetailCustomizations/Private/PrimitiveComponentDetails.cpp
@@ -96,6 +96,38 @@ EVisibility FPrimitiveComponentDetails::IsCustomLockedAxisSelected() const
 	return bVisible ? EVisibility::Visible : EVisibility::Collapsed;
 }
 
+bool FPrimitiveComponentDetails::IsAutoWeldEditable() const
+{
+	for (int32 i = 0; i < ObjectsCustomized.Num(); ++i)
+	{
+		if (UPrimitiveComponent * SceneComponent = Cast<UPrimitiveComponent>(ObjectsCustomized[i].Get()))
+		{
+			if (FBodyInstance * BI = SceneComponent->GetBodyInstance())
+			{
+				if (BI->IsInstanceSimulatingPhysics())
+				{
+					return false;
+				}
+			}
+		}
+	}
+
+	return true;
+}
+
+EVisibility FPrimitiveComponentDetails::IsAutoWeldVisible() const
+{
+	for (int32 i = 0; i < ObjectsCustomized.Num(); ++i)
+	{
+		if (ObjectsCustomized[i].IsValid() && !ObjectsCustomized[i]->IsA(UStaticMeshComponent::StaticClass()))
+		{
+			return EVisibility::Collapsed;
+		}
+	}
+
+	return EVisibility::Visible;
+}
+
 void FPrimitiveComponentDetails::CustomizeDetails( IDetailLayoutBuilder& DetailBuilder )
 {
 	TSharedRef<IPropertyHandle> MobilityHandle = DetailBuilder.GetProperty("Mobility", USceneComponent::StaticClass());
@@ -176,9 +208,10 @@ void FPrimitiveComponentDetails::CustomizeDetails( IDetailLayoutBuilder& DetailB
 						//we only enable bUseAsyncScene if the project uses an AsyncScene
 						PhysicsCategory.AddProperty(ChildProperty).Visibility(TAttribute<EVisibility>(this, &FPrimitiveComponentDetails::IsCustomLockedAxisSelected));
 					}
-					else if (PropName == TEXT("AutoWeld"))
+					else if (PropName == TEXT("bAutoWeld"))
 					{
-						PhysicsCategory.AddProperty(ChildProperty);	//can't ifdef the property out so for now we just customize it to be hidden in UI
+						PhysicsCategory.AddProperty(ChildProperty).Visibility(TAttribute<EVisibility>(this, &FPrimitiveComponentDetails::IsAutoWeldVisible))
+																  .EditCondition(TAttribute<bool>(this, &FPrimitiveComponentDetails::IsAutoWeldEditable), NULL);
 					}
 					else
 					{
diff --git a/Engine/Source/Editor/DetailCustomizations/Private/PrimitiveComponentDetails.h b/Engine/Source/Editor/DetailCustomizations/Private/PrimitiveComponentDetails.h
--- a/Engine/Source/Editor/DetailCustomizations/Private/PrimitiveComponentDetails.h
+++ b/Engine/Source/Editor/DetailCustomizations/Private/PrimitiveComponentDetails.h
@@ -36,5 +36,8 @@ private:
 	TArray< TWeakObjectPtr<UObject> > ObjectsCustomized;
 	TSharedPtr<IPropertyHandle> LockedAxisProperty;
 	EVisibility IsCustomLockedAxisSelected() const;
+
+	bool IsAutoWeldEditable() const;
+	EVisibility IsAutoWeldVisible() const;
 };
 
diff --git a/Engine/Source/Runtime/Engine/Classes/Components/PrimitiveComponent.h b/Engine/Source/Runtime/Engine/Classes/Components/PrimitiveComponent.h
--- a/Engine/Source/Runtime/Engine/Classes/Components/PrimitiveComponent.h
+++ b/Engine/Source/Runtime/Engine/Classes/Components/PrimitiveComponent.h
@@ -1030,18 +1030,18 @@ public:
 	*   @param InParent the component to be physically attached to
 	*   @param InSocketName optional socket to attach component to
 	*/
-	virtual void WeldTo(class USceneComponent* InParent, FName InSocketName = NAME_None) override;
+	virtual void WeldTo(class USceneComponent* InParent, FName InSocketName = NAME_None);
 
 	/**
-	*   UnWelds this component from its parent component. Attachment is maintained (DetachFromParent automatically unwelds)
+	*	Does the actual work for welding.
+	*	@return true if did a true weld of shapes, meaning body initialization is not needed
 	*/
-	virtual void UnWeldFromParent() override;
+	virtual bool WeldToImplementation(USceneComponent * InParent, FName ParentSocketName = NAME_None, bool bWeldSimulatedChild = true);
 
 	/**
-	*	Does the actual work for welding.
-	*	@return true if did a true weld of shapes, meaning body initialization is not needed
+	*   UnWelds this component from its parent component. Attachment is maintained (DetachFromParent automatically unwelds)
 	*/
-	bool WeldToInternal(USceneComponent * InParent, FName ParentSocketName = NAME_None);
+	virtual void UnWeldFromParent();
 
 	/**
 	*	Adds the bodies that are currently welded to the OutWeldedBodies array 
diff --git a/Engine/Source/Runtime/Engine/Classes/Components/SceneComponent.h b/Engine/Source/Runtime/Engine/Classes/Components/SceneComponent.h
--- a/Engine/Source/Runtime/Engine/Classes/Components/SceneComponent.h
+++ b/Engine/Source/Runtime/Engine/Classes/Components/SceneComponent.h
@@ -342,8 +342,14 @@ public:
 	 *   Attach this component to another scene component, optionally at a named socket. It is valid to call this on components whether or not they have been Registered.
 	 *   @param bMaintainWorldTransform	If true, update the relative location/rotation of this component to keep its world position the same
 	 */
-	UFUNCTION(BlueprintCallable, Category="Utilities|Transformation", meta=(AttachType="KeepRelativeOffset"))
-	void AttachTo(class USceneComponent* InParent, FName InSocketName = NAME_None, EAttachLocation::Type AttachType = EAttachLocation::KeepRelativeOffset);
+	void AttachTo(class USceneComponent* InParent, FName InSocketName = NAME_None, EAttachLocation::Type AttachType = EAttachLocation::KeepRelativeOffset, bool bWeldSimulatedBodies = false);
+
+	/**
+	*   Attach this component to another scene component, optionally at a named socket. It is valid to call this on components whether or not they have been Registered.
+	*   @param bMaintainWorldTransform	If true, update the relative location/rotation of this component to keep its world position the same
+	*/
+	UFUNCTION(BlueprintCallable, Category = "Utilities|Transformation", meta = (FriendlyName = "AttachTo", AttachType = "KeepRelativeOffset"))
+	void K2_AttachTo(class USceneComponent* InParent, FName InSocketName = NAME_None, EAttachLocation::Type AttachType = EAttachLocation::KeepRelativeOffset, bool bWeldSimulatedBodies = true);
 
 	/** Zeroes out the relative transform of this component, and calls AttachTo(). Useful for attaching directly to a scene component or socket location  */
 	UFUNCTION(BlueprintCallable, meta=(DeprecatedFunction, DeprecationMessage = "Use AttachTo with EAttachLocation::SnapToTarget option instead"), Category="Utilities|Transformation")
@@ -356,34 +362,6 @@ public:
 	UFUNCTION(BlueprintCallable, Category="Utilities|Transformation")
 	virtual void DetachFromParent(bool bMaintainWorldPosition = false);
 
-	/**
-	*   Welds this component to another scene component, optionally at a named socket. Component is automatically attached if not already
-	*	Welding allows the child physics object to become physically connected to its parent. This is useful for creating compound rigid bodies with correct mass distribution.
-	*   @param InParent the component to be physically attached to
-	*   @param InSocketName optional socket to attach component to
-	*/
-	UFUNCTION(BlueprintCallable, Category = "Physics")
-	virtual void WeldTo(class USceneComponent* InParent, FName InSocketName = NAME_None);
-
-	/**
-	*   UnWelds this component from its parent component. Attachment is maintained (DetachFromParent automatically unwelds)
-	*/
-	UFUNCTION(BlueprintCallable, Category = "Physics")
-	virtual void UnWeldFromParent();
-
-
-	/**
-	* Attempts to weld/unweld the child component to this component in a physically correct way. This re-calculates mass distribution and geometry to achieve correct physics.
-	* This function assumes ChildComponent has already been attached to this component and is now doing the physics fix-up.
-	* The function uses the attachment hierarchy to easily attach/detach multiple components together
-	* @param ChildComponent the component to be physically attached - assumed to already be attached in the scene hierarchy
-	* @param bWeld determines whether to weld or unweld. This flag is used for creating or breaking welded pieces at runtime
-	* @param ParentBoneName potentially used by skeletal mesh to weld child to specific bone on this component
-	* @param ChildBoneName potentially used by skeletal mesh to weld this component to specific bone in child
-	* @return whether or not the weld/unweld was succesfull
-	*/
-	//virtual bool WeldPhysicsBody(USceneComponent * ChildComponent, bool bWeld = true, FName ParentBoneName = NAME_None, FName ChildBoneName = NAME_None) { return false; }
-
 
 	/** 
 	 * Gets the names of all the sockets on the component.
diff --git a/Engine/Source/Runtime/Engine/Classes/Components/SkeletalMeshComponent.h b/Engine/Source/Runtime/Engine/Classes/Components/SkeletalMeshComponent.h
--- a/Engine/Source/Runtime/Engine/Classes/Components/SkeletalMeshComponent.h
+++ b/Engine/Source/Runtime/Engine/Classes/Components/SkeletalMeshComponent.h
@@ -955,7 +955,6 @@ public:
 	void SetNamedMotorsAngularVelocityDrive(bool bEnableSwingDrive, bool bEnableTwistDrive, const TArray<FName>& BoneNames, bool bSetOtherBodiesToComplement = false);
 
 	void GetWeldedBodies(TArray<FBodyInstance*> & OutWeldedBodies, TArray<FName> & OutChildrenLabels) override;
-	bool virtual WeldToInternal(USceneComponent * InParent, FName ParentSocketName = NAME_None){ return false; }	//not currently supported for skeletal mesh
 
 	/**
 	 * Return Transform Matrix for SkeletalMeshComponent considering root motion setups
diff --git a/Engine/Source/Runtime/Engine/Classes/GameFramework/Actor.h b/Engine/Source/Runtime/Engine/Classes/GameFramework/Actor.h
--- a/Engine/Source/Runtime/Engine/Classes/GameFramework/Actor.h
+++ b/Engine/Source/Runtime/Engine/Classes/GameFramework/Actor.h
@@ -732,8 +732,15 @@ public:
 	 *  Attaches the RootComponent of this Actor to the supplied component, optionally at a named socket. It is not valid to call this on components that are not Registered. 
 	 *   @param AttachLocationType	Type of attachment, AbsoluteWorld to keep its world position, RelativeOffset to keep the object's relative offset and SnapTo to snap to the new parent.
 	 */
-	UFUNCTION(BlueprintCallable, meta=(FriendlyName = "AttachActorToComponent", AttachLocationType="KeepRelativeOffset"), Category="Utilities|Transformation")
-	void AttachRootComponentTo(class USceneComponent* InParent, FName InSocketName = NAME_None, EAttachLocation::Type AttachLocationType = EAttachLocation::KeepRelativeOffset);
+	void AttachRootComponentTo(class USceneComponent* InParent, FName InSocketName = NAME_None, EAttachLocation::Type AttachLocationType = EAttachLocation::KeepRelativeOffset, bool bWeldSimulatedBodies = false);
+
+	/**
+	*  Attaches the RootComponent of this Actor to the supplied component, optionally at a named socket. It is not valid to call this on components that are not Registered.
+	*   @param AttachLocationType	Type of attachment, AbsoluteWorld to keep its world position, RelativeOffset to keep the object's relative offset and SnapTo to snap to the new parent.
+	*/
+
+	UFUNCTION(BlueprintCallable, meta = (FriendlyName = "AttachActorToComponent", AttachLocationType = "KeepRelativeOffset"), Category = "Utilities|Transformation")
+	void K2_AttachRootComponentTo(class USceneComponent* InParent, FName InSocketName = NAME_None, EAttachLocation::Type AttachLocationType = EAttachLocation::KeepRelativeOffset, bool bWeldSimulatedBodies = true);
 
 	/**
 	 * Attaches the RootComponent of this Actor to the RootComponent of the supplied actor, optionally at a named socket.
@@ -741,8 +748,15 @@ public:
 	 * @param InSocketName				Socket name to attach to, if any
 	 * @param AttachLocationType	Type of attachment, AbsoluteWorld to keep its world position, RelativeOffset to keep the object's relative offset and SnapTo to snap to the new parent.
 	 */
-	UFUNCTION(BlueprintCallable, meta=(FriendlyName = "AttachActorToActor", AttachLocationType="KeepRelativeOffset"), Category="Utilities|Transformation")
-	void AttachRootComponentToActor(AActor* InParentActor, FName InSocketName = NAME_None, EAttachLocation::Type AttachLocationType = EAttachLocation::KeepRelativeOffset);
+	void AttachRootComponentToActor(AActor* InParentActor, FName InSocketName = NAME_None, EAttachLocation::Type AttachLocationType = EAttachLocation::KeepRelativeOffset, bool bWeldSimulatedBodies = false);
+
+	/**
+	*  Attaches the RootComponent of this Actor to the supplied component, optionally at a named socket. It is not valid to call this on components that are not Registered.
+	*   @param AttachLocationType	Type of attachment, AbsoluteWorld to keep its world position, RelativeOffset to keep the object's relative offset and SnapTo to snap to the new parent.
+	*/
+
+	UFUNCTION(BlueprintCallable, meta = (FriendlyName = "AttachActorToActor", AttachLocationType = "KeepRelativeOffset"), Category = "Utilities|Transformation")
+	void K2_AttachRootComponentToActor(AActor* InParentActor, FName InSocketName = NAME_None, EAttachLocation::Type AttachLocationType = EAttachLocation::KeepRelativeOffset, bool bWeldSimulatedBodies = true);
 
 	/** 
 	 *  Snap the RootComponent of this Actor to the supplied Actor's root component, optionally at a named socket. It is not valid to call this on components that are not Registered. 
diff --git a/Engine/Source/Runtime/Engine/Private/Actor.cpp b/Engine/Source/Runtime/Engine/Private/Actor.cpp
--- a/Engine/Source/Runtime/Engine/Private/Actor.cpp
+++ b/Engine/Source/Runtime/Engine/Private/Actor.cpp
@@ -1098,11 +1098,16 @@ bool AActor::HasNetOwner() const
 	return TopOwner->HasNetOwner();
 }
 
-void AActor::AttachRootComponentTo(USceneComponent* InParent, FName InSocketName, EAttachLocation::Type AttachLocationType /*= EAttachLocation::KeepRelativeOffset */)
+void AActor::K2_AttachRootComponentTo(USceneComponent* InParent, FName InSocketName, EAttachLocation::Type AttachLocationType /*= EAttachLocation::KeepRelativeOffset */, bool bWeldSimulatedBodies /*=true*/)
+{
+	AttachRootComponentTo(InParent, InSocketName, AttachLocationType, bWeldSimulatedBodies);
+}
+
+void AActor::AttachRootComponentTo(USceneComponent* InParent, FName InSocketName, EAttachLocation::Type AttachLocationType /*= EAttachLocation::KeepRelativeOffset */, bool bWeldSimulatedBodies /*=false*/)
 {
 	if(RootComponent && InParent)
 	{
-		RootComponent->AttachTo(InParent, InSocketName, AttachLocationType );
+		RootComponent->AttachTo(InParent, InSocketName, AttachLocationType, bWeldSimulatedBodies);
 
 
 		AttachmentReplication.AttachParent = InParent->GetAttachmentRootActor();
@@ -1145,14 +1150,19 @@ void AActor::OnRep_AttachmentReplication()
 	}
 }
 
-void AActor::AttachRootComponentToActor(AActor* InParentActor, FName InSocketName /*= NAME_None*/, EAttachLocation::Type AttachLocationType /*= EAttachLocation::KeepRelativeOffset */)
+void AActor::K2_AttachRootComponentToActor(AActor* InParentActor, FName InSocketName /*= NAME_None*/, EAttachLocation::Type AttachLocationType /*= EAttachLocation::KeepRelativeOffset */, bool bWeldSimulatedBodies /*=true*/)
+{
+	AttachRootComponentToActor(InParentActor, InSocketName, AttachLocationType, bWeldSimulatedBodies);
+}
+
+void AActor::AttachRootComponentToActor(AActor* InParentActor, FName InSocketName /*= NAME_None*/, EAttachLocation::Type AttachLocationType /*= EAttachLocation::KeepRelativeOffset */, bool bWeldSimulatedBodies /*=false*/)
 {
 	if (RootComponent && InParentActor)
 	{
 		USceneComponent* ParentRootComponent = InParentActor->GetRootComponent();
 		if (ParentRootComponent)
 		{
-			RootComponent->AttachTo(ParentRootComponent, InSocketName, AttachLocationType );
+			RootComponent->AttachTo(ParentRootComponent, InSocketName, AttachLocationType, bWeldSimulatedBodies );
 
 
 			AttachmentReplication.AttachParent = InParentActor;
diff --git a/Engine/Source/Runtime/Engine/Private/PhysicsEngine/BodyInstance.cpp b/Engine/Source/Runtime/Engine/Private/PhysicsEngine/BodyInstance.cpp
--- a/Engine/Source/Runtime/Engine/Private/PhysicsEngine/BodyInstance.cpp
+++ b/Engine/Source/Runtime/Engine/Private/PhysicsEngine/BodyInstance.cpp
@@ -907,7 +907,7 @@ void FBodyInstance::InitBody(UBodySetup* Setup, const FTransform& Transform, UPr
 		//if we have a parent we will now do the weld and exit any further initialization
 		if (ParentPrimComponent && PrimComp->GetWorld() && PrimComp->GetWorld()->IsGameWorld())
 		{
-			if (PrimComp->WeldToInternal(ParentPrimComponent, PrimComp->AttachSocketName))	//welded new simulated body so initialization is done
+			if (PrimComp->WeldToImplementation(ParentPrimComponent, PrimComp->AttachSocketName, false))	//welded new simulated body so initialization is done
 			{
 				return;
 			}
diff --git a/Engine/Source/Runtime/Engine/Private/PrimitiveComponent.cpp b/Engine/Source/Runtime/Engine/Private/PrimitiveComponent.cpp
--- a/Engine/Source/Runtime/Engine/Private/PrimitiveComponent.cpp
+++ b/Engine/Source/Runtime/Engine/Private/PrimitiveComponent.cpp
@@ -329,7 +329,7 @@ void UPrimitiveComponent::OnChildAttached(USceneComponent* ChildComponent)
 	Super::OnChildAttached(ChildComponent);
 }
 
-UPrimitiveComponent * GetRootWelded(UPrimitiveComponent * PrimComponent, FName ParentSocketName = NAME_None, FName * OutSocketName = NULL)
+UPrimitiveComponent * GetRootWelded(UPrimitiveComponent * PrimComponent, FName ParentSocketName = NAME_None, FName * OutSocketName = NULL, bool bAboutToWeld = false)
 {
 	UPrimitiveComponent * Result = NULL;
 	UPrimitiveComponent * RootComponent = Cast<UPrimitiveComponent>(PrimComponent->AttachParent);	//we must find the root component along hierarchy that has bWelded set to true
@@ -337,7 +337,7 @@ UPrimitiveComponent * GetRootWelded(UPrimitiveComponent * PrimComponent, FName P
 	//check that body itself is welded
 	if (FBodyInstance * BI = PrimComponent->GetBodyInstance(ParentSocketName))
 	{
-		if (BI->bWelded == false && BI->bAutoWeld == false)	//we're not welded and we aren't trying to become welded
+		if (bAboutToWeld == false && BI->bWelded == false && BI->bAutoWeld == false)	//we're not welded and we aren't trying to become welded
 		{
 			return NULL;
 		}
@@ -384,7 +384,7 @@ void UPrimitiveComponent::GetWeldedBodies(TArray<FBodyInstance*> & OutWeldedBodi
 	}
 }
 
-bool UPrimitiveComponent::WeldToInternal(USceneComponent * InParent, FName ParentSocketName /* = Name_None */)
+bool UPrimitiveComponent::WeldToImplementation(USceneComponent * InParent, FName ParentSocketName /* = Name_None */, bool bWeldSimulatedChild /* = false */)
 {
 	//WeldToInternal assumes attachment is already done
 	if (AttachParent != InParent || AttachSocketName != ParentSocketName)
@@ -399,9 +399,14 @@ bool UPrimitiveComponent::WeldToInternal(USceneComponent * InParent, FName Paren
 		return false;
 	}
 
+	if (BI->ShouldInstanceSimulatingPhysics() && bWeldSimulatedChild == false)
+	{
+		return false;
+	}
+
 	FName SocketName;
-	UPrimitiveComponent * RootComponent = GetRootWelded(this, ParentSocketName, &SocketName);
-	
+	UPrimitiveComponent * RootComponent = GetRootWelded(this, ParentSocketName, &SocketName, true);
+
 	if (RootComponent)
 	{
 		if (FBodyInstance * RootBI = RootComponent->GetBodyInstance(SocketName))
@@ -438,7 +443,7 @@ void UPrimitiveComponent::WeldTo(USceneComponent* InParent, FName InSocketName /
 		AttachTo(InParent, InSocketName, EAttachLocation::KeepWorldPosition);
 	}
 
-	WeldToInternal(InParent, InSocketName);
+	WeldToImplementation(InParent, InSocketName);
 }
 
 void UPrimitiveComponent::UnWeldFromParent()
@@ -492,7 +497,6 @@ void UPrimitiveComponent::UnWeldFromParent()
 	}
 }
 
-
 void UPrimitiveComponent::DestroyRenderState_Concurrent()
 {
 	// Remove the primitive from the scene.
diff --git a/Engine/Source/Runtime/Engine/Private/SceneComponent.cpp b/Engine/Source/Runtime/Engine/Private/SceneComponent.cpp
--- a/Engine/Source/Runtime/Engine/Private/SceneComponent.cpp
+++ b/Engine/Source/Runtime/Engine/Private/SceneComponent.cpp
@@ -70,15 +70,6 @@ void USceneComponent::OnUpdateTransform(bool bSkipPhysicsMove)
 {
 }
 
-void USceneComponent::WeldTo(class USceneComponent* InParent, FName InSocketName)
-{
-}
-
-void USceneComponent::UnWeldFromParent()
-{
-
-}
-
 void USceneComponent::UpdateComponentToWorldWithParent(USceneComponent * Parent, bool bSkipPhysicsMove)
 {
 	// If our parent hasn't been updated before, we'll need walk up our parent attach hierarchy
@@ -697,7 +688,13 @@ void USceneComponent::AppendDescendants(TArray<USceneComponent*>& Children) cons
 	}
 }
 
-void USceneComponent::AttachTo(class USceneComponent* Parent, FName InSocketName, EAttachLocation::Type AttachType /*= EAttachLocation::KeepRelativeOffset */)
+//This function is used for giving AttachTo different bWeldSimulatedBodies default, but only when called from BP
+void USceneComponent::K2_AttachTo(class USceneComponent* InParent, FName InSocketName, EAttachLocation::Type AttachLocationType, bool bWeldSimulatedBodies /*= true*/)
+{
+	AttachTo(InParent, InSocketName, AttachLocationType, bWeldSimulatedBodies);
+}
+
+void USceneComponent::AttachTo(class USceneComponent* Parent, FName InSocketName, EAttachLocation::Type AttachType /*= EAttachLocation::KeepRelativeOffset */, bool bWeldSimulatedBodies /*= false*/)
 {
 	if(Parent != NULL)
 	{
@@ -760,15 +757,16 @@ void USceneComponent::AttachTo(class USceneComponent* Parent, FName InSocketName
 		{
 			//This code requires some explaining. Inside the editor we allow user to attach physically simulated objects to other objects. This is done for convenience so that users can group things together in hierarchy.
 			//At runtime we must not attach physically simulated objects as it will cause double transform updates, and you should just use a physical constraint if attachment is the desired behavior.
-			//Note that if bAutoWeld is true then it means that they actually do want the bodies to be simulated together, in which case we should not break the attachment
+			//Note if bWeldSimulatedBodies = true then they actually want to keep these objects simulating together
 			//We must fixup the relative location,rotation,scale as the attachment is no longer valid. Blueprint uses simple construction to try and attach before ComponentToWorld has ever been updated, so we cannot rely on it.
 			//As such we must calculate the proper Relative information
 			//Also physics state may not be created yet so we use bSimulatePhysics to determine if the object has any intention of being physically simulated
 			UPrimitiveComponent * PrimitiveComponent = Cast<UPrimitiveComponent>(this);
-			if (PrimitiveComponent && PrimitiveComponent->BodyInstance.bSimulatePhysics && !PrimitiveComponent->BodyInstance.bAutoWeld && GetWorld() && GetWorld()->IsGameWorld())
+
+			if (PrimitiveComponent && PrimitiveComponent->BodyInstance.bSimulatePhysics && !bWeldSimulatedBodies && GetWorld() && GetWorld()->IsGameWorld())
 			{
 				//Since the object is physically simulated it can't be the case that it's a child of object A and being attached to object B (at runtime)
-				UpdateComponentToWorldWithParent(Parent, false);
+				//UpdateComponentToWorldWithParent(Parent, false);
 				RelativeLocation = ComponentToWorld.GetLocation();
 				RelativeRotation = ComponentToWorld.GetRotation().Rotator();
 				RelativeScale3D = ComponentToWorld.GetScale3D();
@@ -868,9 +866,9 @@ void USceneComponent::AttachTo(class USceneComponent* Parent, FName InSocketName
 		{
 			if (FBodyInstance * BI = PrimitiveComponent->GetBodyInstance())
 			{
-				if (BI->bAutoWeld)
+				if (BI->bAutoWeld || bWeldSimulatedBodies)
 				{
-					PrimitiveComponent->WeldToInternal(AttachParent, AttachSocketName);
+					PrimitiveComponent->WeldToImplementation(AttachParent, AttachSocketName, bWeldSimulatedBodies);
 				}
 			}
 		}
@@ -892,7 +890,10 @@ void USceneComponent::DetachFromParent(bool bMaintainWorldPosition)
 {
 	if(AttachParent != NULL)
 	{
-		UnWeldFromParent();
+		if (UPrimitiveComponent * PrimComp = Cast<UPrimitiveComponent>(this))
+		{
+			PrimComp->UnWeldFromParent();
+		}
 
 		// Make sure parent points to us if we're registered
 		checkf(!bRegistered || AttachParent->AttachChildren.Contains(this), TEXT("Attempt to detach SceneComponent '%s' owned by '%s' from AttachParent '%s' while not attached."), *GetName(), (GetOwner() ? *GetOwner()->GetName() : TEXT("Unowned")), *AttachParent->GetName());
