nicer UI response

diff --git a/Engine/Source/Runtime/Engine/Private/Materials/Material.cpp b/Engine/Source/Runtime/Engine/Private/Materials/Material.cpp
--- a/Engine/Source/Runtime/Engine/Private/Materials/Material.cpp
+++ b/Engine/Source/Runtime/Engine/Private/Materials/Material.cpp
@@ -2194,7 +2194,7 @@ bool UMaterial::CanEditChange(const UProperty* InProperty) const
 		{
 			static auto* CVar = IConsoleManager::Get().FindTConsoleVariableDataInt(TEXT("r.DBuffer"));
 
-			return CVar->GetValueOnGameThread() > 0;
+			return MaterialDomain == MD_Surface && CVar->GetValueOnGameThread() > 0;
 		}		
 
 		if (PropertyName == GET_MEMBER_NAME_STRING_CHECKED(UMaterial, BlendableLocation) ||
