UE4: Fix a crash when attempting to delete a world that you currently have loaded. This assert is not necessary. TTP#345088

diff --git a/Engine/Source/Runtime/Engine/Private/LevelUtils.cpp b/Engine/Source/Runtime/Engine/Private/LevelUtils.cpp
--- a/Engine/Source/Runtime/Engine/Private/LevelUtils.cpp
+++ b/Engine/Source/Runtime/Engine/Private/LevelUtils.cpp
@@ -33,9 +33,8 @@ ULevelStreaming* FLevelUtils::FindStreamingLevel(const ULevel* Level)
 {
 	ULevelStreaming* MatchingLevel = NULL;
 
-	if (Level)
+	if (Level && Level->OwningWorld)
 	{
-		check(Level->OwningWorld);
 		for( int32 LevelIndex = 0 ; LevelIndex < Level->OwningWorld->StreamingLevels.Num() ; ++LevelIndex )
 		{
 			ULevelStreaming* CurStreamingLevel = Level->OwningWorld->StreamingLevels[ LevelIndex ];
