Fix detach offset issue for simulated bodies

diff --git a/Engine/Source/Runtime/Engine/Private/SceneComponent.cpp b/Engine/Source/Runtime/Engine/Private/SceneComponent.cpp
--- a/Engine/Source/Runtime/Engine/Private/SceneComponent.cpp
+++ b/Engine/Source/Runtime/Engine/Private/SceneComponent.cpp
@@ -766,10 +766,14 @@ void USceneComponent::AttachTo(class USceneComponent* Parent, FName InSocketName
 			if (PrimitiveComponent && PrimitiveComponent->BodyInstance.bSimulatePhysics && !bWeldSimulatedBodies && GetWorld() && GetWorld()->IsGameWorld())
 			{
 				//Since the object is physically simulated it can't be the case that it's a child of object A and being attached to object B (at runtime)
-				//UpdateComponentToWorldWithParent(Parent, false);
-				RelativeLocation = ComponentToWorld.GetLocation();
-				RelativeRotation = ComponentToWorld.GetRotation().Rotator();
-				RelativeScale3D = ComponentToWorld.GetScale3D();
+				if (bMaintainWorldPosition == false)	//User tried to attach but physically based so detach. However, if they provided relative coordinates we should still get the correct position
+				{
+					UpdateComponentToWorldWithParent(Parent, false);
+					RelativeLocation = ComponentToWorld.GetLocation();
+					RelativeRotation = ComponentToWorld.GetRotation().Rotator();
+					RelativeScale3D = ComponentToWorld.GetScale3D();
+				}
+
 				return;
 			}
 		}
