Fix for particle light per view data oob acces when >1 view and using a CameraOffset module.

If the system contains a CameraOffset Module and Lights then I create an array of indices for each light into the per view data.
The only other way would be to add per view data for all views for all particle lights.
I can imagine cases where this would be lots so I'm trading some memory savings for a bit of complexity here.
Another bonus here is that he emitters don't need to know anything about each other.

diff --git a/Engine/Source/Runtime/Engine/Private/Particles/ParticleSystemRender.cpp b/Engine/Source/Runtime/Engine/Private/Particles/ParticleSystemRender.cpp
--- a/Engine/Source/Runtime/Engine/Private/Particles/ParticleSystemRender.cpp
+++ b/Engine/Source/Runtime/Engine/Private/Particles/ParticleSystemRender.cpp
@@ -1193,13 +1193,33 @@ void GatherParticleLightData(const FDynamicSpriteEmitterReplayDataBase& Source,
 				FVector LightPosition = Source.bUseLocalSpace ? FVector(InLocalToWorld.TransformPosition(ParticlePosition)) : ParticlePosition;
 				
 				// Calculate light positions per-view if we are using a camera offset.
-				// OutParticleLights is accumulated over MULTIPLE emitters.  Some may use camera offset, some may not.
-				// if some do and some don't, then OutParticleLights.PerViewData will cease to have NumLights*NumViews entries and 
-				// GetViewDependentData on it will become bogus and in the worst case, crash from reading out of bounds.
-				// Hack fix is to just add all views always.  Real fix would be detecting whether any emitter was using cameraoffset and 
-				// add for all emitters in that case.
 				if (Source.CameraPayloadOffset != 0)
 				{
+					//Reserve mem for the indices into the per view data.
+					if (InViewFamily.Views.Num() > 1)
+					{	
+						OutParticleLights.InstancePerViewDataIndices.Reserve(OutParticleLights.InstanceData.Num() + ParticleCount);
+						if (InViewFamily.Views.Num() > 1 && OutParticleLights.InstancePerViewDataIndices.Num() == 0)
+						{
+							//If this is the first time we've needed to add these then we need to fill data up to this light.
+							int32 NumLights = OutParticleLights.InstanceData.Num();
+							OutParticleLights.InstancePerViewDataIndices.AddUninitialized(OutParticleLights.InstanceData.Num());
+							for (int32 i = 0; i < NumLights; ++i)
+							{
+								OutParticleLights.InstancePerViewDataIndices[i].PerViewIndex = i;
+								OutParticleLights.InstancePerViewDataIndices[i].bHasPerViewData = false;
+							}
+						}
+
+						//If non-zero, the numbers of instances to indices should always match.
+						check(OutParticleLights.InstancePerViewDataIndices.Num() == OutParticleLights.InstanceData.Num());
+						//Add index of this light into the per veiw data.
+						FSimpleLightInstacePerViewIndexData PerViewIndexData;
+						PerViewIndexData.bHasPerViewData = true;
+						PerViewIndexData.PerViewIndex = OutParticleLights.PerViewData.Num();
+						OutParticleLights.InstancePerViewDataIndices.Add(PerViewIndexData);
+					}
+
 					for(int32 ViewIndex = 0; ViewIndex < InViewFamily.Views.Num(); ++ViewIndex)
 					{
 						FSimpleLightPerViewEntry PerViewData;
@@ -1212,13 +1232,21 @@ void GatherParticleLightData(const FDynamicSpriteEmitterReplayDataBase& Source,
 				}
 				else
 				{
+					//Add index of this light into the per view data if needed.
+					if (OutParticleLights.InstancePerViewDataIndices.Num() != 0)
+					{
+						//If non-zero, the numbers of instances to indices should always match.
+						check(OutParticleLights.InstancePerViewDataIndices.Num() == OutParticleLights.InstanceData.Num());
+						FSimpleLightInstacePerViewIndexData PerViewIndexData;
+						PerViewIndexData.bHasPerViewData = false;
+						PerViewIndexData.PerViewIndex = OutParticleLights.PerViewData.Num();
+						OutParticleLights.InstancePerViewDataIndices.Add(PerViewIndexData);
+					}
+
 					// When not using camera-offset, output one position for all views to share. 
 					FSimpleLightPerViewEntry PerViewData;
 					PerViewData.Position = LightPosition;
-					for (int32 ViewIndex = 0; ViewIndex < InViewFamily.Views.Num(); ++ViewIndex)
-					{
-						OutParticleLights.PerViewData.Add(PerViewData);
-					}									
+					OutParticleLights.PerViewData.Add(PerViewData);
 				}
 
 				// Add an entry for the light instance.
diff --git a/Engine/Source/Runtime/Engine/Public/PrimitiveSceneProxy.h b/Engine/Source/Runtime/Engine/Public/PrimitiveSceneProxy.h
--- a/Engine/Source/Runtime/Engine/Public/PrimitiveSceneProxy.h
+++ b/Engine/Source/Runtime/Engine/Public/PrimitiveSceneProxy.h
@@ -35,6 +35,18 @@ public:
 	FVector Position;
 };
 
+/**
+* Index into the Per-view data for each instance.
+* Most uses wont need to add > 1 per view data.
+* This array will be the same size as InstanceData for uses that require per view data. Otherwise it will be empty.
+*/
+class FSimpleLightInstacePerViewIndexData
+{
+public:
+	uint32 PerViewIndex : 31;
+	uint32 bHasPerViewData;
+};
+
 /** Data pertaining to a set of simple dynamic lights */
 class FSimpleLightArray
 {
@@ -43,6 +55,8 @@ public:
 	TArray<FSimpleLightEntry, TMemStackAllocator<>> InstanceData;
 	/** Per-view data for each light */
 	TArray<FSimpleLightPerViewEntry, TMemStackAllocator<>> PerViewData;
+	/** Indices into the per-view data for each light. */
+	TArray<FSimpleLightInstacePerViewIndexData, TMemStackAllocator<>> InstancePerViewDataIndices;
 
 public:
 
@@ -57,12 +71,16 @@ public:
 		// If InstanceData has an equal number of elements to PerViewData then all views share the same PerViewData.
 		if (InstanceData.Num() == PerViewData.Num())
 		{
+			check(InstancePerViewDataIndices.Num() == 0);
 			return PerViewData[LightIndex];
 		}
 		else 
 		{
-			// Calculate strided per-view index.
-			const int32 PerViewDataIndex = (LightIndex * NumViews) + ViewIndex;
+			check(InstancePerViewDataIndices.Num() == InstanceData.Num());
+
+			// Calculate per-view index.
+			FSimpleLightInstacePerViewIndexData PerViewIndex = InstancePerViewDataIndices[LightIndex];
+			const int32 PerViewDataIndex = PerViewIndex.PerViewIndex + ( PerViewIndex.bHasPerViewData ? ViewIndex : 0 );
 			return PerViewData[PerViewDataIndex];
 		}
 	}
