Weld when changing from kinematic to simulated

diff --git a/Engine/Source/Runtime/Engine/Private/PhysicsEngine/BodyInstance.cpp b/Engine/Source/Runtime/Engine/Private/PhysicsEngine/BodyInstance.cpp
--- a/Engine/Source/Runtime/Engine/Private/PhysicsEngine/BodyInstance.cpp
+++ b/Engine/Source/Runtime/Engine/Private/PhysicsEngine/BodyInstance.cpp
@@ -2043,9 +2043,28 @@ void FBodyInstance::SetInstanceSimulatePhysics(bool bSimulate, bool bMaintainPhy
 #endif
 
 	// If we are enabling simulation, and we are the root body of our component, we detach the component 
-	if ((bSimulate == true) && OwnerComponent.IsValid() && OwnerComponent->IsRegistered() && (OwnerComponent->AttachParent != NULL) && (OwnerComponent->GetBodyInstance() == this))
+	if ((bSimulate == true) && OwnerComponent.IsValid() && OwnerComponent->IsRegistered() && OwnerComponent->GetBodyInstance() == this)
 	{
-		OwnerComponent->DetachFromParent(true);
+		if (OwnerComponent->AttachParent)
+		{
+			OwnerComponent->DetachFromParent(true);
+		}
+		else if (bSimulatePhysics == false)	//if we're switching from kinematic to simulated
+		{
+			//if we are the root in hierarchy we must make sure to update children that are welded
+			TArray<FBodyInstance*> ChildrenBodies;
+			TArray<FName> ChildrenLabels;
+			OwnerComponent->GetWeldedBodies(ChildrenBodies, ChildrenLabels);
+
+			for (int32 ChildIdx = 0; ChildIdx < ChildrenBodies.Num(); ++ChildIdx)
+			{
+				FBodyInstance * ChildBI = ChildrenBodies[ChildIdx];
+				if (ChildBI != this)
+				{
+					Weld(ChildBI, ChildBI->OwnerComponent->GetSocketTransform(ChildrenLabels[ChildIdx]));
+				}
+			}
+		}
 	}
 
 	bSimulatePhysics = bSimulate;
diff --git a/Engine/Source/Runtime/Engine/Private/PrimitiveComponent.cpp b/Engine/Source/Runtime/Engine/Private/PrimitiveComponent.cpp
--- a/Engine/Source/Runtime/Engine/Private/PrimitiveComponent.cpp
+++ b/Engine/Source/Runtime/Engine/Private/PrimitiveComponent.cpp
@@ -369,16 +369,19 @@ UPrimitiveComponent * GetRootWelded(UPrimitiveComponent * PrimComponent, FName P
 
 void UPrimitiveComponent::GetWeldedBodies(TArray<FBodyInstance*> & OutWeldedBodies, TArray<FName> & OutLabels)
 {
-	if (BodyInstance.bWelded)
-	{
-		OutWeldedBodies.Add(&BodyInstance);
-		OutLabels.Add(NAME_None);
+	OutWeldedBodies.Add(&BodyInstance);
+	OutLabels.Add(NAME_None);
 			
-		for (USceneComponent * Child : AttachChildren)
+	for (USceneComponent * Child : AttachChildren)
+	{
+		if (UPrimitiveComponent * PrimChild = Cast<UPrimitiveComponent>(Child))
 		{
-			if (UPrimitiveComponent * PrimChild = Cast<UPrimitiveComponent>(Child))
+			if (FBodyInstance * BI = PrimChild->GetBodyInstance())
 			{
-				PrimChild->GetWeldedBodies(OutWeldedBodies, OutLabels);
+				if (BI->bWelded)
+				{
+					PrimChild->GetWeldedBodies(OutWeldedBodies, OutLabels);
+				}
 			}
 		}
 	}
@@ -467,8 +470,8 @@ void UPrimitiveComponent::UnWeldFromParent()
 
 			bool bHasBodySetup = GetBodySetup() != NULL;
 
-			//we need to temporarily turn off AutoWeld
-			if (bHasBodySetup)
+			//if BodyInstance hasn't already been created we need to initialize it
+			if (bHasBodySetup && NewRootBI->IsValidBodyInstance() == false)
 			{
 				bool bPrevAutoWeld = NewRootBI->bAutoWeld;
 				NewRootBI->bAutoWeld = false;
