Added Owner* GameplayTag methods to OrionAnimInstance. Fixed up wrong parents on BP classes.

Made UAnimInstance::GetOwningPawn const

diff --git a/Engine/Source/Runtime/Engine/Classes/Animation/AnimInstance.h b/Engine/Source/Runtime/Engine/Classes/Animation/AnimInstance.h
--- a/Engine/Source/Runtime/Engine/Classes/Animation/AnimInstance.h
+++ b/Engine/Source/Runtime/Engine/Classes/Animation/AnimInstance.h
@@ -307,8 +307,8 @@ public:
 
 	// kismet event functions
 
-	UFUNCTION(BlueprintPure, Category = "Animation")
-	virtual APawn* TryGetPawnOwner();
+	UFUNCTION(BlueprintCallable, Category = "Animation")
+	virtual APawn* TryGetPawnOwner() const;
 
 public:
 	/** Returns the owning actor of this AnimInstance */
diff --git a/Engine/Source/Runtime/Engine/Private/Animation/AnimInstance.cpp b/Engine/Source/Runtime/Engine/Private/Animation/AnimInstance.cpp
--- a/Engine/Source/Runtime/Engine/Private/Animation/AnimInstance.cpp
+++ b/Engine/Source/Runtime/Engine/Private/Animation/AnimInstance.cpp
@@ -196,7 +196,7 @@ AActor* UAnimInstance::GetOwningActor() const
 	return OwnerComponent->GetOwner();
 }
 
-APawn* UAnimInstance::TryGetPawnOwner()
+APawn* UAnimInstance::TryGetPawnOwner() const
 {
 	USkeletalMeshComponent* OwnerComponent = GetSkelMeshComponent();
 	if (AActor* OwnerActor = OwnerComponent->GetOwner())
