Fixed ensure with CSM on zero sized view frustum

diff --git a/Engine/Source/Runtime/Engine/Private/DirectionalLightComponent.cpp b/Engine/Source/Runtime/Engine/Private/DirectionalLightComponent.cpp
--- a/Engine/Source/Runtime/Engine/Private/DirectionalLightComponent.cpp
+++ b/Engine/Source/Runtime/Engine/Private/DirectionalLightComponent.cpp
@@ -525,7 +525,9 @@ private:
 		{
 			CascadeSphere.W = FMath::Max(CascadeSphere.W, FVector::DistSquared(CascadeFrustumVerts[Index], CascadeSphere.Center));
 		}
-		CascadeSphere.W = FMath::Sqrt(CascadeSphere.W); 
+
+		// Don't allow the bounds to reach 0 (INF)
+		CascadeSphere.W = FMath::Max(FMath::Sqrt(CascadeSphere.W), 1.0f); 
 
 		if (OutCascadeSettings)
 		{
