Fixed compile errors for non-editor build

diff --git a/Engine/Source/Runtime/Engine/Private/DistanceFieldAtlas.cpp b/Engine/Source/Runtime/Engine/Private/DistanceFieldAtlas.cpp
--- a/Engine/Source/Runtime/Engine/Private/DistanceFieldAtlas.cpp
+++ b/Engine/Source/Runtime/Engine/Private/DistanceFieldAtlas.cpp
@@ -287,14 +287,18 @@ uint32 FBuildDistanceFieldThreadRunnable::Run()
 	return 0;
 }
 
-FDistanceFieldAsyncQueue::FDistanceFieldAsyncQueue() :
-	MeshUtilities(NULL)
+FDistanceFieldAsyncQueue::FDistanceFieldAsyncQueue() 
 {
+#if WITH_EDITOR
+	MeshUtilities = NULL;
+#endif
+
 	ThreadRunnable = new FBuildDistanceFieldThreadRunnable(this);
 }
 
 void FDistanceFieldAsyncQueue::AddTask(FAsyncDistanceFieldTask* Task)
 {
+#if WITH_EDITOR
 	if (!MeshUtilities)
 	{
 		MeshUtilities = &FModuleManager::Get().LoadModuleChecked<IMeshUtilities>(TEXT("MeshUtilities"));
@@ -316,10 +320,14 @@ void FDistanceFieldAsyncQueue::AddTask(FAsyncDistanceFieldTask* Task)
 		TScopedPointer<FQueuedThreadPool> WorkerThreadPool = TScopedPointer<FQueuedThreadPool>(CreateWorkerThreadPool());
 		Build(Task, *WorkerThreadPool);
 	}
+#else
+	UE_LOG(LogStaticMesh,Fatal,TEXT("Tried to build a distance field without editor support (this should have been done during cooking)"));
+#endif
 }
 
 void FDistanceFieldAsyncQueue::Build(FAsyncDistanceFieldTask* Task, FQueuedThreadPool& ThreadPool)
 {
+#if WITH_EDITOR
 	const FStaticMeshLODResources& LODModel = Task->StaticMesh->RenderData->LODResources[0];
 
 	MeshUtilities->GenerateSignedDistanceFieldVolumeData(
@@ -332,6 +340,7 @@ void FDistanceFieldAsyncQueue::Build(FAsyncDistanceFieldTask* Task, FQueuedThrea
 		*Task->GeneratedVolumeData);
 
 	CompletedTasks.Push(Task);
+#endif
 }
 
 void FDistanceFieldAsyncQueue::AddReferencedObjects(FReferenceCollector& Collector)
@@ -345,6 +354,7 @@ void FDistanceFieldAsyncQueue::AddReferencedObjects(FReferenceCollector& Collect
 
 void FDistanceFieldAsyncQueue::ProcessAsyncTasks()
 {
+#if WITH_EDITOR
 	TArray<FAsyncDistanceFieldTask*> LocalCompletedTasks;
 	CompletedTasks.PopAll(LocalCompletedTasks);
 
@@ -385,6 +395,7 @@ void FDistanceFieldAsyncQueue::ProcessAsyncTasks()
 	{
 		ThreadRunnable->Launch();
 	}
+#endif
 }
 
 void FDistanceFieldAsyncQueue::Shutdown()
