Fix skeletal mesh welding getting the wrong socket name

diff --git a/Engine/Source/Runtime/Engine/Private/PrimitiveComponent.cpp b/Engine/Source/Runtime/Engine/Private/PrimitiveComponent.cpp
--- a/Engine/Source/Runtime/Engine/Private/PrimitiveComponent.cpp
+++ b/Engine/Source/Runtime/Engine/Private/PrimitiveComponent.cpp
@@ -343,12 +343,14 @@ UPrimitiveComponent * GetRootWelded(UPrimitiveComponent * PrimComponent, FName P
 		}
 	}
 
-	FName SocketName = ParentSocketName; //because of skeletal mesh it's important that we check along the bones that we attached
+	FName PrevSocketName = ParentSocketName;
+	FName SocketName = NAME_None; //because of skeletal mesh it's important that we check along the bones that we attached
 	FBodyInstance * RootBI = NULL;
 	for (; RootComponent; RootComponent = Cast<UPrimitiveComponent>(RootComponent->AttachParent))
 	{
 		Result = RootComponent;
-		SocketName = RootComponent->AttachSocketName;
+		SocketName = PrevSocketName;
+		PrevSocketName = RootComponent->AttachSocketName;
 
 		RootBI = RootComponent->GetBodyInstance(SocketName);
 		if (RootBI && RootBI->bWelded)
