Hopeful fix for mac compile error

diff --git a/Engine/Source/Runtime/Engine/Private/DistanceFieldAtlas.cpp b/Engine/Source/Runtime/Engine/Private/DistanceFieldAtlas.cpp
--- a/Engine/Source/Runtime/Engine/Private/DistanceFieldAtlas.cpp
+++ b/Engine/Source/Runtime/Engine/Private/DistanceFieldAtlas.cpp
@@ -296,6 +296,9 @@ FDistanceFieldAsyncQueue::FDistanceFieldAsyncQueue()
 	ThreadRunnable = new FBuildDistanceFieldThreadRunnable(this);
 }
 
+FDistanceFieldAsyncQueue::~FDistanceFieldAsyncQueue()
+{}
+
 void FDistanceFieldAsyncQueue::AddTask(FAsyncDistanceFieldTask* Task)
 {
 #if WITH_EDITOR
diff --git a/Engine/Source/Runtime/Engine/Public/DistanceFieldAtlas.h b/Engine/Source/Runtime/Engine/Public/DistanceFieldAtlas.h
--- a/Engine/Source/Runtime/Engine/Public/DistanceFieldAtlas.h
+++ b/Engine/Source/Runtime/Engine/Public/DistanceFieldAtlas.h
@@ -163,6 +163,8 @@ public:
 
 	ENGINE_API FDistanceFieldAsyncQueue();
 
+	ENGINE_API virtual ~FDistanceFieldAsyncQueue();
+
 	/** Adds a new build task. */
 	ENGINE_API void AddTask(FAsyncDistanceFieldTask* Task);
 
