#UE4

- More getters for Local Player Context.

diff --git a/Engine/Source/Runtime/Engine/Classes/Engine/LocalPlayer.h b/Engine/Source/Runtime/Engine/Classes/Engine/LocalPlayer.h
--- a/Engine/Source/Runtime/Engine/Classes/Engine/LocalPlayer.h
+++ b/Engine/Source/Runtime/Engine/Classes/Engine/LocalPlayer.h
@@ -67,6 +67,16 @@ struct ENGINE_API FLocalPlayerContext
 		return CastChecked<T>(GetHUD(), ECastCheckedType::NullAllowed);
 	}
 
+	/** Getter for the base pawn of this player */
+	class APawn* GetPawn() const;
+
+	/** Templated getter for the player's pawn */
+	template<class T>
+	FORCEINLINE T* GetPawn() const
+	{
+		return CastChecked<T>(GetPawn(), ECastCheckedType::NullAllowed);
+	}
+
 private:	
 
 	/* Set the local player. */
diff --git a/Engine/Source/Runtime/Engine/Private/LocalPlayer.cpp b/Engine/Source/Runtime/Engine/Private/LocalPlayer.cpp
--- a/Engine/Source/Runtime/Engine/Private/LocalPlayer.cpp
+++ b/Engine/Source/Runtime/Engine/Private/LocalPlayer.cpp
@@ -117,6 +117,13 @@ AHUD* FLocalPlayerContext::GetHUD() const
 	return PC ? PC->MyHUD : NULL;
 }
 
+class APawn* FLocalPlayerContext::GetPawn() const
+{
+	check(LocalPlayer.IsValid())
+	APlayerController* PC = LocalPlayer->PlayerController;
+	return PC ? PC->GetPawn() : NULL;
+}
+
 void FLocalPlayerContext::SetLocalPlayer( const ULocalPlayer* InLocalPlayer )
 {
 	LocalPlayer = InLocalPlayer;
