Fix hardware survey running on Dedicated server

diff --git a/Engine/Source/Runtime/Engine/Private/UnrealEngine.cpp b/Engine/Source/Runtime/Engine/Private/UnrealEngine.cpp
--- a/Engine/Source/Runtime/Engine/Private/UnrealEngine.cpp
+++ b/Engine/Source/Runtime/Engine/Private/UnrealEngine.cpp
@@ -5354,7 +5354,7 @@ void UEngine::TickHardwareSurvey()
 bool UEngine::IsHardwareSurveyRequired()
 {
 	// Analytics must have been initialized FIRST.
-	if (!FEngineAnalytics::IsAvailable())
+	if (!FEngineAnalytics::IsAvailable() || IsRunningDedicatedServer() )
 	{
 		return false;
 	}
