Fixed TScopedPointer issue

diff --git a/Engine/Source/Runtime/Engine/Private/DistanceFieldAtlas.cpp b/Engine/Source/Runtime/Engine/Private/DistanceFieldAtlas.cpp
--- a/Engine/Source/Runtime/Engine/Private/DistanceFieldAtlas.cpp
+++ b/Engine/Source/Runtime/Engine/Private/DistanceFieldAtlas.cpp
@@ -296,9 +296,6 @@ FDistanceFieldAsyncQueue::FDistanceFieldAsyncQueue()
 	ThreadRunnable = new FBuildDistanceFieldThreadRunnable(this);
 }
 
-FDistanceFieldAsyncQueue::~FDistanceFieldAsyncQueue()
-{}
-
 void FDistanceFieldAsyncQueue::AddTask(FAsyncDistanceFieldTask* Task)
 {
 #if WITH_EDITOR
@@ -320,7 +317,7 @@ void FDistanceFieldAsyncQueue::AddTask(FAsyncDistanceFieldTask* Task)
 	}
 	else
 	{
-		TScopedPointer<FQueuedThreadPool> WorkerThreadPool = TScopedPointer<FQueuedThreadPool>(CreateWorkerThreadPool());
+		TScopedPointer<FQueuedThreadPool> WorkerThreadPool(CreateWorkerThreadPool());
 		Build(Task, *WorkerThreadPool);
 	}
 #else
diff --git a/Engine/Source/Runtime/Engine/Public/DistanceFieldAtlas.h b/Engine/Source/Runtime/Engine/Public/DistanceFieldAtlas.h
--- a/Engine/Source/Runtime/Engine/Public/DistanceFieldAtlas.h
+++ b/Engine/Source/Runtime/Engine/Public/DistanceFieldAtlas.h
@@ -163,8 +163,6 @@ public:
 
 	ENGINE_API FDistanceFieldAsyncQueue();
 
-	ENGINE_API virtual ~FDistanceFieldAsyncQueue();
-
 	/** Adds a new build task. */
 	ENGINE_API void AddTask(FAsyncDistanceFieldTask* Task);
 
