fixed material using Subsurface translucency and a single float constant as input

diff --git a/Engine/Shaders/ShaderVersion.usf b/Engine/Shaders/ShaderVersion.usf
--- a/Engine/Shaders/ShaderVersion.usf
+++ b/Engine/Shaders/ShaderVersion.usf
@@ -2,4 +2,4 @@
 // Each time the console command is executed it generates a new GUID. As this file is included
 // in common.usf (which should be included in any shader) it allows to invalidate the shader DDC.
 // 
-// GUID = A7702884495089DC221506B8C4961E6C1
+// GUID = 1B144F2441A3CDA26663BB820BFD049B
\ No newline at end of file
diff --git a/Engine/Source/Runtime/Engine/Private/Materials/HLSLMaterialTranslator.h b/Engine/Source/Runtime/Engine/Private/Materials/HLSLMaterialTranslator.h
--- a/Engine/Source/Runtime/Engine/Private/Materials/HLSLMaterialTranslator.h
+++ b/Engine/Source/Runtime/Engine/Private/Materials/HLSLMaterialTranslator.h
@@ -339,9 +339,13 @@ public:
 
 			EMaterialShadingModel MaterialShadingModel = Material->GetShadingModel();
 
-			if (MaterialShadingModel == MSM_Subsurface || MaterialShadingModel == MSM_PreintegratedSkin || MaterialShadingModel == MSM_SubsurfaceProfile)
+			if (Material->GetMaterialDomain() == MD_Surface
+				&& (MaterialShadingModel == MSM_Subsurface || MaterialShadingModel == MSM_PreintegratedSkin || MaterialShadingModel == MSM_SubsurfaceProfile))
 			{
-				int32 SubsurfaceColor = Material->CompilePropertyAndSetMaterialProperty(MP_SubsurfaceColor, this);
+				// Note we don't test for the blend mode as you can have a translucent material using the subsurface shading model
+
+				// another ForceCast as CompilePropertyAndSetMaterialProperty() can return MCT_Float which we don't want here
+				int32 SubsurfaceColor = ForceCast(Material->CompilePropertyAndSetMaterialProperty(MP_SubsurfaceColor, this), MCT_Float3, true);
 
 				static FName NameSubsurfaceProfile(TEXT("__SubsurfaceProfile"));
 
