#ue4 - CharacterMovement: clamp effects of high friction so as to not reverse direction of acceleration when moving normally or when moving through a fluid.

diff --git a/Engine/Source/Runtime/Engine/Private/CharacterMovementComponent.cpp b/Engine/Source/Runtime/Engine/Private/CharacterMovementComponent.cpp
--- a/Engine/Source/Runtime/Engine/Private/CharacterMovementComponent.cpp
+++ b/Engine/Source/Runtime/Engine/Private/CharacterMovementComponent.cpp
@@ -2148,13 +2148,13 @@ void UCharacterMovementComponent::CalcVelocity(float DeltaTime, float Friction,
 	{
 		const FVector AccelDir = SafeNormalPrecise(Acceleration);
 		const float VelSize = Velocity.Size();
-		Velocity = Velocity - (Velocity - AccelDir * VelSize) * DeltaTime * Friction;
+		Velocity = Velocity - (Velocity - AccelDir * VelSize) * FMath::Min(DeltaTime * Friction, 1.f);
 	}
 
 	// Apply fluid friction
 	if (bFluid)
 	{
-		Velocity = Velocity * (1.f - Friction * DeltaTime);
+		Velocity = Velocity * (1.f - FMath::Min(Friction * DeltaTime, 1.f));
 	}
 
 	// Apply input acceleration
