refactored code, could be even simpler with FAutoConsoleVariableSink but would need to be tested with reloading the module

diff --git a/Engine/Source/Runtime/Engine/Private/UnrealEngine.cpp b/Engine/Source/Runtime/Engine/Private/UnrealEngine.cpp
--- a/Engine/Source/Runtime/Engine/Private/UnrealEngine.cpp
+++ b/Engine/Source/Runtime/Engine/Private/UnrealEngine.cpp
@@ -379,13 +379,19 @@ void RefreshSamplerStatesCallback()
 	}
 }
 
-ENGINE_API void InitializeRenderingCVarsCaching()
+void RefreshEngineSettings()
 {
 	extern void FreeSkeletalMeshBuffersSinkCallback();
-	IConsoleManager::Get().RegisterConsoleVariableSink(FConsoleCommandDelegate::CreateStatic(&RefreshSamplerStatesCallback));
-	IConsoleManager::Get().RegisterConsoleVariableSink(FConsoleCommandDelegate::CreateStatic(&ScalabilityCVarsSinkCallback));
-	IConsoleManager::Get().RegisterConsoleVariableSink(FConsoleCommandDelegate::CreateStatic(&FreeSkeletalMeshBuffersSinkCallback));
-	IConsoleManager::Get().RegisterConsoleVariableSink(FConsoleCommandDelegate::CreateStatic(&SystemResolutionSinkCallback));
+
+	RefreshSamplerStatesCallback();
+	ScalabilityCVarsSinkCallback();
+	FreeSkeletalMeshBuffersSinkCallback();
+	SystemResolutionSinkCallback();
+}
+
+ENGINE_API void InitializeRenderingCVarsCaching()
+{
+	IConsoleManager::Get().RegisterConsoleVariableSink(FConsoleCommandDelegate::CreateStatic(&RefreshEngineSettings));
 
 	// Initialise this to invalid
 	GCachedScalabilityCVars.MaterialQualityLevel = EMaterialQualityLevel::Num;
@@ -397,11 +403,7 @@ ENGINE_API void InitializeRenderingCVarsCaching()
 
 void ShutdownRenderingCVarsCaching()
 {
-	extern void FreeSkeletalMeshBuffersSinkCallback();
-	IConsoleManager::Get().UnregisterConsoleVariableSink(FConsoleCommandDelegate::CreateStatic(&RefreshSamplerStatesCallback));
-	IConsoleManager::Get().UnregisterConsoleVariableSink(FConsoleCommandDelegate::CreateStatic(&ScalabilityCVarsSinkCallback));
-	IConsoleManager::Get().UnregisterConsoleVariableSink(FConsoleCommandDelegate::CreateStatic(&FreeSkeletalMeshBuffersSinkCallback));
-	IConsoleManager::Get().UnregisterConsoleVariableSink(FConsoleCommandDelegate::CreateStatic(&SystemResolutionSinkCallback));
+	IConsoleManager::Get().UnregisterConsoleVariableSink(FConsoleCommandDelegate::CreateStatic(&RefreshEngineSettings));
 }
 
 namespace
