Stop running construction scripts for duplicated actors in PIE

diff --git a/Engine/Source/Runtime/Engine/Classes/Engine/Level.h b/Engine/Source/Runtime/Engine/Classes/Engine/Level.h
--- a/Engine/Source/Runtime/Engine/Classes/Engine/Level.h
+++ b/Engine/Source/Runtime/Engine/Classes/Engine/Level.h
@@ -401,6 +401,8 @@ class ULevel : public ULevelBase, public IInterface_AssetUserData
 	uint32										bRequireFullVisibilityToRender:1;
 	/** Whether this level is specific to client, visibility state will not be replicated to server	*/
 	uint32										bClientOnlyVisible:1;
+	/** Whether this level was duplicated for PIE	*/
+	uint32										bWasDuplicatedForPIE:1;
 	/** Current index into actors array for updating components.							*/
 	int32										CurrentActorIndexForUpdateComponents;
 
@@ -470,6 +472,7 @@ public:
 #endif // WITH_EDITOR
 	virtual void PostLoad() override;
 	virtual void PreSave() override;
+	virtual void PostDuplicate(bool bDuplicateForPIE) override;
 	static void AddReferencedObjects(UObject* InThis, FReferenceCollector& Collector);
 	// End UObject interface.
 
diff --git a/Engine/Source/Runtime/Engine/Private/Level.cpp b/Engine/Source/Runtime/Engine/Private/Level.cpp
--- a/Engine/Source/Runtime/Engine/Private/Level.cpp
+++ b/Engine/Source/Runtime/Engine/Private/Level.cpp
@@ -641,6 +641,13 @@ void ULevel::PostLoad()
 #endif
 }
 
+void ULevel::PostDuplicate(bool bDuplicateForPIE)
+{
+	Super::PostDuplicate(bDuplicateForPIE);
+
+	bWasDuplicatedForPIE = bDuplicateForPIE;
+}
+
 UWorld* ULevel::GetWorld() const
 {
 	return OwningWorld;
diff --git a/Engine/Source/Runtime/Engine/Private/World.cpp b/Engine/Source/Runtime/Engine/Private/World.cpp
--- a/Engine/Source/Runtime/Engine/Private/World.cpp
+++ b/Engine/Source/Runtime/Engine/Private/World.cpp
@@ -1671,17 +1671,18 @@ void UWorld::AddToWorld( ULevel* Level, const FTransform& LevelTransform )
 		// Make sure code thinks components are not currently attached.
 		Level->bAreComponentsCurrentlyRegistered = false;
 
+#if WITH_EDITOR
+			// Pretend here that we are loading package to avoid package dirtying during components registration
+		TGuardValue<bool> IsEditorLoadingPackage(GIsEditorLoadingPackage, (GIsEditor ? true : GIsEditorLoadingPackage));
+#endif
+		// We don't need to rerun construction scripts if we have cooked data or we are playing in editor unless the PIE world was loaded
+		// from disk rather than duplicated
+		const bool bRerunConstructionScript = !(FPlatformProperties::RequiresCookedData() || (IsPlayInEditor() && Level->bWasDuplicatedForPIE));
+		
 		// Incrementally update components.
 		int32 NumComponentsToUpdate = GEngine->LevelStreamingComponentsRegistrationGranularity;
 		do
 		{
-#if WITH_EDITOR
-			// Pretend here that we are loading package to avoid package dirtying during components registration
-			TGuardValue<bool> IsEditorLoadingPackage(GIsEditorLoadingPackage, (GIsEditor ? true : GIsEditorLoadingPackage));
-#endif
-			// We don't need to rerun construction scripts if we have cooked data or we are playing in editor unless the PIE world was loaded
-			// from disk rather than duplicated
-			const bool bRerunConstructionScript = !FPlatformProperties::RequiresCookedData() && (!IsPlayInEditor() || HasAnyFlags(RF_WasLoaded));
 			Level->IncrementalUpdateComponents( (!IsGameWorld() || IsRunningCommandlet()) ? 0 : NumComponentsToUpdate, bRerunConstructionScript );
 		}
 		while( (!bConsiderTimeLimit || !IsTimeLimitExceeded( TEXT("updating components"), StartTime, Level )) && !Level->bAreComponentsCurrentlyRegistered );
@@ -2754,7 +2755,7 @@ void UWorld::InitializeActorsForPlay(const FURL& InURL, bool bResetTime)
 	// Update world and the components of all levels.	
 	// We don't need to rerun construction scripts if we have cooked data or we are playing in editor unless the PIE world was loaded
 	// from disk rather than duplicated
-	const bool bRerunConstructionScript = !FPlatformProperties::RequiresCookedData() && (!IsPlayInEditor() || HasAnyFlags(RF_WasLoaded));
+	const bool bRerunConstructionScript = !(FPlatformProperties::RequiresCookedData() || (IsPlayInEditor() && PersistentLevel->bWasDuplicatedForPIE));
 	UpdateWorldComponents( bRerunConstructionScript, true );
 
 	// Reset indices till we have a chance to rearrange actor list at the end of this function.
