fixed navmesh update after changing properties of NavModifierVolume

diff --git a/Engine/Source/Runtime/Engine/Classes/AI/Navigation/NavModifierVolume.h b/Engine/Source/Runtime/Engine/Classes/AI/Navigation/NavModifierVolume.h
--- a/Engine/Source/Runtime/Engine/Classes/AI/Navigation/NavModifierVolume.h
+++ b/Engine/Source/Runtime/Engine/Classes/AI/Navigation/NavModifierVolume.h
@@ -18,4 +18,8 @@ class ANavModifierVolume : public AVolume, public INavRelevantInterface
 
 	virtual void GetNavigationData(struct FNavigationRelevantData& Data) const override;
 	virtual FBox GetNavigationBounds() const;
+
+#if WITH_EDITOR
+	virtual void PostEditChangeProperty(FPropertyChangedEvent& PropertyChangedEvent) override;
+#endif
 };
diff --git a/Engine/Source/Runtime/Engine/Private/AI/Navigation/NavModifierVolume.cpp b/Engine/Source/Runtime/Engine/Private/AI/Navigation/NavModifierVolume.cpp
--- a/Engine/Source/Runtime/Engine/Private/AI/Navigation/NavModifierVolume.cpp
+++ b/Engine/Source/Runtime/Engine/Private/AI/Navigation/NavModifierVolume.cpp
@@ -25,3 +25,17 @@ FBox ANavModifierVolume::GetNavigationBounds() const
 {
 	return GetComponentsBoundingBox();
 }
+
+#if WITH_EDITOR
+
+void ANavModifierVolume::PostEditChangeProperty(FPropertyChangedEvent& PropertyChangedEvent)
+{
+	Super::PostEditChangeProperty(PropertyChangedEvent);
+
+	if (PropertyChangedEvent.Property && PropertyChangedEvent.Property->GetFName() == GET_MEMBER_NAME_CHECKED(ANavModifierVolume, AreaClass))
+	{
+		UNavigationSystem::UpdateNavOctree(this);
+	}
+}
+
+#endif
\ No newline at end of file
