Ensure that new sounds would be audible before playing

UGameplayStatics::PlaySoundAttached did not pass a location to FAudioDevice::CreateComponent so that it could call IsAudibleSimple on the sound.
Also realised that this would not take overriden attenuation settings into account so added a new parameter to pass that through as well.

diff --git a/Engine/Source/Runtime/Engine/Classes/Sound/SoundBase.h b/Engine/Source/Runtime/Engine/Classes/Sound/SoundBase.h
--- a/Engine/Source/Runtime/Engine/Classes/Sound/SoundBase.h
+++ b/Engine/Source/Runtime/Engine/Classes/Sound/SoundBase.h
@@ -81,8 +81,11 @@ public:
 
 	/** 
 	 * Does a simple range check to all listeners to test hearability
+	 *
+	 * @param Location				Location to check against
+	 * @param AttenuationSettings	Optional Attenuation override if not using settings from the sound
 	 */
-	bool IsAudibleSimple( const FVector Location );
+	bool IsAudibleSimple( const FVector Location, USoundAttenuation* AttenuationSettings = NULL );
 
 	/** 
 	 * Returns the farthest distance at which the sound could be heard
diff --git a/Engine/Source/Runtime/Engine/Private/AudioDevice.cpp b/Engine/Source/Runtime/Engine/Private/AudioDevice.cpp
--- a/Engine/Source/Runtime/Engine/Private/AudioDevice.cpp
+++ b/Engine/Source/Runtime/Engine/Private/AudioDevice.cpp
@@ -2073,7 +2073,7 @@ bool FAudioDevice::LocationIsAudible( FVector Location, float MaxDistance )
 	return( false );
 }
 
-UAudioComponent* FAudioDevice::CreateComponent( USoundBase* Sound, UWorld* World, AActor* Actor, bool bPlay, bool bStopWhenOwnerDestroyed, const FVector* Location )
+UAudioComponent* FAudioDevice::CreateComponent( USoundBase* Sound, UWorld* World, AActor* Actor, bool bPlay, bool bStopWhenOwnerDestroyed, const FVector* Location, USoundAttenuation* AttenuationSettings )
 {
 	UAudioComponent* AudioComponent = NULL;
 
@@ -2085,7 +2085,7 @@ UAudioComponent* FAudioDevice::CreateComponent( USoundBase* Sound, UWorld* World
 			// Don't create component on destroyed actor.
 		}
 		// Either no actor or actor is still alive.
-		else if( Location && !Sound->IsAudibleSimple( *Location ) )
+		else if( Location && !Sound->IsAudibleSimple( *Location, AttenuationSettings ) )
 		{
 			// Don't create a sound component for short sounds that start out of range of any listener
 			UE_LOG(LogAudio, Log, TEXT( "AudioComponent not created for out of range Sound %s" ), *Sound->GetName() );
@@ -2111,6 +2111,7 @@ UAudioComponent* FAudioDevice::CreateComponent( USoundBase* Sound, UWorld* World
 			AudioComponent->bAutoDestroy = bPlay;
 			AudioComponent->bStopWhenOwnerDestroyed = bStopWhenOwnerDestroyed;
 			AudioComponent->bVisualizeComponent	= false;
+			AudioComponent->AttenuationSettings = AttenuationSettings;
 			if (Location)
 			{
 				AudioComponent->SetWorldLocation(*Location);
diff --git a/Engine/Source/Runtime/Engine/Private/GameplayStatics.cpp b/Engine/Source/Runtime/Engine/Private/GameplayStatics.cpp
--- a/Engine/Source/Runtime/Engine/Private/GameplayStatics.cpp
+++ b/Engine/Source/Runtime/Engine/Private/GameplayStatics.cpp
@@ -601,7 +601,7 @@ void UGameplayStatics::PlaySoundAtLocation(UObject* WorldContextObject, class US
 
 	if (GEngine && GEngine->UseSound() && ThisWorld->GetNetMode() != NM_DedicatedServer)
 	{
-		if( Sound->IsAudibleSimple( Location ) )
+		if( Sound->IsAudibleSimple( Location, AttenuationSettings ) )
 		{
 			FActiveSound NewActiveSound;
 			NewActiveSound.World = ThisWorld;
@@ -656,7 +656,7 @@ void UGameplayStatics::PlayDialogueAtLocation(UObject* WorldContextObject, class
 		UWorld* ThisWorld = GEngine->GetWorldFromContextObject(WorldContextObject);
 		const bool bIsInGameWorld = ThisWorld->IsGameWorld();
 
-		if( Sound->IsAudibleSimple( Location ) )
+		if( Sound->IsAudibleSimple( Location, AttenuationSettings ) )
 		{
 			FActiveSound NewActiveSound;
 			NewActiveSound.World = ThisWorld;
@@ -705,7 +705,14 @@ class UAudioComponent* UGameplayStatics::PlaySoundAttached(class USoundBase* Sou
 		return NULL;
 	}
 
-	UAudioComponent* AudioComponent = FAudioDevice::CreateComponent( Sound, AttachToComponent->GetWorld(), AttachToComponent->GetOwner(), false, bStopWhenAttachedToDestroyed );
+	// Location used to check whether to create a component if out of range
+	FVector TestLocation = Location;
+	if (LocationType != EAttachLocation::KeepWorldPosition)
+	{
+		TestLocation = AttachToComponent->GetRelativeTransform().TransformPosition(Location);
+	}
+
+	UAudioComponent* AudioComponent = FAudioDevice::CreateComponent( Sound, AttachToComponent->GetWorld(), AttachToComponent->GetOwner(), false, bStopWhenAttachedToDestroyed, &TestLocation, AttenuationSettings );
 	if(AudioComponent)
 	{
 		const bool bIsInGameWorld = AudioComponent->GetWorld()->IsGameWorld();
@@ -725,7 +732,6 @@ class UAudioComponent* UGameplayStatics::PlaySoundAttached(class USoundBase* Sou
 		AudioComponent->bIsUISound				= !bIsInGameWorld;
 		AudioComponent->bAutoDestroy			= true;
 		AudioComponent->SubtitlePriority		= 10000.f; // Fixme: pass in? Do we want that exposed to blueprints though?
-		AudioComponent->AttenuationSettings		= AttenuationSettings;
 		AudioComponent->Play(StartTime);
 	}
 
@@ -751,7 +757,14 @@ class UAudioComponent* UGameplayStatics::PlayDialogueAttached(class UDialogueWav
 		return NULL;
 	}
 
-	UAudioComponent* AudioComponent = FAudioDevice::CreateComponent( Sound, AttachToComponent->GetWorld(), AttachToComponent->GetOwner(), false, bStopWhenAttachedToDestroyed );
+	// Location used to check whether to create a component if out of range
+	FVector TestLocation = Location;
+	if (LocationType != EAttachLocation::KeepWorldPosition)
+	{
+		TestLocation = AttachToComponent->GetRelativeTransform().TransformPosition(Location);
+	}
+
+	UAudioComponent* AudioComponent = FAudioDevice::CreateComponent( Sound, AttachToComponent->GetWorld(), AttachToComponent->GetOwner(), false, bStopWhenAttachedToDestroyed, &TestLocation, AttenuationSettings );
 	if(AudioComponent)
 	{
 		const bool bIsGameWorld = AudioComponent->GetWorld()->IsGameWorld();
@@ -771,7 +784,6 @@ class UAudioComponent* UGameplayStatics::PlayDialogueAttached(class UDialogueWav
 		AudioComponent->bIsUISound				= !bIsGameWorld;
 		AudioComponent->bAutoDestroy			= true;
 		AudioComponent->SubtitlePriority		= 10000.f; // Fixme: pass in? Do we want that exposed to blueprints though?
-		AudioComponent->AttenuationSettings		= AttenuationSettings;
 		AudioComponent->Play(StartTime);
 	}
 
diff --git a/Engine/Source/Runtime/Engine/Private/SoundBase.cpp b/Engine/Source/Runtime/Engine/Private/SoundBase.cpp
--- a/Engine/Source/Runtime/Engine/Private/SoundBase.cpp
+++ b/Engine/Source/Runtime/Engine/Private/SoundBase.cpp
@@ -41,7 +41,7 @@ float USoundBase::GetMaxAudibleDistance()
 	return 0.f;
 }
 
-bool USoundBase::IsAudibleSimple( const FVector Location )
+bool USoundBase::IsAudibleSimple( const FVector Location, USoundAttenuation* AttenuationSettings )
 {
 	// No audio device means no listeners to check against
 	if( !GEngine || !GEngine->GetAudioDevice() )
@@ -56,7 +56,8 @@ bool USoundBase::IsAudibleSimple( const FVector Location )
 	}
 
 	// Is this SourceActor within the MaxAudibleDistance of any of the listeners?
-	return GEngine->GetAudioDevice()->LocationIsAudible( Location, GetMaxAudibleDistance() );
+	float MaxAudibleDistance = AttenuationSettings != nullptr ? AttenuationSettings->Attenuation.GetMaxDimension() : GetMaxAudibleDistance();
+	return GEngine->GetAudioDevice()->LocationIsAudible( Location, MaxAudibleDistance );
 }
 
 bool USoundBase::IsAudible( const FVector &SourceLocation, const FVector &ListenerLocation, AActor* SourceActor, bool& bIsOccluded, bool bCheckOcclusion )
diff --git a/Engine/Source/Runtime/Engine/Public/AudioDevice.h b/Engine/Source/Runtime/Engine/Public/AudioDevice.h
--- a/Engine/Source/Runtime/Engine/Public/AudioDevice.h
+++ b/Engine/Source/Runtime/Engine/Public/AudioDevice.h
@@ -342,7 +342,7 @@ public:
 	/**
 	 * Creates an audio component to handle playing a sound cue
 	 */
-	static class UAudioComponent* CreateComponent( class USoundBase* Sound, class UWorld* World, AActor*  AActor  = NULL, bool Play = true, bool bStopWhenOwnerDestroyed = false, const FVector* Location = NULL );
+	static class UAudioComponent* CreateComponent( class USoundBase* Sound, class UWorld* World, AActor*  AActor  = NULL, bool Play = true, bool bStopWhenOwnerDestroyed = false, const FVector* Location = NULL, USoundAttenuation* AttenuationSettings = NULL );
 
 	/**
 	 * Adds an active sound to the audio device
