Cap sub-stepping max time to Max num substeps * max sub step delta time

diff --git a/Engine/Source/Runtime/Engine/Private/PhysicsEngine/PhysSubstepTasks.cpp b/Engine/Source/Runtime/Engine/Private/PhysicsEngine/PhysSubstepTasks.cpp
--- a/Engine/Source/Runtime/Engine/Private/PhysicsEngine/PhysSubstepTasks.cpp
+++ b/Engine/Source/Runtime/Engine/Private/PhysicsEngine/PhysSubstepTasks.cpp
@@ -265,10 +265,10 @@ float FPhysSubstepTask::UpdateTime(float UseDelta)
 	float FrameRateInv = 1.f / FrameRate;
 
 	//Figure out how big dt to make for desired framerate
-	DeltaSeconds = UseDelta;
-	NumSubsteps = FMath::CeilToInt(UseDelta * FrameRateInv);
+	DeltaSeconds = FMath::Min(UseDelta, MaxSubSteps * FrameRate);
+	NumSubsteps = FMath::CeilToInt(DeltaSeconds * FrameRateInv);
 	NumSubsteps = FMath::Max(NumSubsteps > MaxSubSteps ? MaxSubSteps : NumSubsteps, (uint32) 1);
-	SubTime = UseDelta / NumSubsteps;
+	SubTime = DeltaSeconds / NumSubsteps;
 
 	return SubTime;
 }
