removed log spam

diff --git a/Engine/Source/Runtime/Engine/Private/Materials/Material.cpp b/Engine/Source/Runtime/Engine/Private/Materials/Material.cpp
--- a/Engine/Source/Runtime/Engine/Private/Materials/Material.cpp
+++ b/Engine/Source/Runtime/Engine/Private/Materials/Material.cpp
@@ -226,8 +226,6 @@ public:
 					*OutValue = 0.0f;
 				}
 
-				FPlatformMisc::LowLevelOutputDebugStringf(TEXT(">>>> Material %p GetScalarValue %f\n"), this, *OutValue);
-
 				return true;
 			}
 
diff --git a/Engine/Source/Runtime/Engine/Private/Materials/MaterialInstance.cpp b/Engine/Source/Runtime/Engine/Private/Materials/MaterialInstance.cpp
--- a/Engine/Source/Runtime/Engine/Private/Materials/MaterialInstance.cpp
+++ b/Engine/Source/Runtime/Engine/Private/Materials/MaterialInstance.cpp
@@ -223,8 +223,6 @@ bool FMaterialInstanceResource::GetScalarValue(
 			*OutValue = 0.0f;
 		}
 
-		FPlatformMisc::LowLevelOutputDebugStringf(TEXT(">>>> MaterialInstance %p GetScalarValue %f\n"), this, *OutValue);
-
 		return true;
 	}
 
