refactored handing of a cvar, removed one cvar

diff --git a/Engine/Source/Runtime/Core/Private/HAL/ConsoleManager.cpp b/Engine/Source/Runtime/Core/Private/HAL/ConsoleManager.cpp
--- a/Engine/Source/Runtime/Core/Private/HAL/ConsoleManager.cpp
+++ b/Engine/Source/Runtime/Core/Private/HAL/ConsoleManager.cpp
@@ -1465,14 +1465,6 @@ static TAutoConsoleVariable<float> CVarAmbientOcclusionStaticFraction(TEXT("r.Am
 	TEXT(" 1: AO affects the stat lighting"),
 	ECVF_Default);
 
-static TAutoConsoleVariable<int32> CVarAtmosphereRender(
-	TEXT("r.AtmosphereRender"),
-	1,
-	TEXT("Defines atmosphere will render or not. Only changed by r.Atmosphere console command\n")
-	TEXT(" 0: off\n")
-	TEXT(" 1: on"),
-	ECVF_ReadOnly);
-
 static TAutoConsoleVariable<int32> CVarShadowQuality(
 	TEXT("r.ShadowQuality"),
 	5,
diff --git a/Engine/Source/Runtime/Engine/Private/Atmosphere/Atmosphere.cpp b/Engine/Source/Runtime/Engine/Private/Atmosphere/Atmosphere.cpp
--- a/Engine/Source/Runtime/Engine/Private/Atmosphere/Atmosphere.cpp
+++ b/Engine/Source/Runtime/Engine/Private/Atmosphere/Atmosphere.cpp
@@ -206,16 +206,24 @@ public:
 };
 #endif
 
+
+static TAutoConsoleVariable<int32> CVarAtmosphereRender(
+	TEXT("r.Atmosphere"),
+	1,
+	TEXT("Defines atmosphere will render or not. Only changed by r.Atmosphere console command.\n")
+	TEXT("Enable/Disable Atmosphere, Load/Unload related data.\n")
+	TEXT(" 0: off\n")
+	TEXT(" 1: on (default)"));
+
 // On CPU
 void UAtmosphericFogComponent::InitResource()
 {
-#if !(UE_BUILD_SHIPPING || UE_BUILD_TEST)
-	static const auto ICVar = IConsoleManager::Get().FindTConsoleVariableDataInt(TEXT("r.AtmosphereRender"));
-	if (ICVar && ICVar->GetValueOnGameThread() == 0)
+	bool bNewAtmosphere = CVarAtmosphereRender.GetValueOnGameThread() != 0;
+
+	if(!bNewAtmosphere)
 	{
-		return; // Don't do intialize resource when r.AtmosphereRender is off
+		return; // Don't do initialize resource when Atmosphere Render is off
 	}
-#endif
 
 	if (PrecomputeCounter.GetValue() >= EValid)
 	{
@@ -788,3 +796,37 @@ void UAtmosphericFogComponent::ApplyComponentInstanceData(TSharedPtr<FComponentI
 	PrecomputeCounter.Set(EValid);
 	InitResource();
 }
+
+/**
+ * Gets called any time cvars change (on the main thread), we check if r.Atmosphere has changed and update the components.
+ */
+static void AtmosphereRenderSinkFunction()
+{
+	bool bNewAtmosphere = CVarAtmosphereRender.GetValueOnGameThread() != 0;
+
+	// by default we assume the state is true
+	static bool GAtmosphere = true;
+
+	if (GAtmosphere != bNewAtmosphere)
+	{
+		GAtmosphere = bNewAtmosphere;
+
+		for(TObjectIterator<UAtmosphericFogComponent> It; It; ++It)
+		{
+			UAtmosphericFogComponent* Comp = *It;
+			if (!Comp->GetOuter()->HasAnyFlags(RF_ClassDefaultObject))
+			{
+				if (bNewAtmosphere && Comp->IsRegistered())
+				{
+					Comp->InitResource();
+				}
+				else
+				{
+					Comp->ReleaseResource();
+				}
+			}
+		}
+	}
+}
+
+FAutoConsoleVariableSink CVarAtmosphereRenderSink(FConsoleCommandDelegate::CreateStatic(&AtmosphereRenderSinkFunction));
diff --git a/Engine/Source/Runtime/Renderer/Private/AtmosphereRendering.cpp b/Engine/Source/Runtime/Renderer/Private/AtmosphereRendering.cpp
--- a/Engine/Source/Runtime/Renderer/Private/AtmosphereRendering.cpp
+++ b/Engine/Source/Runtime/Renderer/Private/AtmosphereRendering.cpp
@@ -1612,64 +1612,11 @@ FAtmosphericFogSceneInfo::~FAtmosphericFogSceneInfo()
 #endif
 }
 
-/** A mip bias applied globally to all textures. */
-static int32 GAtmosphere = 1;
-
-/**
- * Sets the global texture mip bias and forces all textures to update and stream as needed.
- */
-static void SetAtmosphere(const TArray<FString>& Args)
-{
-#if !(UE_BUILD_SHIPPING || UE_BUILD_TEST)
-	if (Args.Num() == 1 && Args[0].Len() > 0)
-	{
-		int32 NewAtmosphere = FPlatformString::Atoi(*Args[0]) > 0 ? 1 : 0;
-
-		if (GAtmosphere != NewAtmosphere )
-		{
-			GAtmosphere = NewAtmosphere;
-
-			IConsoleVariable* ICVar = IConsoleManager::Get().FindConsoleVariable(TEXT("r.AtmosphereRender"));
-			if(ICVar)
-			{
-				ICVar->Set(GAtmosphere);
-			}
-
-			for(TObjectIterator<UAtmosphericFogComponent> It; It; ++It)
-			{
-				UAtmosphericFogComponent* Comp = *It;
-				if (!Comp->GetOuter()->HasAnyFlags(RF_ClassDefaultObject))
-				{
-					if (GAtmosphere > 0 && Comp->IsRegistered())
-					{
-						Comp->InitResource(); // Initialize resource
-					}
-					else
-					{
-						Comp->ReleaseResource();
-					}
-				}
-			}
-		}
-
-		UE_LOG(LogConsoleResponse,Display,TEXT("r.Atmosphere %s"), GAtmosphere ? TEXT("Enabled") : TEXT("Disabled"));
-	}
-#endif
-}
-
-/** Console command for setting the global texture mip bias. */
-static FAutoConsoleCommand GAtmosphereConsoleCmd(
-	TEXT("r.Atmosphere"),
-	TEXT("Enable/Disable Atmosphere, Load/Unload related data."),
-	FConsoleCommandWithArgsDelegate::CreateStatic(&SetAtmosphere)
-	);
-
 bool ShouldRenderAtmosphere(const FSceneViewFamily& Family)
 {
 	const FEngineShowFlags EngineShowFlags = Family.EngineShowFlags;
 
-	return GAtmosphere
-		&& GSupportsVolumeTextureRendering
+	return GSupportsVolumeTextureRendering
 		&& EngineShowFlags.Atmosphere
 		&& !EngineShowFlags.ShaderComplexity
 		&& !EngineShowFlags.StationaryLightOverlap 
