Do component/actor over mouse updates while game is paused

diff --git a/Engine/Source/Runtime/Engine/Classes/GameFramework/PlayerController.h b/Engine/Source/Runtime/Engine/Classes/GameFramework/PlayerController.h
--- a/Engine/Source/Runtime/Engine/Classes/GameFramework/PlayerController.h
+++ b/Engine/Source/Runtime/Engine/Classes/GameFramework/PlayerController.h
@@ -1135,6 +1135,7 @@ public:
 protected:
 
 	/** Internal. */
+	void TickPlayerInput(const float DeltaSeconds, const bool bGamePaused);
 	virtual void ProcessPlayerInput(const float DeltaTime, const bool bGamePaused);
 	virtual void BuildInputStack(TArray<UInputComponent*>& InputStack);
 	void ProcessForceFeedback(const float DeltaTime, const bool bGamePaused);
diff --git a/Engine/Source/Runtime/Engine/Private/PlayerController.cpp b/Engine/Source/Runtime/Engine/Private/PlayerController.cpp
--- a/Engine/Source/Runtime/Engine/Private/PlayerController.cpp
+++ b/Engine/Source/Runtime/Engine/Private/PlayerController.cpp
@@ -2025,49 +2025,7 @@
 		ServerShortTimeout();
 	}
 
-	if (ULocalPlayer* LocalPlayer = Cast<ULocalPlayer>(Player))
-	{
-		if (bEnableMouseOverEvents)
-		{
-			FVector2D MousePosition;
-			FHitResult HitResult;
-			bool bHit = false;
-			
-			if (LocalPlayer->ViewportClient->GetMousePosition(MousePosition))
-			{
-				bHit = GetHitResultAtScreenPosition(MousePosition, CurrentClickTraceChannel, true, /*out*/ HitResult);
-			}
+	TickPlayerInput(DeltaTime, DeltaTime == 0.f);
 
-			UPrimitiveComponent* PreviousComponent = CurrentClickablePrimitive.Get();
-			UPrimitiveComponent* CurrentComponent = (bHit ? HitResult.Component.Get() : NULL);
-
-			UPrimitiveComponent::DispatchMouseOverEvents(PreviousComponent, CurrentComponent);
-
-			CurrentClickablePrimitive = CurrentComponent;
-		}
-
-		if (bEnableTouchOverEvents)
-		{
-			for (int32 TouchIndexInt = 0; TouchIndexInt < EKeys::NUM_TOUCH_KEYS; ++TouchIndexInt)
-			{
-				const ETouchIndex::Type FingerIndex = ETouchIndex::Type(TouchIndexInt);
-
-				FHitResult HitResult;
-				const bool bHit = GetHitResultUnderFinger(FingerIndex, CurrentClickTraceChannel, true, /*out*/ HitResult);
-
-				UPrimitiveComponent* PreviousComponent = CurrentTouchablePrimitives[TouchIndexInt].Get();
-				UPrimitiveComponent* CurrentComponent = (bHit ? HitResult.Component.Get() : NULL);
-
-				UPrimitiveComponent::DispatchTouchOverEvents(FingerIndex, PreviousComponent, CurrentComponent);
-
-				CurrentTouchablePrimitives[TouchIndexInt] = CurrentComponent;
-			}
-		}
-	}
-
-	ProcessPlayerInput(DeltaTime, DeltaTime == 0.f);
-	ProcessForceFeedback(DeltaTime, DeltaTime == 0.f);
-
-
 	if ((Player != NULL) && (Player->PlayerController == this))
 	{
@@ -3667,16 +3625,60 @@
 	ReceivedPlayer();
 }
 
+void APlayerController::TickPlayerInput(const float DeltaSeconds, const bool bGamePaused)
+{
+	if (ULocalPlayer* LocalPlayer = Cast<ULocalPlayer>(Player))
+	{
+		if (bEnableMouseOverEvents)
+		{
+			FVector2D MousePosition;
+			FHitResult HitResult;
+			bool bHit = false;
+			
+			if (LocalPlayer->ViewportClient->GetMousePosition(MousePosition))
+			{
+				bHit = GetHitResultAtScreenPosition(MousePosition, CurrentClickTraceChannel, true, /*out*/ HitResult);
+			}
 
+			UPrimitiveComponent* PreviousComponent = CurrentClickablePrimitive.Get();
+			UPrimitiveComponent* CurrentComponent = (bHit ? HitResult.Component.Get() : NULL);
+
+			UPrimitiveComponent::DispatchMouseOverEvents(PreviousComponent, CurrentComponent);
+
+			CurrentClickablePrimitive = CurrentComponent;
+		}
+
+		if (bEnableTouchOverEvents)
+		{
+			for (int32 TouchIndexInt = 0; TouchIndexInt < EKeys::NUM_TOUCH_KEYS; ++TouchIndexInt)
+			{
+				const ETouchIndex::Type FingerIndex = ETouchIndex::Type(TouchIndexInt);
+
+				FHitResult HitResult;
+				const bool bHit = GetHitResultUnderFinger(FingerIndex, CurrentClickTraceChannel, true, /*out*/ HitResult);
+
+				UPrimitiveComponent* PreviousComponent = CurrentTouchablePrimitives[TouchIndexInt].Get();
+				UPrimitiveComponent* CurrentComponent = (bHit ? HitResult.Component.Get() : NULL);
+
+				UPrimitiveComponent::DispatchTouchOverEvents(FingerIndex, PreviousComponent, CurrentComponent);
+
+				CurrentTouchablePrimitives[TouchIndexInt] = CurrentComponent;
+			}
+		}
+	}
+
+	if (PlayerInput != NULL)
+	{
+		ProcessPlayerInput(DeltaSeconds, bGamePaused);
+		ProcessForceFeedback(DeltaSeconds, bGamePaused);
+	}
+}
+
 void APlayerController::TickActor( float DeltaSeconds, ELevelTick TickType, FActorTickFunction& ThisTickFunction )
 {
 	if (TickType == LEVELTICK_PauseTick && !ShouldPerformFullTickWhenPaused())
 	{
-		if (PlayerInput != NULL)
-		{
-			ProcessPlayerInput(DeltaSeconds, true);
-			ProcessForceFeedback(DeltaSeconds, true);
-		}
+		TickPlayerInput(DeltaSeconds, true);
 
 		// Clear axis inputs from previous frame.
 		RotationInput = FRotator::ZeroRotator;
