Attach no longer uses auto weld.
TermBody unwelds children.
Weld first unwelds.

diff --git a/Engine/Source/Runtime/Engine/Classes/PhysicsEngine/BodyInstance.h b/Engine/Source/Runtime/Engine/Classes/PhysicsEngine/BodyInstance.h
--- a/Engine/Source/Runtime/Engine/Classes/PhysicsEngine/BodyInstance.h
+++ b/Engine/Source/Runtime/Engine/Classes/PhysicsEngine/BodyInstance.h
@@ -208,6 +208,9 @@ public:
 	/** Constraint used to allow for easy DOF setup per bodyinstance */
 	FConstraintInstance * DOFConstraint;
 
+	/** The parent body that we are welded to*/
+	FBodyInstance * WeldParent;
+
 protected:
 
 	/** Whether this instance of the object has its own custom walkability override setting. */
diff --git a/Engine/Source/Runtime/Engine/Private/PhysicsEngine/BodyInstance.cpp b/Engine/Source/Runtime/Engine/Private/PhysicsEngine/BodyInstance.cpp
--- a/Engine/Source/Runtime/Engine/Private/PhysicsEngine/BodyInstance.cpp
+++ b/Engine/Source/Runtime/Engine/Private/PhysicsEngine/BodyInstance.cpp
@@ -206,6 +206,7 @@ FBodyInstance::FBodyInstance()
 , LockedAxisMode(0)
 , CustomLockedAxis(FVector::ZeroVector)
 , DOFConstraint(NULL)
+, WeldParent(NULL)
 , bOverrideWalkableSlopeOnInstance(false)
 , bOverrideMaxDepenetrationVelocity(false)
 , MaxDepenetrationVelocity(0.f)
@@ -1397,6 +1398,27 @@ void TermBodyHelper(int32& SceneIndex, PxRigidActor*& PRigidActor, FBodyInstance
  */
 void FBodyInstance::TermBody()
 {
+
+	if (UPrimitiveComponent* OwnerComponentInst = OwnerComponent.Get())
+	{
+		for (USceneComponent * Child : OwnerComponentInst->AttachChildren)
+		{
+			if (UPrimitiveComponent * PrimChild = Cast<UPrimitiveComponent>(Child))
+			{
+				if (FBodyInstance * ChildBI = PrimChild->GetBodyInstance())
+				{
+					if (ChildBI->bWelded)
+					{
+						if (ChildBI->OwnerComponent.IsValid())
+						{
+							ChildBI->OwnerComponent->UnWeldFromParent();
+						}
+					}
+				}
+			}
+		}
+	}
+
 #if WITH_BOX2D
 	if (BodyInstancePtr != NULL)
 	{
diff --git a/Engine/Source/Runtime/Engine/Private/PrimitiveComponent.cpp b/Engine/Source/Runtime/Engine/Private/PrimitiveComponent.cpp
--- a/Engine/Source/Runtime/Engine/Private/PrimitiveComponent.cpp
+++ b/Engine/Source/Runtime/Engine/Private/PrimitiveComponent.cpp
@@ -409,6 +409,8 @@ bool UPrimitiveComponent::WeldToImplementation(USceneComponent * InParent, FName
 		return false;
 	}
 
+	UnWeldFromParent();	//make sure to unweld from wherever we currently are
+
 	FName SocketName;
 	UPrimitiveComponent * RootComponent = GetRootWelded(this, ParentSocketName, &SocketName, true);
 
@@ -416,6 +418,11 @@ bool UPrimitiveComponent::WeldToImplementation(USceneComponent * InParent, FName
 	{
 		if (FBodyInstance * RootBI = RootComponent->GetBodyInstance(SocketName))
 		{
+			if (BI->WeldParent == RootBI)	//already welded so stop
+			{
+				return true;
+			}
+
 			BI->bWelded = true;
 			//There are multiple cases to handle:
 			//Root is kinematic, simulated
@@ -425,6 +432,7 @@ bool UPrimitiveComponent::WeldToImplementation(USceneComponent * InParent, FName
 			//if root is kinematic simply set child to be kinematic and we're done
 			if (RootComponent->IsSimulatingPhysics(SocketName) == false)
 			{
+				BI->WeldParent = NULL;
 				SetSimulatePhysics(false);
 				return false;	//return false because we need to continue with regular body initialization
 			}
@@ -432,6 +440,7 @@ bool UPrimitiveComponent::WeldToImplementation(USceneComponent * InParent, FName
 			//root is simulated so we actually weld the body
 			FTransform RelativeTM = RootComponent == AttachParent ? GetRelativeTransform() : GetComponentToWorld().GetRelativeTransform(RootComponent->GetComponentToWorld());	//if direct parent we already have relative. Otherwise compute it
 			RootBI->Weld(BI, GetComponentToWorld());
+			BI->WeldParent = RootBI;
 
 			return true;
 		}
@@ -469,6 +478,7 @@ void UPrimitiveComponent::UnWeldFromParent()
 			//create new root
 			RootBI->UnWeld(NewRootBI);
 			NewRootBI->bWelded = false;
+			NewRootBI->WeldParent = NULL;
 
 			bool bHasBodySetup = GetBodySetup() != NULL;
 
diff --git a/Engine/Source/Runtime/Engine/Private/SceneComponent.cpp b/Engine/Source/Runtime/Engine/Private/SceneComponent.cpp
--- a/Engine/Source/Runtime/Engine/Private/SceneComponent.cpp
+++ b/Engine/Source/Runtime/Engine/Private/SceneComponent.cpp
@@ -870,7 +870,7 @@ void USceneComponent::AttachTo(class USceneComponent* Parent, FName InSocketName
 		{
 			if (FBodyInstance * BI = PrimitiveComponent->GetBodyInstance())
 			{
-				if (BI->bAutoWeld || bWeldSimulatedBodies)
+				if (bWeldSimulatedBodies)
 				{
 					PrimitiveComponent->WeldToImplementation(AttachParent, AttachSocketName, bWeldSimulatedBodies);
 				}
