Consider all playing SoundWaves to be referenced, regardless of whether in editor or in game.

diff --git a/Engine/Source/Runtime/Engine/Private/ActiveSound.cpp b/Engine/Source/Runtime/Engine/Private/ActiveSound.cpp
--- a/Engine/Source/Runtime/Engine/Private/ActiveSound.cpp
+++ b/Engine/Source/Runtime/Engine/Private/ActiveSound.cpp
@@ -75,20 +75,16 @@ FArchive& operator<<( FArchive& Ar, FActiveSound* ActiveSound )
 
 void FActiveSound::AddReferencedObjects( FReferenceCollector& Collector)
 {
-#if WITH_EDITOR
-	if( GIsEditor )
+	for (auto WaveInstanceIt(WaveInstances.CreateConstIterator()); WaveInstanceIt; ++WaveInstanceIt)
 	{
-		for (auto WaveInstanceIt(WaveInstances.CreateConstIterator()); WaveInstanceIt; ++WaveInstanceIt)
+		FWaveInstance* WaveInstance = WaveInstanceIt.Value();
+		// Avoid recursing back to the wave instance that sourced this active sound
+		if( WaveInstance )
 		{
-			FWaveInstance* WaveInstance = WaveInstanceIt.Value();
-			// Avoid recursing back to the wave instance that sourced this active sound
-			if( WaveInstance )
-			{
-				WaveInstance->AddReferencedObjects( Collector );
-			}
+			WaveInstance->AddReferencedObjects( Collector );
 		}
 	}
-#endif
+
 	Collector.AddReferencedObject(Sound);
 	Collector.AddReferencedObject(LastReverbVolume);
 	Collector.AddReferencedObject(SoundClassOverride);
