Don't crash during GC if a SoundMode (or other non SoundWave/Cue) is playing

diff --git a/Engine/Source/Runtime/Engine/Private/AudioDevice.cpp b/Engine/Source/Runtime/Engine/Private/AudioDevice.cpp
--- a/Engine/Source/Runtime/Engine/Private/AudioDevice.cpp
+++ b/Engine/Source/Runtime/Engine/Private/AudioDevice.cpp
@@ -2450,9 +2450,11 @@ void FAudioDevice::StopSoundsUsingResource(USoundWave* SoundWave, TArray<UAudioC
 	for (int32 ActiveSoundIndex = ActiveSounds.Num() - 1; ActiveSoundIndex >= 0; --ActiveSoundIndex)
 	{
 		FActiveSound* ActiveSound = ActiveSounds[ActiveSoundIndex];
-		if (ActiveSound->Sound->IsA(USoundWave::StaticClass()))
+		for (auto WaveInstanceIt(ActiveSound->WaveInstances.CreateConstIterator()); WaveInstanceIt; ++WaveInstanceIt)
 		{
-			if (ActiveSound->Sound == SoundWave)
+			// If anything the ActiveSound uses the wave then we stop the sound
+			FWaveInstance* WaveInstance = WaveInstanceIt.Value();
+			if (WaveInstance->WaveData == SoundWave)
 			{
 				if (ActiveSound->AudioComponent.IsValid())
 				{
@@ -2460,26 +2462,7 @@ void FAudioDevice::StopSoundsUsingResource(USoundWave* SoundWave, TArray<UAudioC
 				}
 				ActiveSound->Stop(this);
 				bStoppedSounds = true;
-			}
-		}
-		else
-		{
-			check(ActiveSound->Sound->IsA(USoundCue::StaticClass()));
-
-			for (auto WaveInstanceIt(ActiveSound->WaveInstances.CreateConstIterator()); WaveInstanceIt; ++WaveInstanceIt)
-			{
-				// If anything in the SoundCue uses the wave we're going to stop the whole thing for simplicity
-				FWaveInstance* WaveInstance = WaveInstanceIt.Value();
-				if (WaveInstance->WaveData == SoundWave)
-				{
-					if (ActiveSound->AudioComponent.IsValid())
-					{
-						StoppedComponents.Add(ActiveSound->AudioComponent.Get());
-					}
-					ActiveSound->Stop(this);
-					bStoppedSounds = true;
-					break;
-				}
+				break;
 			}
 		}
 	}
