Add a Blueprint pure IsGamePaused function
Fix comments on SetGamePaused

diff --git a/Engine/Source/Runtime/Engine/Classes/Kismet/GameplayStatics.h b/Engine/Source/Runtime/Engine/Classes/Kismet/GameplayStatics.h
--- a/Engine/Source/Runtime/Engine/Classes/Kismet/GameplayStatics.h
+++ b/Engine/Source/Runtime/Engine/Classes/Kismet/GameplayStatics.h
@@ -145,13 +145,20 @@ class ENGINE_API UGameplayStatics : public UBlueprintFunctionLibrary
 	static void SetGlobalTimeDilation(UObject* WorldContextObject, float TimeDilation);
 
 	/**
-	 * Sets the global time dilation
-	 * @param	bPaused		whether the game should be paused or not
+	 * Sets the game's paused state
+	 * @param	bPaused		Whether the game should be paused or not
 	 * @return	Whether the game was successfully paused/unpaused
 	 */
 	UFUNCTION(BlueprintCallable, Category="Game", meta=(HidePin="WorldContextObject", DefaultToSelf="WorldContextObject") )
 	static bool SetGamePaused(UObject* WorldContextObject, bool bPaused);
 
+	/**
+	 * Returns the game's paused state
+	 * @return	Whether the game is currently paused or not
+	 */
+	UFUNCTION(BlueprintPure, Category="Game", meta=(HidePin="WorldContextObject", DefaultToSelf="WorldContextObject") )
+	static bool IsGamePaused(UObject* WorldContextObject);
+
 
 	/** Hurt locally authoritative actors within the radius. Will only hit components that block the Visibility channel.
 	 * @param BaseDamage - The base damage to apply, i.e. the damage at the origin.
diff --git a/Engine/Source/Runtime/Engine/Private/GameplayStatics.cpp b/Engine/Source/Runtime/Engine/Private/GameplayStatics.cpp
--- a/Engine/Source/Runtime/Engine/Private/GameplayStatics.cpp
+++ b/Engine/Source/Runtime/Engine/Private/GameplayStatics.cpp
@@ -131,6 +131,12 @@ bool UGameplayStatics::SetGamePaused(UObject* WorldContextObject, bool bPaused)
 	}
 }
 
+bool UGameplayStatics::IsGamePaused(UObject* WorldContextObject)
+{
+	UWorld* const World = GEngine->GetWorldFromContextObject( WorldContextObject );
+	return World->IsPaused();
+}
+
 /** @RETURN True if weapon trace from Origin hits component VictimComp.  OutHitResult will contain properties of the hit. */
 static bool ComponentIsDamageableFrom(UPrimitiveComponent* VictimComp, FVector const& Origin, AActor const* IgnoredActor, const TArray<AActor*>& IgnoreActors, ECollisionChannel TraceChannel, FHitResult& OutHitResult)
 {
