Merge from shootercert (UE4-Shootergame):

Fix for the network guid issue during seamless travel when we flush guids on the server. The quick fix is to not even send this property to the server, since we don't need to. This has the added benefit of not only fixing the issue, but saves bandwidth. We still need to look into letting guids persist on the server for a bit longer to rully resolve any future conflicts that could happen.

diff --git a/Engine/Source/Runtime/Engine/Private/DataChannel.cpp b/Engine/Source/Runtime/Engine/Private/DataChannel.cpp
--- a/Engine/Source/Runtime/Engine/Private/DataChannel.cpp
+++ b/Engine/Source/Runtime/Engine/Private/DataChannel.cpp
@@ -1998,8 +1998,12 @@ void UActorChannel::BeginContentBlock( UObject * Obj, FOutBunch &Bunch )
 	Bunch << Obj;
 	NET_CHECKSUM(Bunch);
 
-	UClass *ObjClass = Obj->GetClass();
-	Bunch << ObjClass;
+	if ( Connection->Driver->IsServer() )
+	{
+		// Only the server can tell clients to create objects, so no need for the client to send this to the server
+		UClass *ObjClass = Obj->GetClass();
+		Bunch << ObjClass;
+	}
 
 #if !(UE_BUILD_SHIPPING || UE_BUILD_TEST)
 	if (CVarDoReplicationContextString->GetInt() > 0)
@@ -2011,6 +2015,8 @@ void UActorChannel::BeginContentBlock( UObject * Obj, FOutBunch &Bunch )
 
 void UActorChannel::BeginContentBlockForSubObjectDelete( FOutBunch & Bunch, FNetworkGUID & GuidToDelete )
 {
+	check( Connection->Driver->IsServer() );
+
 	// Send a 0 bit to signify we are dealing with sub-objects
 	Bunch.WriteBit( 0 );
 
@@ -2040,7 +2046,7 @@ void UActorChannel::EndContentBlock( UObject *Obj, FOutBunch &Bunch, FClassNetCa
 
 UObject * UActorChannel::ReadContentBlockHeader( FInBunch & Bunch )
 {
-	const bool IsServer = ( Connection->Driver->ServerConnection == NULL );
+	const bool IsServer = Connection->Driver->IsServer();
 
 	if ( Bunch.ReadBit() )
 	{
@@ -2101,64 +2107,68 @@ UObject * UActorChannel::ReadContentBlockHeader( FInBunch & Bunch )
 		return NULL;
 	}
 
-	// Serialize the class in case we have to spawn it.
-	// Manually serialize the object so that we can get the NetGUID (in order to assign it if we spawn the object here)
-	FNetworkGUID ClassNetGUID;
-	UObject * SubObjClassObj = NULL;
-	Connection->PackageMap->SerializeObject( Bunch, UObject::StaticClass(), SubObjClassObj, &ClassNetGUID );
-
-	// Delete sub-object
-	if ( !ClassNetGUID.IsValid() )
+	if ( !IsServer )
 	{
-		if ( IsServer )
-		{
-			UE_LOG( LogNetTraffic, Error, TEXT( "UActorChannel::ReadContentBlockHeader: Client attempted to delete sub-object. Actor: %s" ), *Actor->GetName() );
-			Bunch.SetError();
-			return NULL;
-		}
+		// Serialize the class in case we have to spawn it.
+		// Manually serialize the object so that we can get the NetGUID (in order to assign it if we spawn the object here)
+		FNetworkGUID ClassNetGUID;
+		UObject * SubObjClassObj = NULL;
+		Connection->PackageMap->SerializeObject( Bunch, UObject::StaticClass(), SubObjClassObj, &ClassNetGUID );
 
-		if ( SubObj )
+		// Delete sub-object
+		if ( !ClassNetGUID.IsValid() )
 		{
-			Actor->OnSubobjectDestroyFromReplication( SubObj );
-			SubObj->MarkPendingKill();
+			if ( SubObj )
+			{
+				Actor->OnSubobjectDestroyFromReplication( SubObj );
+				SubObj->MarkPendingKill();
+			}
+			return NULL;
 		}
-		return NULL;
-	}
 
-	UClass * SubObjClass = Cast< UClass >( SubObjClassObj );
+		UClass * SubObjClass = Cast< UClass >( SubObjClassObj );
 
-	// Valid NetGUID but no class was resolved - this is an error
-	if ( SubObjClass == NULL )
-	{
-		UE_LOG( LogNetTraffic, Error, TEXT( "UActorChannel::ReadContentBlockHeader: Unable to read sub-object class. Actor: %s" ), *Actor->GetName() );
-		Bunch.SetError();
-		return NULL;
-	}
+		if ( SubObjClass == NULL )
+		{
+			UE_LOG( LogNetTraffic, Warning, TEXT( "UActorChannel::ReadContentBlockHeader: Unable to read sub-object class. Actor: %s" ), *Actor->GetName() );
 
-	if ( SubObjClass == UObject::StaticClass() )
-	{
-		UE_LOG( LogNetTraffic, Error, TEXT( "UActorChannel::ReadContentBlockHeader: SubObjClass == UObject::StaticClass(). Actor: %s" ), *Actor->GetName() );
-		Bunch.SetError();
-		return NULL;
-	}
+			// Valid NetGUID but no class was resolved - this is an error
+			if ( SubObj == NULL )
+			{
+				UE_LOG( LogNetTraffic, Error, TEXT( "UActorChannel::ReadContentBlockHeader: Unable to read sub-object class (SubObj == NULL). Actor: %s" ), *Actor->GetName() );
+				Bunch.SetError();
+				return NULL;
+			}
+		}
+		else
+		{
+			if ( SubObjClass == UObject::StaticClass() )
+			{
+				UE_LOG( LogNetTraffic, Error, TEXT( "UActorChannel::ReadContentBlockHeader: SubObjClass == UObject::StaticClass(). Actor: %s" ), *Actor->GetName() );
+				Bunch.SetError();
+				return NULL;
+			}
 
-	if ( SubObjClass->IsChildOf( AActor::StaticClass() ) )
-	{
-		UE_LOG( LogNetTraffic, Error, TEXT( "UActorChannel::ReadContentBlockHeader: Sub-object cannot be actor class. Actor: %s" ), *Actor->GetName() );
-		Bunch.SetError();
-		return NULL;
-	}
+			if ( SubObjClass->IsChildOf( AActor::StaticClass() ) )
+			{
+				UE_LOG( LogNetTraffic, Error, TEXT( "UActorChannel::ReadContentBlockHeader: Sub-object cannot be actor class. Actor: %s" ), *Actor->GetName() );
+				Bunch.SetError();
+				return NULL;
+			}
+		}
 
-	if ( SubObj == NULL )
-	{
-		// Construct the sub-object
-		UE_LOG( LogNetTraffic, Log, TEXT( "UActorChannel::ReadContentBlockHeader: Instantiating sub-object. Class: %s, Actor: %s" ), *SubObjClass->GetName(), *Actor->GetName() );
-		SubObj = ConstructObject< UObject >( SubObjClass, Actor );
-		check( SubObj != NULL );
-		Actor->OnSubobjectCreatedFromReplication( SubObj );
-		Connection->Driver->GuidCache->RegisterNetGUID_Client( NetGUID, SubObj );
+		if ( SubObj == NULL )
+		{
+			// Construct the sub-object
+			UE_LOG( LogNetTraffic, Log, TEXT( "UActorChannel::ReadContentBlockHeader: Instantiating sub-object. Class: %s, Actor: %s" ), *SubObjClass->GetName(), *Actor->GetName() );
+			SubObj = ConstructObject< UObject >( SubObjClass, Actor );
+			check( SubObj != NULL );
+			Actor->OnSubobjectCreatedFromReplication( SubObj );
+			Connection->Driver->GuidCache->RegisterNetGUID_Client( NetGUID, SubObj );
+		}
 	}
 
+	check( SubObj != NULL );
 	check( Cast< AActor >( SubObj ) == NULL );
 
 	if ( !SubObj->IsIn( Actor ) )
