Refactored Montage API to support multi-montages.

diff --git a/Engine/Source/Editor/AnimGraph/Private/AnimPreviewInstance.cpp b/Engine/Source/Editor/AnimGraph/Private/AnimPreviewInstance.cpp
--- a/Engine/Source/Editor/AnimGraph/Private/AnimPreviewInstance.cpp
+++ b/Engine/Source/Editor/AnimGraph/Private/AnimPreviewInstance.cpp
@@ -174,7 +174,7 @@ void UAnimPreviewInstance::MontagePreview_SetReverse(bool bInReverse)
 	if (FAnimMontageInstance * CurMontageInstance = GetActiveMontageInstance())
 	{
 		// copy the current playrate
-		CurMontageInstance->PlayRate = PlayRate;
+		CurMontageInstance->SetPlayRate(PlayRate);
 	}
 }
 
diff --git a/Engine/Source/Runtime/Engine/Classes/Animation/AnimInstance.h b/Engine/Source/Runtime/Engine/Classes/Animation/AnimInstance.h
--- a/Engine/Source/Runtime/Engine/Classes/Animation/AnimInstance.h
+++ b/Engine/Source/Runtime/Engine/Classes/Animation/AnimInstance.h
@@ -334,6 +334,9 @@ public:
 	UFUNCTION()
 	void AnimNotify_Sound(UAnimNotify* Notify);
 
+	/*********************************************************************************************
+	* SlotAnimation
+	********************************************************************************************* */
 public:
 
 	/** Play normal animation asset on the slot node. You can only play one asset (whether montage or animsequence) at a time. */
@@ -348,13 +351,25 @@ public:
 	UFUNCTION(BlueprintCallable, Category="Animation")
 	bool IsPlayingSlotAnimation(UAnimSequenceBase* Asset, FName SlotNodeName ); 
 
-	/** Makes a montage jump to a named section. */
+	/*********************************************************************************************
+	 * AnimMontage
+	 ********************************************************************************************* */
+public:
+	/** Plays an animation montage. Returns the length of the animation montage in seconds. Returns 0.f if failed to play. */
+	UFUNCTION(BlueprintCallable, Category = "Animation")
+	float Montage_Play(UAnimMontage * MontageToPlay, float InPlayRate);
+
+	/** Stops the animation montage. If reference is NULL, it will stop ALL active montages. */
+	UFUNCTION(BlueprintCallable, Category = "Animation")
+	void Montage_Stop(float InBlendOutTime, UAnimMontage * Montage = NULL);
+
+	/** Makes a montage jump to a named section. If Montage reference is NULL, it will do that to all active montages. */
 	UFUNCTION(BlueprintCallable, Category="Animation")
-	void Montage_JumpToSection(FName SectionName);
+	void Montage_JumpToSection(FName SectionName, UAnimMontage * Montage = NULL);
 
-	/** Makes a montage jump to the end of a named section. */
+	/** Makes a montage jump to the end of a named section. If Montage reference is NULL, it will do that to all active montages. */
 	UFUNCTION(BlueprintCallable, Category="Animation")
-	void Montage_JumpToSectionsEnd(FName SectionName);
+	void Montage_JumpToSectionsEnd(FName SectionName, UAnimMontage * Montage = NULL);
 
 	/** Relink new next section AFTER SectionNameToChange in run-time
 	 *	You can link section order the way you like in editor, but in run-time if you'd like to change it dynamically, 
@@ -366,32 +381,93 @@ public:
 	 * @param NextSection	: new next section 
 	 */
 	UFUNCTION(BlueprintCallable, Category="Animation")
-	void Montage_SetNextSection(FName SectionNameToChange, FName NextSection);
+	void Montage_SetNextSection(FName SectionNameToChange, FName NextSection, UAnimMontage * Montage = NULL);
 
 	/** Change AnimMontage play rate. NewPlayRate = 1.0 is the default playback rate. */
 	UFUNCTION(BlueprintCallable, Category="Animation")
 	void Montage_SetPlayRate(UAnimMontage* Montage, float NewPlayRate = 1.f);
 
-	/** Plays an animation montage. Returns the length of the animation montage in seconds. Returns 0.f if failed to play. */
-	UFUNCTION(BlueprintCallable, Category="Animation")
-	float Montage_Play(UAnimMontage* MontageToPlay, float InPlayRate = 1.f);
-
-	/** Stops the animation montage. */
-	UFUNCTION(BlueprintCallable, Category="Animation")
-	void Montage_Stop(float InBlendOutTime = 0.25f);
-
-	/** Returns true if the animation montage is active. */
+	/** Returns true if the animation montage is active. If the Montage reference is NULL, it will return true if any Montage is active. */
 	UFUNCTION(BlueprintCallable, Category="Animation")
 	bool Montage_IsActive(UAnimMontage* Montage);  
 
-	/** Returns true if the animation montage is currently playing. */
+	/** Returns true if the animation montage is currently active and playing. 
+	If reference is NULL, it will return true is ANY montage is currently active and playing. */
 	UFUNCTION(BlueprintCallable, Category="Animation")
 	bool Montage_IsPlaying(UAnimMontage* Montage); 
 
 	/** Returns the name of the current animation montage section. */
 	UFUNCTION(BlueprintCallable, Category="Animation")
-	FName Montage_GetCurrentSection();
+	FName Montage_GetCurrentSection(UAnimMontage* Montage = NULL);
+
+	/** Called when a montage starts blending out, whether interrupted or finished */
+	UPROPERTY(BlueprintAssignable)
+	FOnMontageBlendingOutStartedMCDelegate OnMontageBlendingOut;
+	
+	/** Called when a montage has ended, whether interrupted or finished*/
+	UPROPERTY(BlueprintAssignable)
+	FOnMontageEndedMCDelegate OnMontageEnded;
+
+	/*********************************************************************************************
+	* AnimMontage native C++ interface
+	********************************************************************************************* */
+public:	
+	void Montage_SetEndDelegate(FOnMontageEnded & OnMontageEnded, UAnimMontage * Montage = NULL);
+	
+	void Montage_SetBlendingOutDelegate(FOnMontageBlendingOutStarted & OnMontageBlendingOut, UAnimMontage* Montage = NULL);
+	
+	/** Get pointer to BlendingOutStarted delegate for Montage.
+	If Montage reference is NULL, it will pick the first active montage found. */
+	FOnMontageBlendingOutStarted * Montage_GetBlendingOutDelegate(UAnimMontage * Montage = NULL);
+
+	/** Get Current Montage Position */
+	float Montage_GetPosition(UAnimMontage* Montage);
+	
+	/** Set position. */
+	void Montage_SetPosition(UAnimMontage* Montage, float NewPosition);
+	
+	/** return true if Montage is not currently active. (not valid or blending out) */
+	bool Montage_GetIsStopped(UAnimMontage* Montage);
+
+	/** Get PlayRate for Montage.
+	If Montage reference is NULL, PlayRate for any Active Montage will be returned.
+	If Montage is not playing, 0 is returned. */
+	float Montage_GetPlayRate(UAnimMontage* Montage);
+
+	/** Get next sectionID for given section ID */
+	int32 Montage_GetNextSectionID(UAnimMontage const * const Montage, int32 const & CurrentSectionID) const;
+
+	/** Get a current Active Montage in this AnimInstance. 
+		Note that there might be multiple Active at the same time. This will only return the first active one it finds. **/
+	UAnimMontage * GetCurrentActiveMontage();
+
+	/** Get Currently active montage instance.
+		Note that there might be multiple Active at the same time. This will only return the first active one it finds. **/
+	FAnimMontageInstance * GetActiveMontageInstance();
+
+	/** Get Active FAnimMontageInstance for given Montage asset. Will return NULL if Montage is not currently Active. */
+	FAnimMontageInstance * GetActiveInstanceForMontage(UAnimMontage const & Montage) const;
+
+	/** AnimMontage instances that are running currently
+	* - only one is primarily active per group, and the other ones are blending out
+	*/
+	TArray<struct FAnimMontageInstance*> MontageInstances;
+
+	void OnMontageInstanceStopped(FAnimMontageInstance & StoppedMontageInstance);
+
+protected:
+	/** Map between Active Montages and their FAnimMontageInstance */
+	TMap<class UAnimMontage *, struct FAnimMontageInstance*> ActiveMontagesMap;
+
+	/** Stop all montages that are active **/
+	void StopAllMontages(float BlendOut);
 
+	/** Update weight of montages  **/
+	virtual void Montage_UpdateWeight(float DeltaSeconds);
+	/** Advance montages **/
+	virtual void Montage_Advance(float DeltaSeconds);
+
+public:
 	/** Returns the value of a named curve. */
 	UFUNCTION(BlueprintPure, Category="Animation", meta=(BlueprintProtected = "true"))
 	float GetCurveValue(FName CurveName);
@@ -490,11 +566,6 @@ public:
 		is removed correctly. */
 	TArray<FAnimNotifyEvent> ActiveAnimNotifyState;
 
-	/** AnimMontage instances that are running currently 
-	 * - only one is primarily active, but the other ones are blending out 
-	 */
-	TArray<struct FAnimMontageInstance*> MontageInstances;
-
 	/** Curve Values that are added to trigger in event**/
 	TMap<FName, float>	EventCurves;
 	/** Morph Target Curves that will be used for SkeletalMeshComponent **/
@@ -548,51 +619,7 @@ public:
 	/** Add curve float data using a curve Uid, the name of the curve will be resolved from the skeleton **/
 	void AddCurveValue(const USkeleton::AnimCurveUID Uid, float Value, int32 CurveTypeFlags);
 
-	/* native cpp montage interface functions **/
-	void Montage_SetEndDelegate(FOnMontageEnded & OnMontageEnded);
-	void Montage_SetBlendingOutDelegate(FOnMontageBlendingOutStarted & OnMontageBlendingOut);
-	FOnMontageBlendingOutStarted* Montage_GetBlendingOutDelegate();
-
-
-	/** Called when a montage starts blending out, whether interrupted or finished */
-	UPROPERTY(BlueprintAssignable)
-	FOnMontageBlendingOutStartedMCDelegate OnMontageBlendingOut;
-	
-	/** Called when a montage has ended, whether interrupted or finished*/
-	UPROPERTY(BlueprintAssignable)
-	FOnMontageEndedMCDelegate OnMontageEnded;
-
-public:
-	/** Get Current Montage Position */
-	float Montage_GetPosition(UAnimMontage* Montage);
-	/** Has Montage been stopped? */
-	bool Montage_GetIsStopped(UAnimMontage* Montage);
-
-	/** Set position. */
-	void Montage_SetPosition(UAnimMontage* Montage, float NewPosition);
-
-	/** Get PlayRate */
-	float Montage_GetPlayRate(UAnimMontage* Montage);
-
-	/** Get next sectionID for given section ID */
-	int32 Montage_GetNextSectionID(UAnimMontage* Montage, int32 CurrentSectionID);
-
-	/** Get Current Active Montage in this AnimInstance **/
-	UAnimMontage * GetCurrentActiveMontage();
-
-	/** Get Currently active montage instace **/
-	FAnimMontageInstance* GetActiveMontageInstance();
-
 protected:
-	/** Update weight of montages  **/
-	virtual void Montage_UpdateWeight(float DeltaSeconds);
-	/** Advance montages **/
-	virtual void Montage_Advance(float DeltaSeconds);
-	/** Stop all montages that are active **/
-	void StopAllMontages(float BlendOut);
-	/** Called by blueprint functions that modify the montages current position. */
-	void OnMontagePositionChanged(FAnimMontageInstance* MontageInstance, FName ToSectionName);
-
 	/** 
 	 * Add curve float data, using a curve name. External values should all be added using
 	 * The curve UID to the public version of this method
@@ -616,9 +643,6 @@ private:
 	/** Active Root Motion Montage Instance, if any. */
 	struct FAnimMontageInstance * RootMotionMontageInstance;
 
-	/** Update RootMotionMontageInstance */
-	void UpdateRootMotionMontageInstance();
-
 	// Root motion extracted from animation since the last time ConsumeExtractedRootMotion was called
 	FRootMotionMovementParams ExtractedRootMotion;
 };
diff --git a/Engine/Source/Runtime/Engine/Classes/Animation/AnimMontage.h b/Engine/Source/Runtime/Engine/Classes/Animation/AnimMontage.h
--- a/Engine/Source/Runtime/Engine/Classes/Animation/AnimMontage.h
+++ b/Engine/Source/Runtime/Engine/Classes/Animation/AnimMontage.h
@@ -143,12 +143,6 @@ public:
 	UPROPERTY()
 	TArray<int32> PrevSections;
 
-	UPROPERTY()
-	float Position;
-
-	UPROPERTY()
-	float PlayRate;
-
 	UPROPERTY()
 	bool bPlaying;
 
@@ -166,6 +160,13 @@ public:
 	// transient PreviousWeight - Weight of previous tick
 	float PreviousWeight;
 
+private:
+	UPROPERTY()
+	float Position;
+
+	UPROPERTY()
+	float PlayRate;
+
 public:
 	FAnimMontageInstance()
 		: Montage(NULL)
@@ -173,12 +174,12 @@ public:
 		, Weight(0.f)
 		, BlendTime(0.f)
 		, DefaultBlendTimeMultiplier(1.0f)
-		, Position(0.f)
-		, PlayRate(1.f)
-		, bPlaying(false)		
+		, bPlaying(false)	
 		, AnimInstance(NULL)
 		, bInterrupted(false)
 		, PreviousWeight(0.f)
+		, Position(0.f)
+		, PlayRate(1.f)
 	{
 	}
 
@@ -188,29 +189,46 @@ public:
 		, Weight(0.f)
 		, BlendTime(0.f)
 		, DefaultBlendTimeMultiplier(1.0f)
-		, Position(0.f)
-		, PlayRate(1.f)
 		, bPlaying(false)		
 		, AnimInstance(InAnimInstance)
 		, bInterrupted(false)
-		, PreviousWeight(0.f)		
+		, PreviousWeight(0.f)	
+		, Position(0.f)
+		, PlayRate(1.f)
 	{
 	}
 
-	bool ChangeNextSection(FName SectionName, FName NextSectionName);
-
 	// montage instance interfaces
 	void Play(float InPlayRate = 1.f);
 	void Stop(float BlendOut, bool bInterrupt=true);
 	void Pause();
 	void Initialize(class UAnimMontage * InMontage);
 
-	bool ChangePositionToSection(FName SectionName, bool bEndOfSection = false);
+	bool JumpToSectionName(FName const & SectionName, bool bEndOfSection = false);
+	bool SetNextSectionName(FName const & SectionName, FName const & NewNextSectionName);
+	bool SetNextSectionID(int32 const & SectionID, int32 const & NewNextSectionID);
 
 	bool IsValid() const { return (Montage!=NULL); }
 	bool IsPlaying() const { return IsValid() && bPlaying; }
+	bool IsStopped() const { return DesiredWeight == 0.f; }
+
+	/** Returns true if this montage is active (valid and not blending out) */
+	bool IsActive() const { return (IsValid() && !IsStopped()); }
 
 	void Terminate();
+
+	/**
+	 *  Getters
+	 */
+	float GetPosition() const { return Position; };
+	float GetPlayRate() const { return PlayRate; }
+
+	/** 
+	 * Setters
+	 */
+	void SetPosition(float const & InPosition) { Position = InPosition; }
+	void SetPlayRate(float const & InPlayRate) { PlayRate = InPlayRate; }
+
 	/**
 	 * Montage Tick happens in 2 phases
 	 *
@@ -225,14 +243,19 @@ public:
 	/** Simulate is same as Advance, but without calling any events or touching any of the instance data. So it performs a simulation of advancing the timeline. 
 	 * @fixme laurent can we make Advance use that, so we don't have 2 code paths which risk getting out of sync? */
 	bool SimulateAdvance(float DeltaTime, float & InOutPosition, struct FRootMotionMovementParams & OutRootMotionParams) const;
-	void Advance(float DeltaTime, struct FRootMotionMovementParams & OutRootMotionParams);
+	void Advance(float DeltaTime, struct FRootMotionMovementParams * OutRootMotionParams);
 
 	FName GetCurrentSection() const;
 	FName GetNextSection() const;
+	int32 GetNextSectionID(int32 const & CurrentSectionID) const;
+	FName GetSectionNameFromID(int32 const & SectionID) const;
 
 	// reference has to be managed manually
 	void AddReferencedObjects( FReferenceCollector& Collector );
 private:
+	/** Called by blueprint functions that modify the montages current position. */
+	void OnMontagePositionChanged(FName const & ToSectionName);
+
 	/** Delegate function handlers
 	 */
 	void HandleEvents(float PreviousTrackPos, float CurrentTrackPos, const FBranchingPoint * BranchingPoint);
diff --git a/Engine/Source/Runtime/Engine/Private/Animation/AnimInstance.cpp b/Engine/Source/Runtime/Engine/Private/Animation/AnimInstance.cpp
--- a/Engine/Source/Runtime/Engine/Private/Animation/AnimInstance.cpp
+++ b/Engine/Source/Runtime/Engine/Private/Animation/AnimInstance.cpp
@@ -665,15 +665,13 @@ void UAnimInstance::DisplayDebug(class UCanvas* Canvas, const FDebugDisplayInfo&
 		Canvas->DrawText(RenderFont, Heading, Indent, YPos, 1.f, 1.f, RenderInfo);
 		YPos += YL;
 
-		FAnimMontageInstance* ActiveMontageInstance = GetActiveMontageInstance();
-
 		for (int32 MontageIndex = 0; MontageIndex < MontageInstances.Num(); ++MontageIndex)
 		{
 			FIndenter PlayerIndent(Indent);
 
 			FAnimMontageInstance* MontageInstance = MontageInstances[MontageIndex];
 
-			Canvas->SetLinearDrawColor((MontageInstance == ActiveMontageInstance) ? ActiveColor : TextWhite);
+			Canvas->SetLinearDrawColor((MontageInstance->IsActive()) ? ActiveColor : TextWhite);
 
 			FString MontageEntry = FString::Printf(TEXT("%i) %s Sec: %s W:%.3f DW:%.3f"), MontageIndex, *MontageInstance->Montage->GetName(), *MontageInstance->GetCurrentSection().ToString(), *MontageInstance->GetNextSection().ToString(), MontageInstance->Weight, MontageInstance->DesiredWeight);
 			Canvas->DrawText(RenderFont, MontageEntry, Indent, YPos, 1.f, 1.f, RenderInfo);
@@ -1185,6 +1183,7 @@ void UAnimInstance::GetSlotWeight(FName const & SlotNodeName, float & out_SlotNo
 		if (MontageInstance && MontageInstance->IsValid() && MontageInstance->Montage->IsValidSlot(SlotNodeName))
 		{
 			NodeTotalWeight += MontageInstance->Weight;
+
 			if( !MontageInstance->Montage->IsValidAdditive() )
 			{
 				NonAdditiveTotalWeight += MontageInstance->Weight;
@@ -1264,7 +1263,7 @@ void UAnimInstance::SlotEvaluatePose(FName SlotNodeName, const FA2Pose & SourceP
 
 			// Extract pose from Track
 			UAnimMontage const * const MontageAsset = MontageInstance->Montage;
-			FAnimExtractContext ExtractionContext(MontageInstance->Position, MontageAsset->bEnableRootMotionTranslation, MontageAsset->bEnableRootMotionRotation, MontageAsset->RootMotionRootLock);
+			FAnimExtractContext ExtractionContext(MontageInstance->GetPosition(), MontageAsset->bEnableRootMotionTranslation, MontageAsset->bEnableRootMotionRotation, MontageAsset->RootMotionRootLock);
 			FAnimationRuntime::GetPoseFromAnimTrack(*AnimTrack, RequiredBones, NewPose.Pose.Bones, ExtractionContext);
 
 			TotalWeight += MontageInstance->Weight;
@@ -1502,35 +1501,6 @@ float UAnimInstance::GetCurrentStateElapsedTime(int32 MachineIndex)
 	return 0.0f;
 }
 
-void UAnimInstance::Montage_SetEndDelegate(FOnMontageEnded & OnMontageEnded)
-{
-	FAnimMontageInstance * CurMontageInstance = GetActiveMontageInstance();
-	if ( CurMontageInstance )
-	{
-		CurMontageInstance->OnMontageEnded = OnMontageEnded;
-	}
-}
-
-void UAnimInstance::Montage_SetBlendingOutDelegate(FOnMontageBlendingOutStarted & OnMontageBlendingOut)
-{
-	FAnimMontageInstance * CurMontageInstance = GetActiveMontageInstance();
-	if ( CurMontageInstance )
-	{
-		CurMontageInstance->OnMontageBlendingOutStarted = OnMontageBlendingOut;
-	}
-}
-
-FOnMontageBlendingOutStarted* UAnimInstance::Montage_GetBlendingOutDelegate()
-{
-	FAnimMontageInstance * CurMontageInstance = GetActiveMontageInstance();
-	if ( CurMontageInstance )
-	{
-		return &CurMontageInstance->OnMontageBlendingOutStarted;
-	}
-
-	return NULL;
-}
-
 void UAnimInstance::Montage_UpdateWeight(float DeltaSeconds)
 {
 	// go through all montage instances, and update them
@@ -1546,25 +1516,27 @@ void UAnimInstance::Montage_UpdateWeight(float DeltaSeconds)
 
 void UAnimInstance::Montage_Advance(float DeltaSeconds)
 {
-	bool bUpdateRootMotionMontageInstance = false;
 	FRootMotionMovementParams LocalExtractedRootMotion;
 
 	// go through all montage instances, and update them
 	// and make sure their weight is updated properly
-	for (int32 I=0; I<MontageInstances.Num(); ++I)
+	for (int32 InstanceIndex = 0; InstanceIndex<MontageInstances.Num(); InstanceIndex++)
 	{
 		// should never be NULL
-		ensure(MontageInstances[I]);
-		if( MontageInstances[I] )
+		FAnimMontageInstance * MontageInstance = MontageInstances[InstanceIndex];
+		ensure(MontageInstance);
+		if (MontageInstance)
 		{
-			MontageInstances[I]->Advance(DeltaSeconds, LocalExtractedRootMotion);
+			// Only use root motion from active root motion montage.
+			// we don't support blending/weighting because of networking constraints.
+			bool const bExtractRootMotion = (MontageInstance == GetRootMotionMontageInstance());
+			MontageInstance->Advance(DeltaSeconds, bExtractRootMotion ? &LocalExtractedRootMotion : NULL);
 
-			if( !MontageInstances[I]->IsValid() )
+			if (!MontageInstance->IsValid())
 			{
-				delete MontageInstances[I];
-				MontageInstances.RemoveAt(I);
-				bUpdateRootMotionMontageInstance = true;
-				--I;
+				delete MontageInstance;
+				MontageInstances.RemoveAt(InstanceIndex);
+				--InstanceIndex;
 			}
 #if DO_CHECK && WITH_EDITORONLY_DATA && 0
 			else
@@ -1579,8 +1551,6 @@ void UAnimInstance::Montage_Advance(float DeltaSeconds)
 		}
 	}
 
-	UpdateRootMotionMontageInstance();
-
 	// If Root Motion has been extracted, store it.
 	if ((GetRootMotionMontageInstance() != NULL) && LocalExtractedRootMotion.bHasRootMotion)
 	{
@@ -1651,247 +1621,540 @@ void UAnimInstance::StopSlotAnimation(float InBlendOutTime)
 
 bool UAnimInstance::IsPlayingSlotAnimation(UAnimSequenceBase* Asset, FName SlotNodeName )
 {
-	// check if this is playing
-	FAnimMontageInstance * CurrentInstance = GetActiveMontageInstance();
-	// make sure what is active right now is transient that we created by request
-	if ( CurrentInstance && CurrentInstance->Montage && CurrentInstance->Montage->GetOuter() == GetTransientPackage() )
+	for (int32 InstanceIndex = 0; InstanceIndex < MontageInstances.Num(); InstanceIndex++)
 	{
-		UAnimMontage * CurMontage = CurrentInstance->Montage;
-
-		const FAnimTrack * AnimTrack = CurMontage->GetAnimationData(SlotNodeName);
-		if (AnimTrack && AnimTrack->AnimSegments.Num() == 1)
+		// check if this is playing
+		FAnimMontageInstance * MontageInstance = MontageInstances[InstanceIndex];
+		// make sure what is active right now is transient that we created by request
+		if (MontageInstance && MontageInstance->IsActive() && MontageInstance->IsPlaying())
 		{
-			// find if the 
-			return (AnimTrack->AnimSegments[0].AnimReference == Asset);
+			UAnimMontage * CurMontage = MontageInstance->Montage;
+			if (CurMontage && CurMontage->GetOuter() == GetTransientPackage())
+			{
+				const FAnimTrack * AnimTrack = CurMontage->GetAnimationData(SlotNodeName);
+				if (AnimTrack && AnimTrack->AnimSegments.Num() == 1)
+				{
+					// find if the 
+					return (AnimTrack->AnimSegments[0].AnimReference == Asset);
+				}
+			}
 		}
 	}
 
 	return false;
 }
 
-UAnimMontage * UAnimInstance::GetCurrentActiveMontage()
+/** Play a Montage. Returns Length of Montage in seconds. Returns 0.f if failed to play. */
+float UAnimInstance::Montage_Play(UAnimMontage * MontageToPlay, float InPlayRate)
 {
-	FAnimMontageInstance* CurrentActiveMontage = GetActiveMontageInstance();
-	if (CurrentActiveMontage)
+	if (MontageToPlay && (MontageToPlay->SequenceLength > 0.f))
 	{
-		return CurrentActiveMontage->Montage;
-	}
+		if (CurrentSkeleton->IsCompatible(MontageToPlay->GetSkeleton()))
+		{
+			//@fixme laurent -- only stop montages that belong to same group. 
+			// or to ensure single root motion montage playing.
+			// when stopping old animations, make sure it does give current new blendintime to blend out
+			StopAllMontages(MontageToPlay->BlendInTime);
 
-	return NULL;
-}
+			FAnimMontageInstance * NewInstance = new FAnimMontageInstance(this);
+			check(NewInstance);
 
-FAnimMontageInstance* UAnimInstance::GetActiveMontageInstance()
-{
-	if ( MontageInstances.Num() > 0 )
-	{
-		FAnimMontageInstance * RetVal = MontageInstances.Last();
-		if ( RetVal->IsValid() )
+			NewInstance->Initialize(MontageToPlay);
+			NewInstance->Play(InPlayRate);
+			MontageInstances.Add(NewInstance);
+			ActiveMontagesMap.Add(MontageToPlay, NewInstance);
+
+			// If we are playing root motion, set this instance as the one providing root motion.
+			if (MontageToPlay->bEnableRootMotionTranslation || MontageToPlay->bEnableRootMotionRotation)
+			{
+				RootMotionMontageInstance = NewInstance;
+			}
+
+			return NewInstance->Montage->SequenceLength;
+		}
+		else
 		{
-			return RetVal;
+			UE_LOG(LogAnimation, Warning, TEXT("Playing a Montage (%s) for the wrong Skeleton (%s) instead of (%s)."),
+				*GetNameSafe(MontageToPlay), *GetNameSafe(MontageToPlay->GetSkeleton()), *GetNameSafe(CurrentSkeleton));
 		}
 	}
 
-	return NULL;
+	return 0.f;
 }
 
-void UAnimInstance::OnMontagePositionChanged(FAnimMontageInstance* MontageInstance, FName ToSectionName)
+void UAnimInstance::Montage_Stop(float InBlendOutTime, UAnimMontage * Montage)
 {
-	if(MontageInstance && MontageInstance->bPlaying)
+	if (Montage)
 	{
-		if(MontageInstance->DesiredWeight == 0.f)
+		FAnimMontageInstance * MontageInstance = GetActiveInstanceForMontage(*Montage);
+		if (MontageInstance)
 		{
-			UE_LOG( LogAnimation, Warning, TEXT("Changing section on Montage (%s) to '%s' during blend out. This can cause incorrect visuals!"), 
-					*MontageInstance->Montage->GetName(), *ToSectionName.ToString() );
-			MontageInstance->Play(MontageInstance->PlayRate);
+			MontageInstance->Stop(InBlendOutTime);
+		}
+	}
+	else
+	{
+		// If no Montage reference, do it on all active ones.
+		for (int32 InstanceIndex = 0; InstanceIndex < MontageInstances.Num(); InstanceIndex++)
+		{
+			FAnimMontageInstance * MontageInstance = MontageInstances[InstanceIndex];
+			if (MontageInstance && MontageInstance->IsActive())
+			{
+				MontageInstance->Stop(InBlendOutTime);
+			}
 		}
 	}
 }
 
-void UAnimInstance::Montage_JumpToSection(FName SectionName)
+void UAnimInstance::Montage_JumpToSection(FName SectionName, UAnimMontage * Montage)
 {
-	FAnimMontageInstance * CurMontageInstance = GetActiveMontageInstance();
-	if ( CurMontageInstance && CurMontageInstance->ChangePositionToSection(SectionName, CurMontageInstance->PlayRate < 0.0f) == false )
+	if (Montage)
 	{
-		UE_LOG(LogAnimation, Warning, TEXT("Jumping section to %s failed for Montage (%s) "), 
-			*SectionName.ToString(), *CurMontageInstance->Montage->GetName() );
+		FAnimMontageInstance * MontageInstance = GetActiveInstanceForMontage(*Montage);
+		if (MontageInstance)
+		{
+			bool const bEndOfSection = (MontageInstance->GetPlayRate() < 0.f);
+			MontageInstance->JumpToSectionName(SectionName, bEndOfSection);
+		}
 	}
 	else
 	{
-		OnMontagePositionChanged(CurMontageInstance, SectionName);
+		// If no Montage reference, do it on all active ones.
+		for (int32 InstanceIndex = 0; InstanceIndex < MontageInstances.Num(); InstanceIndex++)
+		{
+			FAnimMontageInstance * MontageInstance = MontageInstances[InstanceIndex];
+			if (MontageInstance && MontageInstance->IsActive())
+			{
+				bool const bEndOfSection = (MontageInstance->GetPlayRate() < 0.f);
+				MontageInstance->JumpToSectionName(SectionName, bEndOfSection);
+			}
+		}
 	}
 }
 
-void UAnimInstance::Montage_JumpToSectionsEnd(FName SectionName)
+void UAnimInstance::Montage_JumpToSectionsEnd(FName SectionName, UAnimMontage * Montage)
 {
-	FAnimMontageInstance * CurMontageInstance = GetActiveMontageInstance();
-	if ( CurMontageInstance && CurMontageInstance->ChangePositionToSection(SectionName, CurMontageInstance->PlayRate >= 0.0f) == false )
+	if (Montage)
 	{
-		UE_LOG(LogAnimation, Warning, TEXT("Jumping section to %s failed for Montage (%s) "), 
-			*SectionName.ToString(), *CurMontageInstance->Montage->GetName() );
+		FAnimMontageInstance * MontageInstance = GetActiveInstanceForMontage(*Montage);
+		if (MontageInstance)
+		{
+			bool const bEndOfSection = (MontageInstance->GetPlayRate() >= 0.f);
+			MontageInstance->JumpToSectionName(SectionName, bEndOfSection);
+		}
 	}
 	else
 	{
-		OnMontagePositionChanged(CurMontageInstance, SectionName);
+		// If no Montage reference, do it on all active ones.
+		for (int32 InstanceIndex = 0; InstanceIndex < MontageInstances.Num(); InstanceIndex++)
+		{
+			FAnimMontageInstance * MontageInstance = MontageInstances[InstanceIndex];
+			if (MontageInstance && MontageInstance->IsActive())
+			{
+				bool const bEndOfSection = (MontageInstance->GetPlayRate() >= 0.f);
+				MontageInstance->JumpToSectionName(SectionName, bEndOfSection);
+			}
+		}
 	}
 }
 
-FName UAnimInstance::Montage_GetCurrentSection()
+void UAnimInstance::Montage_SetNextSection(FName SectionNameToChange, FName NextSection, UAnimMontage * Montage)
 {
-	FAnimMontageInstance * CurMontageInstance = GetActiveMontageInstance();
-	if ( CurMontageInstance )
+	if (Montage)
 	{
-		return CurMontageInstance->GetCurrentSection();
+		FAnimMontageInstance * MontageInstance = GetActiveInstanceForMontage(*Montage);
+		if (MontageInstance)
+		{
+			MontageInstance->SetNextSectionName(SectionNameToChange, NextSection);
+		}
 	}
+	else
+	{
+		bool bFoundOne = false;
 
-	return NAME_None;
+		// If no Montage reference, do it on all active ones.
+		for (int32 InstanceIndex = 0; InstanceIndex < MontageInstances.Num(); InstanceIndex++)
+		{
+			FAnimMontageInstance * MontageInstance = MontageInstances[InstanceIndex];
+			if (MontageInstance && MontageInstance->IsActive())
+			{
+				MontageInstance->SetNextSectionName(SectionNameToChange, NextSection);
+				bFoundOne = true;
+			}
+		}
+
+		if (!bFoundOne)
+		{
+			bFoundOne = true;
+		}
+	}
 }
 
-void UAnimInstance::Montage_SetNextSection(FName SectionNameToChange, FName NextSection)
+void UAnimInstance::Montage_SetPlayRate(UAnimMontage * Montage, float NewPlayRate)
 {
-	FAnimMontageInstance * CurMontageInstance = GetActiveMontageInstance();
-	if ( CurMontageInstance && CurMontageInstance->ChangeNextSection(SectionNameToChange, NextSection) == false )
+	if (Montage)
 	{
-		UE_LOG(LogAnimation, Warning, TEXT("Changing section from %s to %s failed for Montage (%s) "), 
-			*SectionNameToChange.ToString(), *NextSection.ToString(), *CurMontageInstance->Montage->GetName() );
+		FAnimMontageInstance * MontageInstance = GetActiveInstanceForMontage(*Montage);
+		if (MontageInstance)
+		{
+			MontageInstance->SetPlayRate(NewPlayRate);
+		}
 	}
 	else
 	{
-		OnMontagePositionChanged(CurMontageInstance, NextSection);
+		// If no Montage reference, do it on all active ones.
+		for (int32 InstanceIndex = 0; InstanceIndex < MontageInstances.Num(); InstanceIndex++)
+		{
+			FAnimMontageInstance * MontageInstance = MontageInstances[InstanceIndex];
+			if (MontageInstance && MontageInstance->IsActive())
+			{
+				MontageInstance->SetPlayRate(NewPlayRate);
+			}
+		}
 	}
 }
 
-/** Play a Montage. Returns Length of Montage in seconds. Returns 0.f if failed to play. */
-float UAnimInstance::Montage_Play(UAnimMontage * MontageToPlay, float InPlayRate)
+bool UAnimInstance::Montage_IsActive(UAnimMontage * Montage)
 {
-	if( MontageToPlay && (MontageToPlay->SequenceLength > 0.f) ) 
+	if (Montage)
 	{
-		if( CurrentSkeleton->IsCompatible(MontageToPlay->GetSkeleton()) )
+		FAnimMontageInstance * MontageInstance = GetActiveInstanceForMontage(*Montage);
+		if (MontageInstance)
 		{
-			// when stopping old animations, make sure it does give current new blendintime to blend out
-			StopAllMontages(MontageToPlay->BlendInTime);
-
-			FAnimMontageInstance * NewInstance = new FAnimMontageInstance(this);
-			check (NewInstance);
-
-			NewInstance->Initialize(MontageToPlay);
-			NewInstance->Play(InPlayRate);
-			MontageInstances.Add(NewInstance);
-
-			UpdateRootMotionMontageInstance();
-
-			return NewInstance->Montage->SequenceLength;
+			return true;
 		}
-		else
+	}
+	else
+	{
+		// If no Montage reference, return true if there is any active montage.
+		for (int32 InstanceIndex = 0; InstanceIndex < MontageInstances.Num(); InstanceIndex++)
 		{
-			UE_LOG(LogAnimation, Warning, TEXT("Playing a Montage (%s) for the wrong Skeleton (%s) instead of (%s)."), 
-				*GetNameSafe(MontageToPlay), *GetNameSafe(MontageToPlay->GetSkeleton()), *GetNameSafe(CurrentSkeleton));
+			FAnimMontageInstance * MontageInstance = MontageInstances[InstanceIndex];
+			if (MontageInstance && MontageInstance->IsActive())
+			{
+				return true;
+			}
 		}
 	}
 
-	return 0.f;
+	return false;
 }
 
-void UAnimInstance::UpdateRootMotionMontageInstance()
+bool UAnimInstance::Montage_IsPlaying(UAnimMontage * Montage)
 {
-	FAnimMontageInstance * ActiveMontageInstance = GetActiveMontageInstance();
-
-	const bool bValidRootMotionInstance = (ActiveMontageInstance && ActiveMontageInstance->IsValid() && (ActiveMontageInstance->DesiredWeight > 0.f) && ActiveMontageInstance->Montage 
-		&& (ActiveMontageInstance->Montage->bEnableRootMotionTranslation || ActiveMontageInstance->Montage->bEnableRootMotionRotation) );
+	if (Montage)
+	{
+		FAnimMontageInstance * MontageInstance = GetActiveInstanceForMontage(*Montage);
+		if (MontageInstance)
+		{
+			return MontageInstance->IsPlaying();
+		}
+	}
+	else
+	{
+		// If no Montage reference, return true if there is any active playing montage.
+		for (int32 InstanceIndex = 0; InstanceIndex < MontageInstances.Num(); InstanceIndex++)
+		{
+			FAnimMontageInstance * MontageInstance = MontageInstances[InstanceIndex];
+			if (MontageInstance && MontageInstance->IsActive() && MontageInstance->IsPlaying())
+			{
+				return true;
+			}
+		}
+	}
 
-	RootMotionMontageInstance = bValidRootMotionInstance ? ActiveMontageInstance : NULL;
+	return false;
 }
 
-FAnimMontageInstance * UAnimInstance::GetRootMotionMontageInstance() const
+FName UAnimInstance::Montage_GetCurrentSection(UAnimMontage* Montage)
 {
-	return RootMotionMontageInstance;
-}
+	if (Montage)
+	{
+		FAnimMontageInstance * MontageInstance = GetActiveInstanceForMontage(*Montage);
+		if (MontageInstance)
+		{
+			return MontageInstance->GetCurrentSection();
+		}
+	}
+	else
+	{
+		// If no Montage reference, get first active one.
+		for (int32 InstanceIndex = 0; InstanceIndex < MontageInstances.Num(); InstanceIndex++)
+		{
+			FAnimMontageInstance * MontageInstance = MontageInstances[InstanceIndex];
+			if (MontageInstance && MontageInstance->IsActive())
+			{
+				return MontageInstance->GetCurrentSection();
+			}
+		}
+	}
 
-FRootMotionMovementParams UAnimInstance::ConsumeExtractedRootMotion()
-{
-	FRootMotionMovementParams RootMotion = ExtractedRootMotion;
-	ExtractedRootMotion.Clear();
-	return RootMotion;
+	return NAME_None;
 }
 
-void UAnimInstance::Montage_Stop(float InBlendOutTime)
+void UAnimInstance::Montage_SetEndDelegate(FOnMontageEnded & OnMontageEnded, UAnimMontage* Montage)
 {
-	FAnimMontageInstance * CurMontageInstance = GetActiveMontageInstance();
-	if ( CurMontageInstance )
+	if (Montage)
 	{
-		CurMontageInstance->Stop(InBlendOutTime);
+		FAnimMontageInstance * MontageInstance = GetActiveInstanceForMontage(*Montage);
+		if (MontageInstance)
+		{
+			MontageInstance->OnMontageEnded = OnMontageEnded;
+		}
+	}
+	else
+	{
+		// If no Montage reference, do it on all active ones.
+		for (int32 InstanceIndex = 0; InstanceIndex < MontageInstances.Num(); InstanceIndex++)
+		{
+			FAnimMontageInstance * MontageInstance = MontageInstances[InstanceIndex];
+			if (MontageInstance && MontageInstance->IsActive())
+			{
+				MontageInstance->OnMontageEnded = OnMontageEnded;
+			}
+		}
 	}
 }
 
-/** Has Montage been stopped? */
-bool UAnimInstance::Montage_GetIsStopped(UAnimMontage* Montage)
+void UAnimInstance::Montage_SetBlendingOutDelegate(FOnMontageBlendingOutStarted & OnMontageBlendingOut, UAnimMontage* Montage)
 {
-	FAnimMontageInstance* CurMontageInstance = GetActiveMontageInstance();
-	return (!CurMontageInstance || (CurMontageInstance->Montage != Montage) || (CurMontageInstance->DesiredWeight == 0.f));
+	if (Montage)
+	{
+		FAnimMontageInstance * MontageInstance = GetActiveInstanceForMontage(*Montage);
+		if (MontageInstance)
+		{
+			MontageInstance->OnMontageBlendingOutStarted = OnMontageBlendingOut;
+		}
+	}
+	else
+	{
+		// If no Montage reference, do it on all active ones.
+		for (int32 InstanceIndex = 0; InstanceIndex < MontageInstances.Num(); InstanceIndex++)
+		{
+			FAnimMontageInstance * MontageInstance = MontageInstances[InstanceIndex];
+			if (MontageInstance && MontageInstance->IsActive())
+			{
+				MontageInstance->OnMontageBlendingOutStarted = OnMontageBlendingOut;
+			}
+		}
+	}
 }
 
-bool UAnimInstance::Montage_IsActive(UAnimMontage * Montage)
+FOnMontageBlendingOutStarted * UAnimInstance::Montage_GetBlendingOutDelegate(UAnimMontage * Montage)
 {
-	FAnimMontageInstance * CurMontageInstance = GetActiveMontageInstance();
-	return ( CurMontageInstance  && CurMontageInstance->Montage == Montage);
-}
+	if (Montage)
+	{
+		FAnimMontageInstance * MontageInstance = GetActiveInstanceForMontage(*Montage);
+		if (MontageInstance)
+		{
+			return &MontageInstance->OnMontageBlendingOutStarted;
+		}
+	}
+	else
+	{
+		// If no Montage reference, use first active one found.
+		for (int32 InstanceIndex = 0; InstanceIndex < MontageInstances.Num(); InstanceIndex++)
+		{
+			FAnimMontageInstance * MontageInstance = MontageInstances[InstanceIndex];
+			if (MontageInstance && MontageInstance->IsActive())
+			{
+				return &MontageInstance->OnMontageBlendingOutStarted;
+			}
+		}
+	}
 
+	return NULL;
+}
 
-bool UAnimInstance::Montage_IsPlaying(UAnimMontage * Montage)
+void UAnimInstance::Montage_SetPosition(UAnimMontage* Montage, float NewPosition)
 {
-	FAnimMontageInstance * CurMontageInstance = GetActiveMontageInstance();
-	return ( CurMontageInstance  && CurMontageInstance->Montage == Montage && CurMontageInstance->IsPlaying() );
+	if (Montage)
+	{
+		FAnimMontageInstance * MontageInstance = GetActiveInstanceForMontage(*Montage);
+		if (MontageInstance)
+		{
+			MontageInstance->SetPosition(NewPosition);
+		}
+	}
+	else
+	{
+		// If no Montage reference, do it on all active ones.
+		for (int32 InstanceIndex = 0; InstanceIndex < MontageInstances.Num(); InstanceIndex++)
+		{
+			FAnimMontageInstance * MontageInstance = MontageInstances[InstanceIndex];
+			if (MontageInstance && MontageInstance->IsActive())
+			{
+				MontageInstance->SetPosition(NewPosition);
+			}
+		}
+	}
 }
 
 float UAnimInstance::Montage_GetPosition(UAnimMontage* Montage)
 {
-	FAnimMontageInstance * CurMontageInstance = GetActiveMontageInstance();
-	return CurMontageInstance ? CurMontageInstance->Position : 0.f;
+	if (Montage)
+	{
+		FAnimMontageInstance * MontageInstance = GetActiveInstanceForMontage(*Montage);
+		if (MontageInstance)
+		{
+			return MontageInstance->GetPosition();
+		}
+	}
+	else
+	{
+		// If no Montage reference, use first active one found.
+		for (int32 InstanceIndex = 0; InstanceIndex < MontageInstances.Num(); InstanceIndex++)
+		{
+			FAnimMontageInstance * MontageInstance = MontageInstances[InstanceIndex];
+			if (MontageInstance && MontageInstance->IsActive())
+			{
+				return MontageInstance->GetPosition();
+			}
+		}
+	}
+
+	return 0.f;
 }
 
-void UAnimInstance::Montage_SetPosition(UAnimMontage* Montage, float NewPosition)
+bool UAnimInstance::Montage_GetIsStopped(UAnimMontage* Montage)
 {
-	// @laurent we probably want (an option?) to advance time rather than jump? As that skips notifies/events?
-	FAnimMontageInstance* CurMontageInstance = GetActiveMontageInstance();
-	if( CurMontageInstance )
+	if (Montage)
 	{
-		CurMontageInstance->Position = NewPosition;
+		FAnimMontageInstance * MontageInstance = GetActiveInstanceForMontage(*Montage);
+		return (!MontageInstance); // Not active == Stopped.
 	}
+	return true;
 }
 
 float UAnimInstance::Montage_GetPlayRate(UAnimMontage* Montage)
 {
-	FAnimMontageInstance * CurMontageInstance = GetActiveMontageInstance();
-	return CurMontageInstance ? CurMontageInstance->PlayRate : 0.f;
+	if (Montage)
+	{
+		FAnimMontageInstance * MontageInstance = GetActiveInstanceForMontage(*Montage);
+		if (MontageInstance)
+		{
+			return MontageInstance->GetPlayRate();
+		}
+	}
+	else
+	{
+		// If no Montage reference, use first active one found.
+		for (int32 InstanceIndex = 0; InstanceIndex < MontageInstances.Num(); InstanceIndex++)
+		{
+			FAnimMontageInstance * MontageInstance = MontageInstances[InstanceIndex];
+			if (MontageInstance && MontageInstance->IsActive())
+			{
+				return MontageInstance->GetPlayRate();
+			}
+		}
+	}
+
+	return 0.f;
 }
 
-void UAnimInstance::Montage_SetPlayRate(UAnimMontage* Montage, float NewPlayRate)
+int32 UAnimInstance::Montage_GetNextSectionID(UAnimMontage const * const Montage, int32 const & CurrentSectionID) const
 {
-	FAnimMontageInstance* CurMontageInstance = GetActiveMontageInstance();
-	if( CurMontageInstance && (CurMontageInstance->Montage == Montage) && CurMontageInstance->IsPlaying() )
+	if (Montage)
 	{
-		CurMontageInstance->PlayRate = NewPlayRate;
+		FAnimMontageInstance * MontageInstance = GetActiveInstanceForMontage(*Montage);
+		if (MontageInstance)
+		{
+			return (CurrentSectionID < MontageInstance->NextSections.Num()) ? MontageInstance->NextSections[CurrentSectionID] : INDEX_NONE;
+		}
 	}
+	else
+	{
+		// If no Montage reference, use first active one found.
+		for (int32 InstanceIndex = 0; InstanceIndex < MontageInstances.Num(); InstanceIndex++)
+		{
+			FAnimMontageInstance * MontageInstance = MontageInstances[InstanceIndex];
+			if (MontageInstance && MontageInstance->IsActive())
+			{
+				return (CurrentSectionID < MontageInstance->NextSections.Num()) ? MontageInstance->NextSections[CurrentSectionID] : INDEX_NONE;
+			}
+		}
+	}
+
+	return INDEX_NONE;
 }
 
-int32 UAnimInstance::Montage_GetNextSectionID(UAnimMontage* Montage, int32 CurrentSectionID)
+UAnimMontage * UAnimInstance::GetCurrentActiveMontage()
 {
-	FAnimMontageInstance * CurMontageInstance = GetActiveMontageInstance();
-	if( CurMontageInstance && CurrentSectionID < CurMontageInstance->NextSections.Num() )
+	// Start from end, as most recent instances are added at the end of the queue.
+	int32 const NumInstances = MontageInstances.Num();
+	for (int32 InstanceIndex = NumInstances - 1; InstanceIndex >= 0; InstanceIndex--)
 	{
-		return CurMontageInstance->NextSections[CurrentSectionID];
+		FAnimMontageInstance * MontageInstance = MontageInstances[InstanceIndex];
+		if (MontageInstance && MontageInstance->IsActive())
+		{
+			MontageInstance->Montage;
+		}
 	}
 
-	return INDEX_NONE;
+	return NULL;
+}
+
+FAnimMontageInstance* UAnimInstance::GetActiveMontageInstance()
+{
+	// Start from end, as most recent instances are added at the end of the queue.
+	int32 const NumInstances = MontageInstances.Num();
+	for (int32 InstanceIndex = NumInstances - 1; InstanceIndex >= 0; InstanceIndex--)
+	{
+		FAnimMontageInstance * MontageInstance = MontageInstances[InstanceIndex];
+		if (MontageInstance && MontageInstance->IsActive())
+		{
+			return MontageInstance;
+		}
+	}
+
+	return NULL;
 }
 
 void UAnimInstance::StopAllMontages(float BlendOut)
 {
-	for ( int32 Index=MontageInstances.Num()-1; Index>=0; Index-- )
+	for (int32 Index = MontageInstances.Num() - 1; Index >= 0; Index--)
 	{
 		MontageInstances[Index]->Stop(BlendOut, true);
 	}
 }
 
+void UAnimInstance::OnMontageInstanceStopped(FAnimMontageInstance & StoppedMontageInstance)
+{
+	UAnimMontage * MontageStopped = StoppedMontageInstance.Montage;
+	ensure(MontageStopped != NULL);
+
+	// Remove instance for Active List.
+	FAnimMontageInstance** AnimInstancePtr = ActiveMontagesMap.Find(MontageStopped);
+	if (AnimInstancePtr && (*AnimInstancePtr == &StoppedMontageInstance))
+	{
+		ActiveMontagesMap.Remove(MontageStopped);
+	}
+
+	// Clear RootMotionMontageInstance
+	if (RootMotionMontageInstance == &StoppedMontageInstance)
+	{
+		RootMotionMontageInstance = NULL;
+	}
+}
+
+FAnimMontageInstance * UAnimInstance::GetActiveInstanceForMontage(UAnimMontage const & Montage) const
+{
+	FAnimMontageInstance * const * FoundInstancePtr = ActiveMontagesMap.Find(&Montage);
+	return FoundInstancePtr ? *FoundInstancePtr : NULL;
+}
+
+FAnimMontageInstance * UAnimInstance::GetRootMotionMontageInstance() const
+{
+	return RootMotionMontageInstance;
+}
+
+FRootMotionMovementParams UAnimInstance::ConsumeExtractedRootMotion()
+{
+	FRootMotionMovementParams RootMotion = ExtractedRootMotion;
+	ExtractedRootMotion.Clear();
+	return RootMotion;
+}
+
 void UAnimInstance::SetMorphTarget(FName MorphTargetName, float Value)
 {
 	USkeletalMeshComponent * Component = GetOwningComponent();
diff --git a/Engine/Source/Runtime/Engine/Private/Animation/AnimMontage.cpp b/Engine/Source/Runtime/Engine/Private/Animation/AnimMontage.cpp
--- a/Engine/Source/Runtime/Engine/Private/Animation/AnimMontage.cpp
+++ b/Engine/Source/Runtime/Engine/Private/Animation/AnimMontage.cpp
@@ -698,6 +698,7 @@ void FAnimMontageInstance::Stop(float BlendOut, bool bInterrupt)
 
 	if( Montage )
 	{
+		
 		// do not use default Montage->BlendOut  Time
 		// depending on situation, the BlendOut time changes
 		// check where this function gets called and see how we calculate BlendTime
@@ -705,6 +706,8 @@ void FAnimMontageInstance::Stop(float BlendOut, bool bInterrupt)
 
 		if (UAnimInstance* Inst = AnimInstance.Get())
 		{
+			// Let AnimInstance know we are being stopped.
+			Inst->OnMontageInstanceStopped(*this);
 			Inst->OnMontageBlendingOut.Broadcast(Montage, bInterrupt);
 		}
 		OnMontageBlendingOutStarted.ExecuteIfBound(Montage, bInterrupted);
@@ -775,6 +778,7 @@ void FAnimMontageInstance::Terminate()
 	{
 		return;
 	}
+
 	UAnimMontage* OldMontage = Montage;
 	Montage = NULL;
 
@@ -786,35 +790,73 @@ void FAnimMontageInstance::Terminate()
 	OnMontageEnded.ExecuteIfBound(OldMontage, bInterrupted);
 }
 
-bool FAnimMontageInstance::ChangeNextSection(FName SectionName, FName NextSectionName)
+bool FAnimMontageInstance::JumpToSectionName(FName const & SectionName, bool bEndOfSection)
 {
-	int32 SectionID = Montage->GetSectionIndex(SectionName);
-	int32 NextSectionID = Montage->GetSectionIndex(NextSectionName);
-	bool bHasValidNextSection = NextSections.IsValidIndex(SectionID);
-	
+	const int32 SectionID = Montage->GetSectionIndex(SectionName);
+
+	if (Montage->IsValidSectionIndex(SectionID))
+	{
+		FCompositeSection & CurSection = Montage->GetAnimCompositeSection(SectionID);
+		Position = Montage->CalculatePos(CurSection, bEndOfSection ? Montage->GetSectionLength(SectionID) - KINDA_SMALL_NUMBER : 0.0f);
+		OnMontagePositionChanged(SectionName);
+		return true;
+	}
+
+	UE_LOG(LogAnimation, Warning, TEXT("JumpToSectionName %s bEndOfSection: %d failed for Montage %s"),
+		*SectionName.ToString(), bEndOfSection, *GetNameSafe(Montage));
+	return false;
+}
+
+bool FAnimMontageInstance::SetNextSectionName(FName const & SectionName, FName const & NewNextSectionName)
+{
+	int32 const SectionID = Montage->GetSectionIndex(SectionName);
+	int32 const NewNextSectionID = Montage->GetSectionIndex(NewNextSectionName);
+
+	return SetNextSectionID(SectionID, NewNextSectionID);
+}
+
+bool FAnimMontageInstance::SetNextSectionID(int32 const & SectionID, int32 const & NewNextSectionID)
+{
+	bool const bHasValidNextSection = NextSections.IsValidIndex(SectionID);
+
 	// disconnect prev section
-	if ( bHasValidNextSection && 
-		NextSections[SectionID] != INDEX_NONE &&
-		 PrevSections.IsValidIndex(NextSections[SectionID]) )
+	if (bHasValidNextSection && (NextSections[SectionID] != INDEX_NONE) && PrevSections.IsValidIndex(NextSections[SectionID]))
 	{
 		PrevSections[NextSections[SectionID]] = INDEX_NONE;
 	}
+
 	// update in-reverse next section
-	if ( PrevSections.IsValidIndex(NextSectionID) )
+	if (PrevSections.IsValidIndex(NewNextSectionID))
 	{
-		PrevSections[NextSectionID] = SectionID;
+		PrevSections[NewNextSectionID] = SectionID;
 	}
+
 	// update next section for the SectionID
 	// NextSection can be invalid
-	if ( bHasValidNextSection )
+	if (bHasValidNextSection)
 	{
-		NextSections[SectionID] = NextSectionID;
+		NextSections[SectionID] = NewNextSectionID;
+		OnMontagePositionChanged(GetSectionNameFromID(NewNextSectionID));
 		return true;
 	}
 
+	UE_LOG(LogAnimation, Warning, TEXT("SetNextSectionName %s to %s failed for Montage %s"),
+		*GetSectionNameFromID(SectionID).ToString(), *GetSectionNameFromID(NewNextSectionID).ToString(), *GetNameSafe(Montage));
+
 	return false;
 }
 
+void FAnimMontageInstance::OnMontagePositionChanged(FName const & ToSectionName) 
+{
+	if (bPlaying && (DesiredWeight == 0.f))
+	{
+		UE_LOG(LogAnimation, Warning, TEXT("Changing section on Montage (%s) to '%s' during blend out. This can cause incorrect visuals!"),
+			*GetNameSafe(Montage), *ToSectionName.ToString());
+
+		Play(PlayRate);
+	}
+}
+
 FName FAnimMontageInstance::GetCurrentSection() const
 {
 	if ( Montage )
@@ -847,17 +889,20 @@ FName FAnimMontageInstance::GetNextSection() const
 	return NAME_None;
 }
 
-bool FAnimMontageInstance::ChangePositionToSection(FName SectionName, bool bEndOfSection)
+int32 FAnimMontageInstance::GetNextSectionID(int32 const & CurrentSectionID) const
 {
-	const int32 SectionID = Montage->GetSectionIndex(SectionName);
-	if ( Montage->IsValidSectionIndex(SectionID) )
+	return (IsActive() && (CurrentSectionID < NextSections.Num())) ? NextSections[CurrentSectionID] : INDEX_NONE;
+}
+
+FName FAnimMontageInstance::GetSectionNameFromID(int32 const & SectionID) const
+{
+	if (Montage && Montage->IsValidSectionIndex(SectionID))
 	{
-		FCompositeSection & CurSection = Montage->GetAnimCompositeSection(SectionID);
-		Position = Montage->CalculatePos(CurSection, bEndOfSection? Montage->GetSectionLength(SectionID) - KINDA_SMALL_NUMBER : 0.0f);
-		return true;
+		FCompositeSection const & CurrentSection = Montage->GetAnimCompositeSection(SectionID);
+		return CurrentSection.SectionName;
 	}
 
-	return false;
+	return NAME_None;
 }
 
 void FAnimMontageInstance::UpdateWeight(float DeltaTime)
@@ -961,7 +1006,7 @@ bool FAnimMontageInstance::SimulateAdvance(float DeltaTime, float & InOutPositio
 	return true;
 }
 
-void FAnimMontageInstance::Advance(float DeltaTime, FRootMotionMovementParams & OutRootMotionParams)
+void FAnimMontageInstance::Advance(float DeltaTime, struct FRootMotionMovementParams * OutRootMotionParams)
 {
 	if( IsValid() )
 	{
@@ -982,7 +1027,7 @@ void FAnimMontageInstance::Advance(float DeltaTime, FRootMotionMovementParams &
 			const float CombinedPlayRate = PlayRate * Montage->RateScale;
 			const bool bPlayingForward = (CombinedPlayRate > 0.f);
 
-			const bool bExtractRootMotion = (Montage->bEnableRootMotionRotation || Montage->bEnableRootMotionTranslation);
+			const bool bExtractRootMotion = (OutRootMotionParams != NULL) && (Montage->bEnableRootMotionRotation || Montage->bEnableRootMotionTranslation);
 			
 			float DesiredDeltaMove = CombinedPlayRate * DeltaTime;
 			float OriginalMoveDelta = DesiredDeltaMove;
@@ -1033,7 +1078,7 @@ void FAnimMontageInstance::Advance(float DeltaTime, FRootMotionMovementParams &
 						// Extract Root Motion for this time slice, and accumulate it.
 						if( bExtractRootMotion && AnimInstance.IsValid() )
 						{
-							OutRootMotionParams.Accumulate( Montage->ExtractRootMotionFromTrackRange(PrevPosition, Position) );
+							OutRootMotionParams->Accumulate( Montage->ExtractRootMotionFromTrackRange(PrevPosition, Position) );
 						}
 
 						// If about to reach the end of the montage...
diff --git a/Engine/Source/Runtime/Engine/Private/Animation/AnimNode_Slot.cpp b/Engine/Source/Runtime/Engine/Private/Animation/AnimNode_Slot.cpp
--- a/Engine/Source/Runtime/Engine/Private/Animation/AnimNode_Slot.cpp
+++ b/Engine/Source/Runtime/Engine/Private/Animation/AnimNode_Slot.cpp
@@ -76,7 +76,7 @@ void FAnimNode_Slot::GatherDebugData(FNodeDebugData& DebugData)
 			{
 				float Weight;
 				float CurrentAnimPos;
-				if (UAnimSequenceBase* Anim = Segment.GetAnimationData(MontageInstance->Position, CurrentAnimPos, Weight))
+				if (UAnimSequenceBase* Anim = Segment.GetAnimationData(MontageInstance->GetPosition(), CurrentAnimPos, Weight))
 				{
 					FString MontageLine = FString::Printf(TEXT("Montage: '%s' Anim: '%s' Play Time:%.2f W:%.2f%%"), *MontageInstance->Montage->GetName(), *Anim->GetName(), CurrentAnimPos, Weight*100.f);
 					DebugData.BranchFlow(1.0f).AddDebugItem(MontageLine, true);
diff --git a/Engine/Source/Runtime/Engine/Private/Animation/AnimSingleNodeInstance.cpp b/Engine/Source/Runtime/Engine/Private/Animation/AnimSingleNodeInstance.cpp
--- a/Engine/Source/Runtime/Engine/Private/Animation/AnimSingleNodeInstance.cpp
+++ b/Engine/Source/Runtime/Engine/Private/Animation/AnimSingleNodeInstance.cpp
@@ -275,7 +275,7 @@ void UAnimSingleNodeInstance::NativeUpdateAnimation(float DeltaTimeX)
 			FAnimMontageInstance * ActiveMontageInstance = GetActiveMontageInstance();
 			if (ActiveMontageInstance)
 			{
-				CurrentTime = ActiveMontageInstance->Position;
+				CurrentTime = ActiveMontageInstance->GetPosition();
 			}
 			else if (bPlaying)
 			{
@@ -427,7 +427,7 @@ void UAnimSingleNodeInstance::Montage_Advance(float DeltaTime)
 	FAnimMontageInstance * CurMontageInstance = GetActiveMontageInstance();
 	if ( CurMontageInstance )
 	{
-		CurrentTime = CurMontageInstance->Position;
+		CurrentTime = CurMontageInstance->GetPosition();
 	}
 }
 
@@ -478,7 +478,7 @@ void UAnimSingleNodeInstance::SetPlayRate(float InPlayRate)
 
 	if (FAnimMontageInstance* CurMontageInstance = GetActiveMontageInstance())
 	{
-		CurMontageInstance->PlayRate = InPlayRate;
+		CurMontageInstance->SetPlayRate(InPlayRate);
 	}
 }
 
@@ -517,7 +517,7 @@ void UAnimSingleNodeInstance::SetPosition(float InPosition, bool bFireNotifies)
 
 	if (FAnimMontageInstance* CurMontageInstance = GetActiveMontageInstance())
 	{
-		CurMontageInstance->Position = CurrentTime;
+		CurMontageInstance->SetPosition(CurrentTime);
 	}
 
 	// Handle notifies
diff --git a/Engine/Source/Runtime/Engine/Private/Character.cpp b/Engine/Source/Runtime/Engine/Private/Character.cpp
--- a/Engine/Source/Runtime/Engine/Private/Character.cpp
+++ b/Engine/Source/Runtime/Engine/Private/Character.cpp
@@ -1008,7 +1008,7 @@ void ACharacter::SimulatedRootMotionPositionFixup(float DeltaSeconds)
 			if( RestoreReplicatedMove(RootMotionRepMove) )
 			{
 				const float ServerPosition = RootMotionRepMove.RootMotion.Position;
-				const float ClientPosition = ClientMontageInstance->Position;
+				const float ClientPosition = ClientMontageInstance->GetPosition();
 				const float DeltaPosition = (ClientPosition - ServerPosition);
 				if( FMath::Abs(DeltaPosition) > KINDA_SMALL_NUMBER )
 				{
@@ -1018,10 +1018,11 @@ void ACharacter::SimulatedRootMotionPositionFixup(float DeltaSeconds)
 					// Simulate Root Motion for delta move.
 					if( CharacterMovement )
 					{
+						const float MontagePlayRate = ClientMontageInstance->GetPlayRate();
 						// Guess time it takes for this delta track position, so we can get falling physics accurate.
-						if (!FMath::IsNearlyZero(ClientMontageInstance->PlayRate))
+						if (!FMath::IsNearlyZero(MontagePlayRate))
 						{
-							float DeltaTime = DeltaPosition / ClientMontageInstance->PlayRate;
+							const float DeltaTime = DeltaPosition / MontagePlayRate;
 
 							// Even with negative playrate deltatime should be positive.
 							check(DeltaTime > 0.f);
@@ -1073,7 +1074,7 @@ bool ACharacter::CanUseRootMotionRepMove(const FSimulatedRootMotionReplicatedMov
 		{
 			UAnimMontage * AnimMontage = ClientMontageInstance.Montage;
 			const float ServerPosition = RootMotionRepMove.RootMotion.Position;
-			const float ClientPosition = ClientMontageInstance.Position;
+			const float ClientPosition = ClientMontageInstance.GetPosition();
 			const float DeltaPosition = (ClientPosition - ServerPosition);
 			const int32 CurrentSectionIndex = AnimMontage->GetSectionIndexFromPosition(ClientPosition);
 			if( CurrentSectionIndex != INDEX_NONE )
@@ -1086,7 +1087,7 @@ bool ACharacter::CanUseRootMotionRepMove(const FSimulatedRootMotionReplicatedMov
 				// if we are looping and just wrapped over, skip. That's also not easy to handle and not frequent.
 				const bool bHasLooped = (NextSectionIndex == CurrentSectionIndex) && (FMath::Abs(DeltaPosition) > (AnimMontage->GetSectionLength(CurrentSectionIndex) / 2.f));
 				// Can only simulate forward in time, so we need to make sure server move is not ahead of the client.
-				const bool bServerAheadOfClient = ((DeltaPosition * ClientMontageInstance.PlayRate) < 0.f);
+				const bool bServerAheadOfClient = ((DeltaPosition * ClientMontageInstance.GetPlayRate()) < 0.f);
 
 				UE_LOG(LogRootMotion, Log,  TEXT("\t\tACharacter::CanUseRootMotionRepMove ServerPosition: %.3f, ClientPosition: %.3f, DeltaPosition: %.3f, bSameSections: %d, bHasLooped: %d, bServerAheadOfClient: %d"), 
 					ServerPosition, ClientPosition, DeltaPosition, bSameSections, bHasLooped, bServerAheadOfClient);
@@ -1206,7 +1207,7 @@ void ACharacter::PreReplication( IRepChangedPropertyTracker & ChangedPropertyTra
 		RepRootMotion.MovementBase		= BasedMovement.MovementBase;
 		RepRootMotion.MovementBaseBoneName = BasedMovement.BoneName;
 		RepRootMotion.AnimMontage		= RootMotionMontageInstance->Montage;
-		RepRootMotion.Position			= RootMotionMontageInstance->Position;
+		RepRootMotion.Position			= RootMotionMontageInstance->GetPosition();
 
 		DOREPLIFETIME_ACTIVE_OVERRIDE( ACharacter, RepRootMotion, true );
 	}
@@ -1254,7 +1255,7 @@ float ACharacter::PlayAnimMontage(class UAnimMontage* AnimMontage, float InPlayR
 	UAnimInstance * AnimInstance = (Mesh)? Mesh->GetAnimInstance() : NULL; 
 	if( AnimMontage && AnimInstance )
 	{
-		float const Duration =  AnimInstance->Montage_Play(AnimMontage, InPlayRate);
+		float const Duration = AnimInstance->Montage_Play(AnimMontage, InPlayRate);
 
 		if (Duration > 0.f)
 		{
@@ -1275,7 +1276,7 @@ void ACharacter::StopAnimMontage(class UAnimMontage* AnimMontage)
 {
 	UAnimInstance * AnimInstance = (Mesh)? Mesh->GetAnimInstance() : NULL; 
 	UAnimMontage * MontageToStop = (AnimMontage)? AnimMontage : GetCurrentMontage();
-	bool bShouldStopMontage =  AnimInstance && MontageToStop && AnimInstance->Montage_IsPlaying(MontageToStop);
+	bool bShouldStopMontage =  AnimInstance && MontageToStop && !AnimInstance->Montage_GetIsStopped(MontageToStop);
 
 	if ( bShouldStopMontage )
 	{
diff --git a/Engine/Source/Runtime/Engine/Private/CharacterMovementComponent.cpp b/Engine/Source/Runtime/Engine/Private/CharacterMovementComponent.cpp
--- a/Engine/Source/Runtime/Engine/Private/CharacterMovementComponent.cpp
+++ b/Engine/Source/Runtime/Engine/Private/CharacterMovementComponent.cpp
@@ -5669,7 +5669,7 @@ void UCharacterMovementComponent::SendClientAdjustment()
 			ClientAdjustRootMotionPosition
 				(
 				ServerData->PendingAdjustment.TimeStamp,
-				CharacterOwner->GetRootMotionAnimMontageInstance()->Position,
+				CharacterOwner->GetRootMotionAnimMontageInstance()->GetPosition(),
 				ServerData->PendingAdjustment.NewLoc,
 				CompressedRotation,
 				ServerData->PendingAdjustment.NewVel.Z,
@@ -5897,7 +5897,7 @@ void UCharacterMovementComponent::ClientAdjustRootMotionPosition_Implementation(
 		FAnimMontageInstance * RootMotionMontageInstance = CharacterOwner->GetRootMotionAnimMontageInstance();
 		if( RootMotionMontageInstance )
 		{
-			RootMotionMontageInstance->Position = ServerMontageTrackPosition;
+			RootMotionMontageInstance->SetPosition(ServerMontageTrackPosition);
 			CharacterOwner->bClientResimulateRootMotion = true;
 		}
 	}
@@ -6467,7 +6467,7 @@ void FSavedMove_Character::PostUpdate(ACharacter* Character, FSavedMove_Characte
 		if( RootMotionMontageInstance )
 		{
 			RootMotionMontage = RootMotionMontageInstance->Montage;
-			RootMotionTrackPosition = RootMotionMontageInstance->Position;
+			RootMotionTrackPosition = RootMotionMontageInstance->GetPosition();
 			RootMotionMovement = Character->ClientRootMotionParams;
 		}
 	}
@@ -6569,8 +6569,9 @@ void FSavedMove_Character::PrepMoveFor(ACharacter* Character)
 			if( RootMotionMontageInstance && (RootMotionMontage == RootMotionMontageInstance->Montage) )
 			{
 				RootMotionMovement.Clear();
-				RootMotionMontageInstance->SimulateAdvance(DeltaTime, RootMotionMontageInstance->Position, RootMotionMovement);
-				RootMotionTrackPosition = RootMotionMontageInstance->Position;
+				RootMotionTrackPosition = RootMotionMontageInstance->GetPosition();
+				RootMotionMontageInstance->SimulateAdvance(DeltaTime, RootMotionTrackPosition, RootMotionMovement);
+				RootMotionMontageInstance->SetPosition(RootMotionTrackPosition);
 			}
 		}
 		Character->CharacterMovement->RootMotionParams = RootMotionMovement;
