Fix: Trail emitters won't work properly if their source is simulated in local space (emitter transform won't be taken into account).  Adding EmitterToWorld translation to trail particle spawn positions to fix this issue.

diff --git a/Engine/Source/Runtime/Engine/Private/Particles/ParticleTrail2EmitterInstance.cpp b/Engine/Source/Runtime/Engine/Private/Particles/ParticleTrail2EmitterInstance.cpp
--- a/Engine/Source/Runtime/Engine/Private/Particles/ParticleTrail2EmitterInstance.cpp
+++ b/Engine/Source/Runtime/Engine/Private/Particles/ParticleTrail2EmitterInstance.cpp
@@ -2264,7 +2264,7 @@ bool FParticleRibbonEmitterInstance::ResolveSourcePoint(int32 InTrailIdx,
 					FBaseParticle* SourceParticle = (SourceIndices[InTrailIdx] != -1) ? SourceEmitter->GetParticleDirect(SourceIndices[InTrailIdx]) : NULL;
 					if (SourceParticle != NULL)
 					{
-						OutPosition = SourceParticle->Location;
+						OutPosition = SourceParticle->Location + SourceEmitter->SimulationToWorld.GetOrigin();
 						OutTangent = SourceParticle->Location - SourceParticle->OldLocation;
 						SourceTimes[InTrailIdx] = SourceParticle->RelativeTime;
 					}
