Remove old cvar that is unnecessary

diff --git a/Engine/Source/Runtime/Engine/Private/NetworkDriver.cpp b/Engine/Source/Runtime/Engine/Private/NetworkDriver.cpp
--- a/Engine/Source/Runtime/Engine/Private/NetworkDriver.cpp
+++ b/Engine/Source/Runtime/Engine/Private/NetworkDriver.cpp
@@ -1726,16 +1726,6 @@ FActorPriority::FActorPriority(class UNetConnection* InConnection, struct FActor
 	}
 }
 
-int GNetAllowConnectionSkipping = 0;
-
-static FAutoConsoleVariableRef CVarNetRemoteRoleOnCloseTest(
-	TEXT("net.AllowConnectionSkipping"),
-	GNetAllowConnectionSkipping,
-	TEXT("If true, we dont have to process all clients on the same frame"),
-	ECVF_Default
-	);
-
-
 int32 UNetDriver::ServerReplicateActors(float DeltaSeconds)
 {
 	SCOPE_CYCLE_COUNTER(STAT_NetServerRepActorsTime);
@@ -1753,29 +1743,26 @@ int32 UNetDriver::ServerReplicateActors(float DeltaSeconds)
 
 	int32 NumClientsToTick = ClientConnections.Num();
 
-	if( !GNetAllowConnectionSkipping)
-	{
-		// by default only throttle update for listen servers unless specified on the commandline
-		static bool bForceClientTickingThrottle = FParse::Param(FCommandLine::Get(),TEXT("limitclientticks"));
-		if (bForceClientTickingThrottle || GetNetMode() == NM_ListenServer)
-		{
-			// determine how many clients to tick this frame based on GEngine->NetTickRate (always tick at least one client), double for lan play
-			static float DeltaTimeOverflow = 0.f;
-			// updates are doubled for lan play
-			static bool LanPlay = FParse::Param(FCommandLine::Get(),TEXT("lanplay"));
-			//@todo - ideally we wouldn't want to tick more clients with a higher deltatime as that's not going to be good for performance and probably saturate bandwidth in hitchy situations, maybe 
-			// come up with a solution that is greedier with higher framerates, but still won't risk saturating server upstream bandwidth
-			float ClientUpdatesThisFrame = GEngine->NetClientTicksPerSecond * (DeltaSeconds + DeltaTimeOverflow) * (LanPlay ? 2.f : 1.f);
-			NumClientsToTick = FMath::Min<int32>(NumClientsToTick,FMath::TruncToInt(ClientUpdatesThisFrame));
-			//UE_LOG(LogNet, Log, TEXT("%2.3f: Ticking %d clients this frame, %2.3f/%2.4f"),GetWorld()->GetTimeSeconds(),NumClientsToTick,DeltaSeconds,ClientUpdatesThisFrame);
-			if (NumClientsToTick == 0)
-			{
-				// if no clients are ticked this frame accumulate the time elapsed for the next frame
-				DeltaTimeOverflow += DeltaSeconds;
-				return 0;
-			}
-			DeltaTimeOverflow = 0.f;
+	// by default only throttle update for listen servers unless specified on the commandline
+	static bool bForceClientTickingThrottle = FParse::Param(FCommandLine::Get(),TEXT("limitclientticks"));
+	if (bForceClientTickingThrottle || GetNetMode() == NM_ListenServer)
+	{
+		// determine how many clients to tick this frame based on GEngine->NetTickRate (always tick at least one client), double for lan play
+		static float DeltaTimeOverflow = 0.f;
+		// updates are doubled for lan play
+		static bool LanPlay = FParse::Param(FCommandLine::Get(),TEXT("lanplay"));
+		//@todo - ideally we wouldn't want to tick more clients with a higher deltatime as that's not going to be good for performance and probably saturate bandwidth in hitchy situations, maybe 
+		// come up with a solution that is greedier with higher framerates, but still won't risk saturating server upstream bandwidth
+		float ClientUpdatesThisFrame = GEngine->NetClientTicksPerSecond * (DeltaSeconds + DeltaTimeOverflow) * (LanPlay ? 2.f : 1.f);
+		NumClientsToTick = FMath::Min<int32>(NumClientsToTick,FMath::TruncToInt(ClientUpdatesThisFrame));
+		//UE_LOG(LogNet, Log, TEXT("%2.3f: Ticking %d clients this frame, %2.3f/%2.4f"),GetWorld()->GetTimeSeconds(),NumClientsToTick,DeltaSeconds,ClientUpdatesThisFrame);
+		if (NumClientsToTick == 0)
+		{
+			// if no clients are ticked this frame accumulate the time elapsed for the next frame
+			DeltaTimeOverflow += DeltaSeconds;
+			return 0;
 		}
+		DeltaTimeOverflow = 0.f;
 	}
 
 	bool bNetRelevantActorCount = false;
