Lower verbosity on rejected RPC's. Requested by SteveP, and makes sense. This really isn't an error, and more by design.

diff --git a/Engine/Source/Runtime/Engine/Private/DataReplication.cpp b/Engine/Source/Runtime/Engine/Private/DataReplication.cpp
--- a/Engine/Source/Runtime/Engine/Private/DataReplication.cpp
+++ b/Engine/Source/Runtime/Engine/Private/DataReplication.cpp
@@ -701,10 +701,11 @@ bool FObjectReplicator::ReceivedBunch( FInBunch &Bunch, const FReplicationFlags
 			}
 			else
 			{
-				UE_LOG( LogNet, Warning, TEXT( "Rejected unwanted function %s in %s" ), *Message.ToString(), *Object->GetFullName() );
+				UE_LOG( LogNet, Verbose, TEXT( "Rejected unwanted function %s in %s" ), *Message.ToString(), *Object->GetFullName() );
 
 				if ( !OwningChannel->Connection->TrackLogsPerSecond() )	// This will disconnect the client if we get here too often
 				{
+					UE_LOG( LogNet, Error, TEXT( "Rejected too many unwanted functions %s in %s" ), *Message.ToString(), *Object->GetFullName() );
 					return false;
 				}
 			}
