#ue4 - UPrimitiveComponent::SyncComponentToRBPhysics() can ignore scale when trying to move the component to match the physics location.

- Scale provided to PhysX doesn't always match our scale, but we don't use the scale of the result here anyway.

diff --git a/Engine/Source/Runtime/Engine/Private/PrimitiveComponentPhysics.cpp b/Engine/Source/Runtime/Engine/Private/PrimitiveComponentPhysics.cpp
--- a/Engine/Source/Runtime/Engine/Private/PrimitiveComponentPhysics.cpp
+++ b/Engine/Source/Runtime/Engine/Private/PrimitiveComponentPhysics.cpp
@@ -417,16 +417,20 @@ void UPrimitiveComponent::SyncComponentToRBPhysics()
 	AActor* Owner = GetOwner();
 	if(Owner != NULL)
 	{
-		const bool bStillInWorld = Owner->CheckStillInWorld();
-		if (!bStillInWorld || Owner->IsPendingKill() || IsPendingKill() || !IsSimulatingPhysics())
+		if (Owner->IsPendingKill() || !Owner->CheckStillInWorld())
 		{
 			return;
 		}
 	}
 
+	if (IsPendingKill() || !IsSimulatingPhysics())
+	{
+		return;
+	}
+
 	// See if the transform is actually different, and if so, move the component to match physics
 	const FTransform NewTransform = UseBI->GetUnrealWorldTransform();	
-	if(!NewTransform.Equals(ComponentToWorld))
+	if(!NewTransform.EqualsNoScale(ComponentToWorld))
 	{
 		const FVector MoveBy = NewTransform.GetLocation() - ComponentToWorld.GetLocation();
 		const FRotator NewRotation = NewTransform.Rotator();
