Fix compile error.

diff --git a/Engine/Source/Runtime/Engine/Private/Animation/AnimInstance.cpp b/Engine/Source/Runtime/Engine/Private/Animation/AnimInstance.cpp
--- a/Engine/Source/Runtime/Engine/Private/Animation/AnimInstance.cpp
+++ b/Engine/Source/Runtime/Engine/Private/Animation/AnimInstance.cpp
@@ -2087,7 +2087,7 @@ UAnimMontage * UAnimInstance::GetCurrentActiveMontage()
 		FAnimMontageInstance * MontageInstance = MontageInstances[InstanceIndex];
 		if (MontageInstance && MontageInstance->IsActive())
 		{
-			MontageInstance->Montage;
+			return MontageInstance->Montage;
 		}
 	}
 
