#ue4 - Fix walking on skeletal meshes that simulate physics (player would teleport around), and fix applied forces to physics objects underneath the player.

- Reading bone locations on skeletal meshes that are simulating physics needs to be done in the same tick group across frames, until double buffering is introduced.
- Forces applied to objects we stand on were using the wrong impact point. They also now use the Mass variable of CharacterMovement (which can be scaled for this force by StandingDownwardForceScale).
- Pushing physics objects has been refactored out of HandleImpact to a new function: ApplyImpactPhysicsForces.
- Landing on an object can apply a force, though this new function. Physics forces still need work.

diff --git a/Engine/Source/Runtime/Engine/Classes/GameFramework/CharacterMovementComponent.h b/Engine/Source/Runtime/Engine/Classes/GameFramework/CharacterMovementComponent.h
--- a/Engine/Source/Runtime/Engine/Classes/GameFramework/CharacterMovementComponent.h
+++ b/Engine/Source/Runtime/Engine/Classes/GameFramework/CharacterMovementComponent.h
@@ -327,7 +327,7 @@ public:
 	/** Mass of pawn (for when momentum is imparted to it). */
 	UPROPERTY(Category="Character Movement", EditAnywhere, BlueprintReadWrite, meta=(ClampMin="0", UIMin="0"))
 	float Mass;
-	
+
 	/** If enabled, the player will interact with physics objects when walking into them. */
 	UPROPERTY(Category="Agent Physics", EditAnywhere, BlueprintReadWrite)
 	bool bEnablePhysicsInteraction;
@@ -345,6 +345,10 @@ public:
 	UPROPERTY(Category="Agent Physics", EditAnywhere, BlueprintReadWrite, meta=(editcondition = "bEnablePhysicsInteraction"))
 	bool bScalePushForceToVelocity;
 
+	/** Force applied to objects we stand on (due to Mass and Gravity) is scaled by this amount. */
+	UPROPERTY(Category="Agent Physics", EditAnywhere, BlueprintReadWrite, meta=(editcondition = "bEnablePhysicsInteraction"))
+	float StandingDownwardForceScale;
+
 	/** Initial impulse force to apply when the player bounces into a blocking physics object. */
 	UPROPERTY(Category="Agent Physics", EditAnywhere, BlueprintReadWrite, meta=(editcondition = "bEnablePhysicsInteraction"))
 	float InitialPushForceFactor;
@@ -369,7 +373,7 @@ public:
 	UPROPERTY(Category="Agent Physics", EditAnywhere, BlueprintReadWrite, meta=(editcondition = "bEnablePhysicsInteraction"))
 	float MaxTouchForce;
 
-	/** Force per kg applied constanly to all overlapping components. */
+	/** Force per kg applied constantly to all overlapping components. */
 	UPROPERTY(Category="Agent Physics", EditAnywhere, BlueprintReadWrite, meta=(editcondition = "bEnablePhysicsInteraction"))
 	float RepulsionForce;
 
@@ -1098,9 +1102,17 @@ protected:
 	/** Overridden to set bJustTeleported to true, so we don't make incorrect velocity calculations based on adjusted movement. */
 	virtual bool ResolvePenetration(const FVector& Adjustment, const FHitResult& Hit, const FRotator& NewRotation) override;
 
-	/** Don't call for intermediate parts of move! */
+	/** Handle a blocking impact. Calls ApplyImpactPhysicsForces for the hit, if bEnablePhysicsInteraction is true. */
 	virtual void HandleImpact(FHitResult const& Hit, float TimeSlice=0.f, const FVector& MoveDelta = FVector::ZeroVector) override;
 
+	/**
+	 * Apply physics forces to the impacted component, if bEnablePhysicsInteraction is true.
+	 * @param Impact				HitResult that resulted in the impact
+	 * @param ImpactAcceleration	Acceleration of the character at the time of impact
+	 * @param ImpactVelocity		Velocity of the character at the time of impact
+	 */
+	virtual void ApplyImpactPhysicsForces(const FHitResult& Impact, const FVector& ImpactAcceleration, const FVector& ImpactVelocity);
+
 	/** Custom version of SlideAlongSurface that handles different movement modes separately; namely during walking physics we might not want to slide up slopes. */
 	virtual float SlideAlongSurface(const FVector& Delta, float Time, const FVector& Normal, FHitResult &Hit, bool bHandleImpact) override;
 
diff --git a/Engine/Source/Runtime/Engine/Private/CharacterMovementComponent.cpp b/Engine/Source/Runtime/Engine/Private/CharacterMovementComponent.cpp
--- a/Engine/Source/Runtime/Engine/Private/CharacterMovementComponent.cpp
+++ b/Engine/Source/Runtime/Engine/Private/CharacterMovementComponent.cpp
@@ -205,7 +205,8 @@ UCharacterMovementComponent::UCharacterMovementComponent(const class FPostConstr
 	bCanWalkOffLedgesWhenCrouching = false;
 
 	bEnablePhysicsInteraction = true;
-	InitialPushForceFactor = 500.0f;;
+	StandingDownwardForceScale = 1.0f;
+	InitialPushForceFactor = 500.0f;
 	PushForceFactor = 750000.0f;
 	PushForcePointZOffsetFactor = -0.75f;
 	bPushForceScaledToMass = false;
@@ -826,25 +827,17 @@ void UCharacterMovementComponent::TickComponent(float DeltaTime, enum ELevelTick
 
 	if (bEnablePhysicsInteraction)
 	{
-		// Apply downwards force when walking on top of physics objects
-		if (CharacterOwner->GetMovementBase() != NULL)
+		if (CurrentFloor.HitResult.IsValidBlockingHit())
 		{
-			UPrimitiveComponent* BaseComp = CharacterOwner->GetMovementBase();
-
-			if (BaseComp->IsSimulatingPhysics())
+			// Apply downwards force when walking on top of physics objects
+			if (UPrimitiveComponent* BaseComp = CurrentFloor.HitResult.GetComponent())
 			{
-				UPrimitiveComponent* CharacterComp = CharacterOwner->Mesh.Get();
-
-#if WITH_EDITOR
-				// We need to calculate the mass here, as the bodies, when kinematic, and do not return proper mass
-				float CharacterMass = CharacterComp != NULL ? CharacterComp->CalculateMass() : 1.0f;
-#else
-				float CharacterMass = 1.0f;
-#endif
-
-				const float GravZ = GetGravityZ();
-				const FVector ForceLocation = UpdatedComponent->GetComponentLocation() - CharacterOwner->CapsuleComponent->GetScaledCapsuleHalfHeight();
-				BaseComp->AddForceAtLocation(FVector(0,0, GravZ * CharacterMass), ForceLocation);
+				if (StandingDownwardForceScale != 0.f && BaseComp->IsAnySimulatingPhysics())
+				{
+					const float GravZ = GetGravityZ();
+					const FVector ForceLocation = CurrentFloor.HitResult.ImpactPoint;
+					BaseComp->AddForceAtLocation(FVector(0.f, 0.f, GravZ * Mass * StandingDownwardForceScale), ForceLocation, CurrentFloor.HitResult.BoneName);
+				}
 			}
 		}
 
@@ -1174,14 +1167,23 @@ void UCharacterMovementComponent::MaybeUpdateBasedMovement(float DeltaSeconds)
 	UPrimitiveComponent* MovementBase = CharacterOwner->GetMovementBase();
 	if (MovementBaseUtility::UseRelativeLocation(MovementBase))
 	{
-		if(!MovementBase->IsSimulatingPhysics())
+		const bool bBaseIsSimulatingPhysics = MovementBase->IsSimulatingPhysics();
+		
+		// Temporarily disabling deferred tick on skeletal mesh components that sim physics.
+		// We need to be consistent on when we read the bone locations for those, and while this reads
+		// the wrong location, the relative changes (which is what we care about) will be accurate.
+		const bool bAllowDefer = (bBaseIsSimulatingPhysics && !Cast<USkeletalMeshComponent>(MovementBase));
+		
+		if (!bBaseIsSimulatingPhysics || !bAllowDefer)
 		{
 			UpdateBasedMovement(DeltaSeconds);
+			PreClothComponentTick.SetTickFunctionEnable(false);
 		}
 		else
 		{
 			// defer movement base update until after physics
 			bDeferUpdateBasedMovement = true;
+			PreClothComponentTick.SetTickFunctionEnable(true);
 		}
 	}
 }
@@ -3304,7 +3306,7 @@ void UCharacterMovementComponent::PhysWalking(float deltaTime, int32 Iterations)
 
 		// Save current values
 		UPrimitiveComponent * const OldBase = GetMovementBase();
-		const FVector OldBaseLocation = (OldBase != NULL) ? OldBase->GetComponentLocation() : FVector::ZeroVector;
+		const FVector PreviousBaseLocation = (OldBase != NULL) ? OldBase->GetComponentLocation() : FVector::ZeroVector;
 		const FVector OldLocation = UpdatedComponent->GetComponentLocation();
 		const FFindFloorResult OldFloor = CurrentFloor;
 
@@ -3376,7 +3378,7 @@ void UCharacterMovementComponent::PhysWalking(float deltaTime, int32 Iterations)
 			if ( !NewDelta.IsZero() )
 			{
 				// first revert this move
-				RevertMove(OldLocation, OldBase, OldBaseLocation, OldFloor, false);
+				RevertMove(OldLocation, OldBase, PreviousBaseLocation, OldFloor, false);
 
 				// avoid repeated ledge moves if the first one fails
 				bTriedLedgeMove = true;
@@ -3398,7 +3400,7 @@ void UCharacterMovementComponent::PhysWalking(float deltaTime, int32 Iterations)
 				bCheckedFall = true;
 
 				// revert this move
-				RevertMove(OldLocation, OldBase, OldBaseLocation, OldFloor, true);
+				RevertMove(OldLocation, OldBase, PreviousBaseLocation, OldFloor, true);
 				remainingTime = 0.f;
 				break;
 			}
@@ -3580,7 +3582,10 @@ void UCharacterMovementComponent::SetPostLandedPhysics(const FHitResult& Hit)
 		}
 		else
 		{
+			const FVector PreImpactAccel = Acceleration + (IsFalling() ? FVector(0.f, 0.f, GetGravityZ()) : FVector::ZeroVector);
+			const FVector PreImpactVelocity = Velocity;
 			SetMovementMode(MOVE_Walking);
+			ApplyImpactPhysicsForces(Hit, PreImpactAccel, PreImpactVelocity);
 		}
 	}
 }
@@ -4639,65 +4644,77 @@ void UCharacterMovementComponent::HandleImpact(FHitResult const& Impact, float T
 		NotifyBumpedPawn(OtherPawn);
 	}
 
-	if (bEnablePhysicsInteraction && Impact.Component != NULL && Impact.Component->IsAnySimulatingPhysics() && Impact.bBlockingHit)
+	if (bEnablePhysicsInteraction)
 	{
-		FVector ForcePoint = Impact.ImpactPoint;
+		const FVector ForceAccel = Acceleration + (IsFalling() ? FVector(0.f, 0.f, GetGravityZ()) : FVector::ZeroVector);
+		ApplyImpactPhysicsForces(Impact, ForceAccel, Velocity);
+	}
+}
 
-		FBodyInstance* BI = Impact.Component->GetBodyInstance(Impact.BoneName);
-		
-		float BodyMass = 1.0f;
-		
-		if (BI != NULL)
+void UCharacterMovementComponent::ApplyImpactPhysicsForces(const FHitResult& Impact, const FVector& ImpactAcceleration, const FVector& ImpactVelocity)
+{
+	if (bEnablePhysicsInteraction && Impact.bBlockingHit)
+	{
+		UPrimitiveComponent* ImpactComponent = Impact.GetComponent();
+		if (ImpactComponent != NULL && ImpactComponent->IsAnySimulatingPhysics())
 		{
-			BodyMass = FMath::Max(BI->GetBodyMass(), 1.0f);
+			FVector ForcePoint = Impact.ImpactPoint;
 
-			FBox Bounds = BI->GetBodyBounds();
+			FBodyInstance* BI = ImpactComponent->GetBodyInstance(Impact.BoneName);
 
-			FVector Center, Extents;
-			Bounds.GetCenterAndExtents(Center, Extents);
+			float BodyMass = 1.0f;
 
-			if (!Extents.IsNearlyZero())
+			if (BI != NULL)
 			{
-				ForcePoint.Z = Center.Z + Extents.Z * PushForcePointZOffsetFactor;
+				BodyMass = FMath::Max(BI->GetBodyMass(), 1.0f);
+
+				FBox Bounds = BI->GetBodyBounds();
+
+				FVector Center, Extents;
+				Bounds.GetCenterAndExtents(Center, Extents);
+
+				if (!Extents.IsNearlyZero())
+				{
+					ForcePoint.Z = Center.Z + Extents.Z * PushForcePointZOffsetFactor;
+				}
 			}
-		}
-		
-		FVector Force = Impact.ImpactNormal * -1.0f;
-		Force.Normalize();
-			
-		float PushForceModificator = 1.0f;
 
-		FVector CurrentVelocity = Impact.Component->GetPhysicsLinearVelocity();
-		FVector VirtualVelocity = Acceleration.SafeNormal() * GetMaxSpeed();
+			FVector Force = Impact.ImpactNormal * -1.0f;
 
-		float Dot = 0.0f;
-		
-		if (bScalePushForceToVelocity && !CurrentVelocity.IsNearlyZero())
-		{			
-			Dot = CurrentVelocity | VirtualVelocity;
-			
-			if (Dot > 0.0f && Dot < 1.0f)
-			{
-				PushForceModificator *= Dot;
+			float PushForceModificator = 1.0f;
+
+			const FVector ComponentVelocity = ImpactComponent->GetPhysicsLinearVelocity();
+			const FVector VirtualVelocity = ImpactAcceleration.IsZero() ? ImpactVelocity : ImpactAcceleration.SafeNormal() * GetMaxSpeed();
+
+			float Dot = 0.0f;
+
+			if (bScalePushForceToVelocity && !ComponentVelocity.IsNearlyZero())
+			{			
+				Dot = ComponentVelocity | VirtualVelocity;
+
+				if (Dot > 0.0f && Dot < 1.0f)
+				{
+					PushForceModificator *= Dot;
+				}
 			}
-		}
 
-		if (bPushForceScaledToMass)
-		{
-			PushForceModificator *= BodyMass;
-		}
+			if (bPushForceScaledToMass)
+			{
+				PushForceModificator *= BodyMass;
+			}
 
-		Force *= PushForceModificator;
+			Force *= PushForceModificator;
 
-		if (CurrentVelocity.IsNearlyZero())
-		{
-			Force *= InitialPushForceFactor;
-			Impact.Component->AddImpulseAtLocation(Force, ForcePoint, Impact.BoneName);
-		}
-		else
-		{
-			Force *= PushForceFactor;
-			Impact.Component->AddForceAtLocation(Force, ForcePoint, Impact.BoneName);
+			if (ComponentVelocity.IsNearlyZero())
+			{
+				Force *= InitialPushForceFactor;
+				ImpactComponent->AddImpulseAtLocation(Force, ForcePoint, Impact.BoneName);
+			}
+			else
+			{
+				Force *= PushForceFactor;
+				ImpactComponent->AddForceAtLocation(Force, ForcePoint, Impact.BoneName);
+			}
 		}
 	}
 }
@@ -6155,12 +6172,6 @@ void UCharacterMovementComponent::RegisterComponentTickFunctions(bool bRegister)
 		if (SetupActorComponentTickFunction(&PreClothComponentTick))
 		{
 			PreClothComponentTick.Target = this;
-
-			// If primary tick is registered, add a prerequisite to it
-			if(PrimaryComponentTick.bCanEverTick)
-			{
-				PreClothComponentTick.AddPrerequisite(this, PrimaryComponentTick); 
-			}
 		}
 	}
 	else
