nicer help

diff --git a/Engine/Source/Runtime/Engine/Private/UnrealEngine.cpp b/Engine/Source/Runtime/Engine/Private/UnrealEngine.cpp
--- a/Engine/Source/Runtime/Engine/Private/UnrealEngine.cpp
+++ b/Engine/Source/Runtime/Engine/Private/UnrealEngine.cpp
@@ -142,9 +142,10 @@ static TAutoConsoleVariable<int32> CVarAllowOneFrameThreadLag(
 static FAutoConsoleVariable CVarSystemResolution(
 	TEXT("r.SetRes"),
 	TEXT("1280x720w"),
-	TEXT("Set the display resolution for the current game view. Has no effect in the editor.")
-	TEXT("  Format e.g. 1280x720w")
-	TEXT("  	   e.g. 1920x1080f")
+	TEXT("Set the display resolution for the current game view. Has no effect in the editor.\n")
+	TEXT("e.g. 1280x720w for windowed\n")
+	TEXT("     1920x1080f for fullscreen\n")
+	TEXT("     1920x1080wf for windowed fullscreen")
 	);
 
 static TAutoConsoleVariable<float> CVarDepthOfFieldNearBlurSizeThreshold(
