Added Landscape Layer Sample material expression

diff --git a/Engine/Source/Runtime/Engine/Classes/Materials/MaterialExpressionLandscapeLayerBlend.h b/Engine/Source/Runtime/Engine/Classes/Materials/MaterialExpressionLandscapeLayerBlend.h
--- a/Engine/Source/Runtime/Engine/Classes/Materials/MaterialExpressionLandscapeLayerBlend.h
+++ b/Engine/Source/Runtime/Engine/Classes/Materials/MaterialExpressionLandscapeLayerBlend.h
@@ -23,7 +23,7 @@ struct FLayerBlendInput
 	FName LayerName;
 
 	UPROPERTY(EditAnywhere, Category=LayerBlendInput)
-	TEnumAsByte<enum ELandscapeLayerBlendType> BlendType;
+	TEnumAsByte<ELandscapeLayerBlendType> BlendType;
 
 	UPROPERTY(meta = (RequiredInput = "false", ToolTip = "Defaults to 'ConstLayerInput' if not specified"))
 	FExpressionInput LayerInput;
@@ -44,9 +44,9 @@ struct FLayerBlendInput
 
 	FLayerBlendInput()
 		: BlendType(0)
-		, PreviewWeight(0.f)		
-		, ConstLayerInput(0.f, 0.f, 0.f)
-		, ConstHeightInput(0.f)
+		, PreviewWeight(0.0f)
+		, ConstLayerInput(0.0f, 0.0f, 0.0f)
+		, ConstHeightInput(0.0f)
 	{
 	}
 };
@@ -57,7 +57,7 @@ class UMaterialExpressionLandscapeLayerBlend : public UMaterialExpression
 	GENERATED_UCLASS_BODY()
 
 	UPROPERTY(EditAnywhere, Category=MaterialExpressionLandscapeLayerBlend)
-	TArray<struct FLayerBlendInput> Layers;
+	TArray<FLayerBlendInput> Layers;
 
 	/** GUID that should be unique within the material, this is used for parameter renaming. */
 	UPROPERTY()
diff --git a/Engine/Source/Runtime/Engine/Classes/Materials/MaterialExpressionLandscapeLayerSample.h b/Engine/Source/Runtime/Engine/Classes/Materials/MaterialExpressionLandscapeLayerSample.h
new file mode 100644
--- /dev/null
+++ b/Engine/Source/Runtime/Engine/Classes/Materials/MaterialExpressionLandscapeLayerSample.h
@@ -0,0 +1,43 @@
+// Copyright 1998-2014 Epic Games, Inc. All Rights Reserved.
+
+
+#pragma once
+#include "Materials/MaterialExpression.h"
+#include "MaterialExpressionLandscapeLayerSample.generated.h"
+
+UCLASS(collapsecategories, hidecategories=Object)
+class UMaterialExpressionLandscapeLayerSample : public UMaterialExpression
+{
+	GENERATED_UCLASS_BODY()
+
+	UPROPERTY(EditAnywhere, Category=MaterialExpressionLandscapeLayerWeight)
+	FName ParameterName;
+
+	UPROPERTY(EditAnywhere, Category=MaterialExpressionLandscapeLayerWeight)
+	float PreviewWeight;
+
+	/** GUID that should be unique within the material, this is used for parameter renaming. */
+	UPROPERTY()
+	FGuid ExpressionGUID;
+
+public:
+
+	// Begin UMaterialExpression Interface
+	virtual int32 Compile(class FMaterialCompiler* Compiler, int32 OutputIndex, int32 MultiplexIndex) override;
+	virtual void GetCaption(TArray<FString>& OutCaptions) const override;
+	virtual UTexture* GetReferencedTexture() override;
+	// End UMaterialExpression Interface
+
+	ENGINE_API virtual FGuid& GetParameterExpressionId() override
+	{
+		return ExpressionGUID;
+	}
+
+	/**
+	 * Called to get list of parameter names for static parameter sets
+	 */
+	void GetAllParameterNames(TArray<FName> &OutParameterNames, TArray<FGuid> &OutParameterIds);
+};
+
+
+
diff --git a/Engine/Source/Runtime/Engine/Private/Landscape/LandscapeEdit.cpp b/Engine/Source/Runtime/Engine/Private/Landscape/LandscapeEdit.cpp
--- a/Engine/Source/Runtime/Engine/Private/Landscape/LandscapeEdit.cpp
+++ b/Engine/Source/Runtime/Engine/Private/Landscape/LandscapeEdit.cpp
@@ -7,6 +7,7 @@ LandscapeEdit.cpp: Landscape editing
 #include "EnginePrivate.h"
 #include "Materials/MaterialExpressionLandscapeVisibilityMask.h"
 #include "Materials/MaterialExpressionLandscapeLayerWeight.h"
+#include "Materials/MaterialExpressionLandscapeLayerSample.h"
 #include "Materials/MaterialExpressionLandscapeLayerBlend.h"
 #include "Materials/MaterialExpressionLandscapeLayerSwitch.h"
 #include "Landscape/LandscapeDataAccess.h"
@@ -1598,24 +1599,25 @@ TArray<FName> ALandscapeProxy::GetLayersFromMaterial(UMaterialInterface* Materia
 		const TArray<UMaterialExpression*>& Expressions = Material->GetMaterial()->Expressions;
 
 		// TODO: *Unconnected* layer expressions?
-		for (auto ItExpressions = Expressions.CreateConstIterator(); ItExpressions; ItExpressions++)
+		for (UMaterialExpression* Expression : Expressions)
 		{
-			UMaterialExpressionLandscapeLayerWeight* LayerWeightExpression = Cast<UMaterialExpressionLandscapeLayerWeight>(*ItExpressions);
-			UMaterialExpressionLandscapeLayerSwitch* LayerSwitchExpression = Cast<UMaterialExpressionLandscapeLayerSwitch>(*ItExpressions);
-			UMaterialExpressionLandscapeLayerBlend* LayerBlendExpression = Cast<UMaterialExpressionLandscapeLayerBlend>(*ItExpressions);
-			if (LayerWeightExpression)
+			if (auto LayerWeightExpression = Cast<UMaterialExpressionLandscapeLayerWeight>(Expression))
 			{
 				Result.AddUnique(LayerWeightExpression->ParameterName);
 			}
-			else if (LayerSwitchExpression)
+			else if (auto LayerSampleExpression = Cast<UMaterialExpressionLandscapeLayerSample>(Expression))
+			{
+				Result.AddUnique(LayerSampleExpression->ParameterName);
+			}
+			else if (auto LayerSwitchExpression = Cast<UMaterialExpressionLandscapeLayerSwitch>(Expression))
 			{
 				Result.AddUnique(LayerSwitchExpression->ParameterName);
 			}
-			else if (LayerBlendExpression)
+			else if (auto LayerBlendExpression = Cast<UMaterialExpressionLandscapeLayerBlend>(Expression))
 			{
-				for (auto ItExpressionLayers = LayerBlendExpression->Layers.CreateConstIterator(); ItExpressionLayers; ItExpressionLayers++)
+				for (const auto& ExpressionLayer : LayerBlendExpression->Layers)
 				{
-					Result.AddUnique(ItExpressionLayers->LayerName);
+					Result.AddUnique(ExpressionLayer.LayerName);
 				}
 			}
 		}
diff --git a/Engine/Source/Runtime/Engine/Private/Materials/Material.cpp b/Engine/Source/Runtime/Engine/Private/Materials/Material.cpp
--- a/Engine/Source/Runtime/Engine/Private/Materials/Material.cpp
+++ b/Engine/Source/Runtime/Engine/Private/Materials/Material.cpp
@@ -12,6 +12,7 @@
 #include "Materials/MaterialExpressionFontSampleParameter.h"
 #include "Materials/MaterialExpressionLandscapeLayerBlend.h"
 #include "Materials/MaterialExpressionLandscapeLayerSwitch.h"
+#include "Materials/MaterialExpressionLandscapeLayerSample.h"
 #include "Materials/MaterialExpressionLandscapeLayerWeight.h"
 #include "Materials/MaterialExpressionLandscapeVisibilityMask.h"
 #include "Materials/MaterialExpressionMaterialFunctionCall.h"
@@ -1042,6 +1043,7 @@ void UMaterial::GetAllTerrainLayerWeightParameterNames(TArray<FName> &OutParamet
 	OutParameterNames.Empty();
 	OutParameterIds.Empty();
 	GetAllParameterNames<UMaterialExpressionLandscapeLayerWeight>(OutParameterNames, OutParameterIds);
+	GetAllParameterNames<UMaterialExpressionLandscapeLayerSample>(OutParameterNames, OutParameterIds);
 	GetAllParameterNames<UMaterialExpressionLandscapeLayerSwitch>(OutParameterNames, OutParameterIds);
 	GetAllParameterNames<UMaterialExpressionLandscapeLayerBlend>(OutParameterNames, OutParameterIds);
 	GetAllParameterNames<UMaterialExpressionLandscapeVisibilityMask>(OutParameterNames, OutParameterIds);
diff --git a/Engine/Source/Runtime/Engine/Private/Materials/MaterialExpressions.cpp b/Engine/Source/Runtime/Engine/Private/Materials/MaterialExpressions.cpp
--- a/Engine/Source/Runtime/Engine/Private/Materials/MaterialExpressions.cpp
+++ b/Engine/Source/Runtime/Engine/Private/Materials/MaterialExpressions.cpp
@@ -57,6 +57,7 @@
 #include "Materials/MaterialExpressionLandscapeLayerCoords.h"
 #include "Materials/MaterialExpressionLandscapeLayerSwitch.h"
 #include "Materials/MaterialExpressionLandscapeLayerWeight.h"
+#include "Materials/MaterialExpressionLandscapeLayerSample.h"
 #include "Materials/MaterialExpressionLandscapeVisibilityMask.h"
 #include "Materials/MaterialExpressionLightmapUVs.h"
 #include "Materials/MaterialExpressionLightmassReplace.h"
@@ -173,6 +174,7 @@ bool IsAllowedExpressionType(UClass* Class, bool bMaterialFunction)
 			&& !Class->IsChildOf(UMaterialExpressionDynamicParameter::StaticClass())
 			&& !Class->IsChildOf(UMaterialExpressionFontSampleParameter::StaticClass())
 			&& !Class->IsChildOf(UMaterialExpressionLandscapeLayerWeight::StaticClass())
+			&& !Class->IsChildOf(UMaterialExpressionLandscapeLayerSample::StaticClass())
 			&& !Class->IsChildOf(UMaterialExpressionLandscapeLayerSwitch::StaticClass())
 			&& !Class->IsChildOf(UMaterialExpressionLandscapeLayerBlend::StaticClass())
 			&& !Class->IsChildOf(UMaterialExpressionLandscapeVisibilityMask::StaticClass());
@@ -8477,6 +8479,53 @@ void UMaterialExpressionAntialiasedTextureMask::SetDefaultTexture()
 	Texture = LoadObject<UTexture2D>(NULL, TEXT("/Engine/EngineResources/DefaultTexture.DefaultTexture"), NULL, LOAD_None, NULL);
 }
 
+///////////////////////////////////////////////////////////////////////////////
+// UMaterialExpressionLandscapeLayerSample
+///////////////////////////////////////////////////////////////////////////////
+UMaterialExpressionLandscapeLayerSample::UMaterialExpressionLandscapeLayerSample(const class FPostConstructInitializeProperties& PCIP)
+	: Super(PCIP)
+{
+	// Structure to hold one-time initialization
+	struct FConstructorStatics
+	{
+		FString NAME_Landscape;
+		FConstructorStatics()
+			: NAME_Landscape(LOCTEXT("Landscape", "Landscape").ToString())
+		{
+		}
+	};
+	static FConstructorStatics ConstructorStatics;
+
+	bIsParameterExpression = true;
+	MenuCategories.Add(ConstructorStatics.NAME_Landscape);
+}
+
+int32 UMaterialExpressionLandscapeLayerSample::Compile(class FMaterialCompiler* Compiler, int32 OutputIndex, int32 MultiplexIndex)
+{
+	const int32 WeightCode = Compiler->StaticTerrainLayerWeight(ParameterName, Compiler->Constant(PreviewWeight));
+
+	return WeightCode;
+}
+
+UTexture* UMaterialExpressionLandscapeLayerSample::GetReferencedTexture()
+{
+	return GEngine->WeightMapPlaceholderTexture;
+}
+
+void UMaterialExpressionLandscapeLayerSample::GetCaption(TArray<FString>& OutCaptions) const
+{
+	OutCaptions.Add(FString::Printf(TEXT("Sample '%s'"), *ParameterName.ToString()));
+}
+
+void UMaterialExpressionLandscapeLayerSample::GetAllParameterNames(TArray<FName> &OutParameterNames, TArray<FGuid> &OutParameterIds)
+{
+	if (!OutParameterNames.Contains(ParameterName))
+	{
+		OutParameterNames.Add(ParameterName);
+		OutParameterIds.Add(ExpressionGUID);
+	}
+}
+
 ///////////////////////////////////////////////////////////////////////////////
 // UMaterialExpressionLandscapeLayerWeight
 ///////////////////////////////////////////////////////////////////////////////
