Fix UE 4.4 generating thumbnails for landscape textures/materials (bloating map size).
Additionally fix ULandscapeComponent::GetAllReferencedTexturesAndMaterials() (now GetGeneratedTexturesAndMaterialInstances) to only return landscape-generated material instances, it previously returned external MIC assets if a MIC was assigned as the landscape material.

diff --git a/Engine/Source/Runtime/Engine/Classes/Landscape/LandscapeComponent.h b/Engine/Source/Runtime/Engine/Classes/Landscape/LandscapeComponent.h
--- a/Engine/Source/Runtime/Engine/Classes/Landscape/LandscapeComponent.h
+++ b/Engine/Source/Runtime/Engine/Classes/Landscape/LandscapeComponent.h
@@ -300,8 +300,8 @@ public:
 	/** Get the level in which the owning actor resides */
 	class ULevel* GetLevel() const;
 
-	/** Returns all textures and materials used by this component. */
-	ENGINE_API void GetAllReferencedTexturesAndMaterials(TArray<UObject*>& OutTexturesAndMaterials) const;
+	/** Returns all generated textures and material instances used by this component. */
+	ENGINE_API void GetGeneratedTexturesAndMaterialInstances(TArray<UObject*>& OutTexturesAndMaterials) const;
 
 	/** @todo document */
 	ENGINE_API class ALandscapeProxy* GetLandscapeProxy() const;
diff --git a/Engine/Source/Runtime/Engine/Private/Landscape/Landscape.cpp b/Engine/Source/Runtime/Engine/Private/Landscape/Landscape.cpp
--- a/Engine/Source/Runtime/Engine/Private/Landscape/Landscape.cpp
+++ b/Engine/Source/Runtime/Engine/Private/Landscape/Landscape.cpp
@@ -16,6 +16,7 @@ Landscape.cpp: Terrain rendering
 #include "TargetPlatform.h"
 #include "Landscape/Landscape.h"
 #include "Landscape/LandscapeMeshCollisionComponent.h"
+#include "Landscape/LandscapeMaterialInstanceConstant.h"
 #include "Landscape/LandscapeSplinesComponent.h"
 #include "Landscape/LandscapeInfo.h"
 #include "Landscape/LandscapeLayerInfoObject.h"
@@ -427,19 +428,21 @@ void ULandscapeComponent::PostLoad()
 	if (GIsEditor && !HasAnyFlags(RF_ClassDefaultObject))
 	{
 		// Move the MICs and Textures back to the Package if they're currently in the level
+		// Moving them into the level caused them to be duplicated when running PIE, which is *very very slow*, so we've reverted that change
+		// Also clear the public flag to avoid various issues, e.g. generating and saving thumbnails that can never be seen
 		{
 			ULevel* Level = GetLevel();
 			if (ensure(Level))
 			{
 				TArray<UObject*> ObjectsToMoveFromLevelToPackage;
-				GetAllReferencedTexturesAndMaterials(ObjectsToMoveFromLevelToPackage);
+				GetGeneratedTexturesAndMaterialInstances(ObjectsToMoveFromLevelToPackage);
 
 				UPackage* MyPackage = GetOutermost();
 				for (auto* Obj : ObjectsToMoveFromLevelToPackage)
 				{
-					if (Obj && Obj->GetOuter() == Level)
+					Obj->ClearFlags(RF_Public);
+					if (Obj->GetOuter() == Level)
 					{
-						Obj->ClearFlags(RF_Public);
 						Obj->Rename(NULL, MyPackage, REN_DoNotDirty | REN_DontCreateRedirectors | REN_ForceNoResetLoaders | REN_NonTransactional);
 					}
 				}
@@ -574,14 +577,11 @@ ULevel* ULandscapeComponent::GetLevel() const
 	return MyOwner ? MyOwner->GetLevel() : NULL;
 }
 
-void ULandscapeComponent::GetAllReferencedTexturesAndMaterials(TArray<UObject*>& OutTexturesAndMaterials) const
+void ULandscapeComponent::GetGeneratedTexturesAndMaterialInstances(TArray<UObject*>& OutTexturesAndMaterials) const
 {
-	for (UMaterialInstance* CurrentMIC = MaterialInstance; CurrentMIC; CurrentMIC = Cast<UMaterialInstance>(CurrentMIC->Parent))
+	for (UMaterialInstance* CurrentMIC = MaterialInstance; CurrentMIC && CurrentMIC->IsA<ULandscapeMaterialInstanceConstant>(); CurrentMIC = Cast<UMaterialInstance>(CurrentMIC->Parent))
 	{
-		if (CurrentMIC)
-		{
-			OutTexturesAndMaterials.Add(CurrentMIC);
-		}
+		OutTexturesAndMaterials.Add(CurrentMIC);
 	}
 
 	if (HeightmapTexture)
@@ -591,10 +591,7 @@ void ULandscapeComponent::GetAllReferencedTexturesAndMaterials(TArray<UObject*>&
 
 	for (auto* Tex : WeightmapTextures)
 	{
-		if (Tex)
-		{
-			OutTexturesAndMaterials.Add(Tex);
-		}
+		OutTexturesAndMaterials.Add(Tex);
 	}
 
 	if (XYOffsetmapTexture)
diff --git a/Engine/Source/Runtime/Engine/Private/World.cpp b/Engine/Source/Runtime/Engine/Private/World.cpp
--- a/Engine/Source/Runtime/Engine/Private/World.cpp
+++ b/Engine/Source/Runtime/Engine/Private/World.cpp
@@ -5133,7 +5133,7 @@ void UWorld::GetLandscapeTexturesAndMaterials(ULevel* Level, TArray<UObject*>& O
 		ULandscapeComponent* LandscapeComponent = Cast<ULandscapeComponent>(ObjInLevel);
 		if (LandscapeComponent)
 		{
-			LandscapeComponent->GetAllReferencedTexturesAndMaterials(OutTexturesAndMaterials);
+			LandscapeComponent->GetGeneratedTexturesAndMaterialInstances(OutTexturesAndMaterials);
 		}
 	}
 }
