Fix operator== on FLandscapeEditorLayerSettings

diff --git a/Engine/Source/Runtime/Engine/Classes/Landscape/LandscapeProxy.h b/Engine/Source/Runtime/Engine/Classes/Landscape/LandscapeProxy.h
--- a/Engine/Source/Runtime/Engine/Classes/Landscape/LandscapeProxy.h
+++ b/Engine/Source/Runtime/Engine/Classes/Landscape/LandscapeProxy.h
@@ -66,15 +66,10 @@ struct FLandscapeEditorLayerSettings
 	{
 	}
 
-	bool operator==(const FLandscapeEditorLayerSettings& rhs)
+	bool operator==(const FLandscapeEditorLayerSettings& rhs) const
 	{
 		return LayerInfoObj == rhs.LayerInfoObj;
 	}
-
-	bool operator==(const ULandscapeLayerInfoObject*& rhs)
-	{
-		return LayerInfoObj == rhs;
-	}
 #endif // WITH_EDITORONLY_DATA
 };
 
