Prevent GetPlayerController, GetPlayerPawn, GetPlayerCharacter, GetPlayerCameraManager, CreatePlayer being called in ConstructionScript

diff --git a/Engine/Source/Runtime/Engine/Classes/Kismet/GameplayStatics.h b/Engine/Source/Runtime/Engine/Classes/Kismet/GameplayStatics.h
--- a/Engine/Source/Runtime/Engine/Classes/Kismet/GameplayStatics.h
+++ b/Engine/Source/Runtime/Engine/Classes/Kismet/GameplayStatics.h
@@ -79,26 +79,26 @@ class ENGINE_API UGameplayStatics : public UBlueprintFunctionLibrary
 	static class UGameInstance* GetGameInstance(UObject* WorldContextObject);
 
 	/** Returns the player controller at the specified player index */
-	UFUNCTION(BlueprintPure, Category="Game", meta=(HidePin="WorldContextObject", DefaultToSelf="WorldContextObject"))
+	UFUNCTION(BlueprintPure, Category="Game", meta=(HidePin="WorldContextObject", DefaultToSelf="WorldContextObject", UnsafeDuringActorConstruction="true"))
 	static class APlayerController* GetPlayerController(UObject* WorldContextObject, int32 PlayerIndex);
 
 	/** Returns the player pawn at the specified player index */
-	UFUNCTION(BlueprintPure, Category="Game", meta=(HidePin="WorldContextObject", DefaultToSelf="WorldContextObject"))
+	UFUNCTION(BlueprintPure, Category="Game", meta=(HidePin="WorldContextObject", DefaultToSelf="WorldContextObject", UnsafeDuringActorConstruction="true"))
 	static class APawn* GetPlayerPawn(UObject* WorldContextObject, int32 PlayerIndex);
 
 	/** Returns the player character (NULL if the player pawn doesn't exist OR is not a character) at the specified player index */
-	UFUNCTION(BlueprintPure, Category="Game", meta=(HidePin="WorldContextObject", DefaultToSelf="WorldContextObject"))
+	UFUNCTION(BlueprintPure, Category="Game", meta=(HidePin="WorldContextObject", DefaultToSelf="WorldContextObject", UnsafeDuringActorConstruction="true"))
 	static class ACharacter* GetPlayerCharacter(UObject* WorldContextObject, int32 PlayerIndex);
 
 	/** Returns the player's camera manager for the specified player index */
-	UFUNCTION(BlueprintPure, Category="Game", meta=(HidePin="WorldContextObject", DefaultToSelf="WorldContextObject"))
+	UFUNCTION(BlueprintPure, Category="Game", meta=(HidePin="WorldContextObject", DefaultToSelf="WorldContextObject", UnsafeDuringActorConstruction="true"))
 	static class APlayerCameraManager* GetPlayerCameraManager(UObject* WorldContextObject, int32 PlayerIndex);
 
 	/** Create a new player for this game.  
 	 *  @param ControllerId		The ID of the controller that the should control the newly created player.  A value of -1 specifies to use the next available ID
 	 *  @param bSpawnPawn		Whether a pawn should be spawned immediately. If false a pawn will not be created until transition to the next map.
 	 */
-	UFUNCTION(BlueprintCallable, Category="Game", meta=(HidePin="WorldContextObject", DefaultToSelf="WorldContextObject", AdvancedDisplay="2"))
+	UFUNCTION(BlueprintCallable, Category="Game", meta=(HidePin="WorldContextObject", DefaultToSelf="WorldContextObject", AdvancedDisplay="2", UnsafeDuringActorConstruction="true"))
 	static class APlayerController* CreatePlayer(UObject* WorldContextObject, int32 ControllerId = -1, bool bSpawnPawn = true);
 
 	// --- Level Streaming functions ------------------------
