[GitHub] GetActorsInSelectionRectangle now takes 'bIncludeNonCollidingComponents' option, init's SelectionRectangle and empties OutActors array
https://github.com/EpicGames/UnrealEngine/pull/370 (EverNewJoy)

diff --git a/Engine/Source/Runtime/Engine/Classes/GameFramework/HUD.h b/Engine/Source/Runtime/Engine/Classes/GameFramework/HUD.h
--- a/Engine/Source/Runtime/Engine/Classes/GameFramework/HUD.h
+++ b/Engine/Source/Runtime/Engine/Classes/GameFramework/HUD.h
@@ -419,12 +419,13 @@ public:
 	 * @param FirstPoint					The first point, or anchor of the marquee box. Where the dragging of the marquee started in screen space.
 	 * @param SecondPoint					The second point, where the mouse cursor currently is / the other point of the box selection, in screen space.
 	 * @return OutActors					The actors that are within the selection box according to selection rule
+	 * @param bIncludeNonCollidingComponents 	Whether to include even non-colliding components of the actor when determining its bounds
 	 * @param bActorMustBeFullyEnclosed  	The Selection rule: whether the selection box can partially intersect Actor, or must fully enclose the Actor.
 	 *
 	 * returns false if selection could not occur. Make sure template class is extending AActor.
 	 */
 	template <typename ClassFilter>
-	bool GetActorsInSelectionRectangle(const FVector2D& FirstPoint, const FVector2D& SecondPoint, TArray<ClassFilter*>& OutActors, bool bActorMustBeFullyEnclosed = false)
+	bool GetActorsInSelectionRectangle(const FVector2D& FirstPoint, const FVector2D& SecondPoint, TArray<ClassFilter*>& OutActors, bool bIncludeNonCollidingComponents = true, bool bActorMustBeFullyEnclosed = false)
 	{
 		//Is Actor subclass?
 		if (!ClassFilter::StaticClass()->IsChildOf(AActor::StaticClass()))
@@ -434,7 +435,7 @@ public:
 
 		//Run Inner Function, output to Base AActor Array
 		TArray<AActor*> OutActorsBaseArray;
-		GetActorsInSelectionRectangle(ClassFilter::StaticClass(), FirstPoint, SecondPoint, OutActorsBaseArray, bActorMustBeFullyEnclosed);
+		GetActorsInSelectionRectangle(ClassFilter::StaticClass(), FirstPoint, SecondPoint, OutActorsBaseArray, bIncludeNonCollidingComponents, bActorMustBeFullyEnclosed);
 
 		//Construct casted template type array
 		for (AActor* EachActor : OutActorsBaseArray)
@@ -457,11 +458,12 @@ public:
 	 * @param FirstPoint					The first point, or anchor of the marquee box. Where the dragging of the marquee started in screen space.
 	 * @param SecondPoint					The second point, where the mouse cursor currently is / the other point of the box selection, in screen space.
 	 * @return OutActors					The actors that are within the selection box according to selection rule
+	 * @param bIncludeNonCollidingComponents 	Whether to include even non-colliding components of the actor when determining its bounds
 	 * @param bActorMustBeFullyEnclosed  	The Selection rule: whether the selection box can partially intersect Actor, or must fully enclose the Actor.
 	 *
 	 */
 	UFUNCTION(BlueprintPure, Category=HUD)
-	void GetActorsInSelectionRectangle(TSubclassOf<class AActor> ClassFilter, const FVector2D& FirstPoint, const FVector2D& SecondPoint, TArray<AActor*>& OutActors, bool bActorMustBeFullyEnclosed = false);
+	void GetActorsInSelectionRectangle(TSubclassOf<class AActor> ClassFilter, const FVector2D& FirstPoint, const FVector2D& SecondPoint, TArray<AActor*>& OutActors, bool bIncludeNonCollidingComponents = true, bool bActorMustBeFullyEnclosed = false);
 
 
 	/**
diff --git a/Engine/Source/Runtime/Engine/Private/HUD.cpp b/Engine/Source/Runtime/Engine/Private/HUD.cpp
--- a/Engine/Source/Runtime/Engine/Private/HUD.cpp
+++ b/Engine/Source/Runtime/Engine/Private/HUD.cpp
@@ -644,10 +644,14 @@ void AHUD::Deproject(float ScreenX, float ScreenY, FVector& WorldPosition, FVect
 }
 
 
-void AHUD::GetActorsInSelectionRectangle(TSubclassOf<class AActor> ClassFilter, const FVector2D& FirstPoint, const FVector2D& SecondPoint, TArray<AActor*>& OutActors, bool bActorMustBeFullyEnclosed)
+void AHUD::GetActorsInSelectionRectangle(TSubclassOf<class AActor> ClassFilter, const FVector2D& FirstPoint, const FVector2D& SecondPoint, TArray<AActor*>& OutActors, bool bIncludeNonCollidingComponents, bool bActorMustBeFullyEnclosed)
 {
+	// Because this is a HUD function it is likely to get called each tick,
+	// so make sure any previous contents of the out actor array have been cleared!
+	OutActors.Empty();
+
 	//Create Selection Rectangle from Points
-	FBox2D SelectionRectangle;
+	FBox2D SelectionRectangle(0);
 
 	//This method ensures that an appropriate rectangle is generated, 
 	//		no matter what the coordinates of first and second point actually are.
@@ -676,7 +680,7 @@ void AHUD::GetActorsInSelectionRectangle(TSubclassOf<class AActor> ClassFilter,
 		AActor* EachActor = *Itr;
 
 		//Get Actor Bounds				//casting to base class, checked by template in the .h
-		const FBox EachActorBounds = Cast<AActor>(EachActor)->GetComponentsBoundingBox(true); /* All Components */
+		const FBox EachActorBounds = Cast<AActor>(EachActor)->GetComponentsBoundingBox(bIncludeNonCollidingComponents); /* All Components? */
 
 		//Center
 		const FVector BoxCenter = EachActorBounds.GetCenter();
