EditInline tag is removed. It is replaced with Instanced.

diff --git a/Engine/Plugins/Experimental/Paper2D/Source/Paper2D/Classes/PaperTileMap.h b/Engine/Plugins/Experimental/Paper2D/Source/Paper2D/Classes/PaperTileMap.h
--- a/Engine/Plugins/Experimental/Paper2D/Source/Paper2D/Classes/PaperTileMap.h
+++ b/Engine/Plugins/Experimental/Paper2D/Source/Paper2D/Classes/PaperTileMap.h
@@ -68,7 +68,7 @@ class PAPER2D_API UPaperTileMap : public UDataAsset //@TODO: Just to make it eas
 	UMaterialInterface* Material;
 
 	// The list of layers
-	UPROPERTY(Category=Debug, VisibleAnywhere, EditInline) //@TODO: Remove this
+	UPROPERTY(Category=Debug, VisibleAnywhere, Instanced) //@TODO: Remove this
 	TArray<class UPaperTileLayer*> TileLayers;
 
 	// Collision domain (no collision, 2D, or 3D)
@@ -86,7 +86,7 @@ class PAPER2D_API UPaperTileMap : public UDataAsset //@TODO: Just to make it eas
 public:
 #if WITH_EDITORONLY_DATA
 	/** Importing data and options used for this tile map */
-	UPROPERTY(Category=ImportSettings, VisibleAnywhere, EditInline)
+	UPROPERTY(Category=ImportSettings, VisibleAnywhere, Instanced)
 	class UAssetImportData* AssetImportData;
 #endif
 
diff --git a/Engine/Plugins/Experimental/Paper2D/Source/Paper2D/Classes/PaperTileMapActor.h b/Engine/Plugins/Experimental/Paper2D/Source/Paper2D/Classes/PaperTileMapActor.h
--- a/Engine/Plugins/Experimental/Paper2D/Source/Paper2D/Classes/PaperTileMapActor.h
+++ b/Engine/Plugins/Experimental/Paper2D/Source/Paper2D/Classes/PaperTileMapActor.h
@@ -11,7 +11,7 @@ class APaperTileMapActor : public AActor
 {
 	GENERATED_UCLASS_BODY()
 
-	UPROPERTY(Category=TileMapActor, VisibleAnywhere, EditInline)
+	UPROPERTY(Category=TileMapActor, VisibleAnywhere)
 	TSubobjectPtr<class UPaperTileMapRenderComponent> RenderComponent;
 
 	// AActor interface
diff --git a/Engine/Plugins/Experimental/Paper2D/Source/Paper2D/Classes/PaperTileMapRenderComponent.h b/Engine/Plugins/Experimental/Paper2D/Source/Paper2D/Classes/PaperTileMapRenderComponent.h
--- a/Engine/Plugins/Experimental/Paper2D/Source/Paper2D/Classes/PaperTileMapRenderComponent.h
+++ b/Engine/Plugins/Experimental/Paper2D/Source/Paper2D/Classes/PaperTileMapRenderComponent.h
@@ -39,7 +39,7 @@ private:
 
 public:
 	// The tile map used by this component
-	UPROPERTY(Category=Setup, EditAnywhere, EditInline, BlueprintReadOnly)
+	UPROPERTY(Category=Setup, EditAnywhere, Instanced, BlueprintReadOnly)
 	class UPaperTileMap* TileMap;
 
 protected:
diff --git a/Engine/Plugins/Experimental/Paper2D/Source/PaperJsonImporter/Private/PaperSpriteSheet.h b/Engine/Plugins/Experimental/Paper2D/Source/PaperJsonImporter/Private/PaperSpriteSheet.h
--- a/Engine/Plugins/Experimental/Paper2D/Source/PaperJsonImporter/Private/PaperSpriteSheet.h
+++ b/Engine/Plugins/Experimental/Paper2D/Source/PaperJsonImporter/Private/PaperSpriteSheet.h
@@ -26,7 +26,7 @@ public:
 
 #if WITH_EDITORONLY_DATA
 	// Import data for this 
-	UPROPERTY(VisibleAnywhere, EditInline, Category=ImportSettings)
+	UPROPERTY(VisibleAnywhere, Instanced, Category=ImportSettings)
 	class UAssetImportData* AssetImportData;
 #endif
 };
diff --git a/Engine/Source/Editor/UnrealEd/Classes/Factories/FbxImportUI.h b/Engine/Source/Editor/UnrealEd/Classes/Factories/FbxImportUI.h
--- a/Engine/Source/Editor/UnrealEd/Classes/Factories/FbxImportUI.h
+++ b/Engine/Source/Editor/UnrealEd/Classes/Factories/FbxImportUI.h
@@ -80,15 +80,15 @@ class UFbxImportUI : public UObject
 	uint32 bImportTextures:1;
 
 	/** Import data used when importing static meshes */
-	UPROPERTY(EditAnywhere, editinline, Category=StaticMesh)
+	UPROPERTY(EditAnywhere, Instanced, Category=StaticMesh)
 	class UFbxStaticMeshImportData* StaticMeshImportData;
 
 	/** Import data used when importing skeletal meshes */
-	UPROPERTY(EditAnywhere, editinline, Category=SkeletalMesh)
+	UPROPERTY(EditAnywhere, Instanced, Category=SkeletalMesh)
 	class UFbxSkeletalMeshImportData* SkeletalMeshImportData;
 
 	/** Import data used when importing animations */
-	UPROPERTY(EditAnywhere, editinline, Category=Animation)
+	UPROPERTY(EditAnywhere, Instanced, Category=Animation)
 	class UFbxAnimSequenceImportData* AnimSequenceImportData;
 
 	/** Type of asset to import from the FBX file */
@@ -96,7 +96,7 @@ class UFbxImportUI : public UObject
 	bool	bPreserveLocalTransform;
 
 	/** Import data used when importing textures */
-	UPROPERTY(EditAnywhere, editinline, Category=Textures)
+	UPROPERTY(EditAnywhere, Instanced, Category=Textures)
 	class UFbxTextureImportData* TextureImportData;
 
 	// Begin UObject Interface
diff --git a/Engine/Source/Editor/UnrealEd/Classes/Factories/TrueTypeFontFactory.h b/Engine/Source/Editor/UnrealEd/Classes/Factories/TrueTypeFontFactory.h
--- a/Engine/Source/Editor/UnrealEd/Classes/Factories/TrueTypeFontFactory.h
+++ b/Engine/Source/Editor/UnrealEd/Classes/Factories/TrueTypeFontFactory.h
@@ -13,7 +13,7 @@ class UTrueTypeFontFactory : public UFontFactory, public FReimportHandler
 	GENERATED_UCLASS_BODY()
 
 	/** Import options for this font */
-	UPROPERTY(EditAnywhere, editinline, Category=TrueTypeFontFactory, meta=(ToolTip="Import options for the font"))
+	UPROPERTY(EditAnywhere, Instanced, Category = TrueTypeFontFactory, meta = (ToolTip = "Import options for the font"))
 	class UFontImportOptions* ImportOptions;
 
 	/** True when the font dialog was shown for this factory during the non-legacy creation process */
diff --git a/Engine/Source/Editor/UnrealEd/Classes/MaterialEditor/MaterialEditorInstanceConstant.h b/Engine/Source/Editor/UnrealEd/Classes/MaterialEditor/MaterialEditorInstanceConstant.h
--- a/Engine/Source/Editor/UnrealEd/Classes/MaterialEditor/MaterialEditorInstanceConstant.h
+++ b/Engine/Source/Editor/UnrealEd/Classes/MaterialEditor/MaterialEditorInstanceConstant.h
@@ -18,7 +18,7 @@ struct FEditorParameterGroup
 	UPROPERTY()
 	FName GroupName;
 
-	UPROPERTY(EditAnywhere,  editfixedsize, editinline, Category=EditorParameterGroup)
+	UPROPERTY(EditAnywhere, editfixedsize, Instanced, Category=EditorParameterGroup)
 	TArray<class UDEditorParameterValue*> Parameters;
 
 };
diff --git a/Engine/Source/Programs/UnrealHeaderTool/Private/HeaderParser.cpp b/Engine/Source/Programs/UnrealHeaderTool/Private/HeaderParser.cpp
--- a/Engine/Source/Programs/UnrealHeaderTool/Private/HeaderParser.cpp
+++ b/Engine/Source/Programs/UnrealHeaderTool/Private/HeaderParser.cpp
@@ -2495,7 +2495,7 @@ bool FHeaderParser::GetVarType
 			}
 			else if (Specifier == TEXT("EditInline"))
 			{
-				Flags |= CPF_EditInline;
+				FError::Throwf(TEXT("EditInline is deprecated. Remove it, or use Instanced instead."));
 			}
 			else if (Specifier == TEXT("NoClear"))
 			{
@@ -3270,7 +3270,7 @@ bool FHeaderParser::GetVarType
 	// Perform some more specific validation on the property flags
 	if ((VarProperty.PropertyFlags & CPF_EditInline) && VarProperty.Type != CPT_ObjectReference)
 	{
-		FError::Throwf(TEXT("'EditInline' is only allowed on object property (or array of objects)"));
+		FError::Throwf(TEXT("'Instanced' is only allowed on object property (or array of objects)"));
 	}
 
 	if ( VarProperty.IsObject() && VarProperty.MetaClass == NULL && (VarProperty.PropertyFlags&CPF_Config) != 0 )
diff --git a/Engine/Source/Runtime/AIModule/Classes/BehaviorTree/BlackboardData.h b/Engine/Source/Runtime/AIModule/Classes/BehaviorTree/BlackboardData.h
--- a/Engine/Source/Runtime/AIModule/Classes/BehaviorTree/BlackboardData.h
+++ b/Engine/Source/Runtime/AIModule/Classes/BehaviorTree/BlackboardData.h
@@ -21,7 +21,7 @@ struct FBlackboardEntry
 #endif // WITH_EDITORONLY_DATA
 
 	/** key type and additional properties */
-	UPROPERTY(EditAnywhere, EditInline, Category=Blackboard)
+	UPROPERTY(EditAnywhere, Instanced, Category=Blackboard)
 	UBlackboardKeyType* KeyType;
 };
 
diff --git a/Engine/Source/Runtime/CoreUObject/Public/UObject/ObjectBase.h b/Engine/Source/Runtime/CoreUObject/Public/UObject/ObjectBase.h
--- a/Engine/Source/Runtime/CoreUObject/Public/UObject/ObjectBase.h
+++ b/Engine/Source/Runtime/CoreUObject/Public/UObject/ObjectBase.h
@@ -741,9 +741,6 @@ namespace UP
 		/// Object property can be exported with it's owner.
 		Export,
 
-		/// Edit this object reference inline in the editor.
-		EditInline,
-
 		/// Hide clear (and browse) button in the editor.
 		NoClear,
 
diff --git a/Engine/Source/Runtime/Engine/Classes/Animation/AnimSequence.h b/Engine/Source/Runtime/Engine/Classes/Animation/AnimSequence.h
--- a/Engine/Source/Runtime/Engine/Classes/Animation/AnimSequence.h
+++ b/Engine/Source/Runtime/Engine/Classes/Animation/AnimSequence.h
@@ -366,7 +366,7 @@ class UAnimSequence : public UAnimSequenceBase
 	 * The compression scheme that was most recently used to compress this animation.
 	 * May be NULL.
 	 */
-	UPROPERTY(EditInline, Category=Compression, VisibleAnywhere)
+	UPROPERTY(Instanced, Category=Compression, VisibleAnywhere)
 	class UAnimCompress* CompressionScheme;
 #endif // WITH_EDITORONLY_DATA
 
@@ -476,7 +476,7 @@ class UAnimSequence : public UAnimSequenceBase
 	uint32 bWasCompressedWithoutTranslations:1;
 
 	/** Importing data and options used for this mesh */
-	UPROPERTY(EditAnywhere, editinline, Category=Reimport)
+	UPROPERTY(EditAnywhere, Instanced, Category=Reimport)
 	class UAssetImportData* AssetImportData;
 
 	/***  for Reimport **/
diff --git a/Engine/Source/Runtime/Engine/Classes/Animation/AnimationAsset.h b/Engine/Source/Runtime/Engine/Classes/Animation/AnimationAsset.h
--- a/Engine/Source/Runtime/Engine/Classes/Animation/AnimationAsset.h
+++ b/Engine/Source/Runtime/Engine/Classes/Animation/AnimationAsset.h
@@ -417,7 +417,7 @@ public:
 
 #if WITH_EDITORONLY_DATA
 	/** Information for thumbnail rendering */
-	UPROPERTY(VisibleAnywhere, EditInline, Category = Thumbnail)
+	UPROPERTY(VisibleAnywhere, Instanced, Category = Thumbnail)
 	class UThumbnailInfo* ThumbnailInfo;
 #endif
 
diff --git a/Engine/Source/Runtime/Engine/Classes/Camera/PlayerCameraManager.h b/Engine/Source/Runtime/Engine/Classes/Camera/PlayerCameraManager.h
--- a/Engine/Source/Runtime/Engine/Classes/Camera/PlayerCameraManager.h
+++ b/Engine/Source/Runtime/Engine/Classes/Camera/PlayerCameraManager.h
@@ -261,7 +261,7 @@ public:
 	/////////////////////
 	
 	/** Camera modifier for cone-driven screen shakes */
-	UPROPERTY(EditAnywhere, BlueprintReadWrite, editinline, transient, Category=PlayerCameraManager)
+	UPROPERTY(EditAnywhere, BlueprintReadWrite, Instanced, transient, Category = PlayerCameraManager)
 	class UCameraModifier_CameraShake* CameraShakeCamMod;
 
 protected:
diff --git a/Engine/Source/Runtime/Engine/Classes/Engine/Blueprint.h b/Engine/Source/Runtime/Engine/Classes/Engine/Blueprint.h
--- a/Engine/Source/Runtime/Engine/Classes/Engine/Blueprint.h
+++ b/Engine/Source/Runtime/Engine/Classes/Engine/Blueprint.h
@@ -408,7 +408,7 @@ protected:
 	TWeakObjectPtr< class UWorld > CurrentWorldBeingDebugged;
 public:
 	/** Information for thumbnail rendering */
-	UPROPERTY(VisibleAnywhere, EditInline, Category=Thumbnail)
+	UPROPERTY(VisibleAnywhere, Instanced, Category=Thumbnail)
 	class UThumbnailInfo* ThumbnailInfo;
 
 	/** The blueprint is currently compiled */
diff --git a/Engine/Source/Runtime/Engine/Classes/Engine/Brush.h b/Engine/Source/Runtime/Engine/Classes/Engine/Brush.h
--- a/Engine/Source/Runtime/Engine/Classes/Engine/Brush.h
+++ b/Engine/Source/Runtime/Engine/Classes/Engine/Brush.h
@@ -113,7 +113,7 @@ class ABrush : public AActor
 	TSubobjectPtr<class UBrushComponent> BrushComponent;
 
 #if WITH_EDITORONLY_DATA
-	UPROPERTY(VisibleAnywhere, EditInline, Category=BrushBuilder)
+	UPROPERTY(VisibleAnywhere, Instanced, Category=BrushBuilder)
 	class UBrushBuilder* BrushBuilder;
 #endif
 
diff --git a/Engine/Source/Runtime/Engine/Classes/Engine/DestructibleMesh.h b/Engine/Source/Runtime/Engine/Classes/Engine/DestructibleMesh.h
--- a/Engine/Source/Runtime/Engine/Classes/Engine/DestructibleMesh.h
+++ b/Engine/Source/Runtime/Engine/Classes/Engine/DestructibleMesh.h
@@ -422,12 +422,12 @@ class UDestructibleMesh : public USkeletalMesh
 	TArray<struct FFractureEffect> FractureEffects;
 
 	/** Physics data.  Fields from BodySetup which are relevant to the DestructibleMesh will be used. */
-	UPROPERTY(EditAnywhere, editinline, Category=DestructibleMesh)
+	UPROPERTY(EditAnywhere, Instanced, Category=DestructibleMesh)
 	class UBodySetup* BodySetup;
 
 #if WITH_EDITORONLY_DATA
 	/** Information used to author an NxDestructibleAsset*/
-	UPROPERTY(instanced, editinline)
+	UPROPERTY(instanced)
 	class UDestructibleFractureSettings* FractureSettings;
 	
 	/** Static mesh this destructible mesh is created from. Is NULL if not created from a static mesh */
diff --git a/Engine/Source/Runtime/Engine/Classes/Engine/SkeletalMesh.h b/Engine/Source/Runtime/Engine/Classes/Engine/SkeletalMesh.h
--- a/Engine/Source/Runtime/Engine/Classes/Engine/SkeletalMesh.h
+++ b/Engine/Source/Runtime/Engine/Classes/Engine/SkeletalMesh.h
@@ -623,7 +623,7 @@ public:
 	class UPhysicsAsset* BoundsPreviewAsset_DEPRECATED;
 
 	/** Importing data and options used for this mesh */
-	UPROPERTY(EditAnywhere, editinline, Category=Reimport)
+	UPROPERTY(EditAnywhere, Instanced, Category = Reimport)
 	class UAssetImportData* AssetImportData;
 
 	/** Path to the resource used to construct this skeletal mesh */
@@ -635,7 +635,7 @@ public:
 	FString SourceFileTimestamp_DEPRECATED;
 
 	/** Information for thumbnail rendering */
-	UPROPERTY(VisibleAnywhere, EditInline, Category=Thumbnail)
+	UPROPERTY(VisibleAnywhere, Instanced, Category = Thumbnail)
 	class UThumbnailInfo* ThumbnailInfo;
 
 	/** Optimization settings used to simplify LODs of this mesh. */
diff --git a/Engine/Source/Runtime/Engine/Classes/Engine/StaticMesh.h b/Engine/Source/Runtime/Engine/Classes/Engine/StaticMesh.h
--- a/Engine/Source/Runtime/Engine/Classes/Engine/StaticMesh.h
+++ b/Engine/Source/Runtime/Engine/Classes/Engine/StaticMesh.h
@@ -360,7 +360,7 @@ class UStaticMesh : public UObject, public IInterface_CollisionDataProvider, pub
 	int32 LightMapCoordinateIndex;
 
 	// Physics data.
-	UPROPERTY(EditAnywhere, transient, duplicatetransient, editinline, Category=StaticMesh)
+	UPROPERTY(EditAnywhere, transient, duplicatetransient, Instanced, Category = StaticMesh)
 	class UBodySetup* BodySetup;
 
 	// Artist-accessible options.
@@ -404,11 +404,11 @@ class UStaticMesh : public UObject, public IInterface_CollisionDataProvider, pub
 
 #if WITH_EDITORONLY_DATA
 	/** Default settings when using this mesh for instanced foliage */
-	UPROPERTY(EditAnywhere, AdvancedDisplay, editinline, Category=StaticMesh)
+	UPROPERTY(EditAnywhere, AdvancedDisplay, Instanced, Category=StaticMesh)
 	UFoliageType_InstancedStaticMesh* FoliageDefaultSettings;
 
 	/** Importing data and options used for this mesh */
-	UPROPERTY(VisibleAnywhere, editinline, Category=ImportSettings)
+	UPROPERTY(VisibleAnywhere, Instanced, Category=ImportSettings)
 	class UAssetImportData* AssetImportData;
 
 	/** Path to the resource used to construct this static mesh */
@@ -420,7 +420,7 @@ class UStaticMesh : public UObject, public IInterface_CollisionDataProvider, pub
 	FString SourceFileTimestamp_DEPRECATED;
 
 	/** Information for thumbnail rendering */
-	UPROPERTY(VisibleAnywhere, EditInline, Category=Thumbnail)
+	UPROPERTY(VisibleAnywhere, Instanced, Category=Thumbnail)
 	class UThumbnailInfo* ThumbnailInfo;
 
 	/** The stored camera position to use as a default for the static mesh editor */
@@ -454,12 +454,12 @@ protected:
 	int32 ElementToIgnoreForTexFactor;
 
 	/** Array of user data stored with the asset */
-	UPROPERTY(EditAnywhere, AdvancedDisplay, editinline, Category=StaticMesh)
+	UPROPERTY(EditAnywhere, AdvancedDisplay, Instanced, Category = StaticMesh)
 	TArray<UAssetUserData*> AssetUserData;
 
 public:
 	/** Pre-build navigation collision */
-	UPROPERTY(VisibleAnywhere, transient, duplicatetransient, editinline, Category=Navigation)
+	UPROPERTY(VisibleAnywhere, transient, duplicatetransient, Instanced, Category = Navigation)
 	class UNavCollision* NavCollision;
 	
 public:
diff --git a/Engine/Source/Runtime/Engine/Classes/Engine/World.h b/Engine/Source/Runtime/Engine/Classes/Engine/World.h
--- a/Engine/Source/Runtime/Engine/Classes/Engine/World.h
+++ b/Engine/Source/Runtime/Engine/Classes/Engine/World.h
@@ -483,7 +483,7 @@ class ENGINE_API UWorld : public UObject, public FNetworkNotify
 	TArray<class AActor*> ActiveGroupActors;
 
 	/** Information for thumbnail rendering */
-	UPROPERTY(VisibleAnywhere, EditInline, Category=Thumbnail)
+	UPROPERTY(VisibleAnywhere, Instanced, Category=Thumbnail)
 	class UThumbnailInfo* ThumbnailInfo;
 #endif // WITH_EDITORONLY_DATA
 
diff --git a/Engine/Source/Runtime/Engine/Classes/GameFramework/Controller.h b/Engine/Source/Runtime/Engine/Classes/GameFramework/Controller.h
--- a/Engine/Source/Runtime/Engine/Classes/GameFramework/Controller.h
+++ b/Engine/Source/Runtime/Engine/Classes/GameFramework/Controller.h
@@ -51,7 +51,7 @@ private:
 
 public:
 	/** PlayerState containing replicated information about the player using this controller (only exists for players, not NPCs). */
-	UPROPERTY(editinline, replicatedUsing=OnRep_PlayerState, BlueprintReadOnly, Category="Controller")
+	UPROPERTY(replicatedUsing=OnRep_PlayerState, BlueprintReadOnly, Category="Controller")
 	class APlayerState* PlayerState;
 
 protected:
diff --git a/Engine/Source/Runtime/Engine/Classes/GameFramework/Pawn.h b/Engine/Source/Runtime/Engine/Classes/GameFramework/Pawn.h
--- a/Engine/Source/Runtime/Engine/Classes/GameFramework/Pawn.h
+++ b/Engine/Source/Runtime/Engine/Classes/GameFramework/Pawn.h
@@ -97,7 +97,7 @@ public:
 	TSubclassOf<AController>  AIControllerClass;
 
 	/** If Pawn is possessed by a player, points to his playerstate.  Needed for network play as controllers are not replicated to clients. */
-	UPROPERTY(editinline, replicatedUsing=OnRep_PlayerState, BlueprintReadOnly, Category="Pawn")
+	UPROPERTY(replicatedUsing=OnRep_PlayerState, BlueprintReadOnly, Category="Pawn")
 	APlayerState* PlayerState;
 
 	/** Replicated so we can see where remote clients are looking. */
@@ -109,7 +109,7 @@ public:
 	AController* LastHitBy;
 
 	/** Controller currently possessing this Actor */
-	UPROPERTY(editinline, replicatedUsing=OnRep_Controller)
+	UPROPERTY(replicatedUsing=OnRep_Controller)
 	AController* Controller;
 
 	/** Max difference between pawn's Rotation.Yaw and GetDesiredRotation().Yaw for pawn to be considered as having reached its desired rotation */
diff --git a/Engine/Source/Runtime/Engine/Classes/GameFramework/PlayerController.h b/Engine/Source/Runtime/Engine/Classes/GameFramework/PlayerController.h
--- a/Engine/Source/Runtime/Engine/Classes/GameFramework/PlayerController.h
+++ b/Engine/Source/Runtime/Engine/Classes/GameFramework/PlayerController.h
@@ -101,7 +101,7 @@ class ENGINE_API APlayerController : public AController
 	// Camera/view related variables
 
 	/** Camera manager associated with this Player Controller. */
-	UPROPERTY(EditInline, BlueprintReadOnly, Category=PlayerController)
+	UPROPERTY(BlueprintReadOnly, Category=PlayerController)
 	class APlayerCameraManager* PlayerCameraManager;
 
 	/** PlayerCamera class should be set for each game, otherwise Engine.PlayerCameraManager is used */
diff --git a/Engine/Source/Runtime/Engine/Classes/Materials/MaterialFunction.h b/Engine/Source/Runtime/Engine/Classes/Materials/MaterialFunction.h
--- a/Engine/Source/Runtime/Engine/Classes/Materials/MaterialFunction.h
+++ b/Engine/Source/Runtime/Engine/Classes/Materials/MaterialFunction.h
@@ -100,7 +100,7 @@ public:
 	uint32 CombinedOutputTypes;
 
 	/** Information for thumbnail rendering */
-	UPROPERTY(VisibleAnywhere, EditInline, Category=Thumbnail)
+	UPROPERTY(VisibleAnywhere, Instanced, Category = Thumbnail)
 	class UThumbnailInfo* ThumbnailInfo;
 #endif
 };
diff --git a/Engine/Source/Runtime/Engine/Classes/Materials/MaterialInterface.h b/Engine/Source/Runtime/Engine/Classes/Materials/MaterialInterface.h
--- a/Engine/Source/Runtime/Engine/Classes/Materials/MaterialInterface.h
+++ b/Engine/Source/Runtime/Engine/Classes/Materials/MaterialInterface.h
@@ -189,7 +189,7 @@ public:
 	FStringAssetReference PreviewMesh;
 
 	/** Information for thumbnail rendering */
-	UPROPERTY(VisibleAnywhere, EditInline, Category=Thumbnail)
+	UPROPERTY(VisibleAnywhere, Instanced, Category = Thumbnail)
 	class UThumbnailInfo* ThumbnailInfo;
 
 private:
diff --git a/Engine/Source/Runtime/Engine/Classes/Particles/Event/ParticleModuleEventGenerator.h b/Engine/Source/Runtime/Engine/Classes/Particles/Event/ParticleModuleEventGenerator.h
--- a/Engine/Source/Runtime/Engine/Classes/Particles/Event/ParticleModuleEventGenerator.h
+++ b/Engine/Source/Runtime/Engine/Classes/Particles/Event/ParticleModuleEventGenerator.h
@@ -45,7 +45,7 @@ struct FParticleEvent_GenerateInfo
 	FName CustomName;
 
 	/** The events we want to fire off when this event has been generated */
-	UPROPERTY(EditAnywhere, editinline, Category=ParticleEvent_GenerateInfo)
+	UPROPERTY(EditAnywhere, Instanced, Category = ParticleEvent_GenerateInfo)
 	TArray<class UParticleModuleEventSendToGame*> ParticleModuleEventsToSendToGame;
 
 	FParticleEvent_GenerateInfo()
diff --git a/Engine/Source/Runtime/Engine/Classes/Particles/ParticleSystemComponent.h b/Engine/Source/Runtime/Engine/Classes/Particles/ParticleSystemComponent.h
--- a/Engine/Source/Runtime/Engine/Classes/Particles/ParticleSystemComponent.h
+++ b/Engine/Source/Runtime/Engine/Classes/Particles/ParticleSystemComponent.h
@@ -400,7 +400,7 @@ public:
 	uint32 bIsViewRelevanceDirty:1;
 
 	/** Array of replay clips for this particle system component.  These are serialized to disk.  You really should never add anything to this in the editor.  It's exposed so that you can delete clips if you need to, but be careful when doing so! */
-	UPROPERTY(editinline)
+	UPROPERTY()
 	TArray<class UParticleSystemReplay*> ReplayClips;
 
 	/** Current particle 'replay state'.  This setting controls whether we're currently simulating/rendering particles normally, or whether we should capture or playback particle replay data instead. */
diff --git a/Engine/Source/Runtime/Engine/Classes/Sound/AmbientSound.h b/Engine/Source/Runtime/Engine/Classes/Sound/AmbientSound.h
--- a/Engine/Source/Runtime/Engine/Classes/Sound/AmbientSound.h
+++ b/Engine/Source/Runtime/Engine/Classes/Sound/AmbientSound.h
@@ -16,7 +16,7 @@ class AAmbientSound : public AActor
 	uint32 bAutoPlay_DEPRECATED:1;
 
 	/** Audio component to play */
-	UPROPERTY(EditInLine, Category=Sound, VisibleAnywhere, BlueprintReadOnly,meta=(ExposeFunctionCategories="Sound,Audio,Audio|Components|Audio"))
+	UPROPERTY(Category=Sound, VisibleAnywhere, BlueprintReadOnly,meta=(ExposeFunctionCategories="Sound,Audio,Audio|Components|Audio"))
 	TSubobjectPtr<class UAudioComponent> AudioComponent;
 	
 	/** Dummy sound node property to force instantiation of subobject.	*/
diff --git a/Engine/Source/Runtime/SlateCore/Public/Styling/SlateWidgetStyleAsset.h b/Engine/Source/Runtime/SlateCore/Public/Styling/SlateWidgetStyleAsset.h
--- a/Engine/Source/Runtime/SlateCore/Public/Styling/SlateWidgetStyleAsset.h
+++ b/Engine/Source/Runtime/SlateCore/Public/Styling/SlateWidgetStyleAsset.h
@@ -14,7 +14,7 @@ class SLATECORE_API USlateWidgetStyleAsset : public UObject
 		  
 public:  
 	/**  */
-	UPROPERTY(Category=Appearance, EditAnywhere, EditInline)
+	UPROPERTY(Category=Appearance, EditAnywhere, Instanced)
 	USlateWidgetStyleContainerBase* CustomStyle;
 
 	template< class WidgetStyleType >            
diff --git a/Engine/Source/Runtime/UMG/Public/Components/Widget.h b/Engine/Source/Runtime/UMG/Public/Components/Widget.h
--- a/Engine/Source/Runtime/UMG/Public/Components/Widget.h
+++ b/Engine/Source/Runtime/UMG/Public/Components/Widget.h
@@ -65,7 +65,7 @@ public:
 	/**
 	 * The parent slot of the UWidget.  Allows us to easily inline edit the layout controlling this widget.
 	 */
-	UPROPERTY(EditInline, EditAnywhere, BlueprintReadOnly, Category=Layout, meta=( ShowOnlyInnerProperties ))
+	UPROPERTY(Instanced, EditAnywhere, BlueprintReadOnly, Category = Layout, meta = (ShowOnlyInnerProperties))
 	UPanelSlot* Slot;
 
 	/** Sets whether this widget can be modified interactively by the user */
