Better support for Instanced tag. Objects can be defined in CDO. Subobjects are not removed when owner bp is recompiled, etc.. ttp#333611

diff --git a/Engine/Source/Editor/PropertyEditor/Private/UserInterface/PropertyEditor/SPropertyEditorEditInline.cpp b/Engine/Source/Editor/PropertyEditor/Private/UserInterface/PropertyEditor/SPropertyEditorEditInline.cpp
--- a/Engine/Source/Editor/PropertyEditor/Private/UserInterface/PropertyEditor/SPropertyEditorEditInline.cpp
+++ b/Engine/Source/Editor/PropertyEditor/Private/UserInterface/PropertyEditor/SPropertyEditorEditInline.cpp
@@ -130,32 +130,10 @@ void SPropertyEditorEditInline::GetDesiredWidth( float& OutMinDesiredWidth, floa
 
 bool SPropertyEditorEditInline::Supports( const FPropertyNode* InTreeNode, int32 InArrayIdx )
 {
-	if (InTreeNode->IsEditConst())
-	{
-		return false;
-	}
-
-	bool bCanEditInline = false;
-
-	const FObjectPropertyNode* ObjectPropertyNode = InTreeNode->FindObjectItemParent();
-	if (ObjectPropertyNode && (InTreeNode->HasNodeFlags(EPropertyNodeFlags::EditInline) != 0))
-	{
-		// Assume objects with the editinline flag can be editinline.  If any of the parent object is a default subobject then we cannot allow editinline
-		bCanEditInline = true;
-		for( int32 ObjectIndex = 0; ObjectIndex < ObjectPropertyNode->GetNumObjects(); ++ObjectIndex )
-		{
-			const UObject* Object = ObjectPropertyNode->GetUObject(ObjectIndex);
-
-			if( Object && Object->IsDefaultSubobject() && Object->IsTemplate() )
-			{
-				bCanEditInline = false;
-				break;
-			}
-		}
-
-	}
-
-	return bCanEditInline;
+	return InTreeNode 
+		&& !InTreeNode->IsEditConst()
+		&& InTreeNode->FindObjectItemParent()
+		&& InTreeNode->HasNodeFlags(EPropertyNodeFlags::EditInline);
 }
 
 bool SPropertyEditorEditInline::Supports( const TSharedRef< class FPropertyEditor >& InPropertyEditor )
@@ -222,7 +200,11 @@ void SPropertyEditorEditInline::OnClassPicked(UClass* InClass)
 			UObject*		Object = Itor->Get();
 			UObject*		UseOuter			= ( InClass->IsChildOf( UClass::StaticClass() ) ? Cast<UClass>(Object)->GetDefaultObject() : Object );
 			EObjectFlags	MaskedOuterFlags	= UseOuter ? UseOuter->GetMaskedFlags(RF_PropagateToSubObjects) : RF_NoFlags;
-			UObject*		NewObject			= StaticConstructObject( InClass, UseOuter, NAME_None, MaskedOuterFlags, NULL );
+			if (UseOuter && UseOuter->HasAnyFlags(RF_ClassDefaultObject | RF_ArchetypeObject))
+			{
+				MaskedOuterFlags |= RF_ArchetypeObject;
+			}
+			UObject*		NewObject = StaticConstructObject(InClass, UseOuter, NAME_None, MaskedOuterFlags, NULL);
 
 			NewValues.Add( NewObject->GetPathName() );
 		}
diff --git a/Engine/Source/Runtime/CoreUObject/Private/UObject/PropertyBaseObject.cpp b/Engine/Source/Runtime/CoreUObject/Private/UObject/PropertyBaseObject.cpp
--- a/Engine/Source/Runtime/CoreUObject/Private/UObject/PropertyBaseObject.cpp
+++ b/Engine/Source/Runtime/CoreUObject/Private/UObject/PropertyBaseObject.cpp
@@ -238,76 +238,12 @@ bool UObjectPropertyBase::ParseObjectPropertyValue( const UProperty* Property, U
 
 const TCHAR* UObjectPropertyBase::ImportText_Internal( const TCHAR* InBuffer, void* Data, int32 PortFlags, UObject* Parent, FOutputDevice* ErrorText ) const
 {
-	UObject* TemplateOwner = NULL;
-	if (PropertyFlags & CPF_InstancedReference)
-	{
-		for ( UObject* CheckOuter = Parent; CheckOuter; CheckOuter = CheckOuter->GetOuter() )
-		{
-			if ( CheckOuter->HasAnyFlags(RF_ClassDefaultObject) )
-			{
-				TemplateOwner = CheckOuter;
-				break;
-			}
-		}
-	}
 	const TCHAR* Buffer = InBuffer;
-	// if TemplateOwner != NULL, Data is pointing to a class default object and this is a component
-	if(TemplateOwner)
-	{
-		// For default property assignment, only allow references to components declared in the default properties of the class.
-		FString Temp;
-		Buffer = UPropertyHelpers::ReadToken( Buffer, Temp, 1 );
-		if( !Buffer )
-		{
-			return NULL;
-		}
-		if( Temp==TEXT("None") )
-		{
-			SetObjectPropertyValue(Data, NULL);
-		}
-		else
-		{
-			UObject* Result = NULL;
+	UObject* Result = NULL;
 
-			// Try to find the component through the parent class.
-			UClass*	OuterClass = TemplateOwner->GetClass();
-			while(OuterClass)
-			{
-				UObject* ComponentDefaultObject = OuterClass->GetDefaultSubobjectByName(FName(*Temp, FNAME_Find, true));
+	bool bOk = ParseObjectPropertyValue(this, Parent, PropertyClass, PortFlags, Buffer, Result);
 
-				if(ComponentDefaultObject && ComponentDefaultObject->IsA(PropertyClass))
-				{
-					Result = ComponentDefaultObject;
-					break;
-				}
-				OuterClass = NULL;
-			}
-
-			check(!Result || Result->IsA(PropertyClass));
-
-			SetObjectPropertyValue(Data, Result);
-
-			// If we couldn't find it or load it, we'll have to do without it.
-			if( !Result )
-			{
-				if (PortFlags & PPF_CheckReferences)
-				{
-					ErrorText->Logf(ELogVerbosity::Warning, TEXT("%s: unresolved reference to '%s'"), *GetFullName(), InBuffer);
-				}
-				return NULL;
-			}
-
-		}
-	}
-	else
-	{
-		// non-component path
-		UObject* Result = NULL;
-
-		bool bOk = ParseObjectPropertyValue(this, Parent, PropertyClass, PortFlags, Buffer, Result);
-
-		SetObjectPropertyValue(Data, Result);
-	}
+	SetObjectPropertyValue(Data, Result);
 	return Buffer;
 }
 
diff --git a/Engine/Source/Runtime/Engine/Private/UnrealEngine.cpp b/Engine/Source/Runtime/Engine/Private/UnrealEngine.cpp
--- a/Engine/Source/Runtime/Engine/Private/UnrealEngine.cpp
+++ b/Engine/Source/Runtime/Engine/Private/UnrealEngine.cpp
@@ -9582,7 +9582,7 @@ static TAutoConsoleVariable<int32> CVarDumpCopyPropertiesForUnrelatedObjects(
 	TEXT("Dump the objects that are cross class copied")
 	);
 
-struct FFindEditInlineSubobjectHelper
+struct FFindInstancedReferenceSubobjectHelper
 {
 	template<typename T>
 	static void Get(const UStruct* Struct, const uint8* ContainerAddress, T& OutObjects)
@@ -9590,10 +9590,8 @@ struct FFindEditInlineSubobjectHelper
 		check(ContainerAddress);
 		for (UProperty* Prop = (Struct ? Struct->RefLink : NULL); Prop; Prop = Prop->NextRef)
 		{
-			if (Prop->HasAnyPropertyFlags(CPF_EditInline))
+			if (Prop->HasAnyPropertyFlags(CPF_InstancedReference))
 			{
-				// Collect Objects referenced in EditInline properties
-
 				auto ObjectProperty = Cast<const UObjectProperty>(Prop);
 				if (ObjectProperty)
 				{
@@ -9624,10 +9622,8 @@ struct FFindEditInlineSubobjectHelper
 					}
 				}
 			}
-			else
+			else if (Prop->HasAnyPropertyFlags(CPF_ContainsInstancedReference))
 			{
-				// Look for EditInline properties in structs varaibles
-
 				auto StructProperty = Cast<const UStructProperty>(Prop);
 				if (StructProperty && StructProperty->Struct)
 				{
@@ -9795,18 +9791,19 @@ void UEngine::CopyPropertiesForUnrelatedObjects(UObject* OldObject, UObject* New
 		}
 	}
 
+	if (OldObject->GetClass()->HasAnyClassFlags(CLASS_HasInstancedReference) &&
+		NewObject->GetClass()->HasAnyClassFlags(CLASS_HasInstancedReference))
 	{
 		TSet<UObject*> OldEditInlineObjects;
-		FFindEditInlineSubobjectHelper::Get(OldObject->GetClass(), reinterpret_cast<uint8*>(OldObject), OldEditInlineObjects);
+		FFindInstancedReferenceSubobjectHelper::Get(OldObject->GetClass(), reinterpret_cast<uint8*>(OldObject), OldEditInlineObjects);
 		if (OldEditInlineObjects.Num())
 		{
 			TSet<UObject*> NewEditInlineObjects;
-			FFindEditInlineSubobjectHelper::Get(NewObject->GetClass(), reinterpret_cast<uint8*>(NewObject), NewEditInlineObjects);
-			// if object exist in both sets, it's outer is old object, and its class is "editinline"
+			FFindInstancedReferenceSubobjectHelper::Get(NewObject->GetClass(), reinterpret_cast<uint8*>(NewObject), NewEditInlineObjects);
 			for (auto Obj : NewEditInlineObjects)
 			{
 				const bool bProperOuter = (Obj->GetOuter() == OldObject);
-				const bool bEditInlineNew = Obj->GetClass()->HasAllClassFlags(CLASS_EditInlineNew);
+				const bool bEditInlineNew = Obj->GetClass()->HasAnyClassFlags(CLASS_EditInlineNew | CLASS_DefaultToInstanced);
 				if (bProperOuter && bEditInlineNew)
 				{
 					const bool bKeptByOld = OldEditInlineObjects.Contains(Obj);
