Better handling of child processes when compiling shaders.

- Shader compiler manager now uses handles instead of PIDs (PIDs have numerous problems - can come and go without us noticing, can be reused, making us wait for a wrong process, have to be wait()ed for on Unix).
- Also, sprinkled handle closing logic into SCM code. SCM should be refactored to treat created processes like opened files and do not hold on to their handles indefinitely, but current code works too.
- A method to dispose of process handles added (CloseProc(), matching CreateProc()).
- Linux specific: IsProcRunning() will now wait() if it detects that child is a zombie, to allow higher level code holding on to process handles of no longer running processes without taking kernel resources.

diff --git a/Engine/Source/Runtime/Core/Private/GenericPlatform/GenericPlatformProcess.cpp b/Engine/Source/Runtime/Core/Private/GenericPlatform/GenericPlatformProcess.cpp
--- a/Engine/Source/Runtime/Core/Private/GenericPlatform/GenericPlatformProcess.cpp
+++ b/Engine/Source/Runtime/Core/Private/GenericPlatform/GenericPlatformProcess.cpp
@@ -169,6 +169,16 @@ void FGenericPlatformProcess::WaitForProc( FProcHandle & ProcessHandle )
 	UE_LOG(LogHAL, Fatal, TEXT("FGenericPlatformProcess::WaitForProc not implemented on this platform"));
 }
 
+void FGenericPlatformProcess::CloseProc(FProcHandle & ProcessHandle)
+{
+	// make sure we don't leave running processes (this is especially important for Linux)
+	check(!FPlatformProcess::IsProcRunning(ProcessHandle));
+
+	// Wait() should return instantly if the process isn't running (which it shouldn't).
+	// However, this is important that we wait for it on Linux/Unix platforms so we don't leave zombies.
+	FPlatformProcess::WaitForProc(ProcessHandle);
+}
+
 void FGenericPlatformProcess::TerminateProc( FProcHandle & ProcessHandle, bool KillTree )
 {
 	UE_LOG(LogHAL, Fatal, TEXT("FGenericPlatformProcess::TerminateProc not implemented on this platform"));
diff --git a/Engine/Source/Runtime/Core/Private/Linux/LinuxPlatformProcess.cpp b/Engine/Source/Runtime/Core/Private/Linux/LinuxPlatformProcess.cpp
--- a/Engine/Source/Runtime/Core/Private/Linux/LinuxPlatformProcess.cpp
+++ b/Engine/Source/Runtime/Core/Private/Linux/LinuxPlatformProcess.cpp
@@ -583,13 +583,20 @@ bool FProcHandle::IsRunning()
 				}
 				else
 				{
-					// since we used WNOWAIT, we don't have to collect all the info 
-					// even if child is indeed a zombie
 					bIsRunning = ( SignalInfo.si_pid != Get() );
 					break;
 				}
 			}
 		}
+
+		// If child is a zombie, wait() immediately to free up kernel resources. Higher level code
+		// (e.g. shader compiling manager) can hold on to handle of no longer running process for longer,
+		// which is a dubious, but valid behavior. We don't want to keep zombie around though.
+		if (!bIsRunning)
+		{
+			UE_LOG(LogHAL, Log, TEXT("Child %d is no more running (zombie), Wait()ing immediately."), Get() );
+			Wait();
+		}
 	}
 
 	return bIsRunning;
diff --git a/Engine/Source/Runtime/Core/Public/GenericPlatform/GenericPlatformProcess.h b/Engine/Source/Runtime/Core/Public/GenericPlatform/GenericPlatformProcess.h
--- a/Engine/Source/Runtime/Core/Public/GenericPlatform/GenericPlatformProcess.h
+++ b/Engine/Source/Runtime/Core/Public/GenericPlatform/GenericPlatformProcess.h
@@ -302,22 +302,31 @@ struct CORE_API FGenericPlatformProcess
 	 * @return	The process handle for use in other process functions
 	 */
 	static FProcHandle CreateProc( const TCHAR* URL, const TCHAR* Parms, bool bLaunchDetached, bool bLaunchHidden, bool bLaunchReallyHidden, uint32* OutProcessID, int32 PriorityModifier, const TCHAR* OptionalWorkingDirectory, void* PipeWrite );
+
 	/** Returns true if the specified process is running 
-	*
-	* @param ProcessHandle handle returned from FPlatformProcess::CreateProc
-	* @return true if the process is still running
-	*/
+	 *
+	 * @param ProcessHandle handle returned from FPlatformProcess::CreateProc
+	 * @return true if the process is still running
+	 */
 	static bool IsProcRunning( FProcHandle & ProcessHandle );
+	
 	/** Waits for a process to stop
-	*
-	* @param ProcessHandle handle returned from FPlatformProcess::CreateProc
-	*/
+	 *
+	 * @param ProcessHandle handle returned from FPlatformProcess::CreateProc
+	 */
 	static void WaitForProc( FProcHandle & ProcessHandle );
+
+	/** Cleans up FProcHandle after we're done with it.
+	 *
+	 * @param ProcessHandle handle returned from FPlatformProcess::CreateProc.
+	 */
+	static void CloseProc( FProcHandle & ProcessHandle );
+
 	/** Terminates a process
-	*
-	* @param ProcessHandle handle returned from FPlatformProcess::CreateProc
-	* @param KillTree Whether the entire process tree should be terminated.
-	*/
+	 *
+	 * @param ProcessHandle handle returned from FPlatformProcess::CreateProc
+	 * @param KillTree Whether the entire process tree should be terminated.
+	 */
 	static void TerminateProc( FProcHandle & ProcessHandle, bool KillTree = false );
 	/** Retrieves the termination status of the specified process. **/
 	static bool GetProcReturnCode( FProcHandle & ProcHandle, int32* ReturnCode );
diff --git a/Engine/Source/Runtime/Engine/Private/ShaderCompiler/ShaderCompiler.cpp b/Engine/Source/Runtime/Engine/Private/ShaderCompiler/ShaderCompiler.cpp
--- a/Engine/Source/Runtime/Engine/Private/ShaderCompiler/ShaderCompiler.cpp
+++ b/Engine/Source/Runtime/Engine/Private/ShaderCompiler/ShaderCompiler.cpp
@@ -309,8 +309,8 @@ struct FPipeWorkerInfo
 /** Information tracked for each shader compile worker process instance. */
 struct FShaderCompileWorkerInfo
 {
-	/** Process Id of the worker app once launched.  0 is reserved to indicate no process. */
-	uint32 WorkerAppId;
+	/** Process handle of the worker app once launched.  Invalid handle means no process. */
+	FProcHandle WorkerProcess;
 
 	/** Tracks whether tasks have been issued to the worker. */
 	bool bIssuedTasksToWorker;	
@@ -335,7 +335,6 @@ struct FShaderCompileWorkerInfo
 	TArray<FShaderCompileJob*> QueuedJobs;
 
 	FShaderCompileWorkerInfo() :
-		WorkerAppId(0),
 		bIssuedTasksToWorker(false),		
 		bLaunchedWorker(false),
 		bComplete(false),
@@ -346,6 +345,12 @@ struct FShaderCompileWorkerInfo
 	{
 	}
 
+	// warning: not virtual
+	~FShaderCompileWorkerInfo()
+	{
+		check(!WorkerProcess.IsValid() || !FPlatformProcess::IsProcRunning(WorkerProcess));
+	}
+
 	void CreatePipeAndNewTask(uint32 WorkerIndex, uint32 ProcessId)
 	{
 #if PLATFORM_SUPPORTS_NAMED_PIPES
@@ -360,7 +365,7 @@ struct FShaderCompileWorkerInfo
 			if (bWorkerForPipeWasLaunched)
 			{
 				// Make sure it's alive
-				if (!FShaderCompilingManager::IsShaderCompilerWorkerRunning(WorkerAppId))
+				if (!FShaderCompilingManager::IsShaderCompilerWorkerRunning(WorkerProcess))
 				{
 					bAllocNameForPipe = true;
 					bWorkerForPipeWasLaunched = false;
@@ -710,12 +715,16 @@ void FShaderCompileThreadRunnable::LaunchWorkerIfNeeded(FShaderCompileWorkerInfo
 			LastCheckForWorkersTime = CurrentTime;
 		}
 
-		if (bCheckForWorkerRunning && !FShaderCompilingManager::IsShaderCompilerWorkerRunning(CurrentWorkerInfo.WorkerAppId))
+		if (bCheckForWorkerRunning && !FShaderCompilingManager::IsShaderCompilerWorkerRunning(CurrentWorkerInfo.WorkerProcess))
 		{
 			// Worker died, so clear this pipe and make a new one
 			CurrentWorkerInfo.PipeWorker.DestroyPipe();
 			CurrentWorkerInfo.bWorkerForPipeWasLaunched = false;
 
+			// clean up the proc handle
+			FPlatformProcess::CloseProc(CurrentWorkerInfo.WorkerProcess);
+			CurrentWorkerInfo.WorkerProcess = FProcHandle();
+
 			check(GShaderPipeConfig.bUseNamedPipes && !GShaderPipeConfig.bSingleJobPerNamedPipeProcess);
 			CurrentWorkerInfo.CreatePipeAndNewTask(WorkerIndex, GShaderCompilingManager->ProcessId);
 		}
@@ -727,7 +736,9 @@ void FShaderCompileThreadRunnable::LaunchWorkerIfNeeded(FShaderCompileWorkerInfo
 
 		// Store the Id with this thread so that we will know not to launch it again
 		const FString& PipeName = CurrentWorkerInfo.PipeWorker.NamedPipe.GetName();
-		CurrentWorkerInfo.WorkerAppId = Manager->LaunchWorker(WorkingDirectory, Manager->ProcessId, WorkerIndex, PipeName, PipeName, true, !GShaderPipeConfig.bReuseNamedPipeAndProcess);		
+		// make sure we don't overwrite a running process
+		check(!CurrentWorkerInfo.WorkerProcess.IsValid());
+		CurrentWorkerInfo.WorkerProcess = Manager->LaunchWorker(WorkingDirectory, Manager->ProcessId, WorkerIndex, PipeName, PipeName, true, !GShaderPipeConfig.bReuseNamedPipeAndProcess);		
 		CurrentWorkerInfo.bLaunchedWorker = true;
 		CurrentWorkerInfo.bWorkerForPipeWasLaunched = true;
 	}
@@ -754,6 +765,12 @@ bool FShaderCompileThreadRunnable::LaunchWorkersIfNeeded()
 		if (CurrentWorkerInfo.QueuedJobs.Num() == 0 )
 		{
 			// Skip if nothing to do
+			// Also, use the opportunity to free OS resources by cleaning up handles of no more running processes
+			if (CurrentWorkerInfo.WorkerProcess.IsValid() && !FShaderCompilingManager::IsShaderCompilerWorkerRunning(CurrentWorkerInfo.WorkerProcess))
+			{
+				FPlatformProcess::CloseProc(CurrentWorkerInfo.WorkerProcess);
+				CurrentWorkerInfo.WorkerProcess = FProcHandle();
+			}
 			continue;
 		}
 
@@ -765,32 +782,37 @@ bool FShaderCompileThreadRunnable::LaunchWorkersIfNeeded()
 		else
 #endif	// PLATFORM_SUPPORTS_NAMED_PIPES
 		{
-			if (CurrentWorkerInfo.WorkerAppId == 0 || (bCheckForWorkerRunning && !FShaderCompilingManager::IsShaderCompilerWorkerRunning(CurrentWorkerInfo.WorkerAppId)))
+			if (!CurrentWorkerInfo.WorkerProcess.IsValid() || (bCheckForWorkerRunning && !FShaderCompilingManager::IsShaderCompilerWorkerRunning(CurrentWorkerInfo.WorkerProcess)))
 			{
+				// @TODO: dubious design - worker should not be launched unless we know there's more work to do.
 				bool bLaunchAgain = true;
 
 				// Detect when the worker has exited due to fatal error
 				// bLaunchedWorker check here is necessary to distinguish between 'process isn't running because it crashed' and 'process isn't running because it exited cleanly and the outputfile was already consumed'
-				if (CurrentWorkerInfo.WorkerAppId != 0 && CurrentWorkerInfo.bLaunchedWorker)
+				if (CurrentWorkerInfo.WorkerProcess.IsValid())
 				{
-					const FString WorkingDirectory = Manager->AbsoluteShaderBaseWorkingDirectory + FString::FromInt(WorkerIndex) + TEXT("/");
-					const FString OutputFileNameAndPath = WorkingDirectory + TEXT("WorkerOutputOnly.out");
+					// shader compiler exited one way or another, so clear out the stale PID.
+					FPlatformProcess::CloseProc(CurrentWorkerInfo.WorkerProcess);
+					CurrentWorkerInfo.WorkerProcess = FProcHandle();
 
-					if (FPlatformFileManager::Get().GetPlatformFile().FileExists(*OutputFileNameAndPath))
+					if (CurrentWorkerInfo.bLaunchedWorker)
 					{
-						// If the worker is no longer running but it successfully wrote out the output, no need to assert
-						bLaunchAgain = false;						
-					}
-					else
-					{
-						UE_LOG(LogShaderCompilers, Error, TEXT("ShaderCompileWorker terminated unexpectedly!  Falling back to directly compiling which will be very slow.  Thread %u."), WorkerIndex);
+						const FString WorkingDirectory = Manager->AbsoluteShaderBaseWorkingDirectory + FString::FromInt(WorkerIndex) + TEXT("/");
+						const FString OutputFileNameAndPath = WorkingDirectory + TEXT("WorkerOutputOnly.out");
 
-						bAbandonWorkers = true;
-						break;
-					}
+						if (FPlatformFileManager::Get().GetPlatformFile().FileExists(*OutputFileNameAndPath))
+						{
+							// If the worker is no longer running but it successfully wrote out the output, no need to assert
+							bLaunchAgain = false;
+						}
+						else
+						{
+							UE_LOG(LogShaderCompilers, Error, TEXT("ShaderCompileWorker terminated unexpectedly!  Falling back to directly compiling which will be very slow.  Thread %u."), WorkerIndex);
 
-					// shader compiler exited one way or another, so clear out the stale PID.
-					CurrentWorkerInfo.WorkerAppId = 0;
+							bAbandonWorkers = true;
+							break;
+						}
+					}
 				}
 
 				if (bLaunchAgain)
@@ -799,8 +821,8 @@ bool FShaderCompileThreadRunnable::LaunchWorkersIfNeeded()
 					FString InputFileName(TEXT("WorkerInputOnly.in"));
 					FString OutputFileName(TEXT("WorkerOutputOnly.out"));
 
-					// Store the Id with this thread so that we will know not to launch it again
-					CurrentWorkerInfo.WorkerAppId = Manager->LaunchWorker(WorkingDirectory, Manager->ProcessId, WorkerIndex, InputFileName, OutputFileName, false, false);					
+					// Store the handle with this thread so that we will know not to launch it again
+					CurrentWorkerInfo.WorkerProcess = Manager->LaunchWorker(WorkingDirectory, Manager->ProcessId, WorkerIndex, InputFileName, OutputFileName, false, false);					
 					CurrentWorkerInfo.bLaunchedWorker = true;
 				}
 			}
@@ -1131,8 +1153,8 @@ void FShaderCompilingManager::AddJobs(TArray<FShaderCompileJob*>& NewJobs, bool
 	}
 }
 
-/** Launches the worker, returns the launched Process Id. */
-uint32 FShaderCompilingManager::LaunchWorker(const FString& WorkingDirectory, uint32 InProcessId, uint32 ThreadId, const FString& WorkerInputFile, const FString& WorkerOutputFile, bool bUseNamedPipes, bool bSingleConnectionPipe)
+/** Launches the worker, returns the launched process handle. */
+FProcHandle FShaderCompilingManager::LaunchWorker(const FString& WorkingDirectory, uint32 InProcessId, uint32 ThreadId, const FString& WorkerInputFile, const FString& WorkerOutputFile, bool bUseNamedPipes, bool bSingleConnectionPipe)
 {
 	// Setup the parameters that the worker application needs
 	// Surround the working directory with double quotes because it may contain a space 
@@ -1163,7 +1185,7 @@ uint32 FShaderCompilingManager::LaunchWorker(const FString& WorkingDirectory, ui
 	// Note: Set breakpoint here and launch the shadercompileworker with WorkerParameters a cmd-line
 	const TCHAR* WorkerParametersText = *WorkerParameters;
 	FPlatformMisc::LowLevelOutputDebugStringf(TEXT("Launching shader compile worker w/ WorkerParameters\n\t%s\n"), WorkerParametersText);
-	return 17;
+	return FProcHandle(17);
 #else
 #if UE_BUILD_DEBUG && PLATFORM_LINUX
 	FPlatformMisc::LowLevelOutputDebugStringf(TEXT("Launching shader compile worker:\n\t%s\n"), *WorkerParameters);
@@ -1176,10 +1198,7 @@ uint32 FShaderCompilingManager::LaunchWorker(const FString& WorkingDirectory, ui
 		UE_LOG(LogShaderCompilers, Fatal, TEXT( "Couldn't launch %s! Make sure the file is in your binaries folder." ), *ShaderCompileWorkerName );
 	}
 
-	// Close the handle as we no longer need it.
-	WorkerHandle.Close();
-
-	return WorkerId;
+	return WorkerHandle;
 #endif
 }
 
@@ -1965,17 +1984,9 @@ void FShaderCompilingManager::ProcessAsyncResults(bool bLimitExecutionTime, bool
 	}
 }
 
-bool FShaderCompilingManager::IsShaderCompilerWorkerRunning(uint32 ProcessId)
+bool FShaderCompilingManager::IsShaderCompilerWorkerRunning(FProcHandle & WorkerHandle)
 {
-	bool bRunning = false;
-	// make sure the PID we're inspecting is actually a shader compiler worker.
-	FString ProcessName = FPlatformProcess::GetApplicationName(ProcessId);
-
-	if (ProcessName.Find(FString(TEXT("ShaderCompileWorker"))) != -1)
-	{
-		bRunning = FPlatformProcess::IsApplicationRunning(ProcessId);
-	}
-	return bRunning;
+	return FPlatformProcess::IsProcRunning(WorkerHandle);
 }
 
 /** Enqueues a shader compile job with GShaderCompilingManager. */
diff --git a/Engine/Source/Runtime/Engine/Public/ShaderCompiler.h b/Engine/Source/Runtime/Engine/Public/ShaderCompiler.h
--- a/Engine/Source/Runtime/Engine/Public/ShaderCompiler.h
+++ b/Engine/Source/Runtime/Engine/Public/ShaderCompiler.h
@@ -213,8 +213,8 @@ private:
 	 */
 	double WorkersBusyTime;
 
-	/** Launches the worker, returns the launched Process Id. */
-	uint32 LaunchWorker(const FString& WorkingDirectory, uint32 ProcessId, uint32 ThreadId, const FString& WorkerInputFile, const FString& WorkerOutputFile, bool bUseNamedPipes, bool bSingleConnectionPipe);
+	/** Launches the worker, returns the launched process handle. */
+	FProcHandle LaunchWorker(const FString& WorkingDirectory, uint32 ProcessId, uint32 ThreadId, const FString& WorkerInputFile, const FString& WorkerOutputFile, bool bUseNamedPipes, bool bSingleConnectionPipe);
 
 	/** Blocks on completion of the given shader maps. */
 	void BlockOnShaderMapCompletion(const TArray<int32>& ShaderMapIdsToFinishCompiling, TMap<int32, FShaderMapFinalizeResults>& CompiledShaderMaps);
@@ -311,9 +311,9 @@ public:
 	ENGINE_API void ProcessAsyncResults(bool bLimitExecutionTime, bool bBlockOnGlobalShaderCompletion);
 
 	/**
-	 * Returns true if the given PID is actually a ShaderCompilerWorker, and it's still running.
+	 * Returns true if the given shader compile worker is still running.
 	 */
-	static bool IsShaderCompilerWorkerRunning(uint32 ProcessId);
+	static bool IsShaderCompilerWorkerRunning(FProcHandle & WorkerHandle);
 };
 
 /** The global shader compiling thread manager. */
