Replaced crude csv parsing with FCsvParser.

This addresses TTP#343109 - LIVE: EDITOR: Importing String fields containing commas from .csv (comma separated) files fails

diff --git a/Engine/Source/Runtime/Engine/Classes/Curves/CurveBase.h b/Engine/Source/Runtime/Engine/Classes/Curves/CurveBase.h
--- a/Engine/Source/Runtime/Engine/Classes/Curves/CurveBase.h
+++ b/Engine/Source/Runtime/Engine/Classes/Curves/CurveBase.h
@@ -475,12 +475,15 @@ public:
 	// End FCurveOwnerInterface
 
 	// Begin UCurveBase interface
+
+#if WITH_EDITOR
 	/** 
 	 *	Create curve from CSV style comma-separated string. 
 	 *	@return	Set of problems encountered while processing input
 	 */
 	TArray<FString> CreateCurveFromCSVString(const FString& InString);
-	
+#endif
+
 	/** Reset all curve data */
 	void ResetCurve();
 
diff --git a/Engine/Source/Runtime/Engine/Classes/Engine/CurveTable.h b/Engine/Source/Runtime/Engine/Classes/Engine/CurveTable.h
--- a/Engine/Source/Runtime/Engine/Classes/Engine/CurveTable.h
+++ b/Engine/Source/Runtime/Engine/Classes/Engine/CurveTable.h
@@ -61,6 +61,7 @@ class UCurveTable : public UObject
 	/** Output entire contents of table as JSON */
 	ENGINE_API bool WriteTableAsJSON(const TSharedRef< TJsonWriter<TCHAR, TPrettyJsonPrintPolicy<TCHAR> > >& JsonWriter) const;
 
+#if WITH_EDITOR
 	/** 
 	 *	Create table from CSV style comma-separated string. 
 	 *	RowCurve must be defined before calling this function. 
@@ -69,6 +70,7 @@ class UCurveTable : public UObject
 	 *	@return	Set of problems encountered while processing input
 	 */
 	ENGINE_API TArray<FString> CreateTableFromCSVString(const FString& InString, ERichCurveInterpMode InterpMode = RCIM_Linear);
+#endif
 
 	/** Empty the table info (will not clear RowCurve) */
 	ENGINE_API void EmptyTable();
diff --git a/Engine/Source/Runtime/Engine/Classes/Engine/DataTable.h b/Engine/Source/Runtime/Engine/Classes/Engine/DataTable.h
--- a/Engine/Source/Runtime/Engine/Classes/Engine/DataTable.h
+++ b/Engine/Source/Runtime/Engine/Classes/Engine/DataTable.h
@@ -155,7 +155,7 @@ public:
 	/** Get array of all the column titles */
 	ENGINE_API TArray<FString> GetColumnTitles() const;
 
-	TArray<UProperty*> GetTablePropertyArray(const FString& FirstRowString, UStruct* RowStruct, TArray<FString>& OutProblems);
+	TArray<UProperty*> GetTablePropertyArray(const TArray<const TCHAR*>& Cells, UStruct* RowStruct, TArray<FString>& OutProblems);
 
 	/** Get array for each row in the table. The first row is the titles*/
 	ENGINE_API TArray< TArray<FString> > GetTableData() const;
diff --git a/Engine/Source/Runtime/Engine/Engine.Build.cs b/Engine/Source/Runtime/Engine/Engine.Build.cs
--- a/Engine/Source/Runtime/Engine/Engine.Build.cs
+++ b/Engine/Source/Runtime/Engine/Engine.Build.cs
@@ -66,6 +66,11 @@ public class Engine : ModuleRules
 				"VectorVM",
 			}
         );
+        
+		if (Target.Configuration != UnrealTargetConfiguration.Shipping)
+		{
+            PrivateDependencyModuleNames.AddRange(new string[] { "CsvParser" });
+		}
 
         CircularlyReferencedDependentModules.Add("AIModule");
 
diff --git a/Engine/Source/Runtime/Engine/Private/CurveBase.cpp b/Engine/Source/Runtime/Engine/Private/CurveBase.cpp
--- a/Engine/Source/Runtime/Engine/Private/CurveBase.cpp
+++ b/Engine/Source/Runtime/Engine/Private/CurveBase.cpp
@@ -868,6 +868,9 @@ void UCurveBase::ResetCurve()
 	}
 }
 
+#if WITH_EDITOR
+#include "CsvParser.h"
+
 TArray<FString> UCurveBase::CreateCurveFromCSVString(const FString& InString)
 {
 	// Array used to store problems about curve import
@@ -876,11 +879,10 @@ TArray<FString> UCurveBase::CreateCurveFromCSVString(const FString& InString)
 	TArray<FRichCurveEditInfo> Curves = GetCurves();
 	const int32 NumCurves = Curves.Num();
 
-	// Split one giant string into one string per row
-	TArray<FString> RowStrings;
-	InString.ParseIntoArray(&RowStrings, TEXT("\r\n"), true);
+	const FCsvParser Parser(InString);
+	const FCsvParser::FRows& Rows = Parser.GetRows();
 
-	if(RowStrings.Num() == 0)
+	if(Rows.Num() == 0)
 	{
 		OutProblems.Add(FString(TEXT("No data.")));
 		return OutProblems;
@@ -890,11 +892,10 @@ TArray<FString> UCurveBase::CreateCurveFromCSVString(const FString& InString)
 	ResetCurve();
 
 	// Each row represents a point
-	for(int32 RowIdx=0; RowIdx<RowStrings.Num(); RowIdx++)
+	for(int32 RowIdx=0; RowIdx<Rows.Num(); RowIdx++)
 	{
-		TArray<FString> CellStrings;
-		RowStrings[RowIdx].ParseIntoArray(&CellStrings, TEXT(","), false);
-		const int32 NumCells = CellStrings.Num();
+		const TArray<const TCHAR*>& Cells = Rows[RowIdx];
+		const int32 NumCells = Cells.Num();
 
 		// Need at least two cell, Time and one Value
 		if(NumCells < 2)
@@ -903,13 +904,13 @@ TArray<FString> UCurveBase::CreateCurveFromCSVString(const FString& InString)
 			continue;
 		}
 
-		float Time = FCString::Atof(*CellStrings[0]);
+		float Time = FCString::Atof(Cells[0]);
 		for(int32 CellIdx=1; CellIdx<NumCells && CellIdx<(NumCurves+1); CellIdx++)
 		{
 			FRichCurve* Curve = Curves[CellIdx-1].CurveToEdit;
 			if(Curve != NULL)
 			{
-				FKeyHandle KeyHandle = Curve->AddKey(Time, FCString::Atof(*CellStrings[CellIdx]));
+				FKeyHandle KeyHandle = Curve->AddKey(Time, FCString::Atof(Cells[CellIdx]));
 				Curve->SetKeyInterpMode(KeyHandle, RCIM_Linear);
 			}
 		}
@@ -930,6 +931,8 @@ TArray<FString> UCurveBase::CreateCurveFromCSVString(const FString& InString)
 	return OutProblems;
 }
 
+#endif	// WITH_EDITOR
+
 //////////////////////////////////////////////////////////////////////////
 
 int32 FIntegralCurve::GetNumKeys() const
diff --git a/Engine/Source/Runtime/Engine/Private/CurveTable.cpp b/Engine/Source/Runtime/Engine/Private/CurveTable.cpp
--- a/Engine/Source/Runtime/Engine/Private/CurveTable.cpp
+++ b/Engine/Source/Runtime/Engine/Private/CurveTable.cpp
@@ -216,34 +216,32 @@ void UCurveTable::EmptyTable()
 
 
 /** */
-void GetCurveValues(const FString& RowString, TArray<float>* Values)
+void GetCurveValues(const TArray<const TCHAR*>& Cells, TArray<float>* Values)
 {
-	// Find the x values from he first row
-	TArray<FString> ValuesAsStrings;
-	RowString.ParseIntoArray(&ValuesAsStrings, TEXT(","), false);
-
 	// Need at least 2 columns, first column is skipped, will contain row names
-	if(ValuesAsStrings.Num() >= 2)
+	if(Cells.Num() >= 2)
 	{
 		// first element always NULL - as first column is row names
-		for(int32 ColIdx = 1; ColIdx < ValuesAsStrings.Num(); ColIdx++)
+		for(int32 ColIdx = 1; ColIdx < Cells.Num(); ColIdx++)
 		{
-			Values->Add(FCString::Atof(*ValuesAsStrings[ColIdx]));
+			Values->Add(FCString::Atof(Cells[ColIdx]));
 		}
 	}
 }
 
+#if WITH_EDITOR
+#include "CsvParser.h"
+
 TArray<FString> UCurveTable::CreateTableFromCSVString(const FString& InString, ERichCurveInterpMode InterpMode)
 {
 	// Array used to store problems about table creation
 	TArray<FString> OutProblems;
 
-	// Split one giant string into one string per row
-	TArray<FString> RowStrings;
-	InString.ParseIntoArray(&RowStrings, TEXT("\r\n"), true);
+	const FCsvParser Parser(InString);
+	const FCsvParser::FRows& Rows = Parser.GetRows();
 
 	// Must have at least 2 rows (x values + y values for at least one row)
-	if(RowStrings.Num() <= 1)
+	if(Rows.Num() <= 1)
 	{
 		OutProblems.Add(FString(TEXT("Too few rows.")));
 		return OutProblems;
@@ -253,23 +251,22 @@ TArray<FString> UCurveTable::CreateTableFromCSVString(const FString& InString, E
 	EmptyTable();
 
 	TArray<float> XValues;
-	GetCurveValues(RowStrings[0], &XValues);
+	GetCurveValues(Rows[0], &XValues);
 
 	// Iterate over rows
-	for(int32 RowIdx = 1; RowIdx < RowStrings.Num(); RowIdx++)
+	for(int32 RowIdx = 1; RowIdx < Rows.Num(); RowIdx++)
 	{
-		TArray<FString> CellStrings;
-		RowStrings[RowIdx].ParseIntoArray(&CellStrings, TEXT(","), false);
+		const TArray<const TCHAR*>& Row = Rows[RowIdx];
 
 		// Need at least 1 cells (row name)
-		if(CellStrings.Num() < 1)
+		if(Row.Num() < 1)
 		{
 			OutProblems.Add(FString::Printf(TEXT("Row '%d' has too few cells."), RowIdx));
 			continue;
 		}
 
 		// Get row name
-		FName RowName = MakeValidName(CellStrings[0]);
+		FName RowName = MakeValidName(Row[0]);
 
 		// Check its not 'none'
 		if(RowName == NAME_None)
@@ -286,7 +283,7 @@ TArray<FString> UCurveTable::CreateTableFromCSVString(const FString& InString, E
 		}
 
 		TArray<float> YValues;
-		GetCurveValues(RowStrings[RowIdx], &YValues);
+		GetCurveValues(Row, &YValues);
 
 		if(XValues.Num() != YValues.Num())
 		{
@@ -309,6 +306,8 @@ TArray<FString> UCurveTable::CreateTableFromCSVString(const FString& InString, E
 	return OutProblems;
 }
 
+#endif	// WITH_EDITOR
+
 //////////////////////////////////////////////////////////////////////////
 
 
diff --git a/Engine/Source/Runtime/Engine/Private/DataTable.cpp b/Engine/Source/Runtime/Engine/Private/DataTable.cpp
--- a/Engine/Source/Runtime/Engine/Private/DataTable.cpp
+++ b/Engine/Source/Runtime/Engine/Private/DataTable.cpp
@@ -235,6 +235,8 @@ UProperty* UDataTable::FindTableProperty(const FName& PropertyName) const
 
 #if WITH_EDITOR || HACK_HEADER_GENERATOR
 
+#include "CsvParser.h"
+
 struct FPropertyDisplayNameHelper
 {
 	static inline FString Get(const UProperty* Prop, const FString& DefaultName)
@@ -387,27 +389,25 @@ bool UDataTable::WriteTableAsJSON(const TSharedRef< TJsonWriter<TCHAR, TPrettyJs
 }
 
 /** Get array of UProperties that corresponds to columns in the table */
-TArray<UProperty*> UDataTable::GetTablePropertyArray(const FString& FirstRowString, UStruct* InRowStruct, TArray<FString>& OutProblems)
+TArray<UProperty*> UDataTable::GetTablePropertyArray(const TArray<const TCHAR*>& Cells, UStruct* InRowStruct, TArray<FString>& OutProblems)
 {
 	TArray<UProperty*> ColumnProps;
 
 	// Get list of all expected properties from the struct
 	TArray<FName> ExpectedPropNames = GetStructPropertyNames(InRowStruct);
 
-	// Find the column names from first row
-	TArray<FString> ColumnNameStrings;
-	FirstRowString.ParseIntoArray(&ColumnNameStrings, TEXT(","), false);
-
 	// Need at least 2 columns, first column is skipped, will contain row names
-	if(ColumnNameStrings.Num() > 0)
+	if(Cells.Num() > 1)
 	{
-		ColumnProps.AddZeroed( ColumnNameStrings.Num() );
+		ColumnProps.AddZeroed( Cells.Num() );
 
 		// first element always NULL - as first column is row names
 
-		for(int32 ColIdx=1; ColIdx<ColumnNameStrings.Num(); ColIdx++)
+		for (int32 ColIdx = 1; ColIdx < Cells.Num(); ++ColIdx)
 		{
-			FName PropName = MakeValidName(ColumnNameStrings[ColIdx]);
+			const TCHAR* ColumnValue = Cells[ColIdx];
+
+			FName PropName = MakeValidName(ColumnValue);
 			if(PropName == NAME_None)
 			{
 				OutProblems.Add(FString(TEXT("Missing name for column %d."), ColIdx));
@@ -419,7 +419,7 @@ TArray<UProperty*> UDataTable::GetTablePropertyArray(const FString& FirstRowStri
 				for (TFieldIterator<UProperty> It(InRowStruct); It && !ColumnProp; ++It)
 				{
 					const auto DisplayName = FPropertyDisplayNameHelper::Get(*It, FString());
-					ColumnProp = (!DisplayName.IsEmpty() && (DisplayName == ColumnNameStrings[ColIdx])) ? *It : NULL;
+					ColumnProp = (!DisplayName.IsEmpty() && DisplayName == ColumnValue) ? *It : NULL;
 				}
 
 				// Didn't find a property with this name, problem..
@@ -475,46 +475,49 @@ TArray<FString> UDataTable::CreateTableFromCSVString(const FString& InString)
 		OutProblems.Add(FString(TEXT("No RowStruct specified.")));
 		return OutProblems;
 	}
+	if (InString.IsEmpty())
+	{
+		OutProblems.Add(FString(TEXT("Input data is empty.")));
+		return OutProblems;
+	}
 
-	// Split one giant string into one string per row
-	TArray<FString> RowStrings;
-	InString.ParseIntoArray(&RowStrings, TEXT("\r\n"), true);
+	const FCsvParser Parser(InString);
+	const auto& Rows = Parser.GetRows();
 
 	// Must have at least 2 rows (column names + data)
-	if(RowStrings.Num() <= 1)
+	if(Rows.Num() <= 1)
 	{
 		OutProblems.Add(FString(TEXT("Too few rows.")));
 		return OutProblems;
 	}
 
 	// Find property for each column
-	TArray<UProperty*> ColumnProps = GetTablePropertyArray(RowStrings[0], RowStruct, OutProblems);
+	TArray<UProperty*> ColumnProps = GetTablePropertyArray(Rows[0], RowStruct, OutProblems);
 
 	// Empty existing data
 	EmptyTable();
 
 	// Iterate over rows
-	for(int32 RowIdx=1; RowIdx<RowStrings.Num(); RowIdx++)
+	for(int32 RowIdx=1; RowIdx<Rows.Num(); RowIdx++)
 	{
-		TArray<FString> CellStrings;
-		RowStrings[RowIdx].ParseIntoArray(&CellStrings, TEXT(","), false);
+		const TArray<const TCHAR*>& Cells = Rows[RowIdx];
 
 		// Need at least 1 cells (row name)
-		if(CellStrings.Num() < 1)
+		if(Cells.Num() < 1)
 		{
 			OutProblems.Add(FString::Printf(TEXT("Row '%d' has too few cells."), RowIdx));
 			continue;
 		}
 
 		// Need enough columns in the properties!
-		if( ColumnProps.Num() < CellStrings.Num() )
+		if( ColumnProps.Num() < Cells.Num() )
 		{
 			OutProblems.Add(FString::Printf(TEXT("Row '%d' has more cells than properties, is there a malformed string?"), RowIdx));
 			continue;
 		}
 
 		// Get row name
-		FName RowName = MakeValidName(CellStrings[0]);
+		FName RowName = MakeValidName(Cells[0]);
 
 		// Check its not 'none'
 		if(RowName == NAME_None)
@@ -546,11 +549,12 @@ TArray<FString> UDataTable::CreateTableFromCSVString(const FString& InString)
 		RowMap.Add(RowName, RowData);
 
 		// Now iterate over cells (skipping first cell, that was row name)
-		for(int32 CellIdx=1; CellIdx<CellStrings.Num(); CellIdx++)
+		for(int32 CellIdx=1; CellIdx<Cells.Num(); CellIdx++)
 		{
 			// Try and assign string to data using the column property
 			UProperty* ColumnProp = ColumnProps[CellIdx];
-			FString Error = AssignStringToProperty(CellStrings[CellIdx], ColumnProp, RowData);
+			const FString CellValue = Cells[CellIdx];
+			FString Error = AssignStringToProperty(CellValue, ColumnProp, RowData);
 
 			// If we failed, output a problem string
 			if(Error.Len() > 0)
@@ -558,12 +562,12 @@ TArray<FString> UDataTable::CreateTableFromCSVString(const FString& InString)
 				FString ColumnName = (ColumnProp != NULL) 
 					? FPropertyDisplayNameHelper::Get(ColumnProp, ColumnProp->GetName())
 					: FString(TEXT("NONE"));
-				OutProblems.Add(FString::Printf(TEXT("Problem assigning string '%s' to property '%s' on row '%s' : %s"), *CellStrings[CellIdx], *ColumnName, *RowName.ToString(), *Error));
+				OutProblems.Add(FString::Printf(TEXT("Problem assigning string '%s' to property '%s' on row '%s' : %s"), *CellValue, *ColumnName, *RowName.ToString(), *Error));
 			}
 		}
 
 		// Problem if we didn't have enough cells on this row
-		if(CellStrings.Num() < ColumnProps.Num())
+		if(Cells.Num() < ColumnProps.Num())
 		{
 			OutProblems.Add(FString::Printf(TEXT("Too few cells on row '%s'."), *RowName.ToString()));			
 		}
