Fix id on landscape physics shapes so that ignoring landscape on traces works

diff --git a/Engine/Source/Runtime/Engine/Private/Landscape/LandscapeCollision.cpp b/Engine/Source/Runtime/Engine/Private/Landscape/LandscapeCollision.cpp
--- a/Engine/Source/Runtime/Engine/Private/Landscape/LandscapeCollision.cpp
+++ b/Engine/Source/Runtime/Engine/Private/Landscape/LandscapeCollision.cpp
@@ -144,7 +144,7 @@ void ULandscapeHeightfieldCollisionComponent::CreatePhysicsState()
 
 				// Setup filtering
 				PxFilterData PQueryFilterData, PSimFilterData;
-				CreateShapeFilterData(GetCollisionObjectType(), GetUniqueID(), GetCollisionResponseToChannels(), 0, 0, PQueryFilterData, PSimFilterData, true, false, true);
+				CreateShapeFilterData(GetCollisionObjectType(), GetOwner()->GetUniqueID(), GetCollisionResponseToChannels(), 0, 0, PQueryFilterData, PSimFilterData, true, false, true);
 
 				// Heightfield is used for simple and complex collision
 				PQueryFilterData.word3 |= (EPDF_SimpleCollision | EPDF_ComplexCollision);
@@ -169,7 +169,7 @@ void ULandscapeHeightfieldCollisionComponent::CreatePhysicsState()
 						CollisionResponse.SetAllChannels(ECollisionResponse::ECR_Ignore);
 						CollisionResponse.SetResponse(ECollisionChannel::ECC_Visibility, ECR_Block);
 						PxFilterData PQueryFilterDataEd, PSimFilterDataEd;
-						CreateShapeFilterData(ECollisionChannel::ECC_Visibility, GetUniqueID(), CollisionResponse, 0, 0, PQueryFilterDataEd, PSimFilterDataEd, true, false, true);
+						CreateShapeFilterData(ECollisionChannel::ECC_Visibility, GetOwner()->GetUniqueID(), CollisionResponse, 0, 0, PQueryFilterDataEd, PSimFilterDataEd, true, false, true);
 
 						PQueryFilterDataEd.word3 |= (EPDF_SimpleCollision | EPDF_ComplexCollision);
 						HeightFieldEdShapeSync->setQueryFilterData(PQueryFilterDataEd);
@@ -721,7 +721,7 @@ void ULandscapeMeshCollisionComponent::CreatePhysicsState()
 						CollisionResponse.SetAllChannels(ECollisionResponse::ECR_Ignore);
 						CollisionResponse.SetResponse(ECollisionChannel::ECC_Visibility, ECR_Block);
 						PxFilterData PQueryFilterDataEd, PSimFilterDataEd;
-						CreateShapeFilterData(ECollisionChannel::ECC_Visibility, GetUniqueID(), CollisionResponse, 0, 0, PQueryFilterDataEd, PSimFilterDataEd, true, false, true);
+						CreateShapeFilterData(ECollisionChannel::ECC_Visibility, GetOwner()->GetUniqueID(), CollisionResponse, 0, 0, PQueryFilterDataEd, PSimFilterDataEd, true, false, true);
 
 						PQueryFilterDataEd.word3 |= (EPDF_SimpleCollision | EPDF_ComplexCollision);
 						MeshShapeEdSync->setQueryFilterData(PQueryFilterDataEd);
