Fix re-compile shader crash when destructible chunks are moving

diff --git a/Engine/Source/Runtime/Engine/Private/DestructibleComponent.cpp b/Engine/Source/Runtime/Engine/Private/DestructibleComponent.cpp
--- a/Engine/Source/Runtime/Engine/Private/DestructibleComponent.cpp
+++ b/Engine/Source/Runtime/Engine/Private/DestructibleComponent.cpp
@@ -354,9 +354,13 @@ void UDestructibleComponent::DestroyPhysicsState()
 	{
 		if(UWorld * World = GetWorld())
 		{
-			if(FPhysScene * PhysScene = World->GetPhysicsScene())
+			if (FPhysScene * PhysScene = World->GetPhysicsScene())
 			{
-				PhysScene->DeferredCommandHandler.DeferredRelease(ApexDestructibleActor);
+				const uint32 SceneType = BodyInstance.UseAsyncScene() ? PST_Async : PST_Sync;
+				PxScene * PScene = PhysScene->GetPhysXScene(SceneType);
+				SCOPED_SCENE_WRITE_LOCK(PScene);
+				ApexDestructibleActor->release();
+				//Deferring here is difficult because a call to CreatePhysicsState will re-use the same chunk info buffer for new chunks. Note this may be possible, but need to fix crash now
 			}
 		}
 		
@@ -800,6 +804,7 @@ void UDestructibleComponent::SetChunkVisible( int32 ChunkIndex, bool bVisible )
 				PActor->userData = UserData;
 
 
+
 				// Set collision response to non-root chunks
 				if (GetDestructibleMesh()->ApexDestructibleAsset->getChunkParentIndex(ChunkIndex) >= 0)
 				{
