Allow sound class to be set for voip audio component

diff --git a/Engine/Source/Runtime/Engine/Classes/Sound/AudioSettings.h b/Engine/Source/Runtime/Engine/Classes/Sound/AudioSettings.h
--- a/Engine/Source/Runtime/Engine/Classes/Sound/AudioSettings.h
+++ b/Engine/Source/Runtime/Engine/Classes/Sound/AudioSettings.h
@@ -14,12 +14,18 @@ class ENGINE_API UAudioSettings
 {
 	GENERATED_UCLASS_BODY()
 
+	/** The SoundClass assigned to newly created sounds */
 	UPROPERTY(config, EditAnywhere, Category="Audio", meta=(AllowedClasses="SoundClass", DisplayName="Default Sound Class"))
 	FStringAssetReference DefaultSoundClassName;
 
+	/** The SoundMix to use as base when no other system has specified a Base SoundMix */
 	UPROPERTY(config, EditAnywhere, Category="Audio", meta=(AllowedClasses="SoundMix"))
 	FStringAssetReference DefaultBaseSoundMix;
 	
+	/** Sound class to be used for the VOIP audio component */
+	UPROPERTY(config, EditAnywhere, Category="Audio", meta=(AllowedClasses="SoundClass"))
+	FStringAssetReference VoiPSoundClass;
+
 	UPROPERTY(config, EditAnywhere, Category="Audio", AdvancedDisplay, meta=(ClampMin=0.1,ClampMax=1.5))
 	float LowPassFilterResonance;
 
diff --git a/Engine/Source/Runtime/Online/OnlineSubsystemUtils/Private/OnlineSubsystemUtils.cpp b/Engine/Source/Runtime/Online/OnlineSubsystemUtils/Private/OnlineSubsystemUtils.cpp
--- a/Engine/Source/Runtime/Online/OnlineSubsystemUtils/Private/OnlineSubsystemUtils.cpp
+++ b/Engine/Source/Runtime/Online/OnlineSubsystemUtils/Private/OnlineSubsystemUtils.cpp
@@ -59,6 +59,12 @@ UAudioComponent* CreateVoiceAudioComponent(uint32 SampleRate)
 			{
 				AudioComponent->bIsUISound = true;
 				AudioComponent->SetVolumeMultiplier(1.5f);
+
+				const FStringAssetReference VoiPSoundClassName = GetDefault<UAudioSettings>()->VoiPSoundClass;
+				if (VoiPSoundClassName.IsValid())
+				{
+					AudioComponent->SoundClassOverride = LoadObject<USoundClass>(NULL, *VoiPSoundClassName.ToString());
+				}
 			}
 			else
 			{
