Fix unresolved external in non-shipping monolithic game builds

diff --git a/Engine/Source/Runtime/Engine/Engine.Build.cs b/Engine/Source/Runtime/Engine/Engine.Build.cs
--- a/Engine/Source/Runtime/Engine/Engine.Build.cs
+++ b/Engine/Source/Runtime/Engine/Engine.Build.cs
@@ -66,8 +66,8 @@ public class Engine : ModuleRules
 				"VectorVM",
 			}
         );
-        
-		if (Target.Configuration != UnrealTargetConfiguration.Shipping)
+
+        if (UEBuildConfiguration.bBuildEditor)
 		{
             PrivateDependencyModuleNames.AddRange(new string[] { "CsvParser" });
 		}
