PS4 Flip management now handled by the GPU.  Fixes unnecessary CPU stall.

diff --git a/Engine/Source/Runtime/Engine/Private/Slate/SceneViewport.cpp b/Engine/Source/Runtime/Engine/Private/Slate/SceneViewport.cpp
--- a/Engine/Source/Runtime/Engine/Private/Slate/SceneViewport.cpp
+++ b/Engine/Source/Runtime/Engine/Private/Slate/SceneViewport.cpp
@@ -1051,7 +1051,7 @@ void FSceneViewport::EnqueueBeginRenderFrame()
 void FSceneViewport::BeginRenderFrame(FRHICommandListImmediate& RHICmdList)
 {
 	check( IsInRenderingThread() );
-	if( bUseSeparateRenderTarget )
+	if (bUseSeparateRenderTarget)
 	{
 		SetRenderTarget(RHICmdList,  RenderTargetTextureRHI,  FTexture2DRHIRef() );
 	}
@@ -1059,7 +1059,7 @@ void FSceneViewport::BeginRenderFrame(FRHICommandListImmediate& RHICmdList)
 	{
 		// Get the backbuffer render target to render directly to it
 		RenderTargetTextureRHI = RHICmdList.GetViewportBackBuffer(ViewportRHI);
-		if (GRHIShaderPlatform != SP_METAL)
+		if (GRHIRequiresEarlyBackBufferRenderTarget)
 		{
 			// unused set render targets are bad on Metal
 			SetRenderTarget(RHICmdList, RenderTargetTextureRHI, FTexture2DRHIRef());
diff --git a/Engine/Source/Runtime/IOS/MetalRHI/Private/MetalRHI.cpp b/Engine/Source/Runtime/IOS/MetalRHI/Private/MetalRHI.cpp
--- a/Engine/Source/Runtime/IOS/MetalRHI/Private/MetalRHI.cpp
+++ b/Engine/Source/Runtime/IOS/MetalRHI/Private/MetalRHI.cpp
@@ -62,6 +62,7 @@ FMetalDynamicRHI::FMetalDynamicRHI()
 // 	GReadTexturePoolSizeFromIni = true;
 
 	GRHISupportsBaseVertexIndex = false;
+	GRHIRequiresEarlyBackBufferRenderTarget = false;
 
 	GMaxTextureDimensions = 4096;
 	GMaxTextureMipCount = FPlatformMath::CeilLogTwo( GMaxTextureDimensions ) + 1;
diff --git a/Engine/Source/Runtime/RHI/Private/RHI.cpp b/Engine/Source/Runtime/RHI/Private/RHI.cpp
--- a/Engine/Source/Runtime/RHI/Private/RHI.cpp
+++ b/Engine/Source/Runtime/RHI/Private/RHI.cpp
@@ -186,6 +186,7 @@ bool GTriggerGPUProfile = false;
 bool GRHISupportsTextureStreaming = false;
 bool GSupportsDepthBoundsTest = false;
 bool GRHISupportsBaseVertexIndex = true;
+bool GRHIRequiresEarlyBackBufferRenderTarget = true;
 
 /** Whether we are profiling GPU hitches. */
 bool GTriggerGPUHitchProfile = false;
diff --git a/Engine/Source/Runtime/RHI/Public/RHI.h b/Engine/Source/Runtime/RHI/Public/RHI.h
--- a/Engine/Source/Runtime/RHI/Public/RHI.h
+++ b/Engine/Source/Runtime/RHI/Public/RHI.h
@@ -202,6 +202,9 @@ extern RHI_API int32 GNumPrimitivesDrawnRHI;
 /** Whether or not the RHI can handle a non-zero BaseVertexIndex - extra SetStreamSource calls will be needed if this is false */
 extern RHI_API bool GRHISupportsBaseVertexIndex;
 
+/** Whether or not the engine should set the BackBuffer as a render target early in the frame. */
+extern RHI_API bool GRHIRequiresEarlyBackBufferRenderTarget;
+
 
 /** Called once per frame only from within an RHI. */
 extern RHI_API void RHIPrivateBeginFrame();
