nicer UI when using postprocess materials

diff --git a/Engine/Source/Runtime/Engine/Private/Materials/Material.cpp b/Engine/Source/Runtime/Engine/Private/Materials/Material.cpp
--- a/Engine/Source/Runtime/Engine/Private/Materials/Material.cpp
+++ b/Engine/Source/Runtime/Engine/Private/Materials/Material.cpp
@@ -2200,6 +2200,22 @@ bool UMaterial::CanEditChange(const UProperty* InProperty) const
 			return MaterialDomain == MD_Surface && CVar->GetValueOnGameThread() > 0;
 		}		
 
+		if(MaterialDomain == MD_PostProcess)
+		{
+			// some settings don't make sense for postprocess materials
+
+			if (PropertyName == GET_MEMBER_NAME_STRING_CHECKED(UMaterial, bTangentSpaceNormal) ||
+				PropertyName == GET_MEMBER_NAME_STRING_CHECKED(UMaterial, TwoSided) ||
+				PropertyName == GET_MEMBER_NAME_STRING_CHECKED(UMaterial, bDisableDepthTest) ||
+				PropertyName == GET_MEMBER_NAME_STRING_CHECKED(UMaterial, bUseMaterialAttributes) ||
+				PropertyName == GET_MEMBER_NAME_STRING_CHECKED(UMaterial, bFullyRough) ||
+				PropertyName == GET_MEMBER_NAME_STRING_CHECKED(UMaterial, bUseLightmapDirectionality)
+				)
+			{
+				return false;
+			}
+		}
+
 		if (PropertyName == GET_MEMBER_NAME_STRING_CHECKED(UMaterial, BlendableLocation) ||
 			PropertyName == GET_MEMBER_NAME_STRING_CHECKED(UMaterial, BlendablePriority)
 			)
