Fix for Zen birds on ES2 and Metal, where we don't have instanced mesh particles; passing in velocity and relative time as constants every draw call instead.

diff --git a/Engine/Shaders/MeshParticleVertexFactory.usf b/Engine/Shaders/MeshParticleVertexFactory.usf
--- a/Engine/Shaders/MeshParticleVertexFactory.usf
+++ b/Engine/Shaders/MeshParticleVertexFactory.usf
@@ -17,6 +17,8 @@ float4x4 PreviousLocalToWorld;
 	float4 Transform3;
 	float4 SubUVParams;
 	float SubUVLerp;
+	float4 ParticleDirection;
+	float RelativeTime;
 	float4 DynamicParameter;
 	float4 ParticleColor;
 #endif
@@ -311,8 +313,8 @@ FVertexFactoryIntermediates GetVertexFactoryIntermediates(FVertexFactoryInput In
 	Intermediates.ParticleVelocity = Input.ParticleVelocity;
 	Intermediates.RelativeTime = Input.SubUVLerpAndRelTime.y;
 #else
-	Intermediates.ParticleVelocity = float4(0,1,0,0); // TODO: We don't have this data when instancing is disabled. Should we? 
-	Intermediates.RelativeTime = 0;
+	Intermediates.ParticleVelocity = ParticleDirection;
+	Intermediates.RelativeTime = RelativeTime;
 #endif
 
 	return Intermediates;
diff --git a/Engine/Source/Runtime/Engine/Private/MeshParticleVertexFactory.cpp b/Engine/Source/Runtime/Engine/Private/MeshParticleVertexFactory.cpp
--- a/Engine/Source/Runtime/Engine/Private/MeshParticleVertexFactory.cpp
+++ b/Engine/Source/Runtime/Engine/Private/MeshParticleVertexFactory.cpp
@@ -21,7 +21,9 @@ public:
 		Transform3.Bind(ParameterMap,TEXT("Transform3"));
 		SubUVParams.Bind(ParameterMap,TEXT("SubUVParams"));
 		SubUVLerp.Bind(ParameterMap,TEXT("SubUVLerp"));
-		DynamicParameter.Bind(ParameterMap,TEXT("DynamicParameter"));
+		ParticleDirection.Bind(ParameterMap, TEXT("ParticleDirection"));
+		RelativeTime.Bind(ParameterMap, TEXT("RelativeTime"));
+		DynamicParameter.Bind(ParameterMap, TEXT("DynamicParameter"));
 		ParticleColor.Bind(ParameterMap,TEXT("ParticleColor"));
 	}
 
@@ -32,6 +34,8 @@ public:
 		Ar << Transform3;
 		Ar << SubUVParams;
 		Ar << SubUVLerp;
+		Ar << ParticleDirection;
+		Ar << RelativeTime;
 		Ar << DynamicParameter;
 		Ar << ParticleColor;
 	}
@@ -54,6 +58,8 @@ public:
 			SetShaderValue(RHICmdList, VertexShaderRHI, Transform3, Vertex->Transform[2]);
 			SetShaderValue(RHICmdList, VertexShaderRHI, SubUVParams, FVector4((float)Vertex->SubUVParams[0], (float)Vertex->SubUVParams[1], (float)Vertex->SubUVParams[2], (float)Vertex->SubUVParams[3]));
 			SetShaderValue(RHICmdList, VertexShaderRHI, SubUVLerp, Vertex->SubUVLerp);
+			SetShaderValue(RHICmdList, VertexShaderRHI, ParticleDirection, Vertex->Velocity);
+			SetShaderValue(RHICmdList, VertexShaderRHI, RelativeTime, Vertex->RelativeTime);
 			SetShaderValue(RHICmdList, VertexShaderRHI, DynamicParameter, FVector4(DynamicVertex->DynamicValue[0], DynamicVertex->DynamicValue[1], DynamicVertex->DynamicValue[2], DynamicVertex->DynamicValue[3]));
 			SetShaderValue(RHICmdList, VertexShaderRHI, ParticleColor, FVector4(Vertex->Color.Component(0), Vertex->Color.Component(1), Vertex->Color.Component(2), Vertex->Color.Component(3)));
 		}
@@ -66,6 +72,8 @@ private:
 	FShaderParameter Transform3;
 	FShaderParameter SubUVParams;
 	FShaderParameter SubUVLerp;
+	FShaderParameter ParticleDirection;
+	FShaderParameter RelativeTime;
 	FShaderParameter DynamicParameter;
 	FShaderParameter ParticleColor;
 };
