Disable WorldPositionOffset in UI materials for now

diff --git a/Engine/Source/Runtime/Engine/Private/Materials/Material.cpp b/Engine/Source/Runtime/Engine/Private/Materials/Material.cpp
--- a/Engine/Source/Runtime/Engine/Private/Materials/Material.cpp
+++ b/Engine/Source/Runtime/Engine/Private/Materials/Material.cpp
@@ -3645,6 +3645,8 @@ bool UMaterial::IsPropertyActive(EMaterialProperty InProperty)const
 		Active = true;
 		break;
 	case MP_WorldPositionOffset:
+		Active = !bUsedWithUI;
+		break;
 	case MP_MaterialAttributes:
 	default:
 		Active = true;
