Landscape : Fixed Tessellation indices problem for Game Mode

diff --git a/Engine/Source/Runtime/Engine/Private/Landscape/LandscapeRender.cpp b/Engine/Source/Runtime/Engine/Private/Landscape/LandscapeRender.cpp
--- a/Engine/Source/Runtime/Engine/Private/Landscape/LandscapeRender.cpp
+++ b/Engine/Source/Runtime/Engine/Private/Landscape/LandscapeRender.cpp
@@ -699,7 +699,7 @@ void FLandscapeComponentSceneProxy::CreateRenderThreadResources()
 	SharedBuffers = FLandscapeComponentSceneProxy::SharedBuffersMap.FindRef(SharedBuffersKey);
 	if( SharedBuffers == NULL )
 	{
-		SharedBuffers = new FLandscapeSharedBuffers(SharedBuffersKey, SubsectionSizeQuads, NumSubsections, GetScene()->GetFeatureLevel());
+		SharedBuffers = new FLandscapeSharedBuffers(SharedBuffersKey, SubsectionSizeQuads, NumSubsections, GetScene()->GetFeatureLevel(), bRequiresAdjacencyInformation);
 		FLandscapeComponentSceneProxy::SharedBuffersMap.Add(SharedBuffersKey, SharedBuffers);
 
 		if (!XYOffsetmapTexture)
@@ -724,10 +724,30 @@ void FLandscapeComponentSceneProxy::CreateRenderThreadResources()
 	{
 		if (SharedBuffers->AdjacencyIndexBuffers == NULL)
 		{
+			ensure(SharedBuffers->NumIndexBuffers > 0);
+			if (SharedBuffers->IndexBuffers[0])
+			{
+				// Recreate Index Buffers, this case happens only there are Landscape Components using different material (one uses tessellation, other don't use it) 
+				if (SharedBuffers->bUse32BitIndices && !((FRawStaticIndexBuffer16or32<uint32>*)SharedBuffers->IndexBuffers[0])->Num())
+				{
+					SharedBuffers->CreateIndexBuffers<uint32>(GetScene()->GetFeatureLevel(), bRequiresAdjacencyInformation);
+				}
+				else if (!((FRawStaticIndexBuffer16or32<uint16>*)SharedBuffers->IndexBuffers[0])->Num())
+				{
+					SharedBuffers->CreateIndexBuffers<uint16>(GetScene()->GetFeatureLevel(), bRequiresAdjacencyInformation);
+				}
+			}
+
 			SharedBuffers->AdjacencyIndexBuffers = new FLandscapeSharedAdjacencyIndexBuffer(SharedBuffers);
 			FLandscapeComponentSceneProxy::SharedAdjacencyIndexBufferMap.Add(SharedBuffersKey, SharedBuffers->AdjacencyIndexBuffers);
 		}
 		SharedBuffers->AdjacencyIndexBuffers->AddRef();
+
+		// Delayed Initialize for IndexBuffers
+		for (int i = 0; i < SharedBuffers->NumIndexBuffers; i++)
+		{
+			SharedBuffers->IndexBuffers[i]->InitResource();
+		}
 	}
 
 	// Assign vertex factory
@@ -1924,10 +1944,10 @@ void FLandscapeVertexBuffer::InitRHI()
 //
 
 template <typename INDEX_TYPE>
-void FLandscapeSharedBuffers::CreateIndexBuffers(ERHIFeatureLevel::Type InFeatureLevel)
+void FLandscapeSharedBuffers::CreateIndexBuffers(ERHIFeatureLevel::Type InFeatureLevel, bool bRequiresAdjacencyInformation)
 {
 	if (InFeatureLevel == ERHIFeatureLevel::ES2)
-{
+	{
 		if (!bVertexScoresComputed)
 		{
 			bVertexScoresComputed = ComputeVertexScores();
@@ -1958,11 +1978,11 @@ void FLandscapeSharedBuffers::CreateIndexBuffers(ERHIFeatureLevel::Type InFeatur
 			// ES2 version
 			float MipRatio = (float)SubsectionSizeQuads / (float)LodSubsectionSizeQuads; // Morph current MIP to base MIP
 
-		for( int32 SubY=0;SubY<NumSubsections;SubY++ )
-		{
-			for( int32 SubX=0;SubX<NumSubsections;SubX++ )
+			for( int32 SubY=0;SubY<NumSubsections;SubY++ )
 			{
-				TArray<INDEX_TYPE> SubIndices;
+				for( int32 SubX=0;SubX<NumSubsections;SubX++ )
+				{
+					TArray<INDEX_TYPE> SubIndices;
 					SubIndices.Empty(FMath::Square(LodSubsectionSizeQuads) * 6);
 
 					int32& MaxIndex = IndexRanges[Mip].MaxIndex[SubX][SubY];
@@ -1971,9 +1991,9 @@ void FLandscapeSharedBuffers::CreateIndexBuffers(ERHIFeatureLevel::Type InFeatur
 					MinIndex = MAX_int32;
 
 					for (int32 y = 0; y < LodSubsectionSizeQuads; y++)
-				{
-						for (int32 x = 0; x < LodSubsectionSizeQuads; x++)
 					{
+						for (int32 x = 0; x < LodSubsectionSizeQuads; x++)
+						{
 							int32 x0 = FMath::RoundToInt((float)x * MipRatio);
 							int32 y0 = FMath::RoundToInt((float)y * MipRatio);
 							int32 x1 = FMath::RoundToInt((float)(x + 1) * MipRatio);
@@ -1999,9 +2019,9 @@ void FLandscapeSharedBuffers::CreateIndexBuffers(ERHIFeatureLevel::Type InFeatur
 							{
 								i00 = VertexCount++;
 								VertexMap.Add(Key00, i00);
-						}
-						else
-						{
+							}
+							else
+							{
 								i00 = *KeyPtr;
 							}
 
@@ -2014,27 +2034,27 @@ void FLandscapeSharedBuffers::CreateIndexBuffers(ERHIFeatureLevel::Type InFeatur
 							else
 							{
 								i10 = *KeyPtr;
-						}
+							}
 
 							KeyPtr = VertexMap.Find(Key11);
 							if (KeyPtr == NULL)
 							{
 								i11 = VertexCount++;
 								VertexMap.Add(Key11, i11);
-					}
+							}
 							else
 							{
 								i11 = *KeyPtr;
-				}
+							}
 
 							KeyPtr = VertexMap.Find(Key01);
 							if (KeyPtr == NULL)
-				{
+							{
 								i01 = VertexCount++;
 								VertexMap.Add(Key01, i01);
 							}
 							else
-					{
+							{
 								i01 = *KeyPtr;
 							}
 
@@ -2115,22 +2135,31 @@ void FLandscapeSharedBuffers::CreateIndexBuffers(ERHIFeatureLevel::Type InFeatur
 					MaxIndexFull = FMath::Max<int32>(MaxIndexFull, MaxIndex);
 					MinIndexFull = FMath::Min<int32>(MinIndexFull, MinIndex);
 				}
-	}
+			}
 
 			check(MinIndexFull <= (uint32)((INDEX_TYPE)(~(INDEX_TYPE)0)));
 			check(NewIndices.Num() == ExpectedNumIndices);
-	}
+		}
 
 		// Create and init new index buffer with index data
-		auto IndexBuffer = new FRawStaticIndexBuffer16or32<INDEX_TYPE>(false);
+		FRawStaticIndexBuffer16or32<INDEX_TYPE>* IndexBuffer = (FRawStaticIndexBuffer16or32<INDEX_TYPE>*)IndexBuffers[Mip];
+		if (!IndexBuffer)
+		{
+			IndexBuffer = new FRawStaticIndexBuffer16or32<INDEX_TYPE>(false);
+		}
 		IndexBuffer->AssignNewBuffer(NewIndices);
-		IndexBuffer->InitResource();
+
+		// Delay init resource to keep CPU data until create AdjacencyIndexbuffers
+		if (!bRequiresAdjacencyInformation)
+		{
+			IndexBuffer->InitResource();
+		}
 
 		IndexBuffers[Mip] = IndexBuffer;
 	}
 }
 
-FLandscapeSharedBuffers::FLandscapeSharedBuffers(int32 InSharedBuffersKey, int32 InSubsectionSizeQuads, int32 InNumSubsections, ERHIFeatureLevel::Type InFeatureLevel)
+FLandscapeSharedBuffers::FLandscapeSharedBuffers(int32 InSharedBuffersKey, int32 InSubsectionSizeQuads, int32 InNumSubsections, ERHIFeatureLevel::Type InFeatureLevel, bool bRequiresAdjacencyInformation)
 :	SharedBuffersKey(InSharedBuffersKey)
 ,	NumIndexBuffers(FMath::CeilLogTwo(InSubsectionSizeQuads+1))
 ,	SubsectionSizeVerts(InSubsectionSizeQuads+1)
@@ -2138,6 +2167,7 @@ FLandscapeSharedBuffers::FLandscapeSharedBuffers(int32 InSharedBuffersKey, int32
 ,	VertexFactory(NULL)
 ,	VertexBuffer(NULL)
 ,	AdjacencyIndexBuffers(NULL)
+,	bUse32BitIndices(false)
 {
 	if (InFeatureLevel > ERHIFeatureLevel::ES2)
 	{
@@ -2145,16 +2175,18 @@ FLandscapeSharedBuffers::FLandscapeSharedBuffers(int32 InSharedBuffersKey, int32
 		VertexBuffer = new FLandscapeVertexBuffer(SubsectionSizeVerts, NumSubsections);
 	}
 	IndexBuffers = new FIndexBuffer*[NumIndexBuffers];
+	FMemory::Memzero(IndexBuffers, sizeof(FIndexBuffer*)* NumIndexBuffers);
 	IndexRanges = new FLandscapeIndexRanges[NumIndexBuffers]();
 
 	// See if we need to use 16 or 32-bit index buffers
 	if( FMath::Square(SubsectionSizeVerts) * FMath::Square(NumSubsections) > 65535 )
 	{
-		CreateIndexBuffers<uint32>(InFeatureLevel);
+		bUse32BitIndices = true;
+		CreateIndexBuffers<uint32>(InFeatureLevel, bRequiresAdjacencyInformation);
 	}
 	else
 	{
-		CreateIndexBuffers<uint16>(InFeatureLevel);
+		CreateIndexBuffers<uint16>(InFeatureLevel, bRequiresAdjacencyInformation);
 	}
 }
 
@@ -2185,7 +2217,7 @@ FLandscapeSharedBuffers::~FLandscapeSharedBuffers()
 template<typename IndexType>
 static void BuildLandscapeAdjacencyIndexBuffer(int32 LODSubsectionSizeQuads, const FRawStaticIndexBuffer16or32<IndexType>* Indices, TArray<IndexType>& OutPnAenIndices)
 {
-	if (Indices)
+	if (Indices && Indices->Num())
 	{
 		// Landscape use regular grid, so only expand Index buffer works
 		// PN AEN Dominant Corner
diff --git a/Engine/Source/Runtime/Engine/Private/Landscape/LandscapeRenderMobile.cpp b/Engine/Source/Runtime/Engine/Private/Landscape/LandscapeRenderMobile.cpp
--- a/Engine/Source/Runtime/Engine/Private/Landscape/LandscapeRenderMobile.cpp
+++ b/Engine/Source/Runtime/Engine/Private/Landscape/LandscapeRenderMobile.cpp
@@ -197,7 +197,7 @@ void FLandscapeComponentSceneProxyMobile::CreateRenderThreadResources()
 	SharedBuffers = FLandscapeComponentSceneProxy::SharedBuffersMap.FindRef(SharedBuffersKey);
 	if( SharedBuffers == NULL )
 	{
-		SharedBuffers = new FLandscapeSharedBuffers(SharedBuffersKey, SubsectionSizeQuads, NumSubsections, GetScene()->GetFeatureLevel());
+		SharedBuffers = new FLandscapeSharedBuffers(SharedBuffersKey, SubsectionSizeQuads, NumSubsections, GetScene()->GetFeatureLevel(), false);
 		FLandscapeComponentSceneProxy::SharedBuffersMap.Add(SharedBuffersKey, SharedBuffers);
 	}
 
diff --git a/Engine/Source/Runtime/Engine/Public/Landscape/LandscapeRender.h b/Engine/Source/Runtime/Engine/Public/Landscape/LandscapeRender.h
--- a/Engine/Source/Runtime/Engine/Public/Landscape/LandscapeRender.h
+++ b/Engine/Source/Runtime/Engine/Public/Landscape/LandscapeRender.h
@@ -295,11 +295,12 @@ public:
 	FIndexBuffer** IndexBuffers;
 	FLandscapeIndexRanges* IndexRanges;
 	FLandscapeSharedAdjacencyIndexBuffer* AdjacencyIndexBuffers;
+	bool bUse32BitIndices;
 
-	FLandscapeSharedBuffers(int32 SharedBuffersKey, int32 SubsectionSizeQuads, int32 NumSubsections, ERHIFeatureLevel::Type InFeatureLevel);
+	FLandscapeSharedBuffers(int32 SharedBuffersKey, int32 SubsectionSizeQuads, int32 NumSubsections, ERHIFeatureLevel::Type InFeatureLevel, bool bRequiresAdjacencyInformation);
 
 	template <typename INDEX_TYPE>
-	void CreateIndexBuffers(ERHIFeatureLevel::Type InFeatureLevel);
+	void CreateIndexBuffers(ERHIFeatureLevel::Type InFeatureLevel, bool bRequiresAdjacencyInformation);
 
 	virtual ~FLandscapeSharedBuffers();
 };
