Batched elements rendering now has an extra couple of permutations for gamma correcting input textures on mobile.

diff --git a/Engine/Shaders/SimpleElementPixelShader.usf b/Engine/Shaders/SimpleElementPixelShader.usf
--- a/Engine/Shaders/SimpleElementPixelShader.usf
+++ b/Engine/Shaders/SimpleElementPixelShader.usf
@@ -25,6 +25,15 @@ void ReplicateChannelSimpleElementShader(inout float4 BaseColor)
 bool bEnableEditorPrimitiveDepthTest;
 float4 ScreenToPixel;
 
+MaterialFloat4 ColourTexture2DSample(Texture2D Tex, SamplerState Sampler, float2 UV)
+{
+	MaterialFloat4 Sample = Tex.Sample(Sampler, UV);
+#if SRGB_INPUT_TEXTURE && (ES2_PROFILE || ES3_1_PROFILE)
+	Sample.rgb = Sample.rgb * Sample.rgb;
+#endif
+	return Sample;
+}
+
 #if FEATURE_LEVEL >= FEATURE_LEVEL_SM4
 
 void PerformEditorCompositingDepthTest(float4 SvPosition)
@@ -61,7 +70,8 @@ void Main(
 	}
 #endif
 	 
-	float4 BaseColor = Texture2DSample(InTexture, InTextureSampler,TextureCoordinate);
+	float4 BaseColor = ColourTexture2DSample(InTexture, InTextureSampler,TextureCoordinate);
+
 	ReplicateChannelSimpleElementShader(BaseColor);
 	OutColor = RETURN_COLOR(BaseColor * Color);	
 
@@ -97,7 +107,7 @@ void GammaMain(
 	}
 #endif
 
-	float4 BaseColor = Texture2DSample(InTexture, InTextureSampler,TextureCoordinate);
+	float4 BaseColor = ColourTexture2DSample(InTexture, InTextureSampler,TextureCoordinate);
 	ReplicateChannelSimpleElementShader(BaseColor);
 	OutColor = BaseColor * Color;
 	if( Gamma != 1.0 )
@@ -135,7 +145,7 @@ void GammaMaskedMain(
 	}
 #endif
 
-	float4 BaseColor = Texture2DSample(InTexture, InTextureSampler,TextureCoordinate);
+	float4 BaseColor = ColourTexture2DSample(InTexture, InTextureSampler,TextureCoordinate);
 	ReplicateChannelSimpleElementShader(BaseColor);
 	clip(BaseColor.a - ClipRef);
 	OutColor.rgba = BaseColor * Color;
diff --git a/Engine/Source/Runtime/Engine/Private/BatchedElements.cpp b/Engine/Source/Runtime/Engine/Private/BatchedElements.cpp
--- a/Engine/Source/Runtime/Engine/Private/BatchedElements.cpp
+++ b/Engine/Source/Runtime/Engine/Private/BatchedElements.cpp
@@ -356,7 +356,8 @@ static void SetHitTestingBlendState(FRHICommandList& RHICmdList, ESimpleElementB
 }
 
 FGlobalBoundShaderState FBatchedElements::SimpleBoundShaderState;
-FGlobalBoundShaderState FBatchedElements::RegularBoundShaderState;
+FGlobalBoundShaderState FBatchedElements::RegularSRGBBoundShaderState;
+FGlobalBoundShaderState FBatchedElements::RegularLinearBoundShaderState;
 FGlobalBoundShaderState FBatchedElements::MaskedBoundShaderState;
 FGlobalBoundShaderState FBatchedElements::DistanceFieldBoundShaderState;
 FGlobalBoundShaderState FBatchedElements::HitTestingBoundShaderState;
@@ -476,12 +477,24 @@ void FBatchedElements::PrepareShaders(
 			    // also don't support alpha testing to floating point render targets
 				RHICmdList.SetBlendState(TStaticBlendState<>::GetRHI());
 
-				TShaderMapRef<FSimpleElementMaskedGammaPS> MaskedPixelShader(GetGlobalShaderMap(FeatureLevel));
-				SetGlobalBoundShaderState(RHICmdList, FeatureLevel, MaskedBoundShaderState, GSimpleElementVertexDeclaration.VertexDeclarationRHI,
-					*VertexShader, *MaskedPixelShader);
+				if (Texture->bSRGB)
+				{
+					TShaderMapRef<FSimpleElementMaskedGammaPS_SRGB> MaskedPixelShader(GetGlobalShaderMap(FeatureLevel));
+					SetGlobalBoundShaderState(RHICmdList, FeatureLevel, MaskedBoundShaderState, GSimpleElementVertexDeclaration.VertexDeclarationRHI,
+						*VertexShader, *MaskedPixelShader);
+
+					MaskedPixelShader->SetEditorCompositingParameters(RHICmdList, View, DepthTexture);
+					MaskedPixelShader->SetParameters(RHICmdList, Texture, Gamma, GBatchedElementAlphaRefVal / 255.0f, BlendMode);
+				}
+				else
+				{
+					TShaderMapRef<FSimpleElementMaskedGammaPS_Linear> MaskedPixelShader(GetGlobalShaderMap(FeatureLevel));
+					SetGlobalBoundShaderState(RHICmdList, FeatureLevel, MaskedBoundShaderState, GSimpleElementVertexDeclaration.VertexDeclarationRHI,
+						*VertexShader, *MaskedPixelShader);
 
-				MaskedPixelShader->SetEditorCompositingParameters(RHICmdList, View, DepthTexture );
-				MaskedPixelShader->SetParameters(RHICmdList, Texture,Gamma,GBatchedElementAlphaRefVal / 255.0f,BlendMode);
+					MaskedPixelShader->SetEditorCompositingParameters(RHICmdList, View, DepthTexture);
+					MaskedPixelShader->SetParameters(RHICmdList, Texture, Gamma, GBatchedElementAlphaRefVal / 255.0f, BlendMode);
+				}
 		    }
 		    // render distance field elements
 		    else if (
@@ -546,21 +559,36 @@ void FBatchedElements::PrepareShaders(
     
 			    if (FMath::Abs(Gamma - 1.0f) < KINDA_SMALL_NUMBER)
 			    {
-					TShaderMapRef<FSimpleElementPS> RegularPixelShader(GetGlobalShaderMap(FeatureLevel));
+					TShaderMapRef<FSimpleElementPS> PixelShader(GetGlobalShaderMap(FeatureLevel));
 					SetGlobalBoundShaderState(RHICmdList, FeatureLevel, SimpleBoundShaderState, GSimpleElementVertexDeclaration.VertexDeclarationRHI,
-						*VertexShader, *RegularPixelShader );
+						*VertexShader, *PixelShader);
 
-					RegularPixelShader->SetEditorCompositingParameters(RHICmdList, View, DepthTexture);
-				    RegularPixelShader->SetParameters(RHICmdList, Texture);
+					PixelShader->SetParameters(RHICmdList, Texture);
+					PixelShader->SetEditorCompositingParameters(RHICmdList, View, DepthTexture);
 			    }
 			    else
 			    {
-					TShaderMapRef<FSimpleElementGammaPS> RegularPixelShader(GetGlobalShaderMap(FeatureLevel));
-					SetGlobalBoundShaderState(RHICmdList, FeatureLevel, RegularBoundShaderState, GSimpleElementVertexDeclaration.VertexDeclarationRHI,
-						*VertexShader, *RegularPixelShader );
+					FSimpleElementGammaBasePS* BasePixelShader;
+					FGlobalBoundShaderState* BoundShaderState;
+
+					if (Texture->bSRGB)
+					{
+						static TShaderMapRef<FSimpleElementGammaPS_SRGB> PixelShader(GetGlobalShaderMap(FeatureLevel));
+						BasePixelShader = *PixelShader;
+						BoundShaderState = &RegularSRGBBoundShaderState;
+					}
+					else
+					{
+						static TShaderMapRef<FSimpleElementGammaPS_Linear> PixelShader(GetGlobalShaderMap(FeatureLevel));
+						BasePixelShader = *PixelShader;
+						BoundShaderState = &RegularLinearBoundShaderState;
+					}
+
+					SetGlobalBoundShaderState(RHICmdList, FeatureLevel, *BoundShaderState, GSimpleElementVertexDeclaration.VertexDeclarationRHI,
+						*VertexShader, BasePixelShader);
 
-				    RegularPixelShader->SetParameters(RHICmdList, Texture,Gamma,BlendMode);
-					RegularPixelShader->SetEditorCompositingParameters(RHICmdList, View, DepthTexture);
+					BasePixelShader->SetParameters(RHICmdList, Texture, Gamma, BlendMode);
+					BasePixelShader->SetEditorCompositingParameters(RHICmdList, View, DepthTexture);
 			    }
 		    }
 	    }
diff --git a/Engine/Source/Runtime/Engine/Private/SimpleElementShaders.cpp b/Engine/Source/Runtime/Engine/Private/SimpleElementShaders.cpp
--- a/Engine/Source/Runtime/Engine/Private/SimpleElementShaders.cpp
+++ b/Engine/Source/Runtime/Engine/Private/SimpleElementShaders.cpp
@@ -42,7 +42,7 @@ void FSimpleElementVS::ModifyCompilationEnvironment(EShaderPlatform Platform, FS
 	Simple element pixel shaders.
 ------------------------------------------------------------------------------*/
 
-FSimpleElementPS::FSimpleElementPS(const ShaderMetaType::CompiledShaderInitializerType& Initializer):
+FSimpleElementPS::FSimpleElementPS(const ShaderMetaType::CompiledShaderInitializerType& Initializer) :
 	FGlobalShader(Initializer)
 {
 	InTexture.Bind(Initializer.ParameterMap,TEXT("InTexture"));
@@ -106,40 +106,40 @@ bool FSimpleElementPS::Serialize(FArchive& Ar)
 	return bShaderHasOutdatedParameters;
 }
 
-FSimpleElementGammaPS::FSimpleElementGammaPS(const ShaderMetaType::CompiledShaderInitializerType& Initializer):
+FSimpleElementGammaBasePS::FSimpleElementGammaBasePS(const ShaderMetaType::CompiledShaderInitializerType& Initializer) :
 	FSimpleElementPS(Initializer)
 {
 	Gamma.Bind(Initializer.ParameterMap,TEXT("Gamma"));
 }
 
-void FSimpleElementGammaPS::SetParameters(FRHICommandList& RHICmdList, const FTexture* Texture,float GammaValue,ESimpleElementBlendMode BlendMode)
+void FSimpleElementGammaBasePS::SetParameters(FRHICommandList& RHICmdList, const FTexture* Texture, float GammaValue, ESimpleElementBlendMode BlendMode)
 {
 	FSimpleElementPS::SetParameters(RHICmdList, Texture);
 	SetShaderValue(RHICmdList, GetPixelShader(),Gamma,GammaValue);
 }
 
-bool FSimpleElementGammaPS::Serialize(FArchive& Ar)
+bool FSimpleElementGammaBasePS::Serialize(FArchive& Ar)
 {
 	bool bShaderHasOutdatedParameters = FSimpleElementPS::Serialize(Ar);
 	Ar << Gamma;
 	return bShaderHasOutdatedParameters;
 }
 
-FSimpleElementMaskedGammaPS::FSimpleElementMaskedGammaPS(const ShaderMetaType::CompiledShaderInitializerType& Initializer):
-	FSimpleElementGammaPS(Initializer)
+FSimpleElementMaskedGammaBasePS::FSimpleElementMaskedGammaBasePS(const ShaderMetaType::CompiledShaderInitializerType& Initializer) :
+	FSimpleElementGammaBasePS(Initializer)
 {
 	ClipRef.Bind(Initializer.ParameterMap,TEXT("ClipRef"), SPF_Mandatory);
 }
 
-void FSimpleElementMaskedGammaPS::SetParameters(FRHICommandList& RHICmdList, const FTexture* Texture,float Gamma,float ClipRefValue,ESimpleElementBlendMode BlendMode)
+void FSimpleElementMaskedGammaBasePS::SetParameters(FRHICommandList& RHICmdList, const FTexture* Texture, float Gamma, float ClipRefValue, ESimpleElementBlendMode BlendMode)
 {
-	FSimpleElementGammaPS::SetParameters(RHICmdList, Texture,Gamma,BlendMode);
+	FSimpleElementGammaBasePS::SetParameters(RHICmdList, Texture,Gamma,BlendMode);
 	SetShaderValue(RHICmdList, GetPixelShader(),ClipRef,ClipRefValue);
 }
 
-bool FSimpleElementMaskedGammaPS::Serialize(FArchive& Ar)
+bool FSimpleElementMaskedGammaBasePS::Serialize(FArchive& Ar)
 {
-	bool bShaderHasOutdatedParameters = FSimpleElementGammaPS::Serialize(Ar);
+	bool bShaderHasOutdatedParameters = FSimpleElementGammaBasePS::Serialize(Ar);
 	Ar << ClipRef;
 	return bShaderHasOutdatedParameters;
 }
@@ -150,7 +150,7 @@ bool FSimpleElementMaskedGammaPS::Serialize(FArchive& Ar)
 * @param Initializer - shader initialization container
 */
 FSimpleElementDistanceFieldGammaPS::FSimpleElementDistanceFieldGammaPS(const ShaderMetaType::CompiledShaderInitializerType& Initializer)
-:	FSimpleElementMaskedGammaPS(Initializer)
+	: FSimpleElementMaskedGammaBasePS(Initializer)
 {
 	SmoothWidth.Bind(Initializer.ParameterMap,TEXT("SmoothWidth"));
 	EnableShadow.Bind(Initializer.ParameterMap,TEXT("EnableShadow"));
@@ -190,7 +190,7 @@ void FSimpleElementDistanceFieldGammaPS::SetParameters(
 	ESimpleElementBlendMode BlendMode
 	)
 {
-	FSimpleElementMaskedGammaPS::SetParameters(RHICmdList, Texture,Gamma,ClipRef,BlendMode);
+	FSimpleElementMaskedGammaBasePS::SetParameters(RHICmdList, Texture,Gamma,ClipRef,BlendMode);
 	SetShaderValue(RHICmdList, GetPixelShader(),SmoothWidth,SmoothWidthValue);		
 	SetPixelShaderBool(RHICmdList, GetPixelShader(),EnableShadow,EnableShadowValue);
 	if (EnableShadowValue)
@@ -219,7 +219,7 @@ void FSimpleElementDistanceFieldGammaPS::SetParameters(
 */
 bool FSimpleElementDistanceFieldGammaPS::Serialize(FArchive& Ar)
 {
-	bool bShaderHasOutdatedParameters = FSimpleElementMaskedGammaPS::Serialize(Ar);
+	bool bShaderHasOutdatedParameters = FSimpleElementMaskedGammaBasePS::Serialize(Ar);
 	Ar << SmoothWidth;
 	Ar << EnableShadow;
 	Ar << ShadowDirection;
@@ -291,9 +291,11 @@ bool FSimpleElementColorChannelMaskPS::Serialize(FArchive& Ar)
 ------------------------------------------------------------------------------*/
 
 IMPLEMENT_SHADER_TYPE(,FSimpleElementVS,TEXT("SimpleElementVertexShader"),TEXT("Main"),SF_Vertex);
-IMPLEMENT_SHADER_TYPE(,FSimpleElementPS,TEXT("SimpleElementPixelShader"),TEXT("Main"),SF_Pixel);
-IMPLEMENT_SHADER_TYPE(,FSimpleElementGammaPS,TEXT("SimpleElementPixelShader"),TEXT("GammaMain"),SF_Pixel);
-IMPLEMENT_SHADER_TYPE(,FSimpleElementMaskedGammaPS,TEXT("SimpleElementPixelShader"),TEXT("GammaMaskedMain"),SF_Pixel);
+IMPLEMENT_SHADER_TYPE(,FSimpleElementPS, TEXT("SimpleElementPixelShader"), TEXT("Main"), SF_Pixel);
+IMPLEMENT_SHADER_TYPE(template<>, FSimpleElementGammaPS_SRGB, TEXT("SimpleElementPixelShader"), TEXT("GammaMain"), SF_Pixel);
+IMPLEMENT_SHADER_TYPE(template<>, FSimpleElementGammaPS_Linear, TEXT("SimpleElementPixelShader"), TEXT("GammaMain"), SF_Pixel);
+IMPLEMENT_SHADER_TYPE(template<>, FSimpleElementMaskedGammaPS_SRGB, TEXT("SimpleElementPixelShader"), TEXT("GammaMaskedMain"), SF_Pixel);
+IMPLEMENT_SHADER_TYPE(template<>, FSimpleElementMaskedGammaPS_Linear, TEXT("SimpleElementPixelShader"), TEXT("GammaMaskedMain"), SF_Pixel);
 IMPLEMENT_SHADER_TYPE(,FSimpleElementDistanceFieldGammaPS,TEXT("SimpleElementPixelShader"),TEXT("GammaDistanceFieldMain"),SF_Pixel);
 IMPLEMENT_SHADER_TYPE(,FSimpleElementHitProxyPS,TEXT("SimpleElementHitProxyPixelShader"),TEXT("Main"),SF_Pixel);
 IMPLEMENT_SHADER_TYPE(,FSimpleElementColorChannelMaskPS,TEXT("SimpleElementColorChannelMaskPixelShader"),TEXT("Main"),SF_Pixel);
diff --git a/Engine/Source/Runtime/Engine/Private/Texture2D.cpp b/Engine/Source/Runtime/Engine/Private/Texture2D.cpp
--- a/Engine/Source/Runtime/Engine/Private/Texture2D.cpp
+++ b/Engine/Source/Runtime/Engine/Private/Texture2D.cpp
@@ -1208,6 +1208,7 @@ FTexture2DResource::FTexture2DResource( UTexture2D* InOwner, int32 InitialMipCou
 {
 	// HDR images are stored in linear but still require gamma correction to display correctly.
 	bIgnoreGammaConversions = !Owner->SRGB && Owner->CompressionSettings != TC_HDR;
+	bSRGB = InOwner->SRGB;
 
 	// First request to create the resource. Decrement the counter to signal that the resource is not ready
 	// for streaming yet.
diff --git a/Engine/Source/Runtime/Engine/Public/BatchedElements.h b/Engine/Source/Runtime/Engine/Public/BatchedElements.h
--- a/Engine/Source/Runtime/Engine/Public/BatchedElements.h
+++ b/Engine/Source/Runtime/Engine/Public/BatchedElements.h
@@ -297,8 +297,10 @@ private:
 
 	/** bound shader state for the fast path */
 	static FGlobalBoundShaderState SimpleBoundShaderState;
-	/** bound shader state for the regular mesh elements */
-	static FGlobalBoundShaderState RegularBoundShaderState;
+	/** bound shader state for the regular mesh elements with a linear texture */
+	static FGlobalBoundShaderState RegularLinearBoundShaderState;
+	/** bound shader state for the regular mesh elements with an sRGB texture */
+	static FGlobalBoundShaderState RegularSRGBBoundShaderState;
 	/** bound shader state for masked mesh elements */
 	static FGlobalBoundShaderState MaskedBoundShaderState;
 	/** bound shader state for masked mesh elements */
diff --git a/Engine/Source/Runtime/Engine/Public/SimpleElementShaders.h b/Engine/Source/Runtime/Engine/Public/SimpleElementShaders.h
--- a/Engine/Source/Runtime/Engine/Public/SimpleElementShaders.h
+++ b/Engine/Source/Runtime/Engine/Public/SimpleElementShaders.h
@@ -1,3 +1,4 @@
+
 // Copyright 1998-2014 Epic Games, Inc. All Rights Reserved.
 
 /*==============================================================================
@@ -39,7 +40,7 @@ private:
  */
 class FSimpleElementPS : public FGlobalShader
 {
-	DECLARE_SHADER_TYPE(FSimpleElementPS,Global);
+	DECLARE_SHADER_TYPE(FSimpleElementPS, Global);
 public:
 
 	static bool ShouldCache(EShaderPlatform Platform) { return true; }
@@ -72,15 +73,13 @@ private:
 /**
  * A pixel shader for rendering a texture on a simple element.
  */
-class FSimpleElementGammaPS : public FSimpleElementPS
+class FSimpleElementGammaBasePS : public FSimpleElementPS
 {
-	DECLARE_SHADER_TYPE(FSimpleElementGammaPS,Global);
+	DECLARE_SHADER_TYPE(FSimpleElementGammaBasePS, Global);
 public:
 
-	static bool ShouldCache(EShaderPlatform Platform) { return true; }
-
-	FSimpleElementGammaPS(const ShaderMetaType::CompiledShaderInitializerType& Initializer);
-	FSimpleElementGammaPS() {}
+	FSimpleElementGammaBasePS(const ShaderMetaType::CompiledShaderInitializerType& Initializer);
+	FSimpleElementGammaBasePS() {}
 
 	void SetParameters(FRHICommandList& RHICmdList, const FTexture* Texture,float GammaValue,ESimpleElementBlendMode BlendMode);
 
@@ -90,18 +89,33 @@ private:
 	FShaderParameter Gamma;
 };
 
+template <bool bSRGBTexture>
+class FSimpleElementGammaPS : public FSimpleElementGammaBasePS
+{
+	DECLARE_SHADER_TYPE(FSimpleElementGammaPS, Global);
+public:
+
+	FSimpleElementGammaPS(const ShaderMetaType::CompiledShaderInitializerType& Initializer) : FSimpleElementGammaBasePS(Initializer) {}
+	FSimpleElementGammaPS() {}
+
+	static bool ShouldCache(EShaderPlatform Platform) { return true; }
+
+	static void ModifyCompilationEnvironment(EShaderPlatform Platform, FShaderCompilerEnvironment& OutEnvironment)
+	{
+		OutEnvironment.SetDefine(TEXT("SRGB_INPUT_TEXTURE"), bSRGBTexture ? TEXT("1") : TEXT("0"));
+	}
+};
+
 /**
  * A pixel shader for rendering a masked texture on a simple element.
  */
-class FSimpleElementMaskedGammaPS : public FSimpleElementGammaPS
+class FSimpleElementMaskedGammaBasePS : public FSimpleElementGammaBasePS
 {
-	DECLARE_SHADER_TYPE(FSimpleElementMaskedGammaPS,Global);
+	DECLARE_SHADER_TYPE(FSimpleElementMaskedGammaBasePS, Global);
 public:
 
-	static bool ShouldCache(EShaderPlatform Platform) { return true; }
-
-	FSimpleElementMaskedGammaPS(const ShaderMetaType::CompiledShaderInitializerType& Initializer);
-	FSimpleElementMaskedGammaPS() {}
+	FSimpleElementMaskedGammaBasePS(const ShaderMetaType::CompiledShaderInitializerType& Initializer);
+	FSimpleElementMaskedGammaBasePS() {}
 
 	void SetParameters(FRHICommandList& RHICmdList, const FTexture* Texture,float Gamma,float ClipRefValue,ESimpleElementBlendMode BlendMode);
 
@@ -111,10 +125,26 @@ private:
 	FShaderParameter ClipRef;
 };
 
+template <bool bSRGBTexture>
+class FSimpleElementMaskedGammaPS : public FSimpleElementMaskedGammaBasePS
+{
+	DECLARE_SHADER_TYPE(FSimpleElementMaskedGammaPS, Global);
+public:
+
+	FSimpleElementMaskedGammaPS(const ShaderMetaType::CompiledShaderInitializerType& Initializer) : FSimpleElementMaskedGammaBasePS(Initializer) {}
+	FSimpleElementMaskedGammaPS() {}
+
+	static bool ShouldCache(EShaderPlatform Platform) { return true; }
+	static void ModifyCompilationEnvironment(EShaderPlatform Platform, FShaderCompilerEnvironment& OutEnvironment)
+	{
+		OutEnvironment.SetDefine(TEXT("SRGB_INPUT_TEXTURE"), bSRGBTexture ? TEXT("1") : TEXT("0"));
+	}
+};
+
 /**
 * A pixel shader for rendering a masked texture using signed distance filed for alpha on a simple element.
 */
-class FSimpleElementDistanceFieldGammaPS : public FSimpleElementMaskedGammaPS
+class FSimpleElementDistanceFieldGammaPS : public FSimpleElementMaskedGammaBasePS
 {
 	DECLARE_SHADER_TYPE(FSimpleElementDistanceFieldGammaPS,Global);
 public:
@@ -261,3 +291,8 @@ private:
 	FShaderParameter ColorWeights; 
 	FShaderParameter Gamma;
 };
+
+typedef FSimpleElementGammaPS<true> FSimpleElementGammaPS_SRGB;
+typedef FSimpleElementGammaPS<false> FSimpleElementGammaPS_Linear;
+typedef FSimpleElementMaskedGammaPS<true> FSimpleElementMaskedGammaPS_SRGB;
+typedef FSimpleElementMaskedGammaPS<false> FSimpleElementMaskedGammaPS_Linear;
diff --git a/Engine/Source/Runtime/RenderCore/Public/RenderResource.h b/Engine/Source/Runtime/RenderCore/Public/RenderResource.h
--- a/Engine/Source/Runtime/RenderCore/Public/RenderResource.h
+++ b/Engine/Source/Runtime/RenderCore/Public/RenderResource.h
@@ -323,6 +323,11 @@ public:
 	 */
 	bool				bIgnoreGammaConversions;
 
+	/** 
+	 * Is the pixel data in this texture sRGB?
+	 **/
+	bool				bSRGB;
+
 	/** Default constructor. */
 	FTexture()
 	: TextureRHI(NULL)
@@ -331,6 +336,7 @@ public:
 	, LastRenderTime(-FLT_MAX)
 	, bGreyScaleFormat(false)
 	, bIgnoreGammaConversions(false)
+	, bSRGB(false)
 	{}
 
 	// Destructor
