Mark several classes as 'notplaceable' to remove them from the All Classes list
Improve comments on some classes
Remove TriggerBase/Box/Capsule/Sphere.h from Engine.h

diff --git a/Engine/Plugins/Runtime/CableComponent/Source/CableComponent/Classes/CableActor.h b/Engine/Plugins/Runtime/CableComponent/Source/CableComponent/Classes/CableActor.h
--- a/Engine/Plugins/Runtime/CableComponent/Source/CableComponent/Classes/CableActor.h
+++ b/Engine/Plugins/Runtime/CableComponent/Source/CableComponent/Classes/CableActor.h
@@ -4,11 +4,13 @@
 
 #include "CableActor.generated.h"
 
+/** An actor that renders a simulated cable */
 UCLASS(hidecategories=(Input,Collision,Replication), showcategories=("Input|MouseInput", "Input|TouchInput"))
 class ACableActor : public AActor
 {
 	GENERATED_UCLASS_BODY()
 
+	/** Cable component that performs simulation and rendering */
 	UPROPERTY(Category=Cable, VisibleAnywhere, BlueprintReadOnly)
 	TSubobjectPtr<class UCableComponent> CableComponent;
 };
diff --git a/Engine/Plugins/Script/ScriptPlugin/Source/ScriptPlugin/Classes/ScriptTestActor.h b/Engine/Plugins/Script/ScriptPlugin/Source/ScriptPlugin/Classes/ScriptTestActor.h
--- a/Engine/Plugins/Script/ScriptPlugin/Source/ScriptPlugin/Classes/ScriptTestActor.h
+++ b/Engine/Plugins/Script/ScriptPlugin/Source/ScriptPlugin/Classes/ScriptTestActor.h
@@ -5,7 +5,7 @@
 /**
 * Test Actor. To be removed at some point.
 */
-UCLASS(BlueprintType, MinimalAPI, EarlyAccessPreview)
+UCLASS(BlueprintType, MinimalAPI, EarlyAccessPreview, notplaceable)
 class AScriptTestActor : public AActor
 {
 	GENERATED_UCLASS_BODY()
diff --git a/Engine/Source/Developer/FunctionalTesting/Classes/FunctionalAITest.h b/Engine/Source/Developer/FunctionalTesting/Classes/FunctionalAITest.h
--- a/Engine/Source/Developer/FunctionalTesting/Classes/FunctionalAITest.h
+++ b/Engine/Source/Developer/FunctionalTesting/Classes/FunctionalAITest.h
@@ -94,7 +94,7 @@ struct FAITestSpawnSet
 	{}
 };
 
-UCLASS(Blueprintable, MinimalAPI)
+UCLASS(Blueprintable, MinimalAPI, notplaceable)
 class AFunctionalAITest : public AFunctionalTest
 {
 	GENERATED_UCLASS_BODY()
diff --git a/Engine/Source/Developer/FunctionalTesting/Classes/FunctionalTest.h b/Engine/Source/Developer/FunctionalTesting/Classes/FunctionalTest.h
--- a/Engine/Source/Developer/FunctionalTesting/Classes/FunctionalTest.h
+++ b/Engine/Source/Developer/FunctionalTesting/Classes/FunctionalTest.h
@@ -22,7 +22,7 @@ namespace EFunctionalTestResult
 DECLARE_DYNAMIC_MULTICAST_DELEGATE(FFunctionalTestEventSignature);
 DECLARE_DELEGATE_OneParam(FFunctionalTestDoneSignature, class AFunctionalTest*);
 
-UCLASS(Blueprintable, MinimalAPI)
+UCLASS(Blueprintable, MinimalAPI, notplaceable)
 class AFunctionalTest : public AActor
 {
 	GENERATED_UCLASS_BODY()
diff --git a/Engine/Source/Developer/GameplayDebugger/Classes/GameplayDebuggingReplicator.h b/Engine/Source/Developer/GameplayDebugger/Classes/GameplayDebuggingReplicator.h
--- a/Engine/Source/Developer/GameplayDebugger/Classes/GameplayDebuggingReplicator.h
+++ b/Engine/Source/Developer/GameplayDebugger/Classes/GameplayDebuggingReplicator.h
@@ -12,7 +12,7 @@ class AGameplayDebuggingHUDComponent;
 
 DECLARE_MULTICAST_DELEGATE_OneParam(FOnSelectionChanged, class AActor*);
 
-UCLASS(config = Engine, NotBlueprintable, Transient, hidecategories = Actor)
+UCLASS(config = Engine, NotBlueprintable, Transient, hidecategories = Actor, notplaceable)
 class GAMEPLAYDEBUGGER_API AGameplayDebuggingReplicator : public AActor
 {
 	GENERATED_UCLASS_BODY()
diff --git a/Engine/Source/Editor/LevelEditor/Private/LevelEditorCreateActorMenu.cpp b/Engine/Source/Editor/LevelEditor/Private/LevelEditorCreateActorMenu.cpp
--- a/Engine/Source/Editor/LevelEditor/Private/LevelEditorCreateActorMenu.cpp
+++ b/Engine/Source/Editor/LevelEditor/Private/LevelEditorCreateActorMenu.cpp
@@ -9,6 +9,7 @@
 #include "ClassIconFinder.h"
 #include "LevelEditorActions.h"
 #include "IPlacementModeModule.h"
+#include "Engine/TriggerBase.h"
 
 class SMenuThumbnail : public SCompoundWidget
 {
diff --git a/Engine/Source/Editor/UnrealEd/Private/Factories/ActorFactory.cpp b/Engine/Source/Editor/UnrealEd/Private/Factories/ActorFactory.cpp
--- a/Engine/Source/Editor/UnrealEd/Private/Factories/ActorFactory.cpp
+++ b/Engine/Source/Editor/UnrealEd/Private/Factories/ActorFactory.cpp
@@ -28,6 +28,10 @@ ActorFactory.cpp:
 
 #include "VectorField/VectorField.h"
 
+#include "Engine/TriggerBox.h"
+#include "Engine/TriggerSphere.h"
+#include "Engine/TriggerCapsule.h"
+
 DEFINE_LOG_CATEGORY_STATIC(LogActorFactory, Log, All);
 
 #define LOCTEXT_NAMESPACE "ActorFactory"
diff --git a/Engine/Source/Runtime/Engine/Classes/AI/Navigation/NavigationTestingActor.h b/Engine/Source/Runtime/Engine/Classes/AI/Navigation/NavigationTestingActor.h
--- a/Engine/Source/Runtime/Engine/Classes/AI/Navigation/NavigationTestingActor.h
+++ b/Engine/Source/Runtime/Engine/Classes/AI/Navigation/NavigationTestingActor.h
@@ -26,7 +26,7 @@ namespace ENavCostDisplay
 	};
 }
 
-UCLASS(hidecategories=(Object, Actor, Input, Rendering), showcategories=("Input|MouseInput", "Input|TouchInput"), Blueprintable)
+UCLASS(hidecategories=(Object, Actor, Input, Rendering), showcategories=("Input|MouseInput", "Input|TouchInput"), Blueprintable, notplaceable)
 class ENGINE_API ANavigationTestingActor : public AActor, public INavAgentInterface, public INavPathObserverInterface
 {
 	GENERATED_UCLASS_BODY()
diff --git a/Engine/Source/Runtime/Engine/Classes/AI/Navigation/RecastNavMesh.h b/Engine/Source/Runtime/Engine/Classes/AI/Navigation/RecastNavMesh.h
--- a/Engine/Source/Runtime/Engine/Classes/AI/Navigation/RecastNavMesh.h
+++ b/Engine/Source/Runtime/Engine/Classes/AI/Navigation/RecastNavMesh.h
@@ -410,7 +410,7 @@ namespace FNavMeshConfig
 }
 
 
-UCLASS(config=Engine, defaultconfig, hidecategories=(Input,Rendering,Tags,Transform,"Utilities|Transformation",Actor,Layers,Replication))
+UCLASS(config=Engine, defaultconfig, hidecategories=(Input,Rendering,Tags,Transform,"Utilities|Transformation",Actor,Layers,Replication), notplaceable)
 class ENGINE_API ARecastNavMesh : public ANavigationData
 {
 	GENERATED_UCLASS_BODY()
diff --git a/Engine/Source/Runtime/Engine/Classes/Atmosphere/AtmosphericFog.h b/Engine/Source/Runtime/Engine/Classes/Atmosphere/AtmosphericFog.h
--- a/Engine/Source/Runtime/Engine/Classes/Atmosphere/AtmosphericFog.h
+++ b/Engine/Source/Runtime/Engine/Classes/Atmosphere/AtmosphericFog.h
@@ -4,15 +4,21 @@
 #pragma once
 #include "AtmosphericFog.generated.h"
 
+/** 
+ *	A placeable fog actor that simulates atmospheric light scattering  
+ *	@see https://docs.unrealengine.com/latest/INT/Engine/Actors/FogEffects/AtmosphericFog/index.html
+ */
 UCLASS(showcategories=(Movement, Rendering, "Utilities|Transformation", "Input|MouseInput", "Input|TouchInput"), ClassGroup=Fog, hidecategories=(Info,Object,Input), MinimalAPI)
 class AAtmosphericFog : public AInfo
 {
 	GENERATED_UCLASS_BODY()
 
+	/** Main fog component */
 	UPROPERTY(BlueprintReadOnly, VisibleAnywhere, Category=Atmosphere)
 	TSubobjectPtr<class UAtmosphericFogComponent> AtmosphericFogComponent;
 
 #if WITH_EDITORONLY_DATA
+	/** Arrow component to indicate default sun rotation */
 	UPROPERTY()
 	TSubobjectPtr<class UArrowComponent> ArrowComponent;
 #endif
diff --git a/Engine/Source/Runtime/Engine/Classes/Camera/CameraActor.h b/Engine/Source/Runtime/Engine/Classes/Camera/CameraActor.h
--- a/Engine/Source/Runtime/Engine/Classes/Camera/CameraActor.h
+++ b/Engine/Source/Runtime/Engine/Classes/Camera/CameraActor.h
@@ -1,12 +1,12 @@
 // Copyright 1998-2014 Epic Games, Inc. All Rights Reserved.
 
-/*
- * CameraActor
- */
 
 #pragma once
 #include "CameraActor.generated.h"
 
+/** 
+ * A camera that can be placed in a level.
+ */
 UCLASS(ClassGroup=Common, hideCategories=(Input, Rendering), showcategories=("Input|MouseInput", "Input|TouchInput"), MinimalAPI, Blueprintable)
 class ACameraActor : public AActor
 {
@@ -14,12 +14,12 @@ class ACameraActor : public AActor
 
 private:
 
-	/* Specifies which player controller, if any, should automatically use this Camera when the controller is active. */
+	/** Specifies which player controller, if any, should automatically use this Camera when the controller is active. */
 	UPROPERTY(Category="AutoPlayerActivation", EditAnywhere)
 	TEnumAsByte<EAutoReceiveInput::Type> AutoActivateForPlayer;
 
 public:
-	// The camera component for this camera
+	/** The camera component for this camera */
 	UPROPERTY(Category=CameraActor, VisibleAnywhere, BlueprintReadOnly)
 	TSubobjectPtr<class UCameraComponent> CameraComponent;
 
diff --git a/Engine/Source/Runtime/Engine/Classes/Engine/BlockingVolume.h b/Engine/Source/Runtime/Engine/Classes/Engine/BlockingVolume.h
--- a/Engine/Source/Runtime/Engine/Classes/Engine/BlockingVolume.h
+++ b/Engine/Source/Runtime/Engine/Classes/Engine/BlockingVolume.h
@@ -1,14 +1,10 @@
 // Copyright 1998-2014 Epic Games, Inc. All Rights Reserved.
 
-//=============================================================================
-// BlockingVolume:  a bounding volume
-// used to block certain classes of actors
-// primary use is to provide collision for non-zero extent traces around static meshes
-//=============================================================================
 
 #pragma once
 #include "BlockingVolume.generated.h"
 
+/** An invisible volume used to block other actors */
 UCLASS(MinimalAPI)
 class ABlockingVolume : public AVolume
 {
diff --git a/Engine/Source/Runtime/Engine/Classes/Engine/BoxReflectionCapture.h b/Engine/Source/Runtime/Engine/Classes/Engine/BoxReflectionCapture.h
--- a/Engine/Source/Runtime/Engine/Classes/Engine/BoxReflectionCapture.h
+++ b/Engine/Source/Runtime/Engine/Classes/Engine/BoxReflectionCapture.h
@@ -1,13 +1,12 @@
 // Copyright 1998-2014 Epic Games, Inc. All Rights Reserved.
 
-/**
- * 
- *
- */
-
 #pragma once
 #include "BoxReflectionCapture.generated.h"
 
+/** 
+ *	Actor used to capture the scene for reflection in a box shape 
+ *	@see https://docs.unrealengine.com/latest/INT/Resources/ContentExamples/Reflections/1_3/index.html
+ */
 UCLASS(hidecategories=(Collision, Attachment, Actor), MinimalAPI)
 class ABoxReflectionCapture : public AReflectionCapture
 {
diff --git a/Engine/Source/Runtime/Engine/Classes/Engine/NiagaraActor.h b/Engine/Source/Runtime/Engine/Classes/Engine/NiagaraActor.h
--- a/Engine/Source/Runtime/Engine/Classes/Engine/NiagaraActor.h
+++ b/Engine/Source/Runtime/Engine/Classes/Engine/NiagaraActor.h
@@ -3,7 +3,7 @@
 #pragma once
 #include "NiagaraActor.generated.h"
 
-UCLASS()
+UCLASS(notplaceable)
 class ANiagaraActor : public AActor
 {
 	GENERATED_UCLASS_BODY()
diff --git a/Engine/Source/Runtime/Engine/Classes/Engine/SphereReflectionCapture.h b/Engine/Source/Runtime/Engine/Classes/Engine/SphereReflectionCapture.h
--- a/Engine/Source/Runtime/Engine/Classes/Engine/SphereReflectionCapture.h
+++ b/Engine/Source/Runtime/Engine/Classes/Engine/SphereReflectionCapture.h
@@ -1,23 +1,20 @@
 // Copyright 1998-2014 Epic Games, Inc. All Rights Reserved.
 
-/**
- * 
- *
- */
-
 #pragma once
 #include "SphereReflectionCapture.generated.h"
 
-//
-// Forward declarations.
-//
 class UDrawSphereComponent;
-UCLASS(hidecategories=(Collision, Attachment, Actor), MinimalAPI)
+
+/** 
+ *	Actor used to capture the scene for reflection in a sphere shape 
+ *	@see https://docs.unrealengine.com/latest/INT/Resources/ContentExamples/Reflections/1_4
+ */
+UCLASS(hidecategories = (Collision, Attachment, Actor), MinimalAPI)
 class ASphereReflectionCapture : public AReflectionCapture
 {
 	GENERATED_UCLASS_BODY()
 
-	// Reference to the draw capture radius component
+	/** Sphere component used to visualize the capture radius */
 	UPROPERTY()
 	TSubobjectPtr<UDrawSphereComponent> DrawCaptureRadius;
 
diff --git a/Engine/Source/Runtime/Engine/Classes/Engine/TriggerBase.h b/Engine/Source/Runtime/Engine/Classes/Engine/TriggerBase.h
--- a/Engine/Source/Runtime/Engine/Classes/Engine/TriggerBase.h
+++ b/Engine/Source/Runtime/Engine/Classes/Engine/TriggerBase.h
@@ -1,9 +1,5 @@
 // Copyright 1998-2014 Epic Games, Inc. All Rights Reserved.
 
-/**
- * An actor used to generate collision events (begin/end overlap) as inputs into the scripting system.
- */
-
 #pragma once
 #include "GameFramework/Actor.h"
 #include "TriggerBase.generated.h"
@@ -11,16 +7,17 @@
 class UShapeComponent;
 class UBillboardComponent;
 
+/** An actor used to generate collision events (begin/end overlap) in the level. */
 UCLASS(ClassGroup=Common, abstract, ConversionRoot, MinimalAPI)
 class ATriggerBase : public AActor
 {
 	GENERATED_UCLASS_BODY()
 
-	/** Base collision component */
+	/** Shape component used for collision */
 	UPROPERTY(Category=TriggerBase, VisibleAnywhere, BlueprintReadOnly)
 	TSubobjectPtr<UShapeComponent> CollisionComponent;
 
-	// Reference to the billboard component
+	/** Billboard used to see the trigger in the editor */
 	UPROPERTY()
 	TSubobjectPtr<UBillboardComponent> SpriteComponent;
 };
diff --git a/Engine/Source/Runtime/Engine/Classes/Engine/TriggerBox.h b/Engine/Source/Runtime/Engine/Classes/Engine/TriggerBox.h
--- a/Engine/Source/Runtime/Engine/Classes/Engine/TriggerBox.h
+++ b/Engine/Source/Runtime/Engine/Classes/Engine/TriggerBox.h
@@ -5,6 +5,7 @@
 #include "Engine/TriggerBase.h"
 #include "TriggerBox.generated.h"
 
+/** A box shaped trigger, used to generate overlap events in the level */
 UCLASS(MinimalAPI)
 class ATriggerBox : public ATriggerBase
 {
diff --git a/Engine/Source/Runtime/Engine/Classes/Engine/TriggerCapsule.h b/Engine/Source/Runtime/Engine/Classes/Engine/TriggerCapsule.h
--- a/Engine/Source/Runtime/Engine/Classes/Engine/TriggerCapsule.h
+++ b/Engine/Source/Runtime/Engine/Classes/Engine/TriggerCapsule.h
@@ -5,6 +5,7 @@
 #include "Engine/TriggerBase.h"
 #include "TriggerCapsule.generated.h"
 
+/** A capsule shaped trigger, used to generate overlap events in the level */
 UCLASS(MinimalAPI)
 class ATriggerCapsule : public ATriggerBase
 {
diff --git a/Engine/Source/Runtime/Engine/Classes/Engine/TriggerSphere.h b/Engine/Source/Runtime/Engine/Classes/Engine/TriggerSphere.h
--- a/Engine/Source/Runtime/Engine/Classes/Engine/TriggerSphere.h
+++ b/Engine/Source/Runtime/Engine/Classes/Engine/TriggerSphere.h
@@ -5,6 +5,7 @@
 #include "Engine/TriggerBase.h"
 #include "TriggerSphere.generated.h"
 
+/** A sphere shaped trigger, used to generate overlap events in the level */
 UCLASS(MinimalAPI)
 class ATriggerSphere : public ATriggerBase
 {
diff --git a/Engine/Source/Runtime/Engine/Classes/GameFramework/Volume.h b/Engine/Source/Runtime/Engine/Classes/GameFramework/Volume.h
--- a/Engine/Source/Runtime/Engine/Classes/GameFramework/Volume.h
+++ b/Engine/Source/Runtime/Engine/Classes/GameFramework/Volume.h
@@ -1,13 +1,13 @@
 // Copyright 1998-2014 Epic Games, Inc. All Rights Reserved.
 
-//=============================================================================
-// Volume:  a bounding volume
-//=============================================================================
-
 #pragma once
 #include "Engine/Brush.h"
 #include "Volume.generated.h"
 
+/** 
+ *	An editable 3D volume placed in a level. Different types of volumes perform different functions
+ *	@see https://docs.unrealengine.com/latest/INT/Engine/Actors/Volumes
+ */
 UCLASS(MinimalAPI, showcategories=Collision, hidecategories=(Brush, Physics), abstract, ConversionRoot)
 class AVolume : public ABrush
 {
diff --git a/Engine/Source/Runtime/Engine/Classes/Sound/AmbientSound.h b/Engine/Source/Runtime/Engine/Classes/Sound/AmbientSound.h
--- a/Engine/Source/Runtime/Engine/Classes/Sound/AmbientSound.h
+++ b/Engine/Source/Runtime/Engine/Classes/Sound/AmbientSound.h
@@ -1,25 +1,21 @@
 // Copyright 1998-2014 Epic Games, Inc. All Rights Reserved.
 
-
-// Base ambient sound actor
-
 #pragma once
 #include "AmbientSound.generated.h"
 
+/** A sound actor that can be placed in a level */
 UCLASS(AutoExpandCategories=Audio, ClassGroup=Sounds, MinimalAPI, hidecategories(Collision, Input, Game), showcategories=("Input|MouseInput", "Input|TouchInput", "Game|Damage"))
 class AAmbientSound : public AActor
 {
 	GENERATED_UCLASS_BODY()
 
-	/** Should the audio component automatically play on load? */
 	UPROPERTY()
 	uint32 bAutoPlay_DEPRECATED:1;
 
-	/** Audio component to play */
+	/** Audio component that handles sound playing */
 	UPROPERTY(Category=Sound, VisibleAnywhere, BlueprintReadOnly,meta=(ExposeFunctionCategories="Sound,Audio,Audio|Components|Audio"))
 	TSubobjectPtr<class UAudioComponent> AudioComponent;
 	
-	/** Dummy sound node property to force instantiation of subobject.	*/
 	UPROPERTY(instanced)
 	class UDEPRECATED_SoundNodeAmbient* SoundNodeInstance_DEPRECATED;
 
diff --git a/Engine/Source/Runtime/Engine/Public/Engine.h b/Engine/Source/Runtime/Engine/Public/Engine.h
--- a/Engine/Source/Runtime/Engine/Public/Engine.h
+++ b/Engine/Source/Runtime/Engine/Public/Engine.h
@@ -116,10 +116,6 @@
 #include "Engine/StaticMeshActor.h"
 #include "Engine/TargetPoint.h"
 #include "Engine/TextRenderActor.h"
-#include "Engine/TriggerBase.h"
-#include "Engine/TriggerBox.h"
-#include "Engine/TriggerCapsule.h"
-#include "Engine/TriggerSphere.h"
 #include "VectorField/VectorFieldVolume.h"
 #include "Components/ApplicationLifecycleComponent.h"
 #include "Engine/DataAsset.h"
diff --git a/Engine/Source/Runtime/GameplayAbilities/Public/Abilities/GameplayAbilityTargetActor.h b/Engine/Source/Runtime/GameplayAbilities/Public/Abilities/GameplayAbilityTargetActor.h
--- a/Engine/Source/Runtime/GameplayAbilities/Public/Abilities/GameplayAbilityTargetActor.h
+++ b/Engine/Source/Runtime/GameplayAbilities/Public/Abilities/GameplayAbilityTargetActor.h
@@ -5,7 +5,7 @@
 #include "GameplayAbilityTargetActor.generated.h"
 
 
-UCLASS(Blueprintable, abstract)
+UCLASS(Blueprintable, abstract, notplaceable)
 class GAMEPLAYABILITIES_API AGameplayAbilityTargetActor : public AActor
 {
 	GENERATED_UCLASS_BODY()
diff --git a/Engine/Source/Runtime/GameplayAbilities/Public/Abilities/GameplayAbilityWorldReticle.h b/Engine/Source/Runtime/GameplayAbilities/Public/Abilities/GameplayAbilityWorldReticle.h
--- a/Engine/Source/Runtime/GameplayAbilities/Public/Abilities/GameplayAbilityWorldReticle.h
+++ b/Engine/Source/Runtime/GameplayAbilities/Public/Abilities/GameplayAbilityWorldReticle.h
@@ -5,7 +5,7 @@
 
 class AGameplayAbilityTargetActor;
 
-UCLASS(Blueprintable)
+UCLASS(Blueprintable, notplaceable)
 class GAMEPLAYABILITIES_API AGameplayAbilityWorldReticle : public AActor
 {
 	GENERATED_UCLASS_BODY()
diff --git a/Engine/Source/Runtime/GameplayAbilities/Public/AbilitySystemTestPawn.h b/Engine/Source/Runtime/GameplayAbilities/Public/AbilitySystemTestPawn.h
--- a/Engine/Source/Runtime/GameplayAbilities/Public/AbilitySystemTestPawn.h
+++ b/Engine/Source/Runtime/GameplayAbilities/Public/AbilitySystemTestPawn.h
@@ -6,7 +6,7 @@
 #include "GameplayCueView.h"
 #include "AbilitySystemTestPawn.generated.h"
 
-UCLASS(Blueprintable, BlueprintType)
+UCLASS(Blueprintable, BlueprintType, notplaceable)
 class GAMEPLAYABILITIES_API AAbilitySystemTestPawn : public ADefaultPawn, public IGameplayCueInterface, public IAbilitySystemInterface
 {
 	GENERATED_UCLASS_BODY()
diff --git a/Engine/Source/Runtime/GameplayAbilities/Public/GameplayCueActor.h b/Engine/Source/Runtime/GameplayAbilities/Public/GameplayCueActor.h
--- a/Engine/Source/Runtime/GameplayAbilities/Public/GameplayCueActor.h
+++ b/Engine/Source/Runtime/GameplayAbilities/Public/GameplayCueActor.h
@@ -13,7 +13,7 @@
  */
 
 
-UCLASS(BlueprintType)
+UCLASS(BlueprintType, notplaceable)
 class GAMEPLAYABILITIES_API AGameplayCueActor : public AActor
 {
 	// Defines how a GameplayCue is translated into viewable components
diff --git a/Engine/Source/Runtime/Online/OnlineSubsystemUtils/Classes/OnlineBeaconHostObject.h b/Engine/Source/Runtime/Online/OnlineSubsystemUtils/Classes/OnlineBeaconHostObject.h
--- a/Engine/Source/Runtime/Online/OnlineSubsystemUtils/Classes/OnlineBeaconHostObject.h
+++ b/Engine/Source/Runtime/Online/OnlineSubsystemUtils/Classes/OnlineBeaconHostObject.h
@@ -16,7 +16,7 @@ class UNetConnection;
  * By defining a beacon type and implementing the ability to spawn unique AOnlineBeaconClients, any two instances of the engine
  * can communicate with each other without be officially connecting through normal Unreal networking
  */
-UCLASS(transient, config=Engine)
+UCLASS(transient, config=Engine, notplaceable)
 class ONLINESUBSYSTEMUTILS_API AOnlineBeaconHostObject : public AActor
 {
 	GENERATED_UCLASS_BODY()
