Fixed const correctness on HUD::Project and HUD::Deproject
(Github pull request https://github.com/EpicGames/UnrealEngine/pull/357)

diff --git a/Engine/Source/Runtime/Engine/Classes/Engine/Canvas.h b/Engine/Source/Runtime/Engine/Classes/Engine/Canvas.h
--- a/Engine/Source/Runtime/Engine/Classes/Engine/Canvas.h
+++ b/Engine/Source/Runtime/Engine/Classes/Engine/Canvas.h
@@ -348,7 +348,7 @@ public:
 	static void WrapString( FTextSizingParameters& Parameters, const float InCurX, const TCHAR* const pText, TArray<FWrappedStringElement>& out_Lines, FWordWrapper::FWrappedLineData* const OutWrappedLineData = nullptr);
 
 	/** Transforms a 3D world-space vector into 2D screen coordinates. */
-	FVector Project(FVector Location);
+	FVector Project(FVector Location) const;
 
 	/** 
 	 * Transforms 2D screen coordinates into a 3D world-space origin and direction.
@@ -356,7 +356,7 @@ public:
 	 * @param WorldOrigin (out) - world-space origin vector
 	 * @param WorldDirection (out) - world-space direction vector
 	 */
-	void Deproject(FVector2D ScreenPos, /*out*/ FVector& WorldOrigin, /*out*/ FVector& WorldDirection);
+	void Deproject(FVector2D ScreenPos, /*out*/ FVector& WorldOrigin, /*out*/ FVector& WorldDirection) const;
 
 	/**
 	 * Calculate the length of a string, taking text wrapping into account.
diff --git a/Engine/Source/Runtime/Engine/Classes/GameFramework/HUD.h b/Engine/Source/Runtime/Engine/Classes/GameFramework/HUD.h
--- a/Engine/Source/Runtime/Engine/Classes/GameFramework/HUD.h
+++ b/Engine/Source/Runtime/Engine/Classes/GameFramework/HUD.h
@@ -399,12 +399,12 @@ public:
 	void DrawMaterialSimple(UMaterialInterface* Material, float ScreenX, float ScreenY, float ScreenW, float ScreenH, float Scale=1.f, bool bScalePosition=false);
 
 	/** Transforms a 3D world-space vector into 2D screen coordinates */
-	UFUNCTION(BlueprintPure, Category=HUD)
-	FVector Project(FVector Location);
+	UFUNCTION(BlueprintCallable, Category = HUD)
+	FVector Project(FVector Location) const;
 	
 	/** Transforms a 2D screen location into a 3D location and direction */
-	UFUNCTION(BlueprintPure, Category=HUD)
-	void Deproject(float ScreenX, float ScreenY, FVector& WorldPosition, FVector& WorldDirection);
+	UFUNCTION(BlueprintCallable, Category = HUD)
+	void Deproject(float ScreenX, float ScreenY, FVector& WorldPosition, FVector& WorldDirection) const;
 
 
 	/**
diff --git a/Engine/Source/Runtime/Engine/Private/HUD.cpp b/Engine/Source/Runtime/Engine/Private/HUD.cpp
--- a/Engine/Source/Runtime/Engine/Private/HUD.cpp
+++ b/Engine/Source/Runtime/Engine/Private/HUD.cpp
@@ -625,7 +625,7 @@ void AHUD::DrawTextureSimple(UTexture* Texture, float ScreenX, float ScreenY, fl
 	}
 }
 
-FVector AHUD::Project(FVector Location)
+FVector AHUD::Project(FVector Location) const
 {
 	if (IsCanvasValid_WarnIfNot())
 	{
@@ -634,7 +634,7 @@ FVector AHUD::Project(FVector Location)
 	return FVector(0,0,0);
 }
 
-void AHUD::Deproject(float ScreenX, float ScreenY, FVector& WorldPosition, FVector& WorldDirection)
+void AHUD::Deproject(float ScreenX, float ScreenY, FVector& WorldPosition, FVector& WorldDirection) const
 {
 	WorldPosition = WorldDirection = FVector(0,0,0);
 	if (IsCanvasValid_WarnIfNot())
diff --git a/Engine/Source/Runtime/Engine/Private/UserInterface/Canvas.cpp b/Engine/Source/Runtime/Engine/Private/UserInterface/Canvas.cpp
--- a/Engine/Source/Runtime/Engine/Private/UserInterface/Canvas.cpp
+++ b/Engine/Source/Runtime/Engine/Private/UserInterface/Canvas.cpp
@@ -1491,7 +1491,7 @@ void UCanvas::TextSize(UFont* InFont, const FString& InText, float& XL, float& Y
 	YL = YLi;
 }
 
-FVector UCanvas::Project( FVector Location )
+FVector UCanvas::Project(FVector Location) const
 {
 	FPlane V(0,0,0,0);
 
@@ -1514,7 +1514,7 @@ FVector UCanvas::Project( FVector Location )
 	return resultVec;
 }
 
-void UCanvas::Deproject(FVector2D ScreenPos, /*out*/ FVector& WorldOrigin, /*out*/ FVector& WorldDirection)
+void UCanvas::Deproject(FVector2D ScreenPos, /*out*/ FVector& WorldOrigin, /*out*/ FVector& WorldDirection) const
 {
 	if (SceneView != NULL)
 	{
