Reverted the change making development-time AI-debugging actors not placable #UE4

diff --git a/Engine/Source/Developer/FunctionalTesting/Classes/FunctionalAITest.h b/Engine/Source/Developer/FunctionalTesting/Classes/FunctionalAITest.h
--- a/Engine/Source/Developer/FunctionalTesting/Classes/FunctionalAITest.h
+++ b/Engine/Source/Developer/FunctionalTesting/Classes/FunctionalAITest.h
@@ -94,7 +94,7 @@ struct FAITestSpawnSet
 	{}
 };
 
-UCLASS(Blueprintable, MinimalAPI, notplaceable)
+UCLASS(Blueprintable, MinimalAPI)
 class AFunctionalAITest : public AFunctionalTest
 {
 	GENERATED_UCLASS_BODY()
diff --git a/Engine/Source/Developer/FunctionalTesting/Classes/FunctionalTest.h b/Engine/Source/Developer/FunctionalTesting/Classes/FunctionalTest.h
--- a/Engine/Source/Developer/FunctionalTesting/Classes/FunctionalTest.h
+++ b/Engine/Source/Developer/FunctionalTesting/Classes/FunctionalTest.h
@@ -22,7 +22,7 @@ namespace EFunctionalTestResult
 DECLARE_DYNAMIC_MULTICAST_DELEGATE(FFunctionalTestEventSignature);
 DECLARE_DELEGATE_OneParam(FFunctionalTestDoneSignature, class AFunctionalTest*);
 
-UCLASS(Blueprintable, MinimalAPI, notplaceable)
+UCLASS(Blueprintable, MinimalAPI)
 class AFunctionalTest : public AActor
 {
 	GENERATED_UCLASS_BODY()
diff --git a/Engine/Source/Runtime/Engine/Classes/AI/Navigation/NavigationTestingActor.h b/Engine/Source/Runtime/Engine/Classes/AI/Navigation/NavigationTestingActor.h
--- a/Engine/Source/Runtime/Engine/Classes/AI/Navigation/NavigationTestingActor.h
+++ b/Engine/Source/Runtime/Engine/Classes/AI/Navigation/NavigationTestingActor.h
@@ -26,7 +26,7 @@ namespace ENavCostDisplay
 	};
 }
 
-UCLASS(hidecategories=(Object, Actor, Input, Rendering), showcategories=("Input|MouseInput", "Input|TouchInput"), Blueprintable, notplaceable)
+UCLASS(hidecategories=(Object, Actor, Input, Rendering), showcategories=("Input|MouseInput", "Input|TouchInput"), Blueprintable)
 class ENGINE_API ANavigationTestingActor : public AActor, public INavAgentInterface, public INavPathObserverInterface
 {
 	GENERATED_UCLASS_BODY()
