Core - TTP#COOKED: FPSChart commands have no affect, fixed by adding a proper path to the saved files

diff --git a/Engine/Source/Runtime/Core/Private/Misc/Core.cpp b/Engine/Source/Runtime/Core/Private/Misc/Core.cpp
--- a/Engine/Source/Runtime/Core/Private/Misc/Core.cpp
+++ b/Engine/Source/Runtime/Core/Private/Misc/Core.cpp
@@ -317,6 +317,7 @@ DEFINE_LOG_CATEGORY(LogContentComparisonCommandlet);
 DEFINE_LOG_CATEGORY(LogNetPackageMap);
 DEFINE_LOG_CATEGORY(LogNetSerialization);
 DEFINE_LOG_CATEGORY(LogMemory);
+DEFINE_LOG_CATEGORY(LogProfilingDebugging);
 
 DEFINE_LOG_CATEGORY(LogTemp);
 
diff --git a/Engine/Source/Runtime/Core/Private/ProfilingDebugging/ProfilingHelpers.cpp b/Engine/Source/Runtime/Core/Private/ProfilingDebugging/ProfilingHelpers.cpp
--- a/Engine/Source/Runtime/Core/Private/ProfilingDebugging/ProfilingHelpers.cpp
+++ b/Engine/Source/Runtime/Core/Private/ProfilingDebugging/ProfilingHelpers.cpp
@@ -52,9 +52,11 @@ int32 GetChangeListNumberForPerfTesting()
  **/
 void SendDataToPCViaUnrealConsole( const FString& NotifyType, const FString& FullFileName )
 {
-	//UE_LOG(LogProfilingDebugging, Warning, TEXT("SendDataToPCViaUnrealConsole %s%s"), *NotifyType, *FullFileName );
+	const FString AbsoluteFilename = IFileManager::Get().ConvertToAbsolutePathForExternalAppForRead( *FullFileName );
 
-	const FString NotifyString = NotifyType + FullFileName;
+	UE_LOG( LogProfilingDebugging, Warning, TEXT( "SendDataToPCViaUnrealConsole %s%s" ), *NotifyType, *AbsoluteFilename );
+
+	const FString NotifyString = NotifyType + AbsoluteFilename + LINE_TERMINATOR;
 	
 	// send it across via UnrealConsole
 	FMsg::SendNotificationString( *NotifyString );
diff --git a/Engine/Source/Runtime/Core/Public/HAL/OutputDevices.h b/Engine/Source/Runtime/Core/Public/HAL/OutputDevices.h
--- a/Engine/Source/Runtime/Core/Public/HAL/OutputDevices.h
+++ b/Engine/Source/Runtime/Core/Public/HAL/OutputDevices.h
@@ -535,6 +535,7 @@ CORE_API DECLARE_LOG_CATEGORY_EXTERN(LogContentComparisonCommandlet, Log, All);
 CORE_API DECLARE_LOG_CATEGORY_EXTERN(LogNetPackageMap, Warning, All);
 CORE_API DECLARE_LOG_CATEGORY_EXTERN(LogNetSerialization, Warning, All);
 CORE_API DECLARE_LOG_CATEGORY_EXTERN(LogMemory, Log, All);
+CORE_API DECLARE_LOG_CATEGORY_EXTERN(LogProfilingDebugging, Log, All);
 
 // Temporary log category, generally you should not check things in that use this
 CORE_API DECLARE_LOG_CATEGORY_EXTERN(LogTemp, Log, All);
diff --git a/Engine/Source/Runtime/Core/Public/ProfilingDebugging/MallocProfiler.h b/Engine/Source/Runtime/Core/Public/ProfilingDebugging/MallocProfiler.h
--- a/Engine/Source/Runtime/Core/Public/ProfilingDebugging/MallocProfiler.h
+++ b/Engine/Source/Runtime/Core/Public/ProfilingDebugging/MallocProfiler.h
@@ -8,8 +8,6 @@
 #include "UMemoryDefines.h"
 #include "CompressedGrowableBuffer.h"
 
-DECLARE_LOG_CATEGORY_EXTERN(LogProfilingDebugging, Log, All);
-
 
 #if USE_MALLOC_PROFILER
 
diff --git a/Engine/Source/Runtime/Engine/Private/ChartCreation.cpp b/Engine/Source/Runtime/Engine/Private/ChartCreation.cpp
--- a/Engine/Source/Runtime/Engine/Private/ChartCreation.cpp
+++ b/Engine/Source/Runtime/Engine/Private/ChartCreation.cpp
@@ -726,7 +726,7 @@ void UEngine::DumpFPSChartToStatsLog( float TotalTime, float DeltaTime, int32 Nu
 		// Flush, close and delete.
 		delete OutputFile;
 
-		SendDataToPCViaUnrealConsole( TEXT("UE_PROFILER!FPSCHART:"), *(ChartName) );
+		UE_LOG( LogProfilingDebugging, Warning, TEXT( "FPS Chart (logfile) saved to %s" ), *IFileManager::Get().ConvertToAbsolutePathForExternalAppForRead( *ChartName ) );
 	}
 #endif // ALLOW_DEBUG_FILES
 }
@@ -871,7 +871,7 @@ void UEngine::DumpFPSChartToHTML( float TotalTime, float DeltaTime, int32 NumFra
 		// Save the resulting file back to disk.
 		FFileHelper::SaveStringToFile( FPSChart, *FPSChartFilename );
 
-		SendDataToPCViaUnrealConsole( TEXT("UE_PROFILER!FPSCHART:"), *(FPSChartFilename) );
+		UE_LOG( LogProfilingDebugging, Warning, TEXT( "FPS Chart (HTML) saved to %s" ), *IFileManager::Get().ConvertToAbsolutePathForExternalAppForRead( *FPSChartFilename ) );
 	}
 	else
 	{
