More robust code in EdGraph

diff --git a/Engine/Source/Runtime/Engine/Private/EdGraph/EdGraphNode.cpp b/Engine/Source/Runtime/Engine/Private/EdGraph/EdGraphNode.cpp
--- a/Engine/Source/Runtime/Engine/Private/EdGraph/EdGraphNode.cpp
+++ b/Engine/Source/Runtime/Engine/Private/EdGraph/EdGraphNode.cpp
@@ -159,7 +159,8 @@ void UEdGraphNode::DestroyNode()
 
 const class UEdGraphSchema* UEdGraphNode::GetSchema() const
 {
-	return GetGraph()->GetSchema();
+	UEdGraph* ParentGraph = GetGraph();
+	return ParentGraph ? ParentGraph->GetSchema() : NULL;
 }
 
 bool UEdGraphNode::IsCompatibleWithGraph(UEdGraph const* Graph) const
diff --git a/Engine/Source/Runtime/Engine/Private/EdGraph/EdGraphPin.cpp b/Engine/Source/Runtime/Engine/Private/EdGraph/EdGraphPin.cpp
--- a/Engine/Source/Runtime/Engine/Private/EdGraph/EdGraphPin.cpp
+++ b/Engine/Source/Runtime/Engine/Private/EdGraph/EdGraphPin.cpp
@@ -191,7 +191,8 @@ void UEdGraphPin::CopyPersistentDataFromOldPin(const UEdGraphPin& SourcePin)
 const class UEdGraphSchema* UEdGraphPin::GetSchema() const
 {
 #if WITH_EDITOR
-	return GetOwningNode()->GetGraph()->GetSchema();
+	auto OwnerNode = GetOwningNodeUnchecked();
+	return OwnerNode ? OwnerNode->GetSchema() : NULL;
 #else
 	return NULL;
 #endif	//WITH_EDITOR
@@ -211,7 +212,17 @@ FString UEdGraphPin::GetDefaultAsString() const
 
 FText UEdGraphPin::GetDisplayName() const
 {
-	return FText::FromString(GetSchema()->GetPinDisplayName(this));
+	FString StrName;
+	auto Schema = GetSchema();
+	if (Schema)
+	{
+		StrName = Schema->GetPinDisplayName(this);
+	}
+	else
+	{
+		StrName = (!PinFriendlyName.IsEmpty()) ? PinFriendlyName.ToString() : PinName;
+	}
+	return FText::FromString(StrName);
 }
 
 const FString UEdGraphPin::GetLinkInfoString( const FString& InFunctionName, const FString& InInfoData, const UEdGraphPin* InToPin ) const
