#ue4 - Make ABrush, AVolume, and APhysics volume ENGINE_API to export all functions.

TTP 345281

diff --git a/Engine/Source/Runtime/Engine/Classes/Engine/Brush.h b/Engine/Source/Runtime/Engine/Classes/Engine/Brush.h
--- a/Engine/Source/Runtime/Engine/Classes/Engine/Brush.h
+++ b/Engine/Source/Runtime/Engine/Classes/Engine/Brush.h
@@ -72,8 +72,8 @@ struct FGeomSelection
 	}
 };
 
-UCLASS(hidecategories=(Object, Collision, Display, Rendering, Physics, Input), showcategories=("Input|MouseInput", "Input|TouchInput"), MinimalAPI, NotBlueprintable, ConversionRoot)
-class ABrush : public AActor
+UCLASS(hidecategories=(Object, Collision, Display, Rendering, Physics, Input), showcategories=("Input|MouseInput", "Input|TouchInput"), NotBlueprintable, ConversionRoot)
+class ENGINE_API ABrush : public AActor
 {
 	GENERATED_UCLASS_BODY()
 
@@ -132,7 +132,7 @@ class ABrush : public AActor
 	/** Delegate used for notifications when PostRegisterAllComponents is called for a Brush */
 	DECLARE_MULTICAST_DELEGATE_OneParam(FOnBrushRegistered, ABrush*);
 	/** Function to get the 'brush registered' delegate */
-	static ENGINE_API FOnBrushRegistered& GetOnBrushRegisteredDelegate()
+	static FOnBrushRegistered& GetOnBrushRegisteredDelegate()
 	{
 		return OnBrushRegistered;
 	}
@@ -140,12 +140,12 @@ class ABrush : public AActor
 
 	
 	// Begin UObject interface.
-	ENGINE_API virtual void PostLoad() override;
+	virtual void PostLoad() override;
 #if WITH_EDITOR
-	ENGINE_API virtual void PostEditMove(bool bFinished) override;
-	ENGINE_API virtual void PostEditChangeProperty(FPropertyChangedEvent& PropertyChangedEvent) override;
+	virtual void PostEditMove(bool bFinished) override;
+	virtual void PostEditChangeProperty(FPropertyChangedEvent& PropertyChangedEvent) override;
 #endif // WITH_EDITOR
-	ENGINE_API virtual bool Modify(bool bAlwaysMarkDirty = false) override;
+	virtual bool Modify(bool bAlwaysMarkDirty = false) override;
 
 	virtual bool NeedsLoadForClient() const override
 	{ 
@@ -159,27 +159,27 @@ class ABrush : public AActor
 
 	
 	// Begin AActor interface
-	ENGINE_API virtual void Destroyed() override;
-	ENGINE_API virtual void PostRegisterAllComponents() override;
+	virtual void Destroyed() override;
+	virtual void PostRegisterAllComponents() override;
 	
 	virtual bool IsLevelBoundsRelevant() const override;
 
 #if WITH_EDITOR
 	virtual void CheckForErrors() override;
 
-	ENGINE_API virtual void SetIsTemporarilyHiddenInEditor( bool bIsHidden ) override;
+	virtual void SetIsTemporarilyHiddenInEditor( bool bIsHidden ) override;
 
 	// End AActor interface
 
-	ENGINE_API virtual void InitPosRotScale();
-	ENGINE_API virtual void CopyPosRotScaleFrom( ABrush* Other );
+	virtual void InitPosRotScale();
+	virtual void CopyPosRotScaleFrom( ABrush* Other );
 
 private: 
 	/** An array to keep track of all the levels that need rebuilding. This is checked via NeedsRebuild() in the editor tick and triggers a csg rebuild. */
 	static TArray< TWeakObjectPtr< ULevel > > LevelsToRebuild;
 
 	/** Delegate called when PostRegisterAllComponents is called for a Brush */
-	static ENGINE_API FOnBrushRegistered OnBrushRegistered;
+	static FOnBrushRegistered OnBrushRegistered;
 public:
 	/**
 	 * Called to see if any of the levels need rebuilding
@@ -187,21 +187,21 @@ public:
 	 * @param	OutLevels if specified, provides a copy of the levels array
 	 * @return	true if the csg needs to be rebuilt on the next editor tick.	
 	 */
-	ENGINE_API static bool NeedsRebuild(TArray< TWeakObjectPtr< ULevel > >* OutLevels = NULL){if(OutLevels){*OutLevels = LevelsToRebuild;} return(LevelsToRebuild.Num() > 0);}
+	static bool NeedsRebuild(TArray< TWeakObjectPtr< ULevel > >* OutLevels = NULL){if(OutLevels){*OutLevels = LevelsToRebuild;} return(LevelsToRebuild.Num() > 0);}
 	/**
 	 * Called upon finishing the csg rebuild to clear the rebuild bool.
 	 */
-	ENGINE_API static void OnRebuildDone(){LevelsToRebuild.Empty();}
+	static void OnRebuildDone(){LevelsToRebuild.Empty();}
 	/**
 	 * Called to make not of the level that needs rebuilding
 	 *
 	 * @param	InLevel The level that needs rebuilding
 	 */
-	ENGINE_API static void SetNeedRebuild(ULevel* InLevel){if(InLevel){LevelsToRebuild.AddUnique(InLevel);}}
+	static void SetNeedRebuild(ULevel* InLevel){if(InLevel){LevelsToRebuild.AddUnique(InLevel);}}
 #else
-	ENGINE_API static bool NeedsRebuild(TArray< TWeakObjectPtr< ULevel > >* OutLevels = NULL){return false;}
-	ENGINE_API static void OnRebuildDone(){}
-	ENGINE_API static void SetNeedRebuild(ULevel* InLevel){}
+	static bool NeedsRebuild(TArray< TWeakObjectPtr< ULevel > >* OutLevels = NULL){return false;}
+	static void OnRebuildDone(){}
+	static void SetNeedRebuild(ULevel* InLevel){}
 #endif//WITH_EDITOR
 
 	/** @return true if this is a static brush */
@@ -216,13 +216,13 @@ public:
 	// ABrush interface.
 
 	/** Returns the prepivot FVector of the RootComponent of this actor */
-	ENGINE_API FVector GetPrePivot() const;
+	FVector GetPrePivot() const;
 
 	/** Set the PrePivot of the root component */
-	ENGINE_API void SetPrePivot(const FVector& InPrePivot);
+	void SetPrePivot(const FVector& InPrePivot);
 
 	/** Figures out the best color to use for this brushes wireframe drawing.	*/
-	ENGINE_API virtual FColor GetWireColor() const;
+	virtual FColor GetWireColor() const;
 
 
 	/**
diff --git a/Engine/Source/Runtime/Engine/Classes/GameFramework/DefaultPhysicsVolume.h b/Engine/Source/Runtime/Engine/Classes/GameFramework/DefaultPhysicsVolume.h
--- a/Engine/Source/Runtime/Engine/Classes/GameFramework/DefaultPhysicsVolume.h
+++ b/Engine/Source/Runtime/Engine/Classes/GameFramework/DefaultPhysicsVolume.h
@@ -15,7 +15,7 @@ class ADefaultPhysicsVolume : public APhysicsVolume
 	GENERATED_UCLASS_BODY()
 
 	// Begin AActor interface
-	virtual void Destroyed() override;
+	ENGINE_API virtual void Destroyed() override;
 	// End AActor interface
 };
 
diff --git a/Engine/Source/Runtime/Engine/Classes/GameFramework/PhysicsVolume.h b/Engine/Source/Runtime/Engine/Classes/GameFramework/PhysicsVolume.h
--- a/Engine/Source/Runtime/Engine/Classes/GameFramework/PhysicsVolume.h
+++ b/Engine/Source/Runtime/Engine/Classes/GameFramework/PhysicsVolume.h
@@ -9,14 +9,14 @@
  * Each AActor is affected at any time by one PhysicsVolume
  */
 
-UCLASS(MinimalAPI)
-class APhysicsVolume : public AVolume
+UCLASS()
+class ENGINE_API APhysicsVolume : public AVolume
 {
 	GENERATED_UCLASS_BODY()
 
 	// Begin UObject interface.
 #if WITH_EDITOR
-	ENGINE_API virtual void LoadedFromAnotherClass(const FName& OldClassName) override;
+	virtual void LoadedFromAnotherClass(const FName& OldClassName) override;
 #endif // WITH_EDITOR	
 	// End UObject interface.
 
@@ -48,14 +48,14 @@ class APhysicsVolume : public AVolume
 
 public:
 	/** @Returns the Z component of the current world gravity. */
-	ENGINE_API virtual float GetGravityZ() const;
+	virtual float GetGravityZ() const;
 	
 	// Called when actor enters a volume
-	ENGINE_API virtual void ActorEnteredVolume(class AActor* Other);
+	virtual void ActorEnteredVolume(class AActor* Other);
 	
 	// Called when actor leaves a volume, Other can be NULL
-	ENGINE_API virtual void ActorLeavingVolume(class AActor* Other);
+	virtual void ActorLeavingVolume(class AActor* Other);
 
 	/** Given a known overlap with the given component, validate that it meets the rules imposed by bPhysicsOnContact. */
-	ENGINE_API virtual bool IsOverlapInVolume(const class USceneComponent& TestComponent) const;
+	virtual bool IsOverlapInVolume(const class USceneComponent& TestComponent) const;
 };
diff --git a/Engine/Source/Runtime/Engine/Classes/GameFramework/Volume.h b/Engine/Source/Runtime/Engine/Classes/GameFramework/Volume.h
--- a/Engine/Source/Runtime/Engine/Classes/GameFramework/Volume.h
+++ b/Engine/Source/Runtime/Engine/Classes/GameFramework/Volume.h
@@ -8,8 +8,8 @@
  *	An editable 3D volume placed in a level. Different types of volumes perform different functions
  *	@see https://docs.unrealengine.com/latest/INT/Engine/Actors/Volumes
  */
-UCLASS(MinimalAPI, showcategories=Collision, hidecategories=(Brush, Physics), abstract, ConversionRoot)
-class AVolume : public ABrush
+UCLASS(showcategories=Collision, hidecategories=(Brush, Physics), abstract, ConversionRoot)
+class ENGINE_API AVolume : public ABrush
 {
 	GENERATED_UCLASS_BODY()
 
@@ -18,7 +18,7 @@ class AVolume : public ABrush
 	DECLARE_MULTICAST_DELEGATE_OneParam(FOnVolumeShapeChanged, AVolume& )
 
 	/** Function to get the 'Volume imported' delegate */
-	static ENGINE_API FOnVolumeShapeChanged& GetOnVolumeShapeChangedDelegate()
+	static FOnVolumeShapeChanged& GetOnVolumeShapeChangedDelegate()
 	{
 		return OnVolumeShapeChanged;
 	}
@@ -31,24 +31,24 @@ public:
 	* Function that gets called from within Map_Check to allow this actor to check itself
 	* for any potential errors and register them with map check dialog.
 	*/
-	ENGINE_API virtual void CheckForErrors() override;
+	virtual void CheckForErrors() override;
 #endif // WITH_EDITOR
 	
-	ENGINE_API virtual bool IsLevelBoundsRelevant() const override;
+	virtual bool IsLevelBoundsRelevant() const override;
 	// End AActor Interface
 
 	// Begin Brush Interface
-	ENGINE_API virtual bool IsStaticBrush() const override;
-	ENGINE_API virtual bool IsVolumeBrush() const override;
+	virtual bool IsStaticBrush() const override;
+	virtual bool IsVolumeBrush() const override;
 	// End Brush Interface
 
 	/** @returns true if a sphere/point (with optional radius CheckRadius) overlaps this volume */
-	ENGINE_API bool EncompassesPoint(FVector Point, float SphereRadius=0.f, float* OutDistanceToPoint = 0);
+	bool EncompassesPoint(FVector Point, float SphereRadius=0.f, float* OutDistanceToPoint = 0);
 
 	//Begin UObject Interface
 #if WITH_EDITOR
-	ENGINE_API virtual void PostEditImport() override;
-	ENGINE_API virtual void PostEditChangeProperty(FPropertyChangedEvent& PropertyChangedEvent) override;
+	virtual void PostEditImport() override;
+	virtual void PostEditChangeProperty(FPropertyChangedEvent& PropertyChangedEvent) override;
 #endif // WITH_EDITOR
 	//End UObject Interface
 
