#ue4 - SpectatorPawn is no longer placeable in the editor.

diff --git a/Engine/Source/Runtime/Engine/Classes/GameFramework/SpectatorPawn.h b/Engine/Source/Runtime/Engine/Classes/GameFramework/SpectatorPawn.h
--- a/Engine/Source/Runtime/Engine/Classes/GameFramework/SpectatorPawn.h
+++ b/Engine/Source/Runtime/Engine/Classes/GameFramework/SpectatorPawn.h
@@ -11,7 +11,7 @@
 #include "SpectatorPawn.generated.h"
 
 
-UCLASS(config=Game, Blueprintable, BlueprintType)
+UCLASS(config=Game, Blueprintable, BlueprintType, notplaceable)
 class ENGINE_API ASpectatorPawn : public ADefaultPawn
 {
 	GENERATED_UCLASS_BODY()
