Routing hit events for welded bodies

diff --git a/Engine/Source/Runtime/Engine/Classes/PhysicsEngine/BodyInstance.h b/Engine/Source/Runtime/Engine/Classes/PhysicsEngine/BodyInstance.h
--- a/Engine/Source/Runtime/Engine/Classes/PhysicsEngine/BodyInstance.h
+++ b/Engine/Source/Runtime/Engine/Classes/PhysicsEngine/BodyInstance.h
@@ -22,6 +22,7 @@ namespace physx
 	class PxGeometry;
 	class PxShape;
 	class PxTransform;
+	struct PxContactPair;
 }
 #endif // WITH_PHYSX
 
@@ -45,6 +46,8 @@ namespace ELockedAxis
 	};
 }
 
+struct FCollisionNotifyInfo;
+
 USTRUCT()
 struct ENGINE_API FCollisionResponse
 {
@@ -211,6 +214,11 @@ public:
 	/** The parent body that we are welded to*/
 	FBodyInstance * WeldParent;
 
+#if WITH_PHYSX
+	/** Figures out the new FCollisionNotifyInfo needed for pending notification. It adds it, and then returns an array that goes from pair index to notify collision index */
+	static TArray<int32> AddCollisionNotifyInfo(const FBodyInstance * Body0, const FBodyInstance * Body1, const physx::PxContactPair * Pairs, uint32 NumPairs, TArray<FCollisionNotifyInfo> & PendingNotifyInfos);
+#endif
+
 protected:
 
 	/** Whether this instance of the object has its own custom walkability override setting. */
diff --git a/Engine/Source/Runtime/Engine/Private/PhysicsEngine/BodyInstance.cpp b/Engine/Source/Runtime/Engine/Private/PhysicsEngine/BodyInstance.cpp
--- a/Engine/Source/Runtime/Engine/Private/PhysicsEngine/BodyInstance.cpp
+++ b/Engine/Source/Runtime/Engine/Private/PhysicsEngine/BodyInstance.cpp
@@ -1345,6 +1345,57 @@ void FBodyInstance::InitBody(UBodySetup* Setup, const FTransform& Transform, UPr
 }
 #endif // UE_WITH_PHYSICS
 
+#if WITH_PHYSX
+TArray<int32> FBodyInstance::AddCollisionNotifyInfo(const FBodyInstance * Body0, const FBodyInstance * Body1, const physx::PxContactPair * Pairs, uint32 NumPairs, TArray<FCollisionNotifyInfo> & PendingNotifyInfos)
+{
+	TArray<int32> PairNotifyMapping;
+	PairNotifyMapping.Empty(NumPairs);
+
+	TMap<const FBodyInstance *, TMap<const FBodyInstance *, int32> > BodyPairNotifyMap;
+	for (uint32 PairIdx = 0; PairIdx < NumPairs; ++PairIdx)
+	{
+		const PxContactPair* Pair = Pairs + PairIdx;
+		// Get the two shapes that are involved in the collision
+		const PxShape* Shape0 = Pair->shapes[0];
+		check(Shape0);
+		const PxShape* Shape1 = Pair->shapes[1];
+		check(Shape1);
+
+		PairNotifyMapping.Add(-1);	//start as -1 because we can have collisions that we don't want to actually record collision
+
+		const FBodyInstance * const* SubBody0Res = Body0->ShapeToBodyMap.Find(Shape0);
+		const FBodyInstance * const* SubBody1Res = Body1->ShapeToBodyMap.Find(Shape1);
+
+		if (SubBody0Res == NULL){ SubBody0Res = &Body0;  }	//body may have no welding in which case just use the body
+		if (SubBody1Res == NULL){ SubBody1Res = &Body1; }
+		
+		const FBodyInstance * SubBody0 = *SubBody0Res;
+		const FBodyInstance * SubBody1 = *SubBody1Res;
+
+		if (SubBody0->bNotifyRigidBodyCollision || SubBody1->bNotifyRigidBodyCollision)
+		{
+			TMap<const FBodyInstance *, int32> & SubBodyNotifyMap = BodyPairNotifyMap.FindOrAdd(SubBody0);
+			int32 * NotifyInfoIndex = SubBodyNotifyMap.Find(SubBody1);
+
+			if (NotifyInfoIndex == NULL)
+			{
+				FCollisionNotifyInfo * NotifyInfo = new (PendingNotifyInfos) FCollisionNotifyInfo;
+				NotifyInfo->bCallEvent0 = (SubBody0->bNotifyRigidBodyCollision);
+				NotifyInfo->Info0.SetFrom(SubBody0);
+				NotifyInfo->bCallEvent1 = (SubBody1->bNotifyRigidBodyCollision);
+				NotifyInfo->Info1.SetFrom(SubBody1);
+
+				NotifyInfoIndex = &SubBodyNotifyMap.Add(SubBody0, PendingNotifyInfos.Num() - 1);
+			}
+
+			PairNotifyMapping[PairIdx] = *NotifyInfoIndex;
+		}
+	}
+
+	return PairNotifyMapping;
+}
+#endif
+
 //helper function for TermBody to avoid code duplication between scenes
 #if WITH_PHYSX
 void TermBodyHelper(int32& SceneIndex, PxRigidActor*& PRigidActor, FBodyInstance* BodyInstance)
diff --git a/Engine/Source/Runtime/Engine/Private/PhysicsEngine/PhysXSupport.cpp b/Engine/Source/Runtime/Engine/Private/PhysicsEngine/PhysXSupport.cpp
--- a/Engine/Source/Runtime/Engine/Private/PhysicsEngine/PhysXSupport.cpp
+++ b/Engine/Source/Runtime/Engine/Private/PhysicsEngine/PhysXSupport.cpp
@@ -446,18 +446,21 @@ void FPhysXSimEventCallback::onContact(const PxContactPairHeader& PairHeader, co
 	FPhysScene* PhysScene = FPhysxUserData::Get<FPhysScene>(PScene->userData);
 	check(PhysScene);
 
-	// Allocate and fill in the NotifyInfo struct.
-	FCollisionNotifyInfo* NotifyInfo = new(PhysScene->PendingCollisionNotifies) FCollisionNotifyInfo;
-	NotifyInfo->bCallEvent0 = (BodyInst0->bNotifyRigidBodyCollision);
-	NotifyInfo->Info0.SetFrom(BodyInst0);
-	NotifyInfo->bCallEvent1 = (BodyInst1->bNotifyRigidBodyCollision);
-	NotifyInfo->Info1.SetFrom(BodyInst1);
-
-	FCollisionImpactData* ImpactInfo = &(NotifyInfo->RigidCollisionData);
+	uint32 PreAddingCollisionNotify = PhysScene->PendingCollisionNotifies.Num() - 1;
+	TArray<int32> PairNotifyMapping = FBodyInstance::AddCollisionNotifyInfo(BodyInst0, BodyInst1, Pairs, NumPairs, PhysScene->PendingCollisionNotifies);
 
 	// Iterate through contact points
 	for(uint32 PairIdx=0; PairIdx<NumPairs; PairIdx++)
 	{
+		int32 NotifyIdx = PairNotifyMapping[PairIdx];
+		if (NotifyIdx == -1)	//the body instance this pair belongs to is not listening for events
+		{
+			continue;
+		}
+
+		FCollisionNotifyInfo * NotifyInfo = &PhysScene->PendingCollisionNotifies[NotifyIdx];
+		FCollisionImpactData* ImpactInfo = &(NotifyInfo->RigidCollisionData);
+
 		const PxContactPair* Pair = Pairs + PairIdx;
 
 		// Get the two shapes that are involved in the collision
@@ -495,11 +498,16 @@ void FPhysXSimEventCallback::onContact(const PxContactPairHeader& PairHeader, co
 		}	
 	}
 
-	// Discard pairs that don't generate any force (eg. have been rejected through a modify contact callback).
-	if(ImpactInfo->TotalNormalImpulse.SizeSquared() < KINDA_SMALL_NUMBER)
+	for (int32 NotifyIdx = PreAddingCollisionNotify + 1; NotifyIdx < PhysScene->PendingCollisionNotifies.Num(); NotifyIdx++)
 	{
-		PhysScene->PendingCollisionNotifies.RemoveAt(PhysScene->PendingCollisionNotifies.Num()-1);
-		return;
+		FCollisionNotifyInfo * NotifyInfo = &PhysScene->PendingCollisionNotifies[NotifyIdx];
+		FCollisionImpactData* ImpactInfo = &(NotifyInfo->RigidCollisionData);
+		// Discard pairs that don't generate any force (eg. have been rejected through a modify contact callback).
+		if (ImpactInfo->TotalNormalImpulse.SizeSquared() < KINDA_SMALL_NUMBER)
+		{
+			PhysScene->PendingCollisionNotifies.RemoveAt(NotifyIdx);
+			NotifyIdx--;
+		}
 	}
 }
 
