Renderer/UMG - The material interface now allows the Used with UI flag to be set using the flags enum.

diff --git a/Engine/Source/Runtime/Engine/Classes/Materials/MaterialInterface.h b/Engine/Source/Runtime/Engine/Classes/Materials/MaterialInterface.h
--- a/Engine/Source/Runtime/Engine/Classes/Materials/MaterialInterface.h
+++ b/Engine/Source/Runtime/Engine/Classes/Materials/MaterialInterface.h
@@ -27,6 +27,7 @@ enum EMaterialUsage
 	MATUSAGE_Landscape,
 	MATUSAGE_InstancedStaticMeshes,
 	MATUSAGE_Clothing,
+	MATUSAGE_UI,
 	MATUSAGE_MAX,
 };
 
diff --git a/Engine/Source/Runtime/Engine/Private/Materials/Material.cpp b/Engine/Source/Runtime/Engine/Private/Materials/Material.cpp
--- a/Engine/Source/Runtime/Engine/Private/Materials/Material.cpp
+++ b/Engine/Source/Runtime/Engine/Private/Materials/Material.cpp
@@ -701,6 +701,7 @@ bool UMaterial::GetUsageByFlag(EMaterialUsage Usage) const
 		case MATUSAGE_SplineMesh: UsageValue = bUsedWithSplineMeshes; break;
 		case MATUSAGE_InstancedStaticMeshes: UsageValue = bUsedWithInstancedStaticMeshes; break;
 		case MATUSAGE_Clothing: UsageValue = bUsedWithClothing; break;
+		case MATUSAGE_UI: UsageValue = bUsedWithUI; break;
 		default: UE_LOG(LogMaterial, Fatal,TEXT("Unknown material usage: %u"), (int32)Usage);
 	};
 	return UsageValue;
@@ -783,6 +784,10 @@ void UMaterial::SetUsageByFlag(EMaterialUsage Usage, bool NewValue)
 		{
 			bUsedWithClothing = NewValue; break;
 		}
+		case MATUSAGE_UI:
+		{
+			bUsedWithUI = NewValue; break;
+		}
 		default: UE_LOG(LogMaterial, Fatal,TEXT("Unknown material usage: %u"), (int32)Usage);
 	};
 #if WITH_EDITOR
@@ -806,6 +811,7 @@ FString UMaterial::GetUsageName(EMaterialUsage Usage) const
 		case MATUSAGE_SplineMesh: UsageName = TEXT("bUsedWithSplineMeshes"); break;
 		case MATUSAGE_InstancedStaticMeshes: UsageName = TEXT("bUsedWithInstancedStaticMeshes"); break;
 		case MATUSAGE_Clothing: UsageName = TEXT("bUsedWithClothing"); break;
+		case MATUSAGE_UI: UsageName = TEXT("bUsedWithUI"); break;
 		default: UE_LOG(LogMaterial, Fatal,TEXT("Unknown material usage: %u"), (int32)Usage);
 	};
 	return UsageName;
