Distinguish between root wanting hit event and it turning hit event on for welded child

diff --git a/Engine/Source/Runtime/Engine/Classes/PhysicsEngine/BodyInstance.h b/Engine/Source/Runtime/Engine/Classes/PhysicsEngine/BodyInstance.h
--- a/Engine/Source/Runtime/Engine/Classes/PhysicsEngine/BodyInstance.h
+++ b/Engine/Source/Runtime/Engine/Classes/PhysicsEngine/BodyInstance.h
@@ -178,6 +178,9 @@ public:
 	/** determines if the body is currently welded */
 	uint32 bWelded : 1;
 
+	/** needed to distinguish between the root wanting hit events, and the root turning hit events on for one of its welded children */
+	uint32 bWeldedNotifyRigidBodyCollision : 1;
+
 	/** If object should start awake, or if it should initially be sleeping */
 	UPROPERTY(EditAnywhere, BlueprintReadOnly, Category=Physics, meta=(editcondition = "bSimulatePhysics"))
 	uint32 bStartAwake:1;
diff --git a/Engine/Source/Runtime/Engine/Private/PhysicsEngine/BodyInstance.cpp b/Engine/Source/Runtime/Engine/Private/PhysicsEngine/BodyInstance.cpp
--- a/Engine/Source/Runtime/Engine/Private/PhysicsEngine/BodyInstance.cpp
+++ b/Engine/Source/Runtime/Engine/Private/PhysicsEngine/BodyInstance.cpp
@@ -199,6 +199,7 @@ FBodyInstance::FBodyInstance()
 , bSimulatePhysics(false)
 , bAutoWeld(false)
 , bWelded(false)
+, bWeldedNotifyRigidBodyCollision(false)
 , bStartAwake(true)
 , bEnableGravity(true)
 , bUseAsyncScene(false)
@@ -625,7 +626,7 @@ void FBodyInstance::UpdatePhysicsFilterData(bool bForceSimpleAsComplex)
 
 	// Grab collision setting from body instance
 	TEnumAsByte<ECollisionEnabled::Type> UseCollisionEnabled = GetCollisionEnabled(); // this checks actor override
-	bool bUseNotifyRBCollision = bNotifyRigidBodyCollision;
+	bool bUseNotifyRBCollision = bNotifyRigidBodyCollision || bWeldedNotifyRigidBodyCollision;
 	FCollisionResponseContainer UseResponse = CollisionResponses.GetResponseContainer();
 
 	// Get skelmeshcomp ID
@@ -649,7 +650,7 @@ void FBodyInstance::UpdatePhysicsFilterData(bool bForceSimpleAsComplex)
 		}
 
 		UseResponse = FCollisionResponseContainer::CreateMinContainer(UseResponse, SkelMeshComp->BodyInstance.CollisionResponses.GetResponseContainer());
-		bUseNotifyRBCollision = bUseNotifyRBCollision && SkelMeshComp->BodyInstance.bNotifyRigidBodyCollision;
+		bUseNotifyRBCollision = bUseNotifyRBCollision && (SkelMeshComp->BodyInstance.bNotifyRigidBodyCollision || SkelMeshComp->BodyInstance.bWeldedNotifyRigidBodyCollision);
 	}
 
 #if WITH_EDITOR
@@ -1571,6 +1572,11 @@ bool FBodyInstance::Weld(FBodyInstance* TheirBody, const FTransform& TheirTM)
 		BI = TheirBody;
 	}
 
+	if (TheirBody->bNotifyRigidBodyCollision)
+	{
+		bWeldedNotifyRigidBodyCollision = true;
+	}
+
 	PostShapeChange();
 
 	//remove their body from scenes
@@ -1592,11 +1598,15 @@ void FBodyInstance::UnWeld(FBodyInstance* BI)
 	int32 NumSyncShapes = 0;
 	TArray<PxShape *> PShapes = GetAllShapes(NumSyncShapes);
 
+	bool bNeedsNotification = false;
+
 	for (int32 ShapeIdx = 0; ShapeIdx < NumSyncShapes; ++ShapeIdx)
 	{
 		PxShape* PShape = PShapes[ShapeIdx];
 		if (FBodyInstance ** BIPtrPtr = ShapeToBodyMap.Find(PShape))
 		{
+			bNeedsNotification |= (*BIPtrPtr)->bNotifyRigidBodyCollision;
+
 			if (*BIPtrPtr == BI)
 			{
 				RigidActorSync->detachShape(*PShape);
@@ -1610,6 +1620,7 @@ void FBodyInstance::UnWeld(FBodyInstance* BI)
 		PxShape* PShape = PShapes[ShapeIdx];
 		if (FBodyInstance ** BIPtrPtr = ShapeToBodyMap.Find(PShape))
 		{
+			bNeedsNotification |= (*BIPtrPtr)->bNotifyRigidBodyCollision;
 			if (*BIPtrPtr == BI)
 			{
 				RigidActorAsync->detachShape(*PShape);
@@ -1618,6 +1629,8 @@ void FBodyInstance::UnWeld(FBodyInstance* BI)
 		}
 	}
 
+	bWeldedNotifyRigidBodyCollision = bNeedsNotification;
+
 	PostShapeChange();
 #endif
 }
