Make UpdateSimulatedPosition virtual (UT team/SteveP request)

diff --git a/Engine/Source/Runtime/Engine/Classes/GameFramework/Character.h b/Engine/Source/Runtime/Engine/Classes/GameFramework/Character.h
--- a/Engine/Source/Runtime/Engine/Classes/GameFramework/Character.h
+++ b/Engine/Source/Runtime/Engine/Classes/GameFramework/Character.h
@@ -629,7 +629,7 @@ public:
 	bool RestoreReplicatedMove(const FSimulatedRootMotionReplicatedMove & RootMotionRepMove);
 	
 	/** Called on client after position update is received to actually move the character. */
-	void UpdateSimulatedPosition(const FVector & NewLocation, const FRotator & NewRotation);
+	virtual void UpdateSimulatedPosition(const FVector & NewLocation, const FRotator & NewRotation);
 
 	/** Replicated Root Motion montage */
 	UPROPERTY(ReplicatedUsing=OnRep_RootMotion)
