#ue4 - Add keywords for MakeRot / BreakRot so that typing "MakeRotator" / "BreakRotator" shows the functions.

diff --git a/Engine/Source/Runtime/Engine/Classes/Kismet/KismetMathLibrary.h b/Engine/Source/Runtime/Engine/Classes/Kismet/KismetMathLibrary.h
--- a/Engine/Source/Runtime/Engine/Classes/Kismet/KismetMathLibrary.h
+++ b/Engine/Source/Runtime/Engine/Classes/Kismet/KismetMathLibrary.h
@@ -1012,7 +1012,7 @@ class UKismetMathLibrary : public UBlueprintFunctionLibrary
 	static FVector GetUpVector(FRotator InRot);
 
 	/** Makes a rotator {Pitch, Yaw, Roll} */
-	UFUNCTION(BlueprintPure, Category="Math|Rotator", meta=(Keywords="construct build rotation rotate", NativeMakeFunc))
+	UFUNCTION(BlueprintPure, Category="Math|Rotator", meta=(Keywords="construct build rotation rotate rotator makerotator", NativeMakeFunc))
 	static FRotator MakeRot(float Pitch, float Yaw, float Roll);
 	
 	/** Find a rotation for an object at Start location to point at Target location. */
@@ -1020,47 +1020,47 @@ class UKismetMathLibrary : public UBlueprintFunctionLibrary
 	static FRotator FindLookAtRotation(const FVector& Start, const FVector& Target);
 
 	/** Builds a rotator given only a XAxis. Y and Z are unspecified but will be orthonormal. XAxis need not be normalized. */
-	UFUNCTION(BlueprintPure, Category="Math|Rotator", meta=(Keywords="construct build rotation rotate"))
+	UFUNCTION(BlueprintPure, Category="Math|Rotator", meta=(Keywords="construct build rotation rotate rotator makerotator"))
 	static FRotator MakeRotFromX(const FVector& X);
 
 	/** Builds a rotation matrix given only a YAxis. X and Z are unspecified but will be orthonormal. YAxis need not be normalized. */
-	UFUNCTION(BlueprintPure, Category="Math|Rotator", meta=(Keywords="construct build rotation rotate"))
+	UFUNCTION(BlueprintPure, Category="Math|Rotator", meta=(Keywords="construct build rotation rotate rotator makerotator"))
 	static FRotator MakeRotFromY(const FVector& Y);
 
 	/** Builds a rotation matrix given only a ZAxis. X and Y are unspecified but will be orthonormal. ZAxis need not be normalized. */
-	UFUNCTION(BlueprintPure, Category="Math|Rotator", meta=(Keywords="construct build rotation rotate"))
+	UFUNCTION(BlueprintPure, Category="Math|Rotator", meta=(Keywords="construct build rotation rotate rotator makerotator"))
 	static FRotator MakeRotFromZ(const FVector& Z);
 
 	/** Builds a matrix with given X and Y axes. X will remain fixed, Y may be changed minimally to enforce orthogonality. Z will be computed. Inputs need not be normalized. */
-	UFUNCTION(BlueprintPure, Category="Math|Rotator", meta=(Keywords="construct build rotation rotate"))
+	UFUNCTION(BlueprintPure, Category="Math|Rotator", meta=(Keywords="construct build rotation rotate rotator makerotator"))
 	static FRotator MakeRotFromXY(const FVector& X, const FVector& Y);
 
 	/** Builds a matrix with given X and Z axes. X will remain fixed, Z may be changed minimally to enforce orthogonality. Y will be computed. Inputs need not be normalized. */
-	UFUNCTION(BlueprintPure, Category="Math|Rotator", meta=(Keywords="construct build rotation rotate"))
+	UFUNCTION(BlueprintPure, Category="Math|Rotator", meta=(Keywords="construct build rotation rotate rotator makerotator"))
 	static FRotator MakeRotFromXZ(const FVector& X, const FVector& Z);
 
 	/** Builds a matrix with given Y and X axes. Y will remain fixed, X may be changed minimally to enforce orthogonality. Z will be computed. Inputs need not be normalized. */
-	UFUNCTION(BlueprintPure, Category="Math|Rotator", meta=(Keywords="construct build rotation rotate"))
+	UFUNCTION(BlueprintPure, Category="Math|Rotator", meta=(Keywords="construct build rotation rotate rotator makerotator"))
 	static FRotator MakeRotFromYX(const FVector& Y, const FVector& X);
 
 	/** Builds a matrix with given Y and Z axes. Y will remain fixed, Z may be changed minimally to enforce orthogonality. X will be computed. Inputs need not be normalized. */
-	UFUNCTION(BlueprintPure, Category="Math|Rotator", meta=(Keywords="construct build rotation rotate"))
+	UFUNCTION(BlueprintPure, Category="Math|Rotator", meta=(Keywords="construct build rotation rotate rotator makerotator"))
 	static FRotator MakeRotFromYZ(const FVector& Y, const FVector& Z);
 
 	/** Builds a matrix with given Z and X axes. Z will remain fixed, X may be changed minimally to enforce orthogonality. Y will be computed. Inputs need not be normalized. */
-	UFUNCTION(BlueprintPure, Category="Math|Rotator", meta=(Keywords="construct build rotation rotate"))
+	UFUNCTION(BlueprintPure, Category="Math|Rotator", meta=(Keywords="construct build rotation rotate rotator makerotator"))
 	static FRotator MakeRotFromZX(const FVector& Z, const FVector& X);
 
 	/** Builds a matrix with given Z and Y axes. Z will remain fixed, Y may be changed minimally to enforce orthogonality. X will be computed. Inputs need not be normalized. */
-	UFUNCTION(BlueprintPure, Category="Math|Rotator", meta=(Keywords="construct build rotation rotate"))
+	UFUNCTION(BlueprintPure, Category="Math|Rotator", meta=(Keywords="construct build rotation rotate rotator makerotator"))
 	static FRotator MakeRotFromZY(const FVector& Z, const FVector& Y);
 
 	/** Breaks apart a rotator into Pitch, Yaw, Roll */
-	UFUNCTION(BlueprintPure, Category="Math|Rotator", meta=(Keywords="rotation rotate", NativeBreakFunc))
+	UFUNCTION(BlueprintPure, Category="Math|Rotator", meta=(Keywords="rotation rotate rotator breakrotator", NativeBreakFunc))
 	static void BreakRot(FRotator InRot, float& Pitch, float& Yaw, float& Roll);
 
 	/** Breaks apart a rotator into its component axes */
-	UFUNCTION(BlueprintPure, Category="Math|Rotator", meta=(Keywords="rotation rotate"))
+	UFUNCTION(BlueprintPure, Category="Math|Rotator", meta=(Keywords="rotation rotate rotator breakrotator"))
 	static void BreakRotIntoAxes(const FRotator& InRot, FVector& X, FVector& Y, FVector& Z);
 
 	/** Make a transform from location, rotation and scale */
