Allow multiple keys to be mapped for opening console

diff --git a/Engine/Config/BaseInput.ini b/Engine/Config/BaseInput.ini
--- a/Engine/Config/BaseInput.ini
+++ b/Engine/Config/BaseInput.ini
@@ -8,7 +8,7 @@ DefaultTouchInterface=/Engine/MobileResources/HUD/DefaultVirtualJoysticks.Defaul
 bAlwaysShowTouchInterface=false
 bShowConsoleOnFourFingerTap=true
 
-ConsoleKey=Tilde
++ConsoleKeys=Tilde
 bRequireCtrlToNavigateAutoComplete=False
 
 ;-----------------------------------------------------------------------------------------
diff --git a/Engine/Source/Runtime/Engine/Classes/GameFramework/InputSettings.h b/Engine/Source/Runtime/Engine/Classes/GameFramework/InputSettings.h
--- a/Engine/Source/Runtime/Engine/Classes/GameFramework/InputSettings.h
+++ b/Engine/Source/Runtime/Engine/Classes/GameFramework/InputSettings.h
@@ -56,9 +56,12 @@ class ENGINE_API UInputSettings
 	FStringAssetReference DefaultTouchInterface;
 
 	/** The key which opens the console. */
-	UPROPERTY(config, EditAnywhere, Category="Console")
-	FKey ConsoleKey;
+	UPROPERTY(config)
+	FKey ConsoleKey_DEPRECATED;
 
+	/** The keys which open the console. */
+	UPROPERTY(config, EditAnywhere, Category="Console")
+	TArray<FKey> ConsoleKeys;
 
 	// UObject interface
 #if WITH_EDITOR
diff --git a/Engine/Source/Runtime/Engine/Private/UserInterface/Console.cpp b/Engine/Source/Runtime/Engine/Private/UserInterface/Console.cpp
--- a/Engine/Source/Runtime/Engine/Private/UserInterface/Console.cpp
+++ b/Engine/Source/Runtime/Engine/Private/UserInterface/Console.cpp
@@ -551,7 +551,7 @@ bool UConsole::InputKey_InputLine( int32 ControllerId, FKey Key, EInputEvent Eve
 	bool bModifierDown = bCtrl;
 	FModifierKeysState KeyState = FSlateApplication::Get().GetModifierKeys();
 	bModifierDown |= KeyState.IsAltDown() || KeyState.IsCommandDown() || KeyState.IsShiftDown() || KeyState.IsControlDown();
-	if ( Key==GetDefault<UInputSettings>()->ConsoleKey && Event == IE_Pressed && !bModifierDown )
+	if ( GetDefault<UInputSettings>()->ConsoleKeys.Contains(Key) && Event == IE_Pressed && !bModifierDown )
 	{
 		if (ConsoleState == NAME_Typing)
 		{
diff --git a/Engine/Source/Runtime/Engine/Private/UserInterface/InputSettings.cpp b/Engine/Source/Runtime/Engine/Private/UserInterface/InputSettings.cpp
--- a/Engine/Source/Runtime/Engine/Private/UserInterface/InputSettings.cpp
+++ b/Engine/Source/Runtime/Engine/Private/UserInterface/InputSettings.cpp
@@ -15,6 +15,12 @@ void UInputSettings::PostInitProperties()
 {
 	Super::PostInitProperties();
 
+	if (ConsoleKey_DEPRECATED.IsValid())
+	{
+		ConsoleKeys.Empty(1);
+		ConsoleKeys.Add(ConsoleKey_DEPRECATED);
+	}
+
 	TMap<FName, int32> UniqueAxisConfigNames;
 	for (int32 Index = 0; Index < AxisConfig.Num(); ++Index)
 	{
