Refactor welding to embed welded body in shape data.
Refactor GetBodyInstance to automatically resolve to welded body
Fix scene queries to automatically route to original component

diff --git a/Engine/Source/Runtime/Engine/Classes/Components/DestructibleComponent.h b/Engine/Source/Runtime/Engine/Classes/Components/DestructibleComponent.h
--- a/Engine/Source/Runtime/Engine/Classes/Components/DestructibleComponent.h
+++ b/Engine/Source/Runtime/Engine/Classes/Components/DestructibleComponent.h
@@ -122,7 +122,7 @@ public:
 	// End UActorComponent interface.
 
 	// Begin UPrimitiveComponent interface.
-	virtual FBodyInstance* GetBodyInstance(FName BoneName = NAME_None) const override;
+	virtual FBodyInstance* GetBodyInstance(FName BoneName = NAME_None, bool bGetWelded = true) const override;
 	virtual bool CanEditSimulatePhysics() override;
 	virtual bool IsAnySimulatingPhysics() const override;
 
diff --git a/Engine/Source/Runtime/Engine/Classes/Components/PrimitiveComponent.h b/Engine/Source/Runtime/Engine/Classes/Components/PrimitiveComponent.h
--- a/Engine/Source/Runtime/Engine/Classes/Components/PrimitiveComponent.h
+++ b/Engine/Source/Runtime/Engine/Classes/Components/PrimitiveComponent.h
@@ -981,7 +981,18 @@ public:
 	virtual int32 GetNumMaterials() const;
 	
 	/** Get a BodyInstance from this component. The supplied name is used in the SkeletalMeshComponent case. A name of NAME_None in the skeletal case gives the root body instance. */
-	virtual FBodyInstance* GetBodyInstance(FName BoneName = NAME_None) const;
+
+
+	/**
+	* returns BodyInstance of the component.
+	*
+	* @param BoneName				Used to get body associated with specific bone. NAME_None automatically gets the root most body
+	* @param bGetWelded				If the component has been welded to another component and bGetWelded is true we return the single welded BodyInstance that is used in the simulation
+	*
+	* @return		Returns the BodyInstance based on various states (does component have multiple bodies? Is the body welded to another body?)
+	*/
+
+	virtual FBodyInstance* GetBodyInstance(FName BoneName = NAME_None, bool bGetWelded = true) const;
 
 	/** 
 	 * returns Distance to closest Body Instance surface. 
diff --git a/Engine/Source/Runtime/Engine/Classes/Components/SkeletalMeshComponent.h b/Engine/Source/Runtime/Engine/Classes/Components/SkeletalMeshComponent.h
--- a/Engine/Source/Runtime/Engine/Classes/Components/SkeletalMeshComponent.h
+++ b/Engine/Source/Runtime/Engine/Classes/Components/SkeletalMeshComponent.h
@@ -808,7 +808,7 @@ public:
 	// Begin UPrimitiveComponent interface.
 	virtual class UBodySetup* GetBodySetup() override;
 	virtual bool CanEditSimulatePhysics() override;
-	virtual FBodyInstance* GetBodyInstance(FName BoneName = NAME_None) const override;
+	virtual FBodyInstance* GetBodyInstance(FName BoneName = NAME_None, bool bGetWelded = true) const override;
 	virtual void UpdatePhysicsToRBChannels() override;
 	virtual void SetAllPhysicsAngularVelocity(FVector const& NewVel, bool bAddToCurrent = false) override;
 	virtual void SetAllPhysicsPosition(FVector NewPos) override;
diff --git a/Engine/Source/Runtime/Engine/Classes/PhysicsEngine/BodyInstance.h b/Engine/Source/Runtime/Engine/Classes/PhysicsEngine/BodyInstance.h
--- a/Engine/Source/Runtime/Engine/Classes/PhysicsEngine/BodyInstance.h
+++ b/Engine/Source/Runtime/Engine/Classes/PhysicsEngine/BodyInstance.h
@@ -220,6 +220,7 @@ public:
 #if WITH_PHYSX
 	/** Figures out the new FCollisionNotifyInfo needed for pending notification. It adds it, and then returns an array that goes from pair index to notify collision index */
 	static TArray<int32> AddCollisionNotifyInfo(const FBodyInstance * Body0, const FBodyInstance * Body1, const physx::PxContactPair * Pairs, uint32 NumPairs, TArray<FCollisionNotifyInfo> & PendingNotifyInfos);
+
 #endif
 
 protected:
@@ -334,9 +335,6 @@ public:
 
 
 
-#if WITH_PHYSX
-	TMap<physx::PxShape*, FBodyInstance*> ShapeToBodyMap;
-#endif
 
 	/** 
 	 * Takes two body instances and welds them together to create a single simulated rigid body. Returns true if success.
diff --git a/Engine/Source/Runtime/Engine/Private/DestructibleComponent.cpp b/Engine/Source/Runtime/Engine/Private/DestructibleComponent.cpp
--- a/Engine/Source/Runtime/Engine/Private/DestructibleComponent.cpp
+++ b/Engine/Source/Runtime/Engine/Private/DestructibleComponent.cpp
@@ -1108,7 +1108,7 @@ void UDestructibleComponent::ResetFakeBodyInstance( FFakeBodyInstanceState& Prev
 }
 #endif
 
-FBodyInstance* UDestructibleComponent::GetBodyInstance( FName BoneName /*= NAME_None*/ ) const
+FBodyInstance* UDestructibleComponent::GetBodyInstance( FName BoneName /*= NAME_None*/, bool) const
 {
 #if WITH_APEX
 	if (ApexDestructibleActor != NULL)
diff --git a/Engine/Source/Runtime/Engine/Private/PhysicsEngine/BodyInstance.cpp b/Engine/Source/Runtime/Engine/Private/PhysicsEngine/BodyInstance.cpp
--- a/Engine/Source/Runtime/Engine/Private/PhysicsEngine/BodyInstance.cpp
+++ b/Engine/Source/Runtime/Engine/Private/PhysicsEngine/BodyInstance.cpp
@@ -849,6 +849,10 @@ void FBodyInstance::InitBody(UBodySetup* Setup, const FTransform& Transform, UPr
 	check(Setup);
 	AActor* Owner = PrimComp ? PrimComp->GetOwner() : NULL;
 
+#if WITH_PHYSX
+	PhysxUserData = FPhysxUserData(this);
+#endif
+
 	// Make the debug name for this geometry...
 	FString DebugName(TEXT(""));
 #if (WITH_EDITORONLY_DATA || UE_BUILD_DEBUG || LOOKING_FOR_PERF_ISSUES) && !(UE_BUILD_SHIPPING || UE_BUILD_TEST) && !NO_LOGGING
@@ -1131,7 +1135,6 @@ void FBodyInstance::InitBody(UBodySetup* Setup, const FTransform& Transform, UPr
 #endif
 
 #if WITH_PHYSX
-	PhysxUserData = FPhysxUserData(this);
 
 	// If there is already a body instanced, or there is no scene to create it into, do nothing.
 	if (GetPxRigidActor() != NULL || !InRBScene)
@@ -1364,15 +1367,12 @@ TArray<int32> FBodyInstance::AddCollisionNotifyInfo(const FBodyInstance * Body0,
 
 		PairNotifyMapping.Add(-1);	//start as -1 because we can have collisions that we don't want to actually record collision
 
-		const FBodyInstance * const* SubBody0Res = Body0->ShapeToBodyMap.Find(Shape0);
-		const FBodyInstance * const* SubBody1Res = Body1->ShapeToBodyMap.Find(Shape1);
+		const FBodyInstance * SubBody0 = FPhysxUserData::Get<FBodyInstance>(Shape0->userData);
+		const FBodyInstance * SubBody1 = FPhysxUserData::Get<FBodyInstance>(Shape1->userData);
 
-		if (SubBody0Res == NULL){ SubBody0Res = &Body0;  }	//body may have no welding in which case just use the body
-		if (SubBody1Res == NULL){ SubBody1Res = &Body1; }
+		if (SubBody0 == NULL) { SubBody0 = Body0; }
+		if (SubBody1 == NULL) { SubBody1 = Body1; }
 		
-		const FBodyInstance * SubBody0 = *SubBody0Res;
-		const FBodyInstance * SubBody1 = *SubBody1Res;
-
 		if (SubBody0->bNotifyRigidBodyCollision || SubBody1->bNotifyRigidBodyCollision)
 		{
 			TMap<const FBodyInstance *, int32> & SubBodyNotifyMap = BodyPairNotifyMap.FindOrAdd(SubBody0);
@@ -1395,10 +1395,9 @@ TArray<int32> FBodyInstance::AddCollisionNotifyInfo(const FBodyInstance * Body0,
 
 	return PairNotifyMapping;
 }
-#endif
+
 
 //helper function for TermBody to avoid code duplication between scenes
-#if WITH_PHYSX
 void TermBodyHelper(int32& SceneIndex, PxRigidActor*& PRigidActor, FBodyInstance* BodyInstance)
 {
 	if (SceneIndex)
@@ -1533,7 +1532,7 @@ bool FBodyInstance::Weld(FBodyInstance* TheirBody, const FTransform& TheirTM)
 #if WITH_PHYSX
 	
 	//first time we're welding with root so we need to grab the shapes
-	if (ShapeToBodyMap.Num() == 0)
+	/*if (ShapeToBodyMap.Num() == 0)
 	{
 		int32 NumSyncShapes = 0;
 		TArray<PxShape *> PShapes = GetAllShapes(NumSyncShapes);
@@ -1544,7 +1543,7 @@ bool FBodyInstance::Weld(FBodyInstance* TheirBody, const FTransform& TheirTM)
 			FBodyInstance *& BI = ShapeToBodyMap.FindOrAdd(PShape);
 			BI = this;
 		}
-	}
+	}*/
 
 	TArray<PxShape *> PNewShapes;
 
@@ -1568,8 +1567,10 @@ bool FBodyInstance::Weld(FBodyInstance* TheirBody, const FTransform& TheirTM)
 	for (int32 ShapeIdx = 0; ShapeIdx < PNewShapes.Num(); ++ShapeIdx)
 	{
 		PxShape* PShape = PNewShapes[ShapeIdx];
-		FBodyInstance *& BI = ShapeToBodyMap.FindOrAdd(PShape);
-		BI = TheirBody;
+		//FBodyInstance *& BI = ShapeToBodyMap.FindOrAdd(PShape);
+		//BI = TheirBody;
+
+		PShape->userData = &TheirBody->PhysxUserData;
 	}
 
 	if (TheirBody->bNotifyRigidBodyCollision)
@@ -1589,7 +1590,7 @@ bool FBodyInstance::Weld(FBodyInstance* TheirBody, const FTransform& TheirTM)
 	return false;
 }
 
-void FBodyInstance::UnWeld(FBodyInstance* BI)
+void FBodyInstance::UnWeld(FBodyInstance* TheirBI)
 {
 	//@TODO: BOX2D: Implement Weld
 
@@ -1601,6 +1602,37 @@ void FBodyInstance::UnWeld(FBodyInstance* BI)
 	bool bNeedsNotification = false;
 
 	for (int32 ShapeIdx = 0; ShapeIdx < NumSyncShapes; ++ShapeIdx)
+	{
+		PxShape* PShape = PShapes[ShapeIdx];
+		if (FBodyInstance * BI = FPhysxUserData::Get<FBodyInstance>(PShape->userData))
+		{
+			bNeedsNotification |= BI->bNotifyRigidBodyCollision;
+
+			if (TheirBI == BI)
+			{
+				PShape->userData = NULL;
+				RigidActorSync->detachShape(*PShape);
+			}
+		}
+	}
+
+	for (int32 ShapeIdx = NumSyncShapes; ShapeIdx <PShapes.Num(); ++ShapeIdx)
+	{
+		PxShape* PShape = PShapes[ShapeIdx];
+		if (FBodyInstance * BI = FPhysxUserData::Get<FBodyInstance>(PShape->userData))
+		{
+			bNeedsNotification |= BI->bNotifyRigidBodyCollision;
+
+			if (TheirBI == BI)
+			{
+				PShape->userData = NULL;
+				RigidActorSync->detachShape(*PShape);
+			}
+		}
+	}
+
+
+	/*for (int32 ShapeIdx = 0; ShapeIdx < NumSyncShapes; ++ShapeIdx)
 	{
 		PxShape* PShape = PShapes[ShapeIdx];
 		if (FBodyInstance ** BIPtrPtr = ShapeToBodyMap.Find(PShape))
@@ -1609,6 +1641,7 @@ void FBodyInstance::UnWeld(FBodyInstance* BI)
 
 			if (*BIPtrPtr == BI)
 			{
+				PShape->userData = NULL;
 				RigidActorSync->detachShape(*PShape);
 				ShapeToBodyMap.Remove(PShape);
 			}
@@ -1623,11 +1656,12 @@ void FBodyInstance::UnWeld(FBodyInstance* BI)
 			bNeedsNotification |= (*BIPtrPtr)->bNotifyRigidBodyCollision;
 			if (*BIPtrPtr == BI)
 			{
+				PShape->userData = NULL;
 				RigidActorAsync->detachShape(*PShape);
 				ShapeToBodyMap.Remove(PShape);
 			}
 		}
-	}
+	}*/
 
 	bWeldedNotifyRigidBodyCollision = bNeedsNotification;
 
diff --git a/Engine/Source/Runtime/Engine/Private/PrimitiveComponent.cpp b/Engine/Source/Runtime/Engine/Private/PrimitiveComponent.cpp
--- a/Engine/Source/Runtime/Engine/Private/PrimitiveComponent.cpp
+++ b/Engine/Source/Runtime/Engine/Private/PrimitiveComponent.cpp
@@ -329,188 +329,6 @@ void UPrimitiveComponent::OnChildAttached(USceneComponent* ChildComponent)
 	Super::OnChildAttached(ChildComponent);
 }
 
-UPrimitiveComponent * GetRootWelded(UPrimitiveComponent * PrimComponent, FName ParentSocketName = NAME_None, FName * OutSocketName = NULL, bool bAboutToWeld = false)
-{
-	UPrimitiveComponent * Result = NULL;
-	UPrimitiveComponent * RootComponent = Cast<UPrimitiveComponent>(PrimComponent->AttachParent);	//we must find the root component along hierarchy that has bWelded set to true
-
-	//check that body itself is welded
-	if (FBodyInstance * BI = PrimComponent->GetBodyInstance(ParentSocketName))
-	{
-		if (bAboutToWeld == false && BI->bWelded == false && BI->bAutoWeld == false)	//we're not welded and we aren't trying to become welded
-		{
-			return NULL;
-		}
-	}
-
-	FName PrevSocketName = ParentSocketName;
-	FName SocketName = NAME_None; //because of skeletal mesh it's important that we check along the bones that we attached
-	FBodyInstance * RootBI = NULL;
-	for (; RootComponent; RootComponent = Cast<UPrimitiveComponent>(RootComponent->AttachParent))
-	{
-		Result = RootComponent;
-		SocketName = PrevSocketName;
-		PrevSocketName = RootComponent->AttachSocketName;
-
-		RootBI = RootComponent->GetBodyInstance(SocketName);
-		if (RootBI && RootBI->bWelded)
-		{
-			continue;
-		}
-
-		break;
-	}
-
-	if (OutSocketName)
-	{
-		*OutSocketName = SocketName;
-	}
-
-	return Result;
-}
-
-void UPrimitiveComponent::GetWeldedBodies(TArray<FBodyInstance*> & OutWeldedBodies, TArray<FName> & OutLabels)
-{
-	OutWeldedBodies.Add(&BodyInstance);
-	OutLabels.Add(NAME_None);
-			
-	for (USceneComponent * Child : AttachChildren)
-	{
-		if (UPrimitiveComponent * PrimChild = Cast<UPrimitiveComponent>(Child))
-		{
-			if (FBodyInstance * BI = PrimChild->GetBodyInstance())
-			{
-				if (BI->bWelded)
-				{
-					PrimChild->GetWeldedBodies(OutWeldedBodies, OutLabels);
-				}
-			}
-		}
-	}
-}
-
-bool UPrimitiveComponent::WeldToImplementation(USceneComponent * InParent, FName ParentSocketName /* = Name_None */, bool bWeldSimulatedChild /* = false */)
-{
-	//WeldToInternal assumes attachment is already done
-	if (AttachParent != InParent || AttachSocketName != ParentSocketName)
-	{
-		return false;
-	}
-
-	//Check that we can actually our own socket name
-	FBodyInstance * BI = GetBodyInstance();
-	if (BI == NULL)
-	{
-		return false;
-	}
-
-	if (BI->ShouldInstanceSimulatingPhysics() && bWeldSimulatedChild == false)
-	{
-		return false;
-	}
-
-	UnWeldFromParent();	//make sure to unweld from wherever we currently are
-
-	FName SocketName;
-	UPrimitiveComponent * RootComponent = GetRootWelded(this, ParentSocketName, &SocketName, true);
-
-	if (RootComponent)
-	{
-		if (FBodyInstance * RootBI = RootComponent->GetBodyInstance(SocketName))
-		{
-			if (BI->WeldParent == RootBI)	//already welded so stop
-			{
-				return true;
-			}
-
-			BI->bWelded = true;
-			//There are multiple cases to handle:
-			//Root is kinematic, simulated
-			//Child is kinematic, simulated
-			//Child always inherits from root
-
-			//if root is kinematic simply set child to be kinematic and we're done
-			if (RootComponent->IsSimulatingPhysics(SocketName) == false)
-			{
-				BI->WeldParent = NULL;
-				SetSimulatePhysics(false);
-				return false;	//return false because we need to continue with regular body initialization
-			}
-
-			//root is simulated so we actually weld the body
-			FTransform RelativeTM = RootComponent == AttachParent ? GetRelativeTransform() : GetComponentToWorld().GetRelativeTransform(RootComponent->GetComponentToWorld());	//if direct parent we already have relative. Otherwise compute it
-			RootBI->Weld(BI, GetComponentToWorld());
-			BI->WeldParent = RootBI;
-
-			return true;
-		}
-	}
-
-	return false;
-}
-
-void UPrimitiveComponent::WeldTo(USceneComponent* InParent, FName InSocketName /* = NAME_None */)
-{
-	//automatically attach if needed
-	if (AttachParent != InParent || AttachSocketName != InSocketName)
-	{
-		AttachTo(InParent, InSocketName, EAttachLocation::KeepWorldPosition);
-	}
-
-	WeldToImplementation(InParent, InSocketName);
-}
-
-void UPrimitiveComponent::UnWeldFromParent()
-{
-	FBodyInstance * NewRootBI = GetBodyInstance();
-	if (NewRootBI == NULL || NewRootBI->bWelded == false)
-	{
-		return;
-	}
-
-	FName SocketName;
-	UPrimitiveComponent * RootComponent = GetRootWelded(this, AttachSocketName, &SocketName);
-
-	if (RootComponent)
-	{
-		if (FBodyInstance * RootBI = RootComponent->GetBodyInstance(SocketName))
-		{
-			//create new root
-			RootBI->UnWeld(NewRootBI);
-			NewRootBI->bWelded = false;
-			NewRootBI->WeldParent = NULL;
-
-			bool bHasBodySetup = GetBodySetup() != NULL;
-
-			//if BodyInstance hasn't already been created we need to initialize it
-			if (bHasBodySetup && NewRootBI->IsValidBodyInstance() == false)
-			{
-				bool bPrevAutoWeld = NewRootBI->bAutoWeld;
-				NewRootBI->bAutoWeld = false;
-				NewRootBI->InitBody(GetBodySetup(), GetComponentToWorld(), this, GetWorld()->GetPhysicsScene());
-				NewRootBI->bAutoWeld = bPrevAutoWeld;
-			}
-
-			//now weld its children to it
-			TArray<FBodyInstance*> ChildrenBodies;
-			TArray<FName> ChildrenLabels;
-			GetWeldedBodies(ChildrenBodies, ChildrenLabels);
-
-			for (int32 ChildIdx = 0; ChildIdx < ChildrenBodies.Num(); ++ChildIdx)
-			{
-				FBodyInstance * ChildBI = ChildrenBodies[ChildIdx];
-				if (ChildBI != NewRootBI)
-				{
-					RootBI->UnWeld(NewRootBI);
-					if (bHasBodySetup)
-					{
-						NewRootBI->Weld(ChildBI, ChildBI->OwnerComponent->GetSocketTransform(ChildrenLabels[ChildIdx]));
-					}
-				}
-			}
-		}
-	}
-}
 
 void UPrimitiveComponent::DestroyRenderState_Concurrent()
 {
diff --git a/Engine/Source/Runtime/Engine/Private/PrimitiveComponentPhysics.cpp b/Engine/Source/Runtime/Engine/Private/PrimitiveComponentPhysics.cpp
--- a/Engine/Source/Runtime/Engine/Private/PrimitiveComponentPhysics.cpp
+++ b/Engine/Source/Runtime/Engine/Private/PrimitiveComponentPhysics.cpp
@@ -440,11 +440,209 @@ void UPrimitiveComponent::SyncComponentToRBPhysics()
 	}
 }
 
-FBodyInstance* UPrimitiveComponent::GetBodyInstance(FName BoneName) const
+
+UPrimitiveComponent * GetRootWelded(const UPrimitiveComponent * PrimComponent, FName ParentSocketName = NAME_None, FName * OutSocketName = NULL, bool bAboutToWeld = false)
 {
-	return const_cast<FBodyInstance*>(&BodyInstance);
+	UPrimitiveComponent * Result = NULL;
+	UPrimitiveComponent * RootComponent = Cast<UPrimitiveComponent>(PrimComponent->AttachParent);	//we must find the root component along hierarchy that has bWelded set to true
+
+	//check that body itself is welded
+	if (FBodyInstance * BI = PrimComponent->GetBodyInstance(ParentSocketName, false))
+	{
+		if (bAboutToWeld == false && BI->bWelded == false && BI->bAutoWeld == false)	//we're not welded and we aren't trying to become welded
+		{
+			return NULL;
+		}
+	}
+
+	FName PrevSocketName = ParentSocketName;
+	FName SocketName = NAME_None; //because of skeletal mesh it's important that we check along the bones that we attached
+	FBodyInstance * RootBI = NULL;
+	for (; RootComponent; RootComponent = Cast<UPrimitiveComponent>(RootComponent->AttachParent))
+	{
+		Result = RootComponent;
+		SocketName = PrevSocketName;
+		PrevSocketName = RootComponent->AttachSocketName;
+
+		RootBI = RootComponent->GetBodyInstance(SocketName, false);
+		if (RootBI && RootBI->bWelded)
+		{
+			continue;
+		}
+
+		break;
+	}
+
+	if (OutSocketName)
+	{
+		*OutSocketName = SocketName;
+	}
+
+	return Result;
+}
+
+void UPrimitiveComponent::GetWeldedBodies(TArray<FBodyInstance*> & OutWeldedBodies, TArray<FName> & OutLabels)
+{
+	OutWeldedBodies.Add(&BodyInstance);
+	OutLabels.Add(NAME_None);
+
+	for (USceneComponent * Child : AttachChildren)
+	{
+		if (UPrimitiveComponent * PrimChild = Cast<UPrimitiveComponent>(Child))
+		{
+			if (FBodyInstance * BI = PrimChild->GetBodyInstance(NAME_None, false))
+			{
+				if (BI->bWelded)
+				{
+					PrimChild->GetWeldedBodies(OutWeldedBodies, OutLabels);
+				}
+			}
+		}
+	}
 }
 
+bool UPrimitiveComponent::WeldToImplementation(USceneComponent * InParent, FName ParentSocketName /* = Name_None */, bool bWeldSimulatedChild /* = false */)
+{
+	//WeldToInternal assumes attachment is already done
+	if (AttachParent != InParent || AttachSocketName != ParentSocketName)
+	{
+		return false;
+	}
+
+	//Check that we can actually our own socket name
+	FBodyInstance * BI = GetBodyInstance(NAME_None, false);
+	if (BI == NULL)
+	{
+		return false;
+	}
+
+	if (BI->ShouldInstanceSimulatingPhysics() && bWeldSimulatedChild == false)
+	{
+		return false;
+	}
+
+	UnWeldFromParent();	//make sure to unweld from wherever we currently are
+
+	FName SocketName;
+	UPrimitiveComponent * RootComponent = GetRootWelded(this, ParentSocketName, &SocketName, true);
+
+	if (RootComponent)
+	{
+		if (FBodyInstance * RootBI = RootComponent->GetBodyInstance(SocketName, false))
+		{
+			if (BI->WeldParent == RootBI)	//already welded so stop
+			{
+				return true;
+			}
+
+			BI->bWelded = true;
+			//There are multiple cases to handle:
+			//Root is kinematic, simulated
+			//Child is kinematic, simulated
+			//Child always inherits from root
+
+			//if root is kinematic simply set child to be kinematic and we're done
+			if (RootComponent->IsSimulatingPhysics(SocketName) == false)
+			{
+				BI->WeldParent = NULL;
+				SetSimulatePhysics(false);
+				return false;	//return false because we need to continue with regular body initialization
+			}
+
+			//root is simulated so we actually weld the body
+			FTransform RelativeTM = RootComponent == AttachParent ? GetRelativeTransform() : GetComponentToWorld().GetRelativeTransform(RootComponent->GetComponentToWorld());	//if direct parent we already have relative. Otherwise compute it
+			RootBI->Weld(BI, GetComponentToWorld());
+			BI->WeldParent = RootBI;
+
+			return true;
+		}
+	}
+
+	return false;
+}
+
+void UPrimitiveComponent::WeldTo(USceneComponent* InParent, FName InSocketName /* = NAME_None */)
+{
+	//automatically attach if needed
+	if (AttachParent != InParent || AttachSocketName != InSocketName)
+	{
+		AttachTo(InParent, InSocketName, EAttachLocation::KeepWorldPosition);
+	}
+
+	WeldToImplementation(InParent, InSocketName);
+}
+
+void UPrimitiveComponent::UnWeldFromParent()
+{
+	FBodyInstance * NewRootBI = GetBodyInstance(NAME_None, false);
+	if (NewRootBI == NULL || NewRootBI->bWelded == false)
+	{
+		return;
+	}
+
+	FName SocketName;
+	UPrimitiveComponent * RootComponent = GetRootWelded(this, AttachSocketName, &SocketName);
+
+	if (RootComponent)
+	{
+		if (FBodyInstance * RootBI = RootComponent->GetBodyInstance(SocketName, false))
+		{
+			//create new root
+			RootBI->UnWeld(NewRootBI);
+			NewRootBI->bWelded = false;
+			NewRootBI->WeldParent = NULL;
+
+			bool bHasBodySetup = GetBodySetup() != NULL;
+
+			//if BodyInstance hasn't already been created we need to initialize it
+			if (bHasBodySetup && NewRootBI->IsValidBodyInstance() == false)
+			{
+				bool bPrevAutoWeld = NewRootBI->bAutoWeld;
+				NewRootBI->bAutoWeld = false;
+				NewRootBI->InitBody(GetBodySetup(), GetComponentToWorld(), this, GetWorld()->GetPhysicsScene());
+				NewRootBI->bAutoWeld = bPrevAutoWeld;
+			}
+
+			//now weld its children to it
+			TArray<FBodyInstance*> ChildrenBodies;
+			TArray<FName> ChildrenLabels;
+			GetWeldedBodies(ChildrenBodies, ChildrenLabels);
+
+			for (int32 ChildIdx = 0; ChildIdx < ChildrenBodies.Num(); ++ChildIdx)
+			{
+				FBodyInstance * ChildBI = ChildrenBodies[ChildIdx];
+				if (ChildBI != NewRootBI)
+				{
+					RootBI->UnWeld(NewRootBI);
+					if (bHasBodySetup)
+					{
+						NewRootBI->Weld(ChildBI, ChildBI->OwnerComponent->GetSocketTransform(ChildrenLabels[ChildIdx]));
+					}
+				}
+			}
+		}
+	}
+}
+
+FBodyInstance* UPrimitiveComponent::GetBodyInstance(FName BoneName, bool bGetWelded) const
+{
+	if (bGetWelded && BodyInstance.bWelded)
+	{
+		FName OutSocket;
+		if (UPrimitiveComponent * RootComponentWelded = GetRootWelded(this, AttachSocketName, &OutSocket))
+		{
+			if (FBodyInstance * BI = RootComponentWelded->GetBodyInstance(OutSocket, bGetWelded))
+			{
+				if (BI->bSimulatePhysics)
+				{
+					return BI;
+				}
+			}
+		}
+	}
+
+	return const_cast<FBodyInstance*>(&BodyInstance);
+}
 
 float UPrimitiveComponent::GetDistanceToCollision(const FVector & Point, FVector& ClosestPointOnCollision) const
 {
diff --git a/Engine/Source/Runtime/Engine/Private/SkeletalMeshComponentPhysics.cpp b/Engine/Source/Runtime/Engine/Private/SkeletalMeshComponentPhysics.cpp
--- a/Engine/Source/Runtime/Engine/Private/SkeletalMeshComponentPhysics.cpp
+++ b/Engine/Source/Runtime/Engine/Private/SkeletalMeshComponentPhysics.cpp
@@ -1489,7 +1489,7 @@ FName USkeletalMeshComponent::FindConstraintBoneName( int32 ConstraintIndex )
 }
 
 
-FBodyInstance* USkeletalMeshComponent::GetBodyInstance(FName BoneName) const
+FBodyInstance* USkeletalMeshComponent::GetBodyInstance(FName BoneName, bool) const
 {
 	UPhysicsAsset * const PhysicsAsset = GetPhysicsAsset();
 	FBodyInstance* BodyInst = NULL;
