Add function name to log

diff --git a/Engine/Source/Runtime/Engine/Private/DataReplication.cpp b/Engine/Source/Runtime/Engine/Private/DataReplication.cpp
--- a/Engine/Source/Runtime/Engine/Private/DataReplication.cpp
+++ b/Engine/Source/Runtime/Engine/Private/DataReplication.cpp
@@ -679,7 +679,7 @@ bool FObjectReplicator::ReceivedBunch( FInBunch &Bunch, const FReplicationFlags
 
 			if ( Bunch.IsError() )
 			{
-				UE_LOG( LogNet, Error, TEXT( "ReceivedBunch: ReceivePropertiesForRPC - Bunch.IsError() == true: %s" ), *Object->GetFullName() );
+				UE_LOG( LogNet, Error, TEXT( "ReceivedBunch: ReceivePropertiesForRPC - Bunch.IsError() == true: Function: %s, Object: %s" ), *Message.ToString(), *Object->GetFullName() );
 				return false;
 			}
 
