#ue4 - Add new functions for reading the movement input vector from Pawn and PawnMovementComponent.

- These functions make explicit whether the input being read has been consumed or is pending processing, so components can read the old values after input is consumed rather than the zeroed values for the next frame.
- New functions added to get the pending input (input that has not yet caused movement).
- New functions added to get the last processed input (input that last caused movement).
- APawn::GetMovementInputVector() and UPawnMovementComponent::GetInputVector() have been deprecated.

TTP 345636

diff --git a/Engine/Config/BaseEngine.ini b/Engine/Config/BaseEngine.ini
--- a/Engine/Config/BaseEngine.ini
+++ b/Engine/Config/BaseEngine.ini
@@ -983,7 +983,8 @@ bUseStreamingPause=false
 +K2FieldRedirects=(OldFieldName="MovementComponent.GetModifiedMaxSpeed",NewFieldName="MovementComponent.K2_GetModifiedMaxSpeed")
 +K2FieldRedirects=(OldFieldName="MovementComponent.GetMaxSpeedModifier",NewFieldName="MovementComponent.K2_GetMaxSpeedModifier")
 +K2FieldRedirects=(OldFieldName="CharacterMovementComponent.GetModifiedMaxAcceleration",NewFieldName="CharacterMovementComponent.K2_GetModifiedMaxAcceleration")
-
++K2FieldRedirects=(OldFieldName="Pawn.GetMovementInputVector",NewFieldName="Pawn.K2_GetMovementInputVector")
++K2FieldRedirects=(OldFieldName="PawnMovementComponent.GetInputVector",NewFieldName="PawnMovementComponent.K2_GetInputVector")
 
 +K2FieldRedirects=(OldFieldName="PlayerController.LookRightScale",NewFieldName="PlayerController.InputYawScale")
 +K2FieldRedirects=(OldFieldName="PlayerController.LookUpScale",NewFieldName="PlayerController.InputPitchScale")
diff --git a/Engine/Source/Runtime/Engine/Classes/GameFramework/CharacterMovementComponent.h b/Engine/Source/Runtime/Engine/Classes/GameFramework/CharacterMovementComponent.h
--- a/Engine/Source/Runtime/Engine/Classes/GameFramework/CharacterMovementComponent.h
+++ b/Engine/Source/Runtime/Engine/Classes/GameFramework/CharacterMovementComponent.h
@@ -839,7 +839,7 @@ public:
 	virtual float GetModifiedMaxAcceleration() const;
 	
 	/** @return Maximum acceleration for the current state, based on MaxAcceleration and any additional modifiers. */
-	UFUNCTION(BlueprintCallable, Category="Pawn|Components|CharacterMovement", meta=(DeprecatedFunction, DisplayName="GetModifiedMaxAcceleration()", DeprecationMessage="GetModifiedMaxAcceleration() is deprecated, apply your own modifiers to GetMaxAcceleration() if desired."))
+	UFUNCTION(BlueprintCallable, Category="Pawn|Components|CharacterMovement", meta=(DeprecatedFunction, FriendlyName="GetModifiedMaxAcceleration", DeprecationMessage="GetModifiedMaxAcceleration() is deprecated, apply your own modifiers to GetMaxAcceleration() if desired."))
 	virtual float K2_GetModifiedMaxAcceleration() const;
 
 	/** @return Maximum acceleration for the current state. */
diff --git a/Engine/Source/Runtime/Engine/Classes/GameFramework/Pawn.h b/Engine/Source/Runtime/Engine/Classes/GameFramework/Pawn.h
--- a/Engine/Source/Runtime/Engine/Classes/GameFramework/Pawn.h
+++ b/Engine/Source/Runtime/Engine/Classes/GameFramework/Pawn.h
@@ -323,22 +323,45 @@ protected:
 public:
 	/**
 	 * Add movement input along the given world direction vector (usually normalized) scaled by 'ScaleValue'. If ScaleValue < 0, movement will be in the opposite direction.
-	 * Base Pawn classes won't automatically apply movement, it's up to the user to do so in a Tick. Subclasses such as Character and DefaultPawn automatically handle this input and move.
-	 * @see GetMovementInputVector(), ConsumeMovementInputVector()
+	 * Base Pawn classes won't automatically apply movement, it's up to the user to do so in a Tick event. Subclasses such as Character and DefaultPawn automatically handle this input and move.
+	 * @see GetPendingMovementInputVector(), GetLastMovementInputVector(), ConsumeMovementInputVector()
 	 *
-	 * @param WorldDirection:	Direction in world space to apply input
-	 * @param ScaleValue:		Scale to apply to input. This can be used for analog input, ie a value of 0.5 applies half the normal value.
-	 * @param bForce:			If true always add the input, ignoring the result of IsMoveInputIgnored().
+	 * @param WorldDirection	Direction in world space to apply input
+	 * @param ScaleValue		Scale to apply to input. This can be used for analog input, ie a value of 0.5 applies half the normal value, while -1.0 would reverse the direction.
+	 * @param bForce			If true always add the input, ignoring the result of IsMoveInputIgnored().
 	 */
-	UFUNCTION(BlueprintCallable, Category="Pawn|Input")
+	UFUNCTION(BlueprintCallable, Category="Pawn|Input", meta=(Keywords="AddInput"))
 	virtual void AddMovementInput(FVector WorldDirection, float ScaleValue = 1.0f, bool bForce = false);
 
-	/** Return the input vector in world space. Note that the input should be consumed with ConsumeMovementInputVector() at the end of an update, to prevent accumulation of control input between frames. */
-	UFUNCTION(BlueprintCallable, Category="Pawn|Input")
-	FVector GetMovementInputVector() const;
+	/**
+	 * Return the pending input vector in world space. This is the most up-to-date value of the input vector, pending ConsumeMovementInputVector() which clears it,
+	 * Usually only a PawnMovementComponent will want to read this value, or the Pawn itself if it is responsible for movement.
+	 *
+	 * @return The pending input vector in world space.
+	 * @see AddMovementInput(), GetLastMovementInputVector(), ConsumeMovementInputVector()
+	 */
+	UFUNCTION(BlueprintCallable, Category="Pawn|Input", meta=(Keywords="GetMovementInput GetInput"))
+	FVector GetPendingMovementInputVector() const;
 
-	/** Returns the input vector and resets it to zero. Should be used during an update to prevent accumulation of control input between frames. */
-	UFUNCTION(BlueprintCallable, Category="Pawn|Input")
+	/**
+	 * Return the last input vector in world space that was processed by ConsumeMovementInputVector(), which is usually done by the Pawn or PawnMovementComponent.
+	 * Any user that needs to know about the input that last affected movement should use this function.
+	 * For example an animation update would want to use this, since by default the order of updates in a frame is:
+	 * PlayerController (device input) -> MovementComponent -> Pawn -> Mesh (animations)
+	 *
+	 * @return The last input vector in world space that was processed by ConsumeMovementInputVector().
+	 * @see AddMovementInput(), GetPendingMovementInputVector(), ConsumeMovementInputVector()
+	 */
+	UFUNCTION(BlueprintCallable, Category="Pawn|Input", meta=(Keywords="GetMovementInput GetInput"))
+	FVector GetLastMovementInputVector() const;
+
+	/**
+	 * Returns the pending input vector and resets it to zero.
+	 * This should be used during a movement update (by the Pawn or PawnMovementComponent) to prevent accumulation of control input between frames.
+	 * Copies the pending input vector to the saved input vector (GetLastMovementInputVector()).
+	 * @return The pending input vector.
+	 */
+	UFUNCTION(BlueprintCallable, Category="Pawn|Input", meta=(Keywords="ConsumeInput"))
 	virtual FVector ConsumeMovementInputVector();
 
 	/** Add input (affecting Pitch) to the Controller's ControlRotation, if it is a local PlayerController. */
@@ -359,29 +382,59 @@ public:
 
 protected:
 
-	/** Accumulated control input vector, stored in world space. */
+	/**
+	 * Accumulated control input vector, stored in world space. This is the pending input, which is cleared (zeroed) once consumed.
+	 * @see GetPendingMovementInputVector(), AddMovementInput()
+	 */
 	UPROPERTY(Transient)
 	FVector ControlInputVector;
 
+	/**
+	 * The last control input vector that was processed by ConsumeMovementInputVector().
+	 * @see GetLastMovementInputVector()
+	 */
+	UPROPERTY(Transient)
+	FVector LastControlInputVector;
+
 public:
 	/** Internal function meant for use only within Pawn or by a PawnMovementComponent. Adds movement input if not ignored, or if forced. */
 	void Internal_AddMovementInput(FVector WorldAccel, bool bForce = false);
 
 	/** Internal function meant for use only within Pawn or by a PawnMovementComponent. Returns the value of ControlInputVector. */
-	inline FVector Internal_GetMovementInputVector() const { return ControlInputVector; }
+	inline FVector Internal_GetPendingMovementInputVector() const { return ControlInputVector; }
+
+	/** Internal function meant for use only within Pawn or by a PawnMovementComponent. Returns the value of LastControlInputVector. */
+	inline FVector Internal_GetLastMovementInputVector() const { return LastControlInputVector; }
 
-	/** Internal function meant for use only within Pawn or by a PawnMovementComponent. Returns the vale of ControlInputVector and sets it to zero. */
+	/** Internal function meant for use only within Pawn or by a PawnMovementComponent. LastControlInputVector is updated with initial value of ControlInputVector. Returns ControlInputVector and resets it to zero. */
 	FVector Internal_ConsumeMovementInputVector();
 
 public:
-	/** (DEPRECATED) Launch Character with LaunchVelocity  */
-	UFUNCTION(BlueprintCallable, Category="Pawn", meta=(DeprecatedFunction, DeprecationMessage="Use Character.LaunchCharacter instead"))
-	void LaunchPawn(FVector LaunchVelocity, bool bXYOverride, bool bZOverride);
 
 	/** Add an Actor to ignore by Pawn's movement collision */
 	void MoveIgnoreActorAdd(AActor* ActorToIgnore);
 
 	/** Remove an Actor to ignore by Pawn's movement collision */
 	void MoveIgnoreActorRemove(AActor* ActorToIgnore);
+
+public:
+
+	// DEPRECATED FUNCTIONS
+
+	/** (Deprecated) Launch Character with LaunchVelocity  */
+	UFUNCTION(BlueprintCallable, Category="Pawn", meta=(DeprecatedFunction, DeprecationMessage="Use Character.LaunchCharacter instead"))
+	void LaunchPawn(FVector LaunchVelocity, bool bXYOverride, bool bZOverride);
+
+	/** Internal function meant for use only within Pawn or by a PawnMovementComponent. */
+	DEPRECATED(4.5, "Internal_GetMovementInputVector is deprecated, use Internal_GetPendingMovementInputVector() or Internal_GetLastMovementInputVector() instead.")
+	inline FVector Internal_GetMovementInputVector() const { return Internal_GetPendingMovementInputVector(); }
+
+	/** (Deprecated) Return the input vector in world space. */
+	DEPRECATED(4.5, "GetMovementInputVector() has been deprecated, use either GetPendingMovementInputVector() or GetLastMovementInputVector().")
+	FVector GetMovementInputVector() const;
+
+	/** (Deprecated) Return the input vector in world space. */
+	UFUNCTION(BlueprintCallable, Category="Pawn|Input", meta=(DeprecatedFunction, FriendlyName="GetMovementInputVector", DeprecationMessage="GetMovementInputVector has been deprecated, use either GetPendingMovementInputVector or GetLastMovementInputVector"))
+	FVector K2_GetMovementInputVector() const;	
 };
 
diff --git a/Engine/Source/Runtime/Engine/Classes/GameFramework/PawnMovementComponent.h b/Engine/Source/Runtime/Engine/Classes/GameFramework/PawnMovementComponent.h
--- a/Engine/Source/Runtime/Engine/Classes/GameFramework/PawnMovementComponent.h
+++ b/Engine/Source/Runtime/Engine/Classes/GameFramework/PawnMovementComponent.h
@@ -13,6 +13,8 @@ class ENGINE_API UPawnMovementComponent : public UNavMovementComponent
 {
 	GENERATED_UCLASS_BODY()
 
+public:
+
 	// Overridden to only allow registration with components owned by a Pawn.
 	virtual void SetUpdatedComponent(class UPrimitiveComponent* NewUpdatedComponent) override;
 
@@ -28,11 +30,29 @@ class ENGINE_API UPawnMovementComponent : public UNavMovementComponent
 	UFUNCTION(BlueprintCallable, Category="Pawn|Components|PawnMovement")
 	virtual void AddInputVector(FVector WorldVector, bool bForce = false);
 
-	/** Return the input vector in world space. Note that the input should be consumed with ConsumeInputVector() at the end of an update, to prevent accumulation of control input between frames. */
-	UFUNCTION(BlueprintCallable, Category="Pawn|Components|PawnMovement")
-	FVector GetInputVector() const;
+	/**
+	 * Return the pending input vector in world space. This is the most up-to-date value of the input vector, pending ConsumeMovementInputVector() which clears it.
+	 * PawnMovementComponents implementing movement usually want to use either this or ConsumeInputVector() as these functions represent the most recent state of input.
+	 * @return The pending input vector in world space.
+	 * @see AddInputVector(), ConsumeInputVector(), GetLastInputVector()
+	 */
+	UFUNCTION(BlueprintCallable, Category="Pawn|Components|PawnMovement", meta=(Keywords="GetInput"))
+	FVector GetPendingInputVector() const;
+
+	/**
+	* Return the last input vector in world space that was processed by ConsumeInputVector(), which is usually done by the Pawn or PawnMovementComponent.
+	* Any user that needs to know about the input that last affected movement should use this function.
+	* @return The last input vector in world space that was processed by ConsumeInputVector().
+	* @see AddInputVector(), ConsumeInputVector(), GetPendingInputVector()
+	*/
+	UFUNCTION(BlueprintCallable, Category="Pawn|Components|PawnMovement", meta=(Keywords="GetInput"))
+	FVector GetLastInputVector() const;
 
-	/** Returns the input vector and resets it to zero. Should be used during an update to prevent accumulation of control input between frames. */
+	/* Returns the pending input vector and resets it to zero.
+	 * This should be used during a movement update (by the Pawn or PawnMovementComponent) to prevent accumulation of control input between frames.
+	 * Copies the pending input vector to the saved input vector (GetLastMovementInputVector()).
+	 * @return The pending input vector.
+	 */
 	UFUNCTION(BlueprintCallable, Category="Pawn|Components|PawnMovement")
 	virtual FVector ConsumeInputVector();
 
@@ -44,7 +64,6 @@ class ENGINE_API UPawnMovementComponent : public UNavMovementComponent
 	UFUNCTION(BlueprintCallable, Category="Pawn|Components|PawnMovement")
 	class APawn* GetPawnOwner() const;
 
-public:
 	/** Notify of collision in case we want to react, such as waking up avoidance or pathing code. */
 	virtual void NotifyBumpedPawn(APawn* BumpedPawn) {}
 
@@ -53,4 +72,16 @@ protected:
 	/** Pawn which owns this component. */
 	UPROPERTY()
 	class APawn* PawnOwner;
+
+public:
+
+	// DEPRECATED FUNCTIONS
+
+	/** (Deprecated) Return the input vector in world space. */
+	DEPRECATED(4.5, "GetInputVector() has been deprecated, use either GetPendingInputVector() or GetLastInputVector().")
+	FVector GetInputVector() const;
+
+	/** (Deprecated) Return the input vector in world space. */
+	UFUNCTION(BlueprintCallable, Category="Pawn|Components|PawnMovement", meta=(DeprecatedFunction, FriendlyName="GetInputVector", DeprecationMessage="GetInputVector has been deprecated, use either GetPendingInputVector or GetLastInputVector"))
+	FVector K2_GetInputVector() const;
 };
diff --git a/Engine/Source/Runtime/Engine/Private/FloatingPawnMovement.cpp b/Engine/Source/Runtime/Engine/Private/FloatingPawnMovement.cpp
--- a/Engine/Source/Runtime/Engine/Private/FloatingPawnMovement.cpp
+++ b/Engine/Source/Runtime/Engine/Private/FloatingPawnMovement.cpp
@@ -86,7 +86,7 @@ bool UFloatingPawnMovement::LimitWorldBounds()
 
 void UFloatingPawnMovement::ApplyControlInputToVelocity(float DeltaTime)
 {
-	const FVector ControlAcceleration = GetInputVector().ClampMaxSize(1.f);
+	const FVector ControlAcceleration = GetPendingInputVector().ClampMaxSize(1.f);
 
 	const float AnalogInputModifier = (ControlAcceleration.SizeSquared() > 0.f ? ControlAcceleration.Size() : 0.f);
 	const float MaxPawnSpeed = GetMaxSpeed() * AnalogInputModifier;
diff --git a/Engine/Source/Runtime/Engine/Private/Pawn.cpp b/Engine/Source/Runtime/Engine/Private/Pawn.cpp
--- a/Engine/Source/Runtime/Engine/Private/Pawn.cpp
+++ b/Engine/Source/Runtime/Engine/Private/Pawn.cpp
@@ -573,12 +573,29 @@ void APawn::AddMovementInput(FVector WorldDirection, float ScaleValue, bool bFor
 }
 
 
-FVector APawn::GetMovementInputVector() const
+FVector APawn::GetPendingMovementInputVector() const
 {
 	// There's really no point redirecting to the MovementComponent since GetInputVector is not virtual there, and it just comes back to us.
 	return ControlInputVector;
 }
 
+FVector APawn::GetLastMovementInputVector() const
+{
+	return LastControlInputVector;
+}
+
+// TODO: deprecated, remove
+FVector APawn::GetMovementInputVector() const
+{
+	return GetPendingMovementInputVector();
+}
+
+// TODO: deprecated, remove
+FVector APawn::K2_GetMovementInputVector() const
+{
+	return GetPendingMovementInputVector();
+}
+
 
 FVector APawn::ConsumeMovementInputVector()
 {
@@ -604,9 +621,9 @@ void APawn::Internal_AddMovementInput(FVector WorldAccel, bool bForce /*=false*/
 
 FVector APawn::Internal_ConsumeMovementInputVector()
 {
-	const FVector OldValue = ControlInputVector;
+	LastControlInputVector = ControlInputVector;
 	ControlInputVector = FVector::ZeroVector;
-	return OldValue;
+	return LastControlInputVector;
 }
 
 
diff --git a/Engine/Source/Runtime/Engine/Private/PawnMovementComponent.cpp b/Engine/Source/Runtime/Engine/Private/PawnMovementComponent.cpp
--- a/Engine/Source/Runtime/Engine/Private/PawnMovementComponent.cpp
+++ b/Engine/Source/Runtime/Engine/Private/PawnMovementComponent.cpp
@@ -56,12 +56,29 @@ void UPawnMovementComponent::AddInputVector(FVector WorldAccel, bool bForce /*=f
 	}
 }
 
-FVector UPawnMovementComponent::GetInputVector() const
+FVector UPawnMovementComponent::GetPendingInputVector() const
+{
+	return PawnOwner ? PawnOwner->Internal_GetPendingMovementInputVector() : FVector::ZeroVector;
+}
+
+FVector UPawnMovementComponent::GetLastInputVector() const
 {
-	return PawnOwner ? PawnOwner->Internal_GetMovementInputVector() : FVector::ZeroVector;
+	return PawnOwner ? PawnOwner->Internal_GetLastMovementInputVector() : FVector::ZeroVector;
 }
 
 FVector UPawnMovementComponent::ConsumeInputVector()
 {
 	return PawnOwner ? PawnOwner->Internal_ConsumeMovementInputVector() : FVector::ZeroVector;
 }
+
+// TODO: deprecated, remove
+FVector UPawnMovementComponent::GetInputVector() const
+{
+	return GetPendingInputVector();
+}
+
+// TODO: deprecated, remove
+FVector UPawnMovementComponent::K2_GetInputVector() const
+{
+	return GetPendingInputVector();
+}
