#ue4 - Header cleanup: AActor, USceneComponent

AActor:
- Added comments and doc links.
- Reordered some functions and variables for better grouping. More needs to be done there.

USceneComponent:
- Added class comment and doc links.

diff --git a/Engine/Source/Runtime/Engine/Classes/Components/SceneComponent.h b/Engine/Source/Runtime/Engine/Classes/Components/SceneComponent.h
--- a/Engine/Source/Runtime/Engine/Classes/Components/SceneComponent.h
+++ b/Engine/Source/Runtime/Engine/Classes/Components/SceneComponent.h
@@ -70,7 +70,12 @@ FORCEINLINE void operator|=(EMoveComponentFlags& Dest,EMoveComponentFlags Arg)
 DECLARE_DYNAMIC_MULTICAST_DELEGATE_OneParam(FPhysicsVolumeChanged, class APhysicsVolume*, NewVolume);
 
 
-/** A SceneComponent has a transform and supports attachment, but has no rendering or collision capabilities. Useful as a 'dummy' component in the hierarchy to offset others. */
+/**
+ * A SceneComponent has a transform and supports attachment, but has no rendering or collision capabilities.
+ * Useful as a 'dummy' component in the hierarchy to offset others.
+ * @see https://docs.unrealengine.com/latest/INT/API/Runtime/Engine/Components/USceneComponent/index.html
+ * @see https://docs.unrealengine.com/latest/INT/Programming/UnrealArchitecture/Actors/Components/index.html#scenecomponents
+ */
 UCLASS(ClassGroup=Utility, BlueprintType, HideCategories=(Trigger, PhysicsVolume), meta=(BlueprintSpawnableComponent))
 class ENGINE_API USceneComponent : public UActorComponent
 {
diff --git a/Engine/Source/Runtime/Engine/Classes/GameFramework/Actor.h b/Engine/Source/Runtime/Engine/Classes/GameFramework/Actor.h
--- a/Engine/Source/Runtime/Engine/Classes/GameFramework/Actor.h
+++ b/Engine/Source/Runtime/Engine/Classes/GameFramework/Actor.h
@@ -46,8 +46,12 @@ DECLARE_DELEGATE_RetVal_ThreeParams(bool, FOnProcessEvent, AActor*, UFunction*,
 DECLARE_CYCLE_STAT_EXTERN(TEXT("GetComponentsTime"),STAT_GetComponentsTime,STATGROUP_Engine,ENGINE_API);
 
 /**
- * Base class for an object that can be placed or spawned in a level. Actors may contain a collection of Components, and support network replication.
-*/
+ * Actor is the base class for a UObject that can be placed or spawned in a level.
+ * Actors may contain a collection of UActorComponents, which can be used to control how actors move, how they are rendered, etc.
+ * The other main function of an Actor is the replication of properties and function calls across the network during play.
+ * 
+ * @see https://docs.unrealengine.com/latest/INT/Programming/UnrealArchitecture/Actors/index.html
+ */
 UCLASS(abstract, BlueprintType, Blueprintable, config=Engine)
 class ENGINE_API AActor : public UObject
 {
@@ -63,56 +67,38 @@ class ENGINE_API AActor : public UObject
 
 	GENERATED_UCLASS_BODY()
 
-	/** Primary Actor tick function */
+public:
+
+	/**
+	 * Primary Actor tick function, which calls TickActor().
+	 * Tick functions can be configured to control whether ticking is enabled, at what time during a frame the update occurs, and to set up tick dependencies.
+	 * @see https://docs.unrealengine.com/latest/INT/API/Runtime/Engine/Engine/FTickFunction/index.html
+	 * @see AddTickPrerequisiteActor(), AddTickPrerequisiteComponent()
+	 */
 	UPROPERTY()
 	struct FActorTickFunction PrimaryActorTick;
 
-	/** A fence to track when the primitive is detached from the scene in the rendering thread. */ // @todo UE4 move?
-	FRenderCommandFence DetachFence;
-	
-	/** Array of ActorComponents that is actually serialized per-instance. */
-	UPROPERTY(TextExportTransient)
-	TArray< class UActorComponent* > SerializedComponents;
-
 	/** Allow each actor to run at a different time speed. The DeltaTime for a frame is multiplied by the global TimeDilation (in WorldSettings) and this CustomTimeDilation for this actor's tick.  */
 	UPROPERTY(BlueprintReadWrite, AdvancedDisplay, Category="Misc")
-	float CustomTimeDilation;
-
-private:
-	/** Used for replication (bNetUseOwnerRelevancy & bOnlyRelevantToOwner) and visibility (PrimitiveComponent bOwnerNoSee and bOnlyOwnerSee) */
-	UPROPERTY(replicated)
-	AActor* Owner;
-
-	/** Enables any collision on this actor */
-	UPROPERTY()
-	bool bActorEnableCollision;
+	float CustomTimeDilation;	
 
 public:
-	/** Allows us to only see this Actor in the Editor, and not in the actual game. */
+	/**
+	 * Allows us to only see this Actor in the Editor, and not in the actual game.
+	 * @see SetActorHiddenInGame()
+	 */
 	UPROPERTY(EditAnywhere, Category=Rendering, BlueprintReadOnly, replicated, meta=(DisplayName = "Actor Hidden In Game"))
 	uint32 bHidden:1;
 
-	/** Used for replication of our RootComponent's position and velocity */
-	UPROPERTY(ReplicatedUsing=OnRep_ReplicatedMovement)
-	struct FRepMovement ReplicatedMovement;
-
-	/** Used for replicating attachment of this actor's rootcomponent to another actor */
-	UPROPERTY(replicatedUsing=OnRep_AttachmentReplication)
-	struct FRepAttachment AttachmentReplication;
-
-	/** Called on client when updated AttachmentReplication value is received for this actor. */
-	UFUNCTION()
-	virtual void OnRep_AttachmentReplication();
-
-	/** Tear-off simulation in network play. */
+	/** If true, when the actor is spawned it will be sent to the client but receive no further replication updates from the server afterwards. */
 	UPROPERTY()
 	uint32 bNetTemporary:1;
 
-	/** This actor was loaded directly from the map, and for networking purposes can be addressed by its full path name */
+	/** If true, this actor was loaded directly from the map, and for networking purposes can be addressed by its full path name */
 	UPROPERTY()
 	uint32 bNetStartup:1;
 
-	/** This actor is only relevant to its owner. If this flag is changed during play, all non-owner channels would need to be explicitly closed. */
+	/** If ture, this actor is only relevant to its owner. If this flag is changed during play, all non-owner channels would need to be explicitly closed. */
 	UPROPERTY(Category=Replication, EditDefaultsOnly, BlueprintReadOnly)
 	uint32 bOnlyRelevantToOwner:1;
 
@@ -120,20 +106,30 @@ public:
 	UPROPERTY(Category=Replication, EditDefaultsOnly, BlueprintReadWrite)
 	uint32 bAlwaysRelevant:1;    
 
-	/** Replicate instigator to client (used by bNetTemporary projectiles for example). */
+	/** Replicate instigator to client. */
 	UPROPERTY()
 	uint32 bReplicateInstigator:1;    
 
-	/** If true, replicate movement/location related properties */
+	/**
+	 * If true, replicate movement/location related properties.
+	 * Actor must also be set to replicate.
+	 * @see SetReplicates()
+	 * @see https://docs.unrealengine.com/latest/INT/Gameplay/Networking/Replication/index.html
+	 */
 	UPROPERTY(Replicated, Category=Replication, EditDefaultsOnly)
 	uint32 bReplicateMovement:1;    
 
-	/** if true, this actor is no longer replicated to new clients, and is "torn off" (becomes a ROLE_Authority) on clients to which it was being replicated. */
+	/**
+	 * If true, this actor is no longer replicated to new clients, and is "torn off" (becomes a ROLE_Authority) on clients to which it was being replicated.
+	 * @see TornOff()
+	 */
 	UPROPERTY(replicated)
 	uint32 bTearOff:1;    
 
-	/** whether we already exchanged Role/RemoteRole on the client, as removing then re-adding a streaming level
-	 * causes all initialization to be performed again even though the actor may not have actually been reloaded */
+	/**
+	 * Whether we have already exchanged Role/RemoteRole on the client, as when removing then re-adding a streaming level.
+	 * Causes all initialization to be performed again even though the actor may not have actually been reloaded.
+	 */
 	UPROPERTY(transient)
 	uint32 bExchangedRoles:1;
 
@@ -160,56 +156,101 @@ private:
 	/** Whether FinishSpawning has been called for this Actor.  If it has not, the Actor is in a mal-formed state */
 	uint32 bHasFinishedSpawning:1;
 
+	/**
+	 * Enables any collision on this actor.
+	 * @see SetActorEnableCollision(), GetActorEnableCollision()
+	 */
+	UPROPERTY()
+	uint32 bActorEnableCollision:1;
+
 protected:
-	// This actor will replicate to remote machines
+	/**
+	 * If true, this actor will replicate to remote machines
+	 * @see SetReplicates()
+	 */
 	UPROPERTY(Category = Replication, EditDefaultsOnly, BlueprintReadOnly)
-	uint32 bReplicates : 1;
+	uint32 bReplicates:1;
 
-	// This function should only be used in the constructor of classes that need to set the RemoteRole
-	// to for backwards compatibility purposes
+	/** This function should only be used in the constructor of classes that need to set the RemoteRole for backwards compatibility purposes */
 	void SetRemoteRoleForBackwardsCompat(const ENetRole InRemoteRole) { RemoteRole = InRemoteRole; }
 
 	/**
-	 * Does this actor have an owner responsible for replication (APlayerController typically)
+	 * Does this actor have an owner responsible for replication? (APlayerController typically)
 	 *
 	 * @return true if this actor can call RPCs or false if no such owner chain exists
 	 */
 	virtual bool HasNetOwner() const;
 
 private:
-	// Describes how much control the remote machine has over the actor
-	// old UPROPERTY maintained for backwards compatibility, new non-uproperty as proper value
+	/**
+	 * Describes how much control the remote machine has over the actor.
+	 */
 	UPROPERTY(Replicated, transient)
 	TEnumAsByte<enum ENetRole> RemoteRole;
 
+	/**
+	 * Owner of this Actor, used primarily for replication (bNetUseOwnerRelevancy & bOnlyRelevantToOwner) and visibility (PrimitiveComponent bOwnerNoSee and bOnlyOwnerSee)
+	 * @see SetOwner(), GetOwner()
+	 */
+	UPROPERTY(replicated)
+	AActor* Owner;
+
 public:
 
-	/** Changes the RemoteRole property and handles the cases where the actor needs to be added to the network actor list */
+	/**
+	 * Set whether this actor replicates to network clients. When this actor is spawned on the server it will be sent to clients as well.
+	 * Properties flagged for replication will update on clients if they change on the server.
+	 * Internally changes the RemoteRole property and handles the cases where the actor needs to be added to the network actor list.
+	 * @param bInReplicates Whether this Actor replicates to network clients.
+	 * @see https://docs.unrealengine.com/latest/INT/Gameplay/Networking/Replication/index.html
+	 */
 	UFUNCTION(BlueprintCallable, Category = "Replication")
 	void SetReplicates(bool bInReplicates);
+
+	/** Sets whether or not this Actor is an autonomous proxy, which is an actor on a network client that is controlled by a user on that client. */
 	void SetAutonomousProxy(bool bInAutonomousProxy);
+	
+	/** Copies RemoteRole from another Actor and adds this actor to the list of network actors if necessary. */
 	void CopyRemoteRoleFrom(const AActor* CopyFromActor);
 
+	/** Returns how much control the remote machine has over this actor. */
 	ENetRole GetRemoteRole() const;
 
-	//Describes how much control the local machine has over the actor
+	/** Used for replication of our RootComponent's position and velocity */
+	UPROPERTY(ReplicatedUsing=OnRep_ReplicatedMovement)
+	struct FRepMovement ReplicatedMovement;
+
+	/** Used for replicating attachment of this actor's RootComponent to another actor. */
+	UPROPERTY(replicatedUsing=OnRep_AttachmentReplication)
+	struct FRepAttachment AttachmentReplication;
+
+	/** Called on client when updated AttachmentReplication value is received for this actor. */
+	UFUNCTION()
+	virtual void OnRep_AttachmentReplication();
+
+	/** Describes how much control the local machine has over the actor. */
 	UPROPERTY(Replicated)
 	TEnumAsByte<enum ENetRole> Role;
 
-	/** Dormancy setting for actor taking itself off of the replication list without being destroyed on clients */
-	TEnumAsByte<enum ENetDormancy>	NetDormancy;
+	/** Dormancy setting for actor to take itself off of the replication list without being destroyed on clients. */
+	TEnumAsByte<enum ENetDormancy> NetDormancy;
+
+	/** List of replicated components. Cooked builds won't rebuild this array. */
+	UPROPERTY()
+	TArray< TWeakObjectPtr<class UActorComponent> >	ReplicatedComponents;
 
-	/** Register actor to receive input from a player */
+	/** Automatically registers this actor to receive input from a player. */
 	UPROPERTY(EditAnywhere, Category=Input)
 	TEnumAsByte<EAutoReceiveInput::Type> AutoReceiveInput;
 
+	/** Component that handles input for this actor, if input is enabled. */
 	UPROPERTY()
 	class UInputComponent* InputComponent;
 
 	UPROPERTY()
 	TEnumAsByte<enum EInputConsumeOptions> InputConsumeOption_DEPRECATED;
 
-	/**  Square of the max distance from the client's viewpoint that this actor is relevant and will be replicated. */
+	/** Square of the max distance from the client's viewpoint that this actor is relevant and will be replicated. */
 	UPROPERTY(BlueprintReadOnly, EditDefaultsOnly, Category=Replication)
 	float NetCullDistanceSquared;   
 
@@ -250,36 +291,49 @@ public:
 	/** Called on the actor right before replication occurs */
 	virtual void PreReplication( IRepChangedPropertyTracker & ChangedPropertyTracker );
 
-	/** True to destroy when "finished", meaning all relevant components report that they are done and no timelines or timers are in flight. */
+	/** If true then destroy self when "finished", meaning all relevant components report that they are done and no timelines or timers are in flight. */
 	UPROPERTY(BlueprintReadWrite, Category=Actor)
 	uint32 bAutoDestroyWhenFinished:1;
 
-	/** Can take damage. Must be true for damage events (e.g. ReceiveDamage) to be called */
+	/**
+	 * Whether this actor can take damage. Must be true for damage events (e.g. ReceiveDamage()) to be called.
+	 * @see https://www.unrealengine.com/blog/damage-in-ue4
+	 * @see TakeDamage(), ReceiveDamage()
+	 */
 	UPROPERTY(EditAnywhere, BlueprintReadWrite, SaveGame, Replicated, Category=Actor)
-	uint32 bCanBeDamaged:1;    
+	uint32 bCanBeDamaged:1;
 
-	/** set when actor is about to be deleted (since endstate and other functions called during deletion process before MarkPendingKill is called). See IsPendingKillPending(). */
+	/**
+	 * Set when actor is about to be deleted.
+	 * @see IsPendingKillPending()
+	 */
 	UPROPERTY()
 	uint32 bPendingKillPending:1;    
 
-	/** This actor collides with the world when placing, even if RootComponent collision is disabled */
+	/** This actor collides with the world when placing in the editor or when spawned, even if RootComponent collision is disabled */
 	UPROPERTY()
 	uint32 bCollideWhenPlacing:1;    
 
-	/** Should this actor search for an owned camera component to view thru when used as a view target? */
+	/** If true, this actor should search for an owned camera component to view through when used as a view target. */
 	UPROPERTY(EditAnywhere, BlueprintReadWrite, Category=Actor, AdvancedDisplay)
 	uint32 bFindCameraComponentWhenViewTarget:1;
 
 	/** Pawn responsible for damage caused by this actor. */
 	UPROPERTY(BlueprintReadWrite, replicatedUsing=OnRep_Instigator, meta=(ExposeOnSpawn=true), Category=Actor)
-	class APawn* Instigator;      
+	class APawn* Instigator;
+
+	/** Called on clients when Instigator is replicated. */
+	UFUNCTION()
+	virtual void OnRep_Instigator();
 
-	/** The time this actor was created, relative to World->GetTimeSeconds() */
-	float CreationTime;    
+	/** The time this actor was created, relative to World->GetTimeSeconds().
+	 * @see UWorld::GetTimeSeconds()
+	 */
+	float CreationTime;
 
 	/** Array of Actors whose Owner is this actor */
 	UPROPERTY(transient)
-	TArray<AActor*> Children;    
+	TArray<AActor*> Children;
 	
 	// Animation update rate control.
 protected:
@@ -291,29 +345,34 @@ public:
 	UPROPERTY(Transient)
 	uint32 AnimUpdateRateFrameCount;
 
-	/** get a unique ID to share with all SkinnedMeshComponents in this actor. */
+	/** Get a unique ID to share with all SkinnedMeshComponents in this actor. */
 	uint32 GetAnimUpdateRateShiftTag();
 
 protected:
-	// Collision primitive.
+	/**
+	 * Collision primitive that defines the transform (location, rotation, scale) of this Actor.
+	 */
 	UPROPERTY()
 	class USceneComponent* RootComponent;
 
-	/** The matinee actors that control this actor **/
+	/** The matinee actors that control this actor. */
 	UPROPERTY(transient)
 	TArray<class AMatineeActor*> ControllingMatineeActors;
 
-	/** How long this Actor lives before dying, 0=forever. Note this is the INITIAL value and should not be set once play has begun */
+	/** How long this Actor lives before dying, 0=forever. Note this is the INITIAL value and should not be modified once play has begun. */
 	UPROPERTY(EditDefaultsOnly, BlueprintReadOnly, Category=Actor)
 	float InitialLifeSpan; 
 
-	/** If false, the Blueprint ReceiveTick event will be disabled on dedicated servers. */
+	/**
+	 * If false, the Blueprint ReceiveTick() event will be disabled on dedicated servers.
+	 * @see AllowReceiveTickEventOnDedicatedServer()
+	 */
 	UPROPERTY()
 	uint32 bAllowReceiveTickEventOnDedicatedServer:1;
 
 public:
 
-	/** Return the value of bAllowReceiveTickEventOnDedicatedServer, indicating whether the Blueprint ReceiveTick event will occur on dedicated servers */
+	/** Return the value of bAllowReceiveTickEventOnDedicatedServer, indicating whether the Blueprint ReceiveTick() event will occur on dedicated servers. */
 	FORCEINLINE bool AllowReceiveTickEventOnDedicatedServer() const { return bAllowReceiveTickEventOnDedicatedServer; }
 
 	/** Layer's the actor belongs to.  This is outside of the editoronly data to allow hiding of LD-specified layers at runtime for profiling. */
@@ -327,8 +386,10 @@ protected:
 	uint32 bActorLabelEditable:1;    // Is the actor label editable by the user?
 
 private:
-	/** The friendly name for this actor, displayed in the editor.  You should always use AActor::GetActorLabel() to access the actual label to display,
-		and call AActor::SetActorLabel() or AActor::SetActorLabelUnique() to change the label.  Never set the label directly. */
+	/**
+	 * The friendly name for this actor, displayed in the editor.  You should always use AActor::GetActorLabel() to access the actual label to display,
+	 * and call AActor::SetActorLabel() or AActor::SetActorLabelUnique() to change the label.  Never set the label directly.
+	 */
 	UPROPERTY()
 	FString ActorLabel;
 
@@ -337,7 +398,7 @@ private:
 	FName FolderPath;
 
 public:
-	/** Is hidden within the editor at its startup. */
+	/** Whether this actor is hidden within the editor viewport. */
 	UPROPERTY()
 	uint32 bHiddenEd:1;
 
@@ -346,27 +407,27 @@ protected:
 	UPROPERTY()
 	uint32 bEditable:1;
 
-	/** True if this actor should be listed in the scene outliner */
+	/** Whether this actor should be listed in the scene outliner. */
 	UPROPERTY()
 	uint32 bListedInSceneOutliner:1;
 
 public:
-	/** Is hidden by the layer browser. */
+	/** Whether this actor is hidden by the layer browser. */
 	UPROPERTY()
 	uint32 bHiddenEdLayer:1;
 
 private:
-	/** Is temporarily hidden within the editor; used for show/hide/etc. functionality w/o dirtying the actor */
+	/** Whether this actor is temporarily hidden within the editor; used for show/hide/etc functionality w/o dirtying the actor. */
 	UPROPERTY(transient)
 	uint32 bHiddenEdTemporary:1;
 
 public:
 
-	/** Is hidden by the level browser. */
+	/** Whether this actor is hidden by the level browser. */
 	UPROPERTY(transient)
 	uint32 bHiddenEdLevel:1;
 
-	/** Prevent the actor from being moved in the editor. */
+	/** If true, prevents the actor from being moved in the editor viewport. */
 	UPROPERTY()
 	uint32 bLockLocation:1;
 
@@ -374,15 +435,15 @@ public:
 	UPROPERTY(transient)
 	AActor* GroupActor;
 
-	/** The scale to apply to any billboard components in editor builds (happens in any WITH_EDITOR build, including non-cooked games) */
+	/** The scale to apply to any billboard components in editor builds (happens in any WITH_EDITOR build, including non-cooked games). */
 	UPROPERTY(EditAnywhere, BlueprintReadWrite, Category=Rendering, meta=(DisplayName="Editor Billboard Scale"))
 	float SpriteScale;
 
-	/** Returns how many lights are uncached for this actor */
+	/** Returns how many lights are uncached for this actor. */
 	int32 GetNumUncachedLights();
 #endif // WITH_EDITORONLY_DATA
 
-	/** The Actor that owns the ChildActorComponent that owns this Actor */
+	/** The Actor that owns the UChildActorComponent that owns this Actor. */
 	UPROPERTY()
 	TWeakObjectPtr<AActor> ParentComponentActor;	
 
@@ -399,15 +460,15 @@ public:
 	UPROPERTY()
 	uint32 bActorSeamlessTraveled:1;
 
-	/** Whether this actor should no be affected by world origin shifting  */
+	/** Whether this actor should no be affected by world origin shifting. */
 	UPROPERTY(EditAnywhere, AdvancedDisplay, Category=Actor)
 	uint32 bIgnoresOriginShifting:1;
 	
-	/** Array of tags that can be used for grouping and categorizing. Can also be accessed from scripting */
+	/** Array of tags that can be used for grouping and categorizing. */
 	UPROPERTY(EditAnywhere, BlueprintReadWrite, Category=Tags)
 	TArray<FName> Tags;
 
-	/** Bitflag to represent which views this actor is hidden in, via per-view layer visibility */
+	/** Bitflag to represent which views this actor is hidden in, via per-view layer visibility. */
 	UPROPERTY(transient)
 	uint64 HiddenEditorViews;
 
@@ -418,85 +479,77 @@ public:
 	UPROPERTY(BlueprintAssignable, Category="Game|Damage")
 	FTakeAnyDamageSignature OnTakeAnyDamage;
 
-	/** Called when the actor is damaged by point damage */
+	/** Called when the actor is damaged by point damage. */
 	UPROPERTY(BlueprintAssignable, Category="Game|Damage")
 	FTakePointDamageSignature OnTakePointDamage;
 	
-	/** Called when another actor begins to overlap this actor */
+	/** Called when another actor begins to overlap this actor. */
 	UPROPERTY(BlueprintAssignable, Category="Collision")
 	FActorBeginOverlapSignature OnActorBeginOverlap;
 
-	/** Called when another actor ends overlap with this actor */
+	/** Called when another actor ends overlap with this actor. */
 	UPROPERTY(BlueprintAssignable, Category="Collision")
 	FActorEndOverlapSignature OnActorEndOverlap;
 
-	/** Called when the mouse cursor is moved over this actor and mouse over events are enabled in the player controller */
+	/** Called when the mouse cursor is moved over this actor if mouse over events are enabled in the player controller. */
 	UPROPERTY(BlueprintAssignable, Category="Input|Mouse Input")
 	FActorBeginCursorOverSignature OnBeginCursorOver;
 
-	/** Called when the mouse cursor is moved off this actor and mouse over events are enabled in the player controller */
+	/** Called when the mouse cursor is moved off this actor if mouse over events are enabled in the player controller. */
 	UPROPERTY(BlueprintAssignable, Category="Input|Mouse Input")
 	FActorEndCursorOverSignature OnEndCursorOver;
 
-	/** Called when the left mouse button is clicked while the mouse is over this actor and click events are enabled in the player controller */
+	/** Called when the left mouse button is clicked while the mouse is over this actor and click events are enabled in the player controller. */
 	UPROPERTY(BlueprintAssignable, Category="Input|Mouse Input")
 	FActorOnClickedSignature OnClicked;
 
-	/** Called when the left mouse button is released while the mouse is over this actor and click events are enabled in the player controller */
+	/** Called when the left mouse button is released while the mouse is over this actor and click events are enabled in the player controller. */
 	UPROPERTY(BlueprintAssignable, Category="Input|Mouse Input")
 	FActorOnReleasedSignature OnReleased;
 
-	/** Called when a touch input is received over this actor when touch events are enabled in the player controller */
+	/** Called when a touch input is received over this actor when touch events are enabled in the player controller. */
 	UPROPERTY(BlueprintAssignable, Category="Input|Touch Input")
 	FActorOnInputTouchBeginSignature OnInputTouchBegin;
 		
-	/** Called when a touch input is received over this component when touch events are enabled in the player controller */
+	/** Called when a touch input is received over this component when touch events are enabled in the player controller. */
 	UPROPERTY(BlueprintAssignable, Category="Input|Touch Input")
 	FActorOnInputTouchEndSignature OnInputTouchEnd;
 
-	/** Called when a finger is moved over this actor when touch over events are enabled in the player controller */
+	/** Called when a finger is moved over this actor when touch over events are enabled in the player controller. */
 	UPROPERTY(BlueprintAssignable, Category="Input|Touch Input")
 	FActorBeginTouchOverSignature OnInputTouchEnter;
 
-	/** Called when a finger is moved off this actor when touch over events are enabled in the player controller */
+	/** Called when a finger is moved off this actor when touch over events are enabled in the player controller. */
 	UPROPERTY(BlueprintAssignable, Category="Input|Touch Input")
 	FActorEndTouchOverSignature OnInputTouchLeave;
 
-	/** Called when this actor is involved in a blocking collision */
+	/** Called when this actor is involved in a blocking collision. */
 	UPROPERTY(BlueprintAssignable, Category="Collision")
 	FActorHitSignature OnActorHit;
 
-	/** List of replicated components, cooked builds won't rebuild this array */
-	UPROPERTY()
-	TArray< TWeakObjectPtr<class UActorComponent> >	ReplicatedComponents;
-
-	/** Called on client when Instigator is replicated (only if bReplicateInstigator is true) */
-	UFUNCTION()
-	virtual void OnRep_Instigator();
-
 	/** 
-	 * Pushes this actor on to the stack of input being handled by the player controller
-	 * @param PlayerController - The player controller whose input events we are interested in receiving
+	 * Pushes this actor on to the stack of input being handled by a PlayerController.
+	 * @param PlayerController The PlayerController whose input events we want to receive.
 	 */
 	UFUNCTION(BlueprintCallable, Category="Input")
 	virtual void EnableInput(class APlayerController* PlayerController);
 
 	/** 
-	 * Removes this actor from the stack of input being handled by the player controller
-	 * @param PlayerController - The player controller whose input events we are no longer interested in receiving.  If NULL, actor will be removed from all player controllers
+	 * Removes this actor from the stack of input being handled by a PlayerController.
+	 * @param PlayerController The PlayerController whose input events we no longer want to receive. If null, this actor will stop receiving input from all PlayerControllers.
 	 */
 	UFUNCTION(BlueprintCallable, Category="Input")
 	virtual void DisableInput(class APlayerController* PlayerController);
 
-	/** Gets the value of the input axis if input is enabled for this actor */
+	/** Gets the value of the input axis if input is enabled for this actor. */
 	UFUNCTION(BlueprintCallable, meta=(BlueprintInternalUseOnly="true", BlueprintProtected = "true", HidePin="InputAxisName"))
 	float GetInputAxisValue(const FName InputAxisName) const;
 
-	/** Gets the value of the input axis key if input is enabled for this actor */
+	/** Gets the value of the input axis key if input is enabled for this actor. */
 	UFUNCTION(BlueprintCallable, meta=(BlueprintInternalUseOnly="true", BlueprintProtected = "true", HidePin="InputAxisKey"))
 	float GetInputAxisKeyValue(const FKey InputAxisKey) const;
 
-	/** Gets the value of the input axis key if input is enabled for this actor */
+	/** Gets the value of the input axis key if input is enabled for this actor. */
 	UFUNCTION(BlueprintCallable, meta=(BlueprintInternalUseOnly="true", BlueprintProtected = "true", HidePin="InputAxisKey"))
 	FVector GetInputVectorAxisValue(const FKey InputAxisKey) const;
 
@@ -505,8 +558,8 @@ public:
 	APawn* GetInstigator() const;
 
 	/**
-	 * Get the instigator, given a specific class
-	 * @return The instigator for this weapon, if it is the specified type. NULL, otherwise
+	 * Get the instigator, cast as a specific class.
+	 * @return The instigator for this weapon if it is the specified type, NULL otherwise.
 	 */
 	template <class T>
 	T* GetInstigator() const { return Cast<T>(Instigator); };
@@ -518,16 +571,18 @@ public:
 
 	//=============================================================================
 	// General functions.
-	
-	/** Get the class of this Actor*/
-	UFUNCTION(BlueprintCallable, meta=(DeprecatedFunction, FriendlyName = "GetActorClass"), Category="Class")
-	class UClass* GetActorClass() const;
 
-	/** Get the actor to world transform */
+	/**
+	 * Get the actor-to-world transform.
+	 * @return The transform that transforms from actor space to world space.
+	 */
 	UFUNCTION(BlueprintCallable, meta=(FriendlyName = "GetActorTransform"), Category="Utilities|Transformation")
 	FTransform GetTransform() const;
 
-	/** Returns location of the RootComponent of this Actor*/
+	/** Get the local-to-world transform of the RootComponent. Identical to GetTransform(). */
+	FTransform ActorToWorld() const;
+
+	/** Returns the location of the RootComponent of this Actor */
 	UFUNCTION(BlueprintCallable, meta=(FriendlyName = "GetActorLocation"), Category="Utilities|Transformation")
 	FVector K2_GetActorLocation() const;
 
@@ -535,12 +590,12 @@ public:
 	 *	Move the Actor instantly to the specified location. 
 	 *	@param NewLocation	The new location to teleport the Actor to.
 	 *	@param bSweep		Should we sweep to the destination location. If true, will stop short of the target if blocked by something.
-	 *	@return	Whether the location was successfully set if not swept, or whether movement occurred if swept.
+	 *	@return	Whether the location was successfully set (if not swept), or whether movement occurred at all (if swept).
 	 */
 	UFUNCTION(BlueprintCallable, meta=(FriendlyName = "SetActorLocation"), Category="Utilities|Transformation")
 	bool K2_SetActorLocation(FVector NewLocation, bool bSweep=false);
 
-	/** Returns rotation of the RootComponent of this Actor */
+	/** Returns rotation of the RootComponent of this Actor. */
 	UFUNCTION(BlueprintCallable, meta=(FriendlyName = "GetActorRotation"), Category="Utilities|Transformation")
 	FRotator K2_GetActorRotation() const;
 
@@ -558,7 +613,7 @@ public:
 
 	/**
 	 *	Returns the bounding box of all components that make up this Actor.
-	 *	@param	bOnlyCollidingComponents	If TRUE, will only return the bounding box for components with collision enabled.
+	 *	@param	bOnlyCollidingComponents	If true, will only return the bounding box for components with collision enabled.
 	 */
 	UFUNCTION(BlueprintCallable, Category="Collision", meta=(FriendlyName = "GetActorBounds"))
 	void GetActorBounds(bool bOnlyCollidingComponents, FVector& Origin, FVector& BoxExtent) const;
@@ -790,13 +845,13 @@ public:
 	//==============================================================================
 	// Misc Blueprint support
 
-	/** Getter for the cached world pointer */
+	/** Get the World in which this Actor exists. */
 	UFUNCTION(BlueprintCallable, meta=(DeprecatedFunction))
 	UWorld* K2_GetWorld() const;
 
 	/** 
-	 * Get CustomTimeDilation - this can be used for input control or speed control for slomo
-	 *  We don't want to scale input globally because input can be used for UI, which do not care for TimeDilation 
+	 * Get CustomTimeDilation - this can be used for input control or speed control for slomo.
+	 * We don't want to scale input globally because input can be used for UI, which do not care for TimeDilation.
 	 */
 	UFUNCTION(BlueprintCallable, Category="Utilities|Time")
 	float GetActorTimeDilation() const;
@@ -857,10 +912,11 @@ public:
 	//=============================================================================
 	// Blueprint
 	
-	/** Event when play begins for this actor */
+	/** Event when play begins for this actor. */
 	UFUNCTION(BlueprintImplementableEvent, meta=(FriendlyName = "BeginPlay"))
 	virtual void ReceiveBeginPlay();
 
+	/** Event when play begins for this actor. */
 	virtual void BeginPlay();
 
 	/** Event when this actor takes ANY damage */
@@ -971,7 +1027,7 @@ public:
 	UFUNCTION(BlueprintImplementableEvent, meta = (Keywords = "delete", FriendlyName = "Destroyed"))
 	virtual void ReceiveDestroyed();
 
-	/** Called when the actor is destroyed. */
+	/** Event triggered when the actor is destroyed. */
 	UPROPERTY(BlueprintAssignable)
 	FActorDestroyedSignature OnDestroyed;
 
@@ -980,7 +1036,7 @@ public:
 	UFUNCTION(BlueprintImplementableEvent, meta=(Keywords = "delete", FriendlyName = "End Play"))
 	virtual void ReceiveEndPlay(EEndPlayReason::Type EndPlayReason);
 
-	/** Called when the actor is destroyed. */
+	/** Event triggered when the actor is being removed from a level. */
 	UPROPERTY(BlueprintAssignable)
 	FActorEndPlaySignature OnEndPlay;
 	
@@ -1055,9 +1111,9 @@ public:
 	/** fills ReplicatedMovement property */
 	virtual void GatherCurrentMovement();
 
-	//
-	// See if this actor is owned by TestOwner.
-	//
+	/**
+	 * See if this actor is owned by TestOwner.
+	 */
 	inline bool IsOwnedBy( const AActor* TestOwner ) const
 	{
 		for( const AActor* Arg=this; Arg; Arg=Arg->Owner )
@@ -1068,9 +1124,10 @@ public:
 		return false;
 	}
 
-	/** Returns location of the RootComponent 
-		this is a template for no other reason than to delay compilation until USceneComponent is defined
-	*/ 
+	/**
+	 * Returns location of the RootComponent 
+	 * this is a template for no other reason than to delay compilation until USceneComponent is defined
+	 */ 
 	template<class T>
 	static FORCEINLINE FVector GetActorLocation(T* RootComponent)
 	{
@@ -1082,9 +1139,10 @@ public:
 		return Result;
 	}
 
-	/** Returns rotation of the RootComponent 
-		this is a template for no other reason than to delay compilation until USceneComponent is defined
-	*/ 
+	/**
+	 * Returns rotation of the RootComponent 
+	 * this is a template for no other reason than to delay compilation until USceneComponent is defined
+	 */ 
 	template<class T>
 	static FORCEINLINE FRotator GetActorRotation(T* RootComponent)
 	{
@@ -1096,9 +1154,10 @@ public:
 		return Result;
 	}
 
-	/** Returns scale of the RootComponent 
-		this is a template for no other reason than to delay compilation until USceneComponent is defined
-	*/ 
+	/**
+	 * Returns scale of the RootComponent 
+	 * this is a template for no other reason than to delay compilation until USceneComponent is defined
+	 */ 
 	template<class T>
 	static FORCEINLINE FVector GetActorScale(T* RootComponent)
 	{
@@ -1110,8 +1169,9 @@ public:
 		return Result;
 	}
 
-	/** Returns quaternion of the RootComponent
-		this is a template for no other reason than to delay compilation until USceneComponent is defined
+	/**
+	 * Returns quaternion of the RootComponent
+	 * this is a template for no other reason than to delay compilation until USceneComponent is defined
 	 */ 
 	template<class T>
 	static FORCEINLINE FQuat GetActorQuat(T* RootComponent)
@@ -1131,28 +1191,31 @@ public:
 	DEPRECATED(4.5, "Use GetRootComponent() and cast manually if needed")
 	class UPrimitiveComponent* GetRootPrimitiveComponent() const;
 
-	/** Sets root component to be the specified component.  NewRootComponent's owner should be this actor.  Returns true if successful. */
+	/**
+	 * Sets root component to be the specified component.  NewRootComponent's owner should be this actor.
+	 * @return true if successful
+	 */
 	bool SetRootComponent(class USceneComponent* NewRootComponent);
 
-	/** Returns location of the RootComponent of this Actor*/ 
+	/** Returns the location of the RootComponent of this Actor*/ 
 	FORCEINLINE FVector GetActorLocation() const
 	{
 		return GetActorLocation(RootComponent);
 	}
 
-	/** Returns rotation of the RootComponent of this Actor */
+	/** Returns the rotation of the RootComponent of this Actor */
 	FORCEINLINE FRotator GetActorRotation() const
 	{
 		return GetActorRotation(RootComponent);
 	}
 
-	/** Returns scale of the RootComponent of this Actor */
+	/** Returns the scale of the RootComponent of this Actor */
 	FORCEINLINE FVector GetActorScale() const
 	{
 		return GetActorScale(RootComponent);
 	}
 
-	/** Returns quaternion of the RootComponent of this Actor */
+	/** Returns the quaternion of the RootComponent of this Actor */
 	FORCEINLINE FQuat GetActorQuat() const
 	{
 		return GetActorQuat(RootComponent);
@@ -1413,9 +1476,13 @@ public:
 	/** Update and smooth simulated physic state, replaces PostNetReceiveLocation() and PostNetReceiveVelocity() */
 	virtual void PostNetReceivePhysicState();
 
-	/** Set the owner for this Actor */
+	/** Set the owner of this Actor, used primarily for network replication. */
 	void SetOwner( AActor* NewOwner );
 
+	/**
+	 * Get the owner of this Actor, used primarily for network replication.
+	 * @return Actor that owns this Actor
+	 */
 	UFUNCTION(BlueprintCallable, Category=Actor)
 	AActor* GetOwner() const;
 
@@ -1531,6 +1598,9 @@ public:
 	 */
 	virtual class UNetConnection* GetNetConnection();
 
+	/**
+	 * Gets the net mode for this actor, indicating whether it is a client or server (including standalone/not networked).
+	 */
 	ENetMode GetNetMode() const;
 
 	class UNetDriver * GetNetDriver() const;
@@ -1780,9 +1850,6 @@ public:
 	/** Called once this actor has been deleted */
 	virtual void Destroyed();
 
-	/** Get the local-to-world transform of the RootComponent */
-	FTransform ActorToWorld() const;
-
 	/** Call ReceiveHit, as well as delegates on Actor and Component */
 	void DispatchBlockingHit(UPrimitiveComponent* MyComp, UPrimitiveComponent* OtherComp, bool bSelfMoved, FHitResult const& Hit);
 
@@ -1819,7 +1886,10 @@ public:
 	/** @returns true if the root component is registered and has collision enabled.  */
 	virtual bool IsRootComponentCollisionRegistered() const;
 
-	// Networking - called on client when actor is torn off (bTearOff==true)
+	/**
+	 * Networking - called on client when actor is torn off (bTearOff==true), meaning it's no longer replicated to clients.
+	 * @see bTearOff
+	 */
 	virtual void TornOff();
 
 	//=============================================================================
@@ -1852,6 +1922,7 @@ public:
 	virtual bool CanBeBaseForCharacter(class APawn* Pawn) const;
 
 	/** Apply damage to this actor.
+	 * @see https://www.unrealengine.com/blog/damage-in-ue4
 	 * @param DamageAmount		How much damage to apply
 	 * @param DamageEvent		Data package that fully describes the damage received.
 	 * @param EventInstigator	The Controller responsible for the damage.
@@ -2020,8 +2091,18 @@ public:
 	void ResetOwnedComponents();
 
 private:
+	/**
+	 * All ActorComponents owned by this Actor.
+	 * @see GetComponents()
+	 */
 	TArray<UActorComponent*> OwnedComponents;
 	
+public:
+
+	/** Array of ActorComponents that is actually serialized per-instance. */
+	UPROPERTY(TextExportTransient)
+	TArray<UActorComponent*> SerializedComponents;
+
 public:
 	//=============================================================================
 	// Navigation related functions
@@ -2072,6 +2153,15 @@ public:
 
 	static void MakeNoiseImpl(AActor* NoiseMaker, float Loudness, APawn* NoiseInstigator, const FVector& NoiseLocation);
 	static void SetMakeNoiseDelegate(const FMakeNoiseDelegate& NewDelegate);
+
+	/** A fence to track when the primitive is detached from the scene in the rendering thread. */
+	FRenderCommandFence DetachFence;
+
+// DEPRECATED FUNCTIONS
+
+	/** Get the class of this Actor. */
+	UFUNCTION(BlueprintCallable, meta=(DeprecatedFunction, FriendlyName = "GetActorClass"), Category="Class")
+	class UClass* GetActorClass() const;
 };
 
 
