max agent step height can be defined per navmesh
TTP# 345841

diff --git a/Engine/Source/Runtime/Engine/Classes/AI/Navigation/NavigationTypes.h b/Engine/Source/Runtime/Engine/Classes/AI/Navigation/NavigationTypes.h
--- a/Engine/Source/Runtime/Engine/Classes/AI/Navigation/NavigationTypes.h
+++ b/Engine/Source/Runtime/Engine/Classes/AI/Navigation/NavigationTypes.h
@@ -301,19 +301,25 @@ struct FNavAgentProperties : public FMovementProperties
 	UPROPERTY(EditAnywhere, BlueprintReadWrite, Category = MovementProperties)
 	float AgentHeight;
 
+	/** step height to use, or -1 for default value from navdata's config */
+	UPROPERTY(EditAnywhere, BlueprintReadWrite, Category = MovementProperties)
+	float AgentStepHeight;
+
 	FNavAgentProperties(float Radius = -1.f, float Height = -1.f)
-		: AgentRadius(Radius)
-		, AgentHeight(Height)
+		: AgentRadius(Radius), AgentHeight(Height), AgentStepHeight(-1)
 	{
 	}
 
 	void UpdateWithCollisionComponent(class UShapeComponent* CollisionComponent);
 
 	FORCEINLINE bool IsValid() const { return AgentRadius >= 0 && AgentHeight >= 0; }
+	FORCEINLINE bool HasStepHeightOverride() const { return AgentStepHeight >= 0.0f; }
 
 	FORCEINLINE bool IsEquivalent(const FNavAgentProperties& Other, float Precision = 5.f) const
 	{
-		return FGenericPlatformMath::Abs(AgentRadius - Other.AgentRadius) < Precision && FGenericPlatformMath::Abs(AgentHeight - Other.AgentHeight) < Precision;
+		return FGenericPlatformMath::Abs(AgentRadius - Other.AgentRadius) < Precision &&
+			FGenericPlatformMath::Abs(AgentHeight - Other.AgentHeight) < Precision &&
+			FGenericPlatformMath::Abs(AgentStepHeight - Other.AgentStepHeight) < Precision;
 	}
 
 	bool operator==(const FNavAgentProperties& Other) const
@@ -329,7 +335,7 @@ struct FNavAgentProperties : public FMovementProperties
 
 inline uint32 GetTypeHash(const FNavAgentProperties& A)
 {
-	return (int16(A.AgentRadius) << 16) | int16(A.AgentHeight);
+	return ((int16(A.AgentRadius) << 16) | int16(A.AgentHeight)) ^ int32(A.AgentStepHeight);
 }
 
 USTRUCT()
diff --git a/Engine/Source/Runtime/Engine/Private/AI/Navigation/RecastNavMesh.cpp b/Engine/Source/Runtime/Engine/Private/AI/Navigation/RecastNavMesh.cpp
--- a/Engine/Source/Runtime/Engine/Private/AI/Navigation/RecastNavMesh.cpp
+++ b/Engine/Source/Runtime/Engine/Private/AI/Navigation/RecastNavMesh.cpp
@@ -947,6 +947,11 @@ void ARecastNavMesh::SetConfig(const FNavDataConfig& Src)
 	NavDataConfig = Src; 
 	AgentMaxHeight = AgentHeight = Src.AgentHeight;
 	AgentRadius = Src.AgentRadius;
+
+	if (Src.HasStepHeightOverride())
+	{
+		AgentMaxStepHeight = Src.AgentStepHeight;
+	}
 }
 
 void ARecastNavMesh::FillConfig(FNavDataConfig& Dest)
@@ -954,6 +959,7 @@ void ARecastNavMesh::FillConfig(FNavDataConfig& Dest)
 	Dest = NavDataConfig;
 	Dest.AgentHeight = AgentHeight;
 	Dest.AgentRadius = AgentRadius;
+	Dest.AgentStepHeight = AgentMaxStepHeight;
 }
 
 bool ARecastNavMesh::DoesSupportAgent(const FNavAgentProperties& AgentProps) const
