Added per shape filtering for welded bodies codereview(Lina.Halper)

diff --git a/Engine/Source/Runtime/Engine/Classes/PhysicsEngine/BodyInstance.h b/Engine/Source/Runtime/Engine/Classes/PhysicsEngine/BodyInstance.h
--- a/Engine/Source/Runtime/Engine/Classes/PhysicsEngine/BodyInstance.h
+++ b/Engine/Source/Runtime/Engine/Classes/PhysicsEngine/BodyInstance.h
@@ -178,9 +178,6 @@ public:
 	/** determines if the body is currently welded */
 	uint32 bWelded : 1;
 
-	/** needed to distinguish between the root wanting hit events, and the root turning hit events on for one of its welded children */
-	uint32 bWeldedNotifyRigidBodyCollision : 1;
-
 	/** If object should start awake, or if it should initially be sleeping */
 	UPROPERTY(EditAnywhere, BlueprintReadOnly, Category=Physics, meta=(editcondition = "bSimulatePhysics"))
 	uint32 bStartAwake:1;
@@ -644,6 +641,11 @@ private:
 	 *  @return true if a hit is found
 	 */
 	bool InternalSweepPhysX(struct FHitResult& OutHit, const FVector& Start, const FVector& End, const FCollisionShape& Shape, bool bTraceComplex, const physx::PxRigidActor* RigidBody, const physx::PxGeometry* Geometry) const;
+
+	/** 
+	 * Helper function to update per shape filtering info. This should interface is not very friendly and should only be used from inside FBodyInstance
+	 */
+	void UpdatePhysicsShapeFilterData(uint32 SkelMeshCompID, bool bUseComplexAsSimple, bool bUseSimpleAsComplex, bool bPhysicsStatic, TEnumAsByte<ECollisionEnabled::Type> * CollisionEnabledOverride, FCollisionResponseContainer * ResponseOverride, bool * bNotifyOverride);
 #endif 
 	/**
 	 * Invalidate Collision Profile Name
diff --git a/Engine/Source/Runtime/Engine/Private/PhysicsEngine/BodyInstance.cpp b/Engine/Source/Runtime/Engine/Private/PhysicsEngine/BodyInstance.cpp
--- a/Engine/Source/Runtime/Engine/Private/PhysicsEngine/BodyInstance.cpp
+++ b/Engine/Source/Runtime/Engine/Private/PhysicsEngine/BodyInstance.cpp
@@ -199,7 +199,6 @@ FBodyInstance::FBodyInstance()
 , bSimulatePhysics(false)
 , bAutoWeld(false)
 , bWelded(false)
-, bWeldedNotifyRigidBodyCollision(false)
 , bStartAwake(true)
 , bEnableGravity(true)
 , bUseAsyncScene(false)
@@ -604,101 +603,11 @@ ECollisionEnabled::Type FBodyInstance::GetCollisionEnabled() const
 	}
 }
 
-/** Update the filter data on the physics shapes, based on the owning component flags. */
-void FBodyInstance::UpdatePhysicsFilterData(bool bForceSimpleAsComplex)
-{
-	// Do nothing if no physics actor
-	if (!IsValidBodyInstance())
-	{
-		return;
-	}
-
-	// this can happen in landscape height field collision component
-	if (!BodySetup.IsValid())
-	{
-		return;
-	}
-
-	// Figure out if we are static
-	AActor* Owner = OwnerComponent.IsValid() ? OwnerComponent.Get()->GetOwner() : NULL;
-	int32 OwnerID = (Owner != NULL) ? Owner->GetUniqueID() : 0;
-	const bool bPhysicsStatic = !OwnerComponent.IsValid() || OwnerComponent.Get()->IsWorldGeometry();
-
-	// Grab collision setting from body instance
-	TEnumAsByte<ECollisionEnabled::Type> UseCollisionEnabled = GetCollisionEnabled(); // this checks actor override
-	bool bUseNotifyRBCollision = bNotifyRigidBodyCollision || bWeldedNotifyRigidBodyCollision;
-	FCollisionResponseContainer UseResponse = CollisionResponses.GetResponseContainer();
-
-	// Get skelmeshcomp ID
-	uint32 SkelMeshCompID = 0;
-	if (USkeletalMeshComponent* SkelMeshComp = Cast<USkeletalMeshComponent>(OwnerComponent.Get()))
-	{
-		SkelMeshCompID = SkelMeshComp->GetUniqueID();
-
-		// In skeletal case, collision enable/disable/movement should be overriden by mesh component
-		// being in the physics asset, and not having collision is a waste and it can cause a bug where disconnected bodies
-		UseCollisionEnabled = SkelMeshComp->BodyInstance.CollisionEnabled; 
-		ObjectType = SkelMeshComp->GetCollisionObjectType();
-
-		if (BodySetup->CollisionReponse == EBodyCollisionResponse::BodyCollision_Enabled)
-		{
-			UseResponse.SetAllChannels(ECR_Block);
-		}
-		else if (BodySetup->CollisionReponse == EBodyCollisionResponse::BodyCollision_Disabled)
-		{
-			UseResponse.SetAllChannels(ECR_Ignore);
-		}
-
-		UseResponse = FCollisionResponseContainer::CreateMinContainer(UseResponse, SkelMeshComp->BodyInstance.CollisionResponses.GetResponseContainer());
-		bUseNotifyRBCollision = bUseNotifyRBCollision && (SkelMeshComp->BodyInstance.bNotifyRigidBodyCollision || SkelMeshComp->BodyInstance.bWeldedNotifyRigidBodyCollision);
-	}
-
-#if WITH_EDITOR
-	// if no collision, but if world wants to enable trace collision for components, allow it
-	if ((UseCollisionEnabled == ECollisionEnabled::NoCollision) && Owner && (Owner->IsA(AVolume::StaticClass()) == false))
-	{
-		UWorld* World = Owner->GetWorld();
-		UPrimitiveComponent* PrimComp = OwnerComponent.Get();
-		if (World && World->bEnableTraceCollision && 
-			(PrimComp->IsA(UStaticMeshComponent::StaticClass()) || PrimComp->IsA(USkeletalMeshComponent::StaticClass()) || PrimComp->IsA(UBrushComponent::StaticClass())))
-		{
-			//UE_LOG(LogPhysics, Warning, TEXT("Enabling collision %s : %s"), *GetNameSafe(Owner), *GetNameSafe(OwnerComponent.Get()));
-			// clear all other channel just in case other people using those channels to do something
-			UseResponse.SetAllChannels(ECR_Ignore);
-			UseCollisionEnabled = ECollisionEnabled::QueryOnly;
-		}
-	}
-#endif
-
-	const bool bUseComplexAsSimple = !bForceSimpleAsComplex && (BodySetup.Get()->CollisionTraceFlag == CTF_UseComplexAsSimple);
-	const bool bUseSimpleAsComplex = bForceSimpleAsComplex || (BodySetup.Get()->CollisionTraceFlag == CTF_UseSimpleAsComplex);
-
 #if WITH_PHYSX
+void FBodyInstance::UpdatePhysicsShapeFilterData(uint32 SkelMeshCompID, bool bUseComplexAsSimple, bool bUseSimpleAsComplex, bool bPhysicsStatic, TEnumAsByte<ECollisionEnabled::Type> * CollisionEnabledOverride, FCollisionResponseContainer * ResponseOverride, bool * bNotifyOverride)
+{
 	if (PxRigidActor* PActor = GetPxRigidActor())
 	{
-		// Create the filterdata structs
-		PxFilterData PSimFilterData;
-		PxFilterData PSimpleQueryData;
-		PxFilterData PComplexQueryData;
-		if (UseCollisionEnabled != ECollisionEnabled::NoCollision)
-		{
-			CreateShapeFilterData(ObjectType, OwnerID, UseResponse, SkelMeshCompID, InstanceBodyIndex, PSimpleQueryData, PSimFilterData, bUseCCD && !bPhysicsStatic, bUseNotifyRBCollision, bPhysicsStatic);
-			PComplexQueryData = PSimpleQueryData;
-
-			// Build filterdata variations for complex and simple
-			PSimpleQueryData.word3 |= EPDF_SimpleCollision;
-			if (bUseSimpleAsComplex)
-			{
-				PSimpleQueryData.word3 |= EPDF_ComplexCollision;
-			}
-
-			PComplexQueryData.word3 |= EPDF_ComplexCollision;
-			if (bUseComplexAsSimple)
-			{
-				PComplexQueryData.word3 |= EPDF_SimpleCollision;
-			}
-		}
-
 		// Iterate over all shapes and assign filterdata
 		int32 NumSyncShapes = 0;
 		TArray<PxShape*> AllShapes = GetAllShapes(NumSyncShapes);
@@ -716,6 +625,40 @@ void FBodyInstance::UpdatePhysicsFilterData(bool bForceSimpleAsComplex)
 		for (int32 ShapeIdx = 0; ShapeIdx < AllShapes.Num(); ShapeIdx++)
 		{
 			PxShape* PShape = AllShapes[ShapeIdx];
+			FBodyInstance * BI = FPhysxUserData::Get<FBodyInstance>(PShape->userData);
+			BI = BI ? BI : this;
+
+			const TEnumAsByte<ECollisionEnabled::Type> & UseCollisionEnabled = CollisionEnabledOverride ? *CollisionEnabledOverride : BI->GetCollisionEnabled();
+			const FCollisionResponseContainer & UseResponse = ResponseOverride ? *ResponseOverride : BI->CollisionResponses.GetResponseContainer();
+			bool bUseNotify = bNotifyOverride ? *bNotifyOverride : BI->bNotifyRigidBodyCollision;
+
+
+			AActor* Owner = BI->OwnerComponent.IsValid() ? BI->OwnerComponent.Get()->GetOwner() : NULL;
+			int32 OwnerID = (Owner != NULL) ? Owner->GetUniqueID() : 0;
+
+			// Create the filterdata structs
+			PxFilterData PSimFilterData;
+			PxFilterData PSimpleQueryData;
+			PxFilterData PComplexQueryData;
+			if (UseCollisionEnabled != ECollisionEnabled::NoCollision)
+			{
+				CreateShapeFilterData(BI->ObjectType, OwnerID, UseResponse, SkelMeshCompID, BI->InstanceBodyIndex, PSimpleQueryData, PSimFilterData, BI->bUseCCD && !bPhysicsStatic, bUseNotify, bPhysicsStatic);
+				PComplexQueryData = PSimpleQueryData;
+
+				// Build filterdata variations for complex and simple
+				PSimpleQueryData.word3 |= EPDF_SimpleCollision;
+				if (bUseSimpleAsComplex)
+				{
+					PSimpleQueryData.word3 |= EPDF_ComplexCollision;
+				}
+
+				PComplexQueryData.word3 |= EPDF_ComplexCollision;
+				if (bUseComplexAsSimple)
+				{
+					PComplexQueryData.word3 |= EPDF_SimpleCollision;
+				}
+			}
+
 			PShape->setSimulationFilterData(PSimFilterData);
 
 			// If query collision is enabled..
@@ -804,6 +747,92 @@ void FBodyInstance::UpdatePhysicsFilterData(bool bForceSimpleAsComplex)
 
 		SCENE_UNLOCK_WRITE(AsyncScene);
 	}
+}
+#endif
+
+/** Update the filter data on the physics shapes, based on the owning component flags. */
+void FBodyInstance::UpdatePhysicsFilterData(bool bForceSimpleAsComplex)
+{
+	// Do nothing if no physics actor
+	if (!IsValidBodyInstance())
+	{
+		return;
+	}
+
+	// this can happen in landscape height field collision component
+	if (!BodySetup.IsValid())
+	{
+		return;
+	}
+
+	// Figure out if we are static
+	AActor* Owner = OwnerComponent.IsValid() ? OwnerComponent.Get()->GetOwner() : NULL;
+	const bool bPhysicsStatic = !OwnerComponent.IsValid() || OwnerComponent.Get()->IsWorldGeometry();
+
+	// Grab collision setting from body instance
+	TEnumAsByte<ECollisionEnabled::Type> UseCollisionEnabled = GetCollisionEnabled(); // this checks actor override
+	bool bUseNotifyRBCollision = bNotifyRigidBodyCollision;
+	FCollisionResponseContainer UseResponse = CollisionResponses.GetResponseContainer();
+
+	bool bUseCollisionEnabledOverride = false;
+	bool bResponseOverride = false;
+	bool bNotifyOverride = false;
+
+	// Get skelmeshcomp ID
+	uint32 SkelMeshCompID = 0;
+	if (USkeletalMeshComponent* SkelMeshComp = Cast<USkeletalMeshComponent>(OwnerComponent.Get()))
+	{
+		SkelMeshCompID = SkelMeshComp->GetUniqueID();
+
+		// In skeletal case, collision enable/disable/movement should be overriden by mesh component
+		// being in the physics asset, and not having collision is a waste and it can cause a bug where disconnected bodies
+		UseCollisionEnabled = SkelMeshComp->BodyInstance.CollisionEnabled; 
+		ObjectType = SkelMeshComp->GetCollisionObjectType();
+
+		bUseCollisionEnabledOverride = true;
+
+		if (BodySetup->CollisionReponse == EBodyCollisionResponse::BodyCollision_Enabled)
+		{
+			UseResponse.SetAllChannels(ECR_Block);
+		}
+		else if (BodySetup->CollisionReponse == EBodyCollisionResponse::BodyCollision_Disabled)
+		{
+			UseResponse.SetAllChannels(ECR_Ignore);
+		}
+
+		UseResponse = FCollisionResponseContainer::CreateMinContainer(UseResponse, SkelMeshComp->BodyInstance.CollisionResponses.GetResponseContainer());
+		bUseNotifyRBCollision = bUseNotifyRBCollision && SkelMeshComp->BodyInstance.bNotifyRigidBodyCollision;
+		bResponseOverride = true;
+		bNotifyOverride = true;
+	}
+
+#if WITH_EDITOR
+	// if no collision, but if world wants to enable trace collision for components, allow it
+	if ((UseCollisionEnabled == ECollisionEnabled::NoCollision) && Owner && (Owner->IsA(AVolume::StaticClass()) == false))
+	{
+		UWorld* World = Owner->GetWorld();
+		UPrimitiveComponent* PrimComp = OwnerComponent.Get();
+		if (World && World->bEnableTraceCollision && 
+			(PrimComp->IsA(UStaticMeshComponent::StaticClass()) || PrimComp->IsA(USkeletalMeshComponent::StaticClass()) || PrimComp->IsA(UBrushComponent::StaticClass())))
+		{
+			//UE_LOG(LogPhysics, Warning, TEXT("Enabling collision %s : %s"), *GetNameSafe(Owner), *GetNameSafe(OwnerComponent.Get()));
+			// clear all other channel just in case other people using those channels to do something
+			UseResponse.SetAllChannels(ECR_Ignore);
+			UseCollisionEnabled = ECollisionEnabled::QueryOnly;
+			bResponseOverride = true;
+			bUseCollisionEnabledOverride = true;
+		}
+	}
+#endif
+
+	const bool bUseComplexAsSimple = !bForceSimpleAsComplex && (BodySetup.Get()->CollisionTraceFlag == CTF_UseComplexAsSimple);
+	const bool bUseSimpleAsComplex = bForceSimpleAsComplex || (BodySetup.Get()->CollisionTraceFlag == CTF_UseSimpleAsComplex);
+
+#if WITH_PHYSX
+	TEnumAsByte<ECollisionEnabled::Type> * CollisionEnabledOverride = bUseCollisionEnabledOverride ? &UseCollisionEnabled : NULL;
+	FCollisionResponseContainer * ResponseOverride = bResponseOverride ? &UseResponse : NULL;
+	bool * bNotifyOverridePtr = bNotifyOverride ? &bUseNotifyRBCollision : NULL;
+	UpdatePhysicsShapeFilterData(SkelMeshCompID, bUseComplexAsSimple, bUseSimpleAsComplex, bPhysicsStatic, CollisionEnabledOverride, ResponseOverride, bNotifyOverridePtr);
 #endif
 
 #if WITH_BOX2D
@@ -1573,11 +1602,6 @@ bool FBodyInstance::Weld(FBodyInstance* TheirBody, const FTransform& TheirTM)
 		PShape->userData = &TheirBody->PhysxUserData;
 	}
 
-	if (TheirBody->bNotifyRigidBodyCollision)
-	{
-		bWeldedNotifyRigidBodyCollision = true;
-	}
-
 	PostShapeChange();
 
 	//remove their body from scenes
@@ -1631,40 +1655,6 @@ void FBodyInstance::UnWeld(FBodyInstance* TheirBI)
 		}
 	}
 
-
-	/*for (int32 ShapeIdx = 0; ShapeIdx < NumSyncShapes; ++ShapeIdx)
-	{
-		PxShape* PShape = PShapes[ShapeIdx];
-		if (FBodyInstance ** BIPtrPtr = ShapeToBodyMap.Find(PShape))
-		{
-			bNeedsNotification |= (*BIPtrPtr)->bNotifyRigidBodyCollision;
-
-			if (*BIPtrPtr == BI)
-			{
-				PShape->userData = NULL;
-				RigidActorSync->detachShape(*PShape);
-				ShapeToBodyMap.Remove(PShape);
-			}
-		}
-	}
-
-	for (int32 ShapeIdx = NumSyncShapes; ShapeIdx < PShapes.Num(); ++ShapeIdx)
-	{
-		PxShape* PShape = PShapes[ShapeIdx];
-		if (FBodyInstance ** BIPtrPtr = ShapeToBodyMap.Find(PShape))
-		{
-			bNeedsNotification |= (*BIPtrPtr)->bNotifyRigidBodyCollision;
-			if (*BIPtrPtr == BI)
-			{
-				PShape->userData = NULL;
-				RigidActorAsync->detachShape(*PShape);
-				ShapeToBodyMap.Remove(PShape);
-			}
-		}
-	}*/
-
-	bWeldedNotifyRigidBodyCollision = bNeedsNotification;
-
 	PostShapeChange();
 #endif
 }
