Removed TickAnim from InitAnim as that is unintuitive and seems it shouldn't happen.

diff --git a/Engine/Source/Runtime/Engine/Private/SkeletalMeshComponent.cpp b/Engine/Source/Runtime/Engine/Private/SkeletalMeshComponent.cpp
--- a/Engine/Source/Runtime/Engine/Private/SkeletalMeshComponent.cpp
+++ b/Engine/Source/Runtime/Engine/Private/SkeletalMeshComponent.cpp
@@ -221,7 +221,6 @@ void USkeletalMeshComponent::InitAnim(bool bForceReinit)
 
 		InitializeAnimScriptInstance(bForceReinit);
 
-		TickAnimation(0.f);
 		RefreshBoneTransforms();
 		UpdateComponentToWorld();
 	}
