fix clang compile
#codeview Ori.Cohen

diff --git a/Engine/Source/Runtime/Engine/Private/PhysicsEngine/BodyInstance.cpp b/Engine/Source/Runtime/Engine/Private/PhysicsEngine/BodyInstance.cpp
--- a/Engine/Source/Runtime/Engine/Private/PhysicsEngine/BodyInstance.cpp
+++ b/Engine/Source/Runtime/Engine/Private/PhysicsEngine/BodyInstance.cpp
@@ -628,7 +628,7 @@ void FBodyInstance::UpdatePhysicsShapeFilterData(uint32 SkelMeshCompID, bool bUs
 			FBodyInstance * BI = FPhysxUserData::Get<FBodyInstance>(PShape->userData);
 			BI = BI ? BI : this;
 
-			const TEnumAsByte<ECollisionEnabled::Type> & UseCollisionEnabled = CollisionEnabledOverride ? *CollisionEnabledOverride : BI->GetCollisionEnabled();
+			const TEnumAsByte<ECollisionEnabled::Type> & UseCollisionEnabled = CollisionEnabledOverride ? *CollisionEnabledOverride : (TEnumAsByte<ECollisionEnabled::Type>)BI->GetCollisionEnabled();
 			const FCollisionResponseContainer & UseResponse = ResponseOverride ? *ResponseOverride : BI->CollisionResponses.GetResponseContainer();
 			bool bUseNotify = bNotifyOverride ? *bNotifyOverride : BI->bNotifyRigidBodyCollision;
 
