Exposed smartname code to add unlinked names to the skeleton on load in the runtime aswell as the editor and removed runtime checks on the names.
#ttp 345458 - CRITICAL: ContentExamples: Curve driven animation dont work after packaging a project in either Win32 or Win64.

diff --git a/Engine/Source/Runtime/Engine/Private/Animation/SmartName.cpp b/Engine/Source/Runtime/Engine/Private/Animation/SmartName.cpp
--- a/Engine/Source/Runtime/Engine/Private/Animation/SmartName.cpp
+++ b/Engine/Source/Runtime/Engine/Private/Animation/SmartName.cpp
@@ -45,7 +45,6 @@ bool FSmartNameMapping::AddName(FName Name, UID& OutUid)
 
 bool FSmartNameMapping::GetName(const UID& Uid, FName& OutName) const
 {
-	check(Uid < NextUid);
 	const FName* FoundName = UidMap.Find(Uid);
 	if(FoundName)
 	{
@@ -57,7 +56,6 @@ bool FSmartNameMapping::GetName(const UID& Uid, FName& OutName) const
 
 bool FSmartNameMapping::Rename(const UID& Uid, FName NewName)
 {
-	check(Uid < NextUid);
 	FName* ExistingName = UidMap.Find(Uid);
 	if(ExistingName)
 	{
@@ -69,7 +67,6 @@ bool FSmartNameMapping::Rename(const UID& Uid, FName NewName)
 
 bool FSmartNameMapping::Remove(const UID& Uid)
 {
-	check(Uid < NextUid);
 	bool bRemoved = UidMap.Remove(Uid) != 0;
 
 	if(bRemoved)
