UE4 - Remove Feature Level SM3 & replace it with ES 3.1.
- Fixes up the featurelevel switch on material nodes to default all ES2 inputs to also occur on ES3.1

diff --git a/Engine/Source/Editor/LevelEditor/Private/LevelEditorToolBar.cpp b/Engine/Source/Editor/LevelEditor/Private/LevelEditorToolBar.cpp
--- a/Engine/Source/Editor/LevelEditor/Private/LevelEditorToolBar.cpp
+++ b/Engine/Source/Editor/LevelEditor/Private/LevelEditorToolBar.cpp
@@ -601,7 +601,7 @@ static void MakeShaderModelPreviewMenu(FMenuBuilder& MenuBuilder)
 	{
 		for (int32 i = GMaxRHIFeatureLevel; i >= 0; --i)
 		{
-			if (i != ERHIFeatureLevel::SM3)
+			if (i != ERHIFeatureLevel::ES3_1)
 			{
 				MenuBuilder.AddMenuEntry(FLevelEditorCommands::Get().FeatureLevelPreview[i]);
 			}
diff --git a/Engine/Source/Runtime/Core/Public/UObject/ObjectVersion.h b/Engine/Source/Runtime/Core/Public/UObject/ObjectVersion.h
--- a/Engine/Source/Runtime/Core/Public/UObject/ObjectVersion.h
+++ b/Engine/Source/Runtime/Core/Public/UObject/ObjectVersion.h
@@ -529,7 +529,7 @@ enum EUnrealEngineObjectUE4Version
 	VER_UE4_ANIMATION_REMOVE_NANS,
 	// Change skeleton preview attached assets property type
 	VER_UE4_SKELETON_ASSET_PROPERTY_TYPE_CHANGE,
- 	// Fix some blueprint variables that have the CPF_DisableEditOnTemplate flag set
+	// Fix some blueprint variables that have the CPF_DisableEditOnTemplate flag set
 	// when they shouldn't
 	VER_UE4_FIX_BLUEPRINT_VARIABLE_FLAGS,
 	// Vehicles use Nm for Torque instead of cm and RPM instead of rad/s part two (missed conversion for some variables
@@ -540,7 +540,7 @@ enum EUnrealEngineObjectUE4Version
 	VER_UE4_STATIC_MESH_SCREEN_SIZE_LODS,
 	// Requires test of material coords to ensure they're saved correctly
 	VER_UE4_FIX_MATERIAL_COORDS,
- 	// Changed SpeedTree wind presets to v7
+	// Changed SpeedTree wind presets to v7
 	VER_UE4_SPEEDTREE_WIND_V7,
 	// NeedsLoadForEditorGame added
 	VER_UE4_LOAD_FOR_EDITOR_GAME,
@@ -604,6 +604,8 @@ enum EUnrealEngineObjectUE4Version
 	VER_UE4_RENAME_CROUCHMOVESCHARACTERDOWN,
 	// Lightmap parameters in FMeshBuildSettings
 	VER_UE4_LIGHTMAP_MESH_BUILD_SETTINGS,
+	// Rename SM3 to ES3_1 and updates featurelevel material node selector
+	VER_UE4_RENAME_SM3_TO_ES3_1,
 
 	// -----<new versions can be added before this line>-------------------------------------------------
 	// - this needs to be the last line (see note below)
diff --git a/Engine/Source/Runtime/Engine/Classes/Materials/MaterialExpressionFeatureLevelSwitch.h b/Engine/Source/Runtime/Engine/Classes/Materials/MaterialExpressionFeatureLevelSwitch.h
--- a/Engine/Source/Runtime/Engine/Classes/Materials/MaterialExpressionFeatureLevelSwitch.h
+++ b/Engine/Source/Runtime/Engine/Classes/Materials/MaterialExpressionFeatureLevelSwitch.h
@@ -17,6 +17,10 @@ class UMaterialExpressionFeatureLevelSwitch : public UMaterialExpression
 	UPROPERTY()
 	FExpressionInput Inputs[ERHIFeatureLevel::Num];
 
+	// Begin UObject interface.
+	virtual void Serialize(FArchive& Ar);
+	// End UObject interface.
+
 	// Begin UMaterialExpression Interface
 	virtual int32 Compile(class FMaterialCompiler* Compiler, int32 OutputIndex, int32 MultiplexIndex) override;
 	virtual void GetCaption(TArray<FString>& OutCaptions) const override;
diff --git a/Engine/Source/Runtime/Engine/Private/GPUSkinVertexFactory.cpp b/Engine/Source/Runtime/Engine/Private/GPUSkinVertexFactory.cpp
--- a/Engine/Source/Runtime/Engine/Private/GPUSkinVertexFactory.cpp
+++ b/Engine/Source/Runtime/Engine/Private/GPUSkinVertexFactory.cpp
@@ -85,10 +85,10 @@ uint32 FSharedPoolPolicyData::BucketSizes[NumPoolBucketSizes] = {
 FBoneBufferTypeRef FBoneBufferPoolPolicy::CreateResource(CreationArguments Args)
 {
 	uint32 BufferSize = GetPoolBucketSize(GetPoolBucketIndex(Args));
-    FBoneBuffer Buffer;
+	FBoneBuffer Buffer;
 	FRHIResourceCreateInfo CreateInfo;
-    Buffer.VertexBufferRHI = RHICreateVertexBuffer( BufferSize, (BUF_Dynamic | BUF_ShaderResource), CreateInfo );
-    Buffer.VertexBufferSRV = RHICreateShaderResourceView( Buffer.VertexBufferRHI, sizeof(FVector4), PF_A32B32G32R32F );
+	Buffer.VertexBufferRHI = RHICreateVertexBuffer( BufferSize, (BUF_Dynamic | BUF_ShaderResource), CreateInfo );
+	Buffer.VertexBufferSRV = RHICreateShaderResourceView( Buffer.VertexBufferRHI, sizeof(FVector4), PF_A32B32G32R32F );
 	return Buffer;
 }
 
@@ -765,7 +765,7 @@ void TGPUSkinAPEXClothVertexFactory<bExtraBoneInfluencesT>::ModifyCompilationEnv
 template <bool bExtraBoneInfluencesT>
 bool TGPUSkinAPEXClothVertexFactory<bExtraBoneInfluencesT>::ShouldCache(EShaderPlatform Platform, const class FMaterial* Material, const class FShaderType* ShaderType)
 {
-	return GetMaxSupportedFeatureLevel(Platform) >= ERHIFeatureLevel::SM3
+	return GetMaxSupportedFeatureLevel(Platform) >= ERHIFeatureLevel::SM4
 		&& (Material->IsUsedWithAPEXCloth() || Material->IsSpecialEngineMaterial()) 
 		&& Super::ShouldCache(Platform, Material, ShaderType);
 }
diff --git a/Engine/Source/Runtime/Engine/Private/Landscape/Landscape.cpp b/Engine/Source/Runtime/Engine/Private/Landscape/Landscape.cpp
--- a/Engine/Source/Runtime/Engine/Private/Landscape/Landscape.cpp
+++ b/Engine/Source/Runtime/Engine/Private/Landscape/Landscape.cpp
@@ -670,7 +670,7 @@ FPrimitiveSceneProxy* ULandscapeComponent::CreateSceneProxy()
 {
 	const auto FeatureLevel = GetWorld()->FeatureLevel;
 	FPrimitiveSceneProxy* Proxy = NULL;
-	if (FeatureLevel >= ERHIFeatureLevel::SM3)
+	if (FeatureLevel >= ERHIFeatureLevel::SM4)
 	{
 #if WITH_EDITOR
 		if (EditToolRenderData == NULL)
diff --git a/Engine/Source/Runtime/Engine/Private/Materials/MaterialExpressions.cpp b/Engine/Source/Runtime/Engine/Private/Materials/MaterialExpressions.cpp
--- a/Engine/Source/Runtime/Engine/Private/Materials/MaterialExpressions.cpp
+++ b/Engine/Source/Runtime/Engine/Private/Materials/MaterialExpressions.cpp
@@ -3260,17 +3260,17 @@ void  UMaterialExpressionBreakMaterialAttributes::Serialize( FArchive& Ar )
 {
 	Super::Serialize(Ar);
 
- 	//Switch the connection over to the new input.
- 	if( Ar.UE4Ver() < VER_UE4_MATERIAL_ATTRIBUTES_MULTIPLEX )
- 	{
- 		MaterialAttributes.Expression = Struct.Expression;
- 		MaterialAttributes.OutputIndex = Struct.OutputIndex;
- 		MaterialAttributes.Mask = Struct.Mask;
- 		MaterialAttributes.MaskR = Struct.MaskR;
- 		MaterialAttributes.MaskG = Struct.MaskG;
- 		MaterialAttributes.MaskB = Struct.MaskB;
- 		MaterialAttributes.MaskA = Struct.MaskA;
- 	}
+	//Switch the connection over to the new input.
+	if( Ar.UE4Ver() < VER_UE4_MATERIAL_ATTRIBUTES_MULTIPLEX )
+	{
+		MaterialAttributes.Expression = Struct.Expression;
+		MaterialAttributes.OutputIndex = Struct.OutputIndex;
+		MaterialAttributes.Mask = Struct.Mask;
+		MaterialAttributes.MaskR = Struct.MaskR;
+		MaterialAttributes.MaskG = Struct.MaskG;
+		MaterialAttributes.MaskB = Struct.MaskB;
+		MaterialAttributes.MaskA = Struct.MaskA;
+	}
 }
 
 int32 UMaterialExpressionBreakMaterialAttributes::Compile(class FMaterialCompiler* Compiler, int32 OutputIndex, int32 MultiplexIndex)
@@ -4050,6 +4050,16 @@ bool UMaterialExpressionFeatureLevelSwitch::IsResultMaterialAttributes(int32 Out
 	return false;
 }
 
+void UMaterialExpressionFeatureLevelSwitch::Serialize(FArchive& Ar)
+{
+	Super::Serialize(Ar);
+	if (Ar.IsLoading() && Ar.UE4Ver() < VER_UE4_RENAME_SM3_TO_ES3_1)
+	{
+		// Copy the ES2 input to SM3 (since SM3 will now become ES3_1 and we don't want broken content)
+		Inputs[ERHIFeatureLevel::ES3_1] = Inputs[ERHIFeatureLevel::ES2];
+	}
+}
+
 //
 //	UMaterialExpressionNormalize
 //
@@ -8831,36 +8841,36 @@ int32 UMaterialExpressionLandscapeLayerBlend::Compile(class FMaterialCompiler* C
 		// LB_AlphaBlend layers are blended last
 		if (Layer.BlendType != LB_AlphaBlend)
 		{
-		    // Height input
-		    const int32 HeightCode = Layer.HeightInput.Expression ? Layer.HeightInput.Compile(Compiler, MultiplexIndex) : Compiler->Constant(Layer.ConstHeightInput);
-    
-		    const int32 WeightCode = Compiler->StaticTerrainLayerWeight(Layer.LayerName,Layer.PreviewWeight > 0.0f ? Compiler->Constant(Layer.PreviewWeight) : INDEX_NONE);
-		    if( WeightCode != INDEX_NONE )
-		    {
-			    switch( Layer.BlendType )
-			    {
-			    case LB_WeightBlend:
-				    {
-					    // Store the weight plus accumulate the sum of all weights so far
-					    WeightCodes[LayerIdx] = WeightCode;
-					    WeightSumCode = Compiler->Add(WeightSumCode, WeightCode);
-				    }
-				    break;
-			    case LB_HeightBlend:
-				    {
-					    bNeedsRenormalize = true;
-    
-					    // Modify weight with height
-					    int32 ModifiedWeightCode = Compiler->Clamp(
-						    Compiler->Add(Compiler->Lerp(Compiler->Constant(-1.f), Compiler->Constant(1.f), WeightCode), HeightCode),
-						    Compiler->Constant(0.f), Compiler->Constant(1.f) );
-    
-					    // Store the final weight plus accumulate the sum of all weights so far
-					    WeightCodes[LayerIdx] = ModifiedWeightCode;
-					    WeightSumCode = Compiler->Add(WeightSumCode, ModifiedWeightCode);
-				    }
-				    break;
-			    }
+			// Height input
+			const int32 HeightCode = Layer.HeightInput.Expression ? Layer.HeightInput.Compile(Compiler, MultiplexIndex) : Compiler->Constant(Layer.ConstHeightInput);
+	
+			const int32 WeightCode = Compiler->StaticTerrainLayerWeight(Layer.LayerName,Layer.PreviewWeight > 0.0f ? Compiler->Constant(Layer.PreviewWeight) : INDEX_NONE);
+			if( WeightCode != INDEX_NONE )
+			{
+				switch( Layer.BlendType )
+				{
+				case LB_WeightBlend:
+					{
+						// Store the weight plus accumulate the sum of all weights so far
+						WeightCodes[LayerIdx] = WeightCode;
+						WeightSumCode = Compiler->Add(WeightSumCode, WeightCode);
+					}
+					break;
+				case LB_HeightBlend:
+					{
+						bNeedsRenormalize = true;
+	
+						// Modify weight with height
+						int32 ModifiedWeightCode = Compiler->Clamp(
+							Compiler->Add(Compiler->Lerp(Compiler->Constant(-1.f), Compiler->Constant(1.f), WeightCode), HeightCode),
+							Compiler->Constant(0.f), Compiler->Constant(1.f) );
+	
+						// Store the final weight plus accumulate the sum of all weights so far
+						WeightCodes[LayerIdx] = ModifiedWeightCode;
+						WeightSumCode = Compiler->Add(WeightSumCode, ModifiedWeightCode);
+					}
+					break;
+				}
 			}
 		}
 	}
diff --git a/Engine/Source/Runtime/Engine/Private/SkeletalMesh.cpp b/Engine/Source/Runtime/Engine/Private/SkeletalMesh.cpp
--- a/Engine/Source/Runtime/Engine/Private/SkeletalMesh.cpp
+++ b/Engine/Source/Runtime/Engine/Private/SkeletalMesh.cpp
@@ -1904,7 +1904,7 @@ bool FSkeletalMeshResource::RequiresCPUSkinning(ERHIFeatureLevel::Type FeatureLe
 	const int32 MaxGPUSkinBones = GetFeatureLevelMaxNumberOfBones(FeatureLevel);
 	const int32 MaxBonesPerChunk = GetMaxBonesPerChunk();
 	// Do CPU skinning if we need too many bones per chunk, or if we have too many influences per vertex on lower end
-	return (MaxBonesPerChunk > MaxGPUSkinBones) || (HasExtraBoneInfluences() && FeatureLevel < ERHIFeatureLevel::SM3);
+	return (MaxBonesPerChunk > MaxGPUSkinBones) || (HasExtraBoneInfluences() && FeatureLevel < ERHIFeatureLevel::ES3_1);
 }
 
 /*-----------------------------------------------------------------------------
diff --git a/Engine/Source/Runtime/IOS/MetalRHI/Private/MetalRHI.cpp b/Engine/Source/Runtime/IOS/MetalRHI/Private/MetalRHI.cpp
--- a/Engine/Source/Runtime/IOS/MetalRHI/Private/MetalRHI.cpp
+++ b/Engine/Source/Runtime/IOS/MetalRHI/Private/MetalRHI.cpp
@@ -71,7 +71,7 @@ FMetalDynamicRHI::FMetalDynamicRHI()
 	GMaxTextureArrayLayers = 2048;
 
 	GShaderPlatformForFeatureLevel[ERHIFeatureLevel::ES2] = SP_METAL;
-	GShaderPlatformForFeatureLevel[ERHIFeatureLevel::SM3] = SP_NumPlatforms;
+	GShaderPlatformForFeatureLevel[ERHIFeatureLevel::ES3_1] = SP_NumPlatforms;
 	GShaderPlatformForFeatureLevel[ERHIFeatureLevel::SM4] = SP_NumPlatforms;
 	GShaderPlatformForFeatureLevel[ERHIFeatureLevel::SM5] = SP_NumPlatforms;
 
diff --git a/Engine/Source/Runtime/NullDrv/Private/NullRHI.cpp b/Engine/Source/Runtime/NullDrv/Private/NullRHI.cpp
--- a/Engine/Source/Runtime/NullDrv/Private/NullRHI.cpp
+++ b/Engine/Source/Runtime/NullDrv/Private/NullRHI.cpp
@@ -20,17 +20,17 @@ FNullDynamicRHI::FNullDynamicRHI()
 void FNullDynamicRHI::Init()
 {
 #if PLATFORM_WINDOWS
-    GShaderPlatformForFeatureLevel[ERHIFeatureLevel::ES2] = SP_PCD3D_ES2;
-	GShaderPlatformForFeatureLevel[ERHIFeatureLevel::SM3] = SP_NumPlatforms;
+	GShaderPlatformForFeatureLevel[ERHIFeatureLevel::ES2] = SP_PCD3D_ES2;
+	GShaderPlatformForFeatureLevel[ERHIFeatureLevel::ES3_1] = SP_NumPlatforms;
 	GShaderPlatformForFeatureLevel[ERHIFeatureLevel::SM4] = SP_PCD3D_SM4;
 	GShaderPlatformForFeatureLevel[ERHIFeatureLevel::SM5] = SP_PCD3D_SM5;
 #else
-    GShaderPlatformForFeatureLevel[ERHIFeatureLevel::ES2] = SP_OPENGL_ES2;
-	GShaderPlatformForFeatureLevel[ERHIFeatureLevel::SM3] = SP_NumPlatforms;
+	GShaderPlatformForFeatureLevel[ERHIFeatureLevel::ES2] = SP_OPENGL_ES2;
+	GShaderPlatformForFeatureLevel[ERHIFeatureLevel::ES3_1] = SP_NumPlatforms;
 	GShaderPlatformForFeatureLevel[ERHIFeatureLevel::SM4] = SP_OPENGL_SM4;
 	GShaderPlatformForFeatureLevel[ERHIFeatureLevel::SM5] = SP_OPENGL_SM5;
 #endif
-    
+	
 	check(!GIsRHIInitialized);
 
 	// Notify all initialized FRenderResources that there's a valid RHI device to create their RHI resources for now.
diff --git a/Engine/Source/Runtime/OpenGLDrv/Private/OpenGLDevice.cpp b/Engine/Source/Runtime/OpenGLDrv/Private/OpenGLDevice.cpp
--- a/Engine/Source/Runtime/OpenGLDrv/Private/OpenGLDevice.cpp
+++ b/Engine/Source/Runtime/OpenGLDrv/Private/OpenGLDevice.cpp
@@ -647,7 +647,7 @@ static void InitRHICapabilitiesForGL()
 	GHardwareHiddenSurfaceRemoval = FOpenGL::HasHardwareHiddenSurfaceRemoval();
 
 	GShaderPlatformForFeatureLevel[ERHIFeatureLevel::ES2] = (GMaxRHIFeatureLevel == ERHIFeatureLevel::ES2) ? GRHIShaderPlatform : SP_OPENGL_PCES2;
-	GShaderPlatformForFeatureLevel[ERHIFeatureLevel::SM3] = SP_NumPlatforms;
+	GShaderPlatformForFeatureLevel[ERHIFeatureLevel::ES3_1] = SP_NumPlatforms;
 	GShaderPlatformForFeatureLevel[ERHIFeatureLevel::SM4] = PLATFORM_MAC ? SP_OPENGL_SM4_MAC : SP_OPENGL_SM4;
 	GShaderPlatformForFeatureLevel[ERHIFeatureLevel::SM5] = SP_OPENGL_SM5;
 
diff --git a/Engine/Source/Runtime/RHI/Private/RHI.cpp b/Engine/Source/Runtime/RHI/Private/RHI.cpp
--- a/Engine/Source/Runtime/RHI/Private/RHI.cpp
+++ b/Engine/Source/Runtime/RHI/Private/RHI.cpp
@@ -248,7 +248,7 @@ RHI_API EShaderPlatform GetCurrentRHIShaderPlatform()
 FName FeatureLevelNames[] = 
 {
 	FName(TEXT("ES2")),
-	FName(TEXT("SM3")),
+	FName(TEXT("ES3_1")),
 	FName(TEXT("SM4")),
 	FName(TEXT("SM5")),
 };
diff --git a/Engine/Source/Runtime/RHI/Public/RHIDefinitions.h b/Engine/Source/Runtime/RHI/Public/RHIDefinitions.h
--- a/Engine/Source/Runtime/RHI/Public/RHIDefinitions.h
+++ b/Engine/Source/Runtime/RHI/Public/RHIDefinitions.h
@@ -84,8 +84,8 @@ namespace ERHIFeatureLevel
 	{
 		/** Feature level defined by the core capabilities of OpenGL ES2. */
 		ES2,
-		/** Feature level defined by the capabilities of DX9 Shader Model 3. */
-		SM3,
+		/** Feature level defined by the core capabilities of OpenGL ES3.1 & Metal. */
+		ES3_1,
 		/** Feature level defined by the capabilities of DX10 Shader Model 4. */
 		SM4,
 		/** Feature level defined by the capabilities of DX11 Shader Model 5. */
@@ -650,8 +650,8 @@ inline int32 GetFeatureLevelMaxNumberOfBones(ERHIFeatureLevel::Type FeatureLevel
 	switch (FeatureLevel)
 	{
 	case ERHIFeatureLevel::ES2:
-	case ERHIFeatureLevel::SM3:
 		return 75;
+	case ERHIFeatureLevel::ES3_1:
 	case ERHIFeatureLevel::SM4:
 	case ERHIFeatureLevel::SM5:
 		return 256;
diff --git a/Engine/Source/Runtime/Renderer/Private/BasePassRendering.h b/Engine/Source/Runtime/Renderer/Private/BasePassRendering.h
--- a/Engine/Source/Runtime/Renderer/Private/BasePassRendering.h
+++ b/Engine/Source/Runtime/Renderer/Private/BasePassRendering.h
@@ -103,7 +103,7 @@ public:
 	{
 		bool bShouldCache = TBasePassVertexShaderBaseType<LightMapPolicyType>::ShouldCache(Platform, Material, VertexFactoryType);
 		return bShouldCache 
-			&& (IsFeatureLevelSupported(Platform, ERHIFeatureLevel::SM3) || (IsPCPlatform(Platform)))
+			&& (IsFeatureLevelSupported(Platform, ERHIFeatureLevel::SM4) || (IsPCPlatform(Platform)))
 			&& (!bEnableAtmosphericFog || IsFeatureLevelSupported(Platform, ERHIFeatureLevel::SM4));
 	}
 
@@ -427,7 +427,7 @@ public:
 		const bool bCacheShaders = !bEnableSkyLight || (Material->GetShadingModel() != MSM_Unlit);
 
 		return bCacheShaders
-			&& (IsFeatureLevelSupported(Platform, ERHIFeatureLevel::SM3) || (IsPCPlatform(Platform)))
+			&& (IsFeatureLevelSupported(Platform, ERHIFeatureLevel::SM4) || (IsPCPlatform(Platform)))
 			&& TBasePassPixelShaderBaseType<LightMapPolicyType>::ShouldCache(Platform, Material, VertexFactoryType);
 	}
 
diff --git a/Engine/Source/Runtime/Renderer/Private/Renderer.cpp b/Engine/Source/Runtime/Renderer/Private/Renderer.cpp
--- a/Engine/Source/Runtime/Renderer/Private/Renderer.cpp
+++ b/Engine/Source/Runtime/Renderer/Private/Renderer.cpp
@@ -51,7 +51,7 @@ void FRendererModule::DrawTileMesh(FRHICommandListImmediate& RHICmdList, const F
 		// handle translucent material blend modes
 		if (IsTranslucentBlendMode(MaterialBlendMode))
 		{
-			if (FeatureLevel >= ERHIFeatureLevel::SM3)
+			if (FeatureLevel >= ERHIFeatureLevel::SM4)
 			{
 				FTranslucencyDrawingPolicyFactory::DrawDynamicMesh(RHICmdList, View, FTranslucencyDrawingPolicyFactory::ContextType(), Mesh, false, false, NULL, HitProxyId);
 			}
@@ -73,7 +73,7 @@ void FRendererModule::DrawTileMesh(FRHICommandListImmediate& RHICmdList, const F
 			}
 			else
 			{
-				if (FeatureLevel >= ERHIFeatureLevel::SM3)
+				if (FeatureLevel >= ERHIFeatureLevel::SM4)
 				{
 					FBasePassOpaqueDrawingPolicyFactory::DrawDynamicMesh(RHICmdList, View, FBasePassOpaqueDrawingPolicyFactory::ContextType(false, ESceneRenderTargetsMode::SetTextures), Mesh, false, false, NULL, HitProxyId);
 				}
diff --git a/Engine/Source/Runtime/Renderer/Private/SceneCore.cpp b/Engine/Source/Runtime/Renderer/Private/SceneCore.cpp
--- a/Engine/Source/Runtime/Renderer/Private/SceneCore.cpp
+++ b/Engine/Source/Runtime/Renderer/Private/SceneCore.cpp
@@ -278,7 +278,7 @@ void FStaticMesh::AddToDrawLists(FRHICommandListImmediate& RHICmdList, FScene* S
 {
 	const auto FeatureLevel = Scene->GetFeatureLevel();
 
-	if (CastShadow && FeatureLevel >= ERHIFeatureLevel::SM3)
+	if (CastShadow && FeatureLevel >= ERHIFeatureLevel::SM4)
 	{
 		FShadowDepthDrawingPolicyFactory::AddStaticMesh(Scene, this);
 	}
@@ -302,7 +302,7 @@ void FStaticMesh::AddToDrawLists(FRHICommandListImmediate& RHICmdList, FScene* S
 		return;
 	}
 
-	if (FeatureLevel >= ERHIFeatureLevel::SM3)
+	if (FeatureLevel >= ERHIFeatureLevel::SM4)
 	{
 		extern TAutoConsoleVariable<int32> CVarEarlyZPass;
 		int32 EarlyZPass = CVarEarlyZPass.GetValueOnRenderThread();
diff --git a/Engine/Source/Runtime/Renderer/Private/SceneRendering.cpp b/Engine/Source/Runtime/Renderer/Private/SceneRendering.cpp
--- a/Engine/Source/Runtime/Renderer/Private/SceneRendering.cpp
+++ b/Engine/Source/Runtime/Renderer/Private/SceneRendering.cpp
@@ -841,7 +841,7 @@ FSceneRenderer* FSceneRenderer::CreateSceneRenderer(const FSceneViewFamily* InVi
 
 void FSceneRenderer::RenderCustomDepthPass(FRHICommandListImmediate& RHICmdList)
 {
-	if(FeatureLevel < ERHIFeatureLevel::SM3)
+	if(FeatureLevel < ERHIFeatureLevel::SM4)
 	{
 		// not yet supported on lower end platforms
 		return;
@@ -924,26 +924,26 @@ void FSceneRenderer::OnStartFrame()
  */
 static void RenderViewFamily_RenderThread(FRHICommandListImmediate& RHICmdList, FSceneRenderer* SceneRenderer)
 {
-    FMemMark MemStackMark(FMemStack::Get());
+	FMemMark MemStackMark(FMemStack::Get());
 
 	// update any resources that needed a deferred update
 	FDeferredUpdateResource::UpdateResources(RHICmdList);
 
-    for( int ViewExt = 0; ViewExt < SceneRenderer->ViewFamily.ViewExtensions.Num(); ViewExt++ )
-    {
+	for( int ViewExt = 0; ViewExt < SceneRenderer->ViewFamily.ViewExtensions.Num(); ViewExt++ )
+	{
 		SceneRenderer->ViewFamily.ViewExtensions[ViewExt]->PreRenderViewFamily_RenderThread(SceneRenderer->ViewFamily);
-        for( int ViewIndex = 0; ViewIndex < SceneRenderer->ViewFamily.Views.Num(); ViewIndex++ )
-        {
-            SceneRenderer->ViewFamily.ViewExtensions[ViewExt]->PreRenderView_RenderThread(SceneRenderer->Views[ViewIndex]);
-        }
-    }
+		for( int ViewIndex = 0; ViewIndex < SceneRenderer->ViewFamily.Views.Num(); ViewIndex++ )
+		{
+			SceneRenderer->ViewFamily.ViewExtensions[ViewExt]->PreRenderView_RenderThread(SceneRenderer->Views[ViewIndex]);
+		}
+	}
 
 	if(SceneRenderer->ViewFamily.EngineShowFlags.OnScreenDebug)
 	{
 		GRenderTargetPool.SetEventRecordingActive(true);
 	}
 
-    {
+	{
 		SCOPE_CYCLE_COUNTER(STAT_TotalSceneRenderingTime);
 		
 		if(SceneRenderer->ViewFamily.EngineShowFlags.HitProxies)
@@ -977,7 +977,7 @@ static void RenderViewFamily_RenderThread(FRHICommandListImmediate& RHICmdList,
 
 		GRenderTargetPool.SetEventRecordingActive(false);
 
-        // Delete the scene renderer.
+		// Delete the scene renderer.
 		delete SceneRenderer;
 		{
 			QUICK_SCOPE_CYCLE_COUNTER(STAT_DeleteSceneRenderer_Dispatch);
diff --git a/Engine/Source/Runtime/Renderer/Private/SceneVisibility.cpp b/Engine/Source/Runtime/Renderer/Private/SceneVisibility.cpp
--- a/Engine/Source/Runtime/Renderer/Private/SceneVisibility.cpp
+++ b/Engine/Source/Runtime/Renderer/Private/SceneVisibility.cpp
@@ -362,7 +362,7 @@ static int32 OcclusionCull(FRHICommandListImmediate& RHICmdList, const FScene* S
 	float CurrentRealTime = View.Family->CurrentRealTime;
 	if (ViewState)
 	{
-		if (Scene->GetFeatureLevel() >= ERHIFeatureLevel::SM3)
+		if (Scene->GetFeatureLevel() >= ERHIFeatureLevel::SM4)
 		{
 			bool bClearQueries = !View.Family->EngineShowFlags.HitProxies;
 			bool bSubmitQueries = !View.bDisableQuerySubmissions;
diff --git a/Engine/Source/Runtime/Renderer/Private/ShadowRendering.cpp b/Engine/Source/Runtime/Renderer/Private/ShadowRendering.cpp
--- a/Engine/Source/Runtime/Renderer/Private/ShadowRendering.cpp
+++ b/Engine/Source/Runtime/Renderer/Private/ShadowRendering.cpp
@@ -170,7 +170,7 @@ public:
 					&& (!bUsePositionOnlyStream || VertexFactoryType->SupportsPositionOnly())
 					// Don't render ShadowDepth for translucent unlit materials
 					&& (!IsTranslucentBlendMode(Material->GetBlendMode()) && Material->GetShadingModel() != MSM_Unlit)
-					&& IsFeatureLevelSupported(Platform, ERHIFeatureLevel::SM3);
+					&& IsFeatureLevelSupported(Platform, ERHIFeatureLevel::SM4);
 		}
 	}
 
@@ -514,11 +514,7 @@ public:
 
 	static bool ShouldCache(EShaderPlatform Platform,const FMaterial* Material,const FVertexFactoryType* VertexFactoryType)
 	{
-		if (IsFeatureLevelSupported(Platform, ERHIFeatureLevel::SM3))
-		{
-
-		}
-		else
+		if (!IsFeatureLevelSupported(Platform, ERHIFeatureLevel::SM4))
 		{
 			return (Material->IsSpecialEngineMaterial()
 				// Only compile for masked or lit translucent materials
@@ -550,7 +546,7 @@ public:
 				&& (ShaderMode != PixelShadowDepth_OnePassPointLight || IsFeatureLevelSupported(Platform, ERHIFeatureLevel::SM4))
 				// Don't render ShadowDepth for translucent unlit materials
 				&& (!IsTranslucentBlendMode(Material->GetBlendMode()) && Material->GetShadingModel() != MSM_Unlit)
-				&& IsFeatureLevelSupported(Platform, ERHIFeatureLevel::SM3);
+				&& IsFeatureLevelSupported(Platform, ERHIFeatureLevel::SM4);
 		}
 	}
 
@@ -605,7 +601,7 @@ TGlobalResource<StencilingGeometry::FStencilConeIndexBuffer> StencilingGeometry:
 
 bool FShadowProjectionVS::ShouldCache(EShaderPlatform Platform)
 {
-	return IsFeatureLevelSupported(Platform, ERHIFeatureLevel::SM3);
+	return IsFeatureLevelSupported(Platform, ERHIFeatureLevel::SM4);
 }
 
 void FShadowProjectionVS::SetParameters(FRHICommandList& RHICmdList, const FSceneView& View, const FProjectedShadowInfo* ShadowInfo, const int32 ShadowSplitIndex)
diff --git a/Engine/Source/Runtime/Renderer/Private/ShadowRendering.h b/Engine/Source/Runtime/Renderer/Private/ShadowRendering.h
--- a/Engine/Source/Runtime/Renderer/Private/ShadowRendering.h
+++ b/Engine/Source/Runtime/Renderer/Private/ShadowRendering.h
@@ -1127,7 +1127,7 @@ public:
 
 	static bool ShouldCache(EShaderPlatform Platform)
 	{
-		return IsFeatureLevelSupported(Platform, ERHIFeatureLevel::SM3);
+		return IsFeatureLevelSupported(Platform, ERHIFeatureLevel::SM4);
 	}
 
 	/**
diff --git a/Engine/Source/Runtime/Windows/D3D11RHI/Private/Windows/WindowsD3D11Device.cpp b/Engine/Source/Runtime/Windows/D3D11RHI/Private/Windows/WindowsD3D11Device.cpp
--- a/Engine/Source/Runtime/Windows/D3D11RHI/Private/Windows/WindowsD3D11Device.cpp
+++ b/Engine/Source/Runtime/Windows/D3D11RHI/Private/Windows/WindowsD3D11Device.cpp
@@ -567,7 +567,7 @@ void FD3D11DynamicRHI::InitD3DDevice()
 			&& (DeviceFlags & D3D11_CREATE_DEVICE_SINGLETHREADED) == 0;
 
 		GShaderPlatformForFeatureLevel[ERHIFeatureLevel::ES2] = SP_PCD3D_ES2;
-		GShaderPlatformForFeatureLevel[ERHIFeatureLevel::SM3] = SP_NumPlatforms;
+		GShaderPlatformForFeatureLevel[ERHIFeatureLevel::ES3_1] = SP_NumPlatforms;
 		GShaderPlatformForFeatureLevel[ERHIFeatureLevel::SM4] = SP_PCD3D_SM4;
 		GShaderPlatformForFeatureLevel[ERHIFeatureLevel::SM5] = SP_PCD3D_SM5;
 
