Remove determinism requirements for sending references to components over network (go back to exporting full name through package map)

diff --git a/Engine/Source/Runtime/Engine/Private/PackageMapClient.cpp b/Engine/Source/Runtime/Engine/Private/PackageMapClient.cpp
--- a/Engine/Source/Runtime/Engine/Private/PackageMapClient.cpp
+++ b/Engine/Source/Runtime/Engine/Private/PackageMapClient.cpp
@@ -21,16 +21,6 @@ static const int GUID_PACKET_ACKED		= -1;
  */
 static const int INTERNAL_LOAD_OBJECT_RECURSION_LIMIT = 16;
 
-// Comment this out to disable preallocation of sub-object netguids (will will export them as stable names rather than deterministically produce same list on both sides, which can be fragile)
-#define PRE_ALLOC_SUB_OBJECT_GUIDS
-
-#ifdef PRE_ALLOC_SUB_OBJECT_GUIDS
-#define GET_NET_GUID_INDEX( GUID )					( GUID.Value >> 1 )
-#define GET_NET_GUID_STATIC_MOD( GUID )				( GUID.Value & 1 )
-#define COMPOSE_NET_GUID( Index, IsStatic )			( ( ( Index ) << 1 ) | ( IsStatic ) )
-#define COMPOSE_RELATIVE_NET_GUID( GUID, Index )	( COMPOSE_NET_GUID( GET_NET_GUID_INDEX( GUID ) + Index, GET_NET_GUID_STATIC_MOD( GUID ) ) )
-#endif
-
 static TAutoConsoleVariable<int32> CVarAllowAsyncLoading( TEXT( "net.AllowAsyncLoading" ), 0, TEXT( "Allow async loading" ) );
 static TAutoConsoleVariable<int32> CVarSimulateAsyncLoading( TEXT( "net.SimulateAsyncLoading" ), 0, TEXT( "Simulate async loading" ) );
 
@@ -297,98 +287,6 @@ bool UPackageMapClient::SerializeNewActor(FArchive& Ar, class UActorChannel *Cha
 				}
 			}
 		}
-
-#ifdef PRE_ALLOC_SUB_OBJECT_GUIDS
-		//
-		// Serialize Initially replicating sub-objects here
-		//
-
-		//
-		// We expect that all sub-objects that have stable network names have already been assigned guids, which are relative to their actor owner
-		// Because of this, we can expect the client to deterministically generate the same guids, and we just need to verify with checksum
-		//
-
-		TArray< UObject * > Subobjs;
-
-		Actor->GetSubobjectsWithStableNamesForNetworking( Subobjs );
-
-		if ( Ar.IsSaving() )
-		{
-			// We're not really saving anything necessary here except a checksum to make sure we're in sync
-			// It might be a worth thinking about turning this off in shipping, and trust this is working as intended to save bandwidth if needed
-			uint32 Checksum = 0;
-
-			for ( int32 i = 0; i < Subobjs.Num(); i++ )
-			{
-				UObject * SubObj = Subobjs[i];
-
-				FNetworkGUID SubobjNetGUID = GuidCache->NetGUIDLookup.FindRef( SubObj );
-
-				// Make sure these sub objects have net guids that are relative to this owning actor
-				if ( SubobjNetGUID.Value != COMPOSE_RELATIVE_NET_GUID( NetGUID, i + 1 ) )
-				{
-					// If this happens, someone most likely added a stably name sub-object during gameplay (bad)
-					UE_LOG( LogNetPackageMap, Error, TEXT( "SerializeNewActor: Sub-object net guid out of sync. Actor: %s, SubObj: %s" ), *Actor->GetName(), *SubObj->GetName() );
-					Checksum = 0;		 // Force the checksum on the client to fail, so the client won't have any guids for these sub-objects
-					break;
-				}
-
-				// Evolve checksum so we can sanity check
-				Checksum = FCrc::MemCrc32( &SubobjNetGUID.Value, sizeof( SubobjNetGUID.Value ), Checksum );
-				Checksum = FCrc::StrCrc32( *SubObj->GetName(), Checksum );
-
-				check( !PendingAckGUIDs.Contains( SubobjNetGUID ) );
-			}
-
-			Ar << Checksum;
-		}
-		else
-		{
-			// Load the server checksum, then generate our own, and make sure they match
-			uint32 ServerChecksum = 0;
-			uint32 LocalChecksum = 0;
-			Ar << ServerChecksum;
-
-			TArray< FNetworkGUID > SubObjGuids;
-
-			for ( int32 i = 0; i < Subobjs.Num(); i++ )
-			{
-				// Generate a guid that is relative to our owning actor
-				SubObjGuids.Add( COMPOSE_RELATIVE_NET_GUID( NetGUID, i + 1 ) );
-
-				// Evolve checksum so we can sanity check
-				LocalChecksum = FCrc::MemCrc32( &SubObjGuids[i].Value, sizeof( SubObjGuids[i].Value ), LocalChecksum );
-				LocalChecksum = FCrc::StrCrc32( *Subobjs[i]->GetName(), LocalChecksum );
-
-				UE_LOG( LogNetPackageMap, Log, TEXT( "SerializeNewActor: [Loading] Assigned NetGUID <%s> to subobject %s" ), *SubObjGuids[i].ToString(), *Subobjs[i]->GetPathName() );
-			}
-
-			if ( LocalChecksum != ServerChecksum )
-			{
-				UE_LOG( LogNetPackageMap, Error, TEXT( "SerializeNewActor: Subobject checksum FAILED: %s" ), *Actor->GetFullName() );
-			}
-			else
-			{
-				// Everything worked, assign the guids
-				for ( int32 i = 0; i < Subobjs.Num(); i++ )
-				{
-					FNetworkGUID OldNetGUID = GuidCache->NetGUIDLookup.FindRef( Subobjs[i] );
-
-					if ( OldNetGUID.IsValid() )
-					{
-						if ( OldNetGUID != SubObjGuids[i] )
-						{
-							UE_LOG( LogNetPackageMap, Warning, TEXT( "SerializeNewActor: Changing NetGUID on subobject. Name: %s, OldGUID: %s, NewGUID: %s" ), *Subobjs[i]->GetPathName(), *OldNetGUID.ToString(), *SubObjGuids[i].ToString() );
-						}
-					}
-
-					GuidCache->RegisterNetGUID_Client( SubObjGuids[i], Subobjs[i] );
-
-					UE_LOG( LogNetPackageMap, Log, TEXT( "SerializeNewActor: [Loading] Assigned NetGUID <%s> to subobject %s" ), *SubObjGuids[i].ToString(), *Subobjs[i]->GetPathName() );
-				}
-			}
-		}
-#endif
 	}
 
 	UE_LOG(LogNetPackageMap, Log, TEXT("SerializeNewActor END: Finished Serializing Actor %s <%s> on Channel %d"), Actor ? *Actor->GetName() : TEXT("NULL"), *NetGUID.ToString(), Channel->ChIndex );
@@ -1191,19 +1089,12 @@ bool UPackageMapClient::ShouldSendFullPath( const UObject* Object, const FNetwor
 		return false;
 	}
 
-#ifdef PRE_ALLOC_SUB_OBJECT_GUIDS
-	if ( NetGUID.IsDynamic() )
-	{
-		return false;		// We only export fully stably named objects
-	}
-#else
 	if ( !Object->IsNameStableForNetworking() )
 	{
 		check( !NetGUID.IsDefault() );
 		check( NetGUID.IsDynamic() );
 		return false;		// We only export objects that have stable names
 	}
-#endif
 
 	if ( NetGUID.IsDefault() )
 	{
@@ -1545,27 +1436,6 @@ FNetworkGUID FNetGUIDCache::GetOrAssignNetGUID( const UObject * Object )
 	return AssignNewNetGUID_Server( Object );
 }
 
-#ifdef PRE_ALLOC_SUB_OBJECT_GUIDS
-static const AActor * GetOuterActor( const UObject * Object )
-{
-	Object = Object->GetOuter();
-
-	while ( Object != NULL )
-	{
-		const AActor * Actor = Cast< const AActor >( Object );
-
-		if ( Actor != NULL )
-		{
-			return Actor;
-		}
-
-		Object = Object->GetOuter();
-	}
-
-	return NULL;
-}
-#endif
-
 /**
  *	Generate a new NetGUID for this object and assign it.
  */
@@ -1573,70 +1443,6 @@ FNetworkGUID FNetGUIDCache::AssignNewNetGUID_Server( const UObject * Object )
 {
 	check( IsNetGUIDAuthority() );
 
-#ifdef PRE_ALLOC_SUB_OBJECT_GUIDS
-	const AActor * Actor = Cast< const AActor >( Object );
-
-	// So what is going on here, is we want to make sure the network guid's of stably named sub-objects get 
-	// initialized in sync with their owning actor. What this allows us to do, is not consume any extra bandwidth 
-	// when assigning these guid's to the sub-objects on each client
-	if ( Actor == NULL )
-	{
-		const AActor * OuterActor = GetOuterActor( Object );
-
-		if ( OuterActor != NULL && IsDynamicObject( OuterActor ) )
-		{
-			TArray< UObject * > Subobjs;
-			const_cast< AActor * >( OuterActor )->GetSubobjectsWithStableNamesForNetworking( Subobjs );
-
-			if ( Subobjs.Contains( const_cast< UObject * >( Object ) ) )
-			{
-				// Since we always assign guids to components when the owning actor gets assigned, we assume that if we get here
-				// the owning actor must not have been assigned yet
-				check( !NetGUIDLookup.Contains( OuterActor ) );
-
-				// Assign our owning actor a guid first (which should also assign our net guid if things are working correctly)
-				AssignNewNetGUID_Server( OuterActor );
-
-				return NetGUIDLookup.FindChecked( Object );
-			}
-		}
-	}
-
-#define ALLOC_NEW_NET_GUID( IsStatic ) ( COMPOSE_NET_GUID( ++UniqueNetIDs[ IsStatic ], IsStatic ) )
-
-	// Generate new NetGUID and assign it
-	const int32 IsStatic = IsDynamicObject( Object ) ? 0 : 1;
-
-	const FNetworkGUID NewNetGuid( ALLOC_NEW_NET_GUID( IsStatic ) );
-
-	RegisterNetGUID_Server( NewNetGuid, Object );
-
-	// If this is a dynamic actor, assign all of our stably named sub-objects their net guids now as well (they come as a group)
-	if ( Actor != NULL && !IsStatic )
-	{
-		TArray< UObject * > Subobjs;
-
-		const_cast< AActor * >( Actor )->GetSubobjectsWithStableNamesForNetworking( Subobjs );
-
-		for ( int32 i = 0; i < Subobjs.Num(); i++ )
-		{
-			UObject * SubObj = Subobjs[i];
-
-			check( GetOuterActor( SubObj ) == Actor );
-
-			if ( NetGUIDLookup.Contains( SubObj ) )		// We should NOT have a guid yet
-			{
-				UE_LOG( LogNetPackageMap, Error, TEXT( "AssignNewNetGUID_Server: Sub object already registered: Actor: %s, SubObj: %s, NetGuid: %s" ), *Actor->GetPathName(), *SubObj->GetPathName(), *NetGUIDLookup.FindChecked( SubObj ).ToString() );
-			}
-			const int32 SubIsStatic = IsDynamicObject( SubObj ) ? 0 : 1;
-			check( SubIsStatic == IsStatic );
-			const FNetworkGUID SubobjNetGUID( ALLOC_NEW_NET_GUID( SubIsStatic ) );
-			check( SubobjNetGUID.Value == COMPOSE_RELATIVE_NET_GUID( NewNetGuid, i + 1 ) );
-			RegisterNetGUID_Server( SubobjNetGUID, SubObj );
-		}
-	}
-	return NewNetGuid;
-#else
 #define COMPOSE_NET_GUID( Index, IsStatic )	( ( ( Index ) << 1 ) | ( IsStatic ) )
 #define ALLOC_NEW_NET_GUID( IsStatic )		( COMPOSE_NET_GUID( ++UniqueNetIDs[ IsStatic ], IsStatic ) )
 
@@ -1648,7 +1454,6 @@ FNetworkGUID FNetGUIDCache::AssignNewNetGUID_Server( const UObject * Object )
 	RegisterNetGUID_Server( NewNetGuid, Object );
 
 	return NewNetGuid;
-#endif
 }
 
 void FNetGUIDCache::RegisterNetGUID_Internal( const FNetworkGUID & NetGUID, const FNetGuidCacheObject & CacheObject )
@@ -1722,7 +1527,7 @@ void FNetGUIDCache::RegisterNetGUID_Client( const FNetworkGUID & NetGUID, const
 	{
 		if ( ExistingCacheObjectPtr->PathName != NAME_None )
 		{
-			UE_LOG( LogNetPackageMap, Warning, TEXT( "RegisterNetGUID_Client: Guid with pathname. Path: %s, NetGUID: %s" ), *ExistingCacheObjectPtr->PathName.ToString(), *NetGUID.ToString() );
+			UE_LOG( LogNetPackageMap, Warning, TEXT( "RegisterNetGUID_Client: Guid with pathname. Path: %s, NetGUID: %s, OuterGUID: %s" ), *ExistingCacheObjectPtr->PathName.ToString(), *NetGUID.ToString(), *ExistingCacheObjectPtr->OuterGUID.ToString() );
 		}
 
 		// If this net guid was found but the old object is NULL, this can happen due to:
@@ -1777,9 +1582,6 @@ void FNetGUIDCache::RegisterNetGUIDFromPath_Client( const FNetworkGUID & NetGUID
 {
 	check( !IsNetGUIDAuthority() );		// Server never calls this locally
 	check( !NetGUID.IsDefault() );
-#ifdef PRE_ALLOC_SUB_OBJECT_GUIDS
-	check( !NetGUID.IsDynamic() );		// It only makes sense to do this for static guids
-#endif
 
 	const FNetGuidCacheObject * ExistingCacheObjectPtr = ObjectLookup.Find( NetGUID );
 
@@ -1912,27 +1714,12 @@ UObject * FNetGUIDCache::GetObjectFromNetGUID( const FNetworkGUID & NetGUID, con
 		return NULL;
 	}
 
-#ifdef PRE_ALLOC_SUB_OBJECT_GUIDS
-	if ( NetGUID.IsDynamic() )
-	{
-		// Dynamic objects can never load from their names (name is not stable)
-		return NULL;
-	}
-
-	if ( CacheObjectPtr->PathName == NAME_None )
-	{
-		// Static guids should have a path
-		UE_LOG( LogNetPackageMap, Error, TEXT( "GetObjectFromNetGUID: Static guid on client with no path. NetGUID: %s" ), *NetGUID.ToString() );
-		return NULL;
-	}
-#else
 	if ( CacheObjectPtr->PathName == NAME_None )
 	{
 		// If we don't have a path, assume this is a non stably named guid
 		check( NetGUID.IsDynamic() );
 		return NULL;
 	}
-#endif
 
 	// First, resolve the outer
 	UObject * ObjOuter = NULL;
@@ -1945,8 +1732,11 @@ UObject * FNetGUIDCache::GetObjectFromNetGUID( const FNetworkGUID & NetGUID, con
 		if ( OuterCacheObject == NULL )
 		{
 			// Shouldn't be possible, but just in case...
-			UE_LOG( LogNetPackageMap, Error, TEXT( "GetObjectFromNetGUID: Outer not registered. Path: %s, OuterPath: %s, NetGUID: %s" ), *CacheObjectPtr->PathName.ToString(), OuterCacheObject ? *OuterCacheObject->PathName.ToString() : TEXT( "NULL!!" ), *NetGUID.ToString() );
-			CacheObjectPtr->bIsBroken = 1;	// Set this to 1 so that we don't keep spamming
+			if ( CacheObjectPtr->OuterGUID.IsStatic() )
+			{
+				UE_LOG( LogNetPackageMap, Error, TEXT( "GetObjectFromNetGUID: Static outer not registered. Path: %s, NetGUID: %s, OuterGUID: %s" ), *CacheObjectPtr->PathName.ToString(), *NetGUID.ToString(), *CacheObjectPtr->OuterGUID.ToString() );
+				CacheObjectPtr->bIsBroken = 1;	// Set this to 1 so that we don't keep spamming
+			}
 			return NULL;
 		}
 
