We don't need to send the subobject class if it is stably named for networking (and more sanity checking)

diff --git a/Engine/Source/Runtime/Engine/Private/DataChannel.cpp b/Engine/Source/Runtime/Engine/Private/DataChannel.cpp
--- a/Engine/Source/Runtime/Engine/Private/DataChannel.cpp
+++ b/Engine/Source/Runtime/Engine/Private/DataChannel.cpp
@@ -2001,8 +2001,16 @@ void UActorChannel::BeginContentBlock( UObject * Obj, FOutBunch &Bunch )
 	if ( Connection->Driver->IsServer() )
 	{
 		// Only the server can tell clients to create objects, so no need for the client to send this to the server
-		UClass *ObjClass = Obj->GetClass();
-		Bunch << ObjClass;
+		if ( Obj->IsNameStableForNetworking() )
+		{
+			Bunch.WriteBit( 1 );
+		}
+		else
+		{
+			Bunch.WriteBit( 0 );
+			UClass *ObjClass = Obj->GetClass();
+			Bunch << ObjClass;
+		}
 	}
 
 #if !(UE_BUILD_SHIPPING || UE_BUILD_TEST)
@@ -2026,6 +2034,9 @@ void UActorChannel::BeginContentBlockForSubObjectDelete( FOutBunch & Bunch, FNet
 	Bunch << GuidToDelete;
 	NET_CHECKSUM(Bunch);
 
+	// Send a 0 bit to indicate that this is not a stably named object
+	Bunch.WriteBit( 0 );
+
 	//	-Invalid NetGUID (interpreted as delete)
 	FNetworkGUID InvalidNetGUID;
 	InvalidNetGUID.Reset();
@@ -2100,74 +2111,92 @@ UObject * UActorChannel::ReadContentBlockHeader( FInBunch & Bunch )
 			return NULL;
 		}
 	}
-	else if ( IsServer )
+
+	if ( IsServer )
 	{
-		UE_LOG( LogNetTraffic, Error, TEXT( "ReadContentBlockHeader: Client attempted to create sub-object. Actor: %s" ), *Actor->GetName() );
-		Bunch.SetError();
-		return NULL;
+		// The server should never need to create sub objects
+		if ( SubObj == NULL )
+		{
+			UE_LOG( LogNetTraffic, Error, TEXT( "ReadContentBlockHeader: Client attempted to create sub-object. Actor: %s" ), *Actor->GetName() );
+			Bunch.SetError();
+			return NULL;
+		}
+
+		return SubObj;
 	}
 
-	if ( !IsServer )
+	if ( Bunch.ReadBit() )
 	{
-		// Serialize the class in case we have to spawn it.
-		// Manually serialize the object so that we can get the NetGUID (in order to assign it if we spawn the object here)
-		FNetworkGUID ClassNetGUID;
-		UObject * SubObjClassObj = NULL;
-		Connection->PackageMap->SerializeObject( Bunch, UObject::StaticClass(), SubObjClassObj, &ClassNetGUID );
-
-		// Delete sub-object
-		if ( !ClassNetGUID.IsValid() )
+		// If this is a stably named sub-object, we shouldn't need to create it
+		if ( SubObj == NULL )
 		{
-			if ( SubObj )
-			{
-				Actor->OnSubobjectDestroyFromReplication( SubObj );
-				SubObj->MarkPendingKill();
-			}
+			UE_LOG( LogNetTraffic, Error, TEXT( "ReadContentBlockHeader: Stably named sub-object not found. Actor: %s" ), *Actor->GetName() );
+			Bunch.SetError();
 			return NULL;
 		}
 
-		UClass * SubObjClass = Cast< UClass >( SubObjClassObj );
+		return SubObj;
+	}
+
+	// Serialize the class in case we have to spawn it.
+	// Manually serialize the object so that we can get the NetGUID (in order to assign it if we spawn the object here)
+	FNetworkGUID ClassNetGUID;
+	UObject * SubObjClassObj = NULL;
+	Connection->PackageMap->SerializeObject( Bunch, UObject::StaticClass(), SubObjClassObj, &ClassNetGUID );
 
-		if ( SubObjClass == NULL )
+	// Delete sub-object
+	if ( !ClassNetGUID.IsValid() )
+	{
+		if ( SubObj )
 		{
-			UE_LOG( LogNetTraffic, Warning, TEXT( "UActorChannel::ReadContentBlockHeader: Unable to read sub-object class. Actor: %s" ), *Actor->GetName() );
+			Actor->OnSubobjectDestroyFromReplication( SubObj );
+			SubObj->MarkPendingKill();
+		}
+		return NULL;
+	}
 
-			// Valid NetGUID but no class was resolved - this is an error
-			if ( SubObj == NULL )
-			{
-				UE_LOG( LogNetTraffic, Error, TEXT( "UActorChannel::ReadContentBlockHeader: Unable to read sub-object class (SubObj == NULL). Actor: %s" ), *Actor->GetName() );
-				Bunch.SetError();
-				return NULL;
-			}
+	UClass * SubObjClass = Cast< UClass >( SubObjClassObj );
+
+	if ( SubObjClass == NULL )
+	{
+		UE_LOG( LogNetTraffic, Warning, TEXT( "UActorChannel::ReadContentBlockHeader: Unable to read sub-object class. Actor: %s" ), *Actor->GetName() );
+
+		// Valid NetGUID but no class was resolved - this is an error
+		if ( SubObj == NULL )
+		{
+			UE_LOG( LogNetTraffic, Error, TEXT( "UActorChannel::ReadContentBlockHeader: Unable to read sub-object class (SubObj == NULL). Actor: %s" ), *Actor->GetName() );
+			Bunch.SetError();
+			return NULL;
 		}
-		else
+	}
+	else
+	{
+		if ( SubObjClass == UObject::StaticClass() )
 		{
-			if ( SubObjClass == UObject::StaticClass() )
-			{
-				UE_LOG( LogNetTraffic, Error, TEXT( "UActorChannel::ReadContentBlockHeader: SubObjClass == UObject::StaticClass(). Actor: %s" ), *Actor->GetName() );
-				Bunch.SetError();
-				return NULL;
-			}
-
-			if ( SubObjClass->IsChildOf( AActor::StaticClass() ) )
-			{
-				UE_LOG( LogNetTraffic, Error, TEXT( "UActorChannel::ReadContentBlockHeader: Sub-object cannot be actor class. Actor: %s" ), *Actor->GetName() );
-				Bunch.SetError();
-				return NULL;
-			}
+			UE_LOG( LogNetTraffic, Error, TEXT( "UActorChannel::ReadContentBlockHeader: SubObjClass == UObject::StaticClass(). Actor: %s" ), *Actor->GetName() );
+			Bunch.SetError();
+			return NULL;
 		}
 
-		if ( SubObj == NULL )
+		if ( SubObjClass->IsChildOf( AActor::StaticClass() ) )
 		{
-			// Construct the sub-object
-			UE_LOG( LogNetTraffic, Log, TEXT( "UActorChannel::ReadContentBlockHeader: Instantiating sub-object. Class: %s, Actor: %s" ), *SubObjClass->GetName(), *Actor->GetName() );
-			SubObj = ConstructObject< UObject >( SubObjClass, Actor );
-			check( SubObj != NULL );
-			Actor->OnSubobjectCreatedFromReplication( SubObj );
-			Connection->Driver->GuidCache->RegisterNetGUID_Client( NetGUID, SubObj );
+			UE_LOG( LogNetTraffic, Error, TEXT( "UActorChannel::ReadContentBlockHeader: Sub-object cannot be actor class. Actor: %s" ), *Actor->GetName() );
+			Bunch.SetError();
+			return NULL;
 		}
 	}
 
+	if ( SubObj == NULL )
+	{
+		// Construct the sub-object
+		UE_LOG( LogNetTraffic, Log, TEXT( "UActorChannel::ReadContentBlockHeader: Instantiating sub-object. Class: %s, Actor: %s" ), *SubObjClass->GetName(), *Actor->GetName() );
+		SubObj = ConstructObject< UObject >( SubObjClass, Actor );
+		check( SubObj != NULL );
+		Actor->OnSubobjectCreatedFromReplication( SubObj );
+		Connection->Driver->GuidCache->RegisterNetGUID_Client( NetGUID, SubObj );
+	}
+
+	// Sanity check one last time
 	check( SubObj != NULL );
 	check( Cast< AActor >( SubObj ) == NULL );
 
