Changed curve mismatch resolution from curve id to last known name to mitigate issues with multiple assets containing unlinked curves.

diff --git a/Engine/Source/Runtime/Engine/Private/Animation/AnimSequenceBase.cpp b/Engine/Source/Runtime/Engine/Private/Animation/AnimSequenceBase.cpp
--- a/Engine/Source/Runtime/Engine/Private/Animation/AnimSequenceBase.cpp
+++ b/Engine/Source/Runtime/Engine/Private/Animation/AnimSequenceBase.cpp
@@ -264,7 +264,7 @@ void UAnimSequenceBase::PostLoad()
 			TArray<FFloatCurve*> UnlinkedCurves;
 			for(FFloatCurve& Curve : RawCurveData.FloatCurves)
 			{
-				if(!NameMapping->Exists(Curve.CurveUid))
+				if(!NameMapping->Exists(Curve.LastObservedName))
 				{
 					// The skeleton doesn't know our name. Use the last observed name that was saved with the
 					// curve to create a new name. This can happen if a user saves an animation but not a skeleton
