Add bone name out param to K2_LineTraceComponent

diff --git a/Engine/Source/Runtime/Engine/Classes/Components/PrimitiveComponent.h b/Engine/Source/Runtime/Engine/Classes/Components/PrimitiveComponent.h
--- a/Engine/Source/Runtime/Engine/Classes/Components/PrimitiveComponent.h
+++ b/Engine/Source/Runtime/Engine/Classes/Components/PrimitiveComponent.h
@@ -839,7 +839,7 @@ public:
 
 	/** Perform a line trace against a single component */
 	UFUNCTION(BlueprintCallable, Category="Collision", meta=(FriendlyName = "Line Trace Component", bTraceComplex="true"))	
-	bool K2_LineTraceComponent(FVector TraceStart, FVector TraceEnd, bool bTraceComplex, bool bShowTrace, FVector& HitLocation, FVector& HitNormal);
+	bool K2_LineTraceComponent(FVector TraceStart, FVector TraceEnd, bool bTraceComplex, bool bShowTrace, FVector& HitLocation, FVector& HitNormal, FName & BoneName);
 
 	/** Sets the bRenderCustomDepth property and marks the render state dirty. */
 	UFUNCTION(BlueprintCallable, Category="Rendering")
diff --git a/Engine/Source/Runtime/Engine/Private/PrimitiveComponentPhysics.cpp b/Engine/Source/Runtime/Engine/Private/PrimitiveComponentPhysics.cpp
--- a/Engine/Source/Runtime/Engine/Private/PrimitiveComponentPhysics.cpp
+++ b/Engine/Source/Runtime/Engine/Private/PrimitiveComponentPhysics.cpp
@@ -761,7 +761,7 @@ void UPrimitiveComponent::OnComponentCollisionSettingsChanged()
 	}
 }
 
-bool UPrimitiveComponent::K2_LineTraceComponent(FVector TraceStart, FVector TraceEnd, bool bTraceComplex, bool bShowTrace, FVector& HitLocation, FVector& HitNormal)
+bool UPrimitiveComponent::K2_LineTraceComponent(FVector TraceStart, FVector TraceEnd, bool bTraceComplex, bool bShowTrace, FVector& HitLocation, FVector& HitNormal, FName & BoneName)
 {
 	FHitResult Hit;
 	static FName KismetTraceComponentName(TEXT("KismetTraceComponent"));
@@ -773,12 +773,14 @@ bool UPrimitiveComponent::K2_LineTraceComponent(FVector TraceStart, FVector Trac
 		// Fill in the results if we hit
 		HitLocation = Hit.Location;
 		HitNormal = Hit.Normal;
+		BoneName = Hit.BoneName;
 	}
 	else
 	{
 		// Blank these out to avoid confusion!
 		HitLocation = FVector::ZeroVector;
 		HitNormal = FVector::ZeroVector;
+		BoneName = NAME_None;
 	}
 
 	if( bShowTrace )
