Welding supports update overlaps
FBodyInstance properly handles welded traces sweeps and overlap

diff --git a/Engine/Source/Runtime/Engine/Private/PhysicsEngine/BodyInstance.cpp b/Engine/Source/Runtime/Engine/Private/PhysicsEngine/BodyInstance.cpp
--- a/Engine/Source/Runtime/Engine/Private/PhysicsEngine/BodyInstance.cpp
+++ b/Engine/Source/Runtime/Engine/Private/PhysicsEngine/BodyInstance.cpp
@@ -289,6 +289,13 @@ void FBodyInstance::UpdateFromDeprecatedEnableCollision()
 }
 
 #if WITH_PHYSX
+//Determine that the shape is associated with this subbody (or root body)
+bool ShapeBoundToBody(const PxShape * PShape, const FBodyInstance * SubBody)
+{
+	FBodyInstance * BI = FPhysxUserData::Get<FBodyInstance>(PShape->userData);
+	return (SubBody->WeldParent == NULL && BI == NULL) || (BI == SubBody && BI->WeldParent != NULL);
+}
+
 TArray<PxShape*> FBodyInstance::GetAllShapes(int32& OutNumSyncShapes) const
 {
 	// grab shapes from sync actor
@@ -3250,7 +3257,7 @@ bool FBodyInstance::LineTrace(struct FHitResult& OutHit, const FVector& Start, c
 	if (DeltaMag > KINDA_SMALL_NUMBER)
 	{
 #if WITH_PHYSX
-		const PxRigidActor* RigidBody = GetPxRigidActor();
+		const PxRigidActor* RigidBody = WeldParent ? WeldParent->GetPxRigidActor() : GetPxRigidActor();
 		if ((RigidBody != NULL) && (RigidBody->getNbShapes() != 0))
 		{
 			// Create filter data used to filter collisions, should always return eTOUCH for LineTraceComponent		
@@ -3270,6 +3277,8 @@ bool FBodyInstance::LineTrace(struct FHitResult& OutHit, const FVector& Start, c
 				PxShape* PShape = PShapes[ShapeIdx];
 				check(PShape);
 
+				if (ShapeBoundToBody(PShape, this) == false) { continue;  }
+
 				const PxU32 HitBufferSize = 1;
 				PxRaycastHit PHits[HitBufferSize];
 
@@ -3339,7 +3348,7 @@ bool FBodyInstance::Sweep(struct FHitResult& OutHit, const FVector& Start, const
 
 		bool bSweepHit = false;
 #if WITH_PHYSX
-		const PxRigidActor* RigidBody = GetPxRigidActor();
+		const PxRigidActor* RigidBody = WeldParent ? WeldParent->GetPxRigidActor() : GetPxRigidActor();
 		if ((RigidBody != NULL) && (RigidBody->getNbShapes() != 0) && (OwnerComponent != NULL))
 		{
 			FPhysXShapeAdaptor ShapeAdaptor(FQuat::Identity, CollisionShape);
@@ -3387,6 +3396,8 @@ bool FBodyInstance::InternalSweepPhysX(struct FHitResult& OutHit, const FVector&
 			PxShape* PShape = PShapes[ShapeIdx];
 			check(PShape);
 
+			if (ShapeBoundToBody(PShape, this) == false){ continue; }
+
 			// Filter so we trace against the right kind of collision
 			PxFilterData ShapeFilter = PShape->getQueryFilterData();
 			const bool bShapeIsComplex = (ShapeFilter.word3 & EPDF_ComplexCollision) != 0;
@@ -3505,9 +3516,22 @@ bool FBodyInstance::OverlapTest(const FVector& Position, const FQuat& Rotation,
 	return bHasOverlap;
 }
 
+FTransform RootSpaceToWeldedSpace(const FBodyInstance * BI, const FTransform & RootTM)
+{
+	if (BI->WeldParent && BI->OwnerComponent.IsValid())
+	{
+		FTransform RootToWelded = BI->OwnerComponent->GetRelativeTransform().Inverse();
+		RootToWelded.ScaleTranslation(BI->Scale3D);
+
+		return RootToWelded * RootTM;
+	}
+
+	return RootTM;
+}
+
 bool FBodyInstance::OverlapMulti(TArray<struct FOverlapResult>& InOutOverlaps, const class UWorld* World, const FTransform* pWorldToComponent, const FVector& Pos, const FRotator& Rot, ECollisionChannel TestChannel, const struct FComponentQueryParams& Params, const struct FCollisionResponseParams& ResponseParams, const struct FCollisionObjectQueryParams& ObjectQueryParams) const
 {
-	if (!IsValidBodyInstance())
+	if ( (IsValidBodyInstance() || (WeldParent && WeldParent->IsValidBodyInstance())) == false )
 	{
 		UE_LOG(LogCollision, Log, TEXT("FBodyInstance::OverlapMulti : (%s) No physics data"), *GetBodyDebugName());
 		return false;
@@ -3522,7 +3546,8 @@ bool FBodyInstance::OverlapMulti(TArray<struct FOverlapResult>& InOutOverlaps, c
 	FTransform BodyInstanceSpaceToTestSpace;
 	if (pWorldToComponent)
 	{
-		const FTransform ShapeSpaceToComponentSpace = GetUnrealWorldTransform() * (*pWorldToComponent);
+		const FTransform RootTM = WeldParent ? WeldParent->GetUnrealWorldTransform() : GetUnrealWorldTransform();
+		const FTransform ShapeSpaceToComponentSpace = RootTM * (*pWorldToComponent);
 		BodyInstanceSpaceToTestSpace = ShapeSpaceToComponentSpace * ComponentSpaceToTestSpace;
 	}
 	else
@@ -3530,8 +3555,16 @@ bool FBodyInstance::OverlapMulti(TArray<struct FOverlapResult>& InOutOverlaps, c
 		BodyInstanceSpaceToTestSpace = ComponentSpaceToTestSpace;
 	}
 
+	//We want to test using global position. However, the global position of the body will be in terms of the root body which we are welded to. So we must undo the relative transform so that our shapes are centered
+	//Global = Parent * Relative => Global * RelativeInverse = Parent
+	if (WeldParent)
+	{
+		BodyInstanceSpaceToTestSpace = RootSpaceToWeldedSpace(this, BodyInstanceSpaceToTestSpace);
+	}
+
 #if WITH_PHYSX
-	if (PxRigidActor* PRigidActor = GetPxRigidActor())
+	PxRigidActor* PRigidActor = WeldParent ? WeldParent->GetPxRigidActor() : GetPxRigidActor();
+	if (PRigidActor)
 	{
 		const PxTransform PBodyInstanceSpaceToTestSpace = U2PTransform(BodyInstanceSpaceToTestSpace);
 
@@ -3550,6 +3583,11 @@ bool FBodyInstance::OverlapMulti(TArray<struct FOverlapResult>& InOutOverlaps, c
 				PxShape* PShape = PShapes[ShapeIdx];
 				check(PShape);
 
+				if (ShapeBoundToBody(PShape, this) == false)
+				{
+					continue;
+				}
+
 				// Calc shape global pose
 				const PxTransform PLocalPose = PShape->getLocalPose();
 				const PxTransform PShapeGlobalPose = PBodyInstanceSpaceToTestSpace.transform(PLocalPose);
@@ -3584,7 +3622,7 @@ bool FBodyInstance::OverlapMulti(TArray<struct FOverlapResult>& InOutOverlaps, c
 #if WITH_PHYSX
 bool FBodyInstance::OverlapPhysX(const PxGeometry& PGeom, const PxTransform& ShapePose) const
 {
-	const PxRigidActor* RigidBody = GetPxRigidActor();
+	const PxRigidActor* RigidBody = WeldParent ? WeldParent->GetPxRigidActor() : GetPxRigidActor();
 
 	if (RigidBody==NULL || RigidBody->getNbShapes()==0)
 	{
@@ -3601,10 +3639,13 @@ bool FBodyInstance::OverlapPhysX(const PxGeometry& PGeom, const PxTransform& Sha
 	{
 		const PxShape* PShape = PShapes[ShapeIdx];
 		check(PShape);
-		
-		if(PxGeometryQuery::overlap(PShape->getGeometry().any(), PxShapeExt::getGlobalPose(*PShape, *RigidBody), PGeom, ShapePose))
+
+		if (ShapeBoundToBody(PShape, this) == false)
 		{
-			return true;
+			if (PxGeometryQuery::overlap(PShape->getGeometry().any(), PxShapeExt::getGlobalPose(*PShape, *RigidBody), PGeom, ShapePose))
+			{
+				return true;
+			}
 		}
 	}
 	return false;
