Update UPrimitiveComponent to use GetBodyInstance for cases when welding you actually want to affect root (e.g. SetVelocity)
Update BodyInstance to properly lock physics scene for write before welding/unwelding

diff --git a/Engine/Source/Runtime/Engine/Private/PhysicsEngine/BodyInstance.cpp b/Engine/Source/Runtime/Engine/Private/PhysicsEngine/BodyInstance.cpp
--- a/Engine/Source/Runtime/Engine/Private/PhysicsEngine/BodyInstance.cpp
+++ b/Engine/Source/Runtime/Engine/Private/PhysicsEngine/BodyInstance.cpp
@@ -1567,19 +1567,7 @@ bool FBodyInstance::Weld(FBodyInstance* TheirBody, const FTransform& TheirTM)
 
 #if WITH_PHYSX
 	
-	//first time we're welding with root so we need to grab the shapes
-	/*if (ShapeToBodyMap.Num() == 0)
-	{
-		int32 NumSyncShapes = 0;
-		TArray<PxShape *> PShapes = GetAllShapes(NumSyncShapes);
-
-		for (int32 ShapeIdx = 0; ShapeIdx < PShapes.Num(); ++ShapeIdx)
-		{
-			PxShape* PShape = PShapes[ShapeIdx];
-			FBodyInstance *& BI = ShapeToBodyMap.FindOrAdd(PShape);
-			BI = this;
-		}
-	}*/
+	
 
 	TArray<PxShape *> PNewShapes;
 
@@ -1588,6 +1576,12 @@ bool FBodyInstance::Weld(FBodyInstance* TheirBody, const FTransform& TheirTM)
 
 	FTransform RelativeTM = TheirTM.GetRelativeTransform(MyTM);
 
+	PxScene * PSyncScene = RigidActorSync ? RigidActorSync->getScene() : NULL;
+	PxScene * PAsyncScene = RigidActorAsync ? RigidActorAsync->getScene() : NULL;
+
+	SCOPED_SCENE_WRITE_LOCK(PSyncScene);
+	SCOPED_SCENE_WRITE_LOCK(PAsyncScene);
+	
 	//child body gets placed into the same scenes as parent body
 	if (PxRigidActor* MyBody = RigidActorSync)
 	{
@@ -1632,6 +1626,12 @@ void FBodyInstance::UnWeld(FBodyInstance* TheirBI)
 
 	bool bNeedsNotification = false;
 
+	PxScene * PSyncScene = RigidActorSync ? RigidActorSync->getScene() : NULL;
+	PxScene * PAsyncScene = RigidActorAsync ? RigidActorAsync->getScene() : NULL;
+
+	SCOPED_SCENE_WRITE_LOCK(PSyncScene);
+	SCOPED_SCENE_WRITE_LOCK(PAsyncScene);
+
 	for (int32 ShapeIdx = 0; ShapeIdx < NumSyncShapes; ++ShapeIdx)
 	{
 		PxShape* PShape = PShapes[ShapeIdx];
diff --git a/Engine/Source/Runtime/Engine/Private/PrimitiveComponentPhysics.cpp b/Engine/Source/Runtime/Engine/Private/PrimitiveComponentPhysics.cpp
--- a/Engine/Source/Runtime/Engine/Private/PrimitiveComponentPhysics.cpp
+++ b/Engine/Source/Runtime/Engine/Private/PrimitiveComponentPhysics.cpp
@@ -173,7 +173,11 @@ void UPrimitiveComponent::AddRadialImpulse(FVector Origin, float Radius, float S
 		return;
 	}
 
-	BodyInstance.AddRadialImpulseToBody(Origin, Radius, Strength, Falloff, bVelChange);
+	FBodyInstance* BI = GetBodyInstance();
+	if (BI)
+	{
+		BI->AddRadialImpulseToBody(Origin, Radius, Strength, Falloff, bVelChange);
+	}
 }
 
 
@@ -202,7 +206,11 @@ void UPrimitiveComponent::AddRadialForce(FVector Origin, float Radius, float Str
 		return;
 	}
 
-	BodyInstance.AddRadialForceToBody(Origin, Radius, Strength, Falloff);
+	FBodyInstance* BI = GetBodyInstance();
+	if (BI)
+	{
+		BI->AddRadialForceToBody(Origin, Radius, Strength, Falloff);
+	}
 }
 
 void UPrimitiveComponent::AddTorque(FVector Torque, FName BoneName)
@@ -300,41 +308,73 @@ void UPrimitiveComponent::WakeAllRigidBodies()
 
 void UPrimitiveComponent::SetEnableGravity(bool bGravityEnabled)
 {
-	BodyInstance.SetEnableGravity(bGravityEnabled);
+	FBodyInstance* BI = GetBodyInstance();
+	if (BI)
+	{
+		BI->SetEnableGravity(bGravityEnabled);
+	}
 }
 
 bool UPrimitiveComponent::IsGravityEnabled() const
 {
-	return BodyInstance.bEnableGravity;
+	FBodyInstance* BI = GetBodyInstance();
+	if (BI)
+	{
+		return BI->bEnableGravity;
+	}
+
+	return false;
 }
 
 void UPrimitiveComponent::SetLinearDamping(float InDamping)
 {
-	BodyInstance.LinearDamping = InDamping;
-	BodyInstance.UpdateDampingProperties();
+	FBodyInstance* BI = GetBodyInstance();
+	if (BI)
+	{
+		BI->LinearDamping = InDamping;
+		BI->UpdateDampingProperties();
+	}
 }
 
 float UPrimitiveComponent::GetLinearDamping() const
 {
-	return BodyInstance.LinearDamping;
+
+	FBodyInstance* BI = GetBodyInstance();
+	if (BI)
+	{
+		return BI->LinearDamping;
+	}
+	
+	return 0.f;
 }
 
 void UPrimitiveComponent::SetAngularDamping(float InDamping)
 {
-	BodyInstance.AngularDamping = InDamping;
-	BodyInstance.UpdateDampingProperties();
+	FBodyInstance* BI = GetBodyInstance();
+	if (BI)
+	{
+		BI->AngularDamping = InDamping;
+		BI->UpdateDampingProperties();
+	}
 }
 
 float UPrimitiveComponent::GetAngularDamping() const
 {
-	return BodyInstance.AngularDamping;
+	FBodyInstance* BI = GetBodyInstance();
+	if (BI)
+	{
+		return BI->AngularDamping;
+	}
+	
+	return 0.f;
 }
 
 float UPrimitiveComponent::GetMass() const
 {
-	if (BodyInstance.IsValidBodyInstance())
+	FBodyInstance* BI = GetBodyInstance();
+	if (BI)
 	{
-		return BodyInstance.GetBodyMass();
+		return BI->GetBodyMass();
 	}
 
 	return 0.0f;
