fix literals for box trace blueprint

diff --git a/Engine/Source/Runtime/Engine/Classes/Kismet/KismetSystemLibrary.h b/Engine/Source/Runtime/Engine/Classes/Kismet/KismetSystemLibrary.h
--- a/Engine/Source/Runtime/Engine/Classes/Kismet/KismetSystemLibrary.h
+++ b/Engine/Source/Runtime/Engine/Classes/Kismet/KismetSystemLibrary.h
@@ -621,7 +621,7 @@ class UKismetSystemLibrary : public UBlueprintFunctionLibrary
 	* @return				True if there was a hit, false otherwise.
 	*/
 	UFUNCTION(BlueprintCallable, Category = "Collision", meta = (bIgnoreSelf = "true", HidePin = "WorldContextObject", DefaultToSelf = "WorldContextObject", AutoCreateRefTerm = "ActorsToIgnore", FriendlyName = "SingleBoxTraceByChannel"))
-	static ENGINE_API bool BoxTraceSingle(UObject* WorldContextObject, const FVector Start, const FVector End, const FVector & HalfSize, const FRotator & Orientation, ETraceTypeQuery TraceChannel, bool bTraceComplex, const TArray<AActor*>& ActorsToIgnore, EDrawDebugTrace::Type DrawDebugType, FHitResult& OutHit, bool bIgnoreSelf);
+	static ENGINE_API bool BoxTraceSingle(UObject* WorldContextObject, const FVector Start, const FVector End, const FVector HalfSize, const FRotator Orientation, ETraceTypeQuery TraceChannel, bool bTraceComplex, const TArray<AActor*>& ActorsToIgnore, EDrawDebugTrace::Type DrawDebugType, FHitResult& OutHit, bool bIgnoreSelf);
 
 	/**
 	* Sweeps a box along the given line and returns all hits encountered.
@@ -637,7 +637,7 @@ class UKismetSystemLibrary : public UBlueprintFunctionLibrary
 	* @return				True if there was a hit, false otherwise.
 	*/
 	UFUNCTION(BlueprintCallable, Category = "Collision", meta = (bIgnoreSelf = "true", HidePin = "WorldContextObject", DefaultToSelf = "WorldContextObject", AutoCreateRefTerm = "ActorsToIgnore", FriendlyName = "MultiBoxTraceByChannel"))
-	static ENGINE_API bool BoxTraceMulti(UObject* WorldContextObject, const FVector Start, const FVector End, const FVector & HalfSize, const FRotator & Orientation, ETraceTypeQuery TraceChannel, bool bTraceComplex, const TArray<AActor*>& ActorsToIgnore, EDrawDebugTrace::Type DrawDebugType, TArray<FHitResult>& OutHits, bool bIgnoreSelf);
+	static ENGINE_API bool BoxTraceMulti(UObject* WorldContextObject, const FVector Start, const FVector End, FVector HalfSize, const FRotator Orientation, ETraceTypeQuery TraceChannel, bool bTraceComplex, const TArray<AActor*>& ActorsToIgnore, EDrawDebugTrace::Type DrawDebugType, TArray<FHitResult>& OutHits, bool bIgnoreSelf);
 
 	
 	/**
diff --git a/Engine/Source/Runtime/Engine/Private/KismetSystemLibrary.cpp b/Engine/Source/Runtime/Engine/Private/KismetSystemLibrary.cpp
--- a/Engine/Source/Runtime/Engine/Private/KismetSystemLibrary.cpp
+++ b/Engine/Source/Runtime/Engine/Private/KismetSystemLibrary.cpp
@@ -1164,7 +1164,7 @@ void DrawDebugSweptBox(const UWorld* InWorld, FVector const& Start, FVector cons
 }
 
 
-bool UKismetSystemLibrary::BoxTraceSingle(UObject* WorldContextObject, const FVector Start, const FVector End, const FVector & HalfSize, const FRotator & Orientation, ETraceTypeQuery TraceChannel, bool bTraceComplex, const TArray<AActor*>& ActorsToIgnore, EDrawDebugTrace::Type DrawDebugType, FHitResult& OutHit, bool bIgnoreSelf)
+bool UKismetSystemLibrary::BoxTraceSingle(UObject* WorldContextObject, const FVector Start, const FVector End, const FVector HalfSize, const FRotator Orientation, ETraceTypeQuery TraceChannel, bool bTraceComplex, const TArray<AActor*>& ActorsToIgnore, EDrawDebugTrace::Type DrawDebugType, FHitResult& OutHit, bool bIgnoreSelf)
 {
 	static const FName BoxTraceSingleName(TEXT("BoxTraceSingle"));
 
@@ -1219,7 +1219,7 @@ bool UKismetSystemLibrary::BoxTraceSingle(UObject* WorldContextObject, const FVe
 	return bHit;
 }
 
-bool UKismetSystemLibrary::BoxTraceMulti(UObject* WorldContextObject, const FVector Start, const FVector End, const FVector & HalfSize, const FRotator & Orientation, ETraceTypeQuery TraceChannel, bool bTraceComplex, const TArray<AActor*>& ActorsToIgnore, EDrawDebugTrace::Type DrawDebugType, TArray<FHitResult>& OutHits, bool bIgnoreSelf)
+bool UKismetSystemLibrary::BoxTraceMulti(UObject* WorldContextObject, const FVector Start, const FVector End, const FVector HalfSize, const FRotator Orientation, ETraceTypeQuery TraceChannel, bool bTraceComplex, const TArray<AActor*>& ActorsToIgnore, EDrawDebugTrace::Type DrawDebugType, TArray<FHitResult>& OutHits, bool bIgnoreSelf)
 {
 	static const FName BoxTraceMultiName(TEXT("BoxTraceMulti"));
 
