Fix absolute value being cast to integer.

- abs() normally returns integer, it's safer to use fabs().
- reported as warning by clang 3.5.0

diff --git a/Engine/Source/Runtime/Engine/Private/Rendering/SeparableSSS.cpp b/Engine/Source/Runtime/Engine/Private/Rendering/SeparableSSS.cpp
--- a/Engine/Source/Runtime/Engine/Private/Rendering/SeparableSSS.cpp
+++ b/Engine/Source/Runtime/Engine/Private/Rendering/SeparableSSS.cpp
@@ -106,14 +106,14 @@ void ComputeMirroredSSSKernel(FLinearColor* TargetBuffer, uint32 TargetBufferSiz
 		{
 			float o = -Range + float(i) * step;
 			float sign = o < 0.0f ? -1.0f : 1.0f;
-			kernel[i].A = Range * sign * abs(pow(o, Exponent)) / pow(Range, Exponent);
+			kernel[i].A = Range * sign * fabs(pow(o, Exponent)) / pow(Range, Exponent);
 		}
 
 		// Calculate the weights:
 		for (int32 i = 0; i < nTotalSamples; i++)
 		{
-			float w0 = i > 0 ? abs(kernel[i].A - kernel[i - 1].A) : 0.0f;
-			float w1 = i < nTotalSamples - 1 ? abs(kernel[i].A - kernel[i + 1].A) : 0.0f;
+			float w0 = i > 0 ? fabs(kernel[i].A - kernel[i - 1].A) : 0.0f;
+			float w1 = i < nTotalSamples - 1 ? fabs(kernel[i].A - kernel[i + 1].A) : 0.0f;
 			float area = (w0 + w1) / 2.0f;
 			FVector t = area * SeparableSSS_Profile(kernel[i].A, FalloffColor);
 			kernel[i].R = t.X;
