Fix caller of a function to get rid of if (this) check.

Checking this for being null is reported as warning by clang 3.5.0. Arguably not compiler's call, but it's indeed better to check that condition before calling the function (what if it is turned virtual at some point?).

diff --git a/Engine/Source/Runtime/Engine/Private/PrecomputedLightVolume.cpp b/Engine/Source/Runtime/Engine/Private/PrecomputedLightVolume.cpp
--- a/Engine/Source/Runtime/Engine/Private/PrecomputedLightVolume.cpp
+++ b/Engine/Source/Runtime/Engine/Private/PrecomputedLightVolume.cpp
@@ -259,9 +259,11 @@ void FPrecomputedLightVolume::InterpolateIncidentRadiancePoint(
 		FSHVectorRGB2& AccumulatedIncidentRadiance,
 		FVector& SkyBentNormal) const
 {
-	// Handle being called on a NULL volume, which can happen for a newly created level,
+	// This could be called on a NULL volume for a newly created level. This is now checked at the callsite, but this check provides extra safety
+	checkf(this, TEXT("FPrecomputedLightVolume::InterpolateIncidentRadiancePoint() is called on a null volume. Fix the call site."));
+
 	// Or a volume that hasn't been initialized yet, which can happen if lighting hasn't been built yet.
-	if (this && bInitialized)
+	if (bInitialized)
 	{
 		FBoxCenterAndExtent BoundingBox( WorldPosition, FVector::ZeroVector );
 		// Iterate over the octree nodes containing the query point.
diff --git a/Engine/Source/Runtime/Renderer/Private/IndirectLightingCache.cpp b/Engine/Source/Runtime/Renderer/Private/IndirectLightingCache.cpp
--- a/Engine/Source/Runtime/Renderer/Private/IndirectLightingCache.cpp
+++ b/Engine/Source/Runtime/Renderer/Private/IndirectLightingCache.cpp
@@ -722,12 +722,15 @@ void FIndirectLightingCache::InterpolatePoint(
 	for (int32 VolumeIndex = 0; VolumeIndex < Scene->PrecomputedLightVolumes.Num(); VolumeIndex++)
 	{
 		const FPrecomputedLightVolume* PrecomputedLightVolume = Scene->PrecomputedLightVolumes[VolumeIndex];
-		PrecomputedLightVolume->InterpolateIncidentRadiancePoint(
-			Block.Min + Block.Size / 2, 
-			AccumulatedWeight, 
-			AccumulatedDirectionalShadowing,
-			AccumulatedIncidentRadiance,
-			AccumulatedSkyBentNormal);
+		if (PrecomputedLightVolume)
+		{
+			PrecomputedLightVolume->InterpolateIncidentRadiancePoint(
+				Block.Min + Block.Size / 2,
+				AccumulatedWeight,
+				AccumulatedDirectionalShadowing,
+				AccumulatedIncidentRadiance,
+				AccumulatedSkyBentNormal);
+		}
 	}
 
 	if (AccumulatedWeight > 0)
