Avoiding checking this for null.

This is reported as warning by clang 3.5.0. Arguably not a compiler's call, but it's indeed better to check that condition before calling the function and not inside it.

diff --git a/Engine/Source/Runtime/Engine/Private/PrecomputedLightVolume.cpp b/Engine/Source/Runtime/Engine/Private/PrecomputedLightVolume.cpp
--- a/Engine/Source/Runtime/Engine/Private/PrecomputedLightVolume.cpp
+++ b/Engine/Source/Runtime/Engine/Private/PrecomputedLightVolume.cpp
@@ -259,10 +259,10 @@ void FPrecomputedLightVolume::InterpolateIncidentRadiancePoint(
 		FSHVectorRGB2& AccumulatedIncidentRadiance,
 		FVector& SkyBentNormal) const
 {
-	// This could be called on a NULL volume for a newly created level. This is now checked at the callsite, but this check provides extra safety
+	// This could be called on a NULL volume for a newly created level. This is now checked at the call site, but this check provides extra safety
 	checkf(this, TEXT("FPrecomputedLightVolume::InterpolateIncidentRadiancePoint() is called on a null volume. Fix the call site."));
 
-	// Or a volume that hasn't been initialized yet, which can happen if lighting hasn't been built yet.
+	// Handle being called on a volume that hasn't been initialized yet, which can happen if lighting hasn't been built yet.
 	if (bInitialized)
 	{
 		FBoxCenterAndExtent BoundingBox( WorldPosition, FVector::ZeroVector );
@@ -302,9 +302,11 @@ void FPrecomputedLightVolume::InterpolateIncidentRadianceBlock(
 	TArray<FSHVectorRGB2>& AccumulatedIncidentRadiance,
 	TArray<FVector>& AccumulatedSkyBentNormal) const
 {
-	// Handle being called on a NULL volume, which can happen for a newly created level,
-	// Or a volume that hasn't been initialized yet, which can happen if lighting hasn't been built yet.
-	if (this && bInitialized)
+	// This could be called on a NULL volume for a newly created level. This is now checked at the callsite, but this check provides extra safety
+	checkf(this, TEXT("FPrecomputedLightVolume::InterpolateIncidentRadianceBlock() is called on a null volume. Fix the call site."));
+
+	// Handle being called on a volume that hasn't been initialized yet, which can happen if lighting hasn't been built yet.
+	if (bInitialized)
 	{
 		static TArray<const FVolumeLightingSample*> PotentiallyIntersectingSamples;
 		PotentiallyIntersectingSamples.Reset(100);
diff --git a/Engine/Source/Runtime/Renderer/Private/IndirectLightingCache.cpp b/Engine/Source/Runtime/Renderer/Private/IndirectLightingCache.cpp
--- a/Engine/Source/Runtime/Renderer/Private/IndirectLightingCache.cpp
+++ b/Engine/Source/Runtime/Renderer/Private/IndirectLightingCache.cpp
@@ -850,14 +850,17 @@ void FIndirectLightingCache::InterpolateBlock(
 							&& CellIndex.Z + NumQueryStepCells.Z <= Block.TexelSize);
 
 						// Interpolate from the SH volume lighting samples that Lightmass computed
-						PrecomputedLightVolume->InterpolateIncidentRadianceBlock(
-							BoundingBox, 
-							NumQueryStepCells, 
-							FIntVector(Block.TexelSize), 
-							CellIndex, 
-							AccumulatedWeight, 
-							AccumulatedIncidentRadiance,
-							AccumulatedSkyBentNormal);
+						if (PrecomputedLightVolume)
+						{
+							PrecomputedLightVolume->InterpolateIncidentRadianceBlock(
+								BoundingBox,
+								NumQueryStepCells,
+								FIntVector(Block.TexelSize),
+								CellIndex,
+								AccumulatedWeight,
+								AccumulatedIncidentRadiance,
+								AccumulatedSkyBentNormal);
+						}
 					}
 				}
 			}
@@ -868,6 +871,7 @@ void FIndirectLightingCache::InterpolateBlock(
 		for (int32 VolumeIndex = 0; VolumeIndex < Scene->PrecomputedLightVolumes.Num(); VolumeIndex++)
 		{
 			const FPrecomputedLightVolume* PrecomputedLightVolume = Scene->PrecomputedLightVolumes[VolumeIndex];
+			check(PrecomputedLightVolume);
 			check(PrecomputedLightVolume->IsUsingHighQualityLightMap() == AllowHighQualityLightmaps(Scene->GetFeatureLevel()));
 			// Interpolate from the SH volume lighting samples that Lightmass computed
 			// Query using the bounds of all the samples in this block
