Hide unused categories for WorldSettings

diff --git a/Engine/Source/Runtime/Engine/Classes/GameFramework/WorldSettings.h b/Engine/Source/Runtime/Engine/Classes/GameFramework/WorldSettings.h
--- a/Engine/Source/Runtime/Engine/Classes/GameFramework/WorldSettings.h
+++ b/Engine/Source/Runtime/Engine/Classes/GameFramework/WorldSettings.h
@@ -183,7 +183,7 @@ struct ENGINE_API FNetViewer
 /**
  * Actor containing all script accessible world properties.
  */
-UCLASS(config=game, showcategories=(Rendering), hidecategories=(Actor, Advanced, Display, Events, Object, Attachment, Info, Input, Blueprint, Layers), showcategories=("Input|MouseInput", "Input|TouchInput"), notplaceable)
+UCLASS(config=game, showcategories=(Rendering), hidecategories=(Actor, Advanced, Display, Events, Object, Attachment, Info, Input, Blueprint, Layers, Tags, Replication), showcategories=("Input|MouseInput", "Input|TouchInput"), notplaceable)
 class ENGINE_API AWorldSettings : public AInfo, public IInterface_AssetUserData
 {
 	GENERATED_UCLASS_BODY()
