Fixed motage debug information output to correctly show next section and desired weight.

diff --git a/Engine/Source/Runtime/Engine/Private/Animation/AnimInstance.cpp b/Engine/Source/Runtime/Engine/Private/Animation/AnimInstance.cpp
--- a/Engine/Source/Runtime/Engine/Private/Animation/AnimInstance.cpp
+++ b/Engine/Source/Runtime/Engine/Private/Animation/AnimInstance.cpp
@@ -673,7 +673,7 @@ void UAnimInstance::DisplayDebug(class UCanvas* Canvas, const FDebugDisplayInfo&
 
 			Canvas->SetLinearDrawColor((MontageInstance->IsActive()) ? ActiveColor : TextWhite);
 
-			FString MontageEntry = FString::Printf(TEXT("%i) %s Sec: %s W:%.3f DW:%.3f"), MontageIndex, *MontageInstance->Montage->GetName(), *MontageInstance->GetCurrentSection().ToString(), *MontageInstance->GetNextSection().ToString(), MontageInstance->Weight, MontageInstance->DesiredWeight);
+			FString MontageEntry = FString::Printf(TEXT("%i) %s CurrSec: %s NextSec: %s W:%.3f DW:%.3f"), MontageIndex, *MontageInstance->Montage->GetName(), *MontageInstance->GetCurrentSection().ToString(), *MontageInstance->GetNextSection().ToString(), MontageInstance->Weight, MontageInstance->DesiredWeight);
 			Canvas->DrawText(RenderFont, MontageEntry, Indent, YPos, 1.f, 1.f, RenderInfo);
 			YPos += YL;
 		}
