attempt to fix a crash in navoctree update (parent chain)
ttp# 345908

diff --git a/Engine/Source/Runtime/Engine/Classes/AI/Navigation/NavigationSystem.h b/Engine/Source/Runtime/Engine/Classes/AI/Navigation/NavigationSystem.h
--- a/Engine/Source/Runtime/Engine/Classes/AI/Navigation/NavigationSystem.h
+++ b/Engine/Source/Runtime/Engine/Classes/AI/Navigation/NavigationSystem.h
@@ -639,7 +639,7 @@ protected:
 	TMap<const UObject*, FOctreeElementId> ObjectToOctreeId;
 
 	/** Map of all objects that are tied to indexed navigation parent */
-	TMultiMap<UObject*, UObject*> OctreeChildNodesMap;
+	TMultiMap<UObject*, FWeakObjectPtr> OctreeChildNodesMap;
 
 	/** Map of all custom navigation links, that are relevant for path following */
 	TMap<uint32, INavLinkCustomInterface*> CustomLinksMap;
diff --git a/Engine/Source/Runtime/Engine/Private/AI/Navigation/NavigationSystem.cpp b/Engine/Source/Runtime/Engine/Private/AI/Navigation/NavigationSystem.cpp
--- a/Engine/Source/Runtime/Engine/Private/AI/Navigation/NavigationSystem.cpp
+++ b/Engine/Source/Runtime/Engine/Private/AI/Navigation/NavigationSystem.cpp
@@ -1848,7 +1848,7 @@ FSetElementId UNavigationSystem::RegisterNavOctreeElement(UObject* ElementOwner,
 		UObject* ParentNode = ElementInterface->GetNavigationParent();
 		if (ParentNode)
 		{
-			OctreeChildNodesMap.AddUnique(ParentNode, ElementOwner);
+			OctreeChildNodesMap.AddUnique(ParentNode, FWeakObjectPtr(ElementOwner));
 			bCanAdd = true;
 		}
 		else
@@ -1996,7 +1996,7 @@ void UNavigationSystem::UnregisterNavOctreeElement(UObject* ElementOwner, class
 			// and it's not part of parent chain update
 			// remove it from map and force update on parent to rebuild octree element
 
-			OctreeChildNodesMap.RemoveSingle(ParentNode, ElementOwner);
+			OctreeChildNodesMap.RemoveSingle(ParentNode, FWeakObjectPtr(ElementOwner));
 			UpdateNavOctreeParentChain(ParentNode);
 		}
 	}
@@ -2137,7 +2137,7 @@ void UNavigationSystem::UpdateNavOctreeElement(UObject* ElementOwner, class INav
 
 void UNavigationSystem::UpdateNavOctreeParentChain(UObject* ElementOwner)
 {
-	TArray<UObject*> ChildNodes;
+	TArray<FWeakObjectPtr> ChildNodes;
 
 	OctreeChildNodesMap.MultiFind(ElementOwner, ChildNodes);
 	if (ChildNodes.Num() == 0)
@@ -2152,8 +2152,12 @@ void UNavigationSystem::UpdateNavOctreeParentChain(UObject* ElementOwner)
 	
 	for (int32 Idx = 0; Idx < ChildNodes.Num(); Idx++)
 	{
-		ChildNavInterfaces[Idx] = InterfaceCast<INavRelevantInterface>(ChildNodes[Idx]);
-		UnregisterNavOctreeElement(ChildNodes[Idx], ChildNavInterfaces[Idx], UpdateFlags);
+		if (ChildNodes[Idx].IsValid())
+		{
+			UObject* ChildNodeOb = ChildNodes[Idx].Get();
+			ChildNavInterfaces[Idx] = InterfaceCast<INavRelevantInterface>(ChildNodeOb);
+			UnregisterNavOctreeElement(ChildNodeOb, ChildNavInterfaces[Idx], UpdateFlags);
+		}
 	}
 
 	UnregisterNavOctreeElement(ElementOwner, ElementInterface, UpdateFlags);
@@ -2161,7 +2165,10 @@ void UNavigationSystem::UpdateNavOctreeParentChain(UObject* ElementOwner)
 
 	for (int32 Idx = 0; Idx < ChildNodes.Num(); Idx++)
 	{
-		RegisterNavOctreeElement(ChildNodes[Idx], ChildNavInterfaces[Idx], UpdateFlags);
+		if (ChildNodes[Idx].IsValid())
+		{
+			RegisterNavOctreeElement(ChildNodes[Idx].Get(), ChildNavInterfaces[Idx], UpdateFlags);
+		}
 	}
 }
 
