PDB generation for showing Xbox One shader source in PIX. Currently only enabled for QAGame (so we're not enforcing a choice of PDB directory on subscribers).

diff --git a/Engine/Source/Runtime/Engine/Private/ShaderCompiler/ShaderCompiler.cpp b/Engine/Source/Runtime/Engine/Private/ShaderCompiler/ShaderCompiler.cpp
--- a/Engine/Source/Runtime/Engine/Private/ShaderCompiler/ShaderCompiler.cpp
+++ b/Engine/Source/Runtime/Engine/Private/ShaderCompiler/ShaderCompiler.cpp
@@ -2062,6 +2062,15 @@ void GlobalBeginCompileShader(
 	}
 #endif
 
+	{
+		FString ShaderPDBRoot;
+		GConfig->GetString(TEXT("DevOptions.Shaders"), TEXT("ShaderPDBRoot"), ShaderPDBRoot, GEngineIni);
+		if ( !ShaderPDBRoot.IsEmpty() )
+		{
+			Input.Environment.SetDefine(TEXT("SHADER_PDB_ROOT"), *ShaderPDBRoot);
+		}
+	}
+
 	{
 		static const auto CVar = IConsoleManager::Get().FindTConsoleVariableDataInt(TEXT("Compat.UseDXT5NormalMaps"));
 		Input.Environment.SetDefine(TEXT("DXT5_NORMALMAPS"), CVar ? (CVar->GetValueOnGameThread() != 0) : 0);
