Fixed UNavigationPath.PathPoints not getting updated with changes to the actual observer navigation path #UE4

This addressed TTP#345921

diff --git a/Engine/Source/Runtime/Engine/Classes/AI/Navigation/NavigationPath.h b/Engine/Source/Runtime/Engine/Classes/AI/Navigation/NavigationPath.h
--- a/Engine/Source/Runtime/Engine/Classes/AI/Navigation/NavigationPath.h
+++ b/Engine/Source/Runtime/Engine/Classes/AI/Navigation/NavigationPath.h
@@ -5,6 +5,7 @@
 #include "NavigationPath.generated.h"
 
 class UNavigationPath;
+struct FNavigationPath;
 
 DECLARE_DYNAMIC_MULTICAST_DELEGATE_TwoParams(FOnNavigationPathUpdated, UNavigationPath*, AffectedPath, TEnumAsByte<ENavPathEvent::Type>, PathEvent);
 
@@ -73,4 +74,6 @@ public:
 protected:
 	void DrawDebug(UCanvas* Canvas, APlayerController*);
 	void OnPathEvent(FNavigationPath* Path, ENavPathEvent::Type PathEvent);
+
+	void SetPathPointsFromPath(FNavigationPath& NativePath);
 };
diff --git a/Engine/Source/Runtime/Engine/Private/AI/Navigation/NavigationPath.cpp b/Engine/Source/Runtime/Engine/Private/AI/Navigation/NavigationPath.cpp
--- a/Engine/Source/Runtime/Engine/Private/AI/Navigation/NavigationPath.cpp
+++ b/Engine/Source/Runtime/Engine/Private/AI/Navigation/NavigationPath.cpp
@@ -975,7 +975,15 @@ void UNavigationPath::OnPathEvent(FNavigationPath* UpdatedPath, ENavPathEvent::T
 	if (UpdatedPath == SharedPath.Get())
 	{
 		PathUpdatedNotifier.Broadcast(this, PathEvent);
-		bIsValid = SharedPath.IsValid() && SharedPath->IsValid();
+		if (SharedPath.IsValid() && SharedPath->IsValid())
+		{
+			bIsValid = true;
+			SetPathPointsFromPath(*UpdatedPath);
+		}
+		else
+		{
+			bIsValid = false;
+		}
 	}
 }
 
@@ -1081,11 +1089,7 @@ void UNavigationPath::SetPath(FNavPathSharedPtr NewSharedPath)
 				NewPath->EnableRecalculationOnInvalidation(RecalculateOnInvalidation == ENavigationOptionFlag::Enable);
 			}
 			
-			PathPoints.Reset(NewPath->GetPathPoints().Num());
-			for (const auto& PathPoint : NewPath->GetPathPoints())
-			{
-				PathPoints.Add(PathPoint.Location);
-			}
+			SetPathPointsFromPath(*NewPath);
 		}
 		else
 		{
@@ -1096,3 +1100,11 @@ void UNavigationPath::SetPath(FNavPathSharedPtr NewSharedPath)
 	}
 }
 
+void UNavigationPath::SetPathPointsFromPath(FNavigationPath& NativePath)
+{
+	PathPoints.Reset(NativePath.GetPathPoints().Num());
+	for (const auto& PathPoint : NativePath.GetPathPoints())
+	{
+		PathPoints.Add(PathPoint.Location);
+	}
+}
