- The default value was removed during multi montage work, so reverted back that part.

345456	CRITICAL: REGRESSION: COUCHKNIGHTS: Unable to attack multiple times

diff --git a/Engine/Source/Runtime/Engine/Classes/Animation/AnimInstance.h b/Engine/Source/Runtime/Engine/Classes/Animation/AnimInstance.h
--- a/Engine/Source/Runtime/Engine/Classes/Animation/AnimInstance.h
+++ b/Engine/Source/Runtime/Engine/Classes/Animation/AnimInstance.h
@@ -357,7 +357,7 @@ public:
 public:
 	/** Plays an animation montage. Returns the length of the animation montage in seconds. Returns 0.f if failed to play. */
 	UFUNCTION(BlueprintCallable, Category = "Animation")
-	float Montage_Play(UAnimMontage * MontageToPlay, float InPlayRate);
+	float Montage_Play(UAnimMontage * MontageToPlay, float InPlayRate = 1.f);
 
 	/** Stops the animation montage. If reference is NULL, it will stop ALL active montages. */
 	UFUNCTION(BlueprintCallable, Category = "Animation")
diff --git a/Engine/Source/Runtime/Engine/Private/Animation/AnimInstance.cpp b/Engine/Source/Runtime/Engine/Private/Animation/AnimInstance.cpp
--- a/Engine/Source/Runtime/Engine/Private/Animation/AnimInstance.cpp
+++ b/Engine/Source/Runtime/Engine/Private/Animation/AnimInstance.cpp
@@ -1645,7 +1645,7 @@ bool UAnimInstance::IsPlayingSlotAnimation(UAnimSequenceBase* Asset, FName SlotN
 }
 
 /** Play a Montage. Returns Length of Montage in seconds. Returns 0.f if failed to play. */
-float UAnimInstance::Montage_Play(UAnimMontage * MontageToPlay, float InPlayRate)
+float UAnimInstance::Montage_Play(UAnimMontage * MontageToPlay, float InPlayRate/*= 1.f*/)
 {
 	if (MontageToPlay && (MontageToPlay->SequenceLength > 0.f))
 	{
