Add TickAnimation back to InitAnim till issues can be solved.

diff --git a/Engine/Source/Runtime/Engine/Private/SkeletalMeshComponent.cpp b/Engine/Source/Runtime/Engine/Private/SkeletalMeshComponent.cpp
--- a/Engine/Source/Runtime/Engine/Private/SkeletalMeshComponent.cpp
+++ b/Engine/Source/Runtime/Engine/Private/SkeletalMeshComponent.cpp
@@ -221,6 +221,8 @@ void USkeletalMeshComponent::InitAnim(bool bForceReinit)
 
 		InitializeAnimScriptInstance(bForceReinit);
 
+		TickAnimation(0.f); //Make sure we have a valid pose
+
 		RefreshBoneTransforms();
 		UpdateComponentToWorld();
 	}
