Fix for Mesh Edges with GetDynamicMeshElements

diff --git a/Engine/Source/Runtime/Engine/Private/PrimitiveDrawingUtils.cpp b/Engine/Source/Runtime/Engine/Private/PrimitiveDrawingUtils.cpp
--- a/Engine/Source/Runtime/Engine/Private/PrimitiveDrawingUtils.cpp
+++ b/Engine/Source/Runtime/Engine/Private/PrimitiveDrawingUtils.cpp
@@ -1596,6 +1596,9 @@ void ApplyViewModeOverrides(
 		if (EngineShowFlags.MeshEdges)
 		{
 			FMeshBatch& MeshEdgeElement = Collector.AllocateMesh();
+			MeshEdgeElement = Mesh;
+			// Avoid infinite recursion
+			MeshEdgeElement.bCanApplyViewModeOverrides = false;
 
 			// Draw the mesh's edges in blue, on top of the base geometry.
 			if (MeshEdgeElement.MaterialRenderProxy->GetMaterial(FeatureLevel)->MaterialModifiesMeshPosition_RenderThread())
