nicer property names in the editor

diff --git a/Engine/Source/Runtime/Engine/Classes/Components/DirectionalLightComponent.h b/Engine/Source/Runtime/Engine/Classes/Components/DirectionalLightComponent.h
--- a/Engine/Source/Runtime/Engine/Classes/Components/DirectionalLightComponent.h
+++ b/Engine/Source/Runtime/Engine/Classes/Components/DirectionalLightComponent.h
@@ -15,7 +15,7 @@ class UDirectionalLightComponent : public ULightComponent
 	GENERATED_UCLASS_BODY()
 
 	/** Whether to occlude fog and atmosphere inscattering with screenspace blurred occlusion from this light. */
-	UPROPERTY(EditAnywhere, BlueprintReadOnly, Category=LightShafts)
+	UPROPERTY(EditAnywhere, BlueprintReadOnly, Category=LightShafts, meta=(DisplayName = "Light Shaft Occlusion"))
 	uint32 bEnableLightShaftOcclusion:1;
 
 	/** 
@@ -98,7 +98,7 @@ class UDirectionalLightComponent : public ULightComponent
 	 * Whether to use ray traced distance field area shadows.  The project setting bGenerateMeshDistanceFields must be enabled for this to have effect.
 	 * Distance field shadows have a large per-pixel cost but a low per-object cost (and don't depend on triangle count) so they are effective for distant shadows from a dynamic sun.
 	 */
-	UPROPERTY(EditAnywhere, BlueprintReadOnly, Category=DistanceFieldShadows)
+	UPROPERTY(EditAnywhere, BlueprintReadOnly, Category=DistanceFieldShadows, meta=(DisplayName = "Distance Field Shadows"))
 	bool bUseDistanceFieldShadows;
 
 	/** 
@@ -111,7 +111,7 @@ class UDirectionalLightComponent : public ULightComponent
 	UPROPERTY(EditAnywhere, BlueprintReadOnly, Category=Light, meta=(ShowOnlyInnerProperties))
 	struct FLightmassDirectionalLightSettings LightmassSettings;
 
-	UPROPERTY(EditAnywhere, BlueprintReadOnly, Category=Light)
+	UPROPERTY(EditAnywhere, BlueprintReadOnly, Category=Light, meta=(DisplayName = "Atmosphere Sun Light"))
 	uint32 bUsedAsAtmosphereSunLight : 1;
 
 	UFUNCTION(BlueprintCallable, Category="Rendering|Lighting")
diff --git a/Engine/Source/Runtime/Engine/Classes/Components/LightComponent.h b/Engine/Source/Runtime/Engine/Classes/Components/LightComponent.h
--- a/Engine/Source/Runtime/Engine/Classes/Components/LightComponent.h
+++ b/Engine/Source/Runtime/Engine/Classes/Components/LightComponent.h
@@ -89,7 +89,7 @@ class ENGINE_API ULightComponent : public ULightComponentBase
 	/**
 	 * Whether the light should be injected into the Light Propagation Volume
 	 **/
-	UPROPERTY(EditAnywhere, BlueprintReadOnly, Category=Light, AdvancedDisplay)
+	UPROPERTY(EditAnywhere, BlueprintReadOnly, Category=Light, AdvancedDisplay, meta=(DisplayName = "Dynamic Indirect Lighting"))
 	uint32 bAffectDynamicIndirectLighting : 1;
 
 	/** 
@@ -119,7 +119,7 @@ class ENGINE_API ULightComponent : public ULightComponentBase
 	 * Distance at which the light function should be completely faded to DisabledBrightness.  
 	 * This is useful for hiding aliasing from light functions applied in the distance.
 	 */
-	UPROPERTY(EditAnywhere, BlueprintReadOnly, Category=LightFunction)
+	UPROPERTY(EditAnywhere, BlueprintReadOnly, Category=LightFunction, meta=(DisplayName = "Fade Distance"))
 	float LightFunctionFadeDistance;
 
 	/** 
@@ -134,7 +134,7 @@ class ENGINE_API ULightComponent : public ULightComponentBase
 	 * For directional lights, the color around the light direction will be blurred radially and added back to the scene.
 	 * for point lights, the color on pixels closer than the light's SourceRadius will be blurred radially and added back to the scene.
 	 */
-	UPROPERTY(EditAnywhere, BlueprintReadOnly, Category=LightShafts)
+	UPROPERTY(EditAnywhere, BlueprintReadOnly, Category=LightShafts, meta=(DisplayName = "Light Shaft Bloom"))
 	uint32 bEnableLightShaftBloom:1;
 
 	/** Scales the additive color. */
diff --git a/Engine/Source/Runtime/Engine/Classes/Materials/Material.h b/Engine/Source/Runtime/Engine/Classes/Materials/Material.h
--- a/Engine/Source/Runtime/Engine/Classes/Materials/Material.h
+++ b/Engine/Source/Runtime/Engine/Classes/Materials/Material.h
@@ -388,14 +388,14 @@ public:
 	FMaterialAttributesInput MaterialAttributes;
 
 	/** Indicates that the material should be rendered in the SeparateTranslucency Pass (not affected by DOF, requires bAllowSeparateTranslucency to be set in .ini). */
-	UPROPERTY(EditAnywhere, Category=Translucency)
+	UPROPERTY(EditAnywhere, Category=Translucency, meta=(DisplayName = "Separate Translucency"))
 	uint32 bEnableSeparateTranslucency:1;
 
 	/**
 	 * Indicates that the material should be rendered using responsive anti-aliasing. Improves sharpness of small moving particles such as sparks.
 	 * Only use for small moving features because it will cause aliasing of the background.
 	 */
-	UPROPERTY(EditAnywhere, Category=Translucency)
+	UPROPERTY(EditAnywhere, Category=Translucency, meta=(DisplayName = "Responsive AA"))
 	uint32 bEnableResponsiveAA:1;
 
 	/** Indicates that the material should be rendered without backface culling and the normal should be flipped for backfaces. */
@@ -613,11 +613,11 @@ public:
 	TEnumAsByte<enum EMaterialTessellationMode> D3D11TessellationMode;
 
 	/** Prevents cracks in the surface of the mesh when using tessellation. */
-	UPROPERTY(EditAnywhere, BlueprintReadOnly, Category=Tessellation)
+	UPROPERTY(EditAnywhere, BlueprintReadOnly, Category=Tessellation, meta=(DisplayName = "Crack Free Displacement"))
 	uint32 bEnableCrackFreeDisplacement:1;
 
 	/** Enables adaptive tessellation, which tries to maintain a uniform number of pixels per triangle. */
-	UPROPERTY(EditAnywhere, BlueprintReadOnly, Category=Tessellation)
+	UPROPERTY(EditAnywhere, BlueprintReadOnly, Category=Tessellation, meta=(DisplayName = "Adaptive Tessellation"))
 	uint32 bEnableAdaptiveTessellation:1;
 
 	/** Enables a wireframe view of the mesh the material is applied to.  */
