Cinematic Perf capture test
--Adding a 1 second wait before opening the map just.
--Added more logging.

diff --git a/Engine/Source/Runtime/Engine/Private/Tests/EngineAutomationTests.cpp b/Engine/Source/Runtime/Engine/Private/Tests/EngineAutomationTests.cpp
--- a/Engine/Source/Runtime/Engine/Private/Tests/EngineAutomationTests.cpp
+++ b/Engine/Source/Runtime/Engine/Private/Tests/EngineAutomationTests.cpp
@@ -290,6 +290,7 @@ bool FLoadGameMapCommand::Update()
 	check(GEngine->GetWorldContexts().Num() == 1);
 	check(GEngine->GetWorldContexts()[0].WorldType == EWorldType::Game);
 
+	UE_LOG(LogEngineAutomationTests, Log, TEXT("Loading Map Now. '%s'"), *MapName);
 	GEngine->Exec(GEngine->GetWorldContexts()[0].World(), *FString::Printf(TEXT("Open %s"), *MapName));
 	return true;
 }
@@ -304,6 +305,7 @@ bool FExecWorldStringLatentCommand::Update()
 	check(GEngine->GetWorldContexts().Num() == 1);
 	check(GEngine->GetWorldContexts()[0].WorldType == EWorldType::Game);
 
+	UE_LOG(LogEngineAutomationTests, Log, TEXT("Running Exec Command. '%s'"), *ExecCommand);
 	GEngine->Exec(GEngine->GetWorldContexts()[0].World(), *ExecCommand);
 	return true;
 }
@@ -340,6 +342,7 @@ bool FCinematicFPSPerfTest::RunTest(const FString& Parameters)
 
 		//Load map
 		const FString MapName = Parameters;
+		ADD_LATENT_AUTOMATION_COMMAND(FEngineWaitLatentCommand(1.0f));
 		ADD_LATENT_AUTOMATION_COMMAND(FLoadGameMapCommand(MapName));
 		ADD_LATENT_AUTOMATION_COMMAND(FEngineWaitLatentCommand(1.0f));
 
