#ue4 - Various improvements to ProjectileMovementComponent (acceleration cleanup, sub-stepping control, documentation).

- ComputeMoveDelta() is no longer able to cause movement exceeding max velocity.
- Added some params to some new functions dealing with acceleration and homing targets.
- Added support for custom sub-stepping behavior:
   - New variables: bForceSubStepping, MaxSimulationTimeStep, MaxSimulationIterations.
   - New functions: GetSimulationTimeStep(), ShouldUseSubStepping()

diff --git a/Engine/Source/Runtime/Engine/Classes/GameFramework/ProjectileMovementComponent.h b/Engine/Source/Runtime/Engine/Classes/GameFramework/ProjectileMovementComponent.h
--- a/Engine/Source/Runtime/Engine/Classes/GameFramework/ProjectileMovementComponent.h
+++ b/Engine/Source/Runtime/Engine/Classes/GameFramework/ProjectileMovementComponent.h
@@ -5,7 +5,8 @@
 #include "ProjectileMovementComponent.generated.h"
 
 /**
- * Projectile Movement component updates position of associated PrimitiveComponent during its tick.
+ * ProjectileMovementComponent updates the position of the associated UpdatedComponent during its tick.
+ * Behavior such as bouncing after impacts and homing toward a target are supported.
  * If the updated component is simulating physics, only the initial launch parameters (when initial velocity is non-zero) will affect the projectile,
  * and the physics sim will take over from there.
  */
@@ -40,6 +41,22 @@ class ENGINE_API UProjectileMovementComponent : public UMovementComponent
 	UPROPERTY(EditAnywhere, BlueprintReadWrite, Category=Projectile)
 	uint32 bInitialVelocityInLocalSpace:1;
 
+	/**
+	 * If true, forces sub-stepping to break up movement into discrete smaller steps to improve accuracy of the trajectory.
+	 * Objects that move in a straight line typically do *not* need to set this, as movement always uses continuous collision detection (sweeps) so collision is not missed.
+	 * Sub-stepping is automatically enabled when under the effects of gravity or when homing towards a target.
+	 * @see MaxSimulationTimeStep, MaxSimulationIterations
+	 */
+	UPROPERTY(EditAnywhere, BlueprintReadWrite, Category=ProjectileSimulation)
+	uint32 bForceSubStepping:1;
+
+	/**
+	 * If true, we will accelerate towards our homing target.
+	 * @see HomingTargetComponent, HomingAccelerationMagnitude
+	 */
+	UPROPERTY(EditAnywhere, BlueprintReadWrite, Category = Homing)
+	uint32 bIsHomingProjectile:1;
+
 	/** Custom gravity scale for this projectile.  Set to 0 for no gravity. */
 	UPROPERTY(EditAnywhere, BlueprintReadWrite, Category=Projectile)
 	float ProjectileGravityScale;
@@ -56,7 +73,7 @@ class ENGINE_API UProjectileMovementComponent : public UMovementComponent
 	UPROPERTY(EditAnywhere, BlueprintReadWrite, Category=ProjectileBounces)
 	float Friction;
 
-	/** Consider the projectile at rest if the velocity is below this threshold after a bounce (uu/sec). Ignored if bShouldBounce=false, in which case the projectile stops simulating on the first impact. */
+	/** If velocity is below this threshold after a bounce, stops simulating and triggers the OnProjectileStop event. Ignored if bShouldBounce=false, in which case the projectile stops simulating on the first impact. */
 	UPROPERTY(EditAnywhere, BlueprintReadWrite, Category=ProjectileBounces)
 	float BounceVelocityStopSimulatingThreshold;
 
@@ -68,16 +85,15 @@ class ENGINE_API UProjectileMovementComponent : public UMovementComponent
 	UPROPERTY(BlueprintAssignable)
 	FOnProjectileStopDelegate OnProjectileStop;
 
-	/** The magnitude of our acceleration towards the homing target */
+	/** The magnitude of our acceleration towards the homing target. Overall velocity magnitude will still be limited by MaxSpeed. */
 	UPROPERTY(EditAnywhere, BlueprintReadWrite, Category = Homing)
 	float HomingAccelerationMagnitude;
 
-	/** If true, we will home towards our homing target. */
-	UPROPERTY(EditAnywhere, BlueprintReadWrite, Category = Homing)
-	bool bIsHomingProjectile;
-
-	/** The current target we are homing towards */
-	UPROPERTY(EditAnywhere, BlueprintReadWrite, Category = Homing)
+	/**
+	 * The current target we are homing towards. Can only be set at runtime (when projectile is spawned or updating).
+	 * @see bIsHomingProjectile
+	 */
+	UPROPERTY(VisibleInstanceOnly, BlueprintReadWrite, Category = Homing)
 	TWeakObjectPtr<USceneComponent> HomingTargetComponent;
 
 	/** Sets the velocity to the new value, rotated into Actor space. */
@@ -106,7 +122,7 @@ class ENGINE_API UProjectileMovementComponent : public UMovementComponent
 	bool ShouldApplyGravity() const { return ProjectileGravityScale != 0.f; }
 
 	/** @returns the velocity after DeltaTime */
-	virtual FVector CalculateVelocity(FVector OldVelocity, float DeltaTime, bool bApplyGravity);
+	virtual FVector CalculateVelocity(FVector OldVelocity, float DeltaTime, bool bGravityEnabled);
 
 	/** Clears the reference to UpdatedComponent, fires stop event, and stops ticking. */
 	UFUNCTION(BlueprintCallable, Category="Game|Components|ProjectileMovement")
@@ -114,7 +130,47 @@ class ENGINE_API UProjectileMovementComponent : public UMovementComponent
 
 	bool HasStoppedSimulation() { return UpdatedComponent == NULL; }
 
+	/**
+	 * Compute remaining time step given remaining time and current iterations.
+	 * The last iteration (limited by MaxSimulationIterations) always returns the remaining time, which may violate MaxSimulationTimeStep.
+	 *
+	 * @param RemainingTime		Remaining time in the tick.
+	 * @param Iterations		Current iteration of the tick (starting at 1).
+	 * @return The remaining time step to use for the next sub-step of iteration.
+	 * @see MaxSimulationTimeStep, MaxSimulationIterations
+	 * @see ShouldUseSubStepping()
+	 */
+	float GetSimulationTimeStep(float RemainingTime, int32 Iterations) const;
+
+	/**
+	 * Determine whether or not to use substepping in the projectile motion update.
+	 * If true, GetSimulationTimeStep() will be used to time-slice the update. If false, all remaining time will be used during the tick.
+	 * @return Whether or not to use substepping in the projectile motion update.
+	 * @see GetSimulationTimeStep()
+	 */
+	virtual bool ShouldUseSubStepping() const;
+
+	/**
+	 * Max time delta for each discrete simulation step. This is used if the object is under the effect of 
+	 * Lowering this value can address issues with fast-moving objects or complex collision scenarios, at the cost of performance.
+	 *
+	 * WARNING: if (MaxSimulationTimeStep * MaxSimulationIterations) is too low for the min framerate, the last simulation step may exceed MaxSimulationTimeStep to complete the simulation.
+	 * @see MaxSimulationIterations
+	 */
+	UPROPERTY(EditAnywhere, BlueprintReadWrite, meta=(ClampMin="0.0166", ClampMax="0.50", UIMin="0.0166", UIMax="0.50"), Category=ProjectileSimulation)
+	float MaxSimulationTimeStep;
+
+	/**
+	 * Max number of iterations used for each discrete simulation step
+	 * Increasing this value can address issues with fast-moving objects or complex collision scenarios, at the cost of performance.
+	 *
+	 * WARNING: if (MaxSimulationTimeStep * MaxSimulationIterations) is too low for the min framerate, the last simulation step may exceed MaxSimulationTimeStep to complete the simulation.
+	 */
+	UPROPERTY(EditAnywhere, BlueprintReadWrite, AdvancedDisplay, meta=(ClampMin="1", ClampMax="25", UIMin="1", UIMax="25"), Category=ProjectileSimulation)
+	int32 MaxSimulationIterations;
+
 protected:
+
 	/** @return true if the wall hit has been fully processed, and no further movement is needed */
 	bool HandleHitWall(const FHitResult& Hit, float TimeTick, const FVector& MoveDelta);
 
@@ -124,21 +180,24 @@ protected:
 	/** Computes result of a bounce and returns the new velocity. */
 	virtual FVector ComputeBounceResult(const FHitResult& Hit, float TimeSlice, const FVector& MoveDelta);
 
-	// Don't allow velocity magnitude to exceed MaxSpeed, if MaxSpeed is non-zero.
+	/** Don't allow velocity magnitude to exceed MaxSpeed, if MaxSpeed is non-zero. */
 	UFUNCTION(BlueprintCallable, Category="Game|Components|ProjectileMovement")
 	FVector LimitVelocity(FVector NewVelocity) const;
 
-	// Compute the distance we should move in the given time, at a given a velocity.
-	virtual FVector ComputeMoveDelta(const FVector& InVelocity, float DeltaTime, bool bApplyGravity = true) const;
+	/** Compute the distance we should move in the given time, at a given a velocity. */
+	virtual FVector ComputeMoveDelta(const FVector& InVelocity, float DeltaTime, bool bGravityEnabled = true) const;
 
-	// Compute the acceleration that will be applied
-	virtual FVector ComputeAcceleration(bool bApplyGravity) const;
+	/** Compute the acceleration that will be applied */
+	virtual FVector ComputeAcceleration(const FVector& InVelocity, float DeltaTime, bool bGravityEnabled) const;
 
-	// Allow the projectile to track towards its homing target.
-	virtual FVector ComputeHoming() const;
+	/** Allow the projectile to track towards its homing target. */
+	virtual FVector ComputeHomingAcceleration(const FVector& InVelocity, float DeltaTime, bool bGravityEnabled) const;
 
-	// Compute gravity effect given current physics volume, projectile gravity scale, etc.
+	/** Compute gravity effect given current physics volume, projectile gravity scale, etc. */
 	float GetEffectiveGravityZ() const;
+
+	/** Minimum delta time considered when ticking. Delta times below this are not considered. This is a very small non-zero positive value to avoid potential divide-by-zero in simulation code. */
+	static const float MIN_TICK_TIME;
 };
 
 
diff --git a/Engine/Source/Runtime/Engine/Private/ProjectileMovementComponent.cpp b/Engine/Source/Runtime/Engine/Private/ProjectileMovementComponent.cpp
--- a/Engine/Source/Runtime/Engine/Private/ProjectileMovementComponent.cpp
+++ b/Engine/Source/Runtime/Engine/Private/ProjectileMovementComponent.cpp
@@ -5,11 +5,14 @@
 
 DEFINE_LOG_CATEGORY_STATIC(LogProjectileMovement, Log, All);
 
+const float UProjectileMovementComponent::MIN_TICK_TIME = 0.0002f;
+
 UProjectileMovementComponent::UProjectileMovementComponent(const class FPostConstructInitializeProperties& PCIP)
 	: Super(PCIP)
 {
 	bUpdateOnlyIfRendered = false;
 	bInitialVelocityInLocalSpace = true;
+	bForceSubStepping = false;
 
 	Velocity = FVector(1.f,0.f,0.f);
 
@@ -22,6 +25,9 @@ UProjectileMovementComponent::UProjectileMovementComponent(const class FPostCons
 	HomingAccelerationMagnitude = 0.f;
 
 	bWantsInitializeComponent = true;
+
+	MaxSimulationTimeStep = 0.05f;
+	MaxSimulationIterations = 8;
 }
 
 void UProjectileMovementComponent::Serialize( FArchive& Ar)
@@ -101,15 +107,15 @@ void UProjectileMovementComponent::TickComponent(float DeltaTime, enum ELevelTic
 	FVector OldHitNormal(0.f,0.f,1.f);
 	bool bSliding = false;
 
-	while( RemainingTime > 0.f && (Iterations < 8) && !ActorOwner->IsPendingKill() && UpdatedComponent )
+	while( RemainingTime >= MIN_TICK_TIME && (Iterations < MaxSimulationIterations) && !ActorOwner->IsPendingKill() && UpdatedComponent )
 	{
 		Iterations++;
 
-		// subdivide long ticks if falling to more closely follow parabolic trajectory
-		float TimeTick = (!ShouldApplyGravity() || RemainingTime < 0.05f) ? RemainingTime : FMath::Min(0.05f, RemainingTime*0.5f);
+		// subdivide long ticks to more closely follow parabolic trajectory
+		const float TimeTick = ShouldUseSubStepping() ? GetSimulationTimeStep(RemainingTime, Iterations) : RemainingTime;
 		RemainingTime -= TimeTick;
-		Hit.Time = 1.f;
 
+		Hit.Time = 1.f;
 		const FVector OldVelocity = Velocity;
 		FVector MoveDelta = ComputeMoveDelta(Velocity, TimeTick, !bSliding);
 
@@ -204,11 +210,11 @@ void UProjectileMovementComponent::SetVelocityInLocalSpace(FVector NewVelocity)
 }
 
 
-FVector UProjectileMovementComponent::CalculateVelocity(FVector OldVelocity, float DeltaTime, bool bApplyGravity)
+FVector UProjectileMovementComponent::CalculateVelocity(FVector OldVelocity, float DeltaTime, bool bGravityEnabled)
 {
 	FVector NewVelocity = OldVelocity;
 
-	const FVector Acceleration = ComputeAcceleration(bApplyGravity);
+	const FVector Acceleration = ComputeAcceleration(OldVelocity, DeltaTime, bGravityEnabled);
 	NewVelocity += Acceleration * DeltaTime;
 
 	return LimitVelocity(NewVelocity);
@@ -217,45 +223,58 @@ FVector UProjectileMovementComponent::CalculateVelocity(FVector OldVelocity, flo
 
 FVector UProjectileMovementComponent::LimitVelocity(FVector NewVelocity) const
 {
-	if (GetMaxSpeed() > 0.f)
+	const float CurrentMaxSpeed = GetMaxSpeed();
+	if (CurrentMaxSpeed > 0.f)
 	{
-		NewVelocity = NewVelocity.ClampMaxSize(GetMaxSpeed());
+		NewVelocity = NewVelocity.ClampMaxSize(CurrentMaxSpeed);
 	}
 
 	return NewVelocity;
 }
 
-FVector UProjectileMovementComponent::ComputeMoveDelta(const FVector& InVelocity, float DeltaTime, bool bApplyGravity) const
+FVector UProjectileMovementComponent::ComputeMoveDelta(const FVector& InVelocity, float DeltaTime, bool bGravityEnabled) const
 {
 	// p = p0 + v*t
 	FVector Delta = InVelocity * DeltaTime;
+	
+	const FVector Acceleration = ComputeAcceleration(InVelocity, DeltaTime, bGravityEnabled);
+	if (!Acceleration.IsZero())
+	{
+		// p = p0 + v*t (above) + 1/2*a*t^2 (below)
+		Delta += 0.5f * Acceleration * FMath::Square(DeltaTime);
 
-	// z = z0 + v*t (above) + 1/2*a*t^2
-	const FVector Acceleration = ComputeAcceleration(bApplyGravity);
-	Delta += 0.5f * Acceleration * FMath::Square(DeltaTime);
+		// limit velocity, else acceleration will push this result over the allowed velocity constraint during this timestep.
+		const FVector EffectiveVelocity = Delta / DeltaTime;
+		const FVector ClampedVelocity = LimitVelocity(EffectiveVelocity);
+		if (ClampedVelocity != EffectiveVelocity)
+		{
+			// Maintain direction but change magnitude
+			Delta = Delta.SafeNormal() * (GetMaxSpeed() * DeltaTime);
+		}
+	}
 
 	return Delta;
 }
 
-FVector UProjectileMovementComponent::ComputeAcceleration(bool bApplyGravity) const
+FVector UProjectileMovementComponent::ComputeAcceleration(const FVector& InVelocity, float DeltaTime, bool bGravityEnabled) const
 {
 	FVector Acceleration(FVector::ZeroVector);
 
-	if (bApplyGravity)
+	if (bGravityEnabled)
 	{
 		Acceleration.Z += GetEffectiveGravityZ();
 	}
 
 	if (bIsHomingProjectile && HomingTargetComponent.IsValid())
 	{
-		Acceleration += ComputeHoming();
+		Acceleration += ComputeHomingAcceleration(InVelocity, DeltaTime, bGravityEnabled);
 	}
 
 	return Acceleration;
 }
 
 // Allow the projectile to track towards its homing target.
-FVector UProjectileMovementComponent::ComputeHoming() const
+FVector UProjectileMovementComponent::ComputeHomingAcceleration(const FVector& InVelocity, float DeltaTime, bool bGravityEnabled) const
 {
 	FVector HomingAcceleration = ((HomingTargetComponent->GetComponentLocation() - UpdatedComponent->GetComponentLocation()).SafeNormal() * HomingAccelerationMagnitude);
 	return HomingAcceleration;
@@ -264,7 +283,7 @@ FVector UProjectileMovementComponent::ComputeHoming() const
 float UProjectileMovementComponent::GetEffectiveGravityZ() const
 {
 	// TODO: apply buoyancy if in water
-	return GetGravityZ() * ProjectileGravityScale;
+	return ShouldApplyGravity() ? GetGravityZ() * ProjectileGravityScale : 0.f;
 }
 
 
@@ -283,7 +302,7 @@ bool UProjectileMovementComponent::HandleHitWall(const FHitResult& Hit, float Ti
 	{
 		return true;
 	}
-
+	
 	HandleImpact(Hit, TimeTick, MoveDelta);
 	
 	if( ActorOwner->IsPendingKill() || !UpdatedComponent )
@@ -389,3 +408,37 @@ bool UProjectileMovementComponent::CheckStillInWorld()
 	}
 	return true;
 }
+
+
+bool UProjectileMovementComponent::ShouldUseSubStepping() const
+{
+	return bForceSubStepping || (GetEffectiveGravityZ() != 0.f) || (bIsHomingProjectile && HomingTargetComponent.IsValid());
+}
+
+
+float UProjectileMovementComponent::GetSimulationTimeStep(float RemainingTime, int32 Iterations) const
+{
+	if (RemainingTime > MaxSimulationTimeStep)
+	{
+		if (Iterations < MaxSimulationIterations)
+		{
+			// Subdivide moves to be no longer than MaxSimulationTimeStep seconds
+			RemainingTime = FMath::Min(MaxSimulationTimeStep, RemainingTime * 0.5f);
+		}
+		else
+		{
+			// If this is the last iteration, just use all the remaining time. This is usually better than cutting things short, as the simulation won't move far enough otherwise.
+			// Print a throttled warning.
+#if !(UE_BUILD_SHIPPING || UE_BUILD_TEST)
+			static uint32 s_WarningCount = 0;
+			if ((s_WarningCount++ < 100) || (GFrameCounter & 15) == 0)
+			{
+				UE_LOG(LogProjectileMovement, Warning, TEXT("GetSimulationTimeStep() - Max iterations %d hit while remaining time %.6f > MaxSimulationTimeStep (%.3f) for '%s'"), MaxSimulationIterations, RemainingTime, MaxSimulationTimeStep, *GetPathNameSafe(UpdatedComponent));
+			}
+#endif
+		}
+	}
+
+	// no less than MIN_TICK_TIME (to avoid potential divide-by-zero during simulation).
+	return FMath::Max(MIN_TICK_TIME, RemainingTime);
+}
