#ue4 - Zero extent traces should still use PullBackHit to avoid penetration.

- Prevents bouncing projectiles with zero extents from getting stuck in meshes.

UDN: https://udn.unrealengine.com/questions/219863/zero-extend-capsule-movement.html

diff --git a/Engine/Source/Runtime/Engine/Private/PrimitiveComponent.cpp b/Engine/Source/Runtime/Engine/Private/PrimitiveComponent.cpp
--- a/Engine/Source/Runtime/Engine/Private/PrimitiveComponent.cpp
+++ b/Engine/Source/Runtime/Engine/Private/PrimitiveComponent.cpp
@@ -1318,7 +1318,7 @@ bool UPrimitiveComponent::MoveComponent( const FVector& Delta, const FRotator& N
 			InitSweepCollisionParams(Params, ResponseParam);
 			bool const bHadBlockingHit = GetWorld()->ComponentSweepMulti(Hits, this, TraceStart, TraceEnd, GetComponentRotation(), Params);
 
-			if (Hits.Num() > 0 && !IsZeroExtent())
+			if (Hits.Num() > 0)
 			{
 				const float DeltaSize = FMath::Sqrt(DeltaSizeSq);
 				for(int32 HitIdx=0; HitIdx<Hits.Num(); HitIdx++)
