#ue4 - Fix networking actor replication on clients avoiding changing location/rotation if encroaching at destination.

diff --git a/Engine/Source/Runtime/Engine/Private/ActorReplication.cpp b/Engine/Source/Runtime/Engine/Private/ActorReplication.cpp
--- a/Engine/Source/Runtime/Engine/Private/ActorReplication.cpp
+++ b/Engine/Source/Runtime/Engine/Private/ActorReplication.cpp
@@ -101,8 +101,8 @@ void AActor::PostNetReceiveLocationAndRotation()
 {
 	if( RootComponent && RootComponent->IsRegistered() && (ReplicatedMovement.Location != GetActorLocation() || ReplicatedMovement.Rotation != GetActorRotation()) )
 	{
-		TeleportTo(ReplicatedMovement.Location, ReplicatedMovement.Rotation, false);
-		//SetActorLocationAndRotation(ReplicatedMovement.Location, ReplicatedMovement.Rotation); <-- preferred, but awaiting answer to question about UpdateNavOctree() missing in SceneComponent::MoveComponent
+		TeleportTo(ReplicatedMovement.Location, ReplicatedMovement.Rotation, false, true);
+		//SetActorLocationAndRotation(ReplicatedMovement.Location, ReplicatedMovement.Rotation); <-- preferred for performance, but existing overrides of TeleportTo() will be affected.
 	}
 }
 
