Engine: Expose opt-out for hardware survey thru project Engine..General Settings details panel

diff --git a/Engine/Config/BaseEngine.ini b/Engine/Config/BaseEngine.ini
--- a/Engine/Config/BaseEngine.ini
+++ b/Engine/Config/BaseEngine.ini
@@ -2171,11 +2171,6 @@ EnableNiagara=false
 +CollisionChannelRedirects=(OldName="VehicleMovement",NewName="Vehicle")
 +CollisionChannelRedirects=(OldName="PawnMovement",NewName="Pawn")
 
-[Engine.HardwareSurvey]
-; if we should send home hardware survey in non debug builds
-bEnableHardwareSurvey=true
-bHardwareSurveyDone=false
-
 [Engine.BufferVisualizationMaterials]
 BaseColor=(Material="/Engine/BufferVisualization/BaseColor.BaseColor", Name=LOCTEXT("BaseColorMat", "Base Color"))
 CustomDepth=(Material="/Engine/BufferVisualization/CustomDepth.CustomDepth", Name=LOCTEXT("BaseCustomDepthMat", "Custom Depth"))
diff --git a/Engine/Source/Runtime/Engine/Classes/Engine/Engine.h b/Engine/Source/Runtime/Engine/Classes/Engine/Engine.h
--- a/Engine/Source/Runtime/Engine/Classes/Engine/Engine.h
+++ b/Engine/Source/Runtime/Engine/Classes/Engine/Engine.h
@@ -1069,6 +1069,14 @@ public:
 	UPROPERTY(EditAnywhere, config, Category=Settings)
 	float NearClipPlane;
 
+	/** Should the engine analytics provider be initialized on startup?  Note: The engine/editor must be restarted for changes to take effect. */
+	UPROPERTY(EditAnywhere, config, Category=Settings, AdvancedDisplay)
+	uint32 bEngineAnalyticsProviderEnabled:1;
+
+	/** Can the game/application report anonymous hardware survey statistics back to Epic?  (such as display resolution and GPU model) */
+	UPROPERTY(EditAnywhere, config, Category=Settings, AdvancedDisplay, meta=(EditCondition="bEngineAnalyticsProviderEnabled"))
+	uint32 bHardwareSurveyEnabled:1;
+
 	/** Flag for completely disabling subtitles for localized sounds. */
 	UPROPERTY(EditAnywhere, config, Category=Subtitles)
 	uint32 bSubtitlesEnabled:1;
diff --git a/Engine/Source/Runtime/Engine/Private/EngineAnalytics.cpp b/Engine/Source/Runtime/Engine/Private/EngineAnalytics.cpp
--- a/Engine/Source/Runtime/Engine/Private/EngineAnalytics.cpp
+++ b/Engine/Source/Runtime/Engine/Private/EngineAnalytics.cpp
@@ -119,7 +119,9 @@ void FEngineAnalytics::Initialize()
 	checkf(!bIsInitialized, TEXT("FEngineAnalytics::Initialize called more than once."));
 
 	// Never use analytics when running a commandlet tool
-	const bool bShouldInitAnalytics = !IsRunningCommandlet();
+	check(GEngine);
+	const bool bShouldInitAnalytics = !IsRunningCommandlet() && GEngine->bEngineAnalyticsProviderEnabled;
+
 	if( bShouldInitAnalytics )
 	{
 #if DEBUG_ANALYTICS
diff --git a/Engine/Source/Runtime/Engine/Private/GameEngine.cpp b/Engine/Source/Runtime/Engine/Private/GameEngine.cpp
--- a/Engine/Source/Runtime/Engine/Private/GameEngine.cpp
+++ b/Engine/Source/Runtime/Engine/Private/GameEngine.cpp
@@ -399,6 +399,8 @@ UEngine::UEngine(const class FPostConstructInitializeProperties& PCIP)
 	SelectionHighlightIntensityBillboards = 0.25f;
 
 	bUseSound = true;
+	bEngineAnalyticsProviderEnabled = true;
+	bHardwareSurveyEnabled = true;
 	bPendingHardwareSurveyResults = false;
 	bIsInitialized = false;
 
diff --git a/Engine/Source/Runtime/Engine/Private/UnrealEngine.cpp b/Engine/Source/Runtime/Engine/Private/UnrealEngine.cpp
--- a/Engine/Source/Runtime/Engine/Private/UnrealEngine.cpp
+++ b/Engine/Source/Runtime/Engine/Private/UnrealEngine.cpp
@@ -5388,24 +5388,21 @@ void UEngine::OnLostFocusPause(bool EnablePause)
 
 void UEngine::InitHardwareSurvey()
 {
-	if (GConfig)
-	{
-		bool bEnabled = false;
+	bool bEnabled = bHardwareSurveyEnabled;
 
-		// The hardware survey costs time and we don't want to slow down debug builds.
-		// This is mostly because of the CPU benchmark running in the survey and the results in debug are not being valid.
-#if UE_BUILD_DEBUG == 0
-		GConfig->GetBool(TEXT("Engine.HardwareSurvey"), TEXT("bEnableHardwareSurvey"), bEnabled, GEngineIni);
+	// The hardware survey costs time and we don't want to slow down debug builds.
+	// This is mostly because of the CPU benchmark running in the survey and the results in debug are not being valid.
+#if UE_BUILD_DEBUG
+	bEnabled = false;
 #endif
 
-		if (bEnabled)
+	if (bEnabled)
+	{
+		if (IsHardwareSurveyRequired())
 		{
-			if (IsHardwareSurveyRequired())
-			{
-				bPendingHardwareSurveyResults = true;
-			}
-		}	
-	}
+			bPendingHardwareSurveyResults = true;
+		}
+	}	
 }
 
 void UEngine::TickHardwareSurvey()
