Adjust comment which looks like a markdown link, and causes exception in UnrealDocTool.

diff --git a/Engine/Source/Runtime/Engine/Classes/Engine/NetSerialization.h b/Engine/Source/Runtime/Engine/Classes/Engine/NetSerialization.h
--- a/Engine/Source/Runtime/Engine/Classes/Engine/NetSerialization.h
+++ b/Engine/Source/Runtime/Engine/Classes/Engine/NetSerialization.h
@@ -1035,7 +1035,7 @@ struct TStructOpsTypeTraits< FVector_NetQuantize100 > : public TStructOpsTypeTra
  *	FVector_NetQuantizeNormal
  *
  *	16 bits per component
- *	Valid range: [-1..+1] (inclusive)
+ *	Valid range: -1..+1 inclusive
  */
 USTRUCT()
 struct FVector_NetQuantizeNormal : public FVector
