primitive component's bCanEverAffectNavigation flag can be now changed in instances as well

diff --git a/Engine/Source/Runtime/Engine/Classes/Components/PrimitiveComponent.h b/Engine/Source/Runtime/Engine/Classes/Components/PrimitiveComponent.h
--- a/Engine/Source/Runtime/Engine/Classes/Components/PrimitiveComponent.h
+++ b/Engine/Source/Runtime/Engine/Classes/Components/PrimitiveComponent.h
@@ -387,7 +387,7 @@ public:
 	FBodyInstance BodyInstance;
 
 	/** can this component potentially influence navigation */
-	UPROPERTY(EditDefaultsOnly, Category=Collision, AdvancedDisplay)
+	UPROPERTY(EditAnywhere, Category=Collision, AdvancedDisplay)
 	uint32 bCanEverAffectNavigation:1;
 
 	/** cached navigation relevancy flag for collision updates */
