Properly propogate bIgnoreForFlushing from AudioComponent to ActiveSound

diff --git a/Engine/Source/Runtime/Engine/Private/AudioComponent.cpp b/Engine/Source/Runtime/Engine/Private/AudioComponent.cpp
--- a/Engine/Source/Runtime/Engine/Private/AudioComponent.cpp
+++ b/Engine/Source/Runtime/Engine/Private/AudioComponent.cpp
@@ -161,6 +161,7 @@ void UAudioComponent::PlayInternal(const float StartTime, const float FadeInDura
 
 			NewActiveSound.bShouldRemainActiveIfDropped = bShouldRemainActiveIfDropped;
 			NewActiveSound.bHandleSubtitles = (!bSuppressSubtitles || OnQueueSubtitles.IsBound());
+			NewActiveSound.bIgnoreForFlushing = bIgnoreForFlushing;
 
 			NewActiveSound.bIsUISound = bIsUISound;
 			NewActiveSound.bIsMusic = bIsMusic;
