Fix UnrealDocTool error caused by link to an anchor which doesn't have a link description.

diff --git a/Engine/Source/Runtime/Engine/Classes/Components/SceneComponent.h b/Engine/Source/Runtime/Engine/Classes/Components/SceneComponent.h
--- a/Engine/Source/Runtime/Engine/Classes/Components/SceneComponent.h
+++ b/Engine/Source/Runtime/Engine/Classes/Components/SceneComponent.h
@@ -74,7 +74,7 @@ DECLARE_DYNAMIC_MULTICAST_DELEGATE_OneParam(FPhysicsVolumeChanged, class APhysic
  * A SceneComponent has a transform and supports attachment, but has no rendering or collision capabilities.
  * Useful as a 'dummy' component in the hierarchy to offset others.
  * @see https://docs.unrealengine.com/latest/INT/API/Runtime/Engine/Components/USceneComponent/index.html
- * @see https://docs.unrealengine.com/latest/INT/Programming/UnrealArchitecture/Actors/Components/index.html#scenecomponents
+ * @see [Scene Components](https://docs.unrealengine.com/latest/INT/Programming/UnrealArchitecture/Actors/Components/index.html#scenecomponents)
  */
 UCLASS(ClassGroup=Utility, BlueprintType, HideCategories=(Trigger, PhysicsVolume), meta=(BlueprintSpawnableComponent))
 class ENGINE_API USceneComponent : public UActorComponent
