Fix error caused by UnrealDocTool trying to parse a variable name

diff --git a/Engine/Source/Runtime/Engine/Classes/Materials/MaterialExpressionSphereMask.h b/Engine/Source/Runtime/Engine/Classes/Materials/MaterialExpressionSphereMask.h
--- a/Engine/Source/Runtime/Engine/Classes/Materials/MaterialExpressionSphereMask.h
+++ b/Engine/Source/Runtime/Engine/Classes/Materials/MaterialExpressionSphereMask.h
@@ -22,7 +22,7 @@ class UMaterialExpressionSphereMask : public UMaterialExpression
 	UPROPERTY(meta = (RequiredInput = "false", ToolTip = "Defaults to 'AttenuationRadius' if not specified"))
 	FExpressionInput Radius;
 
-	/** 0..1 for the range of 0% to 100%, if not hooked up the internal constant is used */
+	/** 0..1 for the range of 0\% to 100\%, if not hooked up the internal constant is used */
 	UPROPERTY(meta = (RequiredInput = "false", ToolTip = "Defaults to 'HardnessPercent' if not specified"))
 	FExpressionInput Hardness;
 
