Rearrange comment which UnrealDocTool tries (and fails) to parse as a markdown link.

diff --git a/Engine/Source/Runtime/Engine/Private/SpeedTreeWind.cpp b/Engine/Source/Runtime/Engine/Private/SpeedTreeWind.cpp
--- a/Engine/Source/Runtime/Engine/Private/SpeedTreeWind.cpp
+++ b/Engine/Source/Runtime/Engine/Private/SpeedTreeWind.cpp
@@ -595,7 +595,7 @@ FORCEINLINE float FSpeedTreeWind::RandomFloat(float fMin, float fMax) const
 
 ///////////////////////////////////////////////////////////////////////
 // FSpeedTreeWind::LinearSigmoid
-/// Converts an input value in the range [0.0, 1.0] (\a fInput) to an s-curve.  The parameter \a fLinearness flattens out the s-curve where 0.0 = s-curve and 1.0 = linear.
+/// Converts an input value (\a fInput) in the range [0.0, 1.0] to an s-curve.  The parameter \a fLinearness flattens out the s-curve where 0.0 = s-curve and 1.0 = linear.
 
 FORCEINLINE float FSpeedTreeWind::LinearSigmoid(float fInput, float fLinearness)
 {
