Reformat comment at the top of ActorChannel to appease UnrealDocTool

diff --git a/Engine/Source/Runtime/Engine/Classes/Engine/ActorChannel.h b/Engine/Source/Runtime/Engine/Classes/Engine/ActorChannel.h
--- a/Engine/Source/Runtime/Engine/Classes/Engine/ActorChannel.h
+++ b/Engine/Source/Runtime/Engine/Classes/Engine/ActorChannel.h
@@ -8,26 +8,27 @@
  *
  *	An ActorChannel bunch looks like this:
  *
- *		|----------------------|			
- *		| SpawnInfo		       |  (Spawn Info) Initial bunch only
- *		|  -Actor Class        |	-Created by ActorChannel	
- *		|  -Spawn Loc/Rot      | 
- *      | NetGUID assigns      |
- *		|  -Actor NetGUID      |
- *		|  -Component NetGUIDs |
- *		|----------------------|
- *		|                |
- *		|----------------|
- *		| NetGUID ObjRef |	(Content chunks) x number of replicating objects (Actor + any components)
- * 		|----------------|		-Each chunk created by its own FObjectReplicator instance.
- *      |                |		
- *		| Properties...  |
- *		|                |	
- *		| RPCs...        |
- *      |                |
- *      |----------------|
- *		| </End Tag>     |
- *		|----------------|
+ *		|----------------------|---------------------------------------------------------------------------|
+ *		| SpawnInfo		       | (Spawn Info) Initial bunch only                                           |
+ *		|  -Actor Class        |	-Created by ActorChannel	                                           |
+ *		|  -Spawn Loc/Rot      |                                                                           |
+ *      | NetGUID assigns      |                                                                           |
+ *		|  -Actor NetGUID      |                                                                           |
+ *		|  -Component NetGUIDs |                                                                           |
+ *		|----------------------|---------------------------------------------------------------------------|
+ *		|                      |                                                                           |
+ *		|----------------------|---------------------------------------------------------------------------|
+ *		| NetGUID ObjRef       | (Content chunks) x number of replicating objects (Actor + any components) |
+ * 		|                      |		-Each chunk created by its own FObjectReplicator instance.         |
+ * 		|----------------------|---------------------------------------------------------------------------|
+ *      |                      |		                                                                   |
+ *		| Properties...        |                                                                           |
+ *		|                      |	                                                                       |
+ *		| RPCs...              |                                                                           |
+ *      |                      |                                                                           |
+ *      |----------------------|---------------------------------------------------------------------------|
+ *		| </End Tag>           |                                                                           |
+ *		|----------------------|---------------------------------------------------------------------------|
  * 
  *
  */
