Fix case where deleting collision in static mesh editor cauess cooked physics data to get out of sync.
Add editor code to fix cooked data if out of sync

TTP#345845

diff --git a/Engine/Source/Editor/StaticMeshEditor/Private/StaticMeshEditor.cpp b/Engine/Source/Editor/StaticMeshEditor/Private/StaticMeshEditor.cpp
--- a/Engine/Source/Editor/StaticMeshEditor/Private/StaticMeshEditor.cpp
+++ b/Engine/Source/Editor/StaticMeshEditor/Private/StaticMeshEditor.cpp
@@ -1579,6 +1579,9 @@ void FStaticMeshEditor::DeleteSelectedPrims()
 			// Make sure rendering is done - so we are not changing data being used by collision drawing.
 			FlushRenderingCommands();
 
+			// Make sure to invalidate cooked data
+			StaticMesh->BodySetup->InvalidatePhysicsData();
+
 			// refresh collision change back to staticmesh components
 			RefreshCollisionChange(StaticMesh);
 
diff --git a/Engine/Source/Runtime/Engine/Private/PhysicsEngine/BodySetup.cpp b/Engine/Source/Runtime/Engine/Private/PhysicsEngine/BodySetup.cpp
--- a/Engine/Source/Runtime/Engine/Private/PhysicsEngine/BodySetup.cpp
+++ b/Engine/Source/Runtime/Engine/Private/PhysicsEngine/BodySetup.cpp
@@ -113,15 +113,30 @@ void UBodySetup::CreatePhysicsMeshes()
 			return;
 		}
 
-		ClearPhysicsMeshes(); // Make sure all are cleared before we start
+		FPhysXFormatDataReader CookedDataReader(*FormatData);
 
-		// Create physics objects
-		FPhysXFormatDataReader CookedDataReader( *FormatData );
-		check( !bGenerateNonMirroredCollision || CookedDataReader.ConvexMeshes.Num() == 0 || CookedDataReader.ConvexMeshes.Num() == AggGeom.ConvexElems.Num() );
-		check( !bGenerateMirroredCollision || CookedDataReader.ConvexMeshesNegX.Num() == 0 || CookedDataReader.ConvexMeshesNegX.Num() == AggGeom.ConvexElems.Num() );
+		if (CollisionTraceFlag != CTF_UseComplexAsSimple)
+		{
+			bool bNeedsCooking = bGenerateNonMirroredCollision && CookedDataReader.ConvexMeshes.Num() != AggGeom.ConvexElems.Num();
+			bNeedsCooking = bNeedsCooking || (bGenerateMirroredCollision && CookedDataReader.ConvexMeshesNegX.Num() != AggGeom.ConvexElems.Num());
+			if (bNeedsCooking)	//Because of bugs it's possible to save with out of sync cooked data. In editor we want to fixup this data
+			{
+#if WITH_EDITOR
+				InvalidatePhysicsData();
+				CreatePhysicsMeshes();
+				return;
+#endif
+			}
+		}
+		
+		ClearPhysicsMeshes();
 
 		if (CollisionTraceFlag != CTF_UseComplexAsSimple)
 		{
+
+			ensure(!bGenerateNonMirroredCollision || CookedDataReader.ConvexMeshes.Num() == 0 || CookedDataReader.ConvexMeshes.Num() == AggGeom.ConvexElems.Num());
+			ensure(!bGenerateMirroredCollision || CookedDataReader.ConvexMeshesNegX.Num() == 0 || CookedDataReader.ConvexMeshesNegX.Num() == AggGeom.ConvexElems.Num());
+
 			//If the cooked data no longer has convex meshes, make sure to empty AggGeom.ConvexElems - otherwise we leave NULLS which cause issues, and we also read past the end of CookedDataReader.ConvexMeshes
 			if ((bGenerateNonMirroredCollision && CookedDataReader.ConvexMeshes.Num() == 0) || (bGenerateMirroredCollision && CookedDataReader.ConvexMeshesNegX.Num() == 0))
 			{
