Stop using null to wipe all sounds when doing an AudioDevice flush and allow each world to manage its own state.
Allow bIgnoreForFlushing to work in PIE

diff --git a/Engine/Source/Editor/UnrealEd/Private/EditorServer.cpp b/Engine/Source/Editor/UnrealEd/Private/EditorServer.cpp
--- a/Engine/Source/Editor/UnrealEd/Private/EditorServer.cpp
+++ b/Engine/Source/Editor/UnrealEd/Private/EditorServer.cpp
@@ -1647,7 +1647,7 @@ void UEditorEngine::EditorDestroyWorld( FWorldContext & Context, const FText& Cl
 	// Stop all audio and remove references 
 	if ( GetAudioDevice() )
 	{
-		GetAudioDevice()->Flush(NULL);
+		GetAudioDevice()->Flush(ContextWorld);
 	}
 
 	// Reset the editor transform to avoid loading the new world with an offset if loading a sublevel
diff --git a/Engine/Source/Runtime/Engine/Private/AudioDevice.cpp b/Engine/Source/Runtime/Engine/Private/AudioDevice.cpp
--- a/Engine/Source/Runtime/Engine/Private/AudioDevice.cpp
+++ b/Engine/Source/Runtime/Engine/Private/AudioDevice.cpp
@@ -2139,17 +2139,17 @@ void FAudioDevice::Flush( UWorld* WorldToFlush, bool bClearActivatedReverb )
 {
 	// Stop all audio components attached to the scene
 	bool bFoundIgnoredComponent = false;
-	for( int32 Index = ActiveSounds.Num() - 1; Index >= 0; Index-- )
+	for( int32 Index = ActiveSounds.Num() - 1; Index >= 0; --Index )
 	{
 		FActiveSound* ActiveSound = ActiveSounds[Index];
 		// if we are in the editor we want to always flush the ActiveSounds
-		if( ActiveSound->bIgnoreForFlushing && !GIsEditor )
+		if( ActiveSound->bIgnoreForFlushing )
 		{
 			bFoundIgnoredComponent = true;
 		}
 		else
 		{
-			if( WorldToFlush == NULL || ActiveSound->World == NULL || ActiveSound->World == WorldToFlush )
+			if( WorldToFlush == nullptr || ActiveSound->World == nullptr || ActiveSound->World == WorldToFlush )
 			{
 				ActiveSound->Stop(this);
 			}
@@ -2164,10 +2164,10 @@ void FAudioDevice::Flush( UWorld* WorldToFlush, bool bClearActivatedReverb )
 	if (bClearActivatedReverb)
 	{
 		ActivatedReverbs.Empty();
-		HighestPriorityReverb = NULL;
+		HighestPriorityReverb = nullptr;
 	}
 
-	if( WorldToFlush == NULL )
+	if( WorldToFlush == nullptr )
 	{
 		// Make sure sounds are fully stopped.
 		if( bFoundIgnoredComponent )
diff --git a/Engine/Source/Runtime/Engine/Private/UnrealEngine.cpp b/Engine/Source/Runtime/Engine/Private/UnrealEngine.cpp
--- a/Engine/Source/Runtime/Engine/Private/UnrealEngine.cpp
+++ b/Engine/Source/Runtime/Engine/Private/UnrealEngine.cpp
@@ -8401,7 +8401,7 @@ bool UEngine::LoadMap( FWorldContext& WorldContext, FURL URL, class UPendingNetG
 	// Stop all audio to remove references to current level.
 	if( GEngine && GEngine->GetAudioDevice() )
 	{
-		GEngine->GetAudioDevice()->Flush( NULL );
+		GEngine->GetAudioDevice()->Flush( WorldContext.World() );
 		// reset transient volume
 		GEngine->GetAudioDevice()->TransientMasterVolume = 1.0;
 	}
diff --git a/Engine/Source/Runtime/Engine/Private/World.cpp b/Engine/Source/Runtime/Engine/Private/World.cpp
--- a/Engine/Source/Runtime/Engine/Private/World.cpp
+++ b/Engine/Source/Runtime/Engine/Private/World.cpp
@@ -4429,7 +4429,7 @@ UWorld* FSeamlessTravelHandler::Tick()
 			// Stop all audio to remove references to old world
 			if (GEngine != NULL && GEngine->GetAudioDevice() != NULL)
 			{
-				GEngine->GetAudioDevice()->Flush(NULL);
+				GEngine->GetAudioDevice()->Flush(CurrentWorld);
 			}
 
 			// Copy the standby cheat status
