Make UAnimInstance::GetWorld() not crash if called on the CDO

diff --git a/Engine/Source/Runtime/Engine/Private/Animation/AnimInstance.cpp b/Engine/Source/Runtime/Engine/Private/Animation/AnimInstance.cpp
--- a/Engine/Source/Runtime/Engine/Private/Animation/AnimInstance.cpp
+++ b/Engine/Source/Runtime/Engine/Private/Animation/AnimInstance.cpp
@@ -214,7 +214,8 @@ USkeletalMeshComponent* UAnimInstance::GetOwningComponent() const
 
 UWorld* UAnimInstance::GetWorld() const
 {
-	return GetSkelMeshComponent()->GetWorld();
+	// The CDO isn't owned by a SkelMeshComponent (and doesn't have a World)
+	return (HasAnyFlags(RF_ClassDefaultObject) ? nullptr : GetSkelMeshComponent()->GetWorld());
 }
 
 void UAnimInstance::InitializeAnimation()
