#ue4 - AActor::PostNetReceiveLocationAndRotation() now uses SetActorLocationAndRotation() rather than TeleportTo() for receiving network position/rotation updates.

diff --git a/Engine/Source/Runtime/Engine/Private/ActorReplication.cpp b/Engine/Source/Runtime/Engine/Private/ActorReplication.cpp
--- a/Engine/Source/Runtime/Engine/Private/ActorReplication.cpp
+++ b/Engine/Source/Runtime/Engine/Private/ActorReplication.cpp
@@ -101,8 +101,7 @@ void AActor::PostNetReceiveLocationAndRotation()
 {
 	if( RootComponent && RootComponent->IsRegistered() && (ReplicatedMovement.Location != GetActorLocation() || ReplicatedMovement.Rotation != GetActorRotation()) )
 	{
-		TeleportTo(ReplicatedMovement.Location, ReplicatedMovement.Rotation, false, true);
-		//SetActorLocationAndRotation(ReplicatedMovement.Location, ReplicatedMovement.Rotation); <-- preferred for performance, but existing overrides of TeleportTo() will be affected.
+		SetActorLocationAndRotation(ReplicatedMovement.Location, ReplicatedMovement.Rotation, /*bSweep=*/ false);
 	}
 }
 
