Updated the FStats verification

diff --git a/Engine/Source/Runtime/Engine/Private/Tests/AutomationTestCommon.cpp b/Engine/Source/Runtime/Engine/Private/Tests/AutomationTestCommon.cpp
--- a/Engine/Source/Runtime/Engine/Private/Tests/AutomationTestCommon.cpp
+++ b/Engine/Source/Runtime/Engine/Private/Tests/AutomationTestCommon.cpp
@@ -4,6 +4,7 @@
 #include "Matinee/MatineeActor.h"
 #include "AutomationTestCommon.h"
 
+DEFINE_LOG_CATEGORY_STATIC(LogEngineAutomationLatentCommand, Log, All);
 
 ///////////////////////////////////////////////////////////////////////
 // Common Latent commands which are used across test type. I.e. Engine, Network, etc...
@@ -34,6 +35,8 @@ bool FWaitForMatineeToCompleteLatentCommand::Update()
 
 bool FExecStringLatentCommand::Update()
 {
+	UE_LOG(LogEngineAutomationLatentCommand, Log, TEXT("Executing the console command: '%s'"), *ExecCommand);
+
 	if (GEngine->GameViewport)
 	{
 		GEngine->GameViewport->Exec( NULL, *ExecCommand, *GLog);
diff --git a/Engine/Source/Runtime/Engine/Private/Tests/EngineAutomationTests.cpp b/Engine/Source/Runtime/Engine/Private/Tests/EngineAutomationTests.cpp
--- a/Engine/Source/Runtime/Engine/Private/Tests/EngineAutomationTests.cpp
+++ b/Engine/Source/Runtime/Engine/Private/Tests/EngineAutomationTests.cpp
@@ -112,8 +112,70 @@ bool FStatsVerificationMapTest::RunTest(const FString& Parameters)
 	check(AutomationTestSettings);
 	FString MapName = AutomationTestSettings->AutomationTestmap.FilePath;
 
-	GEngine->Exec(GetSimpleEngineAutomationTestGameWorld(GetTestFlags()), *FString::Printf(TEXT("Open %s"), *MapName));
+	//Get the location of the map being used.
+	FString Filename = FPaths::ConvertRelativePathToFull(MapName);
+
+	if (!FPaths::FileExists(Filename))
+	{
+		UE_LOG(LogEngineAutomationTests, Warning, TEXT("Map doesn't exist: %s.  \nUsing currently loaded map."), *Filename);
+		
+		ADD_LATENT_AUTOMATION_COMMAND(FExecStringLatentCommand(TEXT("stat game")));
+		ADD_LATENT_AUTOMATION_COMMAND(FEngineWaitLatentCommand(1.0));
+		ADD_LATENT_AUTOMATION_COMMAND(FExecStringLatentCommand(TEXT("stat game")));
+
+		ADD_LATENT_AUTOMATION_COMMAND(FExecStringLatentCommand(TEXT("stat scenerendering")));
+		ADD_LATENT_AUTOMATION_COMMAND(FEngineWaitLatentCommand(1.0));
+		ADD_LATENT_AUTOMATION_COMMAND(FExecStringLatentCommand(TEXT("stat scenerendering")));
+
+		ADD_LATENT_AUTOMATION_COMMAND(FExecStringLatentCommand(TEXT("stat memory")));
+		ADD_LATENT_AUTOMATION_COMMAND(FEngineWaitLatentCommand(1.0));
+		ADD_LATENT_AUTOMATION_COMMAND(FExecStringLatentCommand(TEXT("stat memory")));
+
+		ADD_LATENT_AUTOMATION_COMMAND(FExecStringLatentCommand(TEXT("stat slate")));
+		ADD_LATENT_AUTOMATION_COMMAND(FEngineWaitLatentCommand(1.0));
+		ADD_LATENT_AUTOMATION_COMMAND(FExecStringLatentCommand(TEXT("stat slate")));
+
+		return true;
+	}
+	
+	//If the map is located in the engine folder then we'll make the path relative to that.  Otherwise it will be relative to the game content folder.
+	if (Filename.Contains(TEXT("Engine"), ESearchCase::IgnoreCase, ESearchDir::FromStart))
+	{
+		//This will be false if the map is on a different drive.
+		if (FPaths::MakePathRelativeTo(Filename, *FPaths::EngineContentDir()))
+		{
+				FString ShortName = FPaths::GetBaseFilename(Filename);
+				FString PathName = FPaths::GetPath(Filename);
+				MapName = FString::Printf(TEXT("/Engine/%s/%s.%s"), *PathName, *ShortName, *ShortName);
+
+				UE_LOG(LogEngineAutomationTests, Log, TEXT("Opening '%s'"), *MapName);
+
+				GEngine->Exec(GetSimpleEngineAutomationTestGameWorld(GetTestFlags()), *FString::Printf(TEXT("Open %s"), *MapName));
+		}
+		else
+		{
+			UE_LOG(LogEngineAutomationTests, Error, TEXT("Invalid asset path: %s."), *Filename);
+		}
+	}
+	else
+	{
+		//This will be false if the map is on a different drive.
+		if (FPaths::MakePathRelativeTo(Filename, *FPaths::GameContentDir()))
+		{
+			FString ShortName = FPaths::GetBaseFilename(Filename);
+			FString PathName = FPaths::GetPath(Filename);
+			MapName = FString::Printf(TEXT("/Game/%s/%s.%s"), *PathName, *ShortName, *ShortName);
 
+			UE_LOG(LogEngineAutomationTests, Log, TEXT("Opening '%s'"), *MapName);
+
+			GEngine->Exec(GetSimpleEngineAutomationTestGameWorld(GetTestFlags()), *FString::Printf(TEXT("Open %s"), *MapName));
+		}
+		else
+		{
+			UE_LOG(LogEngineAutomationTests, Error, TEXT("Invalid asset path: %s."), *Filename);
+		}
+	}
+	
 	ADD_LATENT_AUTOMATION_COMMAND(FExecStringLatentCommand(TEXT("stat game")));
 	ADD_LATENT_AUTOMATION_COMMAND(FEngineWaitLatentCommand(1.0));
 	ADD_LATENT_AUTOMATION_COMMAND(FExecStringLatentCommand(TEXT("stat game")));
