Remove code that handles out of order packets since at this level, packets are never out of order

diff --git a/Engine/Source/Runtime/CoreUObject/Public/UObject/CoreNet.h b/Engine/Source/Runtime/CoreUObject/Public/UObject/CoreNet.h
--- a/Engine/Source/Runtime/CoreUObject/Public/UObject/CoreNet.h
+++ b/Engine/Source/Runtime/CoreUObject/Public/UObject/CoreNet.h
@@ -153,7 +153,6 @@ struct FPropertyRetirement
 	TSharedPtr<class INetDeltaBaseState> DynamicState;
 
 	FPacketIdRange	OutPacketIdRange;
-	int32			InPacketId;			// Packet received on, INDEX_NONE=none.
 
 	uint32			Reliable		: 1;	// Whether it was sent reliably.
 	uint32			CustomDelta		: 1;	// True if this property uses custom delta compression
@@ -162,7 +161,6 @@ struct FPropertyRetirement
 	FPropertyRetirement()
 		:	Next ( NULL )
 		,	DynamicState ( NULL )
-		,	InPacketId	( INDEX_NONE )
 		,   Reliable( 0 )
 		,   CustomDelta( 0 )
 		,	Config( 0 )
diff --git a/Engine/Source/Runtime/Engine/Private/DataReplication.cpp b/Engine/Source/Runtime/Engine/Private/DataReplication.cpp
--- a/Engine/Source/Runtime/Engine/Private/DataReplication.cpp
+++ b/Engine/Source/Runtime/Engine/Private/DataReplication.cpp
@@ -459,28 +459,10 @@ bool FObjectReplicator::ReceivedBunch( FInBunch &Bunch, const FReplicationFlags
 				}
 			}
 #endif
-
 			if ( !Retirement[ ReplicatedProp->RepIndex ].CustomDelta )
 			{
 				// We hijack a non custom delta property to signify we are using FRepLayout to read the entire property block
-				FPropertyRetirement & Retire = Retirement[ ReplicatedProp->RepIndex ];
-
-				bool bDiscardLayout = false;
-
-				if ( Bunch.PacketId >= Retire.InPacketId ) //!! problem with reliable pkts containing dynamic references, being retransmitted, and overriding newer versions. Want "OriginalPacketId" for retransmissions?
-				{
-					// Receive this new property.
-					Retire.InPacketId = Bunch.PacketId;
-				}
-				else
-				{
-					UE_LOG( LogNet, Fatal, TEXT( "Bunch.PacketId < Retire.InPacketId. PacketId: %i, LastPacketId: %i, Partial: %s, Object: %s" ), Bunch.PacketId, Retire.InPacketId, Bunch.bPartial ? TEXT( "YES" ) : TEXT( "NO" ), *Object->GetFullName() );
-					bDiscardLayout = true;
-				}
-
-				check( Bunch.PacketId >= Retire.InPacketId );
-				
-				if ( !RepLayout->ReceiveProperties( ObjectClass, RepState, (void*)Object, Bunch, bDiscardLayout, bOutHasUnmapped ) )
+				if ( !RepLayout->ReceiveProperties( ObjectClass, RepState, (void*)Object, Bunch, bOutHasUnmapped ) )
 				{
 					UE_LOG( LogNet, Error, TEXT( "ReceiveProperties FAILED %s in %s" ), *ReplicatedProp->GetName(), *Object->GetFullName() );
 					return false;
@@ -506,32 +488,9 @@ bool FObjectReplicator::ReceivedBunch( FInBunch &Bunch, const FReplicationFlags
 				}
 
 				// Pointer to destination.
-				uint8 * DestObj		= (uint8*)Object;
-				uint8 * DestRecent	= RepState->StaticBuffer.Num() ? RepState->StaticBuffer.GetTypedData() : NULL;
+				uint8 * DestRecent = RepState->StaticBuffer.GetTypedData();
 
-				// Check property ordering.
-				FPropertyRetirement & Retire = Retirement[ ReplicatedProp->RepIndex + Element ];
-
-				if ( Bunch.PacketId >= Retire.InPacketId ) //!! problem with reliable pkts containing dynamic references, being retransmitted, and overriding newer versions. Want "OriginalPacketId" for retransmissions?
-				{
-					// Receive this new property.
-					Retire.InPacketId = Bunch.PacketId;
-				}
-				else
-				{
-					UE_LOG( LogNet, Fatal, TEXT( "Bunch.PacketId < Retire.InPacketId. PacketId: %i, LastPacketId: %i, Partial: %s, Object: %s" ), Bunch.PacketId, Retire.InPacketId, Bunch.bPartial ? TEXT( "YES" ) : TEXT( "NO" ), *Object->GetFullName() );
-
-					// Skip this property, because it's out-of-date.
-					UE_LOG( LogNetTraffic, Log, TEXT( "Received out-of-date %s" ), *ReplicatedProp->GetName() );
-
-					DestObj		= NULL;
-					DestRecent	= NULL;
-				}
-
-				check( Bunch.PacketId >= Retire.InPacketId );
-
-				FMemMark Mark(FMemStack::Get());
-				uint8 * Data = DestObj ? ReplicatedProp->ContainerPtrToValuePtr<uint8>(DestObj, Element) : NewZeroed<uint8>(FMemStack::Get(),ReplicatedProp->ElementSize);
+				uint8 * Data = ReplicatedProp->ContainerPtrToValuePtr<uint8>((uint8*)Object, Element);
 				TArray<uint8>	MetaData;
 				PTRINT Offset = 0;
 
@@ -598,8 +557,6 @@ bool FObjectReplicator::ReceivedBunch( FInBunch &Bunch, const FReplicationFlags
 					}
 				}
 
-				Mark.Pop();
-
 				// Successfully received it.
 				UE_LOG( LogNetTraffic, Log, TEXT( " %s - %s - Change: %d" ), *Object->GetName(), *ReplicatedProp->GetName(), PropertyChanged );
 
diff --git a/Engine/Source/Runtime/Engine/Private/RepLayout.cpp b/Engine/Source/Runtime/Engine/Private/RepLayout.cpp
--- a/Engine/Source/Runtime/Engine/Private/RepLayout.cpp
+++ b/Engine/Source/Runtime/Engine/Private/RepLayout.cpp
@@ -108,17 +108,8 @@ static FORCEINLINE void SerializeGenericChecksum( FArchive & Ar )
 	check( Checksum == 0xABADF00D );
 }
 
-static void SerializeReadWritePropertyChecksum( const FRepLayoutCmd & Cmd, const int32 CurCmdIndex, const uint8 * Data, FArchive & Ar, const bool bDiscard )
+static void SerializeReadWritePropertyChecksum( const FRepLayoutCmd & Cmd, const int32 CurCmdIndex, const uint8 * Data, FArchive & Ar )
 {
-	if ( bDiscard )
-	{
-		uint32 MarkerChecksum = 0;
-		uint32 PropertyChecksum = 0;
-		Ar << MarkerChecksum;
-		Ar << PropertyChecksum;
-		return;
-	}
-
 	// Serialize various attributes that will mostly ensure we are working on the same property
 	const uint32 NameHash = GetTypeHash( Cmd.Property->GetName() );
 
@@ -918,7 +909,7 @@ uint16 FRepLayout::SendProperties_r(
 #ifdef ENABLE_PROPERTY_CHECKSUMS
 			if ( WriterState.bDoChecksum )
 			{
-				SerializeReadWritePropertyChecksum( Cmd, CmdIndex, Data + Cmd.Offset, WriterState.Writer, false );
+				SerializeReadWritePropertyChecksum( Cmd, CmdIndex, Data + Cmd.Offset, WriterState.Writer );
 			}
 #endif
 		}
@@ -1036,8 +1027,7 @@ bool FRepLayout::ReadProperty(
 	const FRepLayoutCmd &	Cmd, 
 	const int32				CmdIndex, 
 	uint8 * RESTRICT		StoredData, 
-	uint8 * RESTRICT		Data, 
-	const bool				bDiscard ) const
+	uint8 * RESTRICT		Data ) const
 {
 	if ( !ShouldReadProperty( ReaderState ) )
 	{
@@ -1045,23 +1035,6 @@ bool FRepLayout::ReadProperty(
 		return true;
 	}
 
-	if ( bDiscard )
-	{
-		FMemMark Mark( FMemStack::Get() );
-		uint8 * DiscardBuffer = NewZeroed<uint8>( FMemStack::Get(), Cmd.ElementSize );
-		const int32 OldArrayDim = Cmd.Property->ArrayDim;
-		Cmd.Property->ArrayDim = 1;		// Force this to 1, so that InitializeValue/DestroyValue work on a single item
-		Cmd.Property->InitializeValue( DiscardBuffer );
-		Cmd.Property->NetSerializeItem( ReaderState.Bunch, ReaderState.Bunch.PackageMap, DiscardBuffer );
-		Cmd.Property->DestroyValue( DiscardBuffer );
-		Cmd.Property->ArrayDim = OldArrayDim;
-		Mark.Pop();
-
-		// Read the next property handle
-		ReadNextHandle( ReaderState );
-		return true;
-	}
-
 	const FRepParentCmd & Parent = Parents[Cmd.ParentIndex];
 
 	// This swaps Role/RemoteRole as we write it
@@ -1101,7 +1074,7 @@ bool FRepLayout::ReadProperty(
 #ifdef ENABLE_PROPERTY_CHECKSUMS
 	if ( ReaderState.bDoChecksum )
 	{
-		SerializeReadWritePropertyChecksum( Cmd, CmdIndex, Data + SwappedCmd.Offset, ReaderState.Bunch, false );
+		SerializeReadWritePropertyChecksum( Cmd, CmdIndex, Data + SwappedCmd.Offset, ReaderState.Bunch );
 	}
 #endif
 
@@ -1122,8 +1095,7 @@ bool FRepLayout::ReceiveProperties_AnyArray_r(
 	const int32			ElementSize, 
 	const int32			CmdIndex, 
 	uint8 * RESTRICT	StoredData, 
-	uint8 *	RESTRICT	Data,
-	const bool			bDiscard ) const
+	uint8 *	RESTRICT	Data ) const
 {
 	const int32 CmdStart	= CmdIndex + 1;
 	const int32 CmdEnd		= Cmds[CmdIndex].EndCmd - 1;
@@ -1132,7 +1104,7 @@ bool FRepLayout::ReceiveProperties_AnyArray_r(
 	{
 		const int32 ElementOffset = i * ElementSize;
 
-		if ( !ReceiveProperties_r( ReaderState, UnmappedGuids, AbsOffset + ElementOffset, CmdStart, CmdEnd, StoredData + ElementOffset, Data + ElementOffset, bDiscard ) )
+		if ( !ReceiveProperties_r( ReaderState, UnmappedGuids, AbsOffset + ElementOffset, CmdStart, CmdEnd, StoredData + ElementOffset, Data + ElementOffset ) )
 		{
 			return false;
 		}
@@ -1148,8 +1120,7 @@ bool FRepLayout::ReceiveProperties_DynamicArray_r(
 	const FRepLayoutCmd &	Cmd, 
 	const int32				CmdIndex, 
 	uint8 * RESTRICT		StoredData, 
-	uint8 * RESTRICT		Data,
-	const bool				bDiscard ) const
+	uint8 * RESTRICT		Data ) const
 {
 	if ( !ShouldReadProperty( ReaderState ) )
 	{
@@ -1164,45 +1135,42 @@ bool FRepLayout::ReceiveProperties_DynamicArray_r(
 	// Read the next property handle
 	ReadNextHandle( ReaderState );
 
-	if ( !bDiscard )
-	{
-		FScriptArray * Array = (FScriptArray *)Data;
-		FScriptArray * StoredArray = (FScriptArray *)StoredData;
+	FScriptArray * Array = (FScriptArray *)Data;
+	FScriptArray * StoredArray = (FScriptArray *)StoredData;
 
-		// Since we don't know yet if something under us could be unmapped, go ahead and allocate an array container now
-		FUnmappedGuidMgrElement * NewArrayElement = UnmappedGuids->Map.Find( AbsOffset );
+	// Since we don't know yet if something under us could be unmapped, go ahead and allocate an array container now
+	FUnmappedGuidMgrElement * NewArrayElement = UnmappedGuids->Map.Find( AbsOffset );
 
-		if ( NewArrayElement == NULL )
-		{
-			NewArrayElement = &UnmappedGuids->Map.FindOrAdd( AbsOffset );
-			NewArrayElement->Array = new FUnmappedGuidMgr;
-			NewArrayElement->ParentIndex = Cmd.ParentIndex;
-			NewArrayElement->CmdIndex = CmdIndex;
-		}
+	if ( NewArrayElement == NULL )
+	{
+		NewArrayElement = &UnmappedGuids->Map.FindOrAdd( AbsOffset );
+		NewArrayElement->Array = new FUnmappedGuidMgr;
+		NewArrayElement->ParentIndex = Cmd.ParentIndex;
+		NewArrayElement->CmdIndex = CmdIndex;
+	}
 
-		check( NewArrayElement != NULL );
-		check( NewArrayElement->ParentIndex == Cmd.ParentIndex );
-		check( NewArrayElement->CmdIndex == CmdIndex );
+	check( NewArrayElement != NULL );
+	check( NewArrayElement->ParentIndex == Cmd.ParentIndex );
+	check( NewArrayElement->CmdIndex == CmdIndex );
 
-		UnmappedGuids = NewArrayElement->Array;
+	UnmappedGuids = NewArrayElement->Array;
 
-		const FRepParentCmd & Parent = Parents[Cmd.ParentIndex];
+	const FRepParentCmd & Parent = Parents[Cmd.ParentIndex];
 
-		if ( Array->Num() != ArrayNum && Parent.Property->HasAnyPropertyFlags( CPF_RepNotify ) )
-		{
-			ReaderState.RepState->RepNotifies.AddUnique( Parent.Property );
-		}
+	if ( Array->Num() != ArrayNum && Parent.Property->HasAnyPropertyFlags( CPF_RepNotify ) )
+	{
+		ReaderState.RepState->RepNotifies.AddUnique( Parent.Property );
+	}
 
-		// Set the array to the correct size (if we aren't discarding)
-		FScriptArrayHelper ArrayHelper( (UArrayProperty *)Cmd.Property, Data );
-		ArrayHelper.Resize( ArrayNum );
+	// Set the array to the correct size
+	FScriptArrayHelper ArrayHelper( (UArrayProperty *)Cmd.Property, Data );
+	ArrayHelper.Resize( ArrayNum );
 
-		FScriptArrayHelper StoredArrayHelper( (UArrayProperty *)Cmd.Property, StoredData );
-		StoredArrayHelper.Resize( ArrayNum );
+	FScriptArrayHelper StoredArrayHelper( (UArrayProperty *)Cmd.Property, StoredData );
+	StoredArrayHelper.Resize( ArrayNum );
 
-		Data		= (uint8*)Array->GetData();
-		StoredData	= (uint8*)StoredArray->GetData();
-	}
+	Data		= (uint8*)Array->GetData();
+	StoredData	= (uint8*)StoredArray->GetData();
 
 	const uint16 OldHandle = ReaderState.CurrentHandle;
 
@@ -1210,7 +1178,7 @@ bool FRepLayout::ReceiveProperties_DynamicArray_r(
 	ReaderState.CurrentHandle = 0;
 
 	// Read any changed properties into the array
-	if ( !ReceiveProperties_AnyArray_r( ReaderState, UnmappedGuids, 0, ArrayNum, Cmd.ElementSize, CmdIndex, StoredData, Data, bDiscard ) )
+	if ( !ReceiveProperties_AnyArray_r( ReaderState, UnmappedGuids, 0, ArrayNum, Cmd.ElementSize, CmdIndex, StoredData, Data ) )
 	{
 		return false;
 	}
@@ -1231,8 +1199,7 @@ bool FRepLayout::ReceiveProperties_r(
 	const int32			CmdStart, 
 	const int32			CmdEnd, 
 	uint8 * RESTRICT	StoredData, 
-	uint8 * RESTRICT	Data,
-	const bool			bDiscard ) const
+	uint8 * RESTRICT	Data ) const
 {
 	for ( int32 CmdIndex = CmdStart; CmdIndex < CmdEnd; CmdIndex++ )
 	{
@@ -1242,7 +1209,7 @@ bool FRepLayout::ReceiveProperties_r(
 
 		if ( Cmd.Type == REPCMD_DynamicArray )
 		{
-			if ( !ReceiveProperties_DynamicArray_r( ReaderState, UnmappedGuids, AbsOffset + Cmd.Offset, Cmd, CmdIndex, StoredData + Cmd.Offset, Data + Cmd.Offset, bDiscard ) )
+			if ( !ReceiveProperties_DynamicArray_r( ReaderState, UnmappedGuids, AbsOffset + Cmd.Offset, Cmd, CmdIndex, StoredData + Cmd.Offset, Data + Cmd.Offset ) )
 			{
 				return false;
 			}
@@ -1250,7 +1217,7 @@ bool FRepLayout::ReceiveProperties_r(
 			continue;
 		}
 
-		if ( !ReadProperty( ReaderState, UnmappedGuids, AbsOffset, Cmd, CmdIndex, StoredData, Data, bDiscard ) )
+		if ( !ReadProperty( ReaderState, UnmappedGuids, AbsOffset, Cmd, CmdIndex, StoredData, Data ) )
 		{
 			return false;
 		}
@@ -1259,7 +1226,7 @@ bool FRepLayout::ReceiveProperties_r(
 	return true;
 }
 
-bool FRepLayout::ReceiveProperties( UClass * InObjectClass, FRepState * RESTRICT RepState, void * RESTRICT Data, FNetBitReader & InBunch, bool bDiscard, bool & bOutHasUnmapped ) const
+bool FRepLayout::ReceiveProperties( UClass * InObjectClass, FRepState * RESTRICT RepState, void * RESTRICT Data, FNetBitReader & InBunch, bool & bOutHasUnmapped ) const
 {
 	check( InObjectClass == Owner );
 
@@ -1277,7 +1244,7 @@ bool FRepLayout::ReceiveProperties( UClass * InObjectClass, FRepState * RESTRICT
 	ReadNextHandle( ReaderState );
 
 	// Read all properties
-	if ( !ReceiveProperties_r( ReaderState, &RepState->UnmappedGuids, 0, 0, Cmds.Num() - 1, RepState->StaticBuffer.GetTypedData(), (uint8*)Data, bDiscard ) )
+	if ( !ReceiveProperties_r( ReaderState, &RepState->UnmappedGuids, 0, 0, Cmds.Num() - 1, RepState->StaticBuffer.GetTypedData(), (uint8*)Data ) )
 	{
 		return false;
 	}
@@ -1292,7 +1259,7 @@ bool FRepLayout::ReceiveProperties( UClass * InObjectClass, FRepState * RESTRICT
 #ifdef ENABLE_SUPER_CHECKSUMS
 	if ( InBunch.ReadBit() == 1 )
 	{
-		ValidateWithChecksum( RepState->StaticBuffer.GetTypedData(), InBunch, bDiscard );
+		ValidateWithChecksum( RepState->StaticBuffer.GetTypedData(), InBunch );
 	}
 #endif
 
@@ -1430,22 +1397,8 @@ void FRepLayout::CallRepNotifies( FRepState * RepState, UObject * Object ) const
 	RepState->RepNotifies.Empty();
 }
 
-void FRepLayout::ValidateWithChecksum_DynamicArray_r( const FRepLayoutCmd & Cmd, const int32 CmdIndex, const uint8 * RESTRICT Data, FArchive & Ar, const bool bDiscard ) const
+void FRepLayout::ValidateWithChecksum_DynamicArray_r( const FRepLayoutCmd & Cmd, const int32 CmdIndex, const uint8 * RESTRICT Data, FArchive & Ar ) const
 {
-	if ( bDiscard )
-	{
-		uint16 ArrayNum		= 0;
-		uint16 ElementSize	= 0;
-		Ar << ArrayNum;
-		Ar << ElementSize;
-
-		for ( int32 i = 0; i < ArrayNum; i++ )
-		{
-			ValidateWithChecksum_r( CmdIndex + 1, Cmd.EndCmd - 1, NULL, Ar, bDiscard );
-		}
-		return;
-	}
-
 	FScriptArray * Array = (FScriptArray *)Data;
 
 	uint16 ArrayNum		= Array->Num();
@@ -1468,7 +1421,7 @@ void FRepLayout::ValidateWithChecksum_DynamicArray_r( const FRepLayoutCmd & Cmd,
 
 	for ( int32 i = 0; i < ArrayNum; i++ )
 	{
-		ValidateWithChecksum_r( CmdIndex + 1, Cmd.EndCmd - 1, LocalData + i * ElementSize, Ar, bDiscard );
+		ValidateWithChecksum_r( CmdIndex + 1, Cmd.EndCmd - 1, LocalData + i * ElementSize, Ar );
 	}
 }
 
@@ -1476,8 +1429,7 @@ void FRepLayout::ValidateWithChecksum_r(
 	const int32				CmdStart, 
 	const int32				CmdEnd, 
 	const uint8 * RESTRICT	Data, 
-	FArchive &				Ar,
-	const bool				bDiscard ) const
+	FArchive &				Ar ) const
 {
 	for ( int32 CmdIndex = CmdStart; CmdIndex < CmdEnd; CmdIndex++ )
 	{
@@ -1487,24 +1439,24 @@ void FRepLayout::ValidateWithChecksum_r(
 
 		if ( Cmd.Type == REPCMD_DynamicArray )
 		{
-			ValidateWithChecksum_DynamicArray_r( Cmd, CmdIndex, Data + Cmd.Offset, Ar, bDiscard );
+			ValidateWithChecksum_DynamicArray_r( Cmd, CmdIndex, Data + Cmd.Offset, Ar );
 			CmdIndex = Cmd.EndCmd - 1;	// Jump past children of this array (-1 for ++ in for loop)
 			continue;
 		}
 
-		SerializeReadWritePropertyChecksum( Cmd, CmdIndex - 1, Data + Cmd.Offset, Ar, bDiscard );
+		SerializeReadWritePropertyChecksum( Cmd, CmdIndex - 1, Data + Cmd.Offset, Ar );
 	}
 }
 
-void FRepLayout::ValidateWithChecksum( const void * RESTRICT Data, FArchive & Ar, const bool bDiscard ) const
+void FRepLayout::ValidateWithChecksum( const void * RESTRICT Data, FArchive & Ar ) const
 {
-	ValidateWithChecksum_r( 0, Cmds.Num() - 1, (const uint8*)Data, Ar, bDiscard );
+	ValidateWithChecksum_r( 0, Cmds.Num() - 1, (const uint8*)Data, Ar );
 }
 
 uint32 FRepLayout::GenerateChecksum( const FRepState * RepState ) const
 {
 	FBitWriter Writer( 1024, true );
-	ValidateWithChecksum_r( 0, Cmds.Num() - 1, (const uint8*)RepState->StaticBuffer.GetTypedData(), Writer, false );
+	ValidateWithChecksum_r( 0, Cmds.Num() - 1, (const uint8*)RepState->StaticBuffer.GetTypedData(), Writer );
 
 	return FCrc::MemCrc32( Writer.GetData(), Writer.GetNumBytes(), 0 );
 }
diff --git a/Engine/Source/Runtime/Engine/Public/Net/RepLayout.h b/Engine/Source/Runtime/Engine/Public/Net/RepLayout.h
--- a/Engine/Source/Runtime/Engine/Public/Net/RepLayout.h
+++ b/Engine/Source/Runtime/Engine/Public/Net/RepLayout.h
@@ -321,7 +321,7 @@ public:
 
 	void InitFromObjectClass( UClass * InObjectClass );
 
-	bool ReceiveProperties( UClass * InObjectClass, FRepState * RESTRICT RepState, void * RESTRICT Data, FNetBitReader & InBunch, bool bDiscard, bool & bOutHasUnmapped ) const;
+	bool ReceiveProperties( UClass * InObjectClass, FRepState * RESTRICT RepState, void * RESTRICT Data, FNetBitReader & InBunch, bool & bOutHasUnmapped ) const;
 	void UpdateUnmappedObjects( FRepState *	RepState, UPackageMap * PackageMap, UObject * Object, bool & bOutSomeObjectsWereMapped, bool & bOutHasMoreUnmapped ) const;
 
 	void CallRepNotifies( FRepState * RepState, UObject * Object ) const;
@@ -330,7 +330,7 @@ public:
 	bool AllAcked( FRepState * RepState ) const;
 	bool ReadyForDormancy( FRepState * RepState ) const;
 
-	void ValidateWithChecksum( const void * RESTRICT Data, FArchive & Ar, const bool bDiscard ) const;
+	void ValidateWithChecksum( const void * RESTRICT Data, FArchive & Ar ) const;
 	uint32 GenerateChecksum( const FRepState * RepState ) const;
 
 	void MergeDirtyList( const void * RESTRICT Data, const TArray< uint16 > & Dirty1, const TArray< uint16 > & Dirty2, TArray< uint16 > & MergedDirty ) const;
@@ -418,8 +418,7 @@ private:
 		const FRepLayoutCmd &	Cmd, 
 		const int32				CurrentCmdIndex, 
 		uint8 * RESTRICT		StoredData, 
-		uint8 * RESTRICT		Data, 
-		const bool				bDiscard ) const;
+		uint8 * RESTRICT		Data ) const;
 
 	bool ReceiveProperties_AnyArray_r( 
 		FRepReaderState &	ReaderState, 
@@ -429,8 +428,7 @@ private:
 		const int32			ElementSize, 
 		const int32			CmdIndex, 
 		uint8 * RESTRICT	StoredData, 
-		uint8 *	RESTRICT	Data,
-		const bool			bDiscard ) const;
+		uint8 *	RESTRICT	Data ) const;
 
 	bool ReceiveProperties_DynamicArray_r( 
 		FRepReaderState &		ReaderState, 
@@ -439,8 +437,7 @@ private:
 		const FRepLayoutCmd &	Cmd, 
 		const int32				CmdIndex, 
 		uint8 * RESTRICT		StoredData, 
-		uint8 * RESTRICT		Data,
-		const bool				bDiscard ) const;
+		uint8 * RESTRICT		Data ) const;
 
 	bool ReceiveProperties_r( 
 		FRepReaderState &	ReaderState, 
@@ -449,8 +446,7 @@ private:
 		const int32			CmdStart, 
 		const int32			CmdEnd, 
 		uint8 * RESTRICT	StoredData, 
-		uint8 * RESTRICT	Data,
-		const bool			bDiscard ) const;
+		uint8 * RESTRICT	Data ) const;
 
 	void UpdateUnmappedObjects_r( 
 		FRepState *			RepState, 
@@ -462,8 +458,8 @@ private:
 		bool &				bOutSomeObjectsWereMapped,
 		bool &				bOutHasMoreUnmapped ) const;
 
-	void ValidateWithChecksum_DynamicArray_r( const FRepLayoutCmd & Cmd, const int32 CmdIndex, const uint8 * RESTRICT Data, FArchive & Ar, const bool bDiscard ) const;
-	void ValidateWithChecksum_r( const int32 CmdStart, const int32 CmdEnd, const uint8 * RESTRICT Data, FArchive & Ar, const bool bDiscard ) const;
+	void ValidateWithChecksum_DynamicArray_r( const FRepLayoutCmd & Cmd, const int32 CmdIndex, const uint8 * RESTRICT Data, FArchive & Ar ) const;
+	void ValidateWithChecksum_r( const int32 CmdStart, const int32 CmdEnd, const uint8 * RESTRICT Data, FArchive & Ar ) const;
 
 	void MergeDirtyList_AnyArray_r( FMergeDirtyListState & MergeState, const FRepLayoutCmd & Cmd, const int32 ArrayNum, const int32 ElementSize, const int32 CmdIndex, const uint8 * RESTRICT Data ) const;
 	void MergeDirtyList_DynamicArray_r( FMergeDirtyListState & MergeState, const FRepLayoutCmd & Cmd, const int32 CmdIndex, const uint8 * RESTRICT Data ) const;
