Engine: Add support for a negative play rate on timeline components

diff --git a/Engine/Source/Runtime/Engine/Private/Timeline.cpp b/Engine/Source/Runtime/Engine/Private/Timeline.cpp
--- a/Engine/Source/Runtime/Engine/Private/Timeline.cpp
+++ b/Engine/Source/Runtime/Engine/Private/Timeline.cpp
@@ -356,13 +356,13 @@ void FTimeline::TickTimeline(float DeltaTime)
 
 	if(bPlaying)
 	{
-		float NewPosition = Position;
 		const float TimelineLength = GetTimelineLength();
+		const float EffectiveDeltaTime = DeltaTime * (bReversePlayback ? (-PlayRate) : (PlayRate));
 
-		if(!bReversePlayback)
-		{
-			NewPosition = Position + (DeltaTime * PlayRate);
+		float NewPosition = Position + EffectiveDeltaTime;
 
+		if(EffectiveDeltaTime > 0.0f)
+		{
 			if(NewPosition > TimelineLength)
 			{
 				// If looping, play to end, jump to start, and set target to somewhere near the beginning.
@@ -394,8 +394,6 @@ void FTimeline::TickTimeline(float DeltaTime)
 		}
 		else
 		{
-			NewPosition = Position - (DeltaTime * PlayRate);
-
 			if(NewPosition < 0.f)
 			{
 				// If looping, play to start, jump to end, and set target to somewhere near the end.
