Add ObjectTypes array to particle system

TTP#345755

diff --git a/Engine/Source/Runtime/Engine/Classes/Particles/Collision/ParticleModuleCollision.h b/Engine/Source/Runtime/Engine/Classes/Particles/Collision/ParticleModuleCollision.h
--- a/Engine/Source/Runtime/Engine/Classes/Particles/Collision/ParticleModuleCollision.h
+++ b/Engine/Source/Runtime/Engine/Classes/Particles/Collision/ParticleModuleCollision.h
@@ -50,6 +50,18 @@ class UParticleModuleCollision : public UParticleModuleCollisionBase
 	UPROPERTY(EditAnywhere, Category=Collision)
 	TEnumAsByte<enum EParticleCollisionComplete> CollisionCompletionOption;
 
+	/**
+	 * Which ObjectTypes to collide with
+	 */
+	UPROPERTY(EditAnywhere, Category = Collision)
+	TArray< TEnumAsByte<enum EObjectTypeQuery> > CollisionTypes;
+
+	/**
+	 * ObjectParams created for array of CollisionTypes
+	 */
+	FCollisionObjectQueryParams ObjectParams;
+	
+
 	/** 
 	 *	If true, physic will be applied between a particle and the 
 	 *	object it collides with. 
diff --git a/Engine/Source/Runtime/Engine/Classes/Particles/ParticleSystemComponent.h b/Engine/Source/Runtime/Engine/Classes/Particles/ParticleSystemComponent.h
--- a/Engine/Source/Runtime/Engine/Classes/Particles/ParticleSystemComponent.h
+++ b/Engine/Source/Runtime/Engine/Classes/Particles/ParticleSystemComponent.h
@@ -773,7 +773,7 @@ public:
 	/** Deactivate the system */
 	void DeactivateSystem();
 	// Collision Handling...
-	virtual bool ParticleLineCheck(FHitResult& Hit, AActor* SourceActor, const FVector& End, const FVector& Start, const FVector& HalfExtent);
+	virtual bool ParticleLineCheck(FHitResult& Hit, AActor* SourceActor, const FVector& End, const FVector& Start, const FVector& HalfExtent, const FCollisionObjectQueryParams & ObjectParams);
 
 	/** return true if this psys can tick in any thread */
 	FORCEINLINE bool CanTickInAnyThread()
diff --git a/Engine/Source/Runtime/Engine/Private/Particles/ParticleComponents.cpp b/Engine/Source/Runtime/Engine/Private/Particles/ParticleComponents.cpp
--- a/Engine/Source/Runtime/Engine/Private/Particles/ParticleComponents.cpp
+++ b/Engine/Source/Runtime/Engine/Private/Particles/ParticleComponents.cpp
@@ -3035,21 +3035,22 @@ SIZE_T UParticleSystemComponent::GetResourceSize(EResourceSizeMode::Type Mode)
 }
 
 
-bool UParticleSystemComponent::ParticleLineCheck(FHitResult& Hit, AActor* SourceActor, const FVector& End, const FVector& Start, const FVector& HalfExtent)
+bool UParticleSystemComponent::ParticleLineCheck(FHitResult& Hit, AActor* SourceActor, const FVector& End, const FVector& Start, const FVector& HalfExtent, const FCollisionObjectQueryParams & ObjectParams)
 {
 	ForceAsyncWorkCompletion(ENSURE_AND_STALL);
 	check(GetWorld());
 	static FName NAME_ParticleCollision = FName(TEXT("ParticleCollision"));
+
 	if ( HalfExtent.IsZero() )
 	{
-		return GetWorld()->LineTraceSingle(Hit, Start, End, FCollisionQueryParams(NAME_ParticleCollision, true, SourceActor), FCollisionObjectQueryParams(ECC_WorldStatic));
+		return GetWorld()->LineTraceSingle(Hit, Start, End, FCollisionQueryParams(NAME_ParticleCollision, true, SourceActor), ObjectParams);
 	}
 	else
 	{
 		FCollisionQueryParams BoxParams(false);
 		BoxParams.TraceTag = NAME_ParticleCollision;
 		BoxParams.AddIgnoredActor(SourceActor);
-		return GetWorld()->SweepSingle(Hit, Start, End, FQuat::Identity, FCollisionShape::MakeBox(HalfExtent), BoxParams, FCollisionObjectQueryParams(ECC_WorldStatic));
+		return GetWorld()->SweepSingle(Hit, Start, End, FQuat::Identity, FCollisionShape::MakeBox(HalfExtent), BoxParams, ObjectParams);
 	}
 }
 
diff --git a/Engine/Source/Runtime/Engine/Private/Particles/ParticleModules_Collision.cpp b/Engine/Source/Runtime/Engine/Private/Particles/ParticleModules_Collision.cpp
--- a/Engine/Source/Runtime/Engine/Private/Particles/ParticleModules_Collision.cpp
+++ b/Engine/Source/Runtime/Engine/Private/Particles/ParticleModules_Collision.cpp
@@ -51,6 +51,7 @@ UParticleModuleCollision::UParticleModuleCollision(const class FPostConstructIni
 	bCollideOnlyIfVisible = true;
 	MaxCollisionDistance = 1000.0f;
 	bIgnoreSourceActor = true;
+	CollisionTypes.Add(UEngineTypes::ConvertToObjectType(ECC_WorldStatic));
 }
 
 void UParticleModuleCollision::InitializeDefaults()
@@ -85,6 +86,8 @@ void UParticleModuleCollision::InitializeDefaults()
 		DistributionDelayAmount->Constant = 0.0f;
 		DelayAmount.Distribution = DistributionDelayAmount;
 	}
+
+	ObjectParams = FCollisionObjectQueryParams(CollisionTypes);
 }
 
 void UParticleModuleCollision::PostInitProperties()
@@ -566,7 +569,7 @@ bool UParticleModuleCollision::PerformCollisionCheck(FParticleEmitterInstance* O
 	FHitResult& Hit, AActor* SourceActor, const FVector& End, const FVector& Start, const FVector& Extent)
 {
 	check(Owner && Owner->Component);
-	return Owner->Component->ParticleLineCheck(Hit, SourceActor, End, Start, Extent);
+	return Owner->Component->ParticleLineCheck(Hit, SourceActor, End, Start, Extent, ObjectParams);
 }
 
 /*------------------------------------------------------------------------------
