Blueprint function to change skylight cubemap

diff --git a/Engine/Source/Runtime/Engine/Classes/Components/SkyLightComponent.h b/Engine/Source/Runtime/Engine/Classes/Components/SkyLightComponent.h
--- a/Engine/Source/Runtime/Engine/Classes/Components/SkyLightComponent.h
+++ b/Engine/Source/Runtime/Engine/Classes/Components/SkyLightComponent.h
@@ -153,6 +153,10 @@ class ENGINE_API USkyLightComponent : public ULightComponentBase
 	UFUNCTION(BlueprintCallable, Category="Rendering|Components|SkyLight")
 	void SetLightColor(FLinearColor NewLightColor);
 
+	/** Sets the cubemap used when SourceType is set to SpecifiedCubemap, and causes a skylight update on the next tick. */
+	UFUNCTION(BlueprintCallable, Category="SkyLight")
+	void SetCubemap(UTextureCube* NewCubemap);
+
 	/** Indicates that the capture needs to recapture the scene, adds it to the recapture queue. */
 	void SetCaptureIsDirty();
 
diff --git a/Engine/Source/Runtime/Engine/Private/SkyLightComponent.cpp b/Engine/Source/Runtime/Engine/Private/SkyLightComponent.cpp
--- a/Engine/Source/Runtime/Engine/Private/SkyLightComponent.cpp
+++ b/Engine/Source/Runtime/Engine/Private/SkyLightComponent.cpp
@@ -445,11 +445,23 @@ void USkyLightComponent::SetLightColor(FLinearColor NewLightColor)
 	if (!(IsRegistered() && Mobility == EComponentMobility::Static)
 		&& LightColor != NewColor)
 	{
-		LightColor	= NewColor;
+		LightColor = NewColor;
 		MarkRenderStateDirty();
 	}
 }
 
+void USkyLightComponent::SetCubemap(UTextureCube* NewCubemap)
+{
+	// Can't set color on a static light
+	if (!(IsRegistered() && Mobility == EComponentMobility::Static)
+		&& Cubemap != NewCubemap)
+	{
+		Cubemap = NewCubemap;
+		MarkRenderStateDirty();
+		SetCaptureIsDirty();
+	}
+}
+
 void USkyLightComponent::RecaptureSky()
 {
 	SetCaptureIsDirty();
