Fix case-sensitive issue.

diff --git a/Engine/Source/Runtime/Engine/Private/PhysicsEngine/PhysDerivedData.cpp b/Engine/Source/Runtime/Engine/Private/PhysicsEngine/PhysDerivedData.cpp
--- a/Engine/Source/Runtime/Engine/Private/PhysicsEngine/PhysDerivedData.cpp
+++ b/Engine/Source/Runtime/Engine/Private/PhysicsEngine/PhysDerivedData.cpp
@@ -4,7 +4,7 @@
 #include "TargetPlatform.h"
 
 #if WITH_PHYSX && (WITH_RUNTIME_PHYSICS_COOKING || WITH_EDITOR)
-#include "IPhysxFormatModule.h"
+#include "IPhysXFormatModule.h"
 #include "Landscape/LandscapeMeshCollisionComponent.h"
 
 FDerivedDataPhysXCooker::FDerivedDataPhysXCooker( FName InFormat, UBodySetup* InBodySetup )
