#oculus Integrating fix for fullscreening with an HMD present

diff --git a/Engine/Source/Runtime/Engine/Private/GameEngine.cpp b/Engine/Source/Runtime/Engine/Private/GameEngine.cpp
--- a/Engine/Source/Runtime/Engine/Private/GameEngine.cpp
+++ b/Engine/Source/Runtime/Engine/Private/GameEngine.cpp
@@ -80,11 +80,11 @@ EWindowMode::Type GetWindowModeType(EWindowMode::Type WindowMode)
 {
 	if (FPlatformProperties::SupportsWindowedMode())
 	{
-		if (GEngine && GEngine->HMDDevice.IsValid() && GEngine->HMDDevice->IsFullScreenAllowed())
+		if ((WindowMode != EWindowMode::Windowed) && GEngine && GEngine->HMDDevice.IsValid() && GEngine->HMDDevice->IsFullScreenAllowed())
 		{
 			return EWindowMode::Fullscreen;
 		}
-		
+
 		return WindowMode;
 	}
 
