fixed navigation system update after changing bCanEverAffectNavigation flag

diff --git a/Engine/Source/Runtime/Engine/Private/PrimitiveComponent.cpp b/Engine/Source/Runtime/Engine/Private/PrimitiveComponent.cpp
--- a/Engine/Source/Runtime/Engine/Private/PrimitiveComponent.cpp
+++ b/Engine/Source/Runtime/Engine/Private/PrimitiveComponent.cpp
@@ -510,6 +510,13 @@ void UPrimitiveComponent::PostEditChangeProperty(FPropertyChangedEvent& Property
 	{
 		MarkRenderStateDirty();
 	}
+
+	// update component, ActorComponent's property update locks navigation system 
+	// so it needs to be called directly here
+	if (PropertyThatChanged && PropertyThatChanged->GetFName() == GET_MEMBER_NAME_CHECKED(UPrimitiveComponent, bCanEverAffectNavigation))
+	{
+		UNavigationSystem::UpdateNavOctree(this);
+	}
 }
 
 bool UPrimitiveComponent::CanEditChange(const UProperty* InProperty) const
