Fixed lenght->length typos in FNavigationPath

GitHub pull request #418 - https://github.com/EpicGames/UnrealEngine/pull/418

TTP# 345846 - [GitHub] 418 : Fixed lenght->length typos in FNavigationPath.

diff --git a/Engine/Config/BaseEngine.ini b/Engine/Config/BaseEngine.ini
--- a/Engine/Config/BaseEngine.ini
+++ b/Engine/Config/BaseEngine.ini
@@ -432,6 +432,7 @@ bUseStreamingPause=false
 +TaggedPropertyRedirects=(ClassName="NavLinkProxy",OldPropertyName="NavSegmentLinks",NewPropertyName="SegmentLinks")
 +EnumRedirects=(EnumName="ENavLinkDirection",OldEnumEntry="ESmartNavLinkDir::OneWay",NewEnumEntry="ENavLinkDirection::LeftToRight")
 +EnumRedirects=(EnumName="ENavLinkDirection",OldEnumEntry="ESmartNavLinkDir::BothWays",NewEnumEntry="ENavLinkDirection::BothWays")
++K2FieldRedirects=(OldFieldName="NavigationPath.GetPathLenght",NewFieldName="NavigationPath.GetPathLength")
 
 ;SkeletalMesh redirects
 +TaggedPropertyRedirects=(ClassName="SkeletalMeshLODInfo",OldPropertyName="DisplayFactor",NewPropertyName="ScreenSize")
diff --git a/Engine/Source/Runtime/Engine/Classes/AI/Navigation/NavigationPath.h b/Engine/Source/Runtime/Engine/Classes/AI/Navigation/NavigationPath.h
--- a/Engine/Source/Runtime/Engine/Classes/AI/Navigation/NavigationPath.h
+++ b/Engine/Source/Runtime/Engine/Classes/AI/Navigation/NavigationPath.h
@@ -54,7 +54,7 @@ public:
 	void EnableRecalculationOnInvalidation(TEnumAsByte<ENavigationOptionFlag::Type> DoRecalculation);
 
 	UFUNCTION(BlueprintCallable, Category = "AI|Navigation")
-	float GetPathLenght() const;
+	float GetPathLength() const;
 
 	UFUNCTION(BlueprintCallable, Category = "AI|Navigation")
 	float GetPathCost() const;
diff --git a/Engine/Source/Runtime/Engine/Private/AI/Navigation/NavigationPath.cpp b/Engine/Source/Runtime/Engine/Private/AI/Navigation/NavigationPath.cpp
--- a/Engine/Source/Runtime/Engine/Private/AI/Navigation/NavigationPath.cpp
+++ b/Engine/Source/Runtime/Engine/Private/AI/Navigation/NavigationPath.cpp
@@ -357,8 +357,8 @@ const TArray<FNavigationPortalEdge>* FNavMeshPath::GeneratePathCorridorEdges() c
 #if WITH_RECAST
 	// mz@todo the underlying recast function queries the navmesh a portal at a time, 
 	// which is a waste of performance. A batch-query function has to be added.
-	const int32 CorridorLenght = PathCorridor.Num();
-	if (CorridorLenght != 0 && IsInGameThread() && NavigationDataUsed.IsValid())
+	const int32 CorridorLength = PathCorridor.Num();
+	if (CorridorLength != 0 && IsInGameThread() && NavigationDataUsed.IsValid())
 	{
 		const ARecastNavMesh* MyOwner = Cast<ARecastNavMesh>(GetNavigationDataUsed());
 		MyOwner->GetEdgesForPathCorridor(&PathCorridor, &PathCorridorEdges);
@@ -1041,7 +1041,7 @@ void UNavigationPath::EnableRecalculationOnInvalidation(TEnumAsByte<ENavigationO
 	}
 }
 
-float UNavigationPath::GetPathLenght() const
+float UNavigationPath::GetPathLength() const
 {
 	check((SharedPath.IsValid() && SharedPath->IsValid()) == !!bIsValid);
 	return !!bIsValid ? SharedPath->GetLength() : -1.f;
