Removed Dialogs.h header from DocumentationActor.cpp

diff --git a/Engine/Source/Runtime/Engine/Private/DocumentationActor.cpp b/Engine/Source/Runtime/Engine/Private/DocumentationActor.cpp
--- a/Engine/Source/Runtime/Engine/Private/DocumentationActor.cpp
+++ b/Engine/Source/Runtime/Engine/Private/DocumentationActor.cpp
@@ -5,7 +5,6 @@
 #if WITH_EDITOR
 #include "Components/MaterialBillboardComponent.h"
 #include "IDocumentation.h"
-#include "Dialogs.h"
 #endif
 
 ADocumentationActor::ADocumentationActor(const class FPostConstructInitializeProperties& PCIP)
