fix for missing #if WITH_EDITOR block and #if WITH_EDITORONLY_DATA block when doing a game build
#ue4

diff --git a/Engine/Source/Runtime/Engine/Private/DocumentationActor.cpp b/Engine/Source/Runtime/Engine/Private/DocumentationActor.cpp
--- a/Engine/Source/Runtime/Engine/Private/DocumentationActor.cpp
+++ b/Engine/Source/Runtime/Engine/Private/DocumentationActor.cpp
@@ -49,6 +49,7 @@ EDocumentationActorType::Type ADocumentationActor::GetLinkType() const
 	return LinkType;
 }
 
+#if WITH_EDITOR
 void ADocumentationActor::PostLoad()
 {
 	Super::PostLoad();
@@ -63,9 +64,11 @@ void ADocumentationActor::PostEditChangeProperty(FPropertyChangedEvent& Property
 	
 	UpdateLinkType();
 }
+#endif
 
 void ADocumentationActor::UpdateLinkType()
 {
+#if WITH_EDITORONLY_DATA
 	if (DocumentLink.IsEmpty() == true)
 	{
 		LinkType = EDocumentationActorType::None;
@@ -79,6 +82,7 @@ void ADocumentationActor::UpdateLinkType()
 		LinkType = EDocumentationActorType::UDNLink;
 	}
 	else
+#endif
 	{
 		LinkType = EDocumentationActorType::InvalidLink;
 	}
