Cleaned the stats verification code some more.  It will now load the automation test map.

diff --git a/Engine/Source/Runtime/Engine/Private/Tests/EngineAutomationTests.cpp b/Engine/Source/Runtime/Engine/Private/Tests/EngineAutomationTests.cpp
--- a/Engine/Source/Runtime/Engine/Private/Tests/EngineAutomationTests.cpp
+++ b/Engine/Source/Runtime/Engine/Private/Tests/EngineAutomationTests.cpp
@@ -115,35 +115,24 @@ bool FStatsVerificationMapTest::RunTest(const FString& Parameters)
 	//Get the location of the map being used.
 	FString Filename = FPaths::ConvertRelativePathToFull(MapName);
 
+	//If the filename doesn't exist then we'll make the filename path relative to the engine instead of the game and the retry.
 	if (!FPaths::FileExists(Filename))
 	{
-		UE_LOG(LogEngineAutomationTests, Warning, TEXT("Map doesn't exist: %s.  \nUsing currently loaded map."), *Filename);
-		
-		ADD_LATENT_AUTOMATION_COMMAND(FExecStringLatentCommand(TEXT("stat game")));
-		ADD_LATENT_AUTOMATION_COMMAND(FEngineWaitLatentCommand(1.0));
-		ADD_LATENT_AUTOMATION_COMMAND(FExecStringLatentCommand(TEXT("stat game")));
+		FString EngDirectory = FPaths::ConvertRelativePathToFull(FPaths::EngineContentDir());
 
-		ADD_LATENT_AUTOMATION_COMMAND(FExecStringLatentCommand(TEXT("stat scenerendering")));
-		ADD_LATENT_AUTOMATION_COMMAND(FEngineWaitLatentCommand(1.0));
-		ADD_LATENT_AUTOMATION_COMMAND(FExecStringLatentCommand(TEXT("stat scenerendering")));
+		Filename = EngDirectory / MapName;
 
-		ADD_LATENT_AUTOMATION_COMMAND(FExecStringLatentCommand(TEXT("stat memory")));
-		ADD_LATENT_AUTOMATION_COMMAND(FEngineWaitLatentCommand(1.0));
-		ADD_LATENT_AUTOMATION_COMMAND(FExecStringLatentCommand(TEXT("stat memory")));
-
-		ADD_LATENT_AUTOMATION_COMMAND(FExecStringLatentCommand(TEXT("stat slate")));
-		ADD_LATENT_AUTOMATION_COMMAND(FEngineWaitLatentCommand(1.0));
-		ADD_LATENT_AUTOMATION_COMMAND(FExecStringLatentCommand(TEXT("stat slate")));
-
-		return true;
+		Filename.ReplaceInline(TEXT("/Content/../../"), TEXT("/"),ESearchCase::CaseSensitive);
 	}
-	
-	//If the map is located in the engine folder then we'll make the path relative to that.  Otherwise it will be relative to the game content folder.
-	if (Filename.Contains(TEXT("Engine"), ESearchCase::IgnoreCase, ESearchDir::FromStart))
+
+	if (FPaths::FileExists(Filename))
 	{
-		//This will be false if the map is on a different drive.
-		if (FPaths::MakePathRelativeTo(Filename, *FPaths::EngineContentDir()))
+		//If the map is located in the engine folder then we'll make the path relative to that.  Otherwise it will be relative to the game content folder.
+		if (Filename.Contains(TEXT("Engine"), ESearchCase::IgnoreCase, ESearchDir::FromStart))
 		{
+			//This will be false if the map is on a different drive.
+			if (FPaths::MakePathRelativeTo(Filename, *FPaths::EngineContentDir()))
+			{
 				FString ShortName = FPaths::GetBaseFilename(Filename);
 				FString PathName = FPaths::GetPath(Filename);
 				MapName = FString::Printf(TEXT("/Engine/%s/%s.%s"), *PathName, *ShortName, *ShortName);
@@ -151,31 +140,36 @@ bool FStatsVerificationMapTest::RunTest(const FString& Parameters)
 				UE_LOG(LogEngineAutomationTests, Log, TEXT("Opening '%s'"), *MapName);
 
 				GEngine->Exec(GetSimpleEngineAutomationTestGameWorld(GetTestFlags()), *FString::Printf(TEXT("Open %s"), *MapName));
+			}
+			else
+			{
+				UE_LOG(LogEngineAutomationTests, Error, TEXT("Invalid asset path: %s."), *Filename);
+			}
 		}
 		else
 		{
-			UE_LOG(LogEngineAutomationTests, Error, TEXT("Invalid asset path: %s."), *Filename);
+			//This will be false if the map is on a different drive.
+			if (FPaths::MakePathRelativeTo(Filename, *FPaths::GameContentDir()))
+			{
+				FString ShortName = FPaths::GetBaseFilename(Filename);
+				FString PathName = FPaths::GetPath(Filename);
+				MapName = FString::Printf(TEXT("/Game/%s/%s.%s"), *PathName, *ShortName, *ShortName);
+
+				UE_LOG(LogEngineAutomationTests, Log, TEXT("Opening '%s'"), *MapName);
+
+				GEngine->Exec(GetSimpleEngineAutomationTestGameWorld(GetTestFlags()), *FString::Printf(TEXT("Open %s"), *MapName));
+			}
+			else
+			{
+				UE_LOG(LogEngineAutomationTests, Error, TEXT("Invalid asset path: %s."), *Filename);
+			}
 		}
 	}
 	else
 	{
-		//This will be false if the map is on a different drive.
-		if (FPaths::MakePathRelativeTo(Filename, *FPaths::GameContentDir()))
-		{
-			FString ShortName = FPaths::GetBaseFilename(Filename);
-			FString PathName = FPaths::GetPath(Filename);
-			MapName = FString::Printf(TEXT("/Game/%s/%s.%s"), *PathName, *ShortName, *ShortName);
-
-			UE_LOG(LogEngineAutomationTests, Log, TEXT("Opening '%s'"), *MapName);
-
-			GEngine->Exec(GetSimpleEngineAutomationTestGameWorld(GetTestFlags()), *FString::Printf(TEXT("Open %s"), *MapName));
-		}
-		else
-		{
-			UE_LOG(LogEngineAutomationTests, Error, TEXT("Invalid asset path: %s."), *Filename);
-		}
+		UE_LOG(LogEngineAutomationTests, Warning, TEXT("Automation test map doesn't exist or is not set: %s.  \nUsing the currently loaded map."), *Filename);
 	}
-	
+
 	ADD_LATENT_AUTOMATION_COMMAND(FExecStringLatentCommand(TEXT("stat game")));
 	ADD_LATENT_AUTOMATION_COMMAND(FEngineWaitLatentCommand(1.0));
 	ADD_LATENT_AUTOMATION_COMMAND(FExecStringLatentCommand(TEXT("stat game")));
