A couple more minor updates to the Stat Verification tests.

diff --git a/Engine/Source/Runtime/Engine/Private/Tests/EngineAutomationTests.cpp b/Engine/Source/Runtime/Engine/Private/Tests/EngineAutomationTests.cpp
--- a/Engine/Source/Runtime/Engine/Private/Tests/EngineAutomationTests.cpp
+++ b/Engine/Source/Runtime/Engine/Private/Tests/EngineAutomationTests.cpp
@@ -161,13 +161,13 @@ bool FStatsVerificationMapTest::RunTest(const FString& Parameters)
 			}
 			else
 			{
-				UE_LOG(LogEngineAutomationTests, Error, TEXT("Invalid asset path: %s."), *Filename);
+				UE_LOG(LogEngineAutomationTests, Error, TEXT("Invalid asset path: %s."), *MapName);
 			}
 		}
 	}
 	else
 	{
-		UE_LOG(LogEngineAutomationTests, Warning, TEXT("Automation test map doesn't exist or is not set: %s.  \nUsing the currently loaded map."), *Filename);
+		UE_LOG(LogEngineAutomationTests, Log, TEXT("Automation test map doesn't exist or is not set: %s.  \nUsing the currently loaded map."), *Filename);
 	}
 
 	ADD_LATENT_AUTOMATION_COMMAND(FExecStringLatentCommand(TEXT("stat game")));
