BP: Macro is not recompiled in ConditionalRecompileClass.
ttp342812 external object warning prevents Blueprint'/Game/Characters/Enemies/Smasher/Blueprints/SmasherPlayer.SmasherPlayer' from beingà

diff --git a/Engine/Source/Runtime/Engine/Private/BlueprintGeneratedClass.cpp b/Engine/Source/Runtime/Engine/Private/BlueprintGeneratedClass.cpp
--- a/Engine/Source/Runtime/Engine/Private/BlueprintGeneratedClass.cpp
+++ b/Engine/Source/Runtime/Engine/Private/BlueprintGeneratedClass.cpp
@@ -189,7 +189,7 @@ void UBlueprintGeneratedClass::ConditionalRecompileClass(TArray<UObject*>* ObjLo
 			// Make sure that nodes are up to date, so that we get any updated blueprint signatures
 			FBlueprintEditorUtils::RefreshExternalBlueprintDependencyNodes(GeneratingBP);
 
-			if (GeneratingBP->Status != BS_Error)
+			if ((GeneratingBP->Status != BS_Error) && (GeneratingBP->BlueprintType != EBlueprintType::BPTYPE_MacroLibrary))
 			{
 				FKismetEditorUtilities::RecompileBlueprintBytecode(GeneratingBP, ObjLoaded);
 			}
