[345538] Blueprints - Fix for a crash that could occur after loading a BP asset with a particular arrangement of linked nodes as a dependency of another BP asset.

- Modified UEdGraphPin::CopyPersistentDataFromOldPin() to ensure that the other pin's LinkedTo array remains in the same order as before after the old pin is removed during node reconstruction.
- Added test assets to QAGame for validation purposes.

diff --git a/Engine/Source/Runtime/Engine/Private/EdGraph/EdGraphPin.cpp b/Engine/Source/Runtime/Engine/Private/EdGraph/EdGraphPin.cpp
--- a/Engine/Source/Runtime/Engine/Private/EdGraph/EdGraphPin.cpp
+++ b/Engine/Source/Runtime/Engine/Private/EdGraph/EdGraphPin.cpp
@@ -169,11 +169,25 @@ void UEdGraphPin::CopyPersistentDataFromOldPin(const UEdGraphPin& SourcePin)
 	{
 		UEdGraphPin* OtherPin = SourcePin.LinkedTo[LinkIndex];
 		check(NULL != OtherPin);
-		
+
 		Modify();
 		OtherPin->Modify();
 
-		OtherPin->MakeLinkTo(this);
+		LinkedTo.Add(OtherPin);
+
+		// Unlike MakeLinkTo(), we attempt to ensure that the new pin (this) is inserted at the same position as the old pin (source)
+		// in the OtherPin's LinkedTo array. This is necessary to ensure that the node's position in the execution order will remain
+		// unchanged after nodes are reconstructed, because OtherPin may be connected to more than just this node.
+		int32 Index = OtherPin->LinkedTo.Find(const_cast<UEdGraphPin*>(&SourcePin));
+		if(Index != INDEX_NONE)
+		{
+			OtherPin->LinkedTo.Insert(this, Index);
+		}
+		else
+		{
+			// Fallback to "normal" add, just in case the old pin doesn't exist in the other pin's LinkedTo array for some reason.
+			OtherPin->LinkedTo.Add(this);
+		}
 	}
 
 	// If the source pin is split, then split the new one
