Convert to using WorldContext metadata instead of DefaultToSelf/HidePin

diff --git a/Engine/Source/Developer/FunctionalTesting/Classes/FunctionalTestingManager.h b/Engine/Source/Developer/FunctionalTesting/Classes/FunctionalTestingManager.h
--- a/Engine/Source/Developer/FunctionalTesting/Classes/FunctionalTestingManager.h
+++ b/Engine/Source/Developer/FunctionalTesting/Classes/FunctionalTestingManager.h
@@ -28,7 +28,7 @@ class UFunctionalTestingManager : public UBlueprintFunctionLibrary
 	UPROPERTY(BlueprintAssignable)
 	FFunctionalTestEventSignature OnSetupTests;
 	
-	UFUNCTION(BlueprintCallable, Category="FunctionalTesting", meta=(HidePin="WorldContext", DefaultToSelf="WorldContext" ) )
+	UFUNCTION(BlueprintCallable, Category="FunctionalTesting", meta=(WorldContext="WorldContext", CallableWithoutWorldContext ) )
 	/** Triggers in sequence all functional tests found on the level.
 	 *	@return true if any tests have been triggered */
 	 static bool RunAllFunctionalTests(UObject* WorldContext, bool bNewLog = true, bool bRunLooped = false, bool bWaitForNavigationBuildFinish = true, FString FailedTestsReproString = TEXT(""));
diff --git a/Engine/Source/Editor/BlueprintGraph/Private/CallFunctionHandler.cpp b/Engine/Source/Editor/BlueprintGraph/Private/CallFunctionHandler.cpp
--- a/Engine/Source/Editor/BlueprintGraph/Private/CallFunctionHandler.cpp
+++ b/Engine/Source/Editor/BlueprintGraph/Private/CallFunctionHandler.cpp
@@ -512,14 +512,7 @@ void FKCHandler_CallFunction::RegisterNets(FKismetFunctionContext& Context, UEdG
 			FString const DefaltToSelfPinName = Function->GetMetaData(FBlueprintMetadata::MD_DefaultToSelf);
 			AutoCreateRefTermPinNames.Remove(DefaltToSelfPinName);
 
-			if (Context.Blueprint->ParentClass->HasMetaData(FBlueprintMetadata::MD_ShowWorldContextPin))
-			{
-				RequiresSetValue.Add(DefaltToSelfPinName);
-			}
-			else
-			{
-				DefaultToSelfParamNames.Add(DefaltToSelfPinName);
-			}
+			DefaultToSelfParamNames.Add(DefaltToSelfPinName);
 		}
 		if (Function->HasMetaData(FBlueprintMetadata::MD_WorldContext))
 		{
diff --git a/Engine/Source/Runtime/AIModule/Classes/Blueprint/AIBlueprintHelperLibrary.h b/Engine/Source/Runtime/AIModule/Classes/Blueprint/AIBlueprintHelperLibrary.h
--- a/Engine/Source/Runtime/AIModule/Classes/Blueprint/AIBlueprintHelperLibrary.h
+++ b/Engine/Source/Runtime/AIModule/Classes/Blueprint/AIBlueprintHelperLibrary.h
@@ -13,16 +13,16 @@ class AIMODULE_API UAIBlueprintHelperLibrary : public UBlueprintFunctionLibrary
 {
 	GENERATED_UCLASS_BODY()
 
-	UFUNCTION(BlueprintCallable, meta=(HidePin="WorldContextObject", DefaultToSelf="WorldContextObject", BlueprintInternalUseOnly = "TRUE"))
+	UFUNCTION(BlueprintCallable, meta=(WorldContext="WorldContextObject", BlueprintInternalUseOnly = "TRUE"))
 	static class UAIAsyncTaskBlueprintProxy* CreateMoveToProxyObject(class UObject* WorldContextObject, APawn* Pawn, FVector Destination, AActor* TargetActor = NULL, float AcceptanceRadius = 0, bool bStopOnOverlap = false);
 
 	UFUNCTION(BlueprintCallable, Category="AI", meta=(DefaultToSelf="MessageSource"))
 	static void SendAIMessage(APawn* Target, FName Message, UObject* MessageSource, bool bSuccess = true);
 
-	UFUNCTION(BlueprintCallable, Category="AI", meta=(HidePin="WorldContextObject", DefaultToSelf="WorldContextObject", DeprecatedFunction, DeprecationMessage="Use SpawnAIFromClass"))
+	UFUNCTION(BlueprintCallable, Category="AI", meta=(WorldContext="WorldContextObject", DeprecatedFunction, DeprecationMessage="Use SpawnAIFromClass"))
 	static class APawn* SpawnAI(class UObject* WorldContextObject, UBlueprint* Pawn, class UBehaviorTree* BehaviorTree, FVector Location, FRotator Rotation=FRotator::ZeroRotator, bool bNoCollisionFail=false);
 
-	UFUNCTION(BlueprintCallable, Category="AI", meta=(HidePin="WorldContextObject", DefaultToSelf="WorldContextObject"))
+	UFUNCTION(BlueprintCallable, Category="AI", meta=(WorldContext="WorldContextObject"))
 	static class APawn* SpawnAIFromClass(class UObject* WorldContextObject, TSubclassOf<APawn> PawnClass, class UBehaviorTree* BehaviorTree, FVector Location, FRotator Rotation=FRotator::ZeroRotator, bool bNoCollisionFail=false);
 
 	UFUNCTION(BlueprintPure, Category="AI", meta=(DefaultToSelf="Target"))
diff --git a/Engine/Source/Runtime/CoreUObject/Public/UObject/ObjectBase.h b/Engine/Source/Runtime/CoreUObject/Public/UObject/ObjectBase.h
--- a/Engine/Source/Runtime/CoreUObject/Public/UObject/ObjectBase.h
+++ b/Engine/Source/Runtime/CoreUObject/Public/UObject/ObjectBase.h
@@ -1083,6 +1083,9 @@ namespace UM
 		/// [FunctionMetadata] Used by BlueprintCallable functions to indicate that this function is not to be allowed in the Construction Script.
 		UnsafeDuringActorConstruction,
 
+		/// [FunctionMetadta] Used by BlueprintCallable functions to indicate which parameter is used to determine the World that the operation is occurring within.
+		WorldContext,
+
 	};
 
 	// Metadata usable in UINTERFACE
diff --git a/Engine/Source/Runtime/Engine/Classes/AI/Navigation/NavigationSystem.h b/Engine/Source/Runtime/Engine/Classes/AI/Navigation/NavigationSystem.h
--- a/Engine/Source/Runtime/Engine/Classes/AI/Navigation/NavigationSystem.h
+++ b/Engine/Source/Runtime/Engine/Classes/AI/Navigation/NavigationSystem.h
@@ -200,24 +200,24 @@ public:
 	// Kismet functions
 	//----------------------------------------------------------------------//
 	/** Project a point onto the NavigationData */
-	UFUNCTION(BlueprintPure, Category="AI|Navigation", meta=(HidePin="WorldContext", DefaultToSelf="WorldContext" ) )
+	UFUNCTION(BlueprintPure, Category="AI|Navigation", meta=(WorldContext="WorldContext" ) )
 	static FVector ProjectPointToNavigation(UObject* WorldContext, const FVector& Point, ANavigationData* NavData = NULL, TSubclassOf<UNavigationQueryFilter> FilterClass = NULL);
 	
-	UFUNCTION(BlueprintPure, Category="AI|Navigation", meta=(HidePin="WorldContext", DefaultToSelf="WorldContext" ) )
+	UFUNCTION(BlueprintPure, Category="AI|Navigation", meta=(WorldContext="WorldContext" ) )
 	static FVector GetRandomPoint(UObject* WorldContext, ANavigationData* NavData = NULL, TSubclassOf<UNavigationQueryFilter> FilterClass = NULL);
 
-	UFUNCTION(BlueprintPure, Category="AI|Navigation", meta=(HidePin="WorldContext", DefaultToSelf="WorldContext" ) )
+	UFUNCTION(BlueprintPure, Category="AI|Navigation", meta=(WorldContext="WorldContext" ) )
 	static FVector GetRandomPointInRadius(UObject* WorldContext, const FVector& Origin, float Radius, ANavigationData* NavData = NULL, TSubclassOf<UNavigationQueryFilter> FilterClass = NULL);
 
 	/** Potentially expensive. Use with caution. Consider using UPathFollowingComponent::GetRemainingPathCost instead */
-	UFUNCTION(BlueprintPure, Category="AI|Navigation", meta=(HidePin="WorldContext", DefaultToSelf="WorldContext" ) )
+	UFUNCTION(BlueprintPure, Category="AI|Navigation", meta=(WorldContext="WorldContext" ) )
 	static ENavigationQueryResult::Type GetPathCost(UObject* WorldContext, const FVector& PathStart, const FVector& PathEnd, float& PathCost, ANavigationData* NavData = NULL, TSubclassOf<UNavigationQueryFilter> FilterClass = NULL);
 
 	/** Potentially expensive. Use with caution */
-	UFUNCTION(BlueprintPure, Category="AI|Navigation", meta=(HidePin="WorldContext", DefaultToSelf="WorldContext" ) )
+	UFUNCTION(BlueprintPure, Category="AI|Navigation", meta=(WorldContext="WorldContext" ) )
 	static ENavigationQueryResult::Type GetPathLength(UObject* WorldContext, const FVector& PathStart, const FVector& PathEnd, float& PathLength, ANavigationData* NavData = NULL, TSubclassOf<UNavigationQueryFilter> FilterClass = NULL);
 
-	UFUNCTION(BlueprintPure, Category="AI|Navigation", meta=(HidePin="WorldContext", DefaultToSelf="WorldContext" ) )
+	UFUNCTION(BlueprintPure, Category="AI|Navigation", meta=(WorldContext="WorldContext" ) )
 	static bool IsNavigationBeingBuilt(UObject* WorldContext);
 
 	UFUNCTION(BlueprintCallable, Category="AI|Navigation")
@@ -241,7 +241,7 @@ public:
 	 *	@param Querier if not passed default navigation data will be used
 	 *	@param HitLocation if line was obstructed this will be set to hit location. Otherwise it contains SegmentEnd
 	 *	@return true if line from RayStart to RayEnd was obstructed. Also, true when no navigation data present */
-	UFUNCTION(BlueprintCallable, Category="AI|Navigation", meta=(HidePin="WorldContext", DefaultToSelf="WorldContext" ))
+	UFUNCTION(BlueprintCallable, Category="AI|Navigation", meta=(WorldContext="WorldContext" ))
 	static bool NavigationRaycast(UObject* WorldContext, const FVector& RayStart, const FVector& RayEnd, FVector& HitLocation, TSubclassOf<UNavigationQueryFilter> FilterClass = NULL, AController* Querier = NULL);
 
 	/** delegate type for events that dirty the navigation data ( Params: const FBox& DirtyBounds ) */
diff --git a/Engine/Source/Runtime/Engine/Classes/Kismet/GameplayStatics.h b/Engine/Source/Runtime/Engine/Classes/Kismet/GameplayStatics.h
--- a/Engine/Source/Runtime/Engine/Classes/Kismet/GameplayStatics.h
+++ b/Engine/Source/Runtime/Engine/Classes/Kismet/GameplayStatics.h
@@ -33,11 +33,11 @@ class ENGINE_API UGameplayStatics : public UBlueprintFunctionLibrary
 	// --- Spawning functions ------------------------------
 
 	/** Spawns an instance of a blueprint, but does not automatically run its construction script.  */
-	UFUNCTION(BlueprintCallable, Category="Spawning", meta=(HidePin="WorldContextObject", DefaultToSelf="WorldContextObject", UnsafeDuringActorConstruction = "true", BlueprintInternalUseOnly = "true", DeprecatedFunction, DeprecationMessage="Use BeginSpawningActorFromClass"))
+	UFUNCTION(BlueprintCallable, Category="Spawning", meta=(WorldContext="WorldContextObject", UnsafeDuringActorConstruction = "true", BlueprintInternalUseOnly = "true", DeprecatedFunction, DeprecationMessage="Use BeginSpawningActorFromClass"))
 	static class AActor* BeginSpawningActorFromBlueprint(UObject* WorldContextObject, const class UBlueprint* Blueprint, const FTransform& SpawnTransform, bool bNoCollisionFail);
 
 	/** Spawns an instance of an actor class, but does not automatically run its construction script.  */
-	UFUNCTION(BlueprintCallable, Category="Spawning", meta=(HidePin="WorldContextObject", DefaultToSelf="WorldContextObject", UnsafeDuringActorConstruction = "true", BlueprintInternalUseOnly = "true"))
+	UFUNCTION(BlueprintCallable, Category="Spawning", meta=(WorldContext="WorldContextObject", UnsafeDuringActorConstruction = "true", BlueprintInternalUseOnly = "true"))
 	static class AActor* BeginSpawningActorFromClass(UObject* WorldContextObject, TSubclassOf<AActor> ActorClass, const FTransform& SpawnTransform, bool bNoCollisionFail = false);
 
 	/** 'Finish' spawning an actor.  This will run the construction script. */
@@ -60,7 +60,7 @@ class ENGINE_API UGameplayStatics : public UBlueprintFunctionLibrary
 	 *	@param	ActorClass	Class of Actor to find. Must be specified or result array will be empty.
 	 *	@param	OutActors	Output array of Actors of the specified class.
 	 */
-	UFUNCTION(BlueprintCallable, Category="Utilities",  meta=(HidePin="WorldContextObject", DefaultToSelf="WorldContextObject"))
+	UFUNCTION(BlueprintCallable, Category="Utilities",  meta=(WorldContext="WorldContextObject"))
 	static void GetAllActorsOfClass(UObject* WorldContextObject, TSubclassOf<AActor> ActorClass, TArray<AActor*>& OutActors);
 
 	/** 
@@ -69,50 +69,50 @@ class ENGINE_API UGameplayStatics : public UBlueprintFunctionLibrary
 	 *	@param	Interface	Interface to find. Must be specified or result array will be empty.
 	 *	@param	OutActors	Output array of Actors of the specified class.
 	 */
-	UFUNCTION(BlueprintCallable, Category="Utilities",  meta=(HidePin="WorldContextObject", DefaultToSelf="WorldContextObject"))
+	UFUNCTION(BlueprintCallable, Category="Utilities",  meta=(WorldContext="WorldContextObject"))
 	static void GetAllActorsWithInterface(UObject* WorldContextObject, TSubclassOf<UInterface> Interface, TArray<AActor*>& OutActors);
 
 	// --- Player functions ------------------------------
 
 	/** Returns the game instance object  */
-	UFUNCTION(BlueprintPure, Category="Game", meta=(HidePin="WorldContextObject", DefaultToSelf="WorldContextObject"))
+	UFUNCTION(BlueprintPure, Category="Game", meta=(WorldContext="WorldContextObject"))
 	static class UGameInstance* GetGameInstance(UObject* WorldContextObject);
 
 	/** Returns the player controller at the specified player index */
-	UFUNCTION(BlueprintPure, Category="Game", meta=(HidePin="WorldContextObject", DefaultToSelf="WorldContextObject", UnsafeDuringActorConstruction="true"))
+	UFUNCTION(BlueprintPure, Category="Game", meta=(WorldContext="WorldContextObject", UnsafeDuringActorConstruction="true"))
 	static class APlayerController* GetPlayerController(UObject* WorldContextObject, int32 PlayerIndex);
 
 	/** Returns the player pawn at the specified player index */
-	UFUNCTION(BlueprintPure, Category="Game", meta=(HidePin="WorldContextObject", DefaultToSelf="WorldContextObject", UnsafeDuringActorConstruction="true"))
+	UFUNCTION(BlueprintPure, Category="Game", meta=(WorldContext="WorldContextObject", UnsafeDuringActorConstruction="true"))
 	static class APawn* GetPlayerPawn(UObject* WorldContextObject, int32 PlayerIndex);
 
 	/** Returns the player character (NULL if the player pawn doesn't exist OR is not a character) at the specified player index */
-	UFUNCTION(BlueprintPure, Category="Game", meta=(HidePin="WorldContextObject", DefaultToSelf="WorldContextObject", UnsafeDuringActorConstruction="true"))
+	UFUNCTION(BlueprintPure, Category="Game", meta=(WorldContext="WorldContextObject", UnsafeDuringActorConstruction="true"))
 	static class ACharacter* GetPlayerCharacter(UObject* WorldContextObject, int32 PlayerIndex);
 
 	/** Returns the player's camera manager for the specified player index */
-	UFUNCTION(BlueprintPure, Category="Game", meta=(HidePin="WorldContextObject", DefaultToSelf="WorldContextObject", UnsafeDuringActorConstruction="true"))
+	UFUNCTION(BlueprintPure, Category="Game", meta=(WorldContext="WorldContextObject", UnsafeDuringActorConstruction="true"))
 	static class APlayerCameraManager* GetPlayerCameraManager(UObject* WorldContextObject, int32 PlayerIndex);
 
 	/** Create a new player for this game.  
 	 *  @param ControllerId		The ID of the controller that the should control the newly created player.  A value of -1 specifies to use the next available ID
 	 *  @param bSpawnPawn		Whether a pawn should be spawned immediately. If false a pawn will not be created until transition to the next map.
 	 */
-	UFUNCTION(BlueprintCallable, Category="Game", meta=(HidePin="WorldContextObject", DefaultToSelf="WorldContextObject", AdvancedDisplay="2", UnsafeDuringActorConstruction="true"))
+	UFUNCTION(BlueprintCallable, Category="Game", meta=(WorldContext="WorldContextObject", AdvancedDisplay="2", UnsafeDuringActorConstruction="true"))
 	static class APlayerController* CreatePlayer(UObject* WorldContextObject, int32 ControllerId = -1, bool bSpawnPawn = true);
 
 	// --- Level Streaming functions ------------------------
 	
 	/** Stream the level with the LevelName ; Calling again before it finishes has no effect */
-	UFUNCTION(BlueprintCallable, meta=(HidePin="WorldContextObject", DefaultToSelf="WorldContextObject", Latent = "", LatentInfo = "LatentInfo"), Category="Game")
+	UFUNCTION(BlueprintCallable, meta=(WorldContext="WorldContextObject", Latent = "", LatentInfo = "LatentInfo"), Category="Game")
 	static void LoadStreamLevel(UObject* WorldContextObject, FName LevelName, bool bMakeVisibleAfterLoad, bool bShouldBlockOnLoad, FLatentActionInfo LatentInfo);
 
 	/** Unload a streamed in level */
-	UFUNCTION(BlueprintCallable, meta=(HidePin="WorldContextObject", DefaultToSelf="WorldContextObject", Latent = "", LatentInfo = "LatentInfo"), Category="Game")
+	UFUNCTION(BlueprintCallable, meta=(WorldContext="WorldContextObject", Latent = "", LatentInfo = "LatentInfo"), Category="Game")
 	static void UnloadStreamLevel(UObject* WorldContextObject, FName LevelName, FLatentActionInfo LatentInfo);
 	
 	/** Returns level streaming object with specified level package name */
-	UFUNCTION(BlueprintPure, meta=(HidePin="WorldContextObject", DefaultToSelf="WorldContextObject"), Category="Game")
+	UFUNCTION(BlueprintPure, meta=(WorldContext="WorldContextObject"), Category="Game")
 	static class ULevelStreaming* GetStreamingLevel(UObject* WorldContextObject, FName PackageName);
 	
 	/**
@@ -122,16 +122,16 @@ class ENGINE_API UGameplayStatics : public UBlueprintFunctionLibrary
 	 * @param	bAbsolute			if true options are reset, if false options are carried over from current level
 	 * @param	Options				a string of options to use for the travel URL
 	 */
-	UFUNCTION(BlueprintCallable, meta=(HidePin="WorldContextObject", DefaultToSelf="WorldContextObject", AdvancedDisplay = "2"), Category="Game")
+	UFUNCTION(BlueprintCallable, meta=(WorldContext="WorldContextObject", AdvancedDisplay = "2"), Category="Game")
 	static void OpenLevel(UObject* WorldContextObject, FName LevelName, bool bAbsolute = true, FString Options = FString(TEXT("")));
 
 	// --- Global functions ------------------------------
 
 	/** Returns the current GameMode or NULL if the GameMode can't be retrieved */
-	UFUNCTION(BlueprintPure, Category="Game", meta=(HidePin="WorldContextObject", DefaultToSelf="WorldContextObject") )
+	UFUNCTION(BlueprintPure, Category="Game", meta=(WorldContext="WorldContextObject") )
 	static class AGameMode* GetGameMode(UObject* WorldContextObject);
 
-	UFUNCTION(BlueprintPure, Category="Game", meta=(HidePin="WorldContextObject", DefaultToSelf="WorldContextObject") )
+	UFUNCTION(BlueprintPure, Category="Game", meta=(WorldContext="WorldContextObject") )
 	static class AGameState* GetGameState(UObject* WorldContextObject);
 
 	UFUNCTION(BlueprintPure, meta=(FriendlyName = "GetClass"), Category="Utilities")
@@ -141,7 +141,7 @@ class ENGINE_API UGameplayStatics : public UBlueprintFunctionLibrary
 	 * Sets the global time dilation
 	 * @param	TimeDilation	value to set the global time dilation to
 	 */
-	UFUNCTION(BlueprintCallable, Category="Utilities|Time", meta=(HidePin="WorldContextObject", DefaultToSelf="WorldContextObject") )
+	UFUNCTION(BlueprintCallable, Category="Utilities|Time", meta=(WorldContext="WorldContextObject") )
 	static void SetGlobalTimeDilation(UObject* WorldContextObject, float TimeDilation);
 
 	/**
@@ -149,14 +149,14 @@ class ENGINE_API UGameplayStatics : public UBlueprintFunctionLibrary
 	 * @param	bPaused		Whether the game should be paused or not
 	 * @return	Whether the game was successfully paused/unpaused
 	 */
-	UFUNCTION(BlueprintCallable, Category="Game", meta=(HidePin="WorldContextObject", DefaultToSelf="WorldContextObject") )
+	UFUNCTION(BlueprintCallable, Category="Game", meta=(WorldContext="WorldContextObject") )
 	static bool SetGamePaused(UObject* WorldContextObject, bool bPaused);
 
 	/**
 	 * Returns the game's paused state
 	 * @return	Whether the game is currently paused or not
 	 */
-	UFUNCTION(BlueprintPure, Category="Game", meta=(HidePin="WorldContextObject", DefaultToSelf="WorldContextObject") )
+	UFUNCTION(BlueprintPure, Category="Game", meta=(WorldContext="WorldContextObject") )
 	static bool IsGamePaused(UObject* WorldContextObject);
 
 
@@ -171,7 +171,7 @@ class ENGINE_API UGameplayStatics : public UBlueprintFunctionLibrary
 	 * @param DamagePreventionChannel - Damage will not be applied to victim if there is something between the origin and the victim which blocks traces on this channel
 	 * @return true if damage was applied to at least one actor.
 	 */
-	UFUNCTION(BlueprintCallable, BlueprintAuthorityOnly, Category="Game|Damage", meta=(HidePin="WorldContextObject", DefaultToSelf="WorldContextObject", AutoCreateRefTerm="IgnoreActors"))
+	UFUNCTION(BlueprintCallable, BlueprintAuthorityOnly, Category="Game|Damage", meta=(WorldContext="WorldContextObject", AutoCreateRefTerm="IgnoreActors"))
 		static bool ApplyRadialDamage(UObject* WorldContextObject, float BaseDamage, const FVector& Origin, float DamageRadius, TSubclassOf<class UDamageType> DamageTypeClass, const TArray<AActor*>& IgnoreActors, AActor* DamageCauser = NULL, AController* InstigatedByController = NULL, bool bDoFullDamage = false, ECollisionChannel DamagePreventionChannel = ECC_Visibility);
 	
 	/** Hurt locally authoritative actors within the radius. Will only hit components that block the Visibility channel.
@@ -187,7 +187,7 @@ class ENGINE_API UGameplayStatics : public UBlueprintFunctionLibrary
 	 * @param DamagePreventionChannel - Damage will not be applied to victim if there is something between the origin and the victim which blocks traces on this channel
 	 * @return true if damage was applied to at least one actor.
 	 */
-	UFUNCTION(BlueprintCallable, BlueprintAuthorityOnly, Category="Game|Damage", meta=(HidePin="WorldContextObject", DefaultToSelf="WorldContextObject", AutoCreateRefTerm="IgnoreActors"))
+	UFUNCTION(BlueprintCallable, BlueprintAuthorityOnly, Category="Game|Damage", meta=(WorldContext="WorldContextObject", AutoCreateRefTerm="IgnoreActors"))
 		static bool ApplyRadialDamageWithFalloff(UObject* WorldContextObject, float BaseDamage, float MinimumDamage, const FVector& Origin, float DamageInnerRadius, float DamageOuterRadius, float DamageFalloff, TSubclassOf<class UDamageType> DamageTypeClass, const TArray<AActor*>& IgnoreActors, AActor* DamageCauser = NULL, AController* InstigatedByController = NULL, ECollisionChannel DamagePreventionChannel = ECC_Visibility);
 	
 
@@ -224,7 +224,7 @@ class ENGINE_API UGameplayStatics : public UBlueprintFunctionLibrary
 	 * @param Falloff - Affects falloff of effect as it nears OuterRadius
 	 * @param bOrientShakeTowardsEpicenter - Changes the rotation of shake to point towards epicenter instead of forward
 	 */
-	UFUNCTION(BlueprintCallable, Category="Game|Feedback", meta=(HidePin="WorldContextObject", DefaultToSelf="WorldContextObject", UnsafeDuringActorConstruction = "true"))
+	UFUNCTION(BlueprintCallable, Category="Game|Feedback", meta=(WorldContext="WorldContextObject", UnsafeDuringActorConstruction = "true"))
 	static void PlayWorldCameraShake(UObject* WorldContextObject, TSubclassOf<class UCameraShake> Shake, FVector Epicenter, float InnerRadius, float OuterRadius, float Falloff = 1.f, bool bOrientShakeTowardsEpicenter = false);
 
 	// --- Particle functions ------------------------------
@@ -236,7 +236,7 @@ class ENGINE_API UGameplayStatics : public UBlueprintFunctionLibrary
 	 * @param Rotation - rotation to place the effect in world space	
 	 * @param bAutoDestroy - Whether the component will automatically be destroyed when the particle system completes playing or whether it can be reactivated
 	 */
-	UFUNCTION(BlueprintCallable, Category="Effects|Components|ParticleSystem", meta=(Keywords = "particle system", HidePin="WorldContextObject", DefaultToSelf="WorldContextObject", UnsafeDuringActorConstruction = "true"))
+	UFUNCTION(BlueprintCallable, Category="Effects|Components|ParticleSystem", meta=(Keywords = "particle system", WorldContext="WorldContextObject", UnsafeDuringActorConstruction = "true"))
 	static UParticleSystemComponent* SpawnEmitterAtLocation(UObject* WorldContextObject, class UParticleSystem* EmitterTemplate, FVector Location, FRotator Rotation = FRotator::ZeroRotator, bool bAutoDestroy = true);
 
 	/** Plays the specified effect attached to and following the specified component. The system will go away when the effect is complete. Does not replicate.
@@ -270,7 +270,7 @@ class ENGINE_API UGameplayStatics : public UBlueprintFunctionLibrary
 	 * @param PitchMultiplier - PitchMultiplier
 	 * @param AttenuationSettings - Override attenuation settings package to play sound with
 	 */
-	UFUNCTION(BlueprintCallable, Category="Audio", meta=(HidePin="WorldContextObject", DefaultToSelf="WorldContextObject", AdvancedDisplay = "3", UnsafeDuringActorConstruction = "true"))
+	UFUNCTION(BlueprintCallable, Category="Audio", meta=(WorldContext="WorldContextObject", AdvancedDisplay = "3", UnsafeDuringActorConstruction = "true"))
 	static void PlaySoundAtLocation(UObject* WorldContextObject, class USoundBase* Sound, FVector Location, float VolumeMultiplier = 1.f, float PitchMultiplier = 1.f, float StartTime = 0.f, class USoundAttenuation* AttenuationSettings = NULL);
 
 	/** Plays a dialogue at the given location. This is a fire and forget sound and does not travel with any actor. Replication is also not handled at this point.
@@ -281,7 +281,7 @@ class ENGINE_API UGameplayStatics : public UBlueprintFunctionLibrary
 	 * @param PitchMultiplier - PitchMultiplier
 	 * @param AttenuationSettings - Override attenuation settings package to play sound with
 	 */
-	UFUNCTION(BlueprintCallable, Category="Audio", meta=(HidePin="WorldContextObject", DefaultToSelf="WorldContextObject", AdvancedDisplay = "3", UnsafeDuringActorConstruction = "true"))
+	UFUNCTION(BlueprintCallable, Category="Audio", meta=(WorldContext="WorldContextObject", AdvancedDisplay = "3", UnsafeDuringActorConstruction = "true"))
 	static void PlayDialogueAtLocation(UObject* WorldContextObject, class UDialogueWave* Dialogue, const struct FDialogueContext& Context, FVector Location, float VolumeMultiplier = 1.f, float PitchMultiplier = 1.f, float StartTime = 0.f, class USoundAttenuation* AttenuationSettings = NULL);
 
 	/** Plays a sound attached to and following the specified component. This is a fire and forget sound. Replication is also not handled at this point.
@@ -313,7 +313,7 @@ class ENGINE_API UGameplayStatics : public UBlueprintFunctionLibrary
 	static class UAudioComponent* PlayDialogueAttached(class UDialogueWave* Dialogue, const struct FDialogueContext& Context, class USceneComponent* AttachToComponent, FName AttachPointName = NAME_None, FVector Location = FVector(ForceInit), EAttachLocation::Type LocationType = EAttachLocation::KeepRelativeOffset, bool bStopWhenAttachedToDestroyed = false, float VolumeMultiplier = 1.f, float PitchMultiplier = 1.f, float StartTime = 0.f, class USoundAttenuation* AttenuationSettings = NULL);
 
 	/** DEPRECATED - Use GameplayStatics.PlaySoundAtLocation or GameplayStatics.PlaySoundAttached instead. */
-	UFUNCTION(BlueprintCallable, Category="Audio", meta=(HidePin="WorldContextObject", DefaultToSelf="WorldContextObject", DeprecatedFunction, DeprecationMessage="Use GameplayStatics.PlaySoundAtLocation or GameplayStatics.PlaySoundAttached instead."))
+	UFUNCTION(BlueprintCallable, Category="Audio", meta=(WorldContext="WorldContextObject", DeprecatedFunction, DeprecationMessage="Use GameplayStatics.PlaySoundAtLocation or GameplayStatics.PlaySoundAttached instead."))
 	static void PlaySound(UObject* WorldContextObject, class USoundCue* InSoundCue, class USceneComponent* AttachComponent, FName AttachName=NAME_None, bool bFollow=false, float VolumeMultiplier=1.f, float PitchMultiplier=1.f);
 
 	// --- Audio Functions ----------------------------
@@ -364,7 +364,7 @@ class ENGINE_API UGameplayStatics : public UBlueprintFunctionLibrary
 	 * @param Rotation - rotation to place the decal in world space	
 	 * @param LifeSpan - destroy decal component after time runs out (0 = infinite)
 	 */
-	UFUNCTION(BlueprintCallable, Category="Rendering|Components|Decal", meta=(HidePin="WorldContextObject", DefaultToSelf="WorldContextObject", UnsafeDuringActorConstruction = "true"))
+	UFUNCTION(BlueprintCallable, Category="Rendering|Components|Decal", meta=(WorldContext="WorldContextObject", UnsafeDuringActorConstruction = "true"))
 	static UDecalComponent* SpawnDecalAtLocation(UObject* WorldContextObject, class UMaterialInterface* DecalMaterial, FVector DecalSize, FVector Location, FRotator Rotation = FRotator(-90, 0, 0), float LifeSpan = 0);
 
 	/** Spawns a decal attached to and following the specified component. Does not replicate.
@@ -460,14 +460,14 @@ class ENGINE_API UGameplayStatics : public UBlueprintFunctionLibrary
 	static float GetWorldDeltaSeconds(UObject* WorldContextObject);
 
 	/** Returns time in seconds since world was brought up for play, does NOT stop when game pauses, NOT dilated/clamped */
-	UFUNCTION(BlueprintPure, Category="Utilities|Time", meta=(HidePin="WorldContextObject", DefaultToSelf="WorldContextObject"))
+	UFUNCTION(BlueprintPure, Category="Utilities|Time", meta=(WorldContext="WorldContextObject"))
 	static float GetRealTimeSeconds(UObject* WorldContextObject);
 	
 	/** Returns time in seconds since world was brought up for play, IS stopped when game pauses, NOT dilated/clamped. */
-	UFUNCTION(BlueprintPure, Category="Utilities|Time", meta=(HidePin="WorldContextObject", DefaultToSelf="WorldContextObject"))
+	UFUNCTION(BlueprintPure, Category="Utilities|Time", meta=(WorldContext="WorldContextObject"))
 	static float GetAudioTimeSeconds(UObject* WorldContextObject);
 
-	UFUNCTION(BlueprintPure, Category="Utilities|Time", meta=(HidePin="WorldContextObject", DefaultToSelf="WorldContextObject"))
+	UFUNCTION(BlueprintPure, Category="Utilities|Time", meta=(WorldContext="WorldContextObject"))
 	static void GetAccurateRealTime(UObject* WorldContextObject, int32& Seconds, float& PartialSeconds);
 
 	/** DVRStreaming API */
@@ -498,18 +498,18 @@ class ENGINE_API UGameplayStatics : public UBlueprintFunctionLibrary
 	 * @param bDrawDebug		When true, a debug arc is drawn (red for an invalid arc, green for a valid arc)
 	 * @return					Returns false if there is no valid solution or the valid solutions are blocked.  Returns true otherwise.
 	 */
-	UFUNCTION(BlueprintCallable, Category="Game|Components|ProjectileMovement", FriendlyName="SuggestProjectileVelocity", meta=(HidePin="WorldContextObject", DefaultToSelf="WorldContextObject"))
+	UFUNCTION(BlueprintCallable, Category="Game|Components|ProjectileMovement", FriendlyName="SuggestProjectileVelocity", meta=(WorldContext="WorldContextObject"))
 	static bool BlueprintSuggestProjectileVelocity(UObject* WorldContextObject, FVector& TossVelocity, FVector StartLocation, FVector EndLocation, float LaunchSpeed, float OverrideGravityZ, ESuggestProjVelocityTraceOption::Type TraceOption, float CollisionRadius, bool bFavorHighArc, bool bDrawDebug);
 
 	/** Native version, has more options than the Blueprint version. */
 	static bool SuggestProjectileVelocity(UObject* WorldContextObject, FVector& TossVelocity, FVector StartLocation, FVector EndLocation, float TossSpeed, bool bHighArc = false, float CollisionRadius = 0.f, float OverrideGravityZ = 0, ESuggestProjVelocityTraceOption::Type TraceOption = ESuggestProjVelocityTraceOption::TraceFullPath, const FCollisionResponseParams& ResponseParam = FCollisionResponseParams::DefaultResponseParam, const TArray<AActor*>& ActorsToIgnore = TArray<AActor*>(), bool bDrawDebug = false);
 
 	/** Returns world origin current location */
-	UFUNCTION(BlueprintPure, Category="Game", meta=(HidePin="WorldContextObject", DefaultToSelf="WorldContextObject") )
+	UFUNCTION(BlueprintPure, Category="Game", meta=(WorldContext="WorldContextObject") )
 	static FIntVector GetWorldOriginLocation(UObject* WorldContextObject);
 	
 	/** Requests a new location for a world origin */
-	UFUNCTION(BlueprintCallable, Category="Game", meta=(HidePin="WorldContextObject", DefaultToSelf="WorldContextObject"))
+	UFUNCTION(BlueprintCallable, Category="Game", meta=(WorldContext="WorldContextObject"))
 	static void SetWorldOriginLocation(UObject* WorldContextObject, FIntVector NewLocation);
 };
 
diff --git a/Engine/Source/Runtime/Engine/Classes/Kismet/KismetMaterialLibrary.h b/Engine/Source/Runtime/Engine/Classes/Kismet/KismetMaterialLibrary.h
--- a/Engine/Source/Runtime/Engine/Classes/Kismet/KismetMaterialLibrary.h
+++ b/Engine/Source/Runtime/Engine/Classes/Kismet/KismetMaterialLibrary.h
@@ -10,22 +10,22 @@ class UKismetMaterialLibrary : public UBlueprintFunctionLibrary
 	GENERATED_UCLASS_BODY()
 
 	/** Sets a scalar parameter value on the material collection instance. Logs if ParameterName is invalid. */
-	UFUNCTION(BlueprintCallable, Category="Rendering|Material", meta=(Keywords="SetFloatParameterValue", HidePin="WorldContextObject", DefaultToSelf="WorldContextObject", MaterialParameterCollectionFunction = "true"))
+	UFUNCTION(BlueprintCallable, Category="Rendering|Material", meta=(Keywords="SetFloatParameterValue", WorldContext="WorldContextObject", MaterialParameterCollectionFunction = "true"))
 	static ENGINE_API void SetScalarParameterValue(UObject* WorldContextObject, UMaterialParameterCollection* Collection, FName ParameterName, float ParameterValue);
 
 	/** Sets a vector parameter value on the material collection instance. Logs if ParameterName is invalid. */
-	UFUNCTION(BlueprintCallable, Category="Rendering|Material", meta=(Keywords="SetColorParameterValue", HidePin="WorldContextObject", DefaultToSelf="WorldContextObject", MaterialParameterCollectionFunction = "true"))
+	UFUNCTION(BlueprintCallable, Category="Rendering|Material", meta=(Keywords="SetColorParameterValue", WorldContext="WorldContextObject", MaterialParameterCollectionFunction = "true"))
 	static ENGINE_API void SetVectorParameterValue(UObject* WorldContextObject, UMaterialParameterCollection* Collection, FName ParameterName, const FLinearColor& ParameterValue);
 
 	/** Gets a scalar parameter value from the material collection instance. Logs if ParameterName is invalid. */
-	UFUNCTION(BlueprintCallable, Category="Rendering|Material", meta=(Keywords="GetFloatParameterValue", HidePin="WorldContextObject", DefaultToSelf="WorldContextObject", MaterialParameterCollectionFunction = "true"))
+	UFUNCTION(BlueprintCallable, Category="Rendering|Material", meta=(Keywords="GetFloatParameterValue", WorldContext="WorldContextObject", MaterialParameterCollectionFunction = "true"))
 	static ENGINE_API float GetScalarParameterValue(UObject* WorldContextObject, UMaterialParameterCollection* Collection, FName ParameterName);
 
 	/** Gets a vector parameter value from the material collection instance. Logs if ParameterName is invalid. */
-	UFUNCTION(BlueprintCallable, Category="Rendering|Material", meta=(Keywords="GetColorParameterValue", HidePin="WorldContextObject", DefaultToSelf="WorldContextObject", MaterialParameterCollectionFunction = "true"))
+	UFUNCTION(BlueprintCallable, Category="Rendering|Material", meta=(Keywords="GetColorParameterValue", WorldContext="WorldContextObject", MaterialParameterCollectionFunction = "true"))
 	static ENGINE_API FLinearColor GetVectorParameterValue(UObject* WorldContextObject, UMaterialParameterCollection* Collection, FName ParameterName);
 
 	/** Creates a Dynamic Material Instance which you can modify during gameplay. */
-	UFUNCTION(BlueprintCallable, Category="Rendering|Material", meta=(HidePin="WorldContextObject", DefaultToSelf="WorldContextObject", MaterialParameterCollectionFunction = "true"))
+	UFUNCTION(BlueprintCallable, Category="Rendering|Material", meta=(WorldContext="WorldContextObject", MaterialParameterCollectionFunction = "true"))
 	static ENGINE_API class UMaterialInstanceDynamic* CreateDynamicMaterialInstance(UObject* WorldContextObject, class UMaterialInterface* Parent);
 };
diff --git a/Engine/Source/Runtime/Engine/Classes/Kismet/KismetMathLibrary.h b/Engine/Source/Runtime/Engine/Classes/Kismet/KismetMathLibrary.h
--- a/Engine/Source/Runtime/Engine/Classes/Kismet/KismetMathLibrary.h
+++ b/Engine/Source/Runtime/Engine/Classes/Kismet/KismetMathLibrary.h
@@ -1412,7 +1412,7 @@ class UKismetMathLibrary : public UBlueprintFunctionLibrary
 	 * @outparam	OutRectSideA - Vector oriented and sized to represent one edge of the enclosing rectangle, orthogonal to OutRectSideB
 	 * @outparam	OutRectSideB - Vector oriented and sized to represent one edge of the enclosing rectangle, orthogonal to OutRectSideA
 	*/
-	UFUNCTION(BlueprintCallable, BlueprintAuthorityOnly, Category="Math|Geometry", meta=(HidePin="WorldContextObject", DefaultToSelf="WorldContextObject"))
+	UFUNCTION(BlueprintCallable, BlueprintAuthorityOnly, Category="Math|Geometry", meta=(WorldContext="WorldContextObject"))
 	static void MinimumAreaRectangle(UObject* WorldContextObject, const TArray<FVector>& InVerts, const FVector& SampleSurfaceNormal, FVector& OutRectCenter, FRotator& OutRectRotation, float& OutSideLengthX, float& OutSideLengthY, bool bDebugDraw = false);
 
 	/**
diff --git a/Engine/Source/Runtime/Engine/Classes/Kismet/KismetSystemLibrary.h b/Engine/Source/Runtime/Engine/Classes/Kismet/KismetSystemLibrary.h
--- a/Engine/Source/Runtime/Engine/Classes/Kismet/KismetSystemLibrary.h
+++ b/Engine/Source/Runtime/Engine/Classes/Kismet/KismetSystemLibrary.h
@@ -100,15 +100,15 @@ class UKismetSystemLibrary : public UBlueprintFunctionLibrary
 	 * 
 	 * @param WorldContextObject	World context
 	 */
-	UFUNCTION(BlueprintPure, Category="Utilities|Time", meta=(HidePin="WorldContextObject", DefaultToSelf="WorldContextObject") )
+	UFUNCTION(BlueprintPure, Category="Utilities|Time", meta=(WorldContext="WorldContextObject") )
 	static float GetGameTimeInSeconds(UObject* WorldContextObject);
 
 	/** Returns whether the world this object is in is the host or not */
-	UFUNCTION(BlueprintPure, Category="Networking", meta=(HidePin="WorldContextObject", DefaultToSelf="WorldContextObject") )
+	UFUNCTION(BlueprintPure, Category="Networking", meta=(WorldContext="WorldContextObject") )
 	static bool IsServer(UObject* WorldContextObject);
 
 	/** Returns whether this is running on a dedicated server */
-	UFUNCTION(BlueprintPure, Category="Networking", meta=(HidePin="WorldContextObject", DefaultToSelf="WorldContextObject"))
+	UFUNCTION(BlueprintPure, Category="Networking", meta=(WorldContext="WorldContextObject"))
 	static bool IsDedicatedServer(UObject* WorldContextObject);
 
 	/** Returns whether this is a build that is packaged for distribution */
@@ -188,7 +188,7 @@ class UKismetSystemLibrary : public UBlueprintFunctionLibrary
 	 * @param	bPrintToConsole	Whether or not to print the output to the console
 	 * @param	TextColor		Whether or not to print the output to the console
 	 */
-	UFUNCTION(BlueprintCallable, meta=(HidePin="WorldContextObject", DefaultToSelf="WorldContextObject", Keywords = "log", AdvancedDisplay = "2"), Category="Utilities|String")
+	UFUNCTION(BlueprintCallable, meta=(WorldContext="WorldContextObject", CallableWithoutWorldContext, Keywords = "log", AdvancedDisplay = "2"), Category="Utilities|String")
 	static void PrintString(UObject* WorldContextObject, const FString& InString = FString(TEXT("Hello")), bool bPrintToScreen = true, bool bPrintToLog = true, FLinearColor TextColor = FLinearColor(0.0,0.66,1.0));
 
 	/**
@@ -207,14 +207,14 @@ class UKismetSystemLibrary : public UBlueprintFunctionLibrary
 	 * @param	Command			Command to send to the console
 	 * @param	SpecificPlayer	If specified, the console command will be routed through the specified player
 	 */
-	UFUNCTION(BlueprintCallable, Category="Development",meta=(HidePin="WorldContextObject", DefaultToSelf="WorldContextObject"))
+	UFUNCTION(BlueprintCallable, Category="Development",meta=(WorldContext="WorldContextObject"))
 	static void ExecuteConsoleCommand(UObject* WorldContextObject, const FString& Command, class APlayerController* SpecificPlayer = NULL );
 
 	/** 
 	 *	Exit the current game 
 	 * @param	SpecificPlayer	The specific player to quit the game. If not specified, player 0 will quit.
 	 */
-	UFUNCTION(BlueprintCallable, Category="Game",meta=(HidePin="WorldContextObject", DefaultToSelf="WorldContextObject"))
+	UFUNCTION(BlueprintCallable, Category="Game",meta=(WorldContext="WorldContextObject"))
 	static void QuitGame(UObject* WorldContextObject, class APlayerController* SpecificPlayer, TEnumAsByte<EQuitPreference::Type> QuitPreference);
 
 	//=============================================================================
@@ -227,15 +227,15 @@ class UKismetSystemLibrary : public UBlueprintFunctionLibrary
 	 * @param Duration 		length of delay.
 	 * @param LatentInfo 	The latent action.
 	 */
-	UFUNCTION(BlueprintCallable, Category="Utilities|FlowControl", meta=(Latent, HidePin="WorldContextObject", DefaultToSelf="WorldContextObject", LatentInfo="LatentInfo", Duration="0.2"))
+	UFUNCTION(BlueprintCallable, Category="Utilities|FlowControl", meta=(Latent, WorldContext="WorldContextObject", LatentInfo="LatentInfo", Duration="0.2"))
 	static void	Delay(UObject* WorldContextObject, float Duration, struct FLatentActionInfo LatentInfo );
 
 	/** Delay execution by Duration seconds; Calling again before the delay has expired will reset the countdown to Duration. */
-	UFUNCTION(BlueprintCallable, meta=(Latent, LatentInfo="LatentInfo", HidePin="WorldContextObject", DefaultToSelf="WorldContextObject", Duration="0.2"), Category="Utilities|FlowControl")
+	UFUNCTION(BlueprintCallable, meta=(Latent, LatentInfo="LatentInfo", WorldContext="WorldContextObject", Duration="0.2"), Category="Utilities|FlowControl")
 	static void RetriggerableDelay(UObject* WorldContextObject, float Duration, FLatentActionInfo LatentInfo);
 
 	/** Interpolate a component to the specified relative location and rotation over the course of OverTime seconds. */
-	UFUNCTION(BlueprintCallable, meta=(Latent, LatentInfo="LatentInfo", HidePin="WorldContextObject", DefaultToSelf="WorldContextObject", ExpandEnumAsExecs="MoveAction", OverTime="0.2"), Category="Components")
+	UFUNCTION(BlueprintCallable, meta=(Latent, LatentInfo="LatentInfo", WorldContext="WorldContextObject", ExpandEnumAsExecs="MoveAction", OverTime="0.2"), Category="Components")
 	static void MoveComponentTo(USceneComponent* Component, FVector TargetRelativeLocation, FRotator TargetRelativeRotation, bool bEaseOut, bool bEaseIn, float OverTime, TEnumAsByte<EMoveComponentAction::Type> MoveAction, FLatentActionInfo LatentInfo);
 
 	// --- Timer functions ------------------------------
@@ -419,10 +419,10 @@ class UKismetSystemLibrary : public UBlueprintFunctionLibrary
 	 * @param OutActors		Returned array of actors. Unsorted.
 	 * @return				true if there was an overlap that passed the filters, false otherwise.
 	 */
-	UFUNCTION(BlueprintCallable, Category="Collision", meta=(HidePin="WorldContextObject", DefaultToSelf="WorldContextObject", AutoCreateRefTerm="ActorsToIgnore", FriendlyName = "SphereOverlapActors"))
+	UFUNCTION(BlueprintCallable, Category="Collision", meta=(WorldContext="WorldContextObject", AutoCreateRefTerm="ActorsToIgnore", FriendlyName = "SphereOverlapActors"))
 	static ENGINE_API bool SphereOverlapActors_NEW(UObject* WorldContextObject, const FVector SpherePos, float SphereRadius, const TArray<TEnumAsByte<EObjectTypeQuery> > & ObjectTypes, UClass* ActorClassFilter, const TArray<AActor*>& ActorsToIgnore, TArray<class AActor*>& OutActors);
 
-	UFUNCTION(BlueprintCallable, Category="Collision", meta=(DeprecatedFunction, DeprecationMessage = "Use new SphereOverlaActors", HidePin="WorldContextObject", DefaultToSelf="WorldContextObject", AutoCreateRefTerm="ActorsToIgnore"))
+	UFUNCTION(BlueprintCallable, Category="Collision", meta=(DeprecatedFunction, DeprecationMessage = "Use new SphereOverlaActors", WorldContext="WorldContextObject", AutoCreateRefTerm="ActorsToIgnore"))
 	static ENGINE_API bool SphereOverlapActors_DEPRECATED(UObject* WorldContextObject, const FVector SpherePos, float SphereRadius, EOverlapFilterOption Filter, UClass* ActorClassFilter, const TArray<AActor*>& ActorsToIgnore, TArray<class AActor*>& OutActors);
 
 	/**
@@ -436,10 +436,10 @@ class UKismetSystemLibrary : public UBlueprintFunctionLibrary
 	 * @param OutActors		Returned array of actors. Unsorted.
 	 * @return				true if there was an overlap that passed the filters, false otherwise.
 	 */
-	UFUNCTION(BlueprintCallable, Category="Collision", meta=(HidePin="WorldContextObject", DefaultToSelf="WorldContextObject", AutoCreateRefTerm="ActorsToIgnore", FriendlyName="SphereOverlapComponents"))
+	UFUNCTION(BlueprintCallable, Category="Collision", meta=(WorldContext="WorldContextObject", AutoCreateRefTerm="ActorsToIgnore", FriendlyName="SphereOverlapComponents"))
 	static ENGINE_API bool SphereOverlapComponents_NEW(UObject* WorldContextObject, const FVector SpherePos, float SphereRadius, const TArray<TEnumAsByte<EObjectTypeQuery> > & ObjectTypes, UClass* ComponentClassFilter, const TArray<AActor*>& ActorsToIgnore, TArray<class UPrimitiveComponent*>& OutComponents);
 
-	UFUNCTION(BlueprintCallable, Category="Collision", meta=(DeprecatedFunction, DeprecationMessage = "Use new SphereOverlapComponents", HidePin="WorldContextObject", DefaultToSelf="WorldContextObject", AutoCreateRefTerm="ActorsToIgnore"))
+	UFUNCTION(BlueprintCallable, Category="Collision", meta=(DeprecatedFunction, DeprecationMessage = "Use new SphereOverlapComponents", WorldContext="WorldContextObject", AutoCreateRefTerm="ActorsToIgnore"))
 	static ENGINE_API bool SphereOverlapComponents_DEPRECATED(UObject* WorldContextObject, const FVector SpherePos, float SphereRadius, EOverlapFilterOption Filter, UClass* ComponentClassFilter, const TArray<AActor*>& ActorsToIgnore, TArray<class UPrimitiveComponent*>& OutComponents);
 	
 
@@ -454,10 +454,10 @@ class UKismetSystemLibrary : public UBlueprintFunctionLibrary
 	 * @param OutActors		Returned array of actors. Unsorted.
 	 * @return				true if there was an overlap that passed the filters, false otherwise.
 	 */
-	UFUNCTION(BlueprintCallable, Category="Collision", meta=(HidePin="WorldContextObject", DefaultToSelf="WorldContextObject", AutoCreateRefTerm="ActorsToIgnore", FriendlyName="BoxOverlapActors"))
+	UFUNCTION(BlueprintCallable, Category="Collision", meta=(WorldContext="WorldContextObject", AutoCreateRefTerm="ActorsToIgnore", FriendlyName="BoxOverlapActors"))
 	static ENGINE_API bool BoxOverlapActors_NEW(UObject* WorldContextObject, const FVector BoxPos, FVector BoxExtent, const TArray<TEnumAsByte<EObjectTypeQuery> > & ObjectTypes, UClass* ActorClassFilter, const TArray<AActor*>& ActorsToIgnore, TArray<class AActor*>& OutActors);
 
-	UFUNCTION(BlueprintCallable, Category="Collision", meta=(DeprecatedFunction, DeprecationMessage = "Use new BoxOverlapActors", HidePin="WorldContextObject", DefaultToSelf="WorldContextObject", AutoCreateRefTerm="ActorsToIgnore"))
+	UFUNCTION(BlueprintCallable, Category="Collision", meta=(DeprecatedFunction, DeprecationMessage = "Use new BoxOverlapActors", WorldContext="WorldContextObject", AutoCreateRefTerm="ActorsToIgnore"))
 	static ENGINE_API bool BoxOverlapActors_DEPRECATED(UObject* WorldContextObject, const FVector BoxPos, FVector BoxExtent, EOverlapFilterOption Filter, UClass* ActorClassFilter, const TArray<AActor*>& ActorsToIgnore, TArray<class AActor*>& OutActors);
 
 	/**
@@ -471,10 +471,10 @@ class UKismetSystemLibrary : public UBlueprintFunctionLibrary
 	 * @param OutActors		Returned array of actors. Unsorted.
 	 * @return				true if there was an overlap that passed the filters, false otherwise.
 	 */
-	UFUNCTION(BlueprintCallable, Category="Collision", meta=(HidePin="WorldContextObject", DefaultToSelf="WorldContextObject", AutoCreateRefTerm="ActorsToIgnore", FriendlyName="BoxOverlapComponents"))
+	UFUNCTION(BlueprintCallable, Category="Collision", meta=(WorldContext="WorldContextObject", AutoCreateRefTerm="ActorsToIgnore", FriendlyName="BoxOverlapComponents"))
 	static ENGINE_API bool BoxOverlapComponents_NEW(UObject* WorldContextObject, const FVector BoxPos, FVector Extent, const TArray<TEnumAsByte<EObjectTypeQuery> > & ObjectTypes, UClass* ComponentClassFilter, const TArray<AActor*>& ActorsToIgnore, TArray<class UPrimitiveComponent*>& OutComponents);
 
-	UFUNCTION(BlueprintCallable, Category="Collision", meta=(DeprecatedFunction, DeprecationMessage = "Use new BoxOverlapComponents", HidePin="WorldContextObject", DefaultToSelf="WorldContextObject", AutoCreateRefTerm="ActorsToIgnore"))
+	UFUNCTION(BlueprintCallable, Category="Collision", meta=(DeprecatedFunction, DeprecationMessage = "Use new BoxOverlapComponents", WorldContext="WorldContextObject", AutoCreateRefTerm="ActorsToIgnore"))
 	static ENGINE_API bool BoxOverlapComponents_DEPRECATED(UObject* WorldContextObject, const FVector BoxPos, FVector Extent, EOverlapFilterOption Filter, UClass* ComponentClassFilter, const TArray<AActor*>& ActorsToIgnore, TArray<class UPrimitiveComponent*>& OutComponents);
 
 
@@ -490,10 +490,10 @@ class UKismetSystemLibrary : public UBlueprintFunctionLibrary
 	 * @param OutActors		Returned array of actors. Unsorted.
 	 * @return				true if there was an overlap that passed the filters, false otherwise.
 	 */
-	UFUNCTION(BlueprintCallable, Category="Collision", meta=(HidePin="WorldContextObject", DefaultToSelf="WorldContextObject", AutoCreateRefTerm="ActorsToIgnore", FriendlyName="CapsuleOverlapActors"))
+	UFUNCTION(BlueprintCallable, Category="Collision", meta=(WorldContext="WorldContextObject", AutoCreateRefTerm="ActorsToIgnore", FriendlyName="CapsuleOverlapActors"))
 	static ENGINE_API bool CapsuleOverlapActors_NEW(UObject* WorldContextObject, const FVector CapsulePos, float Radius, float HalfHeight, const TArray<TEnumAsByte<EObjectTypeQuery> > & ObjectTypes, UClass* ActorClassFilter, const TArray<AActor*>& ActorsToIgnore, TArray<class AActor*>& OutActors);
 
-	UFUNCTION(BlueprintCallable, Category="Collision", meta=(DeprecatedFunction, DeprecationMessage = "Use new CapsuleOverlapActors", HidePin="WorldContextObject", DefaultToSelf="WorldContextObject", AutoCreateRefTerm="ActorsToIgnore"))
+	UFUNCTION(BlueprintCallable, Category="Collision", meta=(DeprecatedFunction, DeprecationMessage = "Use new CapsuleOverlapActors", WorldContext="WorldContextObject", AutoCreateRefTerm="ActorsToIgnore"))
 	static ENGINE_API bool CapsuleOverlapActors_DEPRECATED(UObject* WorldContextObject, const FVector CapsulePos, float Radius, float HalfHeight, EOverlapFilterOption Filter, UClass* ActorClassFilter, const TArray<AActor*>& ActorsToIgnore, TArray<class AActor*>& OutActors);
 	/**
 	 * Returns an array of components that overlap the given capsule.
@@ -507,10 +507,10 @@ class UKismetSystemLibrary : public UBlueprintFunctionLibrary
 	 * @param OutActors		Returned array of actors. Unsorted.
 	 * @return				true if there was an overlap that passed the filters, false otherwise.
 	 */
-	UFUNCTION(BlueprintCallable, Category="Collision", meta=(HidePin="WorldContextObject", DefaultToSelf="WorldContextObject", AutoCreateRefTerm="ActorsToIgnore", FriendlyName="CapsuleOverlapComponents") )
+	UFUNCTION(BlueprintCallable, Category="Collision", meta=(WorldContext="WorldContextObject", AutoCreateRefTerm="ActorsToIgnore", FriendlyName="CapsuleOverlapComponents") )
 	static ENGINE_API bool CapsuleOverlapComponents_NEW(UObject* WorldContextObject, const FVector CapsulePos, float Radius, float HalfHeight, const TArray<TEnumAsByte<EObjectTypeQuery> > & ObjectTypes, UClass* ComponentClassFilter, const TArray<AActor*>& ActorsToIgnore, TArray<class UPrimitiveComponent*>& OutComponents);
 
-	UFUNCTION(BlueprintCallable, Category="Collision", meta=(DeprecatedFunction, DeprecationMessage = "Use new CapsuleOverlapComponents", HidePin="WorldContextObject", DefaultToSelf="WorldContextObject", AutoCreateRefTerm="ActorsToIgnore") )
+	UFUNCTION(BlueprintCallable, Category="Collision", meta=(DeprecatedFunction, DeprecationMessage = "Use new CapsuleOverlapComponents", WorldContext="WorldContextObject", AutoCreateRefTerm="ActorsToIgnore") )
 	static ENGINE_API bool CapsuleOverlapComponents_DEPRECATED(UObject* WorldContextObject, const FVector CapsulePos, float Radius, float HalfHeight, EOverlapFilterOption Filter, UClass* ComponentClassFilter, const TArray<AActor*>& ActorsToIgnore, TArray<class UPrimitiveComponent*>& OutComponents);
 	
 
@@ -558,7 +558,7 @@ class UKismetSystemLibrary : public UBlueprintFunctionLibrary
 	 * @param OutHit		Properties of the trace hit.
 	 * @return				True if there was a hit, false otherwise.
 	 */
-	UFUNCTION(BlueprintCallable, Category="Collision", meta=(bIgnoreSelf="true", HidePin="WorldContextObject", DefaultToSelf="WorldContextObject", AutoCreateRefTerm="ActorsToIgnore", FriendlyName = "SingleLineTraceByChannel"))
+	UFUNCTION(BlueprintCallable, Category="Collision", meta=(bIgnoreSelf="true", WorldContext="WorldContextObject", AutoCreateRefTerm="ActorsToIgnore", FriendlyName = "SingleLineTraceByChannel"))
 	static ENGINE_API bool LineTraceSingle_NEW(UObject* WorldContextObject, const FVector Start, const FVector End, ETraceTypeQuery TraceChannel, bool bTraceComplex, const TArray<AActor*>& ActorsToIgnore, EDrawDebugTrace::Type DrawDebugType, FHitResult& OutHit, bool bIgnoreSelf);
 	
 	/**
@@ -573,7 +573,7 @@ class UKismetSystemLibrary : public UBlueprintFunctionLibrary
 	 * @param OutHit		Properties of the trace hit.
 	 * @return				True if there was a hit, false otherwise.
 	 */
-	UFUNCTION(BlueprintCallable, Category="Collision", meta=(bIgnoreSelf="true", HidePin="WorldContextObject", DefaultToSelf="WorldContextObject", AutoCreateRefTerm="ActorsToIgnore", FriendlyName = "MultiLineTraceByChannel"))
+	UFUNCTION(BlueprintCallable, Category="Collision", meta=(bIgnoreSelf="true", WorldContext="WorldContextObject", AutoCreateRefTerm="ActorsToIgnore", FriendlyName = "MultiLineTraceByChannel"))
 	static ENGINE_API bool LineTraceMulti_NEW(UObject* WorldContextObject, const FVector Start, const FVector End, ETraceTypeQuery TraceChannel, bool bTraceComplex, const TArray<AActor*>& ActorsToIgnore, EDrawDebugTrace::Type DrawDebugType, TArray<FHitResult>& OutHits, bool bIgnoreSelf);
 
 	/**
@@ -588,7 +588,7 @@ class UKismetSystemLibrary : public UBlueprintFunctionLibrary
 	 * @param OutHit		Properties of the trace hit.
 	 * @return				True if there was a hit, false otherwise.
 	 */
-	UFUNCTION(BlueprintCallable, Category="Collision", meta=(bIgnoreSelf="true", HidePin="WorldContextObject", DefaultToSelf="WorldContextObject", AutoCreateRefTerm="ActorsToIgnore", FriendlyName = "SingleSphereTraceByChannel"))
+	UFUNCTION(BlueprintCallable, Category="Collision", meta=(bIgnoreSelf="true", WorldContext="WorldContextObject", AutoCreateRefTerm="ActorsToIgnore", FriendlyName = "SingleSphereTraceByChannel"))
 	static ENGINE_API bool SphereTraceSingle_NEW(UObject* WorldContextObject, const FVector Start, const FVector End, float Radius, ETraceTypeQuery TraceChannel, bool bTraceComplex, const TArray<AActor*>& ActorsToIgnore, EDrawDebugTrace::Type DrawDebugType, FHitResult& OutHit, bool bIgnoreSelf);
 
 	/**
@@ -604,7 +604,7 @@ class UKismetSystemLibrary : public UBlueprintFunctionLibrary
 	 * @param OutHits		A list of hits, sorted along the trace from start to finish.  The blocking hit will be the last hit, if there was one.
 	 * @return				True if there was a hit, false otherwise.
 	 */
-	UFUNCTION(BlueprintCallable, Category="Collision", meta=(bIgnoreSelf="true", HidePin="WorldContextObject", DefaultToSelf="WorldContextObject", AutoCreateRefTerm="ActorsToIgnore", FriendlyName = "MultiSphereTraceByChannel"))
+	UFUNCTION(BlueprintCallable, Category="Collision", meta=(bIgnoreSelf="true", WorldContext="WorldContextObject", AutoCreateRefTerm="ActorsToIgnore", FriendlyName = "MultiSphereTraceByChannel"))
 	static ENGINE_API bool SphereTraceMulti_NEW(UObject* WorldContextObject, const FVector Start, const FVector End, float Radius, ETraceTypeQuery TraceChannel, bool bTraceComplex, const TArray<AActor*>& ActorsToIgnore, EDrawDebugTrace::Type DrawDebugType, TArray<FHitResult>& OutHits, bool bIgnoreSelf);
 
 	/**
@@ -654,7 +654,7 @@ class UKismetSystemLibrary : public UBlueprintFunctionLibrary
 	 * @param OutHit		Properties of the trace hit.
 	 * @return				True if there was a hit, false otherwise.
 	 */
-	UFUNCTION(BlueprintCallable, Category="Collision", meta=(bIgnoreSelf="true", HidePin="WorldContextObject", DefaultToSelf="WorldContextObject", AutoCreateRefTerm="ActorsToIgnore", FriendlyName = "SingleCapsuleTraceByChannel"))
+	UFUNCTION(BlueprintCallable, Category="Collision", meta=(bIgnoreSelf="true", WorldContext="WorldContextObject", AutoCreateRefTerm="ActorsToIgnore", FriendlyName = "SingleCapsuleTraceByChannel"))
 	static ENGINE_API bool CapsuleTraceSingle_NEW(UObject* WorldContextObject, const FVector Start, const FVector End, float Radius, float HalfHeight, ETraceTypeQuery TraceChannel, bool bTraceComplex, const TArray<AActor*>& ActorsToIgnore, EDrawDebugTrace::Type DrawDebugType, FHitResult& OutHit, bool bIgnoreSelf);
 
 	/**
@@ -671,7 +671,7 @@ class UKismetSystemLibrary : public UBlueprintFunctionLibrary
 	 * @param OutHits		A list of hits, sorted along the trace from start to finish.  The blocking hit will be the last hit, if there was one.
 	 * @return				True if there was a hit, false otherwise.
 	 */
-	UFUNCTION(BlueprintCallable, Category="Collision", meta=(bIgnoreSelf="true", HidePin="WorldContextObject", DefaultToSelf="WorldContextObject", AutoCreateRefTerm="ActorsToIgnore", FriendlyName = "MultiCapsuleTraceByChannel"))
+	UFUNCTION(BlueprintCallable, Category="Collision", meta=(bIgnoreSelf="true", WorldContext="WorldContextObject", AutoCreateRefTerm="ActorsToIgnore", FriendlyName = "MultiCapsuleTraceByChannel"))
 	static ENGINE_API bool CapsuleTraceMulti_NEW(UObject* WorldContextObject, const FVector Start, const FVector End, float Radius, float HalfHeight, ETraceTypeQuery TraceChannel, bool bTraceComplex, const TArray<AActor*>& ActorsToIgnore, EDrawDebugTrace::Type DrawDebugType, TArray<FHitResult>& OutHits, bool bIgnoreSelf);
 
 	/**
@@ -686,7 +686,7 @@ class UKismetSystemLibrary : public UBlueprintFunctionLibrary
 	 * @param OutHit		Properties of the trace hit.
 	 * @return				True if there was a hit, false otherwise.
 	 */
-	UFUNCTION(BlueprintCallable, Category="Collision", meta=(bIgnoreSelf="true", HidePin="WorldContextObject", DefaultToSelf="WorldContextObject", AutoCreateRefTerm="ActorsToIgnore", FriendlyName = "SingleLineTraceForObjects"))
+	UFUNCTION(BlueprintCallable, Category="Collision", meta=(bIgnoreSelf="true", WorldContext="WorldContextObject", AutoCreateRefTerm="ActorsToIgnore", FriendlyName = "SingleLineTraceForObjects"))
 	static ENGINE_API bool LineTraceSingleForObjects(UObject* WorldContextObject, const FVector Start, const FVector End, const TArray<TEnumAsByte<EObjectTypeQuery> > & ObjectTypes, bool bTraceComplex, const TArray<AActor*>& ActorsToIgnore, EDrawDebugTrace::Type DrawDebugType, FHitResult& OutHit, bool bIgnoreSelf);
 	
 	/**
@@ -701,7 +701,7 @@ class UKismetSystemLibrary : public UBlueprintFunctionLibrary
 	 * @param OutHit		Properties of the trace hit.
 	 * @return				True if there was a hit, false otherwise.
 	 */
-	UFUNCTION(BlueprintCallable, Category="Collision", meta=(bIgnoreSelf="true", HidePin="WorldContextObject", DefaultToSelf="WorldContextObject", AutoCreateRefTerm="ActorsToIgnore", FriendlyName = "MultiLineTraceForObjects"))
+	UFUNCTION(BlueprintCallable, Category="Collision", meta=(bIgnoreSelf="true", WorldContext="WorldContextObject", AutoCreateRefTerm="ActorsToIgnore", FriendlyName = "MultiLineTraceForObjects"))
 	static ENGINE_API bool LineTraceMultiForObjects(UObject* WorldContextObject, const FVector Start, const FVector End, const TArray<TEnumAsByte<EObjectTypeQuery> > & ObjectTypes, bool bTraceComplex, const TArray<AActor*>& ActorsToIgnore, EDrawDebugTrace::Type DrawDebugType, TArray<FHitResult>& OutHits, bool bIgnoreSelf);
 
 	/**
@@ -716,7 +716,7 @@ class UKismetSystemLibrary : public UBlueprintFunctionLibrary
 	 * @param OutHit		Properties of the trace hit.
 	 * @return				True if there was a hit, false otherwise.
 	 */
-	UFUNCTION(BlueprintCallable, Category="Collision", meta=(bIgnoreSelf="true", HidePin="WorldContextObject", DefaultToSelf="WorldContextObject", AutoCreateRefTerm="ActorsToIgnore", FriendlyName = "SingleSphereTraceForObjects"))
+	UFUNCTION(BlueprintCallable, Category="Collision", meta=(bIgnoreSelf="true", WorldContext="WorldContextObject", AutoCreateRefTerm="ActorsToIgnore", FriendlyName = "SingleSphereTraceForObjects"))
 	static ENGINE_API bool SphereTraceSingleForObjects(UObject* WorldContextObject, const FVector Start, const FVector End, float Radius, const TArray<TEnumAsByte<EObjectTypeQuery> > & ObjectTypes, bool bTraceComplex, const TArray<AActor*>& ActorsToIgnore, EDrawDebugTrace::Type DrawDebugType, FHitResult& OutHit, bool bIgnoreSelf);
 
 	/**
@@ -732,7 +732,7 @@ class UKismetSystemLibrary : public UBlueprintFunctionLibrary
 	 * @param OutHits		A list of hits, sorted along the trace from start to finish.  The blocking hit will be the last hit, if there was one.
 	 * @return				True if there was a hit, false otherwise.
 	 */
-	UFUNCTION(BlueprintCallable, Category="Collision", meta=(bIgnoreSelf="true", HidePin="WorldContextObject", DefaultToSelf="WorldContextObject", AutoCreateRefTerm="ActorsToIgnore", FriendlyName = "MultiSphereTraceForObjects"))
+	UFUNCTION(BlueprintCallable, Category="Collision", meta=(bIgnoreSelf="true", WorldContext="WorldContextObject", AutoCreateRefTerm="ActorsToIgnore", FriendlyName = "MultiSphereTraceForObjects"))
 	static ENGINE_API bool SphereTraceMultiForObjects(UObject* WorldContextObject, const FVector Start, const FVector End, float Radius, const TArray<TEnumAsByte<EObjectTypeQuery> > & ObjectTypes, bool bTraceComplex, const TArray<AActor*>& ActorsToIgnore, EDrawDebugTrace::Type DrawDebugType, TArray<FHitResult>& OutHits, bool bIgnoreSelf);
 	
 
@@ -783,7 +783,7 @@ class UKismetSystemLibrary : public UBlueprintFunctionLibrary
 	 * @param OutHit		Properties of the trace hit.
 	 * @return				True if there was a hit, false otherwise.
 	 */
-	UFUNCTION(BlueprintCallable, Category="Collision", meta=(bIgnoreSelf="true", HidePin="WorldContextObject", DefaultToSelf="WorldContextObject", AutoCreateRefTerm="ActorsToIgnore", FriendlyName = "SingleCapsuleTraceForObjects"))
+	UFUNCTION(BlueprintCallable, Category="Collision", meta=(bIgnoreSelf="true", WorldContext="WorldContextObject", AutoCreateRefTerm="ActorsToIgnore", FriendlyName = "SingleCapsuleTraceForObjects"))
 	static ENGINE_API bool CapsuleTraceSingleForObjects(UObject* WorldContextObject, const FVector Start, const FVector End, float Radius, float HalfHeight, const TArray<TEnumAsByte<EObjectTypeQuery> > & ObjectTypes, bool bTraceComplex, const TArray<AActor*>& ActorsToIgnore, EDrawDebugTrace::Type DrawDebugType, FHitResult& OutHit, bool bIgnoreSelf);
 
 	/**
@@ -800,7 +800,7 @@ class UKismetSystemLibrary : public UBlueprintFunctionLibrary
 	 * @param OutHits		A list of hits, sorted along the trace from start to finish.  The blocking hit will be the last hit, if there was one.
 	 * @return				True if there was a hit, false otherwise.
 	 */
-	UFUNCTION(BlueprintCallable, Category="Collision", meta=(bIgnoreSelf="true", HidePin="WorldContextObject", DefaultToSelf="WorldContextObject", AutoCreateRefTerm="ActorsToIgnore", FriendlyName = "MultiCapsuleTraceForObjects"))
+	UFUNCTION(BlueprintCallable, Category="Collision", meta=(bIgnoreSelf="true", WorldContext="WorldContextObject", AutoCreateRefTerm="ActorsToIgnore", FriendlyName = "MultiCapsuleTraceForObjects"))
 	static ENGINE_API bool CapsuleTraceMultiForObjects(UObject* WorldContextObject, const FVector Start, const FVector End, float Radius, float HalfHeight, const TArray<TEnumAsByte<EObjectTypeQuery> > & ObjectTypes, bool bTraceComplex, const TArray<AActor*>& ActorsToIgnore, EDrawDebugTrace::Type DrawDebugType, TArray<FHitResult>& OutHits, bool bIgnoreSelf);
 
 	/**
@@ -815,7 +815,7 @@ class UKismetSystemLibrary : public UBlueprintFunctionLibrary
 	 * @param OutHit		Properties of the trace hit.
 	 * @return				True if there was a hit, false otherwise.
 	 */
-	UFUNCTION(BlueprintCallable, Category="Collision", meta=(DeprecatedFunction, DeprecationMessage = "Use new SingleLineTraceByChannel", bIgnoreSelf="true", HidePin="WorldContextObject", DefaultToSelf="WorldContextObject", AutoCreateRefTerm="ActorsToIgnore", FriendlyName = "SingleLineTraceByChannelDeprecated"))
+	UFUNCTION(BlueprintCallable, Category="Collision", meta=(DeprecatedFunction, DeprecationMessage = "Use new SingleLineTraceByChannel", bIgnoreSelf="true", WorldContext="WorldContextObject", AutoCreateRefTerm="ActorsToIgnore", FriendlyName = "SingleLineTraceByChannelDeprecated"))
 	static ENGINE_API bool LineTraceSingle_DEPRECATED(UObject* WorldContextObject, const FVector Start, const FVector End, ECollisionChannel TraceChannel, bool bTraceComplex, const TArray<AActor*>& ActorsToIgnore, EDrawDebugTrace::Type DrawDebugType, FHitResult& OutHit, bool bIgnoreSelf);
 	
 	/**
@@ -830,7 +830,7 @@ class UKismetSystemLibrary : public UBlueprintFunctionLibrary
 	 * @param OutHit		Properties of the trace hit.
 	 * @return				True if there was a hit, false otherwise.
 	 */
-	UFUNCTION(BlueprintCallable, Category="Collision", meta=(DeprecatedFunction, DeprecationMessage = "Use new MultiLineTraceByChannel", bIgnoreSelf="true", HidePin="WorldContextObject", DefaultToSelf="WorldContextObject", AutoCreateRefTerm="ActorsToIgnore", FriendlyName = "MultiLineTraceByChannelDeprecated"))
+	UFUNCTION(BlueprintCallable, Category="Collision", meta=(DeprecatedFunction, DeprecationMessage = "Use new MultiLineTraceByChannel", bIgnoreSelf="true", WorldContext="WorldContextObject", AutoCreateRefTerm="ActorsToIgnore", FriendlyName = "MultiLineTraceByChannelDeprecated"))
 	static ENGINE_API bool LineTraceMulti_DEPRECATED(UObject* WorldContextObject, const FVector Start, const FVector End, ECollisionChannel TraceChannel, bool bTraceComplex, const TArray<AActor*>& ActorsToIgnore, EDrawDebugTrace::Type DrawDebugType, TArray<FHitResult>& OutHits, bool bIgnoreSelf);
 
 	/**
@@ -845,7 +845,7 @@ class UKismetSystemLibrary : public UBlueprintFunctionLibrary
 	 * @param OutHit		Properties of the trace hit.
 	 * @return				True if there was a hit, false otherwise.
 	 */
-	UFUNCTION(BlueprintCallable, Category="Collision", meta=(DeprecatedFunction, DeprecationMessage = "Use new SingleSphereTraceByChannel", bIgnoreSelf="true", HidePin="WorldContextObject", DefaultToSelf="WorldContextObject", AutoCreateRefTerm="ActorsToIgnore", FriendlyName = "SingleSphereTraceByChannelDeprecated"))
+	UFUNCTION(BlueprintCallable, Category="Collision", meta=(DeprecatedFunction, DeprecationMessage = "Use new SingleSphereTraceByChannel", bIgnoreSelf="true", WorldContext="WorldContextObject", AutoCreateRefTerm="ActorsToIgnore", FriendlyName = "SingleSphereTraceByChannelDeprecated"))
 	static ENGINE_API bool SphereTraceSingle_DEPRECATED(UObject* WorldContextObject, const FVector Start, const FVector End, float Radius, ECollisionChannel TraceChannel, bool bTraceComplex, const TArray<AActor*>& ActorsToIgnore, EDrawDebugTrace::Type DrawDebugType, FHitResult& OutHit, bool bIgnoreSelf);
 
 	/**
@@ -861,7 +861,7 @@ class UKismetSystemLibrary : public UBlueprintFunctionLibrary
 	 * @param OutHits		A list of hits, sorted along the trace from start to finish.  The blocking hit will be the last hit, if there was one.
 	 * @return				True if there was a hit, false otherwise.
 	 */
-	UFUNCTION(BlueprintCallable, Category="Collision", meta=(DeprecatedFunction, DeprecationMessage = "Use new MultiSphereTraceByChannel", bIgnoreSelf="true", HidePin="WorldContextObject", DefaultToSelf="WorldContextObject", AutoCreateRefTerm="ActorsToIgnore", FriendlyName = "MultiSphereTraceByChannelDeprecated"))
+	UFUNCTION(BlueprintCallable, Category="Collision", meta=(DeprecatedFunction, DeprecationMessage = "Use new MultiSphereTraceByChannel", bIgnoreSelf="true", WorldContext="WorldContextObject", AutoCreateRefTerm="ActorsToIgnore", FriendlyName = "MultiSphereTraceByChannelDeprecated"))
 	static ENGINE_API bool SphereTraceMulti_DEPRECATED(UObject* WorldContextObject, const FVector Start, const FVector End, float Radius, ECollisionChannel TraceChannel, bool bTraceComplex, const TArray<AActor*>& ActorsToIgnore, EDrawDebugTrace::Type DrawDebugType, TArray<FHitResult>& OutHits, bool bIgnoreSelf);
 	
 	/**
@@ -878,7 +878,7 @@ class UKismetSystemLibrary : public UBlueprintFunctionLibrary
 	 * @param OutHit		Properties of the trace hit.
 	 * @return				True if there was a hit, false otherwise.
 	 */
-	UFUNCTION(BlueprintCallable, Category="Collision", meta=(DeprecatedFunction, DeprecationMessage = "Use new SingleCapsuleTraceByChannel", bIgnoreSelf="true", HidePin="WorldContextObject", DefaultToSelf="WorldContextObject", AutoCreateRefTerm="ActorsToIgnore", FriendlyName = "SingleCapsuleTraceByChannelDeprecated"))
+	UFUNCTION(BlueprintCallable, Category="Collision", meta=(DeprecatedFunction, DeprecationMessage = "Use new SingleCapsuleTraceByChannel", bIgnoreSelf="true", WorldContext="WorldContextObject", AutoCreateRefTerm="ActorsToIgnore", FriendlyName = "SingleCapsuleTraceByChannelDeprecated"))
 	static ENGINE_API bool CapsuleTraceSingle_DEPRECATED(UObject* WorldContextObject, const FVector Start, const FVector End, float Radius, float HalfHeight, ECollisionChannel TraceChannel, bool bTraceComplex, const TArray<AActor*>& ActorsToIgnore, EDrawDebugTrace::Type DrawDebugType, FHitResult& OutHit, bool bIgnoreSelf);
 
 	/**
@@ -895,7 +895,7 @@ class UKismetSystemLibrary : public UBlueprintFunctionLibrary
 	 * @param OutHits		A list of hits, sorted along the trace from start to finish.  The blocking hit will be the last hit, if there was one.
 	 * @return				True if there was a hit, false otherwise.
 	 */
-	UFUNCTION(BlueprintCallable, Category="Collision", meta=(DeprecatedFunction, DeprecationMessage = "Use new MultiCapsuleTraceByChannel", bIgnoreSelf="true", HidePin="WorldContextObject", DefaultToSelf="WorldContextObject", AutoCreateRefTerm="ActorsToIgnore", FriendlyName = "MultiCapsuleTraceByChannelDeprecated"))
+	UFUNCTION(BlueprintCallable, Category="Collision", meta=(DeprecatedFunction, DeprecationMessage = "Use new MultiCapsuleTraceByChannel", bIgnoreSelf="true", WorldContext="WorldContextObject", AutoCreateRefTerm="ActorsToIgnore", FriendlyName = "MultiCapsuleTraceByChannelDeprecated"))
 	static ENGINE_API bool CapsuleTraceMulti_DEPRECATED(UObject* WorldContextObject, const FVector Start, const FVector End, float Radius, float HalfHeight, ECollisionChannel TraceChannel, bool bTraceComplex, const TArray<AActor*>& ActorsToIgnore, EDrawDebugTrace::Type DrawDebugType, TArray<FHitResult>& OutHits, bool bIgnoreSelf);
 
 	/**
@@ -910,7 +910,7 @@ class UKismetSystemLibrary : public UBlueprintFunctionLibrary
 	 * @param OutHit		Properties of the trace hit.
 	 * @return				True if there was a hit, false otherwise.
 	 */
-	UFUNCTION(BlueprintCallable, Category="Collision", meta=(DeprecatedFunction, DeprecationMessage = "Use new SingleLineTraceForObjects", bIgnoreSelf="true", HidePin="WorldContextObject", DefaultToSelf="WorldContextObject", AutoCreateRefTerm="ActorsToIgnore", FriendlyName = "SingleLineTraceByObjectDeprecated"))
+	UFUNCTION(BlueprintCallable, Category="Collision", meta=(DeprecatedFunction, DeprecationMessage = "Use new SingleLineTraceForObjects", bIgnoreSelf="true", WorldContext="WorldContextObject", AutoCreateRefTerm="ActorsToIgnore", FriendlyName = "SingleLineTraceByObjectDeprecated"))
 	static ENGINE_API bool LineTraceSingleByObject_DEPRECATED(UObject* WorldContextObject, const FVector Start, const FVector End, const TArray<TEnumAsByte<ECollisionChannel> > & ObjectsToTrace, bool bTraceComplex, const TArray<AActor*>& ActorsToIgnore, EDrawDebugTrace::Type DrawDebugType, FHitResult& OutHit, bool bIgnoreSelf);
 	
 	/**
@@ -925,7 +925,7 @@ class UKismetSystemLibrary : public UBlueprintFunctionLibrary
 	 * @param OutHit		Properties of the trace hit.
 	 * @return				True if there was a hit, false otherwise.
 	 */
-	UFUNCTION(BlueprintCallable, Category="Collision", meta=(DeprecatedFunction, DeprecationMessage = "Use new MultiLineTraceForObjects", bIgnoreSelf="true", HidePin="WorldContextObject", DefaultToSelf="WorldContextObject", AutoCreateRefTerm="ActorsToIgnore", FriendlyName = "MultiLineTraceByObjectDeprecated"))
+	UFUNCTION(BlueprintCallable, Category="Collision", meta=(DeprecatedFunction, DeprecationMessage = "Use new MultiLineTraceForObjects", bIgnoreSelf="true", WorldContext="WorldContextObject", AutoCreateRefTerm="ActorsToIgnore", FriendlyName = "MultiLineTraceByObjectDeprecated"))
 	static ENGINE_API bool LineTraceMultiByObject_DEPRECATED(UObject* WorldContextObject, const FVector Start, const FVector End, const TArray<TEnumAsByte<ECollisionChannel> > & ObjectsToTrace, bool bTraceComplex, const TArray<AActor*>& ActorsToIgnore, EDrawDebugTrace::Type DrawDebugType, TArray<FHitResult>& OutHits, bool bIgnoreSelf);
 
 	/**
@@ -940,7 +940,7 @@ class UKismetSystemLibrary : public UBlueprintFunctionLibrary
 	 * @param OutHit		Properties of the trace hit.
 	 * @return				True if there was a hit, false otherwise.
 	 */
-	UFUNCTION(BlueprintCallable, Category="Collision", meta=(DeprecatedFunction, DeprecationMessage = "Use new SingleSphereTraceForObjects", bIgnoreSelf="true", HidePin="WorldContextObject", DefaultToSelf="WorldContextObject", AutoCreateRefTerm="ActorsToIgnore", FriendlyName = "SingleSphereTraceByObjectDeprecated"))
+	UFUNCTION(BlueprintCallable, Category="Collision", meta=(DeprecatedFunction, DeprecationMessage = "Use new SingleSphereTraceForObjects", bIgnoreSelf="true", WorldContext="WorldContextObject", AutoCreateRefTerm="ActorsToIgnore", FriendlyName = "SingleSphereTraceByObjectDeprecated"))
 	static ENGINE_API bool SphereTraceSingleByObject_DEPRECATED(UObject* WorldContextObject, const FVector Start, const FVector End, float Radius, const TArray<TEnumAsByte<ECollisionChannel> > & ObjectsToTrace, bool bTraceComplex, const TArray<AActor*>& ActorsToIgnore, EDrawDebugTrace::Type DrawDebugType, FHitResult& OutHit, bool bIgnoreSelf);
 
 	/**
@@ -956,7 +956,7 @@ class UKismetSystemLibrary : public UBlueprintFunctionLibrary
 	 * @param OutHits		A list of hits, sorted along the trace from start to finish.  The blocking hit will be the last hit, if there was one.
 	 * @return				True if there was a hit, false otherwise.
 	 */
-	UFUNCTION(BlueprintCallable, Category="Collision", meta=(DeprecatedFunction, DeprecationMessage = "Use new MultiSphereTraceForObjects", bIgnoreSelf="true", HidePin="WorldContextObject", DefaultToSelf="WorldContextObject", AutoCreateRefTerm="ActorsToIgnore", FriendlyName = "MultiSphereTraceByObjectDeprecated"))
+	UFUNCTION(BlueprintCallable, Category="Collision", meta=(DeprecatedFunction, DeprecationMessage = "Use new MultiSphereTraceForObjects", bIgnoreSelf="true", WorldContext="WorldContextObject", AutoCreateRefTerm="ActorsToIgnore", FriendlyName = "MultiSphereTraceByObjectDeprecated"))
 	static ENGINE_API bool SphereTraceMultiByObject_DEPRECATED(UObject* WorldContextObject, const FVector Start, const FVector End, float Radius, const TArray<TEnumAsByte<ECollisionChannel> > & ObjectsToTrace, bool bTraceComplex, const TArray<AActor*>& ActorsToIgnore, EDrawDebugTrace::Type DrawDebugType, TArray<FHitResult>& OutHits, bool bIgnoreSelf);
 	
 	/**
@@ -973,7 +973,7 @@ class UKismetSystemLibrary : public UBlueprintFunctionLibrary
 	 * @param OutHit		Properties of the trace hit.
 	 * @return				True if there was a hit, false otherwise.
 	 */
-	UFUNCTION(BlueprintCallable, Category="Collision", meta=(DeprecatedFunction, DeprecationMessage = "Use new SingleCapsuleTraceForObjects", bIgnoreSelf="true", HidePin="WorldContextObject", DefaultToSelf="WorldContextObject", AutoCreateRefTerm="ActorsToIgnore", FriendlyName = "SingleCapsuleTraceByObjectDeprecated"))
+	UFUNCTION(BlueprintCallable, Category="Collision", meta=(DeprecatedFunction, DeprecationMessage = "Use new SingleCapsuleTraceForObjects", bIgnoreSelf="true", WorldContext="WorldContextObject", AutoCreateRefTerm="ActorsToIgnore", FriendlyName = "SingleCapsuleTraceByObjectDeprecated"))
 	static ENGINE_API bool CapsuleTraceSingleByObject_DEPRECATED(UObject* WorldContextObject, const FVector Start, const FVector End, float Radius, float HalfHeight, const TArray<TEnumAsByte<ECollisionChannel> > & ObjectsToTrace, bool bTraceComplex, const TArray<AActor*>& ActorsToIgnore, EDrawDebugTrace::Type DrawDebugType, FHitResult& OutHit, bool bIgnoreSelf);
 
 	/**
@@ -990,7 +990,7 @@ class UKismetSystemLibrary : public UBlueprintFunctionLibrary
 	 * @param OutHits		A list of hits, sorted along the trace from start to finish.  The blocking hit will be the last hit, if there was one.
 	 * @return				True if there was a hit, false otherwise.
 	 */
-	UFUNCTION(BlueprintCallable, Category="Collision", meta=(DeprecatedFunction, DeprecationMessage = "Use new MultiCapsuleTraceForObjects", bIgnoreSelf="true", HidePin="WorldContextObject", DefaultToSelf="WorldContextObject", AutoCreateRefTerm="ActorsToIgnore", FriendlyName = "MultiCapsuleTraceByObjectDeprecated"))
+	UFUNCTION(BlueprintCallable, Category="Collision", meta=(DeprecatedFunction, DeprecationMessage = "Use new MultiCapsuleTraceForObjects", bIgnoreSelf="true", WorldContext="WorldContextObject", AutoCreateRefTerm="ActorsToIgnore", FriendlyName = "MultiCapsuleTraceByObjectDeprecated"))
 	static ENGINE_API bool CapsuleTraceMultiByObject_DEPRECATED(UObject* WorldContextObject, const FVector Start, const FVector End, float Radius, float HalfHeight, const TArray<TEnumAsByte<ECollisionChannel> > & ObjectsToTrace, bool bTraceComplex, const TArray<AActor*>& ActorsToIgnore, EDrawDebugTrace::Type DrawDebugType, TArray<FHitResult>& OutHits, bool bIgnoreSelf);
 
 	/**
@@ -1006,61 +1006,61 @@ class UKismetSystemLibrary : public UBlueprintFunctionLibrary
 	// --- Debug drawing functions ------------------------------
 
 	/** Draw a debug line */
-	UFUNCTION(BlueprintCallable, Category="Rendering|Debug", meta=(HidePin="WorldContextObject", DefaultToSelf="WorldContextObject"))
+	UFUNCTION(BlueprintCallable, Category="Rendering|Debug", meta=(WorldContext="WorldContextObject"))
 	static ENGINE_API void DrawDebugLine(UObject* WorldContextObject, const FVector LineStart, const FVector LineEnd, FLinearColor LineColor, float Duration=0.f, float Thickness = 0.f);
 
 	/** Draw a debug point */
-	UFUNCTION(BlueprintCallable, Category="Rendering|Debug", meta=(HidePin="WorldContextObject", DefaultToSelf="WorldContextObject"))
+	UFUNCTION(BlueprintCallable, Category="Rendering|Debug", meta=(WorldContext="WorldContextObject"))
 	static ENGINE_API void DrawDebugPoint(UObject* WorldContextObject, const FVector Position, float Size, FLinearColor PointColor, float Duration=0.f);
 
 	/** Draw directional arrow, pointing from LineStart to LineEnd. */
-	UFUNCTION(BlueprintCallable, Category="Rendering|Debug", meta=(HidePin="WorldContextObject", DefaultToSelf="WorldContextObject"))
+	UFUNCTION(BlueprintCallable, Category="Rendering|Debug", meta=(WorldContext="WorldContextObject"))
 	static ENGINE_API void DrawDebugArrow(UObject* WorldContextObject, const FVector LineStart, const FVector LineEnd, float ArrowSize, FLinearColor LineColor, float Duration=0.f);
 
 	/** Draw a debug box */
-	UFUNCTION(BlueprintCallable, Category="Rendering|Debug", meta=(HidePin="WorldContextObject", DefaultToSelf="WorldContextObject"))
+	UFUNCTION(BlueprintCallable, Category="Rendering|Debug", meta=(WorldContext="WorldContextObject"))
 	static ENGINE_API void DrawDebugBox(UObject* WorldContextObject, const FVector Center, FVector Extent, FLinearColor LineColor, const FRotator Rotation=FRotator::ZeroRotator, float Duration=0.f);
 
 	/** Draw a debug coordinate system. */
-	UFUNCTION(BlueprintCallable, Category="Rendering|Debug", meta=(HidePin="WorldContextObject", DefaultToSelf="WorldContextObject"))
+	UFUNCTION(BlueprintCallable, Category="Rendering|Debug", meta=(WorldContext="WorldContextObject"))
 	static ENGINE_API void DrawDebugCoordinateSystem(UObject* WorldContextObject, const FVector AxisLoc, const FRotator AxisRot, float Scale=1.f, float Duration=0.f);
 
 	/** Draw a debug sphere */
-	UFUNCTION(BlueprintCallable, Category="Rendering|Debug", meta=(HidePin="WorldContextObject", DefaultToSelf="WorldContextObject"))
+	UFUNCTION(BlueprintCallable, Category="Rendering|Debug", meta=(WorldContext="WorldContextObject"))
 	static ENGINE_API void DrawDebugSphere(UObject* WorldContextObject, const FVector Center, float Radius=100.f, int32 Segments=12, FLinearColor LineColor = FLinearColor::White, float Duration=0.f);
 
 	/** Draw a debug cylinder */
-	UFUNCTION(BlueprintCallable, Category="Rendering|Debug", meta=(HidePin="WorldContextObject", DefaultToSelf="WorldContextObject"))
+	UFUNCTION(BlueprintCallable, Category="Rendering|Debug", meta=(WorldContext="WorldContextObject"))
 	static ENGINE_API void DrawDebugCylinder(UObject* WorldContextObject, const FVector Start, const FVector End, float Radius=100.f, int32 Segments=12, FLinearColor LineColor = FLinearColor::White, float Duration=0.f);
 	
 	/** Draw a debug cone */
-	UFUNCTION(BlueprintCallable, Category="Rendering|Debug", meta=(HidePin="WorldContextObject", DefaultToSelf="WorldContextObject", DeprecatedFunction, DeprecationMessage="DrawDebugCone has been changed to use degrees for angles instead of radians. Place a new DrawDebugCone node and pass your angles as degrees."))
+	UFUNCTION(BlueprintCallable, Category="Rendering|Debug", meta=(WorldContext="WorldContextObject", DeprecatedFunction, DeprecationMessage="DrawDebugCone has been changed to use degrees for angles instead of radians. Place a new DrawDebugCone node and pass your angles as degrees."))
 	static ENGINE_API void DrawDebugCone(UObject* WorldContextObject, const FVector Origin, const FVector Direction, float Length, float AngleWidth, float AngleHeight, int32 NumSides, FLinearColor LineColor);
 
 	/** 
 	 * Draw a debug cone 
 	 * Angles are specified in degrees
 	 */
-	UFUNCTION(BlueprintCallable, Category="Rendering|Debug", meta=(HidePin="WorldContextObject", DefaultToSelf="WorldContextObject", FriendlyName="DrawDebugCone"))
+	UFUNCTION(BlueprintCallable, Category="Rendering|Debug", meta=(WorldContext="WorldContextObject", FriendlyName="DrawDebugCone"))
 	static ENGINE_API void DrawDebugConeInDegrees(UObject* WorldContextObject, const FVector Origin, const FVector Direction, float Length=100.f, float AngleWidth=45.f, float AngleHeight=45.f, int32 NumSides = 12, FLinearColor LineColor = FLinearColor::White, float Duration=0.f);
 
 	/** Draw a debug capsule */
-	UFUNCTION(BlueprintCallable, Category="Rendering|Debug", meta=(HidePin="WorldContextObject", DefaultToSelf="WorldContextObject"))
+	UFUNCTION(BlueprintCallable, Category="Rendering|Debug", meta=(WorldContext="WorldContextObject"))
 	static ENGINE_API void DrawDebugCapsule(UObject* WorldContextObject, const FVector Center, float HalfHeight, float Radius, const FRotator Rotation, FLinearColor LineColor = FLinearColor::White, float Duration=0.f);
 
 	/** Draw a debug string at a 3d world location. */
-	UFUNCTION(BlueprintCallable, Category="Rendering|Debug", meta=(HidePin="WorldContextObject", DefaultToSelf="WorldContextObject"))
+	UFUNCTION(BlueprintCallable, Category="Rendering|Debug", meta=(WorldContext="WorldContextObject"))
 	static ENGINE_API void DrawDebugString(UObject* WorldContextObject, const FVector TextLocation, const FString& Text, class AActor* TestBaseActor = NULL, FLinearColor TextColor = FLinearColor::White, float Duration=0.f);
 	/** 
 	 * Removes all debug strings. 
 	 *
 	 * @param WorldContext	World context
 	 */
-	UFUNCTION(BlueprintCallable, Category="Rendering|Debug", meta=(HidePin="WorldContextObject", DefaultToSelf="WorldContextObject"))
+	UFUNCTION(BlueprintCallable, Category="Rendering|Debug", meta=(WorldContext="WorldContextObject"))
 	static ENGINE_API void FlushDebugStrings(UObject* WorldContextObject);
 
 	/** Draws a debug plane. */
-	UFUNCTION(BlueprintCallable, Category="Rendering|Debug", meta=(HidePin="WorldContextObject", DefaultToSelf="WorldContextObject"))
+	UFUNCTION(BlueprintCallable, Category="Rendering|Debug", meta=(WorldContext="WorldContextObject"))
 	static ENGINE_API void DrawDebugPlane(UObject* WorldContextObject, const FPlane& PlaneCoordinates, const FVector Location, float Size, FLinearColor PlaneColor = FLinearColor::White, float Duration=0.f);
 
 	/** 
@@ -1068,11 +1068,11 @@ class UKismetSystemLibrary : public UBlueprintFunctionLibrary
 	 *
 	 * @param WorldContext	World context
 	 */
-	UFUNCTION(BlueprintCallable, Category="Rendering|Debug", meta=(HidePin="WorldContextObject", DefaultToSelf="WorldContextObject"))
+	UFUNCTION(BlueprintCallable, Category="Rendering|Debug", meta=(WorldContext="WorldContextObject"))
 	static ENGINE_API void FlushPersistentDebugLines(UObject* WorldContextObject);
 
 	/** Draws a debug frustum. */
-	UFUNCTION(BlueprintCallable, Category="Rendering|Debug", meta=(HidePin="WorldContextObject", DefaultToSelf="WorldContextObject"))
+	UFUNCTION(BlueprintCallable, Category="Rendering|Debug", meta=(WorldContext="WorldContextObject"))
 	static ENGINE_API void DrawDebugFrustum(UObject* WorldContextObject, const FTransform& FrustumTransform, FLinearColor FrustumColor = FLinearColor::White, float Duration=0.f);
 
 	/** Draw a debug camera shape. */
diff --git a/Engine/Source/Runtime/UMG/Public/Blueprint/WidgetBlueprintLibrary.h b/Engine/Source/Runtime/UMG/Public/Blueprint/WidgetBlueprintLibrary.h
--- a/Engine/Source/Runtime/UMG/Public/Blueprint/WidgetBlueprintLibrary.h
+++ b/Engine/Source/Runtime/UMG/Public/Blueprint/WidgetBlueprintLibrary.h
@@ -10,7 +10,7 @@ class UWidgetBlueprintLibrary : public UBlueprintFunctionLibrary
 	GENERATED_UCLASS_BODY()
 
 	/** Creates a widget */
-	UFUNCTION(BlueprintCallable, meta=( HidePin="WorldContextObject", DefaultToSelf="WorldContextObject", FriendlyName = "Create Widget", BlueprintInternalUseOnly = "true" ), Category="User Interface|Widget")
+	UFUNCTION(BlueprintCallable, meta=( WorldContext="WorldContextObject", FriendlyName = "Create Widget", BlueprintInternalUseOnly = "true" ), Category="User Interface|Widget")
 	static class UUserWidget* Create(UObject* WorldContextObject, TSubclassOf<class UUserWidget> WidgetType, APlayerController* OwningPlayer);
 
 	UFUNCTION(BlueprintCallable, Category="User Interface|Drag and Drop", meta=( BlueprintInternalUseOnly = "true" ))
