#ue4 - Plane constraints are now supported by RotatingMovementComponent. This only affects translation caused by PivotTranslation.

diff --git a/Engine/Source/Runtime/Engine/Classes/GameFramework/RotatingMovementComponent.h b/Engine/Source/Runtime/Engine/Classes/GameFramework/RotatingMovementComponent.h
--- a/Engine/Source/Runtime/Engine/Classes/GameFramework/RotatingMovementComponent.h
+++ b/Engine/Source/Runtime/Engine/Classes/GameFramework/RotatingMovementComponent.h
@@ -7,7 +7,7 @@
 /**
  * Movement component updates position of associated PrimitiveComponent during its tick.
  */
-UCLASS(ClassGroup=Movement, meta=(BlueprintSpawnableComponent), HideCategories=(PlanarMovement, "Components|Movement|Planar", Velocity))
+UCLASS(ClassGroup=Movement, meta=(BlueprintSpawnableComponent), HideCategories=(Velocity))
 class ENGINE_API URotatingMovementComponent : public UMovementComponent
 {
 	GENERATED_UCLASS_BODY()
diff --git a/Engine/Source/Runtime/Engine/Private/RotatingMovementComponent.cpp b/Engine/Source/Runtime/Engine/Private/RotatingMovementComponent.cpp
--- a/Engine/Source/Runtime/Engine/Private/RotatingMovementComponent.cpp
+++ b/Engine/Source/Runtime/Engine/Private/RotatingMovementComponent.cpp
@@ -13,7 +13,7 @@ URotatingMovementComponent::URotatingMovementComponent(const class FPostConstruc
 
 void URotatingMovementComponent::TickComponent(float DeltaTime, enum ELevelTick TickType, FActorComponentTickFunction *ThisTickFunction)
 {
-	// skip if don't want component updated when not rendered or updated component can't move
+	// skip if we don't want component updated when not rendered or if updated component can't move
 	if ( ShouldSkipUpdate(DeltaTime) )
 	{
 		return;
@@ -30,7 +30,7 @@ void URotatingMovementComponent::TickComponent(float DeltaTime, enum ELevelTick
 	{
 		const FVector OldPivot = OldRotation.RotateVector(PivotTranslation);
 		const FVector NewPivot = NewRotation.RotateVector(PivotTranslation);
-		NewLocation = (NewLocation + OldPivot - NewPivot);
+		NewLocation = NewLocation + ConstrainDirectionToPlane(OldPivot - NewPivot);
 	}
 
 	UpdatedComponent->SetWorldLocationAndRotation(NewLocation, NewRotation);
