missing file (fix build)

diff --git a/Engine/Source/Runtime/Engine/Classes/Engine/StaticMesh.h b/Engine/Source/Runtime/Engine/Classes/Engine/StaticMesh.h
--- a/Engine/Source/Runtime/Engine/Classes/Engine/StaticMesh.h
+++ b/Engine/Source/Runtime/Engine/Classes/Engine/StaticMesh.h
@@ -427,6 +427,10 @@ class UStaticMesh : public UObject, public IInterface_CollisionDataProvider, pub
 	UPROPERTY()
 	FAssetEditorOrbitCameraPosition EditorCameraPosition;
 
+	/** If the user has modified collision in any way or has custom collision imported. Used for determining if to auto generate collision on import */
+	UPROPERTY()
+	bool bCustomizedCollision;
+
 #endif // WITH_EDITORONLY_DATA
 
 	/** For simplified meshes, this is the CRC of the high res mesh we were originally duplicated from. */
