#ue4 - Added DefaultFluidFriction to PhysicsSettings.

- Addresses inability to edit DefaultPhysicsVolume's setting of fluid friction.

diff --git a/Engine/Source/Runtime/Engine/Classes/PhysicsEngine/PhysicsSettings.h b/Engine/Source/Runtime/Engine/Classes/PhysicsEngine/PhysicsSettings.h
--- a/Engine/Source/Runtime/Engine/Classes/PhysicsEngine/PhysicsSettings.h
+++ b/Engine/Source/Runtime/Engine/Classes/PhysicsEngine/PhysicsSettings.h
@@ -78,9 +78,13 @@ class ENGINE_API UPhysicsSettings
 	UPROPERTY(config, EditAnywhere, Category = Constants)
 	float DefaultGravityZ;
 
-	/** Default Terminal Velocity. */
+	/** Default terminal velocity for Physics Volumes. */
 	UPROPERTY(config, EditAnywhere, Category = Constants)
 	float DefaultTerminalVelocity;
+	
+	/** Default fluid friction for Physics Volumes. */
+	UPROPERTY(config, EditAnywhere, Category = Constants)
+	float DefaultFluidFriction;
 
 	/** Enables the use of an async scene */
 	UPROPERTY(config, EditAnywhere, AdvancedDisplay, Category=Simulation)
diff --git a/Engine/Source/Runtime/Engine/Private/PhysicsEngine/PhysicsSettings.cpp b/Engine/Source/Runtime/Engine/Private/PhysicsEngine/PhysicsSettings.cpp
--- a/Engine/Source/Runtime/Engine/Private/PhysicsEngine/PhysicsSettings.cpp
+++ b/Engine/Source/Runtime/Engine/Private/PhysicsEngine/PhysicsSettings.cpp
@@ -6,6 +6,7 @@ UPhysicsSettings::UPhysicsSettings(const class FPostConstructInitializePropertie
 	: Super(PCIP)
 	, DefaultGravityZ(-980.f)
 	, DefaultTerminalVelocity(4000.f)
+	, DefaultFluidFriction(0.3f)
 	, bEnableAsyncScene(false)
 	, bEnable2DPhysics(false)
 	, MaxPhysicsDeltaTime(1.f / 30.f)
diff --git a/Engine/Source/Runtime/Engine/Private/PhysicsVolume.cpp b/Engine/Source/Runtime/Engine/Private/PhysicsVolume.cpp
--- a/Engine/Source/Runtime/Engine/Private/PhysicsVolume.cpp
+++ b/Engine/Source/Runtime/Engine/Private/PhysicsVolume.cpp
@@ -10,7 +10,7 @@ APhysicsVolume::APhysicsVolume(const class FPostConstructInitializeProperties& P
 	static FName CollisionProfileName(TEXT("OverlapAllDynamic"));
 	BrushComponent->SetCollisionProfileName(CollisionProfileName);
 
-	FluidFriction = 0.3f;
+	FluidFriction = UPhysicsSettings::Get()->DefaultFluidFriction;
 	TerminalVelocity = UPhysicsSettings::Get()->DefaultTerminalVelocity;
 	bAlwaysRelevant = true;
 	NetUpdateFrequency = 0.1f;
diff --git a/Engine/Source/Runtime/Engine/Private/World.cpp b/Engine/Source/Runtime/Engine/Private/World.cpp
--- a/Engine/Source/Runtime/Engine/Private/World.cpp
+++ b/Engine/Source/Runtime/Engine/Private/World.cpp
@@ -896,8 +896,9 @@ void UWorld::InitWorld(const InitializationValues IVS)
 		DefaultPhysicsVolume->Priority = -1000000;
 	}
 
-	// update terminal velocity when world is restarted
+	// update default values when world is restarted
 	DefaultPhysicsVolume->TerminalVelocity = UPhysicsSettings::Get()->DefaultTerminalVelocity;
+	DefaultPhysicsVolume->FluidFriction = UPhysicsSettings::Get()->DefaultFluidFriction;
 
 	// Find gravity
 	if (GetPhysicsScene())
