#ue4 - RotatingMovementComponent uses MoveUpdatedComponent() to perform movement.

- Enforces plane constraint, and allows custom movement overrides since MoveUpdatedComponent() is virtual.

diff --git a/Engine/Source/Runtime/Engine/Private/RotatingMovementComponent.cpp b/Engine/Source/Runtime/Engine/Private/RotatingMovementComponent.cpp
--- a/Engine/Source/Runtime/Engine/Private/RotatingMovementComponent.cpp
+++ b/Engine/Source/Runtime/Engine/Private/RotatingMovementComponent.cpp
@@ -25,13 +25,14 @@ void URotatingMovementComponent::TickComponent(float DeltaTime, enum ELevelTick
 	const FQuat NewRotation = bRotationInLocalSpace ? (OldRotation * DeltaRotation) : (DeltaRotation * OldRotation);
 
 	// Compute new location
-	FVector NewLocation = UpdatedComponent->GetComponentLocation();
+	FVector DeltaLocation = FVector::ZeroVector;
 	if (!PivotTranslation.IsZero())
 	{
 		const FVector OldPivot = OldRotation.RotateVector(PivotTranslation);
 		const FVector NewPivot = NewRotation.RotateVector(PivotTranslation);
-		NewLocation = NewLocation + ConstrainDirectionToPlane(OldPivot - NewPivot);
+		DeltaLocation = (OldPivot - NewPivot); // ConstrainDirectionToPlane() not necessary because it's done by MoveUpdatedComponent() below.
 	}
 
-	UpdatedComponent->SetWorldLocationAndRotation(NewLocation, NewRotation);
+	const bool bEnableCollision = false;
+	MoveUpdatedComponent(DeltaLocation, NewRotation.Rotator(), bEnableCollision);
 }
