Engine: Remove RenderingSettings.h and AudioSettings.h from Engine.h, including them where actually used

diff --git a/Engine/Source/Editor/DetailCustomizations/Private/AudioSettingsDetails.cpp b/Engine/Source/Editor/DetailCustomizations/Private/AudioSettingsDetails.cpp
--- a/Engine/Source/Editor/DetailCustomizations/Private/AudioSettingsDetails.cpp
+++ b/Engine/Source/Editor/DetailCustomizations/Private/AudioSettingsDetails.cpp
@@ -2,6 +2,7 @@
 
 #include "DetailCustomizationsPrivatePCH.h"
 #include "AudioSettingsDetails.h"
+#include "Runtime/Engine/Classes/Sound/AudioSettings.h"
 
 TSharedRef<IDetailCustomization> FAudioSettingsDetails::MakeInstance()
 {
diff --git a/Engine/Source/Editor/LevelEditor/Private/SLevelViewport.cpp b/Engine/Source/Editor/LevelEditor/Private/SLevelViewport.cpp
--- a/Engine/Source/Editor/LevelEditor/Private/SLevelViewport.cpp
+++ b/Engine/Source/Editor/LevelEditor/Private/SLevelViewport.cpp
@@ -29,6 +29,7 @@
 #include "SceneOutlinerTreeItems.h"
 #include "BufferVisualizationData.h"
 #include "EditorViewportCommands.h"
+#include "Runtime/Engine/Classes/Engine/RendererSettings.h"
 
 static const FName LevelEditorName("LevelEditor");
 
diff --git a/Engine/Source/Editor/ProjectSettingsViewer/Private/ProjectSettingsViewerModule.cpp b/Engine/Source/Editor/ProjectSettingsViewer/Private/ProjectSettingsViewerModule.cpp
--- a/Engine/Source/Editor/ProjectSettingsViewer/Private/ProjectSettingsViewerModule.cpp
+++ b/Engine/Source/Editor/ProjectSettingsViewer/Private/ProjectSettingsViewerModule.cpp
@@ -4,7 +4,8 @@
 #include "Engine/Console.h"
 #include "ProjectTargetPlatformEditor.h"
 #include "CookerSettings.h"
-
+#include "Runtime/Engine/Classes/Engine/RendererSettings.h"
+#include "Runtime/Engine/Classes/Sound/AudioSettings.h"
 
 #define LOCTEXT_NAMESPACE "FProjectSettingsViewerModule"
 
diff --git a/Engine/Source/Editor/UMGEditor/Private/Designer/SDesignerView.cpp b/Engine/Source/Editor/UMGEditor/Private/Designer/SDesignerView.cpp
--- a/Engine/Source/Editor/UMGEditor/Private/Designer/SDesignerView.cpp
+++ b/Engine/Source/Editor/UMGEditor/Private/Designer/SDesignerView.cpp
@@ -13,6 +13,7 @@
 #include "SDesignerToolBar.h"
 #include "DesignerCommands.h"
 #include "STransformHandle.h"
+#include "Runtime/Engine/Classes/Engine/RendererSettings.h"
 
 #define LOCTEXT_NAMESPACE "UMG"
 
diff --git a/Engine/Source/Editor/UnrealEd/Private/PlayLevel.cpp b/Engine/Source/Editor/UnrealEd/Private/PlayLevel.cpp
--- a/Engine/Source/Editor/UnrealEd/Private/PlayLevel.cpp
+++ b/Engine/Source/Editor/UnrealEd/Private/PlayLevel.cpp
@@ -31,6 +31,7 @@
 #include "EngineAnalytics.h"
 #include "Engine/GameInstance.h"
 #include "EditorAnalytics.h"
+#include "Runtime/Engine/Classes/Engine/RendererSettings.h"
 
 #include "Online.h"
 
diff --git a/Engine/Source/Runtime/Engine/Private/GameEngine.cpp b/Engine/Source/Runtime/Engine/Private/GameEngine.cpp
--- a/Engine/Source/Runtime/Engine/Private/GameEngine.cpp
+++ b/Engine/Source/Runtime/Engine/Private/GameEngine.cpp
@@ -15,7 +15,7 @@
 #include "ConfigCacheIni.h"
 #include "EngineModule.h"
 #include "Engine/GameInstance.h"
-
+#include "Engine/RendererSettings.h"
 #include "AVIWriter.h"
 
 #include "Slate.h"
diff --git a/Engine/Source/Runtime/Engine/Private/RenderingSettings.cpp b/Engine/Source/Runtime/Engine/Private/RenderingSettings.cpp
--- a/Engine/Source/Runtime/Engine/Private/RenderingSettings.cpp
+++ b/Engine/Source/Runtime/Engine/Private/RenderingSettings.cpp
@@ -1,6 +1,7 @@
 // Copyright 1998-2014 Epic Games, Inc. All Rights Reserved.
 
 #include "EnginePrivate.h"
+#include "Engine/RendererSettings.h"
 
 URendererSettings::URendererSettings(const class FPostConstructInitializeProperties& PCIP)
 	: Super(PCIP)
diff --git a/Engine/Source/Runtime/Engine/Public/Engine.h b/Engine/Source/Runtime/Engine/Public/Engine.h
--- a/Engine/Source/Runtime/Engine/Public/Engine.h
+++ b/Engine/Source/Runtime/Engine/Public/Engine.h
@@ -188,7 +188,6 @@
 #include "Animation/AnimNotifies/AnimNotifyState_Trail.h"
 #include "EditorFramework/AssetImportData.h"
 #include "Engine/AssetUserData.h"
-#include "Sound/AudioSettings.h"
 #include "AI/Navigation/AvoidanceManager.h"
 #include "Engine/BlendableInterface.h"
 #include "Engine/BlueprintCore.h"
@@ -286,7 +285,6 @@
 #include "Engine/NetConnection.h"
 #include "Engine/ChildConnection.h"
 #include "Engine/Polys.h"
-#include "Engine/RendererSettings.h"
 #include "Sound/ReverbEffect.h"
 #include "MovieScene/RuntimeMovieScenePlayerInterface.h"
 #include "GameFramework/SaveGame.h"
diff --git a/Engine/Source/Runtime/Online/OnlineSubsystemUtils/Private/OnlineSubsystemUtils.cpp b/Engine/Source/Runtime/Online/OnlineSubsystemUtils/Private/OnlineSubsystemUtils.cpp
--- a/Engine/Source/Runtime/Online/OnlineSubsystemUtils/Private/OnlineSubsystemUtils.cpp
+++ b/Engine/Source/Runtime/Online/OnlineSubsystemUtils/Private/OnlineSubsystemUtils.cpp
@@ -10,6 +10,7 @@
 
 #include "Voice.h"
 #include "SoundDefinitions.h"
+#include "Runtime/Engine/Classes/Sound/AudioSettings.h"
 
 // Testing classes
 #include "Tests/TestFriendsInterface.h"
diff --git a/Engine/Source/Runtime/UMG/Private/WidgetLayoutLibrary.cpp b/Engine/Source/Runtime/UMG/Private/WidgetLayoutLibrary.cpp
--- a/Engine/Source/Runtime/UMG/Private/WidgetLayoutLibrary.cpp
+++ b/Engine/Source/Runtime/UMG/Private/WidgetLayoutLibrary.cpp
@@ -3,6 +3,7 @@
 #include "UMGPrivatePCH.h"
 #include "Slate/SlateBrushAsset.h"
 #include "WidgetLayoutLibrary.h"
+#include "Runtime/Engine/Classes/Engine/RendererSettings.h"
 
 #define LOCTEXT_NAMESPACE "UMG"
 
