Engine: Fix typo in a comment

diff --git a/Engine/Source/Runtime/Engine/Public/SystemSettings.h b/Engine/Source/Runtime/Engine/Public/SystemSettings.h
--- a/Engine/Source/Runtime/Engine/Public/SystemSettings.h
+++ b/Engine/Source/Runtime/Engine/Public/SystemSettings.h
@@ -85,7 +85,7 @@ private:
 
 
 /**
- * Class that actually holds the current system setttings
+ * Class that actually holds the current system settings
  *
  * Derive from FSystemSettingsData instead of holding one purely for backwards
  * compatibility reasons (each data element used to be a member of the class).
