fixed parent node usage in navigation octree, fixed updating parent chain in navigtion octree after removing last attached node

diff --git a/Engine/Source/Runtime/Engine/Private/AI/Navigation/NavigationSystem.cpp b/Engine/Source/Runtime/Engine/Private/AI/Navigation/NavigationSystem.cpp
--- a/Engine/Source/Runtime/Engine/Private/AI/Navigation/NavigationSystem.cpp
+++ b/Engine/Source/Runtime/Engine/Private/AI/Navigation/NavigationSystem.cpp
@@ -1906,7 +1906,8 @@ void UNavigationSystem::AddElementToNavOctree(const FNavigationDirtyElement& Dir
 	{
 		// check if parent node is waiting in queue
 		const FSetElementId ParentRequestId = PendingOctreeUpdates.FindId(FNavigationDirtyElement(ParentNode));
-		if (ParentRequestId.IsValidId())
+		const FOctreeElementId* ParentId = GetObjectsNavOctreeId(ParentNode);
+		if (ParentRequestId.IsValidId() && ParentId == NULL)
 		{
 			FNavigationDirtyElement& ParentNode = PendingOctreeUpdates[ParentRequestId];
 			AddElementToNavOctree(ParentNode);
@@ -1915,11 +1916,11 @@ void UNavigationSystem::AddElementToNavOctree(const FNavigationDirtyElement& Dir
 			ParentNode.bInvalidRequest = true;
 		}
 
-		const FOctreeElementId* ParentId = GetObjectsNavOctreeId(ParentNode);
-		if (ParentId)
+		const FOctreeElementId* UseParentId = ParentId ? ParentId : GetObjectsNavOctreeId(ParentNode);
+		if (UseParentId)
 		{
 			UE_LOG(LogNavOctree, Log, TEXT("ADD %s to %s"), *GetNameSafe(ElementOwner), *GetNameSafe(ParentNode));
-			NavOctree->AppendToNode(*ParentId, DirtyElement.NavInterface, ElementBounds, GeneratedData);
+			NavOctree->AppendToNode(*UseParentId, DirtyElement.NavInterface, ElementBounds, GeneratedData);
 		}
 		else 
 		{
@@ -2137,15 +2138,16 @@ void UNavigationSystem::UpdateNavOctreeElement(UObject* ElementOwner, class INav
 
 void UNavigationSystem::UpdateNavOctreeParentChain(UObject* ElementOwner)
 {
+	INavRelevantInterface* ElementInterface = InterfaceCast<INavRelevantInterface>(ElementOwner);
 	TArray<FWeakObjectPtr> ChildNodes;
 
 	OctreeChildNodesMap.MultiFind(ElementOwner, ChildNodes);
 	if (ChildNodes.Num() == 0)
 	{
+		UpdateNavOctreeElement(ElementOwner, ElementInterface, UpdateFlags);
 		return;
 	}
 
-	INavRelevantInterface* ElementInterface = InterfaceCast<INavRelevantInterface>(ElementOwner);
 	const int32 UpdateFlags = OctreeUpdate_ParentChain | OctreeUpdate_Refresh;
 	TArray<INavRelevantInterface*> ChildNavInterfaces;
 	ChildNavInterfaces.AddZeroed(ChildNodes.Num());
