compilation fix

diff --git a/Engine/Source/Runtime/Engine/Private/AI/Navigation/NavigationSystem.cpp b/Engine/Source/Runtime/Engine/Private/AI/Navigation/NavigationSystem.cpp
--- a/Engine/Source/Runtime/Engine/Private/AI/Navigation/NavigationSystem.cpp
+++ b/Engine/Source/Runtime/Engine/Private/AI/Navigation/NavigationSystem.cpp
@@ -2139,16 +2139,17 @@ void UNavigationSystem::UpdateNavOctreeElement(UObject* ElementOwner, class INav
 void UNavigationSystem::UpdateNavOctreeParentChain(UObject* ElementOwner)
 {
 	INavRelevantInterface* ElementInterface = InterfaceCast<INavRelevantInterface>(ElementOwner);
-	TArray<FWeakObjectPtr> ChildNodes;
+	const int32 UpdateFlags = OctreeUpdate_ParentChain | OctreeUpdate_Refresh;
 
+	TArray<FWeakObjectPtr> ChildNodes;
 	OctreeChildNodesMap.MultiFind(ElementOwner, ChildNodes);
+
 	if (ChildNodes.Num() == 0)
 	{
 		UpdateNavOctreeElement(ElementOwner, ElementInterface, UpdateFlags);
 		return;
 	}
 
-	const int32 UpdateFlags = OctreeUpdate_ParentChain | OctreeUpdate_Refresh;
 	TArray<INavRelevantInterface*> ChildNavInterfaces;
 	ChildNavInterfaces.AddZeroed(ChildNodes.Num());
 	
