Fixed clickboxes when aspect ratio is constrained

diff --git a/Engine/Source/Runtime/Engine/Private/HUD.cpp b/Engine/Source/Runtime/Engine/Private/HUD.cpp
--- a/Engine/Source/Runtime/Engine/Private/HUD.cpp
+++ b/Engine/Source/Runtime/Engine/Private/HUD.cpp
@@ -94,11 +94,11 @@ FVector2D AHUD::GetCoordinateOffset() const
 
 		if (SceneView)
 		{
-			Offset.X = (SceneView->ViewRect.Min.X - SceneView->UnconstrainedViewRect.Min.X) // This accounts for the borders when the aspect ratio is locked
-							- SceneView->UnconstrainedViewRect.Min.X;						// And this will deal with the viewport offset if its a split screen
+			Offset.X = (SceneView->ViewRect.Min.X - SceneView->UnscaledViewRect.Min.X) // This accounts for the borders when the aspect ratio is locked
+				- SceneView->UnscaledViewRect.Min.X;						// And this will deal with the viewport offset if its a split screen
 
-			Offset.Y = (SceneView->ViewRect.Min.Y - SceneView->UnconstrainedViewRect.Min.Y)
-							- SceneView->UnconstrainedViewRect.Min.Y;
+			Offset.Y = (SceneView->ViewRect.Min.Y - SceneView->UnscaledViewRect.Min.Y)
+				- SceneView->UnscaledViewRect.Min.Y;
 		}
 	}
 
