Added metadata for materialgraph nodes

diff --git a/Engine/Source/Editor/GraphEditor/Private/MaterialNodes/SGraphNodeMaterialBase.cpp b/Engine/Source/Editor/GraphEditor/Private/MaterialNodes/SGraphNodeMaterialBase.cpp
--- a/Engine/Source/Editor/GraphEditor/Private/MaterialNodes/SGraphNodeMaterialBase.cpp
+++ b/Engine/Source/Editor/GraphEditor/Private/MaterialNodes/SGraphNodeMaterialBase.cpp
@@ -2,9 +2,13 @@
 
 #include "GraphEditorCommon.h"
 #include "Materials/MaterialExpression.h"
+#include "Materials/MaterialFunction.h"
 #include "SGraphNodeMaterialBase.h"
 #include "ScopedTransaction.h"
 #include "Runtime/Engine/Public/Slate/SceneViewport.h"
+#include "TutorialMetaData.h"
+
+
 
 /**
 * Simple representation of the backbuffer that the preview canvas renders to
@@ -445,3 +449,18 @@ const FSlateBrush* SGraphNodeMaterialBase::GetExpressionPreviewArrow() const
 {
 	return FEditorStyle::GetBrush(MaterialNode->MaterialExpression->bCollapsed ? TEXT("Kismet.TitleBarEditor.ArrowDown") : TEXT("Kismet.TitleBarEditor.ArrowUp"));
 }
+
+void SGraphNodeMaterialBase::PopulateMetaTag(FGraphNodeMetaData* TagMeta) const
+{
+	if (GraphNode != nullptr)
+	{
+		UMaterialGraph* OuterGraph = MaterialNode->GetTypedOuter<UMaterialGraph>();
+		if ((OuterGraph != nullptr) && (MaterialNode->MaterialExpression != nullptr) )
+		{
+			TagMeta->OuterName = OuterGraph->OriginalMaterialFullName;
+			TagMeta->GUID = MaterialNode->MaterialExpression->MaterialExpressionGuid;
+			TagMeta->Tag = FName(*FString::Printf(TEXT("MaterialExprNode_%s_%s"), *TagMeta->OuterName, *TagMeta->GUID.ToString()));
+		}
+		TagMeta->FriendlyName = FString::Printf(TEXT("%s expression node in %s"), *GraphNode->GetNodeTitle(ENodeTitleType::FullTitle).ToString(), *TagMeta->OuterName);		
+	}
+}
diff --git a/Engine/Source/Editor/GraphEditor/Private/MaterialNodes/SGraphNodeMaterialBase.h b/Engine/Source/Editor/GraphEditor/Private/MaterialNodes/SGraphNodeMaterialBase.h
--- a/Engine/Source/Editor/GraphEditor/Private/MaterialNodes/SGraphNodeMaterialBase.h
+++ b/Engine/Source/Editor/GraphEditor/Private/MaterialNodes/SGraphNodeMaterialBase.h
@@ -73,6 +73,9 @@ public:
 
 	UMaterialGraphNode* GetMaterialGraphNode() const {return MaterialNode;}
 
+	/* Populate a meta data tag with information about this graph node */
+	virtual void PopulateMetaTag(class FGraphNodeMetaData* TagMeta) const override;
+
 protected:
 	// SGraphNode interface
 	virtual void AddPin( const TSharedRef<SGraphPin>& PinToAdd ) override;
diff --git a/Engine/Source/Editor/GraphEditor/Private/MaterialNodes/SGraphNodeMaterialResult.cpp b/Engine/Source/Editor/GraphEditor/Private/MaterialNodes/SGraphNodeMaterialResult.cpp
--- a/Engine/Source/Editor/GraphEditor/Private/MaterialNodes/SGraphNodeMaterialResult.cpp
+++ b/Engine/Source/Editor/GraphEditor/Private/MaterialNodes/SGraphNodeMaterialResult.cpp
@@ -2,6 +2,7 @@
 
 #include "GraphEditorCommon.h"
 #include "SGraphNodeMaterialResult.h"
+#include "TutorialMetaData.h"
 
 /////////////////////////////////////////////////////
 // SGraphNodeMaterialResult
@@ -58,3 +59,19 @@ void SGraphNodeMaterialResult::MoveTo(const FVector2D& NewPosition, FNodeSet& No
 	RootNode->Material->MaterialGraph->MaterialDirtyDelegate.ExecuteIfBound();
 }
 
+
+void SGraphNodeMaterialResult::PopulateMetaTag(FGraphNodeMetaData* TagMeta) const
+{
+	if( (GraphNode != nullptr) && (RootNode != nullptr) )
+	{		
+		UMaterialGraph* OuterGraph = RootNode->GetTypedOuter<UMaterialGraph>();
+		if (OuterGraph != nullptr)
+		{
+			TagMeta->OuterName = OuterGraph->OriginalMaterialFullName;
+			// There is only one root node - so we dont need a guid. 
+			TagMeta->Tag = FName(*FString::Printf(TEXT("MaterialResNode_%s"), *TagMeta->OuterName));
+			TagMeta->GUID.Invalidate();
+			TagMeta->FriendlyName = FString::Printf(TEXT("Material Result node in %s"), *TagMeta->OuterName);
+ 		}		
+	}
+}
diff --git a/Engine/Source/Editor/GraphEditor/Private/MaterialNodes/SGraphNodeMaterialResult.h b/Engine/Source/Editor/GraphEditor/Private/MaterialNodes/SGraphNodeMaterialResult.h
--- a/Engine/Source/Editor/GraphEditor/Private/MaterialNodes/SGraphNodeMaterialResult.h
+++ b/Engine/Source/Editor/GraphEditor/Private/MaterialNodes/SGraphNodeMaterialResult.h
@@ -12,6 +12,7 @@ public:
 
 	// SGraphNode interface
 	virtual void CreatePinWidgets() override;
+	virtual void PopulateMetaTag(FGraphNodeMetaData* TagMeta) const override;
 	// End of SGraphNode interface
 
 	// SNodePanel::SNode interface
diff --git a/Engine/Source/Editor/IntroTutorials/Private/STutorialOverlay.cpp b/Engine/Source/Editor/IntroTutorials/Private/STutorialOverlay.cpp
--- a/Engine/Source/Editor/IntroTutorials/Private/STutorialOverlay.cpp
+++ b/Engine/Source/Editor/IntroTutorials/Private/STutorialOverlay.cpp
@@ -148,12 +148,18 @@ void STutorialOverlay::PerformWidgetInteractions(const FTutorialWidgetContent &W
 
 void STutorialOverlay::OpenBrowserForWidgetAnchor(const FTutorialWidgetContent &WidgetContent)
 {
+	bool bTabOpened = false;
 	// Open the required tab if we found it in the map
 	if (WidgetContent.WidgetAnchor.TabTypeToOpen.IsEmpty() == false)
 	{
 		FLevelEditorModule& LevelEditorModule = FModuleManager::GetModuleChecked<FLevelEditorModule>(TEXT("LevelEditor"));
 		TSharedPtr<FTabManager> LevelEditorTabManager = LevelEditorModule.GetLevelEditorTabManager();
-		LevelEditorTabManager->InvokeTab(FName(*WidgetContent.WidgetAnchor.TabTypeToOpen));
+		TSharedPtr<SDockTab> DocTab = LevelEditorTabManager->InvokeTab(FName(*WidgetContent.WidgetAnchor.TabTypeToOpen));
+		bTabOpened = DocTab.IsValid();
+	}
+	if (bTabOpened == false)
+	{
+		//bTabOpened = FAssetEditorManager::Get().OpenEditorForAsset()
 	}
 	
 }
diff --git a/Engine/Source/Editor/IntroTutorials/Private/TutorialStructCustomization.cpp b/Engine/Source/Editor/IntroTutorials/Private/TutorialStructCustomization.cpp
--- a/Engine/Source/Editor/IntroTutorials/Private/TutorialStructCustomization.cpp
+++ b/Engine/Source/Editor/IntroTutorials/Private/TutorialStructCustomization.cpp
@@ -256,7 +256,10 @@ private:
 			
 			// Set the friendly name to the PickedWidget name - we might not have any metadata
 			StructPropertyHandle->GetChildHandle(GET_MEMBER_NAME_CHECKED(FTutorialContentAnchor, FriendlyName))->SetValue(PickedWidgetName);
-			//StructPropertyHandle->GetChildHandle(GET_MEMBER_NAME_CHECKED(FTutorialContentAnchor, FriendlyName))->SetValue(FString());
+			// Reset the other fields
+			StructPropertyHandle->GetChildHandle(GET_MEMBER_NAME_CHECKED(FTutorialContentAnchor, GUIDString))->SetValue(FString());
+			StructPropertyHandle->GetChildHandle(GET_MEMBER_NAME_CHECKED(FTutorialContentAnchor, OuterName))->SetValue(FString());
+			StructPropertyHandle->GetChildHandle(GET_MEMBER_NAME_CHECKED(FTutorialContentAnchor, TabTypeToOpen))->SetValue(FString());
 
 			for (const auto& MetaDataEntry : PickedAllMetaData)
 			{
diff --git a/Engine/Source/Editor/MaterialEditor/Private/MaterialEditor.cpp b/Engine/Source/Editor/MaterialEditor/Private/MaterialEditor.cpp
--- a/Engine/Source/Editor/MaterialEditor/Private/MaterialEditor.cpp
+++ b/Engine/Source/Editor/MaterialEditor/Private/MaterialEditor.cpp
@@ -481,6 +481,8 @@ void FMaterialEditor::InitMaterialEditor( const EToolkitMode::Type Mode, const T
 		}
 	}
 
+	// Store the name of this material (for the tutorial widget meta)
+	Material->MaterialGraph->OriginalMaterialFullName = OriginalMaterial->GetName();
 	Material->MaterialGraph->RebuildGraph();
 	RecenterEditor();
 
@@ -2620,6 +2622,9 @@ UMaterialExpression* FMaterialEditor::CreateNewMaterialExpression(UClass* NewExp
 		NewExpression->MaterialExpressionEditorX = NodePos.X;
 		NewExpression->MaterialExpressionEditorY = NodePos.Y;
 
+		// Create a GUID for the node
+		NewExpression->UpdateMaterialExpressionGuid(true, true);
+
 		if (bAutoAssignResource)
 		{
 			// If the user is adding a texture, automatically assign the currently selected texture to it.
diff --git a/Engine/Source/Editor/UnrealEd/Classes/MaterialGraph/MaterialGraph.h b/Engine/Source/Editor/UnrealEd/Classes/MaterialGraph/MaterialGraph.h
--- a/Engine/Source/Editor/UnrealEd/Classes/MaterialGraph/MaterialGraph.h
+++ b/Engine/Source/Editor/UnrealEd/Classes/MaterialGraph/MaterialGraph.h
@@ -73,6 +73,10 @@ class UMaterialGraph : public UEdGraph
 	/** Toggles the bCollapsed flag of a material expression and updates material editor */
 	FToggleExpressionCollapsed ToggleCollapsedDelegate;
 
+	/** The name of the material that we are editing */
+	UPROPERTY()
+	FString	OriginalMaterialFullName;
+
 public:
 	/**
 	 * Completely rebuild the graph from the material, removing all old nodes
diff --git a/Engine/Source/Runtime/Engine/Classes/Materials/MaterialExpression.h b/Engine/Source/Runtime/Engine/Classes/Materials/MaterialExpression.h
--- a/Engine/Source/Runtime/Engine/Classes/Materials/MaterialExpression.h
+++ b/Engine/Source/Runtime/Engine/Classes/Materials/MaterialExpression.h
@@ -97,6 +97,10 @@ class ENGINE_API UMaterialExpression : public UObject
 	/** Text of last error for this expression */
 	FString LastErrorText;
 
+	/** GUID to uniquely identify this node, to help the tutorials out */
+	UPROPERTY()
+	FGuid MaterialExpressionGuid;
+
 #endif // WITH_EDITORONLY_DATA
 	/** Set to true by RecursiveUpdateRealtimePreview() if the expression's preview needs to be updated in realtime in the material editor. */
 	UPROPERTY()
@@ -280,9 +284,9 @@ class ENGINE_API UMaterialExpression : public UObject
 	 * Connects the specified input expression to the specified output of this expression.
 	 */
 	void ConnectExpression(FExpressionInput* Input, int32 OutputIndex);
-
+	
 	/** 
-	* Generates a GUID for this expression if one doesn't already exist. 
+	* Generates a GUID for the parameter expression if one doesn't already exist and we are one.
 	*
 	* @param bForceGeneration	Whether we should generate a GUID even if it is already valid.
 	*/
@@ -296,6 +300,24 @@ class ENGINE_API UMaterialExpression : public UObject
 		return Dummy;
 	}
 
+	/**
+	* Generates a GUID for this expression if one doesn't already exist.
+	*
+	* @param bForceGeneration	Whether we should generate a GUID even if it is already valid.
+	*/
+	void UpdateMaterialExpressionGuid(bool bForceGeneration, bool bAllowMarkingPackageDirty);
+	
+	/** Return the material expression guid. */
+	virtual FGuid& GetMaterialExpressionId()
+	{		
+#if WITH_EDITORONLY_DATA
+		return MaterialExpressionGuid;
+#else
+		static FGuid Dummy;
+		return Dummy;
+#endif
+	}
+
 	/** Asserts if the expression is not contained by its Material or Function's expressions array. */
 	void ValidateState();
 
diff --git a/Engine/Source/Runtime/Engine/Private/Materials/MaterialExpressions.cpp b/Engine/Source/Runtime/Engine/Private/Materials/MaterialExpressions.cpp
--- a/Engine/Source/Runtime/Engine/Private/Materials/MaterialExpressions.cpp
+++ b/Engine/Source/Runtime/Engine/Private/Materials/MaterialExpressions.cpp
@@ -369,6 +369,7 @@ void UMaterialExpression::CopyMaterialExpressions(const TArray<UMaterialExpressi
 			}
 
 			NewExpression->UpdateParameterGuid(true, true);
+			NewExpression->UpdateMaterialExpressionGuid(true, true);
 
 			UMaterialExpressionFunctionInput* FunctionInput = Cast<UMaterialExpressionFunctionInput>( NewExpression );
 			if( FunctionInput )
@@ -465,6 +466,8 @@ void UMaterialExpression::PostInitProperties()
 	Super::PostInitProperties();
 
 	UpdateParameterGuid(false, false);
+	
+	UpdateMaterialExpressionGuid(false, true);
 }
 
 void UMaterialExpression::PostLoad()
@@ -482,6 +485,8 @@ void UMaterialExpression::PostLoad()
 	}
 
 	UpdateParameterGuid(false, false);
+	
+	UpdateMaterialExpressionGuid(false, false);
 }
 
 void UMaterialExpression::PostDuplicate(bool bDuplicateForPIE)
@@ -491,6 +496,7 @@ void UMaterialExpression::PostDuplicate(bool bDuplicateForPIE)
 	// We do not force a guid regen here because this function is used when the Material Editor makes a copy of a material to edit.
 	// If we forced a GUID regen, it would cause all of the guids for a material to change everytime a material was edited.
 	UpdateParameterGuid(false, true);
+	UpdateMaterialExpressionGuid(false, true);
 }
 
 #if WITH_EDITOR
@@ -869,6 +875,25 @@ void UMaterialExpression::ConnectExpression( FExpressionInput* Input, int32 Outp
 	}
 }
 
+void UMaterialExpression::UpdateMaterialExpressionGuid(bool bForceGeneration, bool bAllowMarkingPackageDirty)
+{
+	// If we are in the editor, and we don't have a valid GUID yet, generate one.
+	if (GIsEditor && !FApp::IsGame())
+	{
+		FGuid& Guid = GetMaterialExpressionId();
+
+		if (bForceGeneration || !Guid.IsValid())
+		{
+			Guid = FGuid::NewGuid();
+
+			if (bAllowMarkingPackageDirty)
+			{
+				MarkPackageDirty();
+			}
+		}
+	}
+}
+
 void UMaterialExpression::UpdateParameterGuid(bool bForceGeneration, bool bAllowMarkingPackageDirty)
 {
 	if (bIsParameterExpression)
