Don't play sounds in worlds that do not allow it

diff --git a/Engine/Source/Runtime/Engine/Private/AudioComponent.cpp b/Engine/Source/Runtime/Engine/Private/AudioComponent.cpp
--- a/Engine/Source/Runtime/Engine/Private/AudioComponent.cpp
+++ b/Engine/Source/Runtime/Engine/Private/AudioComponent.cpp
@@ -129,7 +129,9 @@ void UAudioComponent::Play(float StartTime)
 
 void UAudioComponent::PlayInternal(const float StartTime, const float FadeInDuration, const float FadeVolumeLevel)
 {
-	UE_LOG(LogAudio, Verbose, TEXT( "%g: Playing AudioComponent : '%s' with Sound: '%s'" ), GetWorld() ? GetWorld()->GetAudioTimeSeconds() : 0.0f, *GetFullName(), Sound ? *Sound->GetName() : TEXT( "NULL" ) );
+	UWorld* World = GetWorld();
+
+	UE_LOG(LogAudio, Verbose, TEXT( "%g: Playing AudioComponent : '%s' with Sound: '%s'" ), World ? World->GetAudioTimeSeconds() : 0.0f, *GetFullName(), Sound ? *Sound->GetName() : TEXT( "NULL" ) );
 
 	if (bIsActive)
 	{
@@ -140,7 +142,7 @@ void UAudioComponent::PlayInternal(const float StartTime, const float FadeInDura
 		bAutoDestroy = bCurrentAutoDestroy;
 	}
 
-	if (Sound)
+	if (Sound && (World == nullptr || World->bAllowAudioPlayback))
 	{
 		FAudioDevice* AudioDevice = GEngine ? GEngine->GetAudioDevice() : NULL;
 		if (AudioDevice != NULL)
diff --git a/Engine/Source/Runtime/Engine/Private/GameplayStatics.cpp b/Engine/Source/Runtime/Engine/Private/GameplayStatics.cpp
--- a/Engine/Source/Runtime/Engine/Private/GameplayStatics.cpp
+++ b/Engine/Source/Runtime/Engine/Private/GameplayStatics.cpp
@@ -605,7 +605,7 @@ void UGameplayStatics::PlaySoundAtLocation(UObject* WorldContextObject, class US
 
 	const bool bIsInGameWorld = ThisWorld->IsGameWorld();
 
-	if (GEngine && GEngine->UseSound() && ThisWorld->GetNetMode() != NM_DedicatedServer)
+	if (GEngine && GEngine->UseSound() && ThisWorld->bAllowAudioPlayback && ThisWorld->GetNetMode() != NM_DedicatedServer)
 	{
 		if( Sound->IsAudibleSimple( Location, AttenuationSettings ) )
 		{
@@ -657,9 +657,9 @@ void UGameplayStatics::PlayDialogueAtLocation(UObject* WorldContextObject, class
 		return;
 	}
 
-	if (GEngine && GEngine->UseSound())
+	UWorld* ThisWorld = GEngine->GetWorldFromContextObject(WorldContextObject);
+	if (GEngine && GEngine->UseSound() && ThisWorld->bAllowAudioPlayback)
 	{
-		UWorld* ThisWorld = GEngine->GetWorldFromContextObject(WorldContextObject);
 		const bool bIsInGameWorld = ThisWorld->IsGameWorld();
 
 		if( Sound->IsAudibleSimple( Location, AttenuationSettings ) )
