UMG - You can now configure which side it uses for UI scaling, we default to Shortest Side.  The options are Shortest, Longest, Horizontal, and Vertical.  We can add Custom in the future when we're ready to enhance it further.

diff --git a/Engine/Config/BaseEngine.ini b/Engine/Config/BaseEngine.ini
--- a/Engine/Config/BaseEngine.ini
+++ b/Engine/Config/BaseEngine.ini
@@ -1317,6 +1317,7 @@ RenderPriorityPS3=1001
 SuspendGameIO=True
 
 [/Script/Engine.RendererSettings]
+UIScaleRule=ShortestSide
 UIScaleCurve=(EditorCurveData=(Keys=((Time=480,Value=0.444),(Time=720,Value=0.666),(Time=1080,Value=1.0),(Time=8640,Value=8.0))),ExternalCurve=None)
 
 [/Script/Engine.GameEngine]
diff --git a/Engine/Source/Runtime/Engine/Classes/Engine/RendererSettings.h b/Engine/Source/Runtime/Engine/Classes/Engine/RendererSettings.h
--- a/Engine/Source/Runtime/Engine/Classes/Engine/RendererSettings.h
+++ b/Engine/Source/Runtime/Engine/Classes/Engine/RendererSettings.h
@@ -81,6 +81,26 @@ namespace EAntiAliasingMethodUI
 }
 
 
+/** The Side to use when scaling the UI. */
+UENUM()
+namespace EUIScalingRule
+{
+	enum Type
+	{
+		/** Evaluates the scale curve based on the shortest side of the viewport */
+		ShortestSide,
+		/** Evaluates the scale curve based on the longest side of the viewport */
+		LongestSide,
+		/** Evaluates the scale curve based on the X axis of the viewport */
+		Horizontal,
+		/** Evaluates the scale curve based on the Y axis of the viewport */
+		Vertical,
+		/** Custom - Allows custom rule interpretation */
+		//Custom
+	};
+}
+
+
 /**
  * Implements project settings for the Rendering sub-system.
  */
@@ -226,9 +246,14 @@ class ENGINE_API URendererSettings
 		ToolTip="Screen radius at which wireframe objects are culled. Larger values can improve performance when viewing a scene in wireframe."))
 	float WireframeCullThreshold;
 
+	UPROPERTY(config, EditAnywhere, Category=UI, meta=(
+		DisplayName="DPI Scale Rule",
+		ToolTip="The rule used when trying to decide what scale to apply." ))
+	TEnumAsByte<EUIScalingRule::Type> UIScaleRule;
+
 	UPROPERTY(config, EditAnywhere, Category=UI, meta=(
 		DisplayName="DPI Curve",
-		ToolTip="Controls how the UI is scaled at different resolutions based on the shortest side" ))
+		ToolTip="Controls how the UI is scaled at different resolutions based on the DPI Scale Rule" ))
 	FRuntimeFloatCurve UIScaleCurve;
 
 public:
diff --git a/Engine/Source/Runtime/Engine/Private/RenderingSettings.cpp b/Engine/Source/Runtime/Engine/Private/RenderingSettings.cpp
--- a/Engine/Source/Runtime/Engine/Private/RenderingSettings.cpp
+++ b/Engine/Source/Runtime/Engine/Private/RenderingSettings.cpp
@@ -82,8 +82,34 @@ void URendererSettings::PostEditChangeProperty(FPropertyChangedEvent& PropertyCh
 
 float URendererSettings::GetDPIScaleBasedOnSize(FIntPoint Size) const
 {
-	int32 ShortestSide = FMath::Min(Size.X, Size.Y);
-	const FRichCurve* DPICurve = UIScaleCurve.GetRichCurveConst();
-	float Scale = DPICurve->Eval((float)ShortestSide, 1.0f);
+	float Scale = 1;
+
+	//if ( UIScaleRule == EUIScalingRule::Custom )
+	//{
+
+	//}
+	//else
+	{
+		int32 EvalPoint = 0;
+		switch ( UIScaleRule )
+		{
+		case EUIScalingRule::ShortestSide:
+			EvalPoint = FMath::Min(Size.X, Size.Y);
+			break;
+		case EUIScalingRule::LongestSide:
+			EvalPoint = FMath::Max(Size.X, Size.Y);
+			break;
+		case EUIScalingRule::Horizontal:
+			EvalPoint = Size.X;
+			break;
+		case EUIScalingRule::Vertical:
+			EvalPoint = Size.Y;
+			break;
+		}
+
+		const FRichCurve* DPICurve = UIScaleCurve.GetRichCurveConst();
+		Scale = DPICurve->Eval((float)EvalPoint, 1.0f);
+	}
+
 	return FMath::Max(Scale, 0.01f);
 }
