Fix for landscape export crashing when exporting a landscape with both multiple subsections per component and a visibility map

diff --git a/Engine/Source/Editor/UnrealEd/Private/Fbx/FbxMainExport.cpp b/Engine/Source/Editor/UnrealEd/Private/Fbx/FbxMainExport.cpp
--- a/Engine/Source/Editor/UnrealEd/Private/Fbx/FbxMainExport.cpp
+++ b/Engine/Source/Editor/UnrealEd/Private/Fbx/FbxMainExport.cpp
@@ -2360,9 +2360,9 @@ FbxNode* FFbxExporter::ExportLandscapeToFbx(ALandscapeProxy* Landscape, const TC
 	FbxLayerElementArrayTemplate<FbxVector2>& WeightmapUVs = LayerElementWeightmapUVs->GetDirectArray();
 	WeightmapUVs.Resize(VertexCount);
 
-	TArray<uint8> VisibiltyData;
-	VisibiltyData.Empty(VertexCount);
-	VisibiltyData.AddZeroed(VertexCount);
+	TArray<uint8> VisibilityData;
+	VisibilityData.Empty(VertexCount);
+	VisibilityData.AddZeroed(VertexCount);
 
 	for (int32 ComponentIndex = 0, SelectedComponentIndex = 0; ComponentIndex < Landscape->LandscapeComponents.Num(); ComponentIndex++)
 	{
@@ -2386,9 +2386,12 @@ FbxNode* FFbxExporter::ExportLandscapeToFbx(ALandscapeProxy* Landscape, const TC
 			}
 		}
 
-		for (int32 i = 0; i < CompVisData.Num(); ++i)
+		if (CompVisData.Num() > 0)
 		{
-			VisibiltyData[BaseVertIndex + i] = CompVisData[i];
+			for (int32 i = 0; i < VertexCountPerComponent; ++i)
+			{
+				VisibilityData[BaseVertIndex + i] = CompVisData[CDI.VertexIndexToTexel(i)];
+			}
 		}
 		
 		for (int32 VertIndex = 0; VertIndex < VertexCountPerComponent; VertIndex++)
@@ -2467,7 +2470,7 @@ FbxNode* FFbxExporter::ExportLandscapeToFbx(ALandscapeProxy* Landscape, const TC
 		{
 			for (int32 X = 0; X < ComponentSizeQuads; X++)
 			{
-				if (VisibiltyData[BaseVertIndex + Y * (ComponentSizeQuads + 1) + X] < VisThreshold)
+				if (VisibilityData[BaseVertIndex + Y * (ComponentSizeQuads + 1) + X] < VisThreshold)
 				{
 					Mesh->BeginPolygon();
 					Mesh->AddPolygon(BaseVertIndex + (X + 0) + (Y + 0)*(ComponentSizeQuads + 1));
diff --git a/Engine/Source/Runtime/Engine/Private/Landscape/LandscapeDataAccess.cpp b/Engine/Source/Runtime/Engine/Private/Landscape/LandscapeDataAccess.cpp
--- a/Engine/Source/Runtime/Engine/Private/Landscape/LandscapeDataAccess.cpp
+++ b/Engine/Source/Runtime/Engine/Private/Landscape/LandscapeDataAccess.cpp
@@ -33,6 +33,7 @@ ENGINE_API FLandscapeComponentDataInterface::FLandscapeComponentDataInterface(UL
 
 	ComponentSizeVerts = (Component->ComponentSizeQuads + 1) >> MipLevel;
 	SubsectionSizeVerts = (Component->SubsectionSizeQuads + 1) >> MipLevel;
+	ComponentNumSubsections = Component->NumSubsections;
 
 	if (MipLevel < Component->HeightmapTexture->Source.GetNumMips())
 	{
diff --git a/Engine/Source/Runtime/Engine/Private/Landscape/LandscapeEdit.cpp b/Engine/Source/Runtime/Engine/Private/Landscape/LandscapeEdit.cpp
--- a/Engine/Source/Runtime/Engine/Private/Landscape/LandscapeEdit.cpp
+++ b/Engine/Source/Runtime/Engine/Private/Landscape/LandscapeEdit.cpp
@@ -2329,14 +2329,14 @@ bool ALandscapeProxy::ExportToRawMesh(int32 InExportLOD, FRawMesh& OutRawMesh) c
 			{
 				// Fill indices
 				{
-					int32 SubNumX, SubNumY, SubX, SubY;
-					CDI.ComponentXYToSubsectionXY(x, y, SubNumX, SubNumY, SubX, SubY);
 
 					// Whether this vertex is in hole
 					bool bInvisible = false;
 					if (VisDataMap.Num())
 					{
-						bInvisible = (VisDataMap[x + y * (ComponentSizeQuadsLOD + 1)] >= VisThreshold);
+						int32 TexelX, TexelY;
+						CDI.VertexXYToTexelXY(x, y, TexelX, TexelY);
+						bInvisible = (VisDataMap[CDI.TexelXYToIndex(TexelX, TexelY)] >= VisThreshold);
 					}
 
 					// triangulation matches FLandscapeIndexBuffer constructor
diff --git a/Engine/Source/Runtime/Engine/Public/Landscape/LandscapeDataAccess.h b/Engine/Source/Runtime/Engine/Public/Landscape/LandscapeDataAccess.h
--- a/Engine/Source/Runtime/Engine/Public/Landscape/LandscapeDataAccess.h
+++ b/Engine/Source/Runtime/Engine/Public/Landscape/LandscapeDataAccess.h
@@ -142,8 +142,13 @@ struct FLandscapeComponentDataInterface
 //#if LANDSCAPE_VALIDATE_DATA_ACCESS
 //		check(MipLevel == 0);
 //#endif
-		OutX = QuadIndex % ComponentSizeVerts;
-		OutY = QuadIndex / ComponentSizeVerts;
+		OutX = QuadIndex % (ComponentSizeVerts-1);
+		OutY = QuadIndex / (ComponentSizeVerts-1);
+	}
+
+	int32 VertexXYToIndex(int32 VertX, int32 VertY) const
+	{
+		return VertY * ComponentSizeVerts + VertX;
 	}
 
 	void ComponentXYToSubsectionXY(int32 CompX, int32 CompY, int32& SubNumX, int32& SubNumY, int32& SubX, int32& SubY ) const
@@ -170,6 +175,29 @@ struct FLandscapeComponentDataInterface
 		}
 	}
 
+	void VertexXYToTexelXY(int32 VertX, int32 VertY, int32& OutX, int32& OutY) const
+	{
+		int32 SubNumX, SubNumY, SubX, SubY;
+		ComponentXYToSubsectionXY(VertX, VertY, SubNumX, SubNumY, SubX, SubY);
+
+		OutX = SubNumX * SubsectionSizeVerts + SubX;
+		OutY = SubNumY * SubsectionSizeVerts + SubY;
+	}
+	
+	int32 VertexIndexToTexel(int32 VertexIndex) const
+	{
+		int32 VertX, VertY;
+		VertexIndexToXY(VertexIndex, VertX, VertY);
+		int32 TexelX, TexelY;
+		VertexXYToTexelXY(VertX, VertY, TexelX, TexelY);
+		return TexelXYToIndex(TexelX, TexelY);
+	}
+
+	int32 TexelXYToIndex(int32 TexelX, int32 TexelY) const
+	{
+		return TexelY * ComponentNumSubsections * SubsectionSizeVerts + TexelX;
+	}
+
 	FColor* GetRawHeightData() const
 	{
 		return HeightMipData;
@@ -197,7 +225,7 @@ struct FLandscapeComponentDataInterface
 
 	ENGINE_API bool GetWeightmapTextureData(ULandscapeLayerInfoObject* LayerInfo, TArray<uint8>& OutData);
 
-	FColor* GetHeightData( int32 LocalX, int32 LocalY ) const
+	FColor* GetHeightData(int32 LocalX, int32 LocalY) const
 	{
 #if LANDSCAPE_VALIDATE_DATA_ACCESS
 		check(Component);
@@ -205,13 +233,10 @@ struct FLandscapeComponentDataInterface
 		check(LocalX >=0 && LocalY >=0 && LocalX < ComponentSizeVerts && LocalY < ComponentSizeVerts );
 #endif
 
-		int32 SubNumX;
-		int32 SubNumY;
-		int32 SubX;
-		int32 SubY;
-		ComponentXYToSubsectionXY(LocalX, LocalY, SubNumX, SubNumY, SubX, SubY );
+		int32 TexelX, TexelY;
+		VertexXYToTexelXY(LocalX, LocalY, TexelX, TexelY);
 		
-		return &HeightMipData[SubX + SubNumX*SubsectionSizeVerts + HeightmapComponentOffsetX + (SubY+SubNumY*SubsectionSizeVerts+HeightmapComponentOffsetY)*HeightmapStride];
+		return &HeightMipData[TexelX + HeightmapComponentOffsetX + (TexelY + HeightmapComponentOffsetY) * HeightmapStride];
 	}
 
 	ENGINE_API FColor* GetXYOffsetData(int32 LocalX, int32 LocalY) const;
@@ -322,6 +347,7 @@ private:
 
 	int32 ComponentSizeVerts;
 	int32 SubsectionSizeVerts;
+	int32 ComponentNumSubsections;
 
 	bool bNeedToDeleteDataInterface;
 public:
