#ue4 - GameMode::RestartPlayer() does nothing if PlayerState->bOnlySpectator, meaning the player should only ever be a spectator.

TTP 346019

diff --git a/Engine/Source/Runtime/Engine/Private/GameMode.cpp b/Engine/Source/Runtime/Engine/Private/GameMode.cpp
--- a/Engine/Source/Runtime/Engine/Private/GameMode.cpp
+++ b/Engine/Source/Runtime/Engine/Private/GameMode.cpp
@@ -416,6 +416,13 @@ void AGameMode::RestartPlayer(AController* NewPlayer)
 	}
 
 	UE_LOG(LogGameMode, Verbose, TEXT("RestartPlayer %s"), (NewPlayer && NewPlayer->PlayerState) ? *NewPlayer->PlayerState->PlayerName : TEXT("Unknown"));
+
+	if (NewPlayer->PlayerState && NewPlayer->PlayerState->bOnlySpectator)
+	{
+		UE_LOG(LogGameMode, Verbose, TEXT("RestartPlayer tried to restart a spectator-only player!"));
+		return;
+	}
+
 	AActor* StartSpot = FindPlayerStart(NewPlayer);
 
 	// if a start spot wasn't found,
