Replace check with if in the collision table lookup. It's possible with sub-stepping for this to be a frame behind.

diff --git a/Engine/Source/Runtime/Engine/Private/PhysicsEngine/PhysXSupport.cpp b/Engine/Source/Runtime/Engine/Private/PhysicsEngine/PhysXSupport.cpp
--- a/Engine/Source/Runtime/Engine/Private/PhysicsEngine/PhysXSupport.cpp
+++ b/Engine/Source/Runtime/Engine/Private/PhysicsEngine/PhysXSupport.cpp
@@ -354,12 +354,15 @@ PxFilterFlags PhysXSimFilterShader(	PxFilterObjectAttributes attributes0, PxFilt
 
 		const TMap<uint32, TMap<FRigidBodyIndexPair, bool> *> & CollisionDisableTableLookup = PhysScene->GetCollisionDisableTableLookup();
 		TMap<FRigidBodyIndexPair, bool>* const * DisableTablePtrPtr = CollisionDisableTableLookup.Find(filterData1.word2);
-		check(DisableTablePtrPtr);
-		TMap<FRigidBodyIndexPair,bool>* DisableTablePtr = *DisableTablePtrPtr;
-		FRigidBodyIndexPair BodyPair(filterData0.word0, filterData1.word0); // body indexes are stored in word 0
-		if( DisableTablePtr->Find(BodyPair) )
+		if (DisableTablePtrPtr)		//Since collision table is deferred during sub-stepping it's possible that we won't get the collision disable table until the next frame
 		{
-			return PxFilterFlag::eKILL;
+			TMap<FRigidBodyIndexPair, bool>* DisableTablePtr = *DisableTablePtrPtr;
+			FRigidBodyIndexPair BodyPair(filterData0.word0, filterData1.word0); // body indexes are stored in word 0
+			if (DisableTablePtr->Find(BodyPair))
+			{
+				return PxFilterFlag::eKILL;
+			}
+
 		}
 	}
 	
