Clarified in function comments that the references are not stable long term.

diff --git a/Engine/Source/Runtime/Engine/Classes/Components/InputComponent.h b/Engine/Source/Runtime/Engine/Classes/Components/InputComponent.h
--- a/Engine/Source/Runtime/Engine/Classes/Components/InputComponent.h
+++ b/Engine/Source/Runtime/Engine/Classes/Components/InputComponent.h
@@ -579,6 +579,7 @@ public:
 
 	/**
 	 * Binds a delegate function to an Action defined in the project settings.
+	 * Returned reference is only guaranteed to be valid until another action is bound.
 	 */
 	template<class UserClass>
 	FInputActionBinding& BindAction( const FName ActionName, const EInputEvent KeyEvent, UserClass* Object, typename FInputActionHandlerSignature::TUObjectMethodDelegate< UserClass >::FMethodPtr Func )
@@ -590,6 +591,7 @@ public:
 
 	/**
 	 * Binds a delegate function an Axis defined in the project settings.
+	 * Returned reference is only guaranteed to be valid until another axis is bound.
 	 */
 	template<class UserClass>
 	FInputAxisBinding& BindAxis( const FName AxisName, UserClass* Object, typename FInputAxisHandlerSignature::TUObjectMethodDelegate< UserClass >::FMethodPtr Func )
@@ -603,6 +605,7 @@ public:
 	/**
 	 * Indicates that the InputComponent is interested in knowing the Axis value
 	 * (via GetAxisValue) but does not want a delegate function called each frame.
+	 * Returned reference is only guaranteed to be valid until another axis is bound.
 	 */
 	FInputAxisBinding& BindAxis( const FName AxisName )
 	{
@@ -613,6 +616,7 @@ public:
 
 	/**
 	 * Binds a delegate function for an axis key (e.g. Mouse X).
+	 * Returned reference is only guaranteed to be valid until another axis key is bound.
 	 */
 	template<class UserClass>
 	FInputAxisKeyBinding& BindAxisKey( const FKey AxisKey, UserClass* Object, typename FInputAxisHandlerSignature::TUObjectMethodDelegate< UserClass >::FMethodPtr Func )
@@ -626,6 +630,7 @@ public:
 	/**
 	 * Indicates that the InputComponent is interested in knowing/consuming an axis key's
 	 * value (via GetAxisKeyValue) but does not want a delegate function called each frame.
+	 * Returned reference is only guaranteed to be valid until another axis key is bound.
 	 */
 	FInputAxisKeyBinding& BindAxisKey( const FKey AxisKey )
 	{
@@ -634,7 +639,10 @@ public:
 		return AxisKeyBindings.Last();
 	}
 
-	// Binds a delegate function to a vector axis key (e.g. Tilt)
+	/**
+	 * Binds a delegate function to a vector axis key (e.g. Tilt)
+	 * Returned reference is only guaranteed to be valid until another vector axis key is bound.
+	 */
 	template<class UserClass>
 	FInputVectorAxisBinding& BindVectorAxis( const FKey AxisKey, UserClass* Object, typename FInputVectorAxisHandlerSignature::TUObjectMethodDelegate< UserClass >::FMethodPtr Func )
 	{
@@ -647,6 +655,7 @@ public:
 	/**
 	 * Indicates that the InputComponent is interested in knowing/consuming a vector axis key's
 	 * value (via GetVectorAxisKeyValue) but does not want a delegate function called each frame.
+	 * Returned reference is only guaranteed to be valid until another vector axis key is bound.
 	 */
 	FInputVectorAxisBinding& BindVectorAxis( const FKey AxisKey )
 	{
@@ -657,6 +666,7 @@ public:
 
 	/**
 	 * Binds a chord event to a delegate function.
+	 * Returned reference is only guaranteed to be valid until another input key is bound.
 	 */
 	template<class UserClass>
 	FInputKeyBinding& BindKey( const FInputChord Chord, const EInputEvent KeyEvent, UserClass* Object, typename FInputActionHandlerSignature::TUObjectMethodDelegate< UserClass >::FMethodPtr Func )
@@ -669,6 +679,7 @@ public:
 
 	/**
 	 * Binds a key event to a delegate function.
+	 * Returned reference is only guaranteed to be valid until another input key is bound.
 	 */
 	template<class UserClass>
 	FInputKeyBinding& BindKey( const FKey Key, const EInputEvent KeyEvent, UserClass* Object, typename FInputActionHandlerSignature::TUObjectMethodDelegate< UserClass >::FMethodPtr Func )
@@ -678,6 +689,7 @@ public:
 
 	/**
 	 * Binds this input component to touch events.
+	 * Returned reference is only guaranteed to be valid until another touch event is bound.
 	 */
 	template<class UserClass>
 	FInputTouchBinding& BindTouch( const EInputEvent KeyEvent, UserClass* Object, typename FInputTouchHandlerSignature::TUObjectMethodDelegate< UserClass >::FMethodPtr Func )
@@ -690,6 +702,7 @@ public:
 
 	/**
 	 * Binds a gesture event to a delegate function.
+	 * Returned reference is only guaranteed to be valid until another gesture event is bound.
 	 */
 	template<class UserClass>
 	FInputGestureBinding& BindGesture( const FKey GestureKey, UserClass* Object, typename FInputGestureHandlerSignature::TUObjectMethodDelegate< UserClass >::FMethodPtr Func )
