Fixed compile error

diff --git a/Engine/Source/Runtime/Engine/Private/DirectionalLightComponent.cpp b/Engine/Source/Runtime/Engine/Private/DirectionalLightComponent.cpp
--- a/Engine/Source/Runtime/Engine/Private/DirectionalLightComponent.cpp
+++ b/Engine/Source/Runtime/Engine/Private/DirectionalLightComponent.cpp
@@ -187,7 +187,7 @@ public:
 	/** Returns the number of view dependent shadows this light will create. */
 	virtual int32 GetNumViewDependentWholeSceneShadows(const FSceneView& View) const 
 	{ 
-		int32 NumCascades = GetNumShadowMappedCascades();
+		int32 NumCascades = GetNumShadowMappedCascades(View.MaxShadowCascades);
 
 		if (ShouldCreateRayTracedCascade())
 		{
@@ -297,10 +297,10 @@ public:
 
 private:
 
-	int32 GetNumShadowMappedCascades() const
+	int32 GetNumShadowMappedCascades(int32 MaxShadowCascades) const
 	{
 		int32 NumCascades = WholeSceneDynamicShadowRadius > 0.0f ? DynamicShadowCascades : 0;
-		return FMath::Min<int32>(NumCascades, View.MaxShadowCascades);
+		return FMath::Min<int32>(NumCascades, MaxShadowCascades);
 	}
 
 	float GetCSMMaxDistance() const
@@ -327,7 +327,8 @@ private:
 
 	bool ShouldCreateRayTracedCascade() const
 	{
-		const int32 NumCascades = GetNumShadowMappedCascades();
+		//@todo - handle View.MaxShadowCascades properly
+		const int32 NumCascades = GetNumShadowMappedCascades(10);
 		const float RaytracedShadowDistance = GetDistanceFieldShadowDistance();
 		const bool bCreateWithCSM = NumCascades > 0 && RaytracedShadowDistance > GetCSMMaxDistance();
 		const bool bCreateWithoutCSM = NumCascades == 0 && RaytracedShadowDistance > 0;
