call PhysCustom on simulated characters when in MOVE_Custom

diff --git a/Engine/Source/Runtime/Engine/Private/CharacterMovementComponent.cpp b/Engine/Source/Runtime/Engine/Private/CharacterMovementComponent.cpp
--- a/Engine/Source/Runtime/Engine/Private/CharacterMovementComponent.cpp
+++ b/Engine/Source/Runtime/Engine/Private/CharacterMovementComponent.cpp
@@ -3885,6 +3885,12 @@ void UCharacterMovementComponent::AddForce( FVector Force )
 
 void UCharacterMovementComponent::MoveSmooth(const FVector& InVelocity, const float DeltaSeconds, FStepDownResult* OutStepDownResult)
 {
+	if (MovementMode == MOVE_Custom)
+	{
+		PhysCustom(DeltaSeconds, 1);
+		return;
+	}
+
 	FVector Delta = InVelocity * DeltaSeconds;
 	if (!HasValidData() || Delta.IsZero())
 	{
