fixed navigation path's string pulling

diff --git a/Engine/Source/Runtime/Engine/Classes/AI/Navigation/RecastNavMesh.h b/Engine/Source/Runtime/Engine/Classes/AI/Navigation/RecastNavMesh.h
--- a/Engine/Source/Runtime/Engine/Classes/AI/Navigation/RecastNavMesh.h
+++ b/Engine/Source/Runtime/Engine/Classes/AI/Navigation/RecastNavMesh.h
@@ -81,7 +81,10 @@ struct ENGINE_API FNavMeshPath : public FNavigationPath
 
 	FORCEINLINE void SetWantsStringPulling(const bool bNewWantsStringPulling) { bWantsStringPulling = bNewWantsStringPulling; }
 	FORCEINLINE bool WantsStringPulling() const { return bWantsStringPulling; }
-	FORCEINLINE bool IsStringPulled() const { return bStrigPulled; }
+	FORCEINLINE bool IsStringPulled() const { return bStringPulled; }
+	
+	/** find string pulled path from PathCorridor */
+	void PerformStringPulling(const FVector& StartLoc, const FVector& EndLoc);
 
 	FORCEINLINE void SetWantsPathCorridor(const bool bNewWantsPathCorridor) { bWantsPathCorridor = bNewWantsPathCorridor; }
 	FORCEINLINE bool WantsPathCorridor() const { return bWantsPathCorridor; }
@@ -99,6 +102,8 @@ struct ENGINE_API FNavMeshPath : public FNavigationPath
 
 	void ApplyFlags(int32 NavDataFlags);
 
+	void Reset();
+
 	/** get cost of path, starting from next poly in corridor */
 	virtual float GetCostFromNode(NavNodeRef PathNode) const { return GetCostFromIndex(PathCorridor.Find(PathNode) + 1); }
 
@@ -122,7 +127,7 @@ struct ENGINE_API FNavMeshPath : public FNavigationPath
 
 	FORCEINLINE_DEBUGGABLE float GetTotalPathLength() const
 	{
-		return bStrigPulled ? GetStringPulledLength(0) : GetPathCorridorLength(0);
+		return bStringPulled ? GetStringPulledLength(0) : GetPathCorridorLength(0);
 	}
 
 	FORCEINLINE int32 GetNodeRefIndex(const NavNodeRef NodeRef) const { return PathCorridor.Find(NodeRef); }
@@ -186,7 +191,7 @@ protected:
 	/** does this path contain string pulled path? If true then NumPathVerts > 0 and
 	 *	OutPathVerts contains valid data. If false there's only navigation corridor data
 	 *	available.*/
-	uint32 bStrigPulled : 1;	
+	uint32 bStringPulled : 1;	
 
 	/** If set to true path instance will contain a string pulled version. Otherwise only
 	 *	navigation corridor will be available. Defaults to true */
@@ -868,6 +873,9 @@ public:
 	/** @return true is specified segment is fully on navmesh (respecting the optional filter) */
 	bool IsSegmentOnNavmesh(const FVector& SegmentStart, const FVector& SegmentEnd, TSharedPtr<const FNavigationQueryFilter> Filter = NULL, const UObject* Querier = NULL) const;
 
+	/** finds stringpulled path from given corridor */
+	bool FindStraightPath(const FVector& StartLoc, const FVector& EndLoc, const TArray<NavNodeRef>& PathCorridor, TArray<FNavPathPoint>& PathPoints, TArray<uint32>* CustomLinks = NULL) const;
+
 	/** Runs A* pathfinding on navmesh and collect data for every step */
 	int32 DebugPathfinding(const FPathFindingQuery& Query, TArray<FRecastDebugPathfindingStep>& Steps);
 
diff --git a/Engine/Source/Runtime/Engine/Private/AI/Navigation/NavigationPath.cpp b/Engine/Source/Runtime/Engine/Private/AI/Navigation/NavigationPath.cpp
--- a/Engine/Source/Runtime/Engine/Private/AI/Navigation/NavigationPath.cpp
+++ b/Engine/Source/Runtime/Engine/Private/AI/Navigation/NavigationPath.cpp
@@ -297,13 +297,22 @@ const FNavPathType FNavMeshPath::Type;
 FNavMeshPath::FNavMeshPath()
 	: bCorridorEdgesGenerated(false)
 	, bDynamic(false)
-	, bStrigPulled(false)
+	, bStringPulled(false)
 	, bWantsStringPulling(true)
 	, bWantsPathCorridor(false)
 {
 	PathType = FNavMeshPath::Type;
 }
 
+void FNavMeshPath::Reset()
+{
+	PathPoints.Reset();
+	PathCorridor.Reset();
+	PathCorridorCost.Reset();
+	bStringPulled = false;
+	bCorridorEdgesGenerated = false;
+}
+
 float FNavMeshPath::GetStringPulledLength(const int32 StartingPoint) const
 {
 	if (IsValid() == false || StartingPoint >= PathPoints.Num())
@@ -368,6 +377,15 @@ const TArray<FNavigationPortalEdge>* FNavMeshPath::GeneratePathCorridorEdges() c
 	return &PathCorridorEdges;
 }
 
+void FNavMeshPath::PerformStringPulling(const FVector& StartLoc, const FVector& EndLoc)
+{
+#if WITH_RECAST
+	const ARecastNavMesh* MyOwner = Cast<ARecastNavMesh>(GetNavigationDataUsed());
+	bStringPulled = MyOwner->FindStraightPath(StartLoc, EndLoc, PathCorridor, PathPoints, &CustomLinkIds);
+#endif	// WITH_RECAST
+}
+
+
 #if DEBUG_DRAW_OFFSET
 	UWorld* GInternalDebugWorld_ = NULL;
 #endif
diff --git a/Engine/Source/Runtime/Engine/Private/AI/Navigation/PImplRecastNavMesh.cpp b/Engine/Source/Runtime/Engine/Private/AI/Navigation/PImplRecastNavMesh.cpp
--- a/Engine/Source/Runtime/Engine/Private/AI/Navigation/PImplRecastNavMesh.cpp
+++ b/Engine/Source/Runtime/Engine/Private/AI/Navigation/PImplRecastNavMesh.cpp
@@ -746,10 +746,7 @@ ENavigationQueryResult::Type FPImplRecastNavMesh::FindPath(const FVector& StartL
 	}
 
 	// initialize output
-	// @todo this should be moved into a FNavMeshPath function 
-	Path.GetPathPoints().Reset();
-	Path.PathCorridor.Reset();
-	Path.PathCorridorCost.Reset();
+	Path.Reset();
 
 	// get path corridor
 	dtQueryResult PathResult;
@@ -942,6 +939,8 @@ void FPImplRecastNavMesh::PostProcessPath(dtStatus FindPathStatus, FNavMeshPath&
 
 		if (Path.WantsStringPulling())
 		{
+			FVector UseEndLoc = EndLoc;
+			
 			// if path is partial (path corridor doesn't contain EndPolyID), find new RecastEndPos on last poly in corridor
 			if (dtStatusDetail(FindPathStatus, DT_PARTIAL_RESULT))
 			{
@@ -951,60 +950,11 @@ void FPImplRecastNavMesh::PostProcessPath(dtStatus FindPathStatus, FNavMeshPath&
 				const dtStatus NewEndPointStatus = NavQuery.closestPointOnPoly(LastPolyID, &RecastEndPos.X, NewEndPoint);
 				if (dtStatusSucceed(NewEndPointStatus))
 				{
-					FMemory::Memcpy(&RecastEndPos.X, NewEndPoint, sizeof(NewEndPoint));
+					UseEndLoc = Recast2UnrealPoint(NewEndPoint);
 				}
 			}
 
-			dtQueryResult StringPullResult;
-			const dtStatus StringPullStatus = NavQuery.findStraightPath(&RecastStartPos.X, &RecastEndPos.X,
-				Path.PathCorridor.GetTypedData(), Path.PathCorridor.Num(), StringPullResult, DT_STRAIGHTPATH_AREA_CROSSINGS);
-
-			if (dtStatusSucceed(StringPullStatus))
-			{
-				Path.GetPathPoints().AddZeroed(StringPullResult.size());
-
-				// convert to desired format
-				FNavPathPoint* CurVert = Path.GetPathPoints().GetTypedData();
-				UNavigationSystem* NavSys = NavMeshOwner->GetWorld()->GetNavigationSystem();
-
-				for (int32 VertIdx=0; VertIdx < StringPullResult.size(); ++VertIdx)
-				{
-					const float* CurRecastVert = StringPullResult.getPos(VertIdx);
-					CurVert->Location = Recast2UnrVector(CurRecastVert);
-					CurVert->NodeRef = StringPullResult.getRef(VertIdx);
-
-					FNavMeshNodeFlags CurNodeFlags(0);
-					CurNodeFlags.PathFlags = StringPullResult.getFlag(VertIdx);
-
-					uint8 AreaID = RECAST_DEFAULT_AREA;
-					DetourNavMesh->getPolyArea(CurVert->NodeRef, &AreaID);
-					CurNodeFlags.Area = AreaID;
-
-					const UClass* AreaClass = NavMeshOwner->GetAreaClass(AreaID);
-					const UNavArea* DefArea = AreaClass ? ((UClass*)AreaClass)->GetDefaultObject<UNavArea>() : NULL;
-					CurNodeFlags.AreaFlags = DefArea ? DefArea->GetAreaFlags() : 0;
-
-					CurVert->Flags = CurNodeFlags.Pack();
-
-					// include smart link data
-					// if there will be more "edge types" we change this implementation to something more generic
-					if (CurNodeFlags.PathFlags & DT_STRAIGHTPATH_OFFMESH_CONNECTION)
-					{
-						const dtOffMeshConnection* OffMeshCon = DetourNavMesh->getOffMeshConnectionByRef(CurVert->NodeRef);
-						if (OffMeshCon)
-						{
-							CurVert->CustomLinkId = OffMeshCon->userId;
-							Path.CustomLinkIds.Add(OffMeshCon->userId);
-						}
-					}
-
-					CurVert++;
-				}
-
-				// findStraightPath returns 0 for polyId of last point for some reason, even though it knows the poly id.  We will fill that in correctly with the last poly id of the corridor.
-				// @TODO shouldn't it be the same as EndPolyID? (nope, it could be partial path)
-				Path.GetPathPoints().Last().NodeRef = Path.PathCorridor.Last();
-			}
+			Path.PerformStringPulling(StartLoc, UseEndLoc);
 		}
 		else
 		{
@@ -1032,6 +982,69 @@ void FPImplRecastNavMesh::PostProcessPath(dtStatus FindPathStatus, FNavMeshPath&
 	}
 }
 
+bool FPImplRecastNavMesh::FindStraightPath(const FVector& StartLoc, const FVector& EndLoc, const TArray<NavNodeRef>& PathCorridor, TArray<FNavPathPoint>& PathPoints, TArray<uint32>* CustomLinks) const
+{
+	INITIALIZE_NAVQUERY_SIMPLE(NavQuery, RECAST_MAX_SEARCH_NODES);
+
+	const FVector RecastStartPos = Unreal2RecastPoint(StartLoc);
+	const FVector RecastEndPos = Unreal2RecastPoint(EndLoc);
+	bool bResult = false;
+
+	dtQueryResult StringPullResult;
+	const dtStatus StringPullStatus = NavQuery.findStraightPath(&RecastStartPos.X, &RecastEndPos.X,
+		PathCorridor.GetTypedData(), PathCorridor.Num(), StringPullResult, DT_STRAIGHTPATH_AREA_CROSSINGS);
+
+	PathPoints.Reset();
+	if (dtStatusSucceed(StringPullStatus))
+	{
+		PathPoints.AddZeroed(StringPullResult.size());
+
+		// convert to desired format
+		FNavPathPoint* CurVert = PathPoints.GetTypedData();
+
+		for (int32 VertIdx = 0; VertIdx < StringPullResult.size(); ++VertIdx)
+		{
+			const float* CurRecastVert = StringPullResult.getPos(VertIdx);
+			CurVert->Location = Recast2UnrVector(CurRecastVert);
+			CurVert->NodeRef = StringPullResult.getRef(VertIdx);
+
+			FNavMeshNodeFlags CurNodeFlags(0);
+			CurNodeFlags.PathFlags = StringPullResult.getFlag(VertIdx);
+
+			uint8 AreaID = RECAST_DEFAULT_AREA;
+			DetourNavMesh->getPolyArea(CurVert->NodeRef, &AreaID);
+			CurNodeFlags.Area = AreaID;
+
+			const UClass* AreaClass = NavMeshOwner->GetAreaClass(AreaID);
+			const UNavArea* DefArea = AreaClass ? ((UClass*)AreaClass)->GetDefaultObject<UNavArea>() : NULL;
+			CurNodeFlags.AreaFlags = DefArea ? DefArea->GetAreaFlags() : 0;
+
+			CurVert->Flags = CurNodeFlags.Pack();
+
+			// include smart link data
+			// if there will be more "edge types" we change this implementation to something more generic
+			if (CustomLinks && (CurNodeFlags.PathFlags & DT_STRAIGHTPATH_OFFMESH_CONNECTION))
+			{
+				const dtOffMeshConnection* OffMeshCon = DetourNavMesh->getOffMeshConnectionByRef(CurVert->NodeRef);
+				if (OffMeshCon)
+				{
+					CurVert->CustomLinkId = OffMeshCon->userId;
+					CustomLinks->Add(OffMeshCon->userId);
+				}
+			}
+
+			CurVert++;
+		}
+
+		// findStraightPath returns 0 for polyId of last point for some reason, even though it knows the poly id.  We will fill that in correctly with the last poly id of the corridor.
+		// @TODO shouldn't it be the same as EndPolyID? (nope, it could be partial path)
+		PathPoints.Last().NodeRef = PathCorridor.Last();
+		bResult = true;
+	}
+
+	return bResult;
+}
+
 static bool IsDebugNodeModified(const FRecastDebugPathfindingNode& NodeData, const FRecastDebugPathfindingStep& PreviousStep)
 {
 	const FRecastDebugPathfindingNode* PrevNodeData = PreviousStep.Nodes.Find(NodeData);
diff --git a/Engine/Source/Runtime/Engine/Private/AI/Navigation/PImplRecastNavMesh.h b/Engine/Source/Runtime/Engine/Private/AI/Navigation/PImplRecastNavMesh.h
--- a/Engine/Source/Runtime/Engine/Private/AI/Navigation/PImplRecastNavMesh.h
+++ b/Engine/Source/Runtime/Engine/Private/AI/Navigation/PImplRecastNavMesh.h
@@ -139,6 +139,9 @@ public:
 	//@todo document
 	void GetEdgesForPathCorridor(const TArray<NavNodeRef>* PathCorridor, TArray<FNavigationPortalEdge>* PathCorridorEdges) const;
 
+	/** finds stringpulled path from given corridor */
+	bool FindStraightPath(const FVector& StartLoc, const FVector& EndLoc, const TArray<NavNodeRef>& PathCorridor, TArray<FNavPathPoint>& PathPoints, TArray<uint32>* CustomLinks = NULL) const;
+
 	/** Filters nav polys in PolyRefs with Filter */
 	bool FilterPolys(TArray<NavNodeRef>& PolyRefs, const class FRecastQueryFilter* Filter, const UObject* Owner) const;
 
diff --git a/Engine/Source/Runtime/Engine/Private/AI/Navigation/RecastNavMesh.cpp b/Engine/Source/Runtime/Engine/Private/AI/Navigation/RecastNavMesh.cpp
--- a/Engine/Source/Runtime/Engine/Private/AI/Navigation/RecastNavMesh.cpp
+++ b/Engine/Source/Runtime/Engine/Private/AI/Navigation/RecastNavMesh.cpp
@@ -1772,13 +1772,26 @@ bool ARecastNavMesh::IsSegmentOnNavmesh(const FVector& SegmentStart, const FVect
 	
 	FRaycastResult Result;
 	{
-		SECTION_LOCK_TILES_FOR(this);		
+		SECTION_LOCK_TILES;
 		RecastNavMeshImpl->Raycast2D(SegmentStart, SegmentEnd, GetRightFilterRef(Filter), QueryOwner, Result);
 	}
 
 	return Result.HasHit() == false;
 }
 
+bool ARecastNavMesh::FindStraightPath(const FVector& StartLoc, const FVector& EndLoc, const TArray<NavNodeRef>& PathCorridor, TArray<FNavPathPoint>& PathPoints, TArray<uint32>* CustomLinks) const
+{
+	bool bResult = false;
+	if (RecastNavMeshImpl)
+	{
+		SECTION_LOCK_TILES;
+
+		bResult = RecastNavMeshImpl->FindStraightPath(StartLoc, EndLoc, PathCorridor, PathPoints, CustomLinks);
+	}
+
+	return bResult;
+}
+
 int32 ARecastNavMesh::DebugPathfinding(const FPathFindingQuery& Query, TArray<FRecastDebugPathfindingStep>& Steps)
 {
 	int32 NumSteps = 0;
