Move comment for UApplicationLifecycleComponent
Stop it being included globally

diff --git a/Engine/Source/Runtime/Engine/Classes/Components/ApplicationLifecycleComponent.h b/Engine/Source/Runtime/Engine/Classes/Components/ApplicationLifecycleComponent.h
--- a/Engine/Source/Runtime/Engine/Classes/Components/ApplicationLifecycleComponent.h
+++ b/Engine/Source/Runtime/Engine/Classes/Components/ApplicationLifecycleComponent.h
@@ -1,10 +1,11 @@
 // Copyright 1998-2014 Epic Games, Inc. All Rights Reserved.
-// ApplicationLifecycleComponent.: Component to handle receiving notifications from the OS about application state (activated, suspended, termination, etc).  See FCoreDelegates for details
+// ApplicationLifecycleComponent.:  See FCoreDelegates for details
 
 #pragma once
 
 #include "ApplicationLifecycleComponent.generated.h"
 
+/** Component to handle receiving notifications from the OS about application state (activated, suspended, termination, etc). */
 UCLASS(ClassGroup=Utility, HideCategories=(Activation, "Components|Activation", Collision), meta=(BlueprintSpawnableComponent))
 class ENGINE_API UApplicationLifecycleComponent : public UActorComponent
 {
diff --git a/Engine/Source/Runtime/Engine/Private/ApplicationLifecycleComponent.cpp b/Engine/Source/Runtime/Engine/Private/ApplicationLifecycleComponent.cpp
--- a/Engine/Source/Runtime/Engine/Private/ApplicationLifecycleComponent.cpp
+++ b/Engine/Source/Runtime/Engine/Private/ApplicationLifecycleComponent.cpp
@@ -2,6 +2,7 @@
 // ApplicationLifecycleComponent.cpp: Component to handle receiving notifications from the OS about application state (activated, suspended, termination, etc)
 
 #include "EnginePrivate.h"
+#include "Components/ApplicationLifecycleComponent.h"
 #include "CallbackDevice.h"
 
 UApplicationLifecycleComponent::UApplicationLifecycleComponent(const class FPostConstructInitializeProperties& PCIP)
diff --git a/Engine/Source/Runtime/Engine/Private/EnginePrivate.h b/Engine/Source/Runtime/Engine/Private/EnginePrivate.h
--- a/Engine/Source/Runtime/Engine/Private/EnginePrivate.h
+++ b/Engine/Source/Runtime/Engine/Private/EnginePrivate.h
@@ -81,7 +81,6 @@
 #include "Components/SkeletalMeshComponent.h"
 #include "Engine/TargetPoint.h"
 #include "Engine/TextRenderActor.h"
-#include "Components/ApplicationLifecycleComponent.h"
 #include "Engine/DataAsset.h"
 #include "Components/PawnNoiseEmitterComponent.h"
 #include "PhysicsEngine/PhysicsHandleComponent.h"
diff --git a/Engine/Source/Runtime/Engine/Public/Engine.h b/Engine/Source/Runtime/Engine/Public/Engine.h
--- a/Engine/Source/Runtime/Engine/Public/Engine.h
+++ b/Engine/Source/Runtime/Engine/Public/Engine.h
@@ -117,7 +117,6 @@
 #include "Engine/TargetPoint.h"
 #include "Engine/TextRenderActor.h"
 #include "VectorField/VectorFieldVolume.h"
-#include "Components/ApplicationLifecycleComponent.h"
 #include "Engine/DataAsset.h"
 #include "GameFramework/SpectatorPawnMovement.h"
 #include "Vehicles/VehicleWheel.h"
