Forgot some ifdefs for physics cooking OR runtime physics cooking

diff --git a/Engine/Source/Runtime/Engine/Private/PhysicsEngine/PhysLevel.cpp b/Engine/Source/Runtime/Engine/Private/PhysicsEngine/PhysLevel.cpp
--- a/Engine/Source/Runtime/Engine/Private/PhysicsEngine/PhysLevel.cpp
+++ b/Engine/Source/Runtime/Engine/Private/PhysicsEngine/PhysLevel.cpp
@@ -327,7 +327,7 @@ void InitGamePhys()
 	PxRegisterUnifiedHeightFields(*GPhysXSDK);
 
 
-#if WITH_PHYSICS_COOKING
+#if WITH_PHYSICS_COOKING || WITH_RUNTIME_PHYSICS_COOKING
 	// Create Cooking
 	PxCookingParams PCookingParams(PScale);
 	PCookingParams.meshWeldTolerance = 0.1f; // Weld to 1mm precision
@@ -456,7 +456,7 @@ void TermGamePhys()
 	}
 #endif	// #if WITH_APEX
 
-#if WITH_PHYSICS_COOKING
+#if WITH_PHYSICS_COOKING || WITH_RUNTIME_PHYSICS_COOKING
 	if(GPhysXCooking != NULL)
 	{
 		GPhysXCooking->release(); 
diff --git a/Engine/Source/Runtime/Engine/Private/PhysicsEngine/PhysXLibs.cpp b/Engine/Source/Runtime/Engine/Private/PhysicsEngine/PhysXLibs.cpp
--- a/Engine/Source/Runtime/Engine/Private/PhysicsEngine/PhysXLibs.cpp
+++ b/Engine/Source/Runtime/Engine/Private/PhysicsEngine/PhysXLibs.cpp
@@ -15,7 +15,7 @@
 #if PLATFORM_WINDOWS
 	HMODULE PhysX3CommonHandle = 0;
 	HMODULE	PhysX3Handle = 0;
-	#if WITH_PHYSICS_COOKING
+	#if WITH_PHYSICS_COOKING || WITH_RUNTIME_PHYSICS_COOKING
 		HMODULE	PhysX3CookingHandle = 0;
 	#endif
 	HMODULE	nvToolsExtHandle = 0;
@@ -60,7 +60,7 @@ void LoadPhysXModules()
 			PhysX3CommonHandle = LoadLibraryW(*(RootPhysXPath + "PhysX3CommonDEBUG_x64.dll"));
 			nvToolsExtHandle = LoadLibraryW(*(RootPhysXPath + "nvToolsExt64_1.dll"));
 			PhysX3Handle = LoadLibraryW(*(RootPhysXPath + "PhysX3DEBUG_x64.dll"));
-			#if WITH_PHYSICS_COOKING
+			#if WITH_PHYSICS_COOKING || WITH_RUNTIME_PHYSICS_COOKING
 				PhysX3CookingHandle = LoadLibraryW(*(RootPhysXPath + "PhysX3CookingDEBUG_x64.dll"));
 			#endif
 			#if WITH_APEX
@@ -79,7 +79,7 @@ void LoadPhysXModules()
 			PhysX3CommonHandle = LoadLibraryW(*(RootPhysXPath + "PhysX3Common_x64.dll"));
 			nvToolsExtHandle = LoadLibraryW(*(RootPhysXPath + "nvToolsExt64_1.dll"));
 			PhysX3Handle = LoadLibraryW(*(RootPhysXPath + "PhysX3_x64.dll"));
-			#if WITH_PHYSICS_COOKING
+			#if WITH_PHYSICS_COOKING || WITH_RUNTIME_PHYSICS_COOKING
 				PhysX3CookingHandle = LoadLibraryW(*(RootPhysXPath + "PhysX3Cooking_x64.dll"));
 			#endif
 			#if WITH_APEX
@@ -97,7 +97,7 @@ void LoadPhysXModules()
 			PhysX3CommonHandle = LoadLibraryW(*(RootPhysXPath + "PhysX3CommonPROFILE_x64.dll"));
 			nvToolsExtHandle = LoadLibraryW(*(RootPhysXPath + "nvToolsExt64_1.dll"));
 			PhysX3Handle = LoadLibraryW(*(RootPhysXPath + "PhysX3PROFILE_x64.dll"));
-			#if WITH_PHYSICS_COOKING
+			#if WITH_PHYSICS_COOKING || WITH_RUNTIME_PHYSICS_COOKING
 				PhysX3CookingHandle = LoadLibraryW(*(RootPhysXPath + "PhysX3CookingPROFILE_x64.dll"));
 			#endif
 			#if WITH_APEX
@@ -126,7 +126,7 @@ void LoadPhysXModules()
 			PhysX3CommonHandle = LoadLibraryW(*(RootPhysXPath + "PhysX3CommonDEBUG_x86.dll"));
 			nvToolsExtHandle = LoadLibraryW(*(RootPhysXPath + "nvToolsExt32_1.dll"));
 			PhysX3Handle = LoadLibraryW(*(RootPhysXPath + "PhysX3DEBUG_x86.dll"));
-			#if WITH_PHYSICS_COOKING
+			#if WITH_PHYSICS_COOKING || WITH_RUNTIME_PHYSICS_COOKING
 				PhysX3CookingHandle = LoadLibraryW(*(RootPhysXPath + "PhysX3CookingDEBUG_x86.dll"));
 			#endif
 			#if WITH_APEX
@@ -144,7 +144,7 @@ void LoadPhysXModules()
 			PhysX3CommonHandle = LoadLibraryW(*(RootPhysXPath + "PhysX3Common_x86.dll"));
 			nvToolsExtHandle = LoadLibraryW(*(RootPhysXPath + "nvToolsExt32_1.dll"));
 			PhysX3Handle = LoadLibraryW(*(RootPhysXPath + "PhysX3_x86.dll"));
-			#if WITH_PHYSICS_COOKING
+			#if WITH_PHYSICS_COOKING || WITH_RUNTIME_PHYSICS_COOKING
 				PhysX3CookingHandle = LoadLibraryW(*(RootPhysXPath + "PhysX3Cooking_x86.dll"));
 			#endif
 			#if WITH_APEX
@@ -162,7 +162,7 @@ void LoadPhysXModules()
 			PhysX3CommonHandle = LoadLibraryW(*(RootPhysXPath + "PhysX3CommonPROFILE_x86.dll"));
 			nvToolsExtHandle = LoadLibraryW(*(RootPhysXPath + "nvToolsExt32_1.dll"));
 			PhysX3Handle = LoadLibraryW(*(RootPhysXPath + "PhysX3PROFILE_x86.dll"));
-			#if WITH_PHYSICS_COOKING
+			#if WITH_PHYSICS_COOKING || WITH_RUNTIME_PHYSICS_COOKING
 				PhysX3CookingHandle = LoadLibraryW(*(RootPhysXPath + "PhysX3CookingPROFILE_x86.dll"));
 			#endif
 			#if WITH_APEX
@@ -187,7 +187,7 @@ void UnloadPhysXModules()
 {
 #if PLATFORM_WINDOWS
 	FreeLibrary(PhysX3Handle);
-	#if WITH_PHYSICS_COOKING
+	#if WITH_PHYSICS_COOKING || WITH_RUNTIME_PHYSICS_COOKING
 		FreeLibrary(PhysX3CookingHandle);
 	#endif
 	FreeLibrary(PhysX3CommonHandle);
diff --git a/Engine/Source/Runtime/Engine/Private/PhysicsEngine/PhysXSupport.cpp b/Engine/Source/Runtime/Engine/Private/PhysicsEngine/PhysXSupport.cpp
--- a/Engine/Source/Runtime/Engine/Private/PhysicsEngine/PhysXSupport.cpp
+++ b/Engine/Source/Runtime/Engine/Private/PhysicsEngine/PhysXSupport.cpp
@@ -14,7 +14,7 @@
 PxFoundation*			GPhysXFoundation = NULL;
 PxProfileZoneManager*	GPhysXProfileZoneManager = NULL;
 PxPhysics*				GPhysXSDK = NULL;
-#if WITH_PHYSICS_COOKING
+#if WITH_PHYSICS_COOKING || WITH_RUNTIME_PHYSICS_COOKING
 PxCooking*				GPhysXCooking = NULL;
 #endif
 FPhysXAllocator*		GPhysXAllocator = NULL;
