- Animation Retargeting offline solution using rig

- Rig asset/Customization
- Retargeting Manager
- Support retargeting using the data

diff --git a/Engine/Source/Developer/AssetTools/Private/AssetTools.cpp b/Engine/Source/Developer/AssetTools/Private/AssetTools.cpp
--- a/Engine/Source/Developer/AssetTools/Private/AssetTools.cpp
+++ b/Engine/Source/Developer/AssetTools/Private/AssetTools.cpp
@@ -62,6 +62,7 @@ FAssetTools::FAssetTools()
 	RegisterAssetTypeActions( MakeShareable(new FAssetTypeActions_PhysicsAsset) );
 	RegisterAssetTypeActions( MakeShareable(new FAssetTypeActions_Redirector) );
 	RegisterAssetTypeActions( MakeShareable(new FAssetTypeActions_ReverbEffect) );
+	RegisterAssetTypeActions( MakeShareable(new FAssetTypeActions_Rig) );
 	RegisterAssetTypeActions( MakeShareable(new FAssetTypeActions_SkeletalMesh) );
 	RegisterAssetTypeActions( MakeShareable(new FAssetTypeActions_Skeleton) );
 	RegisterAssetTypeActions( MakeShareable(new FAssetTypeActions_SlateBrush) );
diff --git a/Engine/Source/Developer/AssetTools/Private/AssetToolsPrivatePCH.h b/Engine/Source/Developer/AssetTools/Private/AssetToolsPrivatePCH.h
--- a/Engine/Source/Developer/AssetTools/Private/AssetToolsPrivatePCH.h
+++ b/Engine/Source/Developer/AssetTools/Private/AssetToolsPrivatePCH.h
@@ -62,6 +62,7 @@ DECLARE_LOG_CATEGORY_EXTERN(LogAssetTools, Log, All);
 #include "AssetTypeActions/AssetTypeActions_PhysicsAsset.h"
 #include "AssetTypeActions/AssetTypeActions_Redirector.h"
 #include "AssetTypeActions/AssetTypeActions_ReverbEffect.h"
+#include "AssetTypeActions/AssetTypeActions_Rig.h"
 #include "AssetTypeActions/AssetTypeActions_Skeleton.h"
 #include "AssetTypeActions/AssetTypeActions_SlateBrush.h"
 #include "AssetTypeActions/AssetTypeActions_SlateWidgetStyle.h"
diff --git a/Engine/Source/Developer/AssetTools/Private/AssetTypeActions/AssetTypeActions_Rig.cpp b/Engine/Source/Developer/AssetTools/Private/AssetTypeActions/AssetTypeActions_Rig.cpp
new file mode 100644
--- /dev/null
+++ b/Engine/Source/Developer/AssetTools/Private/AssetTypeActions/AssetTypeActions_Rig.cpp
@@ -0,0 +1,5 @@
+// Copyright 1998-2014 Epic Games, Inc. All Rights Reserved.
+
+#include "AssetToolsPrivatePCH.h"
+
+///////////////////////////////
diff --git a/Engine/Source/Developer/AssetTools/Private/AssetTypeActions/AssetTypeActions_Rig.h b/Engine/Source/Developer/AssetTools/Private/AssetTypeActions/AssetTypeActions_Rig.h
new file mode 100644
--- /dev/null
+++ b/Engine/Source/Developer/AssetTools/Private/AssetTypeActions/AssetTypeActions_Rig.h
@@ -0,0 +1,15 @@
+// Copyright 1998-2014 Epic Games, Inc. All Rights Reserved.
+
+#pragma once
+
+#include "Animation/Rig.h"
+
+class FAssetTypeActions_Rig : public FAssetTypeActions_Base
+{
+public:
+	// IAssetTypeActions Implementation
+	virtual FText GetName() const override { return NSLOCTEXT("AssetTypeActions", "AssetTypeActions_Rig", "Rig"); }
+	virtual FColor GetTypeColor() const override { return FColor(255,201,14); }
+	virtual UClass* GetSupportedClass() const override { return URig::StaticClass(); }
+	virtual uint32 GetCategories() override { return EAssetTypeCategories::Animation; }
+};
\ No newline at end of file
diff --git a/Engine/Source/Developer/AssetTools/Private/AssetTypeActions/AssetTypeActions_Skeleton.cpp b/Engine/Source/Developer/AssetTools/Private/AssetTypeActions/AssetTypeActions_Skeleton.cpp
--- a/Engine/Source/Developer/AssetTools/Private/AssetTypeActions/AssetTypeActions_Skeleton.cpp
+++ b/Engine/Source/Developer/AssetTools/Private/AssetTypeActions/AssetTypeActions_Skeleton.cpp
@@ -13,9 +13,390 @@
 
 #include "Kismet2/KismetEditorUtilities.h"
 #include "Kismet2/BlueprintEditorUtils.h"
+#include "Animation/Rig.h"
 
 #define LOCTEXT_NAMESPACE "AssetTypeActions"
 
+/** 
+* FBoneCheckbox
+* 
+* Context data for the SCreateRigDlg panel check boxes
+*/
+struct FBoneCheckbox
+{
+	FName BoneName;
+	int32 BoneID;
+	bool  bUsed;
+};
+
+/** 
+* FCreateRigDlg
+* 
+* Wrapper class for SCreateRigDlg. This class creates and launches a dialog then awaits the
+* result to return to the user. 
+*/
+class FCreateRigDlg
+{
+public:
+	enum EResult
+	{
+		Cancel = 0,			// No/Cancel, normal usage would stop the current action
+		Confirm = 1,		// Yes/Ok/Etc, normal usage would continue with action
+	};
+
+	FCreateRigDlg( USkeleton* InSkeleton );
+
+	/**  Shows the dialog box and waits for the user to respond. */
+	EResult ShowModal();
+
+	// Map of RequiredBones of (BoneIndex, ParentIndex)
+	TMap<int32, int32> RequiredBones;
+private:
+
+	/** Cached pointer to the modal window */
+	TSharedPtr<SWindow> DialogWindow;
+
+	/** Cached pointer to the merge skeleton widget */
+	TSharedPtr<class SCreateRigDlg> DialogWidget;
+
+	/** The Skeleton to merge bones to */
+	USkeleton* Skeleton;
+};
+
+/** 
+ * Slate panel for choosing which bones to merge into the skeleton
+ */
+class SCreateRigDlg : public SCompoundWidget
+{
+public:
+
+	SLATE_BEGIN_ARGS(SCreateRigDlg)
+		{}
+		/** Window in which this widget resides */
+		SLATE_ATTRIBUTE(TSharedPtr<SWindow>, ParentWindow)
+	SLATE_END_ARGS()	
+
+	/**
+	 * Constructs this widget
+	 *
+	 * @param	InArgs	The declaration data for this widget
+	 */
+	BEGIN_SLATE_FUNCTION_BUILD_OPTIMIZATION
+	void Construct( const FArguments& InArgs )
+	{
+		UserResponse = FCreateRigDlg::Cancel;
+		ParentWindow = InArgs._ParentWindow.Get();
+
+		this->ChildSlot[
+			SNew(SVerticalBox)
+			+SVerticalBox::Slot()
+			.AutoHeight()
+			.Padding(8.0f, 4.0f, 8.0f, 4.0f)
+			[
+				SNew( STextBlock )
+				.Text( LOCTEXT("MergeSkeletonDlgDescription", "Would you like to add following bones to the skeleton?"))
+			]
+			+SVerticalBox::Slot()
+			.AutoHeight()
+			.Padding(8.0f, 4.0f, 8.0f, 4.0f)
+			[
+				SNew(SSeparator)
+			]
+			+SVerticalBox::Slot()
+			.Padding(8.0f, 4.0f, 8.0f, 4.0f)
+			[
+				SNew(SBorder)
+				[
+					SNew(SScrollBox)
+					+SScrollBox::Slot()
+					[
+						//Save this widget so we can populate it later with check boxes
+						SAssignNew(CheckBoxContainer, SVerticalBox)
+					]
+				]
+			]
+			+SVerticalBox::Slot()
+			.AutoHeight()
+			.HAlign(HAlign_Left)
+			.Padding(8.0f, 4.0f, 8.0f, 4.0f)
+			[
+				SNew(SUniformGridPanel)
+				.SlotPadding(FEditorStyle::GetMargin("StandardDialog.SlotPadding"))
+				.MinDesiredSlotWidth(FEditorStyle::GetFloat("StandardDialog.MinDesiredSlotWidth"))
+				.MinDesiredSlotHeight(FEditorStyle::GetFloat("StandardDialog.MinDesiredSlotHeight"))
+				+SUniformGridPanel::Slot(0,0)
+				[
+					SNew(SButton) 
+					.HAlign(HAlign_Center)
+					.ContentPadding(FEditorStyle::GetMargin("StandardDialog.ContentPadding"))
+					.OnClicked(this, &SCreateRigDlg::ChangeAllOptions, true)
+					.Text(LOCTEXT("SkeletonMergeSelectAll", "Select All"))
+				]
+				+SUniformGridPanel::Slot(1,0)
+				[
+					SNew(SButton) 
+					.HAlign(HAlign_Center)
+					.ContentPadding(FEditorStyle::GetMargin("StandardDialog.ContentPadding"))
+					.OnClicked(this, &SCreateRigDlg::ChangeAllOptions, false)
+					.Text(LOCTEXT("SkeletonMergeDeselectAll", "Deselect All"))
+				]
+			]
+			+SVerticalBox::Slot()
+			.AutoHeight()
+			.Padding(8.0f, 4.0f, 8.0f, 4.0f)
+			[
+				SNew(SSeparator)
+			]
+			+SVerticalBox::Slot()
+			.AutoHeight()
+			.HAlign(HAlign_Right)
+			.Padding(8.0f, 4.0f, 8.0f, 4.0f)
+			[
+				SNew(SUniformGridPanel)
+				.SlotPadding(FEditorStyle::GetMargin("StandardDialog.SlotPadding"))
+				.MinDesiredSlotWidth(FEditorStyle::GetFloat("StandardDialog.MinDesiredSlotWidth"))
+				.MinDesiredSlotHeight(FEditorStyle::GetFloat("StandardDialog.MinDesiredSlotHeight"))
+				+SUniformGridPanel::Slot(0,0)
+				[
+					SNew(SButton) 
+					.HAlign(HAlign_Center)
+					.ContentPadding(FEditorStyle::GetMargin("StandardDialog.ContentPadding"))
+					.OnClicked( this, &SCreateRigDlg::OnButtonClick, FCreateRigDlg::Confirm )
+					.Text(LOCTEXT("SkeletonMergeOk", "OK"))
+				]
+				+SUniformGridPanel::Slot(1,0)
+				[
+					SNew(SButton) 
+					.HAlign(HAlign_Center)
+					.ContentPadding(FEditorStyle::GetMargin("StandardDialog.ContentPadding"))
+					.OnClicked( this, &SCreateRigDlg::OnButtonClick, FCreateRigDlg::Cancel )
+					.Text(LOCTEXT("SkeletonMergeCancel", "Cancel"))
+				]
+			]
+		];
+	}
+	END_SLATE_FUNCTION_BUILD_OPTIMIZATION
+
+	/**
+	 * Creates a Slate check box 
+	 *
+	 * @param	Label		Text label for the check box
+	 * @param	ButtonId	The ID for the check box
+	 * @return				The created check box widget
+	 */
+	TSharedRef<SWidget> CreateCheckBox( const FString& Label, int32 ButtonId )
+	{
+		return
+			SNew(SCheckBox)
+			.IsChecked( this, &SCreateRigDlg::IsCheckboxChecked, ButtonId )
+			.OnCheckStateChanged( this, &SCreateRigDlg::OnCheckboxChanged, ButtonId )
+			[
+				SNew(STextBlock).Text(Label)
+			];
+	}
+
+	/**
+	 * Returns the state of the check box
+	 *
+	 * @param	ButtonId	The ID for the check box
+	 * @return				The status of the check box
+	 */
+	ESlateCheckBoxState::Type IsCheckboxChecked( int32 ButtonId ) const
+	{
+		return CheckBoxInfoMap.FindChecked(ButtonId).bUsed ? ESlateCheckBoxState::Checked : ESlateCheckBoxState::Unchecked;
+	}
+
+	/**
+	 * Handler for all check box clicks
+	 *
+	 * @param	NewCheckboxState	The new state of the check box
+	 * @param	CheckboxThatChanged	The ID of the radio button that has changed. 
+	 */
+	void OnCheckboxChanged( ESlateCheckBoxState::Type NewCheckboxState, int32 CheckboxThatChanged )
+	{
+		FBoneCheckbox& Info = CheckBoxInfoMap.FindChecked(CheckboxThatChanged);
+		Info.bUsed = !Info.bUsed;
+	}
+
+	/**
+	 * Handler for the Select All and Deselect All buttons
+	 *
+	 * @param	bNewCheckedState	The new state of the check boxes
+	 */
+	FReply ChangeAllOptions(bool bNewCheckedState)
+	{
+		for(auto Iter = CheckBoxInfoMap.CreateIterator(); Iter; ++Iter)
+		{
+			FBoneCheckbox& Info = Iter.Value();
+			Info.bUsed = bNewCheckedState;
+		}
+		return FReply::Handled();
+	}
+
+	/**
+	 * Populated the dialog with multiple check boxes, each corresponding to a bone
+	 *
+	 * @param	BoneInfos	The list of Bones to populate the dialog with
+	 */
+	void PopulateOptions(TArray<FBoneCheckbox>& BoneInfos)
+	{
+		for(auto Iter = BoneInfos.CreateIterator(); Iter; ++Iter)
+		{
+			FBoneCheckbox& Info = (*Iter);
+			Info.bUsed = true;
+
+			CheckBoxInfoMap.Add(Info.BoneID, Info);
+
+			CheckBoxContainer->AddSlot()
+			.AutoHeight()
+			[
+				CreateCheckBox(Info.BoneName.GetPlainNameString(), Info.BoneID)
+			];
+		}
+	}
+
+	/** 
+	 * Returns the EResult of the button which the user pressed. Closing of the dialog
+	 * in any other way than clicking "Ok" results in this returning a "Cancel" value
+	 */
+	FCreateRigDlg::EResult GetUserResponse() const 
+	{
+		return UserResponse; 
+	}
+
+	/** 
+	 * Returns whether the user selected that bone to be used (checked its respective check box)
+	 */
+	bool IsBoneIncluded(int32 BoneID)
+	{
+		auto* Item = CheckBoxInfoMap.Find(BoneID);
+		return Item ? Item->bUsed : false;
+	}
+
+private:
+	
+	/** 
+	 * Handles when a button is pressed, should be bound with appropriate EResult Key
+	 * 
+	 * @param ButtonID - The return type of the button which has been pressed.
+	 */
+	FReply OnButtonClick(FCreateRigDlg::EResult ButtonID)
+	{
+		ParentWindow->RequestDestroyWindow();
+		UserResponse = ButtonID;
+
+		return FReply::Handled();
+	}
+
+	/** Stores the users response to this dialog */
+	FCreateRigDlg::EResult	 UserResponse;
+
+	/** The slate container that the bone check boxes get added to */
+	TSharedPtr<SVerticalBox>	 CheckBoxContainer;
+	/** Store the check box state for each bone */
+	TMap<int32,FBoneCheckbox> CheckBoxInfoMap;
+
+	/** Pointer to the window which holds this Widget, required for modal control */
+	TSharedPtr<SWindow>			 ParentWindow;
+};
+
+FCreateRigDlg::FCreateRigDlg( USkeleton* InSkeleton )
+{
+	Skeleton = InSkeleton;
+
+	if (FSlateApplication::IsInitialized())
+	{
+		DialogWindow = SNew(SWindow)
+			.Title( LOCTEXT("MergeSkeletonDlgTitle", "Merge Bones") )
+			.SupportsMinimize(false) .SupportsMaximize(false)
+			.ClientSize(FVector2D(350, 500));
+
+		TSharedPtr<SBorder> DialogWrapper = 
+			SNew(SBorder)
+			.BorderImage(FEditorStyle::GetBrush("ToolPanel.GroupBorder"))
+			.Padding(4.0f)
+			[
+				SAssignNew(DialogWidget, SCreateRigDlg)
+				.ParentWindow(DialogWindow)
+			];
+
+		DialogWindow->SetContent(DialogWrapper.ToSharedRef());
+	}
+}
+
+FCreateRigDlg::EResult FCreateRigDlg::ShowModal()
+{
+	RequiredBones.Empty();
+
+	TArray<FBoneCheckbox> BoneInfos;
+
+	// Make a list of all skeleton bone list
+	const FReferenceSkeleton & RefSkeleton = Skeleton->GetReferenceSkeleton();
+	const TArray<FBoneNode> & BoneTree = Skeleton->GetBoneTree();
+	for ( int32 BoneTreeId=0; BoneTreeId<RefSkeleton.GetNum(); ++BoneTreeId )
+	{
+		const FName & BoneName = RefSkeleton.GetBoneName(BoneTreeId);
+
+		FBoneCheckbox Info;
+		Info.BoneID = BoneTreeId;
+		Info.BoneName = BoneName;
+		BoneInfos.Add(Info);
+	}
+
+	if (BoneInfos.Num() == 0)
+	{
+		// something wrong
+		return EResult::Cancel;
+	}
+
+	DialogWidget->PopulateOptions(BoneInfos);
+
+	//Show Dialog
+	GEditor->EditorAddModalWindow(DialogWindow.ToSharedRef());
+	EResult UserResponse = (EResult)DialogWidget->GetUserResponse();
+
+	if(UserResponse == EResult::Confirm)
+	{
+		for ( int32 RefBoneId= 0 ; RefBoneId< RefSkeleton.GetNum() ; ++RefBoneId )
+		{
+			if ( DialogWidget->IsBoneIncluded(RefBoneId) )
+			{
+				// I need to find parent that exists in the RefSkeleton
+				int32 ParentIndex = RefSkeleton.GetParentIndex(RefBoneId);
+				bool bFoundParent = false;
+
+				// make sure RequiredBones already have ParentIndex
+				while (ParentIndex >= 0 )
+				{
+					// if I don't have it yet
+					if ( RequiredBones.Contains(ParentIndex) )
+					{
+						bFoundParent = true;
+						// find the Parent that is related
+						break;
+					}
+					else
+					{
+						ParentIndex = RefSkeleton.GetParentIndex(ParentIndex);					
+					}
+				}
+
+				if (bFoundParent)
+				{
+					RequiredBones.Add(RefBoneId, ParentIndex);
+				}
+				else
+				{
+					RequiredBones.Add(RefBoneId, INDEX_NONE);
+				}
+			}
+		}
+	}
+
+	return (RequiredBones.Num() > 0)? EResult::Confirm : EResult::Cancel;
+}
+
+///////////////////////////////
 void FAssetTypeActions_Skeleton::GetActions( const TArray<UObject*>& InObjects, FMenuBuilder& MenuBuilder )
 {
 	auto Skeletons = GetTypedWeakObjectPtrs<USkeleton>(InObjects);
@@ -46,6 +427,18 @@ void FAssetTypeActions_Skeleton::GetActions( const TArray<UObject*>& InObjects,
 			)
 		);
 
+	if (InObjects.Num() == 1)
+	{
+		MenuBuilder.AddMenuEntry(
+			LOCTEXT("Skeleton_CreateRig", "Create Rig"),
+			LOCTEXT("Skeleton_CreateRigTooltip", "Create Rig from this skeleton."),
+			FSlateIcon(),
+			FUIAction(
+				FExecuteAction::CreateSP(this, &FAssetTypeActions_Skeleton::ExecuteCreateRig, Skeletons),
+				FCanExecuteAction()
+				)
+			);
+	}
 	// @todo ImportAnimation
 	/*
 	MenuBuilder.AddMenuEntry(
@@ -160,6 +553,47 @@ void FAssetTypeActions_Skeleton::ExecuteNewAnimBlueprint(TArray<TWeakObjectPtr<U
 	}
 }
 
+/** Handler for when Create Rig is selected */
+void FAssetTypeActions_Skeleton::ExecuteCreateRig(TArray<TWeakObjectPtr<USkeleton>> Skeletons)
+{
+	if (Skeletons.Num() == 1)
+	{
+		CreateRig(Skeletons[0]);
+	}
+}
+
+void FAssetTypeActions_Skeleton::CreateRig(const TWeakObjectPtr<USkeleton> Skeleton)
+{
+	if(Skeleton.IsValid())
+	{
+		FCreateRigDlg CreateRigDlg(Skeleton.Get());
+		if(CreateRigDlg.ShowModal() == FDlgMergeSkeleton::Confirm)
+		{
+			check(CreateRigDlg.RequiredBones.Num() > 0);
+
+			// Determine an appropriate name
+			FString Name;
+			FString PackageName;
+			CreateUniqueAssetName(Skeleton->GetOutermost()->GetName(), TEXT("Rig"), PackageName, Name);
+
+			// Create the asset, and assign its skeleton
+			FAssetToolsModule& AssetToolsModule = FModuleManager::GetModuleChecked<FAssetToolsModule>("AssetTools");
+			URig* NewAsset = Cast<URig>(AssetToolsModule.Get().CreateAsset(Name, FPackageName::GetLongPackagePath(PackageName), URig::StaticClass(), NULL));
+
+			if(NewAsset)
+			{
+				NewAsset->CreateFromSkeleton(Skeleton.Get(), CreateRigDlg.RequiredBones);
+				NewAsset->MarkPackageDirty();
+
+
+				TArray<UObject*> ObjectsToSync;
+				ObjectsToSync.Add(NewAsset);
+				FAssetTools::Get().SyncBrowserToAssets(ObjectsToSync);
+			}
+		}
+	}
+}
+
 void FAssetTypeActions_Skeleton::CreateAnimationAssets(const TArray<TWeakObjectPtr<USkeleton>>& Skeletons, TSubclassOf<UAnimationAsset> AssetClass, const FString& InPrefix)
 {
 	FAnimAssetCreated OnAssetCreated;
diff --git a/Engine/Source/Developer/AssetTools/Private/AssetTypeActions/AssetTypeActions_Skeleton.h b/Engine/Source/Developer/AssetTools/Private/AssetTypeActions/AssetTypeActions_Skeleton.h
--- a/Engine/Source/Developer/AssetTools/Private/AssetTypeActions/AssetTypeActions_Skeleton.h
+++ b/Engine/Source/Developer/AssetTools/Private/AssetTypeActions/AssetTypeActions_Skeleton.h
@@ -47,10 +47,16 @@ private:
 	/** Handler for when Skeleton Retarget is selected */
 	void ExecuteRetargetSkeleton(TArray<TWeakObjectPtr<USkeleton>> Skeletons);
 
+	/** Handler for when Create Rig is selected */
+	void ExecuteCreateRig(TArray<TWeakObjectPtr<USkeleton>> Skeletons);
+
 private: // Helper functions
 	/** Creates animation assets using the BaseName+Suffix */
 	void CreateAnimationAssets(const TArray<TWeakObjectPtr<USkeleton>>& Skeletons, TSubclassOf<UAnimationAsset> AssetClass, const FString& InPrefix);
 
+	/** Creates animation assets using the BaseName+Suffix */
+	void CreateRig(const TWeakObjectPtr<USkeleton> Skeleton);
+
  	/** 
 	 * Main function for handling retargeting old Skeleton to new Skeleton
 	 * 
diff --git a/Engine/Source/Editor/AnimGraph/Classes/AnimPreviewInstance.h b/Engine/Source/Editor/AnimGraph/Classes/AnimPreviewInstance.h
--- a/Engine/Source/Editor/AnimGraph/Classes/AnimPreviewInstance.h
+++ b/Engine/Source/Editor/AnimGraph/Classes/AnimPreviewInstance.h
@@ -46,6 +46,7 @@ class ANIMGRAPH_API UAnimPreviewInstance : public UAnimSingleNodeInstance
 
 	// Begin UAnimInstance interface
 	virtual void NativeInitializeAnimation() override;
+	virtual void NativeUpdateAnimation(float DeltaTimeX) override;
 	virtual bool NativeEvaluateAnimation(FPoseContext& Output) override;
 	// End UAnimInstance interface
 
@@ -100,6 +101,15 @@ class ANIMGRAPH_API UAnimPreviewInstance : public UAnimSingleNodeInstance
 	 * @param	InBoneName	The name of the existing modification to remove
 	 */
 	void RemoveBoneModification(const FName& InBoneName);
+
+	/**
+	 * Reset all bone modified
+	 */
+	void ResetModifiedBone();
+
+#if WITH_EDITORONLY_DATA
+	bool bForceRetargetBasePose;
+#endif
 };
 
 
diff --git a/Engine/Source/Editor/AnimGraph/Private/AnimPreviewInstance.cpp b/Engine/Source/Editor/AnimGraph/Private/AnimPreviewInstance.cpp
--- a/Engine/Source/Editor/AnimGraph/Private/AnimPreviewInstance.cpp
+++ b/Engine/Source/Editor/AnimGraph/Private/AnimPreviewInstance.cpp
@@ -6,8 +6,12 @@
 
 UAnimPreviewInstance::UAnimPreviewInstance(const class FPostConstructInitializeProperties& PCIP)
 	: Super(PCIP)
+	, SkeletalControlAlpha(1.0f)
+#if WITH_EDITORONLY_DATA
+	, bForceRetargetBasePose(false)
+#endif
 {
-	SkeletalControlAlpha = 1.0f;
+	
 }
 
 void UAnimPreviewInstance::NativeInitializeAnimation()
@@ -20,6 +24,11 @@ void UAnimPreviewInstance::NativeInitializeAnimation()
 	SetPlaying(bCachedIsPlaying);
 }
 
+void UAnimPreviewInstance::ResetModifiedBone()
+{
+	BoneControllers.Empty();
+}
+
 FAnimNode_ModifyBone* UAnimPreviewInstance::FindModifiedBone(const FName& InBoneName)
 {
 	return BoneControllers.FindByPredicate(
@@ -61,9 +70,40 @@ void UAnimPreviewInstance::RemoveBoneModification(const FName& InBoneName)
 	);
 }
 
+void UAnimPreviewInstance::NativeUpdateAnimation(float DeltaTimeX)
+{
+#if WITH_EDITORONLY_DATA
+	if(bForceRetargetBasePose)
+	{
+		// nothing to be done here
+		return;
+	}
+#endif // #if WITH_EDITORONLY_DATA
+
+	Super::NativeUpdateAnimation(DeltaTimeX);
+}
+
 bool UAnimPreviewInstance::NativeEvaluateAnimation(FPoseContext& Output)
 {
-	Super::NativeEvaluateAnimation(Output);
+#if WITH_EDITORONLY_DATA
+	if(bForceRetargetBasePose)
+	{
+		USkeletalMeshComponent * MeshComponent = GetSkelMeshComponent();
+		if(MeshComponent && MeshComponent->SkeletalMesh)
+		{
+			FAnimationRuntime::FillWithRetargetBaseRefPose(Output.Pose.Bones, GetSkelMeshComponent()->SkeletalMesh, RequiredBones);
+		}
+		else
+		{
+			// ideally we'll return just ref pose, but not sure if this will work with LODs
+			FAnimationRuntime::FillWithRefPose(Output.Pose.Bones, RequiredBones);
+		}
+	}
+	else
+#endif // #if WITH_EDITORONLY_DATA
+	{
+		Super::NativeEvaluateAnimation(Output);
+	}
 
 	USkeletalMeshComponent* Component = GetSkelMeshComponent();
 
diff --git a/Engine/Source/Editor/DetailCustomizations/Private/DetailCustomizations.cpp b/Engine/Source/Editor/DetailCustomizations/Private/DetailCustomizations.cpp
--- a/Engine/Source/Editor/DetailCustomizations/Private/DetailCustomizations.cpp
+++ b/Engine/Source/Editor/DetailCustomizations/Private/DetailCustomizations.cpp
@@ -84,6 +84,7 @@
 #include "AudioSettingsDetails.h"
 #include "DateTimeStructCustomization.h"
 #include "TimespanStructCustomization.h"
+#include "RigDetails.h"
 
 IMPLEMENT_MODULE( FDetailCustomizationsModule, DetailCustomizations );
 
@@ -210,6 +211,7 @@ void FDetailCustomizationsModule::RegisterObjectCustomizations()
 	RegisterCustomClassLayout("K2Node_StructMemberSet", FOnGetDetailCustomizationInstance::CreateStatic(&FSkeletalControlNodeDetails::MakeInstance));
 
 	RegisterCustomClassLayout("AnimSequence", FOnGetDetailCustomizationInstance::CreateStatic(&FAnimSequenceDetails::MakeInstance));
+	RegisterCustomClassLayout("Rig", FOnGetDetailCustomizationInstance::CreateStatic(&FRigDetails::MakeInstance));
 
 	RegisterCustomClassLayout("EditorAnimSegment", FOnGetDetailCustomizationInstance::CreateStatic(&FAnimMontageSegmentDetails::MakeInstance));
 	RegisterCustomClassLayout("EditorAnimCompositeSegment", FOnGetDetailCustomizationInstance::CreateStatic(&FAnimMontageSegmentDetails::MakeInstance));
@@ -279,4 +281,4 @@ void FDetailCustomizationsModule::RegisterCustomPropertyTypeLayout(FName Propert
 	static FName PropertyEditor("PropertyEditor");
 	FPropertyEditorModule& PropertyModule = FModuleManager::GetModuleChecked<FPropertyEditorModule>(PropertyEditor);
 	PropertyModule.RegisterCustomPropertyTypeLayout(PropertyTypeName, PropertyTypeLayoutDelegate);
-}
\ No newline at end of file
+}
diff --git a/Engine/Source/Editor/DetailCustomizations/Private/RigDetails.cpp b/Engine/Source/Editor/DetailCustomizations/Private/RigDetails.cpp
new file mode 100644
--- /dev/null
+++ b/Engine/Source/Editor/DetailCustomizations/Private/RigDetails.cpp
@@ -0,0 +1,414 @@
+// Copyright 1998-2014 Epic Games, Inc. All Rights Reserved.
+
+#include "DetailCustomizationsPrivatePCH.h"
+#include "RigDetails.h"
+#include "Animation/Rig.h"
+
+#define LOCTEXT_NAMESPACE	"RigDetails"
+
+// Table Headers for node list
+#define NODE_TABLE_DISPLAYNAME	TEXT("DisplayName")
+#define NODE_TABLE_NODENAME		TEXT("NodeName")
+#define NODE_TABLE_PARENTNAME	TEXT("ParentName")
+
+
+TSharedRef<IDetailCustomization> FRigDetails::MakeInstance()
+{
+	return MakeShareable(new FRigDetails);
+}
+
+void FRigDetails::CustomizeDetails(IDetailLayoutBuilder& DetailBuilder)
+{
+	DetailBuilderPtr = &DetailBuilder;
+	TArray<TWeakObjectPtr<UObject>> Objects;
+
+	DetailBuilder.GetObjectsBeingCustomized(Objects);
+
+	// if more than one, do not customize?
+	if (Objects.Num() > 1)
+	{
+		return;
+	}
+
+	ItemBeingEdited = Objects[0];
+
+
+	IDetailCategoryBuilder& NodeCategory = DetailBuilder.EditCategory("Node");
+	IDetailCategoryBuilder& RigControlCategory = DetailBuilder.EditCategory("Constraint Setup");
+
+	RigControlsPropertyHandle = DetailBuilder.GetProperty("RigControls");
+	NodesPropertyHandle = DetailBuilder.GetProperty("Nodes");
+
+	// since now we can't really resize the array, we'll just allocate everything here
+	// if you reallocate, TSharedPtr<FString> for combo box won't work because
+	uint32 NumElement = 0;
+	check (FPropertyAccess::Fail != NodesPropertyHandle->AsArray()->GetNumElements(NumElement));
+
+	if ( NumElement > 0 )
+	{
+		DisplayNameTextBoxes.AddZeroed(NumElement);
+	}
+
+	check (FPropertyAccess::Fail != RigControlsPropertyHandle->AsArray()->GetNumElements(NumElement));
+	if ( NumElement > 0 )
+	{
+		ParentSpaceOptionList.AddZeroed(NumElement);
+	}
+
+	TSharedRef<FDetailArrayBuilder> NodeArrayBuilder = MakeShareable(new FDetailArrayBuilder(NodesPropertyHandle.ToSharedRef()));
+	NodeArrayBuilder->OnGenerateArrayElementWidget(FOnGenerateArrayElementWidget::CreateSP(this, &FRigDetails::GenerateNodeArrayElementWidget, &DetailBuilder));
+
+	NodeCategory.AddCustomBuilder( NodeArrayBuilder, false );
+
+	TSharedRef<FDetailArrayBuilder> RigControlArrayBuilder = MakeShareable(new FDetailArrayBuilder(RigControlsPropertyHandle.ToSharedRef()));
+	RigControlArrayBuilder->OnGenerateArrayElementWidget(FOnGenerateArrayElementWidget::CreateSP(this, &FRigDetails::GenerateRigControlArrayElementWidget, &DetailBuilder));
+
+	// add custom menu
+	// -> set all to world
+	// -> set all to default parent
+	RigControlCategory.AddCustomRow(TEXT(""))
+	[
+		// two button 1. view 2. save to base pose
+		SNew(SHorizontalBox)
+
+		+SHorizontalBox::Slot()
+		.AutoWidth()
+		.Padding(5)
+		.HAlign(HAlign_Center)
+		[
+			SNew(SButton)
+			.OnClicked(FOnClicked::CreateSP(this, &FRigDetails::OnSetAllToWorld))
+			.HAlign(HAlign_Center)
+			.VAlign(VAlign_Center)
+			.Text(LOCTEXT("SetAllToWorld_ButtonLabel", "Set All Constraints to World"))
+		]
+
+		+SHorizontalBox::Slot()
+		.AutoWidth()
+		.Padding(5)
+		.HAlign(HAlign_Center)
+		[
+			SNew(SButton)
+			.OnClicked(FOnClicked::CreateSP(this, &FRigDetails::OnSetAllToParent))
+			.HAlign(HAlign_Center)
+			.VAlign(VAlign_Center)
+			.Text(LOCTEXT("SetAllToParent_ButtonLabel", "Set All Constraints to Parent"))
+		]
+	];
+
+	RigControlCategory.AddCustomBuilder( RigControlArrayBuilder, false );
+}
+
+void FRigDetails::GenerateNodeArrayElementWidget(TSharedRef<IPropertyHandle> PropertyHandle, int32 ArrayIndex, IDetailChildrenBuilder& ChildrenBuilder, IDetailLayoutBuilder* DetailLayout)
+{
+	TSharedRef<IPropertyHandle> DisplayNameProp = PropertyHandle->GetChildHandle(GET_MEMBER_NAME_CHECKED(FNode, DisplayName)).ToSharedRef();
+	TSharedRef<IPropertyHandle> NodeNameProp = PropertyHandle->GetChildHandle(GET_MEMBER_NAME_CHECKED(FNode, Name)).ToSharedRef();
+	TSharedRef<IPropertyHandle> ParentNameProp = PropertyHandle->GetChildHandle(GET_MEMBER_NAME_CHECKED(FNode, ParentName)).ToSharedRef();
+
+	TSharedPtr<SEditableTextBox> DisplayTextBox;
+
+	// the interface will be node [display name] [parent node]
+	// delegate for display name
+	FString NodeName, ParentNodeName, DisplayString;
+	check (NodeNameProp->GetValueAsDisplayString(NodeName) != FPropertyAccess::Fail);
+	check (ParentNameProp->GetValueAsDisplayString(ParentNodeName) != FPropertyAccess::Fail);
+	check (DisplayNameProp->GetValueAsDisplayString(DisplayString) != FPropertyAccess::Fail);
+
+	ChildrenBuilder.AddChildContent(TEXT(""))
+	[
+		SNew(SHorizontalBox)
+
+		+SHorizontalBox::Slot()
+		.Padding(5, 2)
+		.AutoWidth()
+		[
+			SNew(SBox)
+			.WidthOverride(100)
+			.Content()
+			[
+				SNew(STextBlock)
+				.Text(NodeName)
+				.Font(DetailLayout->GetDetailFontBold())
+			]
+		]
+
+		+SHorizontalBox::Slot()
+		.Padding(5, 2)
+		.AutoWidth()
+		[
+			SNew(SBox)
+			.WidthOverride(150)
+			.Content()
+			[
+				SNew(STextBlock)
+				.Text(FString::Printf(TEXT(" [Parent : %s] "), *ParentNodeName))
+				.Font(DetailLayout->GetDetailFont())
+			]
+		]
+
+		+SHorizontalBox::Slot()
+		.Padding(5, 2)
+		.AutoWidth()
+		[
+			SNew(STextBlock)
+			.Text(FString(TEXT("Display Name")))
+			.Font(DetailLayout->GetDetailFontBold())
+		]
+
+		+SHorizontalBox::Slot()
+		.Padding(5, 2)
+		.FillWidth(1)
+		.HAlign(HAlign_Left)
+		[
+			SNew(SBox)
+			.WidthOverride(250)
+			.Content()
+			[
+				SAssignNew(DisplayTextBox, SEditableTextBox)
+				.Text( TAttribute<FText>::Create( TAttribute<FText>::FGetter::CreateSP(this, &FRigDetails::GetDisplayName, DisplayNameProp) ) )
+				.Font(DetailLayout->GetDetailFont())
+				.OnTextChanged(this, &FRigDetails::OnDisplayNameChanged, DisplayNameProp, ArrayIndex)
+				.OnTextCommitted(this, &FRigDetails::OnDisplayNameCommitted, DisplayNameProp, ArrayIndex)
+				.MinDesiredWidth(200)
+			]
+		]
+	];
+
+	DisplayNameTextBoxes[ArrayIndex] = DisplayTextBox;
+};
+
+FText FRigDetails::GetDisplayName(TSharedRef<IPropertyHandle> DisplayNameProp) const
+{
+	FText DisplayText;
+	
+	check (FPropertyAccess::Success == DisplayNameProp->GetValueAsDisplayText(DisplayText));
+
+	return DisplayText;
+}
+
+void FRigDetails::ValidErrorMessage(const FString & DisplayString, int32 ArrayIndex)
+{
+	if(DisplayNameTextBoxes.IsValidIndex(ArrayIndex))
+	{
+		DisplayNameTextBoxes[ArrayIndex]->SetError(TEXT(""));
+
+		if(DisplayString.Len() == 0)
+		{
+			DisplayNameTextBoxes[ArrayIndex]->SetError(TEXT("Name can't be empty"));
+		}
+		else
+		{
+			// verify if this name is unique
+			URig * Rig = Cast<URig>(ItemBeingEdited.Get());
+			if(Rig)
+			{
+				FString NewText = DisplayString;
+				// make sure that name is unique
+				const TArray<FNode> & Nodes = Rig->GetNodes();
+				int32 NodeIndex = 0;
+				for(auto Node : Nodes)
+				{
+					if(NodeIndex++ != ArrayIndex && Node.DisplayName == NewText)
+					{
+						DisplayNameTextBoxes[ArrayIndex]->SetError(TEXT("Name should be unique."));
+					}
+				}
+			}
+		}
+	}
+}
+
+void FRigDetails::OnDisplayNameChanged(const FText& Text, TSharedRef<IPropertyHandle> DisplayNameProp, int32 ArrayIndex)
+{
+	// still set it since you don't know what they come up with 
+	DisplayNameProp->SetValueFromFormattedString(Text.ToString());
+	ValidErrorMessage(Text.ToString(), ArrayIndex);
+}
+
+void FRigDetails::OnDisplayNameCommitted(const FText& Text, ETextCommit::Type CommitType, TSharedRef<IPropertyHandle> DisplayNameProp, int32 ArrayIndex)
+{
+	// @todo error check here? I basically need mirror string to avoid the issue
+	DisplayNameProp->SetValueFromFormattedString(Text.ToString());
+}
+
+void FRigDetails::GenerateRigControlArrayElementWidget(TSharedRef<IPropertyHandle> PropertyHandle, int32 ArrayIndex, IDetailChildrenBuilder& ChildrenBuilder, IDetailLayoutBuilder* DetailLayout)
+{
+	TSharedRef<IPropertyHandle> NodeNameProp = PropertyHandle->GetChildHandle(GET_MEMBER_NAME_CHECKED(FRigControl, Node)).ToSharedRef();
+	TSharedPtr<IPropertyHandleArray> ConstraintsProp = PropertyHandle->GetChildHandle(GET_MEMBER_NAME_CHECKED(FRigControl, Constraints))->AsArray();
+
+	// translation
+	TSharedPtr<IPropertyHandleArray> TransformConstraintsProp_T = ConstraintsProp->GetElement(0)->GetChildHandle(GET_MEMBER_NAME_CHECKED(FRigControlConstraint, TransformConstraints))->AsArray();
+	TSharedRef<IPropertyHandle> ParentNameProp_T = TransformConstraintsProp_T->GetElement(0)->GetChildHandle(GET_MEMBER_NAME_CHECKED(FRigTransformConstraint, ParentSpace)).ToSharedRef();
+
+	// orientation
+	TSharedPtr<IPropertyHandleArray> TransformConstraintsProp_R = ConstraintsProp->GetElement(1)->GetChildHandle(GET_MEMBER_NAME_CHECKED(FRigControlConstraint, TransformConstraints))->AsArray();
+	TSharedRef<IPropertyHandle> ParentNameProp_R = TransformConstraintsProp_R->GetElement(0)->GetChildHandle(GET_MEMBER_NAME_CHECKED(FRigTransformConstraint, ParentSpace)).ToSharedRef();
+
+	// the interface will be node [display name] [parent node]
+	// delegate for display name
+	FString NodeName, ParentNodeName_T, ParentNodeName_R;
+	check(NodeNameProp->GetValueAsDisplayString(NodeName) != FPropertyAccess::Fail);
+	check(ParentNameProp_T->GetValueAsDisplayString(ParentNodeName_T) != FPropertyAccess::Fail);
+	check(ParentNameProp_R->GetValueAsDisplayString(ParentNodeName_R) != FPropertyAccess::Fail);
+
+	// create string list for picking parent node
+	// make sure you don't include itself and find what is curretn selected item
+	TArray<TSharedPtr<FString>> & ParentNodeOptions = ParentSpaceOptionList[ArrayIndex];
+	ParentNodeOptions.Empty();
+	ParentNodeOptions.Add(MakeShareable(new FString(URig::WorldNodeName.ToString())));
+	URig * Rig = Cast<URig>(ItemBeingEdited.Get());
+	check (Rig);
+	const TArray<FNode> & Nodes = Rig->GetNodes();
+	int32 NodeIndex = 0, ParentIndex_T=0, ParentIndex_R=0;
+	const FNode & CurNode = Nodes[ArrayIndex];
+	for(auto Node : Nodes)
+	{
+		if (NodeIndex != ArrayIndex)
+		{
+			ParentNodeOptions.Add(MakeShareable(new FString(Node.Name.ToString())));
+
+			if (Node.Name.ToString() == ParentNodeName_T)
+			{
+				ParentIndex_T = ParentNodeOptions.Num()-1;
+			}
+
+			if(Node.Name.ToString() == ParentNodeName_R)
+			{
+				ParentIndex_R = ParentNodeOptions.Num()-1;
+			}
+		}
+
+		NodeIndex++;
+	}
+
+	ChildrenBuilder.AddChildContent(TEXT(""))
+	[
+		SNew(SHorizontalBox)
+
+		+SHorizontalBox::Slot()
+		.Padding(5, 2)
+		.AutoWidth()
+		[
+			SNew(SBox)
+			.WidthOverride(100)
+			.Content()
+			[
+				SNew(STextBlock)
+				.Text(NodeName)
+				.Font(DetailLayout->GetDetailFontBold())
+			]
+		]
+
+		+SHorizontalBox::Slot()
+		.Padding(5, 2)
+		.AutoWidth()
+		[
+			SNew(SVerticalBox)
+
+			+SVerticalBox::Slot()
+			.Padding(2)
+			[
+				SNew(SHorizontalBox)
+
+				+SHorizontalBox::Slot()
+				.Padding(2)
+				.AutoWidth()
+				[
+					SNew(SBox)
+					.WidthOverride(100)
+					.Content()
+					[
+						SNew(STextBlock)
+						.Text(FString(TEXT("Translation")))
+						.Font(DetailLayout->GetDetailFontBold())
+					]
+				]
+
+				+SHorizontalBox::Slot()
+				.Padding(2)
+				.FillWidth(1)
+				[
+					SNew(SBox)
+					.WidthOverride(250)
+					.Content()
+					[
+						SNew(STextComboBox)
+						.OptionsSource(&ParentNodeOptions)
+						.InitiallySelectedItem(ParentNodeOptions[ParentIndex_T])
+						.OnSelectionChanged(this, &FRigDetails::OnParentSpaceSelectionChanged, ParentNameProp_T)
+					]
+				]
+			]
+
+			+SVerticalBox::Slot()
+			.Padding(2)
+			[
+				SNew(SHorizontalBox)
+
+				+SHorizontalBox::Slot()
+				.Padding(2)
+				.AutoWidth()
+				[
+					SNew(SBox)
+					.WidthOverride(100)
+					.Content()
+					[
+						SNew(STextBlock)
+						.Text(FString(TEXT("Orientation")))
+						.Font(DetailLayout->GetDetailFontBold())
+					]
+				]
+
+				+SHorizontalBox::Slot()
+				.Padding(2)
+				.FillWidth(1)
+				[
+					SNew(SBox)
+					.WidthOverride(250)
+					.Content()
+					[
+						SNew(STextComboBox)
+						.OptionsSource(&ParentNodeOptions)
+						.InitiallySelectedItem(ParentNodeOptions[ParentIndex_R])
+						.OnSelectionChanged(this, &FRigDetails::OnParentSpaceSelectionChanged, ParentNameProp_R)
+					]
+				]
+			]
+		]
+	];
+}
+
+void FRigDetails::OnParentSpaceSelectionChanged(TSharedPtr<FString> SelectedItem, ESelectInfo::Type SelectInfo, TSharedRef<IPropertyHandle> ParentSpacePropertyHandle)
+{
+	if (SelectedItem.IsValid())
+	{
+		check (ParentSpacePropertyHandle->SetValueFromFormattedString(*SelectedItem.Get()) != FPropertyAccess::Fail);
+	}
+}
+
+FReply FRigDetails::OnSetAllToWorld()
+{
+	URig * Rig = Cast<URig>(ItemBeingEdited.Get());
+	check(Rig);
+
+	const FScopedTransaction Transaction(LOCTEXT("SetAllToWorld_Action", "Set All Transform Constraints to World"));
+	Rig->Modify();
+	Rig->SetAllConstraintsToWorld();
+	DetailBuilderPtr->ForceRefreshDetails();
+
+	return FReply::Handled();
+}
+
+FReply FRigDetails::OnSetAllToParent()
+{
+	URig * Rig = Cast<URig>(ItemBeingEdited.Get());
+	check(Rig);
+
+	const FScopedTransaction Transaction(LOCTEXT("SetAllToParent_Action", "Set All Transform Constraints to Parent"));
+	Rig->Modify();
+	Rig->SetAllConstraintsToParents();
+	DetailBuilderPtr->ForceRefreshDetails();
+
+	return FReply::Handled();
+}
+#undef LOCTEXT_NAMESPACE
diff --git a/Engine/Source/Editor/DetailCustomizations/Private/RigDetails.h b/Engine/Source/Editor/DetailCustomizations/Private/RigDetails.h
new file mode 100644
--- /dev/null
+++ b/Engine/Source/Editor/DetailCustomizations/Private/RigDetails.h
@@ -0,0 +1,47 @@
+// Copyright 1998-2014 Epic Games, Inc. All Rights Reserved.
+
+#pragma once
+
+#include "PreviewScene.h"
+
+class FRigDetails : public IDetailCustomization
+{
+public:
+	/** Makes a new instance of this detail layout class for a specific detail view requesting it */
+	static TSharedRef<IDetailCustomization> MakeInstance();
+
+	/** IDetailCustomization interface */
+	virtual void CustomizeDetails( IDetailLayoutBuilder& DetailBuilder ) override;
+
+private:
+
+	/** Generate a widget for a movie array element */
+	void GenerateNodeArrayElementWidget(TSharedRef<IPropertyHandle> PropertyHandle, int32 ArrayIndex, IDetailChildrenBuilder& ChildrenBuilder, IDetailLayoutBuilder* DetailLayout);
+	void GenerateRigControlArrayElementWidget(TSharedRef<IPropertyHandle> PropertyHandle, int32 ArrayIndex, IDetailChildrenBuilder& ChildrenBuilder, IDetailLayoutBuilder* DetailLayout);
+
+	TSharedPtr<IPropertyHandle> NodesPropertyHandle;
+	TSharedPtr<IPropertyHandle> RigControlsPropertyHandle;
+
+	// array custom boxes - these will stay around as long as this window is up
+	TArray<TSharedPtr<SEditableTextBox>>	DisplayNameTextBoxes;
+	TArray<TArray<TSharedPtr<FString>>>		ParentSpaceOptionList;
+
+	/** we only support one item */
+	TWeakObjectPtr<UObject> ItemBeingEdited;
+
+	// node display text handler
+	void ValidErrorMessage(const FString & DisplayString, int32 ArrayIndex);
+	FText GetDisplayName(TSharedRef<IPropertyHandle> DisplayNameProp) const;
+	void OnDisplayNameChanged(const FText& Text, TSharedRef<IPropertyHandle> DisplayNameProp, int32 ArrayIndex);
+	void OnDisplayNameCommitted(const FText& Text, ETextCommit::Type CommitType, TSharedRef<IPropertyHandle> DisplayNameProp, int32 ArrayIndex);
+
+	// combo box handler
+	void OnParentSpaceSelectionChanged(TSharedPtr<FString> SelectedItem, ESelectInfo::Type SelectInfo, TSharedRef<IPropertyHandle> ParentSpacePropertyHandle);
+
+	// button handlers
+	FReply OnSetAllToWorld();
+	FReply OnSetAllToParent();
+
+	IDetailLayoutBuilder * DetailBuilderPtr;
+};
+
diff --git a/Engine/Source/Editor/EditorStyle/Private/SlateEditorStyle.cpp b/Engine/Source/Editor/EditorStyle/Private/SlateEditorStyle.cpp
--- a/Engine/Source/Editor/EditorStyle/Private/SlateEditorStyle.cpp
+++ b/Engine/Source/Editor/EditorStyle/Private/SlateEditorStyle.cpp
@@ -3136,6 +3136,15 @@ void FSlateEditorStyle::FStyle::SetupGraphEditorStyles()
 					.SetShadowColorAndOpacity( FLinearColor(0.f,0.f,0.f, 0.7f) )
 				);
 			}
+
+			Set( "Persona.RetargetManager.BoldFont",							TTF_CORE_FONT( "Fonts/Roboto-Bold", 10 ) );
+			Set( "Persona.RetargetManager.FilterFont",							TTF_CORE_FONT( "Fonts/Roboto-Regular", 10 ) );
+
+			Set("Persona.RetargetManager.ImportantText", FTextBlockStyle(NormalText)
+				.SetFont(TTF_CORE_FONT("Fonts/Roboto-Bold", 11))
+				.SetColorAndOpacity(FLinearColor(1.0f, 1.0f, 1.0f, 1.0f))
+				.SetShadowOffset(FVector2D(1, 1))
+				.SetShadowColorAndOpacity(FLinearColor(0, 0, 0, 0.9f)));
 		}
 
 		// Behavior Tree Editor
@@ -4212,8 +4221,8 @@ void FSlateEditorStyle::FStyle::SetupPersonaStyle()
 		Set("Persona.ChangeSkeletonPreviewMesh.Small", new IMAGE_BRUSH(TEXT("Icons/icon_Persona_PreviewMesh_40x"), Icon20x20));
 		Set("Persona.AnimNotifyWindow", new IMAGE_BRUSH(TEXT("Icons/icon_Persona_AnimNotift_40x"), Icon40x40));
 		Set("Persona.AnimNotifyWindow.Small", new IMAGE_BRUSH(TEXT("Icons/icon_Persona_AnimNotift_40x"), Icon20x20));
-		Set("Persona.RetargetSourceMgr", new IMAGE_BRUSH(TEXT("Icons/icon_Persona_Retarget_40x"), Icon40x40));
-		Set("Persona.RetargetSourceMgr.Small", new IMAGE_BRUSH(TEXT("Icons/icon_Persona_Retarget_40x"), Icon20x20));
+		Set("Persona.RetargetManager", new IMAGE_BRUSH(TEXT("Icons/icon_Persona_Retarget_40x"), Icon40x40));
+		Set("Persona.RetargetManager.Small", new IMAGE_BRUSH(TEXT("Icons/icon_Persona_Retarget_40x"), Icon20x20));
 		Set("Persona.ImportMesh", new IMAGE_BRUSH(TEXT("Icons/icon_Persona_ImportMesh_40x"), Icon40x40));
 		Set("Persona.ImportMesh.Small", new IMAGE_BRUSH(TEXT("Icons/icon_Persona_ImportMesh_40x"), Icon20x20));
 		Set("Persona.ReimportMesh", new IMAGE_BRUSH(TEXT("Icons/icon_Persona_ReimportMesh_40x"), Icon40x40));
diff --git a/Engine/Source/Editor/Persona/Private/AnimViewportShowCommands.cpp b/Engine/Source/Editor/Persona/Private/AnimViewportShowCommands.cpp
--- a/Engine/Source/Editor/Persona/Private/AnimViewportShowCommands.cpp
+++ b/Engine/Source/Editor/Persona/Private/AnimViewportShowCommands.cpp
@@ -16,6 +16,7 @@ void FAnimViewportShowCommands::RegisterCommands()
 	UI_COMMAND(ProcessRootMotion, "Process Root Motion", "Moves preview based on animation root motion", EUserInterfaceActionType::ToggleButton, FInputGesture());
 
 	UI_COMMAND( ShowReferencePose, "Reference Pose", "Show reference pose on preview mesh", EUserInterfaceActionType::ToggleButton, FInputGesture() );
+	UI_COMMAND( ShowRetargetBasePose, "Retarget Base Pose", "Show retarget Base pose on preview mesh", EUserInterfaceActionType::ToggleButton, FInputGesture() );
 	UI_COMMAND( ShowBound, "Bound", "Show bound on preview mesh", EUserInterfaceActionType::ToggleButton, FInputGesture() );
 	UI_COMMAND( ShowPreviewMesh, "Mesh", "Show the preview mesh", EUserInterfaceActionType::ToggleButton, FInputGesture() );
 
diff --git a/Engine/Source/Editor/Persona/Private/AnimViewportShowCommands.h b/Engine/Source/Editor/Persona/Private/AnimViewportShowCommands.h
--- a/Engine/Source/Editor/Persona/Private/AnimViewportShowCommands.h
+++ b/Engine/Source/Editor/Persona/Private/AnimViewportShowCommands.h
@@ -41,6 +41,9 @@ public:
 
 	/** Show reference pose on preview mesh */
 	TSharedPtr< FUICommandInfo > ShowReferencePose;
+
+	/** Show reference pose on preview mesh */
+	TSharedPtr< FUICommandInfo > ShowRetargetBasePose;
 	
 	/** Show Bound of preview mesh */
 	TSharedPtr< FUICommandInfo > ShowBound;
diff --git a/Engine/Source/Editor/Persona/Private/Persona.cpp b/Engine/Source/Editor/Persona/Private/Persona.cpp
--- a/Engine/Source/Editor/Persona/Private/Persona.cpp
+++ b/Engine/Source/Editor/Persona/Private/Persona.cpp
@@ -751,7 +751,7 @@ void FPersona::ExtendDefaultPersonaToolbar()
 			{
 				ToolbarBuilder.AddToolBarButton(FPersonaCommands::Get().ChangeSkeletonPreviewMesh, NAME_None, LOCTEXT("Toolbar_ChangePreviewMesh", "Set Preview"));
 				ToolbarBuilder.AddToolBarButton(FPersonaCommands::Get().AnimNotifyWindow);
-				ToolbarBuilder.AddToolBarButton(FPersonaCommands::Get().RetargetSourceMgr, NAME_None, LOCTEXT("Toolbar_RetargetSourceMgr", "Retarget Source"));
+				ToolbarBuilder.AddToolBarButton(FPersonaCommands::Get().RetargetManager, NAME_None, LOCTEXT("Toolbar_RetargetManager", "Retarget Manager"));
 				ToolbarBuilder.AddToolBarButton(FPersonaCommands::Get().ImportMesh);
 				ToolbarBuilder.AddToolBarButton(FPersonaCommands::Get().ReimportMesh);
 
@@ -1007,8 +1007,8 @@ void FPersona::CreateDefaultCommands()
 		FIsActionButtonVisible::CreateSP(this, &FPersona::IsInPersonaMode, FPersonaModes::SkeletonDisplayMode)
 		);
 
-	ToolkitCommands->MapAction(FPersonaCommands::Get().RetargetSourceMgr,
-		FExecuteAction::CreateSP(this, &FPersona::OnRetargetSourceMgr),
+	ToolkitCommands->MapAction(FPersonaCommands::Get().RetargetManager,
+		FExecuteAction::CreateSP(this, &FPersona::OnRetargetManager),
 		FCanExecuteAction(),
 		FIsActionChecked(),
 		FIsActionButtonVisible::CreateSP(this, &FPersona::IsInPersonaMode, FPersonaModes::SkeletonDisplayMode)
@@ -2783,9 +2783,9 @@ void FPersona::OnAnimNotifyWindow()
 	TabManager->InvokeTab(FPersonaTabs::SkeletonAnimNotifiesID);
 }
 
-void FPersona::OnRetargetSourceMgr()
+void FPersona::OnRetargetManager()
 {
-	TabManager->InvokeTab(FPersonaTabs::RetargetSourceManagerID);
+	TabManager->InvokeTab(FPersonaTabs::RetargetManagerID);
 }
 
 void FPersona::OnImportAsset(enum EFBXImportType DefaultImportType)
diff --git a/Engine/Source/Editor/Persona/Private/Persona.h b/Engine/Source/Editor/Persona/Private/Persona.h
--- a/Engine/Source/Editor/Persona/Private/Persona.h
+++ b/Engine/Source/Editor/Persona/Private/Persona.h
@@ -560,7 +560,7 @@ private:
 
 	// tool bar actions
 	void OnAnimNotifyWindow();
-	void OnRetargetSourceMgr();
+	void OnRetargetManager();
 	void OnReimportMesh();
 	void OnImportAsset(enum EFBXImportType DefaultImportType);
 	void OnReimportAnimation();
diff --git a/Engine/Source/Editor/Persona/Private/PersonaCommands.cpp b/Engine/Source/Editor/Persona/Private/PersonaCommands.cpp
--- a/Engine/Source/Editor/Persona/Private/PersonaCommands.cpp
+++ b/Engine/Source/Editor/Persona/Private/PersonaCommands.cpp
@@ -13,7 +13,7 @@ void FPersonaCommands::RegisterCommands()
 	UI_COMMAND( ChangeSkeletonPreviewMesh, "Set Preview Mesh as Default", "Changes the skeletons default preview mesh to the current open preview mesh. The skeleton will require saving after this action.", EUserInterfaceActionType::Button, FInputGesture() );
 	UI_COMMAND( RemoveUnusedBones, "Remove Unused Bones from Skeleton", "Removes any bones from the skeleton that are not used by any of its meshes. The skeleton and associated animations will require saving after this action.", EUserInterfaceActionType::Button, FInputGesture() );
 	UI_COMMAND( AnimNotifyWindow, "Anim Notifies", "You can manage animation notifies that belong to the skeleton.", EUserInterfaceActionType::Button, FInputGesture() );
-	UI_COMMAND( RetargetSourceMgr, "Retarget Source Manager", "Create, delete, or edit retarget sources for this skeleton. ", EUserInterfaceActionType::Button, FInputGesture() );
+	UI_COMMAND( RetargetManager, "Retarget Manager", "Manager retarget setups. ", EUserInterfaceActionType::Button, FInputGesture() );
 	UI_COMMAND( ImportMesh, "Import Mesh", "Import new mesh for this skeleton. ", EUserInterfaceActionType::Button, FInputGesture() );
 
 	// mesh menu
diff --git a/Engine/Source/Editor/Persona/Private/PersonaCommands.h b/Engine/Source/Editor/Persona/Private/PersonaCommands.h
--- a/Engine/Source/Editor/Persona/Private/PersonaCommands.h
+++ b/Engine/Source/Editor/Persona/Private/PersonaCommands.h
@@ -27,7 +27,7 @@ public:
 	// Command to show Anim Notify window
 	TSharedPtr<FUICommandInfo> AnimNotifyWindow;
 	// Command to show Retarget Source Manager
-	TSharedPtr<FUICommandInfo> RetargetSourceMgr;
+	TSharedPtr<FUICommandInfo> RetargetManager;
 	// Import Mesh for this Skeleton
 	TSharedPtr<FUICommandInfo> ImportMesh;
 			
diff --git a/Engine/Source/Editor/Persona/Private/SAnimViewportToolBar.cpp b/Engine/Source/Editor/Persona/Private/SAnimViewportToolBar.cpp
--- a/Engine/Source/Editor/Persona/Private/SAnimViewportToolBar.cpp
+++ b/Engine/Source/Editor/Persona/Private/SAnimViewportToolBar.cpp
@@ -457,6 +457,7 @@ TSharedRef<SWidget> SAnimViewportToolBar::GenerateShowMenu() const
 			ShowMenuBuilder.BeginSection("AnimViewportMesh", LOCTEXT("ShowMenu_Actions_Mesh", "Mesh"));
 			{
 				ShowMenuBuilder.AddMenuEntry( Actions.ShowReferencePose );
+				ShowMenuBuilder.AddMenuEntry( Actions.ShowRetargetBasePose );
 				ShowMenuBuilder.AddMenuEntry( Actions.ShowBound );
 				ShowMenuBuilder.AddMenuEntry( Actions.ShowPreviewMesh );
 			}
diff --git a/Engine/Source/Editor/Persona/Private/SAnimationEditorViewport.cpp b/Engine/Source/Editor/Persona/Private/SAnimationEditorViewport.cpp
--- a/Engine/Source/Editor/Persona/Private/SAnimationEditorViewport.cpp
+++ b/Engine/Source/Editor/Persona/Private/SAnimationEditorViewport.cpp
@@ -446,6 +446,12 @@ void SAnimationEditorViewportTabBody::BindCommands()
 		FCanExecuteAction::CreateSP(this, &SAnimationEditorViewportTabBody::CanShowReferencePose),
 		FIsActionChecked::CreateSP(this, &SAnimationEditorViewportTabBody::IsShowReferencePoseEnabled));
 
+	CommandList.MapAction(
+		ViewportShowMenuCommands.ShowRetargetBasePose,
+		FExecuteAction::CreateSP(this, &SAnimationEditorViewportTabBody::ShowRetargetBasePose),
+		FCanExecuteAction::CreateSP(this, &SAnimationEditorViewportTabBody::CanShowRetargetBasePose),
+		FIsActionChecked::CreateSP(this, &SAnimationEditorViewportTabBody::IsShowRetargetBasePoseEnabled));
+
 	CommandList.MapAction(
 		ViewportShowMenuCommands.ShowBound,
 		FExecuteAction::CreateSP(this, &SAnimationEditorViewportTabBody::ShowBound),
@@ -981,6 +987,31 @@ bool SAnimationEditorViewportTabBody::IsShowReferencePoseEnabled() const
 	return false;
 }
 
+void SAnimationEditorViewportTabBody::ShowRetargetBasePose()
+{
+	UDebugSkelMeshComponent* PreviewComponent = PersonaPtr.Pin()->PreviewComponent;
+	if(PreviewComponent && PreviewComponent->PreviewInstance)
+	{
+		PreviewComponent->PreviewInstance->bForceRetargetBasePose = !PreviewComponent->PreviewInstance->bForceRetargetBasePose;
+	}
+}
+
+bool SAnimationEditorViewportTabBody::CanShowRetargetBasePose() const
+{
+	UDebugSkelMeshComponent* PreviewComponent = PersonaPtr.Pin()->PreviewComponent;
+	return PreviewComponent != NULL && PreviewComponent->PreviewInstance;
+}
+
+bool SAnimationEditorViewportTabBody::IsShowRetargetBasePoseEnabled() const
+{
+	UDebugSkelMeshComponent* PreviewComponent = PersonaPtr.Pin()->PreviewComponent;
+	if(PreviewComponent && PreviewComponent->PreviewInstance)
+	{
+		return PreviewComponent->PreviewInstance->bForceRetargetBasePose;
+	}
+	return false;
+}
+
 void SAnimationEditorViewportTabBody::ShowBound()
 {
 	TSharedRef<FAnimationViewportClient> AnimViewportClient = StaticCastSharedRef<FAnimationViewportClient>(LevelViewportClient.ToSharedRef());	
diff --git a/Engine/Source/Editor/Persona/Private/SAnimationEditorViewport.h b/Engine/Source/Editor/Persona/Private/SAnimationEditorViewport.h
--- a/Engine/Source/Editor/Persona/Private/SAnimationEditorViewport.h
+++ b/Engine/Source/Editor/Persona/Private/SAnimationEditorViewport.h
@@ -276,6 +276,11 @@ private:
 	bool CanShowReferencePose() const;
 	bool IsShowReferencePoseEnabled() const;
 
+	/** Called to toggle showing of reference pose on current preview mesh */
+	void ShowRetargetBasePose();
+	bool CanShowRetargetBasePose() const;
+	bool IsShowRetargetBasePoseEnabled() const;
+
 	/** Called to toggle showing of the bounds of the current preview mesh */
 	void ShowBound();
 	bool CanShowBound() const;
diff --git a/Engine/Source/Editor/Persona/Private/SRetargetManager.cpp b/Engine/Source/Editor/Persona/Private/SRetargetManager.cpp
new file mode 100644
--- /dev/null
+++ b/Engine/Source/Editor/Persona/Private/SRetargetManager.cpp
@@ -0,0 +1,290 @@
+// Copyright 1998-2014 Epic Games, Inc. All Rights Reserved.
+
+
+#include "PersonaPrivatePCH.h"
+#include "SRetargetManager.h"
+#include "ObjectTools.h"
+#include "ScopedTransaction.h"
+#include "AssetRegistryModule.h"
+#include "Editor/PropertyEditor/Public/PropertyEditorModule.h"
+#include "Editor/ContentBrowser/Public/ContentBrowserModule.h"
+#include "AssetNotifications.h"
+#include "Animation/Rig.h"
+#include "BoneSelectionWidget.h"
+#include "SRetargetSourceWindow.h"
+#include "SRigWindow.h"
+
+#define LOCTEXT_NAMESPACE "SRetargetManager"
+
+
+//////////////////////////////////////////////////////////////////////////
+// SRetargetManager
+
+void SRetargetManager::Construct(const FArguments& InArgs)
+{
+	PersonaPtr = InArgs._Persona;
+	Skeleton = NULL;
+
+	if ( PersonaPtr.IsValid() )
+	{
+		Skeleton = PersonaPtr.Pin()->GetSkeleton();
+		PersonaPtr.Pin()->RegisterOnPostUndo(FPersona::FOnPostUndo::CreateSP( this, &SRetargetManager::PostUndo ) );
+	}
+
+	const FString DocLink = TEXT("Shared/Editors/Persona");
+	ChildSlot
+	[
+		SNew (SVerticalBox)
+
+		+ SVerticalBox::Slot()
+		.Padding(5, 5)
+		.AutoHeight()
+		[
+			// explainint this is retarget source window
+			// and what it is
+			SNew(STextBlock)
+			.TextStyle( FEditorStyle::Get(), "Persona.RetargetManager.ImportantText" )
+			.Text(LOCTEXT("RetargetSource_Title", "Manager Retarget Source"))
+		]
+
+		+ SVerticalBox::Slot()
+		.Padding(5, 5)
+		.AutoHeight()
+		[
+			// explainint this is retarget source window
+			// and what it is
+			SNew(STextBlock)
+			.AutoWrapText(true)
+			.ToolTip(IDocumentation::Get()->CreateToolTip(LOCTEXT("RetargetSource_Tooltip", "Add/Delete/Rename Retarget Sources."),
+																			NULL,
+																			DocLink,
+																			TEXT("RetargetSource")))
+			.Font(FEditorStyle::GetFontStyle(TEXT("Persona.RetargetManager.FilterFont")))
+			.Text(LOCTEXT("RetargetSource_Description", "You can add/rename/delete Retarget Sources. When you have different proportional meshes per skeleton, you can use this settings to indicate if this animation is from different source. \
+														For example, if your default skeleton is from small guy, and if you have animation for big guy, you can create Retarget Source from the big guy and set it for the animation. \
+														Retargeting system will use this information when extracting animation. "))
+		]
+
+		+ SVerticalBox::Slot()
+		.Padding(2, 5)
+		.FillHeight(0.5)
+		[
+			// construct retarget source window
+			SNew(SRetargetSourceWindow)
+			.Persona(PersonaPtr)
+		]
+
+		+SVerticalBox::Slot()
+		.Padding(5, 5)
+		.AutoHeight()
+		[
+			SNew(SSeparator)
+			.Orientation(Orient_Horizontal)
+		]
+
+		+ SVerticalBox::Slot()
+		.Padding(5, 5)
+		.AutoHeight()
+		[
+			// explainint this is retarget source window
+			// and what it is
+			SNew(STextBlock)
+			.TextStyle(FEditorStyle::Get(), "Persona.RetargetManager.ImportantText")
+			.Text(LOCTEXT("RigTemplate_Title", "Set up Rig"))
+		]
+
+		+ SVerticalBox::Slot()
+		.Padding(5, 5)
+		.AutoHeight()
+		[
+			// explainint this is retarget source window
+			// and what it is
+			SNew(STextBlock)
+			.AutoWrapText(true)
+			.ToolTip(IDocumentation::Get()->CreateToolTip(LOCTEXT("RigSetup_Tooltip", "Set up Rig for retargeting between skeletons."),
+																			NULL,
+																			DocLink,
+																			TEXT("Rig Setup")))
+			.Font(FEditorStyle::GetFontStyle(TEXT("Persona.RetargetManager.FilterFont")))
+			.Text(LOCTEXT("RigTemplate_Description", "You can set up Rig for this skeleton, then when you retarget animation to different skeleton with the same Rig, it will use the information to convert data. "))
+		]
+
+		+ SVerticalBox::Slot()
+		.FillHeight(1)
+		.Padding(2, 5)
+		[
+			// construct rig manager window
+			SNew(SRigWindow)
+			.Persona(PersonaPtr)
+		]
+
+		+SVerticalBox::Slot()
+		.Padding(2, 5)
+		.AutoHeight()
+		[
+			SNew(SSeparator)
+			.Orientation(Orient_Horizontal)
+		]
+
+		+ SVerticalBox::Slot()
+		.Padding(5, 5)
+		.AutoHeight()
+		[
+			// explainint this is retarget source window
+			// and what it is
+			SNew(STextBlock)
+			.TextStyle(FEditorStyle::Get(), "Persona.RetargetManager.ImportantText")
+			.Text(LOCTEXT("BasePose_Title", "Manage Retarget Base Pose"))
+		]
+		// construct base pose options
+		+SVerticalBox::Slot()
+		.Padding(2, 5)
+		.AutoHeight()
+		[
+			// explainint this is retarget source window
+			// and what it is
+			SNew(STextBlock)
+			.AutoWrapText(true)
+			.ToolTip(IDocumentation::Get()->CreateToolTip(LOCTEXT("RetargetBasePose_Tooltip", "Set up base pose for retargeting."),
+																			NULL,
+																			DocLink,
+																			TEXT("Set up Base Pose")))
+			.Font(FEditorStyle::GetFontStyle(TEXT("Persona.RetargetManager.FilterFont")))
+			.Text(LOCTEXT("BasePose_Description", "This information is used when retargeting assets to different skeleton. You need to make sure the ref pose of both mesh is same when retargeting, so you can see the pose and \
+											edit using bone transform widget, and click Save button below. "))
+		]
+
+		+ SVerticalBox::Slot()
+		.AutoHeight()		// This is required to make the scrollbar work, as content overflows Slate containers by default
+		.VAlign(VAlign_Center)
+		.HAlign(HAlign_Right)
+		.Padding(2, 5)
+		[
+			// two button 1. view 2. save to base pose
+			SNew(SHorizontalBox)
+
+			+SHorizontalBox::Slot()
+			.AutoWidth()
+			.HAlign(HAlign_Center)
+			[
+				SNew(SButton)
+				.OnClicked(FOnClicked::CreateSP(this, &SRetargetManager::OnViewRetargetBasePose))
+				.HAlign(HAlign_Center)
+				.VAlign(VAlign_Center)
+				.Text(this, &SRetargetManager::GetToggleRetargetBasePose)
+			]
+
+			+SHorizontalBox::Slot()
+			.AutoWidth()
+			.HAlign(HAlign_Center)
+			[
+				SNew(SButton)
+				.OnClicked(FOnClicked::CreateSP(this, &SRetargetManager::OnSaveRetargetBasePose))
+				.HAlign(HAlign_Center)
+				.VAlign(VAlign_Center)
+				.Text(LOCTEXT("SaveRetargetBasePose_Label", "Save Current Pose to Retarget Base Pose"))
+			]
+		]
+	];
+}
+
+SRetargetManager::~SRetargetManager()
+{
+	if (PersonaPtr.IsValid())
+	{
+		PersonaPtr.Pin()->UnregisterOnPostUndo(this);
+	}
+}
+
+FReply SRetargetManager::OnViewRetargetBasePose()
+{
+	UDebugSkelMeshComponent * PreviewMeshComp = PersonaPtr.Pin()->GetPreviewMeshComponent();
+	if (PreviewMeshComp && PreviewMeshComp->PreviewInstance)
+	{
+		const FScopedTransaction Transaction(LOCTEXT("ViewRetargetBasePose_Action", "Edit Retarget Base Pose"));
+		PreviewMeshComp->PreviewInstance->bForceRetargetBasePose = !PreviewMeshComp->PreviewInstance->bForceRetargetBasePose;
+		PreviewMeshComp->Modify();
+		// reset all bone transform since you don't want to keep any bone transform change
+		PreviewMeshComp->PreviewInstance->ResetModifiedBone();
+		// add root 
+		if (PreviewMeshComp->PreviewInstance->bForceRetargetBasePose)
+		{
+			PreviewMeshComp->BonesOfInterest.Add(0);
+		}
+	}
+
+	return FReply::Handled();
+}
+
+FReply SRetargetManager::OnSaveRetargetBasePose()
+{
+	UDebugSkelMeshComponent * PreviewMeshComp = PersonaPtr.Pin()->GetPreviewMeshComponent();
+	if (PreviewMeshComp && PreviewMeshComp->SkeletalMesh)
+	{
+		USkeletalMesh * PreviewMesh = PreviewMeshComp->SkeletalMesh;
+
+		check (Skeleton == PreviewMesh->Skeleton);
+
+		if (PreviewMesh)
+		{
+			const FScopedTransaction Transaction( LOCTEXT("SaveRetargetBasePose_Action", "Save Retarget Base Pose") );
+			PreviewMesh->Modify();
+			// get space bases and calculate local
+			const TArray<FTransform> & SpaceBases = PreviewMeshComp->SpaceBases;
+			// @todo check to see if skeleton vs preview mesh makes it different for missing bones
+			const FReferenceSkeleton & RefSkeleton = Skeleton->GetReferenceSkeleton();
+			TArray<FTransform> & NewRetargetBasePose = PreviewMesh->RetargetBasePose;
+			// if you're using master pose component in preview, this won't work
+			check (PreviewMesh->RefSkeleton.GetNum() == SpaceBases.Num());
+			int32 TotalNumBones = PreviewMesh->RefSkeleton.GetNum();
+			NewRetargetBasePose.Empty(TotalNumBones);
+			NewRetargetBasePose.AddUninitialized(TotalNumBones);
+
+			for (int32 BoneIndex = 0; BoneIndex<TotalNumBones; ++BoneIndex)
+			{
+				// this is slower, but skeleton can have more bones, so I can't just access
+				// Parent index from Skeleton. Going safer route
+				FName BoneName = PreviewMeshComp->GetBoneName(BoneIndex);
+				FName ParentBoneName = PreviewMeshComp->GetParentBone(BoneName);
+				int32 ParentIndex = RefSkeleton.FindBoneIndex(ParentBoneName);
+
+				if (ParentIndex != INDEX_NONE)
+				{
+					NewRetargetBasePose[BoneIndex] = SpaceBases[BoneIndex].GetRelativeTransform(SpaceBases[ParentIndex]);
+				}
+				else
+				{
+					NewRetargetBasePose[BoneIndex] = SpaceBases[BoneIndex];
+				}
+			}
+
+			// Clear PreviewMeshComp bone modified, they're baked now
+			PreviewMeshComp->PreviewInstance->ResetModifiedBone();
+		}
+	}
+	return FReply::Handled();
+}
+
+void SRetargetManager::PostUndo()
+{
+}
+
+FText SRetargetManager::GetToggleRetargetBasePose() const
+{
+	UDebugSkelMeshComponent * PreviewMeshComp = PersonaPtr.Pin()->GetPreviewMeshComponent();
+	if(PreviewMeshComp && PreviewMeshComp->PreviewInstance)
+	{
+		if (PreviewMeshComp->PreviewInstance->bForceRetargetBasePose)
+		{
+			return LOCTEXT("HideRetargetBasePose_Label", "Hide Current Retarget Base Pose");
+		}
+		else
+		{
+			return LOCTEXT("ViewRetargetBasePose_Label", "View/Edit Current Retarget Base Pose");
+		}
+	}
+
+	return LOCTEXT("InvalidRetargetBasePose_Label", "No Mesh for Retarget Base Pose");
+}
+#undef LOCTEXT_NAMESPACE
+
diff --git a/Engine/Source/Editor/Persona/Private/SRetargetManager.h b/Engine/Source/Editor/Persona/Private/SRetargetManager.h
new file mode 100644
--- /dev/null
+++ b/Engine/Source/Editor/Persona/Private/SRetargetManager.h
@@ -0,0 +1,48 @@
+// Copyright 1998-2014 Epic Games, Inc. All Rights Reserved.
+
+#pragma once
+
+#include "Persona.h"
+
+//////////////////////////////////////////////////////////////////////////
+// SRetargetManager
+
+class SRetargetManager : public SCompoundWidget
+{
+public:
+	SLATE_BEGIN_ARGS( SRetargetManager )
+		: _Persona()
+	{}
+		
+		/* The Persona that owns this table */
+		SLATE_ARGUMENT( TWeakPtr< FPersona >, Persona )
+
+	SLATE_END_ARGS()
+
+	/**
+	* Slate construction function
+	*
+	* @param InArgs - Arguments passed from Slate
+	*
+	*/
+	void Construct( const FArguments& InArgs );
+
+	/**
+	* Destructor - resets the morph targets
+	*
+	*/
+	virtual ~SRetargetManager();
+
+private:
+	FReply OnViewRetargetBasePose();
+	FReply OnSaveRetargetBasePose();
+
+	/** Pointer back to the Persona that owns us */
+	TWeakPtr<FPersona> PersonaPtr;
+	/** The skeletal mesh that we grab the morph targets from */
+	USkeleton* Skeleton;
+
+	/** Delegate for undo/redo transaction **/
+	void PostUndo();
+	FText GetToggleRetargetBasePose() const;	
+};
\ No newline at end of file
diff --git a/Engine/Source/Editor/Persona/Private/SRetargetSourceManager.cpp b/Engine/Source/Editor/Persona/Private/SRetargetSourceManager.cpp
deleted file mode 100644
diff --git a/Engine/Source/Editor/Persona/Private/SRetargetSourceManager.h b/Engine/Source/Editor/Persona/Private/SRetargetSourceManager.h
deleted file mode 100644
diff --git a/Engine/Source/Editor/Persona/Private/SRetargetSourceWindow.cpp b/Engine/Source/Editor/Persona/Private/SRetargetSourceWindow.cpp
new file mode 100644
--- /dev/null
+++ b/Engine/Source/Editor/Persona/Private/SRetargetSourceWindow.cpp
@@ -0,0 +1,757 @@
+// Copyright 1998-2014 Epic Games, Inc. All Rights Reserved.
+
+
+#include "PersonaPrivatePCH.h"
+#include "SRetargetSourceWindow.h"
+#include "ObjectTools.h"
+#include "ScopedTransaction.h"
+#include "AssetRegistryModule.h"
+#include "Editor/PropertyEditor/Public/PropertyEditorModule.h"
+#include "Editor/ContentBrowser/Public/ContentBrowserModule.h"
+#include "AssetNotifications.h"
+
+#define LOCTEXT_NAMESPACE "SRetargetSourceWindow"
+
+static const FName ColumnId_RetargetSourceNameLabel( "Retarget Source Name" );
+static const FName ColumnID_BaseReferenceMeshLabel( "Reference Mesh" );
+
+DECLARE_DELEGATE_TwoParams( FOnRenameCommit, const FName & /*OldName*/,  const FString& /*NewName*/ )
+DECLARE_DELEGATE_RetVal_ThreeParams( bool, FOnVerifyRenameCommit, const FName & /*OldName*/, const FString& /*NewName*/, FText& /*OutErrorMessage*/)
+
+//////////////////////////////////////////////////////////////////////////
+// SRetargetSourceListRow
+
+typedef TSharedPtr< FDisplayedRetargetSourceInfo > FDisplayedRetargetSourceInfoPtr;
+
+class SRetargetSourceListRow
+	: public SMultiColumnTableRow< FDisplayedRetargetSourceInfoPtr >
+{
+public:
+
+	SLATE_BEGIN_ARGS( SRetargetSourceListRow ) {}
+
+	/** The item for this row **/
+	SLATE_ARGUMENT( FDisplayedRetargetSourceInfoPtr, Item )
+
+	/* The SRetargetSourceWindow that handles all retarget sources */
+	SLATE_ARGUMENT( class SRetargetSourceWindow*, RetargetSourceWindow )
+
+	/* Widget used to display the list of retarget sources*/
+	SLATE_ARGUMENT( TSharedPtr<SRetargetSourceListType>, RetargetSourceListView )
+
+	/* Persona used to update the viewport when a weight slider is dragged */
+	SLATE_ARGUMENT( TWeakPtr<FPersona>, Persona )
+
+	/** Delegate for when an asset name has been entered for an item that is in a rename state */
+	SLATE_EVENT( FOnRenameCommit, OnRenameCommit )
+
+	/** Delegate for when an asset name has been entered for an item to verify the name before commit */
+	SLATE_EVENT( FOnVerifyRenameCommit, OnVerifyRenameCommit )
+
+	SLATE_END_ARGS()
+
+	void Construct( const FArguments& InArgs, const TSharedRef<STableViewBase>& OwnerTableView );
+
+	/** Overridden from SMultiColumnTableRow.  Generates a widget for this column of the tree row. */
+	virtual TSharedRef<SWidget> GenerateWidgetForColumn( const FName& ColumnName ) override;
+
+private:
+
+	/* Returns the reference mesh this pose is based on	*/
+	FString GetReferenceMeshName() { return Item->ReferenceMesh->GetPathName(); }
+
+	/* The SRetargetSourceWindow that handles all retarget sources*/
+	SRetargetSourceWindow* RetargetSourceWindow;
+
+	/** Widget used to display the list of retarget sources*/
+	TSharedPtr<SRetargetSourceListType> RetargetSourceListView;
+
+	/** The name and weight of the retarget source*/
+	FDisplayedRetargetSourceInfoPtr	Item;
+
+	/** Pointer back to the Persona that owns us */
+	TWeakPtr<FPersona> PersonaPtr;
+
+	/** Delegate for when an asset name has been entered for an item that is in a rename state */
+	FOnRenameCommit OnRenameCommit;
+
+	/** Delegate for when an asset name has been entered for an item to verify the name before commit */
+	FOnVerifyRenameCommit OnVerifyRenameCommit;
+
+	/** Handles committing a name change */
+	virtual void HandleNameCommitted( const FText& NewText, ETextCommit::Type CommitInfo );
+
+	/** Handles verifying a name change */
+	virtual bool HandleVerifyNameChanged( const FText& NewText, FText& OutErrorMessage );
+};
+
+void SRetargetSourceListRow::Construct( const FArguments& InArgs, const TSharedRef<STableViewBase>& InOwnerTableView )
+{
+	Item = InArgs._Item;
+	RetargetSourceWindow = InArgs._RetargetSourceWindow;
+	RetargetSourceListView = InArgs._RetargetSourceListView;
+	OnRenameCommit = InArgs._OnRenameCommit;
+	OnVerifyRenameCommit = InArgs._OnVerifyRenameCommit;
+
+	PersonaPtr = InArgs._Persona;
+
+	check( Item.IsValid() );
+
+	SMultiColumnTableRow< FDisplayedRetargetSourceInfoPtr >::Construct( FSuperRowType::FArguments(), InOwnerTableView );
+}
+
+TSharedRef< SWidget > SRetargetSourceListRow::GenerateWidgetForColumn( const FName& ColumnName )
+{
+	if ( ColumnName == ColumnId_RetargetSourceNameLabel )
+	{
+		TSharedPtr< SInlineEditableTextBlock > InlineWidget;
+		TSharedRef< SWidget > NewWidget = 
+			SNew( SVerticalBox )
+
+			+ SVerticalBox::Slot()
+			.AutoHeight()
+			.Padding( 0.0f, 4.0f )
+			.VAlign( VAlign_Center )
+			[
+				SAssignNew(InlineWidget, SInlineEditableTextBlock)
+				.Text( FText::FromString(Item->Name.ToString()) )
+				.OnTextCommitted(this, &SRetargetSourceListRow::HandleNameCommitted)
+				.OnVerifyTextChanged(this, &SRetargetSourceListRow::HandleVerifyNameChanged)
+				.HighlightText( RetargetSourceWindow->GetFilterText() )
+				.IsReadOnly(false)
+				.IsSelected(this, &SMultiColumnTableRow< FDisplayedRetargetSourceInfoPtr >::IsSelectedExclusively)
+			];
+
+		Item->OnEnterEditingMode.BindSP( InlineWidget.Get(), &SInlineEditableTextBlock::EnterEditingMode);
+
+		return NewWidget;
+	}
+	else
+	{
+		// Encase the SSpinbox in an SVertical box so we can apply padding. Setting ItemHeight on the containing SListView has no effect :-(
+		return
+			SNew( SVerticalBox )
+
+			+ SVerticalBox::Slot()
+			.AutoHeight()
+			.Padding( 0.0f, 1.0f )
+			.VAlign( VAlign_Center )
+			[
+				SNew( STextBlock )
+				.Text( Item->GetReferenceMeshName() )
+				.HighlightText( RetargetSourceWindow->GetFilterText() )
+			];
+	}
+}
+
+/** Handles committing a name change */
+void SRetargetSourceListRow::HandleNameCommitted( const FText& NewText, ETextCommit::Type CommitInfo )
+{
+	OnRenameCommit.ExecuteIfBound(Item->Name, NewText.ToString());
+}
+
+/** Handles verifying a name change */
+bool SRetargetSourceListRow::HandleVerifyNameChanged( const FText& NewText, FText& OutErrorMessage )
+{
+	return !OnVerifyRenameCommit.IsBound() || OnVerifyRenameCommit.Execute(Item->Name, NewText.ToString(), OutErrorMessage);
+}
+//////////////////////////////////////////////////////////////////////////
+// SRetargetSourceWindow
+
+void SRetargetSourceWindow::Construct(const FArguments& InArgs)
+{
+	PersonaPtr = InArgs._Persona;
+	Skeleton = NULL;
+
+	if ( PersonaPtr.IsValid() )
+	{
+		Skeleton = PersonaPtr.Pin()->GetSkeleton();
+		PersonaPtr.Pin()->RegisterOnPostUndo(FPersona::FOnPostUndo::CreateSP( this, &SRetargetSourceWindow::PostUndo ) );
+	}
+
+	FString SkeletonName = Skeleton ? Skeleton->GetName() : LOCTEXT( "RetargetSourceMeshNameLabel", "No Skeleton Present" ).ToString();
+
+	ChildSlot
+	[
+		SNew( SVerticalBox )
+		
+		+ SVerticalBox::Slot()
+		.AutoHeight()
+		[
+			SNew( SHorizontalBox )
+			+SHorizontalBox::Slot()
+			.AutoWidth()
+			[
+				SNew(STextBlock)
+				.Text(LOCTEXT("Skeleton_Label", "Current Skeleton "))
+				.Font(FEditorStyle::GetFontStyle("Persona.RetargetManager.BoldFont"))
+			]
+
+			+SHorizontalBox::Slot()
+			.HAlign(HAlign_Left)
+			[
+				SNew(STextBlock)
+				.Text(SkeletonName)
+			]
+
+			+SHorizontalBox::Slot()
+			.HAlign(HAlign_Right)
+			[
+				SNew(SButton)
+				.OnClicked(FOnClicked::CreateSP(this, &SRetargetSourceWindow::OnAddRetargetSourceButtonClicked))
+				.HAlign(HAlign_Center)
+				.VAlign(VAlign_Center)
+				.Text(LOCTEXT("AddRetargetSourceButton_Label", "Add New Retarget Source"))
+			]
+		]
+
+		+ SVerticalBox::Slot()
+		.AutoHeight()
+		.Padding(0,2)
+		[
+			SNew(SHorizontalBox)
+			// Filter entry
+			+SHorizontalBox::Slot()
+			.FillWidth( 1 )
+			[
+				SAssignNew( NameFilterBox, SSearchBox )
+				.SelectAllTextWhenFocused( true )
+				.OnTextChanged( this, &SRetargetSourceWindow::OnFilterTextChanged )
+				.OnTextCommitted( this, &SRetargetSourceWindow::OnFilterTextCommitted )
+			]
+		]
+
+		+ SVerticalBox::Slot()
+		.FillHeight( 1.0f )		// This is required to make the scrollbar work, as content overflows Slate containers by default
+		[
+			SAssignNew( RetargetSourceListView, SRetargetSourceListType )
+			.ListItemsSource( &RetargetSourceList )
+			.OnGenerateRow( this, &SRetargetSourceWindow::GenerateRetargetSourceRow )
+			.OnContextMenuOpening( this, &SRetargetSourceWindow::OnGetContextMenuContent )
+			.ItemHeight( 22.0f )
+			.HeaderRow
+			(
+				SNew( SHeaderRow )
+				+ SHeaderRow::Column( ColumnId_RetargetSourceNameLabel )
+				.DefaultLabel( LOCTEXT( "RetargetSourceNameLabel", "Retarget Source Name" ).ToString() )
+
+				+ SHeaderRow::Column( ColumnID_BaseReferenceMeshLabel )
+				.DefaultLabel( LOCTEXT( "RetargetSourceWeightLabel", "Source Mesh" ).ToString() )
+			)
+		]
+	];
+
+	CreateRetargetSourceList();
+}
+
+void SRetargetSourceWindow::OnFilterTextChanged( const FText& SearchText )
+{
+	FilterText = SearchText;
+
+	CreateRetargetSourceList( SearchText.ToString() );
+}
+
+void SRetargetSourceWindow::OnFilterTextCommitted( const FText& SearchText, ETextCommit::Type CommitInfo )
+{
+	// Just do the same as if the user typed in the box
+	OnFilterTextChanged( SearchText );
+}
+
+TSharedRef<ITableRow> SRetargetSourceWindow::GenerateRetargetSourceRow(TSharedPtr<FDisplayedRetargetSourceInfo> InInfo, const TSharedRef<STableViewBase>& OwnerTable)
+{
+	check( InInfo.IsValid() );
+
+	return
+		SNew( SRetargetSourceListRow, OwnerTable )
+		.Persona( PersonaPtr )
+		.Item( InInfo )
+		.RetargetSourceWindow( this )
+		.RetargetSourceListView( RetargetSourceListView )
+		.OnRenameCommit(this, &SRetargetSourceWindow::OnRenameCommit)
+		.OnVerifyRenameCommit(this, &SRetargetSourceWindow::OnVerifyRenameCommit);
+}
+
+
+void SRetargetSourceWindow::OnRenameCommit( const FName & OldName,  const FString& NewName )
+{
+	FString NewNameString = NewName;
+	if (OldName != FName(*NewNameString.Trim().TrimTrailing()))
+	{
+		const FScopedTransaction Transaction( LOCTEXT("RetargetSourceWindow_Rename", "Rename Retarget Source") );
+
+		FReferencePose * Pose = Skeleton->AnimRetargetSources.Find(OldName);
+		if (Pose)
+		{
+			FReferencePose NewPose = *Pose;
+			FName NewPoseName = FName(*NewName);
+			NewPose.PoseName = NewPoseName;
+
+			Skeleton->Modify();
+
+			Skeleton->AnimRetargetSources.Remove(OldName);
+			Skeleton->AnimRetargetSources.Add(NewPoseName, NewPose);
+
+			// need to verify if these pose is used by anybody else
+			FAssetRegistryModule& AssetRegistryModule = FModuleManager::LoadModuleChecked<FAssetRegistryModule>(TEXT("AssetRegistry"));
+
+			TArray<FAssetData> AssetList;
+			TMultiMap<FName, FString> TagsAndValues;
+			TagsAndValues.Add(GET_MEMBER_NAME_CHECKED(UAnimSequence, RetargetSource), OldName.ToString());
+			AssetRegistryModule.Get().GetAssetsByTagValues(TagsAndValues, AssetList);
+
+			// ask users if they'd like to continue and/or fix up
+			if (AssetList.Num() > 0 )
+			{
+				FString ListOfAssets;
+				// if users is sure to delete, delete
+				for (auto Iter=AssetList.CreateConstIterator(); Iter; ++Iter)
+				{
+					ListOfAssets += Iter->AssetName.ToString();
+					ListOfAssets += TEXT("\n");
+				}
+
+
+				// if so, ask if they'd like us to fix the animations as well
+				FString Message = NSLOCTEXT("RetargetSourceEditor", "RetargetSourceRename_FixUpAnimation_Message", "Would you like to fix up the following animation(s)? You'll have to save all the assets in the list.").ToString();
+				Message += TEXT("\n");
+				Message += ListOfAssets;
+
+				EAppReturnType::Type Result = FMessageDialog::Open( EAppMsgType::YesNo, FText::FromString(Message) );
+
+				if (Result == EAppReturnType::Yes)
+				{
+					// now fix up all assets
+					TArray<UObject*> ObjectsToUpdate;
+					for (auto Iter=AssetList.CreateConstIterator(); Iter; ++Iter)
+					{
+						UAnimSequence * AnimSequence = Cast<UAnimSequence>(Iter->GetAsset());
+						if (AnimSequence)
+						{
+							ObjectsToUpdate.Add(AnimSequence);
+
+							AnimSequence->Modify();
+							// clear name
+							AnimSequence->RetargetSource = NewPoseName;
+						}
+					}
+				}
+			}
+
+			Skeleton->CallbackRetargetSourceChanged();
+			FAssetNotifications::SkeletonNeedsToBeSaved(Skeleton);
+			CreateRetargetSourceList( NameFilterBox->GetText().ToString() );
+		}
+	}
+}
+
+bool SRetargetSourceWindow::OnVerifyRenameCommit( const FName & OldName, const FString& NewName, FText& OutErrorMessage)
+{
+	// if same reject
+	FString NewString = NewName;
+
+	if (OldName == FName(*NewString.Trim().TrimTrailing()))
+	{
+		OutErrorMessage = FText::Format (LOCTEXT("RetargetSourceWindowNameSame", "{0} Nothing modified"), FText::FromString(OldName.ToString()) );
+		return false;
+	}
+
+	FReferencePose * Pose = Skeleton->AnimRetargetSources.Find(OldName);
+	if (!Pose)
+	{
+		OutErrorMessage = FText::Format (LOCTEXT("RetargetSourceWindowNameNotFound", "{0} Not found"), FText::FromString(OldName.ToString()) );
+		return false;
+	}
+
+	Pose = Skeleton->AnimRetargetSources.Find(FName(*NewName));
+
+	if (Pose)
+	{
+		OutErrorMessage = FText::Format (LOCTEXT("RetargetSourceWindowNameDuplicated", "{0} already exists"), FText::FromString(NewName) );
+		return false;
+	}
+
+	return true;
+}
+
+TSharedPtr<SWidget> SRetargetSourceWindow::OnGetContextMenuContent() const
+{
+	if ( Skeleton )
+	{
+		const bool bShouldCloseWindowAfterMenuSelection = true;
+		FMenuBuilder MenuBuilder( bShouldCloseWindowAfterMenuSelection, NULL);
+
+		MenuBuilder.BeginSection("RetargetSourceAction", LOCTEXT( "New", "New" ) );
+		{
+			FUIAction Action = FUIAction( FExecuteAction::CreateSP( this, &SRetargetSourceWindow::OnAddRetargetSource ) );
+			const FText Label = LOCTEXT("AddRetargetSourceActionLabel", "Add...");
+			const FText ToolTip = LOCTEXT("AddRetargetSourceActionTooltip", "Add new retarget source.");
+			MenuBuilder.AddMenuEntry( Label, ToolTip, FSlateIcon(), Action);
+		}
+		MenuBuilder.EndSection();
+
+		MenuBuilder.BeginSection("RetargetSourceAction", LOCTEXT( "Selected", "Selected Item Actions" ) );
+		{
+			FUIAction Action = FUIAction( FExecuteAction::CreateSP( this, &SRetargetSourceWindow::OnRenameRetargetSource ), 
+				FCanExecuteAction::CreateSP( this, &SRetargetSourceWindow::CanPerformRename ) );
+			const FText Label = LOCTEXT("RenameRetargetSourceActionLabel", "Rename");
+			const FText ToolTip = LOCTEXT("RenameRetargetSourceActionTooltip", "Rename the selected retarget source.");
+			MenuBuilder.AddMenuEntry( Label, ToolTip, FSlateIcon(), Action);
+		}
+		{
+			FUIAction Action = FUIAction( FExecuteAction::CreateSP( this, &SRetargetSourceWindow::OnDeleteRetargetSource ), 
+				FCanExecuteAction::CreateSP( this, &SRetargetSourceWindow::CanPerformDelete ) );
+			const FText Label = LOCTEXT("DeleteRetargetSourceActionLabel", "Delete");
+			const FText ToolTip = LOCTEXT("DeleteRetargetSourceActionTooltip", "Deletes the selected retarget sources.");
+			MenuBuilder.AddMenuEntry( Label, ToolTip, FSlateIcon(), Action);
+		}
+		{
+			FUIAction Action = FUIAction( FExecuteAction::CreateSP( this, &SRetargetSourceWindow::OnRefreshRetargetSource ), 
+				FCanExecuteAction::CreateSP( this, &SRetargetSourceWindow::CanPerformRefresh ) );
+			const FText Label = LOCTEXT("RefreshRetargetSourceActionLabel", "Refresh");
+			const FText ToolTip = LOCTEXT("RefreshRetargetSourceActionTooltip", "Refreshs the selected retarget sources.");
+			MenuBuilder.AddMenuEntry( Label, ToolTip, FSlateIcon(), Action);
+		}
+		MenuBuilder.EndSection();
+
+		return MenuBuilder.MakeWidget();
+	}
+
+	return TSharedPtr<SWidget>();
+
+
+}
+
+/** @TODO: FIXME: Item to rename. Only valid for adding and this is so hacky, I know**/
+//TSharedPtr<FDisplayedBasePoseInfo> ItemToRename;
+
+void SRetargetSourceWindow::CreateRetargetSourceList( const FString& SearchText, const FName  NewName )
+{
+	RetargetSourceList.Empty();
+
+	if ( Skeleton )
+	{
+		bool bDoFiltering = !SearchText.IsEmpty();
+
+		for (auto Iter=Skeleton->AnimRetargetSources.CreateConstIterator(); Iter; ++Iter)
+		{
+			const FName & Name = Iter.Key();
+			const FReferencePose & RefPose = Iter.Value();
+
+			if ( bDoFiltering )
+			{
+				if (!Name.ToString().Contains( SearchText ))
+				{
+					continue; // Skip items that don't match our filter
+				}
+				if (RefPose.ReferenceMesh && RefPose.ReferenceMesh->GetPathName().Contains( SearchText ))
+				{
+					continue;
+				}
+			}
+
+			TSharedRef<FDisplayedRetargetSourceInfo> Info = FDisplayedRetargetSourceInfo::Make( Name, RefPose.ReferenceMesh );
+
+			if (Name == NewName)
+			{
+				ItemToRename = Info;
+			}
+
+			RetargetSourceList.Add( Info );
+		}
+	}
+
+	RetargetSourceListView->RequestListRefresh();
+}
+
+int32 SRetargetSourceWindow::OnPaint( const FPaintArgs& Args, const FGeometry& AllottedGeometry, const FSlateRect& MyClippingRect, FSlateWindowElementList& OutDrawElements, int32 LayerId, const FWidgetStyle& InWidgetStyle, bool bParentEnabled ) const
+{
+	const int32 Result = SCompoundWidget::OnPaint( Args, AllottedGeometry, MyClippingRect, OutDrawElements, LayerId, InWidgetStyle, bParentEnabled );
+
+	// I need to do this after first paint	
+	if (ItemToRename.IsValid())
+	{
+		ItemToRename->RequestRename();
+		ItemToRename = NULL;
+	}
+
+	return Result;
+}
+
+void SRetargetSourceWindow::AddRetargetSource( const FName Name, USkeletalMesh * ReferenceMesh  )
+{
+	// need to verify if the name is unique, if not create unique name
+	if ( Skeleton )
+	{
+		int32 IntSuffix = 1;
+		FReferencePose * ExistingPose = NULL;
+		FString NewSourceName = TEXT("");
+		do
+		{
+			if ( IntSuffix <= 1 )
+			{
+				NewSourceName = FString::Printf( TEXT("%s"), *Name.ToString());
+			}
+			else
+			{
+				NewSourceName = FString::Printf( TEXT("%s%d"), *Name.ToString(), IntSuffix );
+			}
+
+			ExistingPose = Skeleton->AnimRetargetSources.Find(FName(*NewSourceName));
+			IntSuffix++;
+		}
+		while (ExistingPose != NULL);
+
+		// add new one
+		// remap to skeleton refpose
+		// we have to do this whenever skeleton changes
+		FReferencePose RefPose;
+		RefPose.PoseName = FName(*NewSourceName);
+		RefPose.ReferenceMesh = ReferenceMesh;
+
+		const FScopedTransaction Transaction( LOCTEXT("RetargetSourceWindow_Add", "Add New Retarget Source") );
+		Skeleton->Modify();
+
+		Skeleton->AnimRetargetSources.Add(Name, RefPose);
+		// ask skeleton to update retarget source for the given name
+		Skeleton->UpdateRetargetSource(Name);
+
+		Skeleton->CallbackRetargetSourceChanged();
+		FAssetNotifications::SkeletonNeedsToBeSaved(Skeleton);
+
+		// clear search filter
+		NameFilterBox->SetText(FText::GetEmpty());
+		CreateRetargetSourceList( NameFilterBox->GetText().ToString(), Name );
+	}
+}
+
+void SRetargetSourceWindow::OnAddRetargetSource()
+{
+	// show list of skeletalmeshes that they can choose from
+	FContentBrowserModule& ContentBrowserModule = FModuleManager::Get().LoadModuleChecked<FContentBrowserModule>(TEXT("ContentBrowser"));
+
+	FAssetPickerConfig AssetPickerConfig;
+	AssetPickerConfig.Filter.ClassNames.Add(USkeletalMesh::StaticClass()->GetFName());
+	AssetPickerConfig.OnAssetSelected = FOnAssetSelected::CreateSP(this, &SRetargetSourceWindow::OnAssetSelectedFromMeshPicker);
+	AssetPickerConfig.bAllowNullSelection = false;
+	AssetPickerConfig.InitialAssetViewType = EAssetViewType::Tile;
+	AssetPickerConfig.ThumbnailScale = 0.0f;
+
+	if(Skeleton)
+	{
+		FString SkeletonString = FAssetData(Skeleton).GetExportTextName();
+		AssetPickerConfig.Filter.TagsAndValues.Add(GET_MEMBER_NAME_CHECKED(USkeletalMesh, Skeleton), SkeletonString);
+	}
+
+	TSharedRef<SWidget> Widget = SNew(SBox)
+		.WidthOverride(384)
+		.HeightOverride(768)
+		[
+			SNew(SBorder)
+			.BorderBackgroundColor(FLinearColor(0.25f, 0.25f, 0.25f, 1.f))
+			.Padding( 2 )
+			[
+				SNew(SBorder)
+				.BorderImage( FEditorStyle::GetBrush("ToolPanel.GroupBorder") )
+				.Padding( 8 )
+				[
+					ContentBrowserModule.Get().CreateAssetPicker(AssetPickerConfig)
+				]
+			]
+		];
+
+	FSlateApplication::Get().PushMenu(
+		AsShared(),
+		Widget,
+		FSlateApplication::Get().GetCursorPos(),
+		FPopupTransitionEffect( FPopupTransitionEffect::TopMenu )
+		);
+}
+
+void SRetargetSourceWindow::OnAssetSelectedFromMeshPicker(const FAssetData& AssetData)
+{
+	USkeletalMesh * SelectedMesh = CastChecked<USkeletalMesh> (AssetData.GetAsset());
+
+	// make sure you don't have any more retarget sources from the same mesh
+	for (auto Iter=Skeleton->AnimRetargetSources.CreateConstIterator(); Iter; ++Iter)
+	{
+		const FName & Name = Iter.Key();
+		const FReferencePose & RefPose = Iter.Value();
+
+		if (RefPose.ReferenceMesh == SelectedMesh)
+		{
+			// redundant source exists
+			
+			FSlateApplication::Get().DismissAllMenus();
+			return;
+		}
+	}
+
+	// give temporary name, and make it editable first time
+	AddRetargetSource(FName(*SelectedMesh->GetName()), SelectedMesh);
+	FSlateApplication::Get().DismissAllMenus();
+}
+
+bool SRetargetSourceWindow::CanPerformDelete() const
+{
+	TArray< TSharedPtr< FDisplayedRetargetSourceInfo > > SelectedRows = RetargetSourceListView->GetSelectedItems();
+	return SelectedRows.Num() > 0;
+}
+
+void SRetargetSourceWindow::OnDeleteRetargetSource()
+{
+	TArray< TSharedPtr< FDisplayedRetargetSourceInfo > > SelectedRows = RetargetSourceListView->GetSelectedItems();
+	
+	const FScopedTransaction Transaction( LOCTEXT("RetargetSourceWindow_Delete", "Delete Retarget Source") );
+	for (int RowIndex = 0; RowIndex < SelectedRows.Num(); ++RowIndex)
+	{
+		const FReferencePose * PoseFound = Skeleton->AnimRetargetSources.Find(SelectedRows[RowIndex]->Name);
+		if(PoseFound)
+		{
+			// need to verify if these pose is used by anybody else
+			FAssetRegistryModule& AssetRegistryModule = FModuleManager::LoadModuleChecked<FAssetRegistryModule>(TEXT("AssetRegistry"));
+
+			TArray<FAssetData> AssetList;
+			TMultiMap<FName, FString> TagsAndValues;
+			TagsAndValues.Add(GET_MEMBER_NAME_CHECKED(UAnimSequence, RetargetSource), PoseFound->PoseName.ToString());
+			AssetRegistryModule.Get().GetAssetsByTagValues(TagsAndValues, AssetList);
+
+			// ask users if they'd like to continue and/or fix up
+			if (AssetList.Num() > 0 )
+			{
+				FString ListOfAssets;
+				// if users is sure to delete, delete
+				for (auto Iter=AssetList.CreateConstIterator(); Iter; ++Iter)
+				{
+					ListOfAssets += Iter->AssetName.ToString();
+					ListOfAssets += TEXT("\n");
+				}
+
+				// ask if they'd like to continue deleting this pose regardless animation references
+				FString Message = NSLOCTEXT("RetargetSourceEditor", "RetargetSourceDeleteMessage", "Following animation(s) is(are) referencing this pose. Are you sure if you'd like to delete this pose?").ToString();
+				Message += TEXT("\n\n");
+				Message += ListOfAssets;
+
+				EAppReturnType::Type Result = FMessageDialog::Open( EAppMsgType::YesNo, FText::FromString(Message) );
+
+				if (Result == EAppReturnType::No)
+				{
+					continue;
+				}
+
+				// if so, ask if they'd like us to fix the animations as well
+				Message = NSLOCTEXT("RetargetSourceEditor", "RetargetSourceDelete_FixUpAnimation_Message", "Would you like to fix up the following animation(s)? You'll have to save all the assets in the list.").ToString();
+				Message += TEXT("\n");
+				Message += ListOfAssets;
+
+				Result = FMessageDialog::Open( EAppMsgType::YesNo, FText::FromString(Message) );
+
+				if (Result == EAppReturnType::No)
+				{
+					continue;
+				}
+
+				// now fix up all assets
+				TArray<UObject *> ObjectsToUpdate;
+				for (auto Iter=AssetList.CreateConstIterator(); Iter; ++Iter)
+				{
+					UAnimSequence * AnimSequence = Cast<UAnimSequence>(Iter->GetAsset());
+					if (AnimSequence)
+					{
+						ObjectsToUpdate.Add(AnimSequence);
+
+						AnimSequence->Modify();
+						// clear name
+						AnimSequence->RetargetSource = NAME_None;
+					}
+				}
+			}
+
+			Skeleton->Modify();
+			// delete now
+			Skeleton->AnimRetargetSources.Remove(SelectedRows[RowIndex]->Name);
+			Skeleton->CallbackRetargetSourceChanged();
+			FAssetNotifications::SkeletonNeedsToBeSaved(Skeleton);
+		}
+	}
+
+	CreateRetargetSourceList( NameFilterBox->GetText().ToString() );
+}
+
+bool SRetargetSourceWindow::CanPerformRename() const
+{
+	TArray< TSharedPtr< FDisplayedRetargetSourceInfo > > SelectedRows = RetargetSourceListView->GetSelectedItems();
+	return SelectedRows.Num() == 1;
+}
+
+void SRetargetSourceWindow::OnRenameRetargetSource()
+{
+	TArray< TSharedPtr< FDisplayedRetargetSourceInfo > > SelectedRows = RetargetSourceListView->GetSelectedItems();
+
+	if (ensure (SelectedRows.Num() == 1))
+	{
+		int32 RowIndex = 0;
+		const FReferencePose * PoseFound = Skeleton->AnimRetargetSources.Find(SelectedRows[RowIndex]->Name);
+		if(PoseFound)
+		{
+			// we used to verify if there is any animation referencing and warn them, but that doesn't really help
+			// because you can rename by just double click as well, and it slows process, so removed it
+			// request rename
+			SelectedRows[RowIndex]->RequestRename();
+		}
+	}
+}
+
+bool SRetargetSourceWindow::CanPerformRefresh() const
+{
+	TArray< TSharedPtr< FDisplayedRetargetSourceInfo > > SelectedRows = RetargetSourceListView->GetSelectedItems();
+	return SelectedRows.Num() > 0;
+}
+
+void SRetargetSourceWindow::OnRefreshRetargetSource()
+{
+	TArray< TSharedPtr< FDisplayedRetargetSourceInfo > > SelectedRows = RetargetSourceListView->GetSelectedItems();
+
+	const FScopedTransaction Transaction( LOCTEXT("RetargetSourceWindow_Refresh", "Refresh Retarget Source") );
+	for (int RowIndex = 0; RowIndex < SelectedRows.Num(); ++RowIndex)
+	{
+		Skeleton->Modify();
+		// refresh pose now
+		FName RetargetSourceName = SelectedRows[RowIndex]->Name;
+		Skeleton->UpdateRetargetSource(RetargetSourceName);
+
+		// feedback of pose has been updated
+		FFormatNamedArguments Args;
+		Args.Add( TEXT("RetargetSourceName"), FText::FromString( RetargetSourceName.ToString() ) );
+		FNotificationInfo Info( FText::Format( LOCTEXT("RetargetSourceWindow_RefreshPose", "Retarget Source {RetargetSourceName} is refreshed"), Args ) );
+		Info.ExpireDuration = 5.0f;
+		Info.bUseLargeFont = false;
+		TSharedPtr<SNotificationItem> Notification = FSlateNotificationManager::Get().AddNotification(Info);
+		if ( Notification.IsValid() )
+		{
+			Notification->SetCompletionState( SNotificationItem::CS_None );
+		}
+	}
+
+	FAssetNotifications::SkeletonNeedsToBeSaved(Skeleton);
+}
+
+SRetargetSourceWindow::~SRetargetSourceWindow()
+{
+	if (PersonaPtr.IsValid())
+	{
+		PersonaPtr.Pin()->UnregisterOnPostUndo(this);
+	}
+}
+
+void SRetargetSourceWindow::PostUndo()
+{
+	CreateRetargetSourceList();
+}
+
+FReply SRetargetSourceWindow::OnAddRetargetSourceButtonClicked()
+{
+	OnAddRetargetSource();
+	return FReply::Handled();
+}
+#undef LOCTEXT_NAMESPACE
+
diff --git a/Engine/Source/Editor/Persona/Private/SRetargetSourceWindow.h b/Engine/Source/Editor/Persona/Private/SRetargetSourceWindow.h
new file mode 100644
--- /dev/null
+++ b/Engine/Source/Editor/Persona/Private/SRetargetSourceWindow.h
@@ -0,0 +1,215 @@
+// Copyright 1998-2014 Epic Games, Inc. All Rights Reserved.
+
+#pragma once
+
+#include "Persona.h"
+
+//////////////////////////////////////////////////////////////////////////
+// FDisplayedRetargetSourceInfo
+
+class FDisplayedRetargetSourceInfo
+{
+public:
+	FName Name;
+	USkeletalMesh * ReferenceMesh;
+
+	/** Static function for creating a new item, but ensures that you can only have a TSharedRef to one */
+	static TSharedRef<FDisplayedRetargetSourceInfo> Make(const FName& SourcePose, USkeletalMesh * ReferenceMesh)
+	{
+		return MakeShareable(new FDisplayedRetargetSourceInfo(SourcePose, ReferenceMesh));
+	}
+
+	FString GetReferenceMeshName()
+	{
+		return (ReferenceMesh)? ReferenceMesh->GetPathName() : TEXT("None");
+	}
+
+	/** Requests a rename on the socket item */
+	void RequestRename()
+	{
+		OnEnterEditingMode.ExecuteIfBound();
+	}
+
+	/** Delegate for when the context menu requests a rename */
+	DECLARE_DELEGATE(FOnEnterEditingMode);
+	FOnEnterEditingMode OnEnterEditingMode;
+
+protected:
+	/** Hidden constructor, always use Make above */
+	FDisplayedRetargetSourceInfo(const FName& InSource, USkeletalMesh * RefMesh)
+		: Name( InSource )
+		, ReferenceMesh( RefMesh )
+		, bRenamePending( false )
+	{}
+
+	/** Hidden constructor, always use Make above */
+	FDisplayedRetargetSourceInfo() {}
+
+	bool bRenamePending;
+};
+
+typedef SListView< TSharedPtr<FDisplayedRetargetSourceInfo> > SRetargetSourceListType;
+
+//////////////////////////////////////////////////////////////////////////
+// SRetargetSourceWindow
+
+class SRetargetSourceWindow : public SCompoundWidget
+{
+public:
+	SLATE_BEGIN_ARGS( SRetargetSourceWindow )
+		: _Persona()
+	{}
+		
+		/* The Persona that owns this table */
+		SLATE_ARGUMENT( TWeakPtr< FPersona >, Persona )
+
+	SLATE_END_ARGS()
+
+	/**
+	* Slate construction function
+	*
+	* @param InArgs - Arguments passed from Slate
+	*
+	*/
+	void Construct( const FArguments& InArgs );
+
+	/**
+	* Destructor - resets the morph targets
+	*
+	*/
+	virtual ~SRetargetSourceWindow();
+
+	/**
+	* Filters the SListView when the user changes the search text box (NameFilterBox)
+	*
+	* @param SearchText - The text the user has typed
+	*
+	*/
+	void OnFilterTextChanged( const FText& SearchText );
+
+	/**
+	* Filters the SListView when the user hits enter or clears the search box
+	* Simply calls OnFilterTextChanged
+	*
+	* @param SearchText - The text the user has typed
+	* @param CommitInfo - Not used
+	*
+	*/
+	void OnFilterTextCommitted( const FText& SearchText, ETextCommit::Type CommitInfo );
+
+	/**
+	* Create a widget for an entry in the tree from an info
+	*
+	* @param InInfo - Shared pointer to the morph target we're generating a row for
+	* @param OwnerTable - The table that owns this row
+	*
+	* @return A new Slate widget, containing the UI for this row
+	*/
+	TSharedRef<ITableRow> GenerateRetargetSourceRow(TSharedPtr<FDisplayedRetargetSourceInfo> InInfo, const TSharedRef<STableViewBase>& OwnerTable);
+
+	/**
+	* Handler for adding new retarget source. It displays the asset picker
+	*/
+	void OnAddRetargetSource();
+
+	/**
+	 * Callback for asset picker
+	 */
+	void OnAssetSelectedFromMeshPicker(const FAssetData& AssetData);
+
+	/**
+	* Adds a new retarget source with given name
+	*
+	* @param Name - Name of retarget source
+	* @param ReferenceMesh	- Reference Mesh this pose is based on
+	*
+	*/
+	void AddRetargetSource( const FName Name, USkeletalMesh * ReferenceMesh  );
+
+	/**
+	* Return true if it can delete
+	*
+	*/
+	bool CanPerformDelete() const;
+
+	/**
+	* Handler for the delete of retarget source
+	*/
+	void OnDeleteRetargetSource();
+
+	/**
+	* Return true if it can perform rename
+	*
+	*/
+	bool CanPerformRename() const;
+
+	/**
+	* Handler for the rename of retarget source
+	*/
+	void OnRenameRetargetSource();
+
+	/**
+	* Return true if it can delete
+	*
+	*/
+	bool CanPerformRefresh() const;
+
+	/**
+	* Handler for the delete of retarget source
+	*/
+	void OnRefreshRetargetSource();
+
+	/**
+	* Accessor so our rows can grab the filtertext for highlighting
+	*
+	*/
+	FText& GetFilterText() { return FilterText; }
+
+	// SWidget interface
+	virtual int32 OnPaint( const FPaintArgs& Args, const FGeometry& AllottedGeometry, const FSlateRect& MyClippingRect, FSlateWindowElementList& OutDrawElements, int32 LayerId, const FWidgetStyle& InWidgetStyle, bool bParentEnabled ) const override;
+	// end of SWidget interface
+private:
+
+	/** Handler for context menus */
+	TSharedPtr<SWidget> OnGetContextMenuContent() const;
+
+	/**
+	* Clears and rebuilds the table, according to an optional search string
+	*
+	* @param SearchText - Optional search string
+	*
+	*/
+	void CreateRetargetSourceList( const FString& SearchText = FString(""), const FName  NewName = NAME_None );
+
+	/** Pointer back to the Persona that owns us */
+	TWeakPtr<FPersona> PersonaPtr;
+
+	/** Box to filter to a specific morph target name */
+	TSharedPtr<SSearchBox>	NameFilterBox;
+
+	/** Widget used to display the list of retarget sources */
+	TSharedPtr<SRetargetSourceListType> RetargetSourceListView;
+
+	/** A list of retarget sources. Used by the RetargetSourceListView. */
+	TArray< TSharedPtr<FDisplayedRetargetSourceInfo> > RetargetSourceList;
+
+	/** The skeletal mesh that we grab the morph targets from */
+	USkeleton* Skeleton;
+
+	/** Current text typed into NameFilterBox */
+	FText FilterText;
+
+	/** Item to rename. Only valid for adding **/
+	mutable TSharedPtr<FDisplayedRetargetSourceInfo> ItemToRename;
+
+	/** Rename committed. Called by Inline Widget **/
+	void OnRenameCommit( const FName & OldName,  const FString& NewName );
+	/** Verify Rename is legit or not **/
+	bool OnVerifyRenameCommit( const FName & OldName, const FString& NewName, FText& OutErrorMessage);
+
+	/** Delegate for undo/redo transaction **/
+	void PostUndo();
+
+	/** Button handler **/
+	FReply OnAddRetargetSourceButtonClicked();
+};
\ No newline at end of file
diff --git a/Engine/Source/Editor/Persona/Private/SRigWindow.cpp b/Engine/Source/Editor/Persona/Private/SRigWindow.cpp
new file mode 100644
--- /dev/null
+++ b/Engine/Source/Editor/Persona/Private/SRigWindow.cpp
@@ -0,0 +1,383 @@
+// Copyright 1998-2014 Epic Games, Inc. All Rights Reserved.
+
+
+#include "PersonaPrivatePCH.h"
+#include "SRigWindow.h"
+#include "ObjectTools.h"
+#include "ScopedTransaction.h"
+#include "AssetRegistryModule.h"
+#include "Editor/PropertyEditor/Public/PropertyEditorModule.h"
+#include "Editor/ContentBrowser/Public/ContentBrowserModule.h"
+#include "AssetNotifications.h"
+#include "Animation/Rig.h"
+#include "BoneSelectionWidget.h"
+
+#define LOCTEXT_NAMESPACE "SRigWindow"
+
+static const FName ColumnId_NodeNameLabel( "Node Name" );
+static const FName ColumnID_BoneNameLabel( "Bone" );
+
+DECLARE_DELEGATE_TwoParams(FOnBoneMappingChanged, FName /** NodeName */, FName /** BoneName **/);
+DECLARE_DELEGATE_RetVal_OneParam(FName, FOnGetBoneMapping, FName /** Node Name **/);
+
+//////////////////////////////////////////////////////////////////////////
+// SBoneMappingListRow
+
+typedef TSharedPtr< FDisplayedBoneMappingInfo > FDisplayedBoneMappingInfoPtr;
+
+class SBoneMappingListRow
+	: public SMultiColumnTableRow< FDisplayedBoneMappingInfoPtr >
+{
+public:
+
+	SLATE_BEGIN_ARGS( SBoneMappingListRow ) {}
+
+		/** The item for this row **/
+		SLATE_ARGUMENT( FDisplayedBoneMappingInfoPtr, Item )
+
+		/* The SRigWindow that handles all retarget sources */
+		SLATE_ARGUMENT( class SRigWindow*, RigWindow )
+
+		/* Widget used to display the list of retarget sources*/
+		SLATE_ARGUMENT( TSharedPtr<SBoneMappingListType>, BoneMappingListView )
+
+		/* Persona used to update the viewport when a weight slider is dragged */
+		SLATE_ARGUMENT( TWeakPtr<FPersona>, Persona )
+
+		SLATE_EVENT( FOnBoneMappingChanged, OnBoneMappingChanged)
+
+		SLATE_EVENT( FOnGetBoneMapping, OnGetBoneMapping)
+
+	SLATE_END_ARGS()
+
+	void Construct( const FArguments& InArgs, const TSharedRef<STableViewBase>& OwnerTableView );
+
+	/** Overridden from SMultiColumnTableRow.  Generates a widget for this column of the tree row. */
+	virtual TSharedRef<SWidget> GenerateWidgetForColumn( const FName& ColumnName ) override;
+
+private:
+
+	/* The SRigWindow that handles all retarget sources*/
+	// @todo remove
+	SRigWindow* RigWindow;
+
+	/** Widget used to display the list of retarget sources*/
+	TSharedPtr<SBoneMappingListType> BoneMappingListView;
+
+	/** The name and weight of the retarget source*/
+	FDisplayedBoneMappingInfoPtr	Item;
+
+	/** Pointer back to the Persona that owns us */
+	TWeakPtr<FPersona> PersonaPtr;
+
+	// Bone tree widget delegates
+	void OnBoneSelectionChanged(FName Name);
+	FName GetSelectedBone() const;
+
+	FOnBoneMappingChanged OnBoneMappingChanged;
+	FOnGetBoneMapping OnGetBoneMapping;
+};
+
+void SBoneMappingListRow::Construct( const FArguments& InArgs, const TSharedRef<STableViewBase>& InOwnerTableView )
+{
+	Item = InArgs._Item;
+	RigWindow = InArgs._RigWindow;
+	BoneMappingListView = InArgs._BoneMappingListView;
+	OnBoneMappingChanged = InArgs._OnBoneMappingChanged;
+	OnGetBoneMapping = InArgs._OnGetBoneMapping;
+
+	PersonaPtr = InArgs._Persona;
+
+	check( Item.IsValid() );
+
+	SMultiColumnTableRow< FDisplayedBoneMappingInfoPtr >::Construct( FSuperRowType::FArguments(), InOwnerTableView );
+}
+
+TSharedRef< SWidget > SBoneMappingListRow::GenerateWidgetForColumn( const FName& ColumnName )
+{
+	if ( ColumnName == ColumnId_NodeNameLabel )
+	{
+		TSharedPtr< SInlineEditableTextBlock > InlineWidget;
+		TSharedRef< SWidget > NewWidget = 
+			SNew( SVerticalBox )
+
+			+ SVerticalBox::Slot()
+			.AutoHeight()
+			.Padding( 0.0f, 4.0f )
+			.VAlign( VAlign_Center )
+			[
+				SAssignNew(InlineWidget, SInlineEditableTextBlock)
+				.Text( FText::FromString(Item->GetDisplayName()) )
+				.HighlightText( RigWindow->GetFilterText() )
+				.IsReadOnly(true)
+				.IsSelected(this, &SMultiColumnTableRow< FDisplayedBoneMappingInfoPtr >::IsSelectedExclusively)
+			];
+
+		return NewWidget;
+	}
+	else
+	{
+		check (Item->Skeleton);
+
+		// show bone list
+		// Encase the SSpinbox in an SVertical box so we can apply padding. Setting ItemHeight on the containing SListView has no effect :-(
+		return
+			SNew( SVerticalBox )
+
+			+ SVerticalBox::Slot()
+			.AutoHeight()
+			.Padding( 0.0f, 1.0f )
+			.VAlign( VAlign_Center )
+			[
+				SNew(SBoneSelectionWidget)
+				.Skeleton( Item->Skeleton )
+				.Tooltip( FText::Format( LOCTEXT( "BoneSelectinWidget",  "Select Bone for node {0}" ), FText::FromString( Item->GetDisplayName() ) ))
+				.OnBoneSelectionChanged(this, &SBoneMappingListRow::OnBoneSelectionChanged)
+				.OnGetSelectedBone(this, &SBoneMappingListRow::GetSelectedBone)
+			];
+	}
+}
+
+void SBoneMappingListRow::OnBoneSelectionChanged(FName Name)
+{
+	if (OnBoneMappingChanged.IsBound())
+	{
+		OnBoneMappingChanged.Execute(Item->GetNodeName(), Name);
+	}
+
+	// @todo delete?
+//	Item->BoneName = Name;
+}
+
+FName SBoneMappingListRow::GetSelectedBone() const
+{
+	if (OnGetBoneMapping.IsBound())
+	{
+		return OnGetBoneMapping.Execute(Item->GetNodeName());
+	}
+
+	// @todo delete?
+//	return Item->BoneName
+	
+	return NAME_None;
+}
+
+//////////////////////////////////////////////////////////////////////////
+// SRigWindow
+
+void SRigWindow::Construct(const FArguments& InArgs)
+{
+	PersonaPtr = InArgs._Persona;
+	Skeleton = NULL;
+
+	if ( PersonaPtr.IsValid() )
+	{
+		Skeleton = PersonaPtr.Pin()->GetSkeleton();
+		PersonaPtr.Pin()->RegisterOnPostUndo(FPersona::FOnPostUndo::CreateSP( this, &SRigWindow::PostUndo ) );
+	}
+	
+	// @todo it will crash right nwo without Skeleton
+	check (Skeleton);
+
+	Skeleton->RefreshRigConfig();
+
+	// show list of skeletalmeshes that they can choose from
+	FContentBrowserModule& ContentBrowserModule = FModuleManager::Get().LoadModuleChecked<FContentBrowserModule>(TEXT("ContentBrowser"));
+
+	ChildSlot
+	[
+		SNew( SVerticalBox )
+
+		// first add rig asset picker
+		+ SVerticalBox::Slot()
+		.AutoHeight()
+		[
+			SNew(SHorizontalBox)
+
+			+SHorizontalBox::Slot()
+			.AutoWidth()
+			[
+				SNew(STextBlock)
+				.Text(LOCTEXT("RigNameLabel", "Select Rig "))
+				.Font(FEditorStyle::GetFontStyle("Persona.RetargetManager.BoldFont"))
+			]
+
+			+SHorizontalBox::Slot()
+			[
+				// rig asset picker
+				//ContentBrowserModule.Get().CreateAssetPicker(AssetPickerConfig)
+				SNew(SContentReference)
+				//.WidthOverride(80.f)
+				.AllowClearingReference(true)
+				.AllowSelectingNewAsset(true)
+				.AssetReference(this, &SRigWindow::GetRigObject)
+				.AllowedClass(URig::StaticClass())
+				.OnShouldFilterAsset(this, &SRigWindow::ShouldFilterAssetBased)
+				.OnSetReference(this, &SRigWindow::OnAssetSelected)
+				.AssetPickerSizeOverride(FVector2D(250, 700))
+				.InitialAssetViewType(EAssetViewType::List)
+			]
+		]
+
+		// now show bone mapping
+		+ SVerticalBox::Slot()
+		.AutoHeight()
+		.Padding(0,2)
+		[
+			SNew(SHorizontalBox)
+			// Filter entry
+			+SHorizontalBox::Slot()
+			.FillWidth( 1 )
+			[
+				SAssignNew( NameFilterBox, SSearchBox )
+				.SelectAllTextWhenFocused( true )
+				.OnTextChanged( this, &SRigWindow::OnFilterTextChanged )
+				.OnTextCommitted( this, &SRigWindow::OnFilterTextCommitted )
+			]
+		]
+
+		+ SVerticalBox::Slot()
+		.FillHeight( 1.0f )		// This is required to make the scrollbar work, as content overflows Slate containers by default
+		[
+			SAssignNew( BoneMappingListView, SBoneMappingListType )
+			.ListItemsSource( &BoneMappingList )
+			.OnGenerateRow( this, &SRigWindow::GenerateBoneMappingRow )
+			.ItemHeight( 22.0f )
+			.HeaderRow
+			(
+				SNew( SHeaderRow )
+				+ SHeaderRow::Column( ColumnId_NodeNameLabel )
+				.DefaultLabel( LOCTEXT( "RigWindow_NodeNameLabel", "Node (Rig)" ).ToString() )
+				.FixedWidth(150.f)
+
+				+ SHeaderRow::Column( ColumnID_BoneNameLabel )
+				.DefaultLabel( LOCTEXT( "RigWindow_BoneNameLabel", "Bone (Skeleton)" ).ToString() )
+			)
+		]
+	];
+
+	CreateBoneMappingList();
+}
+
+void SRigWindow::OnFilterTextChanged( const FText& SearchText )
+{
+	FilterText = SearchText;
+
+	CreateBoneMappingList( SearchText.ToString() );
+}
+
+void SRigWindow::OnFilterTextCommitted( const FText& SearchText, ETextCommit::Type CommitInfo )
+{
+	// Just do the same as if the user typed in the box
+	OnFilterTextChanged( SearchText );
+}
+
+TSharedRef<ITableRow> SRigWindow::GenerateBoneMappingRow(TSharedPtr<FDisplayedBoneMappingInfo> InInfo, const TSharedRef<STableViewBase>& OwnerTable)
+{
+	check( InInfo.IsValid() );
+
+	return
+		SNew( SBoneMappingListRow, OwnerTable )
+		.Persona( PersonaPtr )
+		.Item( InInfo )
+		.RigWindow( this )
+		.BoneMappingListView( BoneMappingListView )
+		.OnBoneMappingChanged(this, &SRigWindow::OnBoneMappingChanged)
+		.OnGetBoneMapping(this, &SRigWindow::GetBoneMapping);
+}
+
+void SRigWindow::CreateBoneMappingList( const FString& SearchText)
+{
+	BoneMappingList.Empty();
+
+	const URig * Rig = Skeleton->GetRig();
+
+	if ( Rig )
+	{
+		bool bDoFiltering = !SearchText.IsEmpty();
+		const TArray<FNode> & Nodes = Rig->GetNodes();
+
+		for ( const auto Node : Nodes )
+		{
+			const FName & Name = Node.Name;
+			const FString & DisplayName = Node.DisplayName;
+			const FName & BoneName = Skeleton->GetRigBoneMapping(Name);
+
+			if ( bDoFiltering )
+			{
+				if (!Name.ToString().Contains( SearchText ))
+				{
+					continue; // Skip items that don't match our filter
+				}
+				if (!DisplayName.Contains( SearchText ))
+				{
+					continue;
+				}
+				if (!BoneName.ToString().Contains( SearchText ))
+				{
+					continue;
+				}
+			}
+
+			TSharedRef<FDisplayedBoneMappingInfo> Info = FDisplayedBoneMappingInfo::Make( Name, DisplayName, Skeleton );
+
+			BoneMappingList.Add( Info );
+		}
+	}
+
+	BoneMappingListView->RequestListRefresh();
+}
+
+
+void SRigWindow::OnAssetSelected(UObject* Object)
+{
+	if (Skeleton)
+	{
+		const FScopedTransaction Transaction(LOCTEXT("RigAssetChanged", "Select Rig"));
+		Skeleton->Modify();
+		Skeleton->SetRigConfig(Cast<URig>(Object));
+		CreateBoneMappingList();
+
+		FAssetNotifications::SkeletonNeedsToBeSaved(Skeleton);
+	}
+}
+
+/** Returns true if the asset shouldn't show  */
+bool SRigWindow::ShouldFilterAssetBased(const class FAssetData& AssetData)
+{
+	return (AssetData.GetAsset() == GetRigObject());
+}
+
+UObject* SRigWindow::GetRigObject() const
+{
+	return (Skeleton)? Skeleton->GetRig() : NULL;
+}
+
+SRigWindow::~SRigWindow()
+{
+	if (PersonaPtr.IsValid())
+	{
+		PersonaPtr.Pin()->UnregisterOnPostUndo(this);
+	}
+}
+
+void SRigWindow::PostUndo()
+{
+	CreateBoneMappingList();
+}
+
+void SRigWindow::OnBoneMappingChanged(FName NodeName, FName BoneName)
+{
+	const FScopedTransaction Transaction(LOCTEXT("BoneMappingChanged", "Change Bone Mapping"));
+	Skeleton->Modify();
+
+	Skeleton->SetRigBoneMapping(NodeName, BoneName);
+}
+
+FName SRigWindow::GetBoneMapping(FName NodeName)
+{
+	return Skeleton->GetRigBoneMapping(NodeName);
+}
+
+#undef LOCTEXT_NAMESPACE
+
diff --git a/Engine/Source/Editor/Persona/Private/SRigWindow.h b/Engine/Source/Editor/Persona/Private/SRigWindow.h
new file mode 100644
--- /dev/null
+++ b/Engine/Source/Editor/Persona/Private/SRigWindow.h
@@ -0,0 +1,164 @@
+// Copyright 1998-2014 Epic Games, Inc. All Rights Reserved.
+
+#pragma once
+
+#include "Persona.h"
+
+//////////////////////////////////////////////////////////////////////////
+// FDisplayedBoneMappingInfo
+
+class FDisplayedBoneMappingInfo
+{
+public:
+	FName Name;
+	FString DisplayName;
+	USkeleton * Skeleton;
+
+	/** Static function for creating a new item, but ensures that you can only have a TSharedRef to one */
+	static TSharedRef<FDisplayedBoneMappingInfo> Make(const FName NodeName, const FString DisplayName, USkeleton * InSkeleton)
+	{
+		return MakeShareable(new FDisplayedBoneMappingInfo(NodeName, DisplayName, InSkeleton));
+	}
+
+	FName GetNodeName() const
+	{
+		return Name;
+	}
+
+	FString GetDisplayName() const
+	{
+		return DisplayName;
+	}
+
+protected:
+	/** Hidden constructor, always use Make above */
+	FDisplayedBoneMappingInfo(const FName InNodeName, const FString DisplayName, USkeleton * InSkeleton)
+		: Name( InNodeName )
+		, DisplayName( DisplayName )
+		, Skeleton( InSkeleton )
+	{}
+
+	/** Hidden constructor, always use Make above */
+	FDisplayedBoneMappingInfo() {}
+};
+
+typedef SListView< TSharedPtr<FDisplayedBoneMappingInfo> > SBoneMappingListType;
+
+//////////////////////////////////////////////////////////////////////////
+// SRigWindow
+
+class SRigWindow : public SCompoundWidget
+{
+public:
+	SLATE_BEGIN_ARGS( SRigWindow )
+		: _Persona()
+	{}
+		
+		/* The Persona that owns this table */
+		SLATE_ARGUMENT( TWeakPtr< FPersona >, Persona )
+
+	SLATE_END_ARGS()
+
+	/**
+	* Slate construction function
+	*
+	* @param InArgs - Arguments passed from Slate
+	*
+	*/
+	void Construct( const FArguments& InArgs );
+
+	/**
+	* Destructor - resets the morph targets
+	*
+	*/
+	virtual ~SRigWindow();
+
+	/**
+	* Filters the SListView when the user changes the search text box (NameFilterBox)
+	*
+	* @param SearchText - The text the user has typed
+	*
+	*/
+	void OnFilterTextChanged( const FText& SearchText );
+
+	/**
+	* Filters the SListView when the user hits enter or clears the search box
+	* Simply calls OnFilterTextChanged
+	*
+	* @param SearchText - The text the user has typed
+	* @param CommitInfo - Not used
+	*
+	*/
+	void OnFilterTextCommitted( const FText& SearchText, ETextCommit::Type CommitInfo );
+
+	/**
+	* Create a widget for an entry in the tree from an info
+	*
+	* @param InInfo - Shared pointer to the morph target we're generating a row for
+	* @param OwnerTable - The table that owns this row
+	*
+	* @return A new Slate widget, containing the UI for this row
+	*/
+	TSharedRef<ITableRow> GenerateBoneMappingRow(TSharedPtr<FDisplayedBoneMappingInfo> InInfo, const TSharedRef<STableViewBase>& OwnerTable);
+
+	/**
+	* Return true if it can delete
+	*
+	*/
+	bool CanPerformRefresh() const;
+
+	/**
+	* Handler for the delete of retarget source
+	*/
+	void OnRefreshBoneMappingList();
+
+	/**
+	* Accessor so our rows can grab the filtertext for highlighting
+	*
+	*/
+	FText& GetFilterText() { return FilterText; }
+private:
+
+	/**
+	* Clears and rebuilds the table, according to an optional search string
+	*
+	* @param SearchText - Optional search string
+	*
+	*/
+	void CreateBoneMappingList( const FString& SearchText = FString("") );
+
+	/** Pointer back to the Persona that owns us */
+	TWeakPtr<FPersona> PersonaPtr;
+
+	/** Box to filter to a specific morph target name */
+	TSharedPtr<SSearchBox>	NameFilterBox;
+
+	/** Widget used to display the list of retarget sources */
+	TSharedPtr<SBoneMappingListType> BoneMappingListView;
+
+	/** A list of retarget sources. Used by the BoneMappingListView. */
+	TArray< TSharedPtr<FDisplayedBoneMappingInfo> > BoneMappingList;
+
+	/** The skeletal mesh that we grab the morph targets from */
+	USkeleton* Skeleton;
+
+	/** Current text typed into NameFilterBox */
+	FText FilterText;
+
+	/** Delegate for undo/redo transaction **/
+	void PostUndo();
+
+	/**
+	 * Callback for asset picker
+	 */
+	/* Set new reference for skeletal mesh */
+	void OnAssetSelected(UObject* Object);
+
+	/** Returns true if the asset shouldn't show  */
+	bool ShouldFilterAssetBased(const class FAssetData& AssetData);
+
+	UObject* GetRigObject() const;
+
+	void OnBoneMappingChanged( FName NodeName, FName BoneName );
+	FName GetBoneMapping( FName NodeName );
+};
\ No newline at end of file
diff --git a/Engine/Source/Editor/Persona/Private/Shared/PersonaMode.cpp b/Engine/Source/Editor/Persona/Private/Shared/PersonaMode.cpp
--- a/Engine/Source/Editor/Persona/Private/Shared/PersonaMode.cpp
+++ b/Engine/Source/Editor/Persona/Private/Shared/PersonaMode.cpp
@@ -19,7 +19,7 @@
 const FName FPersonaTabs::MorphTargetsID("MorphTargetsTab");
 const FName FPersonaTabs::SkeletonTreeViewID("SkeletonTreeView");		//@TODO: Name
 // Skeleton Pose manager
-const FName FPersonaTabs::RetargetSourceManagerID("RetargetSourceManager");
+const FName FPersonaTabs::RetargetManagerID("RetargetManager");
 // Anim Blueprint params
 // Explorer
 // Blueprint Defaults
@@ -65,7 +65,7 @@ FPersonaAppMode::FPersonaAppMode(TSharedPtr<class FPersona> InPersona, FName InM
 	PersonaTabFactories.RegisterFactory(MakeShareable(new FPreviewViewportSummoner(InPersona)));
 	PersonaTabFactories.RegisterFactory(MakeShareable(new FMorphTargetTabSummoner(InPersona)));
 	PersonaTabFactories.RegisterFactory(MakeShareable(new FSkeletonAnimNotifiesSummoner(InPersona)));
-	PersonaTabFactories.RegisterFactory(MakeShareable(new FRetargetSourceManagerTabSummoner(InPersona)));
+	PersonaTabFactories.RegisterFactory(MakeShareable(new FRetargetManagerTabSummoner(InPersona)));
 	PersonaTabFactories.RegisterFactory(MakeShareable(new FSkeletonSlotNamesSummoner(InPersona)));
 	PersonaTabFactories.RegisterFactory(MakeShareable(new FSkeletonSlotGroupNamesSummoner(InPersona)));
 	PersonaTabFactories.RegisterFactory(MakeShareable(new FSkeletonCurveNameManagerSummoner(InPersona)));
@@ -224,28 +224,29 @@ TSharedRef<SWidget> FPreviewViewportSummoner::CreateTabBody(const FWorkflowTabSp
 }
 
 /////////////////////////////////////////////////////
-// FRetargetSourceManagerTabSummoner
+// FRetargetManagerTabSummoner
 
-#include "SRetargetSourceManager.h"
+#include "SRetargetManager.h"
 
-FRetargetSourceManagerTabSummoner::FRetargetSourceManagerTabSummoner(TSharedPtr<class FAssetEditorToolkit> InHostingApp)
-	: FWorkflowTabFactory(FPersonaTabs::RetargetSourceManagerID, InHostingApp)
+FRetargetManagerTabSummoner::FRetargetManagerTabSummoner(TSharedPtr<class FAssetEditorToolkit> InHostingApp)
+	: FWorkflowTabFactory(FPersonaTabs::RetargetManagerID, InHostingApp)
 {
-	TabLabel = LOCTEXT("RetargetSourceManagerTabTitle", "Retarget Source Manager");
+	TabLabel = LOCTEXT("RetargetManagerTabTitle", "Retarget Manager");
 
 	EnableTabPadding();
 	bIsSingleton = true;
 
-	ViewMenuDescription = LOCTEXT("RetargetSourceManagerTabView", "Retarget Source Manager");
-	ViewMenuTooltip = LOCTEXT("RetargetSourceManagerTabView_ToolTip", "Manages different source pose for animations (retargeting)");
+	ViewMenuDescription = LOCTEXT("RetargetManagerTabView", "Retarget Manager");
+	ViewMenuTooltip = LOCTEXT("RetargetManagerTabView_ToolTip", "Manages different options for retargeting");
 }
 
-TSharedRef<SWidget> FRetargetSourceManagerTabSummoner::CreateTabBody(const FWorkflowTabSpawnInfo& Info) const
+TSharedRef<SWidget> FRetargetManagerTabSummoner::CreateTabBody(const FWorkflowTabSpawnInfo& Info) const
 {
-	return SNew(SRetargetSourceManager)
+	return SNew(SRetargetManager)
 		.Persona(StaticCastSharedPtr<FPersona>(HostingApp.Pin()));
 }
 
+
 /////////////////////////////////////////////////////
 // FAnimBlueprintDefaultsEditorSummoner
 
diff --git a/Engine/Source/Editor/Persona/Private/Shared/PersonaMode.h b/Engine/Source/Editor/Persona/Private/Shared/PersonaMode.h
--- a/Engine/Source/Editor/Persona/Private/Shared/PersonaMode.h
+++ b/Engine/Source/Editor/Persona/Private/Shared/PersonaMode.h
@@ -16,7 +16,8 @@ struct FPersonaTabs
 	static const FName MorphTargetsID;
 	static const FName SkeletonTreeViewID;
 	// Skeleton Pose manager
-	static const FName RetargetSourceManagerID;
+	static const FName RetargetManagerID;
+	static const FName RigManagerID;
 	// Skeleton/Sockets
 	// Anim Blueprint Params
 	// Explorer
@@ -203,19 +204,19 @@ struct FPreviewViewportSummoner : public FWorkflowTabFactory
 };
 
 /////////////////////////////////////////////////////
-// FRetargetSourceManagerTabSummoner
+// FRetargetManagerTabSummoner
 
-struct FRetargetSourceManagerTabSummoner : public FWorkflowTabFactory
+struct FRetargetManagerTabSummoner : public FWorkflowTabFactory
 {
 public:
-	FRetargetSourceManagerTabSummoner(TSharedPtr<class FAssetEditorToolkit> InHostingApp);
+	FRetargetManagerTabSummoner(TSharedPtr<class FAssetEditorToolkit> InHostingApp);
 
 	virtual TSharedRef<SWidget> CreateTabBody(const FWorkflowTabSpawnInfo& Info) const override;
 
 	// Create a tooltip widget for the tab
 	virtual TSharedPtr<SToolTip> CreateTabToolTipWidget(const FWorkflowTabSpawnInfo& Info) const override
 	{
-		return  IDocumentation::Get()->CreateToolTip(LOCTEXT("RetargetSourceTooltip", "In this panel, you can manage retarget sources for different body types"), NULL, TEXT("Shared/Editors/Persona"), TEXT("RetargetSource_Window"));
+		return  IDocumentation::Get()->CreateToolTip(LOCTEXT("RetargetSourceTooltip", "In this panel, you can manage retarget sources for different body types"), NULL, TEXT("Shared/Editors/Persona"), TEXT("RetargetManager"));
 	}
 };
 
diff --git a/Engine/Source/Editor/UnrealEd/Private/AnimationEditorUtils.cpp b/Engine/Source/Editor/UnrealEd/Private/AnimationEditorUtils.cpp
--- a/Engine/Source/Editor/UnrealEd/Private/AnimationEditorUtils.cpp
+++ b/Engine/Source/Editor/UnrealEd/Private/AnimationEditorUtils.cpp
@@ -47,11 +47,6 @@ namespace AnimationEditorUtils
 		{
 			AssetCreated.Execute(ObjectsToSync);
 		}
-// 		if(ObjectsToSync.Num() > 0)
-// 		{
-// 			FContentBrowserModule& ContentBrowserModule = FModuleManager::Get().LoadModuleChecked<FContentBrowserModule>("ContentBrowser");
-// 			ContentBrowserModule.Get().SyncBrowserToAssets(ObjectsToSync, /*bAllowLockedBrowsers=*/true);
-// 		}
 	}
 
 	template <typename TFactory, typename T>
diff --git a/Engine/Source/Runtime/Core/Private/Math/Transform.cpp b/Engine/Source/Runtime/Core/Private/Math/Transform.cpp
--- a/Engine/Source/Runtime/Core/Private/Math/Transform.cpp
+++ b/Engine/Source/Runtime/Core/Private/Math/Transform.cpp
@@ -46,7 +46,7 @@ FString FTransform::ToString() const
 	const FVector T(GetTranslation());
 	const FVector S(GetScale3D());
 
-	return FString::Printf(TEXT("%f,%f,%f|%f,%f,%f|%f,%f,%f"), T.X, T.Y, T.Z, R.Pitch, R.Yaw, R.Roll, S.X, S.Y, S.Z);
+	return FString::Printf(TEXT("%f,%f,%f|P(%f),Y(%f),R(%f)|%f,%f,%f"), T.X, T.Y, T.Z, R.Pitch, R.Yaw, R.Roll, S.X, S.Y, S.Z);
 }
 
 bool FTransform::InitFromString( const FString & Source )
diff --git a/Engine/Source/Runtime/Core/Private/Math/TransformVectorized.cpp b/Engine/Source/Runtime/Core/Private/Math/TransformVectorized.cpp
--- a/Engine/Source/Runtime/Core/Private/Math/TransformVectorized.cpp
+++ b/Engine/Source/Runtime/Core/Private/Math/TransformVectorized.cpp
@@ -46,7 +46,7 @@ FString FTransform::ToString() const
 	const FVector T(GetTranslation());
 	const FVector S(GetScale3D());
 
-	return FString::Printf(TEXT("%f,%f,%f|%f,%f,%f|%f,%f,%f"), T.X, T.Y, T.Z, R.Pitch, R.Yaw, R.Roll, S.X, S.Y, S.Z);
+	return FString::Printf(TEXT("%f,%f,%f|P(%f),Y(%f),R(%f)|%f,%f,%f"), T.X, T.Y, T.Z, R.Pitch, R.Yaw, R.Roll, S.X, S.Y, S.Z);
 }
 
 bool FTransform::InitFromString( const FString & Source )
diff --git a/Engine/Source/Runtime/Engine/Classes/Animation/AnimSequence.h b/Engine/Source/Runtime/Engine/Classes/Animation/AnimSequence.h
--- a/Engine/Source/Runtime/Engine/Classes/Animation/AnimSequence.h
+++ b/Engine/Source/Runtime/Engine/Classes/Animation/AnimSequence.h
@@ -724,6 +724,11 @@ private:
 	 * Remap NaN tracks from the RawAnimation data and recompress
 	 */	
 	void RemoveNaNTracks();
+
+	/** Retargeting functions */
+	bool ConvertAnimationDataToRiggingData(TArray<struct FRawAnimSequenceTrack> & RawRiggingAnimationData);
+	bool ConvertRiggingDataToAnimationData(TArray<struct FRawAnimSequenceTrack> & RawRiggingAnimationData);
+	int32 GetSpaceBasedAnimationData(TArray< TArray<FTransform> > & AnimationDataInComponentSpace, TArray<struct FRawAnimSequenceTrack> * RawRiggingAnimationData) const;
 #endif
 
 	/**
diff --git a/Engine/Source/Runtime/Engine/Classes/Animation/Rig.h b/Engine/Source/Runtime/Engine/Classes/Animation/Rig.h
new file mode 100644
--- /dev/null
+++ b/Engine/Source/Runtime/Engine/Classes/Animation/Rig.h
@@ -0,0 +1,195 @@
+// Copyright 1998-2014 Epic Games, Inc. All Rights Reserved.
+
+/** 
+ * This is the definition for a Rig that is used for retargeting animations
+ *
+ */
+
+#pragma once
+
+// @todo should I keep the hierarchy or not? - does it matter if hand is child of shoulder? - I think we should keep hierarchy
+// @todo should we support reset data 
+// @todo does it make sense to have "no constraint" on certain data? What does that mean? Just World? 
+
+#include "Rig.generated.h"
+
+/** Rig Controller for bone transform **/
+USTRUCT()
+struct FNode
+{
+	GENERATED_USTRUCT_BODY()
+
+	/** This is Display Name where it will be used to display in Retarget Manager. This name has to be unique. */
+	UPROPERTY(EditAnywhere, Category="FNode")
+	FString DisplayName;
+
+	/** Name of the original node. We don't allow to change this. This is used for identity.**/
+	UPROPERTY(VisibleAnywhere, Category="FNode")
+	FName Name;
+
+	/** We save Parent Node but if the parent node is removed, it will reset to root */
+	UPROPERTY(VisibleAnywhere, Category="FNode")
+	FName ParentName;
+
+	/** Absolute transform of the node. Hoping to use this data in the future to render*/
+	UPROPERTY()
+	FTransform Transform;
+
+	FNode()
+		: Name(NAME_None)
+		, ParentName(NAME_None)
+	{
+	}
+
+	FNode(FName InNodeName, FName InParentName, FTransform InTransform)
+		: DisplayName(InNodeName.ToString())
+		, Name(InNodeName)
+		, ParentName(InParentName)
+		, Transform(InTransform)
+
+	{
+	}
+};
+
+/** Control Constraint Type */
+UENUM()
+namespace EControlConstraint
+{
+	enum Type
+	{
+		/** Rotation constraint. */
+		Orientation,
+		/** Translation constraint. */
+		Translation,
+		/** Max Number */
+		Max
+	};
+}
+
+/** Constraint Transform Type. - currently unused*/
+UENUM()
+namespace EConstraintTransform
+{
+	enum Type
+	{
+		/** Absoluate value. */
+		Absoluate,
+		/** Apply relative transform from ref pose. */
+		Relative,
+	};
+}
+
+USTRUCT()
+struct FRigTransformConstraint
+{
+	GENERATED_USTRUCT_BODY()
+
+	/** What transform type **/
+	UPROPERTY(/*EditAnywhere, Category="FRigControlConstraint"*/)
+	TEnumAsByte<EConstraintTransform::Type>	TranformType;
+
+	/** Parent space that are define **/
+	UPROPERTY(EditAnywhere, Category="FRigControlConstraint")
+	FName	ParentSpace;
+
+	/** Weight of the influence - for future*/
+	UPROPERTY(/*EditAnywhere, Category="FRigControlConstraint"*/)
+	float	Weight;
+};
+
+/** This defines what constraint it is defined */
+USTRUCT()
+struct FRigControlConstraint
+{
+	GENERATED_USTRUCT_BODY()
+
+	/** What transform type **/
+	UPROPERTY(EditAnywhere, Category="FRigControlConstraint")
+	TArray<FRigTransformConstraint>			TransformConstraints;
+};
+
+
+/** This is a mapping table between bone in a particular skeletal mesh and bone of this skeleton set. */
+USTRUCT()
+struct FRigControl
+{
+	GENERATED_USTRUCT_BODY()
+
+	UPROPERTY(EditAnywhere, Category="FRigControl")
+	FName					Node;
+
+	UPROPERTY(EditAnywhere, Category="FRigControl")
+	FRigControlConstraint	Constraints[EControlConstraint::Type::Max];
+};
+
+DECLARE_DELEGATE_RetVal_OneParam(int32, FGetParentIndex, FName);
+
+/**
+ *	URig : that has rigging data for skeleton
+ *		- used for retargeting
+ *		- support to share different animations
+ */
+UCLASS(hidecategories=Object, MinimalAPI)
+class URig : public UObject
+{
+	GENERATED_UCLASS_BODY()
+
+private:
+
+	/** Skeleton bone tree - each contains name and parent index**/
+	UPROPERTY(EditAnywhere, Category = Rig, EditFixedSize)
+	TArray<FRigControl> RigControls;
+
+	/** Skeleton bone tree - each contains name and parent index**/
+	UPROPERTY(EditAnywhere, Category=Rig, EditFixedSize)
+	TArray<FNode> Nodes;
+
+public:
+
+#if WITH_EDITOR
+	// Begin UObject Interface
+	virtual void PostEditChangeProperty(FPropertyChangedEvent& PropertyChangedEvent) override;
+	virtual void PostLoad() override;
+	// End UObject Interface
+
+	// node related functions
+	int32 GetNodeNum() const;
+	const FNode * GetNode(int32 NodeIndex) const;
+	FName GetNodeName(int32 NodeIndex) const;
+	FName GetParentNode(FName & NodeName) const;
+	int32 FindNode(const FName & NodeName) const;
+	const TArray<FNode> & GetNodes() const { return Nodes; }
+
+	// create from skeleton
+	ENGINE_API void CreateFromSkeleton(const USkeleton * Skeleton, const TMap<int32, int32> & RequiredBones);
+	ENGINE_API void SetAllConstraintsToParents();
+	ENGINE_API void SetAllConstraintsToWorld();
+
+	// rig control related
+	int32 GetRigControlNum() const;
+	const TArray<FRigControl> & GetRigControls() const { return RigControls; }
+	const FRigControl * GetRigControl(int32 ControlIndex) const;
+	const FRigControl * GetRigControlByNodeName(FName NodeName) const;
+	int32 FindRigControlByNodeName(FName NodeName) const;
+	int32 FindTransformParentNode(int32 NodeIndex, bool bTranslate, int32 Index=0) const;
+#endif
+
+	// not sure if we'd like to keep this
+	ENGINE_API static FName WorldNodeName;
+
+private:
+
+#if WITH_EDITOR
+	// for now these are privates since we don't have good control yet
+	bool AddNode(FName Name, FName ParentNode, FTransform Transform);
+	bool DeleteNode(FName Name);
+	// rig constraint 
+	bool AddRigConstraint(FName NodeName, EControlConstraint::Type	ConstraintType, EConstraintTransform::Type	TranformType, FName ParentSpace, float Weight = 1.f);
+
+#endif // WITH_EDITOR
+	
+	// not useful so far
+	void CalculateComponentSpace(int32 NodeIndex, const FTransform & LocalTransform, const TArray<FTransform> & TransformBuffer, const FGetParentIndex& DelegateToGetParentIndex, FTransform & OutComponentSpaceTransform) const;
+	void CalculateLocalSpace(int32 NodeIndex, const FTransform & ComponentTransform, const TArray<FTransform> & TransformBuffer, const FGetParentIndex& DelegateToGetParentIndex, FTransform & OutLocalSpaceTransform) const;
+};
+
diff --git a/Engine/Source/Runtime/Engine/Classes/Animation/Skeleton.h b/Engine/Source/Runtime/Engine/Classes/Animation/Skeleton.h
--- a/Engine/Source/Runtime/Engine/Classes/Animation/Skeleton.h
+++ b/Engine/Source/Runtime/Engine/Classes/Animation/Skeleton.h
@@ -141,6 +141,49 @@ struct FBoneReductionSetting
 		return (BonesToRemove.Contains(BoneName));
 	}
 };
+
+USTRUCT()
+struct FNameMapping
+{
+	GENERATED_USTRUCT_BODY()
+
+	UPROPERTY()
+	FName NodeName;
+
+	UPROPERTY()
+	FName BoneName;
+
+	FNameMapping()
+		: NodeName(NAME_None)
+		, BoneName(NAME_None)
+	{
+	}
+
+	FNameMapping(FName InNodeName)
+		: NodeName(InNodeName)
+		, BoneName(NAME_None)
+	{
+	}
+
+	FNameMapping(FName InNodeName, FName InBoneName)
+		: NodeName(InNodeName)
+		, BoneName(InBoneName)
+	{
+	}
+};
+
+USTRUCT()
+struct FRigConfiguration
+{
+	GENERATED_USTRUCT_BODY()
+
+	UPROPERTY()
+	class URig * Rig;
+
+	// @todo in the future we can make this to be run-time data
+	UPROPERTY()
+	TArray<FNameMapping> BoneMappingTable;
+};
 /**
  *	USkeleton : that links between mesh and animation
  *		- Bone hierarchy for animations
@@ -168,6 +211,9 @@ private:
 	/** Guid for skeleton */
 	FGuid Guid;
 
+	UPROPERTY()
+	FRigConfiguration RigConfig;
+
 	/** Conversion function. Remove when VER_UE4_REFERENCE_SKELETON_REFACTOR is removed. */
 	void ConvertToFReferenceSkeleton();
 
@@ -296,6 +342,7 @@ public:
 
 	/** Returns the skeletons preview mesh, loading it if necessary */
 	ENGINE_API USkeletalMesh* GetPreviewMesh(bool bFindIfNotSet=false);
+	ENGINE_API USkeletalMesh* GetPreviewMesh() const;
 
 	/** Returns the skeletons preview mesh, loading it if necessary */
 	ENGINE_API void SetPreviewMesh(USkeletalMesh* PreviewMesh, bool bMarkAsDirty=true);
@@ -576,12 +623,20 @@ public:
 	// Asset registry information for animation curves
 	ENGINE_API static const FName CurveTag;
 	ENGINE_API static const TCHAR CurveTagDelimiter;
+
+	// rig Configs
+	ENGINE_API void SetRigConfig(URig * Rig);
+	ENGINE_API FName GetRigBoneMapping(const FName & NodeName) const;
+	ENGINE_API bool SetRigBoneMapping(const FName & NodeName, FName BoneName);
+	// verify if it has all latest data
+	ENGINE_API void RefreshRigConfig();
+	int32 FindRigBoneMapping(const FName & NodeName) const;
+	ENGINE_API URig * GetRig() const;
 #endif
 
 public:
 	// this should be outside of editor because slot node is initialize to it
 	ENGINE_API static const FName DefaultSlotGroupName;
-
 private:
 	void RegenerateGuid();
 };
diff --git a/Engine/Source/Runtime/Engine/Classes/Engine/SkeletalMesh.h b/Engine/Source/Runtime/Engine/Classes/Engine/SkeletalMesh.h
--- a/Engine/Source/Runtime/Engine/Classes/Engine/SkeletalMesh.h
+++ b/Engine/Source/Runtime/Engine/Classes/Engine/SkeletalMesh.h
@@ -687,6 +687,10 @@ public:
 	UPROPERTY()
 	float FloorOffset;
 
+	/** This is buffer that saves pose that is used by retargeting*/
+	UPROPERTY()
+	TArray<FTransform> RetargetBasePose;
+
 #endif
 
 	/** Clothing asset data */
diff --git a/Engine/Source/Runtime/Engine/Private/Animation/AnimSequence.cpp b/Engine/Source/Runtime/Engine/Private/Animation/AnimSequence.cpp
--- a/Engine/Source/Runtime/Engine/Private/Animation/AnimSequence.cpp
+++ b/Engine/Source/Runtime/Engine/Private/Animation/AnimSequence.cpp
@@ -13,6 +13,7 @@
 #include "Animation/AnimCompress_RevertToRaw.h"
 #include "Animation/AnimNotifies/AnimNotify.h"
 #include "Animation/AnimNotifies/AnimNotifyState.h"
+#include "Animation/Rig.h"
 
 #define USE_SLERP 0
 
@@ -1569,63 +1570,87 @@ bool UAnimSequence::IsValidAdditive() const
 }
 
 #if WITH_EDITOR
-void FillUpSpaceBasesRefPose(USkeleton * Skeleton, TArray<FTransform> &SpaceBaseRefPose)
+
+void FillUpTransformBasedOnRig(const USkeleton * Skeleton, TArray<FTransform> & NodeSpaceBases, TArray<FTransform> &Rotations, TArray<FTransform> & Translations)
 {
-	check (Skeleton);
+	TArray<FTransform> SpaceBases;
+	FAnimationRuntime::FillUpSpaceBasesRetargetBasePose(Skeleton, SpaceBases);
 
-	const TArray<FTransform> & ReferencePose = Skeleton->GetReferenceSkeleton().GetRefBonePose();
-	SpaceBaseRefPose.Empty(ReferencePose.Num());
-	SpaceBaseRefPose.AddUninitialized(ReferencePose.Num());
+	const URig * Rig = Skeleton->GetRig();
 
-	// initialize to identity since some of them don't have tracks
-	for(int Index=0; Index <SpaceBaseRefPose.Num(); ++Index)
+	if (Rig && Rig->GetNodeNum() > 0)
 	{
-		int32 ParentIndex = Skeleton->GetReferenceSkeleton().GetParentIndex(Index);
-		if(ParentIndex != INDEX_NONE)
-		{
-			SpaceBaseRefPose[Index] = ReferencePose[Index] * SpaceBaseRefPose[ParentIndex];
-		}
-		else
-		{
-			SpaceBaseRefPose[Index] = ReferencePose[Index];
-		}
-	}
-}
+		int32 NodeNum = Rig->GetNodeNum();
 
-void UAnimSequence::RemapTracksToNewSkeleton( USkeleton * NewSkeleton, bool bConvertSpaces )
-{
+		NodeSpaceBases.Empty(NodeNum);
+		NodeSpaceBases.AddUninitialized(NodeNum);
 
-	// this only replaces the primary one, it doesn't replace old ones
-	TArray<struct FTrackToSkeletonMap> NewTrackToSkeletonMapTable;
-	NewTrackToSkeletonMapTable.Empty(AnimationTrackNames.Num());
-	NewTrackToSkeletonMapTable.AddUninitialized(AnimationTrackNames.Num());
-	for ( int32 Track = 0; Track < AnimationTrackNames.Num(); ++Track )
-	{
-		int32 BoneIndex = NewSkeleton->GetReferenceSkeleton().FindBoneIndex(AnimationTrackNames[Track]);
-		NewTrackToSkeletonMapTable[Track].BoneTreeIndex = BoneIndex;
-	}
+		Rotations.Empty(NodeNum);
+		Rotations.AddUninitialized(NodeNum);
 
-	// now I have all NewTrack To Skeleton Map Table
-	// I'll need to compare with old tracks and copy over if SkeletonIndex == 0
-	// if SkeletonIndex != 0, we need to see if we can 
-	for (int32 TableId = 0; TableId <NewTrackToSkeletonMapTable.Num(); ++TableId)
-	{
-		if ( ensure(TrackToSkeletonMapTable.IsValidIndex(TableId)) )
+		Translations.Empty(NodeNum);
+		Translations.AddUninitialized(NodeNum);
+
+		for ( int32 Index = 0; Index < NodeNum; ++Index )
 		{
-			if ( NewTrackToSkeletonMapTable[TableId].BoneTreeIndex != INDEX_NONE )
+			const FName & NodeName = Rig->GetNodeName(Index);
+			const FName & BoneName = Skeleton->GetRigBoneMapping(NodeName);
+			const int32 & BoneIndex = Skeleton->GetReferenceSkeleton().FindBoneIndex(BoneName);
+
+			if (BoneIndex == INDEX_NONE)
 			{
-				TrackToSkeletonMapTable[TableId].BoneTreeIndex = NewTrackToSkeletonMapTable[TableId].BoneTreeIndex;
+				// initialize with SpaceBases - assuming World Based
+				NodeSpaceBases[Index] = FTransform::Identity;
+				Rotations[Index] = FTransform::Identity;
+				Translations[Index] = FTransform::Identity;
 			}
-			else
+			else 
 			{
-				// if not found, delete the track data
-				RemoveTrack(TableId);
-				NewTrackToSkeletonMapTable.RemoveAt(TableId);
-				--TableId;
+				// initialize with SpaceBases - assuming World Based
+				NodeSpaceBases[Index] = SpaceBases[BoneIndex];
+				Rotations[Index] = SpaceBases[BoneIndex];
+				Translations[Index] = SpaceBases[BoneIndex];
+
+				const FRigControl * RigControl = Rig->GetRigControlByNodeName(NodeName);
+
+				if (RigControl != NULL)
+				{
+					// orientation constraint			
+					const auto & RotConstraint = RigControl->Constraints[EControlConstraint::Type::Orientation];
+
+					if (RotConstraint.TransformConstraints.Num() > 0)
+					{
+						const FName & ParentBoneName = Skeleton->GetRigBoneMapping(RotConstraint.TransformConstraints[0].ParentSpace);
+						const int32 & ParentBoneIndex = Skeleton->GetReferenceSkeleton().FindBoneIndex(ParentBoneName);
+
+						if (ParentBoneIndex != INDEX_NONE)
+						{
+							Rotations[Index] = SpaceBases[BoneIndex].GetRelativeTransform(SpaceBases[ParentBoneIndex]);
+						}
+					}
+
+					// translation constraint
+					const auto & TransConstraint = RigControl->Constraints[EControlConstraint::Type::Translation];
+
+					if (TransConstraint.TransformConstraints.Num() > 0)
+					{
+						const FName & ParentBoneName = Skeleton->GetRigBoneMapping(TransConstraint.TransformConstraints[0].ParentSpace);
+						const int32 & ParentBoneIndex = Skeleton->GetReferenceSkeleton().FindBoneIndex(ParentBoneName);
+
+						if (ParentBoneIndex != INDEX_NONE)
+						{
+							// I think translation has to include rotation, otherwise it won't work
+							Translations[Index] = SpaceBases[BoneIndex].GetRelativeTransform(SpaceBases[ParentBoneIndex]);
+						}
+					}
+				}
 			}
 		}
 	}
+}
 
+void UAnimSequence::RemapTracksToNewSkeleton( USkeleton * NewSkeleton, bool bConvertSpaces )
+{
 	// this is not cheap, so make sure it only happens in editor
 
 	// @Todo : currently additive will work fine since we don't bake anything except when we extract
@@ -1634,19 +1659,233 @@ void UAnimSequence::RemapTracksToNewSkeleton( USkeleton * NewSkeleton, bool bCon
 	{
 		USkeleton * OldSkeleton = GetSkeleton();
 
-		// @IMPORTANT: this converts animation data to new skeleton
-		// this requires reference pose to be relative in same position
-		// if reference pose is totally different - i.e. A pose vs T pose then it won't work
-		if(OldSkeleton)
+		// first check if both has same rig, if so, we'll retarget using it
+		if (OldSkeleton && OldSkeleton->GetRig() != NULL && NewSkeleton->GetRig() == OldSkeleton->GetRig())
 		{
+			const URig * Rig = OldSkeleton->GetRig();
+
+			// we'll have to save the relative space bases transform from old ref pose to new refpose
+			TArray<FTransform> RelativeToNewSpaceBases;
+			// save the ratio of translation change
+			TArray<float> OldToNewTranslationRatio;
+			// create relative transform in component space between old skeleton and new skeleton
+			{
+				// first calculate component space ref pose to get the relative transform between
+				// two ref poses. It is very important update ref pose before getting here. 
+				TArray<FTransform> NewRotations, NewSpaceBases, OldRotations;
+				TArray<FTransform> NewTranslations, OldSpaceBases, OldTranslations;
+				// get the spacebases transform
+				FillUpTransformBasedOnRig(NewSkeleton, NewSpaceBases, NewRotations, NewTranslations);
+				FillUpTransformBasedOnRig(OldSkeleton, OldSpaceBases, OldRotations, OldTranslations);
+
+				// now we'd like to get the relative transform from old to new ref pose in component space
+				// PK2*K2 = PK1*K1*theta where theta => P1*R1*theta = P2*R2 
+				// where	P1 - parent transform in component space for original skeleton
+				//			R1 - local space of the current bone for original skeleton
+				//			P2 - parent transform in component space for new skeleton
+				//			R2 - local space of the current bone for new skeleton
+				// what we're looking for is theta, so that we can apply that to animated transform
+				int32 NumNodes = Rig->GetNodeNum();
+				// saves the theta data per node
+				RelativeToNewSpaceBases.AddUninitialized(NumNodes);
+				// saves the translation conversion datao
+				OldToNewTranslationRatio.AddUninitialized(NumNodes);
+
+				// calculate the relative transform to new skeleton
+				// so that we can apply the delta in component space
+				for (int32 NodeIndex = 0; NodeIndex < NumNodes; ++NodeIndex)
+				{
+					// theta (RelativeToNewTransform) = (P1*R1)^(-1) * P2*R2 where theta => P1*R1*theta = P2*R2
+					RelativeToNewSpaceBases[NodeIndex] = NewSpaceBases[NodeIndex].GetRelativeTransform(OldSpaceBases[NodeIndex]); 
+// 					UE_LOG(LogAnimation, Warning, TEXT("NodeName : %s"), *Rig->GetNodeName(NodeIndex).ToString());
+// 					UE_LOG(LogAnimation, Warning, TEXT("Old Translation : %s"), *OldSpaceBases[NodeIndex].ToString());
+// 					UE_LOG(LogAnimation, Warning, TEXT("New Translation : %s"), *NewSpaceBases[NodeIndex].ToString());
+// 					UE_LOG(LogAnimation, Warning, TEXT("Diff Translation : %s"), *RelativeToNewSpaceBases[NodeIndex].ToString());
+
+					// also savees the translation difference between old to new
+					FVector OldTranslation = OldTranslations[NodeIndex].GetTranslation();
+					FVector NewTranslation = NewTranslations[NodeIndex].GetTranslation();
+
+					float OldTranslationSize = OldTranslation.Size();
+					float NewTranslationSize = NewTranslation.Size();
+					OldToNewTranslationRatio[NodeIndex] = (FMath::IsNearlyZero(OldTranslationSize)) ? 1.f/*do not touch new translation size*/ : NewTranslationSize / OldTranslationSize;
+				}
+			}
+
+			TArray<struct FRawAnimSequenceTrack> RawRiggingAnimationData;
+
+			// now convert animation data to rig data
+			ConvertAnimationDataToRiggingData(RawRiggingAnimationData);
+
+			// now convert to space bases animation 
+			int32 NumNodes = Rig->GetNodeNum();
+			TArray< TArray<FTransform> > ComponentSpaceAnimations, ConvertedLocalSpaceAnimations, ConvertedSpaceAnimations;
+			ComponentSpaceAnimations.AddZeroed(NumNodes);
+			ConvertedSpaceAnimations.AddZeroed(NumNodes);
+			ConvertedLocalSpaceAnimations.AddZeroed(NumNodes);
+
+			int32 NumKeys = NumFrames;
+			float Interval = (NumKeys)? SequenceLength/NumKeys : 0.f;
+
+			// allocate arrays
+			for(int32 NodeIndex=0; NodeIndex<NumNodes; ++NodeIndex)
+			{
+				ComponentSpaceAnimations[NodeIndex].AddUninitialized(NumKeys);
+				ConvertedLocalSpaceAnimations[NodeIndex].AddUninitialized(NumKeys);
+				ConvertedSpaceAnimations[NodeIndex].AddUninitialized(NumKeys);
+			}
+
+
+			for (int32 NodeIndex=0; NodeIndex<NumNodes; ++NodeIndex)		
+			{
+				auto & RawAnimation = RawRiggingAnimationData[NodeIndex];
+
+				// find rotation parent node
+				int32 RotParentNodeIndex = Rig->FindTransformParentNode(NodeIndex, false);
+				int32 TransParentNodeIndex = Rig->FindTransformParentNode(NodeIndex, true);
+				// fill up keys - calculate PK1 * K1
+				for(int32 Key=0; Key<NumKeys; ++Key)
+				{
+					FTransform AnimatedLocalKey;
+					ExtractBoneTransform(RawRiggingAnimationData, AnimatedLocalKey, NodeIndex, Interval*Key);
+
+					AnimatedLocalKey.ScaleTranslation(OldToNewTranslationRatio[NodeIndex]);
+
+					if(RotParentNodeIndex != INDEX_NONE)
+					{
+						FQuat ComponentSpaceRotation = ComponentSpaceAnimations[RotParentNodeIndex][Key].GetRotation() * AnimatedLocalKey.GetRotation();
+						ComponentSpaceAnimations[NodeIndex][Key].SetRotation(ComponentSpaceRotation);
+					}
+					else
+					{
+						ComponentSpaceAnimations[NodeIndex][Key].SetRotation(AnimatedLocalKey.GetRotation());
+					}
+
+					if (TransParentNodeIndex != INDEX_NONE)
+					{
+						FVector ComponentSpaceTranslation = ComponentSpaceAnimations[TransParentNodeIndex][Key].TransformPosition(AnimatedLocalKey.GetTranslation());
+						ComponentSpaceAnimations[NodeIndex][Key].SetTranslation(ComponentSpaceTranslation);
+						ComponentSpaceAnimations[NodeIndex][Key].SetScale3D(AnimatedLocalKey.GetScale3D());
+					}
+					else
+					{
+						ComponentSpaceAnimations[NodeIndex][Key].SetTranslation(AnimatedLocalKey.GetTranslation());
+						ComponentSpaceAnimations[NodeIndex][Key].SetScale3D(AnimatedLocalKey.GetScale3D());
+					}
+				}
+			}
+
+			// now animation is converted to component space
+			TArray<struct FRawAnimSequenceTrack> NewRawAnimationData = RawRiggingAnimationData;
+			for (int32 NodeIndex=0; NodeIndex<NumNodes; ++NodeIndex)
+			{
+				// find rotation parent node
+				int32 RotParentNodeIndex = Rig->FindTransformParentNode(NodeIndex, false);
+				int32 TransParentNodeIndex = Rig->FindTransformParentNode(NodeIndex, true);
+
+				for(int32 Key=0; Key<NumKeys; ++Key)
+				{
+					// now convert to the new space and save to local spaces
+					// clear translation;
+					RelativeToNewSpaceBases[NodeIndex].SetTranslation(FVector::ZeroVector);
+					ConvertedSpaceAnimations[NodeIndex][Key] = RelativeToNewSpaceBases[NodeIndex] * ComponentSpaceAnimations[NodeIndex][Key];
+
+					if(RotParentNodeIndex != INDEX_NONE)
+					{
+						FQuat LocalRotation = ConvertedSpaceAnimations[RotParentNodeIndex][Key].GetRotation().Inverse() * ConvertedSpaceAnimations[NodeIndex][Key].GetRotation();
+						ConvertedLocalSpaceAnimations[NodeIndex][Key].SetRotation(LocalRotation);
+					}
+					else
+					{
+						ConvertedLocalSpaceAnimations[NodeIndex][Key].SetRotation(ConvertedSpaceAnimations[NodeIndex][Key].GetRotation());
+					}
+
+					if(TransParentNodeIndex != INDEX_NONE)
+					{
+						FVector LocalTranslation = ConvertedSpaceAnimations[NodeIndex][Key].GetRelativeTransform(ConvertedSpaceAnimations[TransParentNodeIndex][Key]).GetTranslation();
+						ConvertedLocalSpaceAnimations[NodeIndex][Key].SetTranslation(LocalTranslation);
+						ConvertedLocalSpaceAnimations[NodeIndex][Key].SetScale3D(ConvertedSpaceAnimations[NodeIndex][Key].GetScale3D());
+					}
+					else
+					{
+						ConvertedLocalSpaceAnimations[NodeIndex][Key].SetTranslation(ConvertedSpaceAnimations[NodeIndex][Key].GetTranslation());
+						ConvertedLocalSpaceAnimations[NodeIndex][Key].SetScale3D(ConvertedSpaceAnimations[NodeIndex][Key].GetScale3D());
+					}
+				}
+
+				auto & RawAnimation = NewRawAnimationData[NodeIndex];
+				RawAnimation.PosKeys.Empty(NumKeys);
+				RawAnimation.PosKeys.AddUninitialized(NumKeys);
+				RawAnimation.RotKeys.Empty(NumKeys);
+				RawAnimation.RotKeys.AddUninitialized(NumKeys);
+				RawAnimation.ScaleKeys.Empty(NumKeys);
+				RawAnimation.ScaleKeys.AddUninitialized(NumKeys);
+
+				for(int32 Key=0; Key<NumKeys; ++Key)
+				{
+					RawAnimation.PosKeys[Key] = ConvertedLocalSpaceAnimations[NodeIndex][Key].GetLocation();
+					RawAnimation.RotKeys[Key] = ConvertedLocalSpaceAnimations[NodeIndex][Key].GetRotation();
+					RawAnimation.ScaleKeys[Key] = ConvertedLocalSpaceAnimations[NodeIndex][Key].GetScale3D();
+				}
+			}
+
+			RawRiggingAnimationData = NewRawAnimationData;
+
+			// set new skeleton
+			SetSkeleton(NewSkeleton);
+
+			// convert back to animated data with new skeleton
+			ConvertRiggingDataToAnimationData(RawRiggingAnimationData);
+		}
+		// @todo end rig testing
+		// @IMPORTANT: now otherwise this will try to do bone to bone mapping
+		else if(OldSkeleton)
+		{
+			// this only replaces the primary one, it doesn't replace old ones
+			TArray<struct FTrackToSkeletonMap> NewTrackToSkeletonMapTable;
+			NewTrackToSkeletonMapTable.Empty(AnimationTrackNames.Num());
+			NewTrackToSkeletonMapTable.AddUninitialized(AnimationTrackNames.Num());
+			for (int32 Track = 0; Track < AnimationTrackNames.Num(); ++Track)
+			{
+				int32 BoneIndex = NewSkeleton->GetReferenceSkeleton().FindBoneIndex(AnimationTrackNames[Track]);
+				NewTrackToSkeletonMapTable[Track].BoneTreeIndex = BoneIndex;
+			}
+
+			// now I have all NewTrack To Skeleton Map Table
+			// I'll need to compare with old tracks and copy over if SkeletonIndex == 0
+			// if SkeletonIndex != 0, we need to see if we can 
+			for (int32 TableId = 0; TableId < NewTrackToSkeletonMapTable.Num(); ++TableId)
+			{
+				if (ensure(TrackToSkeletonMapTable.IsValidIndex(TableId)))
+				{
+					if (NewTrackToSkeletonMapTable[TableId].BoneTreeIndex != INDEX_NONE)
+					{
+						TrackToSkeletonMapTable[TableId].BoneTreeIndex = NewTrackToSkeletonMapTable[TableId].BoneTreeIndex;
+					}
+					else
+					{
+						// if not found, delete the track data
+						RemoveTrack(TableId);
+						NewTrackToSkeletonMapTable.RemoveAt(TableId);
+						--TableId;
+					}
+				}
+			}
+
+			if (TrackToSkeletonMapTable.Num() == 0)
+			{
+				// no bones to retarget
+				// return with error
+				//@todo fail message
+			}
 			// make sure you do update reference pose before coming here
 			
 			// first calculate component space ref pose to get the relative transform between
 			// two ref poses. It is very important update ref pose before getting here. 
 			TArray<FTransform> NewSpaceBaseRefPose, OldSpaceBaseRefPose, RelativeToNewTransform;
 			// get the spacebases transform
-			FillUpSpaceBasesRefPose(NewSkeleton, NewSpaceBaseRefPose);
-			FillUpSpaceBasesRefPose(OldSkeleton, OldSpaceBaseRefPose);
+			FAnimationRuntime::FillUpSpaceBasesRefPose(NewSkeleton, NewSpaceBaseRefPose);
+			FAnimationRuntime::FillUpSpaceBasesRefPose(OldSkeleton, OldSpaceBaseRefPose);
 
 			const TArray<FTransform> & OldRefPose = OldSkeleton->GetReferenceSkeleton().GetRefBonePose();
 			const TArray<FTransform> & NewRefPose = NewSkeleton->GetReferenceSkeleton().GetRefBonePose();
@@ -1820,11 +2059,44 @@ void UAnimSequence::RemapTracksToNewSkeleton( USkeleton * NewSkeleton, bool bCon
 			}
 			RawAnimationData = NewRawAnimationData;
 		}
-	}
+		else
+		{
+			// this only replaces the primary one, it doesn't replace old ones
+			TArray<struct FTrackToSkeletonMap> NewTrackToSkeletonMapTable;
+			NewTrackToSkeletonMapTable.Empty(AnimationTrackNames.Num());
+			NewTrackToSkeletonMapTable.AddUninitialized(AnimationTrackNames.Num());
+			for (int32 Track = 0; Track < AnimationTrackNames.Num(); ++Track)
+			{
+				int32 BoneIndex = NewSkeleton->GetReferenceSkeleton().FindBoneIndex(AnimationTrackNames[Track]);
+				NewTrackToSkeletonMapTable[Track].BoneTreeIndex = BoneIndex;
+			}
 
-	// I have to set this here in order for compression
-	// that has to happen outside of this after Skeleton changes
-	SetSkeleton(NewSkeleton);
+			// now I have all NewTrack To Skeleton Map Table
+			// I'll need to compare with old tracks and copy over if SkeletonIndex == 0
+			// if SkeletonIndex != 0, we need to see if we can 
+			for (int32 TableId = 0; TableId < NewTrackToSkeletonMapTable.Num(); ++TableId)
+			{
+				if (ensure(TrackToSkeletonMapTable.IsValidIndex(TableId)))
+				{
+					if (NewTrackToSkeletonMapTable[TableId].BoneTreeIndex != INDEX_NONE)
+					{
+						TrackToSkeletonMapTable[TableId].BoneTreeIndex = NewTrackToSkeletonMapTable[TableId].BoneTreeIndex;
+					}
+					else
+					{
+						// if not found, delete the track data
+						RemoveTrack(TableId);
+						NewTrackToSkeletonMapTable.RemoveAt(TableId);
+						--TableId;
+					}
+				}
+			}
+		}
+
+		// I have to set this here in order for compression
+		// that has to happen outside of this after Skeleton changes
+		SetSkeleton(NewSkeleton);
+	}
 
 	PostProcessSequence();
 }
@@ -1993,6 +2265,430 @@ bool UAnimSequence::AddLoopingInterpolation()
 
 	return false;
 }
+int32 FindParentNodeIndex(URig * Rig, USkeleton * Skeleton, FName ParentNodeName)
+{
+	const int32 & ParentNodeIndex = Rig->FindNode(ParentNodeName);
+	const FName & ParentBoneName = Skeleton->GetRigBoneMapping(ParentNodeName);
+	
+	return Skeleton->GetReferenceSkeleton().FindBoneIndex(ParentBoneName);
+}
+
+int32 UAnimSequence::GetSpaceBasedAnimationData(TArray< TArray<FTransform> > & AnimationDataInComponentSpace, TArray<struct FRawAnimSequenceTrack> * RawRiggingAnimationData) const
+{
+	USkeleton * Skeleton = GetSkeleton();
+
+	check(Skeleton);
+	const FReferenceSkeleton & RefSkeleton = Skeleton->GetReferenceSkeleton();
+	int32 NumBones = RefSkeleton.GetNum();
+
+	AnimationDataInComponentSpace.Empty(NumBones);
+	AnimationDataInComponentSpace.AddZeroed(NumBones);
+
+	// 2d array of animated time [boneindex][time key]
+	int32 NumKeys = NumFrames;
+	float Interval = (NumKeys) ? SequenceLength / NumKeys : 0.f;
+
+	// allocate arrays
+	for (int32 BoneIndex = 0; BoneIndex < NumBones; ++BoneIndex)
+	{
+		AnimationDataInComponentSpace[BoneIndex].AddUninitialized(NumKeys);
+	}
+
+	if (RawRiggingAnimationData)
+	{
+		const URig * Rig = Skeleton->GetRig();
+
+		check(Rig);
+
+		TArray<bool> bCompleteComponentSpaces;
+		bCompleteComponentSpaces.AddZeroed(NumBones);
+
+		// @Todo sort rig control by parent -> child, so that we can calculate parent first before child
+		// but that depends on each skeleton
+
+		const TArray<FRigControl> & RigControls = Rig->GetRigControls();
+		// first put all the component space
+		for (const auto & RigControl : RigControls)
+		{
+			const FName & NodeName = RigControl.Node;
+			const int32 & NodeIndex = Rig->FindNode(NodeName);
+			const FName & BoneName = Skeleton->GetRigBoneMapping(NodeName);
+			const int32 & BoneIndex = RefSkeleton.FindBoneIndex(BoneName);
+
+			if (BoneIndex != INDEX_NONE)
+			{
+				// now calculate the component space
+				const TArray<FRigTransformConstraint>	& RotTransformConstraints = RigControl.Constraints[EControlConstraint::Type::Orientation].TransformConstraints;
+
+				FQuat ComponentRotation;
+				FTransform ComponentTranslation;
+				FVector ComponentScale = FVector(1.f);
+
+				//if (RotTransformConstraints.Num() > 0)
+				{
+					const FName & ParentNodeName = RotTransformConstraints[0].ParentSpace;
+					const int32 & ParentNodeIndex = Rig->FindNode(ParentNodeName);
+					const FName & ParentBoneName = Skeleton->GetRigBoneMapping(ParentNodeName);
+					const int32 & ParentBoneIndex = RefSkeleton.FindBoneIndex(ParentBoneName);
+
+					if (ParentBoneIndex != INDEX_NONE)
+					{
+						// make sure parent is calculated first
+						check(bCompleteComponentSpaces[ParentBoneIndex]);
+
+						for (int32 Key = 0; Key < NumKeys; ++Key)
+						{
+							ComponentRotation = AnimationDataInComponentSpace[ParentBoneIndex][Key].GetRotation() * (*RawRiggingAnimationData)[NodeIndex].RotKeys[Key];
+							AnimationDataInComponentSpace[BoneIndex][Key].SetRotation(ComponentRotation);
+						}
+					}
+					else
+					{
+						for (int32 Key = 0; Key < NumKeys; ++Key)
+						{
+							ComponentRotation = (*RawRiggingAnimationData)[NodeIndex].RotKeys[Key];
+							AnimationDataInComponentSpace[BoneIndex][Key].SetRotation(ComponentRotation);
+						}
+					}
+
+					bCompleteComponentSpaces[BoneIndex] = true;
+				}
+
+				const TArray<FRigTransformConstraint>	& PosTransformConstraints = RigControl.Constraints[EControlConstraint::Type::Translation].TransformConstraints;
+
+				//if (PosTransformConstraints.Num() > 0)
+				{
+					const FName & ParentNodeName = PosTransformConstraints[0].ParentSpace;
+					const int32 & ParentNodeIndex = Rig->FindNode(ParentNodeName);
+					const FName & ParentBoneName = Skeleton->GetRigBoneMapping(ParentNodeName);
+					const int32 & ParentBoneIndex = RefSkeleton.FindBoneIndex(ParentBoneName);
+
+					if (ParentBoneIndex != INDEX_NONE)
+					{
+						// make sure parent is calculated first
+						check(bCompleteComponentSpaces[ParentBoneIndex]);
+
+						for (int32 Key = 0; Key < NumKeys; ++Key)
+						{
+							ComponentTranslation = FTransform((*RawRiggingAnimationData)[NodeIndex].PosKeys[Key]) * AnimationDataInComponentSpace[ParentBoneIndex][Key];
+							AnimationDataInComponentSpace[BoneIndex][Key].SetTranslation(ComponentTranslation.GetTranslation());
+							AnimationDataInComponentSpace[BoneIndex][Key].SetScale3D(ComponentScale);
+						}
+					}
+					else
+					{
+						for (int32 Key = 0; Key < NumKeys; ++Key)
+						{
+							ComponentTranslation = FTransform((*RawRiggingAnimationData)[NodeIndex].PosKeys[Key]);
+							AnimationDataInComponentSpace[BoneIndex][Key].SetTranslation(ComponentTranslation.GetTranslation());
+							AnimationDataInComponentSpace[BoneIndex][Key].SetScale3D(ComponentScale);
+						}
+					}
+				}
+			}
+		}
+
+		for (int32 BoneIndex = 0; BoneIndex < NumBones; ++BoneIndex)
+		{
+			// fill up all not calculated ones
+			if (bCompleteComponentSpaces[BoneIndex] == false)
+			{
+				int32 ParentIndex = RefSkeleton.GetParentIndex(BoneIndex);
+				const FTransform & LocalSpace = RefSkeleton.GetRefBonePose()[BoneIndex];
+				if (ParentIndex != INDEX_NONE)
+				{
+					for (int32 Key = 0; Key < NumKeys; ++Key)
+					{
+						AnimationDataInComponentSpace[BoneIndex][Key] = LocalSpace * AnimationDataInComponentSpace[ParentIndex][Key];
+					}
+				}
+				else
+				{
+					for (int32 Key = 0; Key < NumKeys; ++Key)
+					{
+						AnimationDataInComponentSpace[BoneIndex][Key] = LocalSpace;
+					}
+				}
+			}
+		}
+	}
+	else
+	{
+		// now calculating old animated space bases
+		// this one calculates aniamted space per bones and per key
+		for (int32 BoneIndex = 0; BoneIndex < NumBones; ++BoneIndex)
+		{
+			FName BoneName = Skeleton->GetReferenceSkeleton().GetBoneName(BoneIndex);
+			int32 TrackIndex = AnimationTrackNames.Find(BoneName);
+			int32 ParentBoneIndex = Skeleton->GetReferenceSkeleton().GetParentIndex(BoneIndex);
+
+			if (TrackIndex != INDEX_NONE)
+			{
+				auto & RawAnimation = RawAnimationData[TrackIndex];
+				// fill up keys - calculate PK1 * K1
+				for (int32 Key = 0; Key < NumKeys; ++Key)
+				{
+					FTransform AnimatedLocalKey;
+					ExtractBoneTransform(RawAnimationData, AnimatedLocalKey, TrackIndex, Interval*Key);
+
+					if (ParentBoneIndex != INDEX_NONE)
+					{
+						AnimationDataInComponentSpace[BoneIndex][Key] = AnimatedLocalKey * AnimationDataInComponentSpace[ParentBoneIndex][Key];
+					}
+					else
+					{
+						AnimationDataInComponentSpace[BoneIndex][Key] = AnimatedLocalKey;
+					}
+				}
+			}
+			else
+			{
+				// get local spaces from refpose and use that to fill it up
+				FTransform LocalTransform = Skeleton->GetReferenceSkeleton().GetRefBonePose()[BoneIndex];
+
+				for (int32 Key = 0; Key < NumKeys; ++Key)
+				{
+					if (ParentBoneIndex != INDEX_NONE)
+					{
+						AnimationDataInComponentSpace[BoneIndex][Key] = LocalTransform * AnimationDataInComponentSpace[ParentBoneIndex][Key];
+					}
+					else
+					{
+						AnimationDataInComponentSpace[BoneIndex][Key] = LocalTransform;
+					}
+				}
+			}
+		}
+
+	}
+
+	return AnimationDataInComponentSpace.Num();
+}
+
+bool UAnimSequence::ConvertAnimationDataToRiggingData(TArray<struct FRawAnimSequenceTrack> & RawRiggingAnimationData)
+{
+	USkeleton * Skeleton = GetSkeleton();
+	if (Skeleton && Skeleton->GetRig())
+	{
+		const URig * Rig = Skeleton->GetRig();
+		const int32 & NumNodes = Rig->GetNodeNum();
+		TArray< TArray<FTransform> > AnimationDataInComponentSpace;
+		int32 NumBones = GetSpaceBasedAnimationData(AnimationDataInComponentSpace, NULL);
+
+		if (NumBones > 0)
+		{
+			RawRiggingAnimationData.Empty(NumNodes);
+			RawRiggingAnimationData.AddZeroed(NumNodes);
+
+			// first we copy all space bases back to it
+			for (int32 NodeIndex = 0; NodeIndex < NumNodes; ++NodeIndex)
+			{
+				struct FRawAnimSequenceTrack & Track = RawRiggingAnimationData[NodeIndex];
+				const FName & NodeName = Rig->GetNodeName(NodeIndex);
+				const FName & BoneName = Skeleton->GetRigBoneMapping(NodeName);
+				const int32 & BoneIndex = Skeleton->GetReferenceSkeleton().FindBoneIndex(BoneName);
+
+				Track.PosKeys.Empty(NumFrames);
+				Track.RotKeys.Empty(NumFrames);
+				Track.ScaleKeys.Empty(NumFrames);
+				Track.PosKeys.AddUninitialized(NumFrames);
+				Track.RotKeys.AddUninitialized(NumFrames);
+				Track.ScaleKeys.AddUninitialized(NumFrames);
+
+				if (BoneIndex != INDEX_NONE)
+				{
+					int32 RigConstraintIndex = Rig->FindRigControlByNodeName(NodeName);
+
+					if (RigConstraintIndex != INDEX_NONE)
+					{
+						const auto * RigConstraint = Rig->GetRigControl(RigConstraintIndex);
+
+						// apply orientation - for now only one
+						const TArray<FRigTransformConstraint> & RotationTransformConstraint = RigConstraint->Constraints[EControlConstraint::Type::Orientation].TransformConstraints;
+
+						if (RotationTransformConstraint.Num() > 0)
+						{
+							const FName & ParentSpace = RotationTransformConstraint[0].ParentSpace;
+							const FName & ParentBoneName = Skeleton->GetRigBoneMapping(ParentSpace);
+							const int32 & ParentBoneIndex = Skeleton->GetReferenceSkeleton().FindBoneIndex(ParentBoneName);
+							if (ParentBoneIndex != INDEX_NONE)
+							{
+								// if no rig control, component space is used
+								for (int32 KeyIndex = 0; KeyIndex < NumFrames; ++KeyIndex)
+								{
+									FTransform ParentTransform = AnimationDataInComponentSpace[ParentBoneIndex][KeyIndex];
+									FTransform RelativeTransform = AnimationDataInComponentSpace[BoneIndex][KeyIndex].GetRelativeTransform(ParentTransform);
+									Track.RotKeys[KeyIndex] = RelativeTransform.GetRotation();
+								}
+							}
+							else
+							{
+								// if no rig control, component space is used
+								for (int32 KeyIndex = 0; KeyIndex < NumFrames; ++KeyIndex)
+								{
+									Track.RotKeys[KeyIndex] = AnimationDataInComponentSpace[BoneIndex][KeyIndex].GetRotation();
+								}
+							}
+						}
+						else
+						{
+							// if no rig control, component space is used
+							for (int32 KeyIndex = 0; KeyIndex < NumFrames; ++KeyIndex)
+							{
+								Track.RotKeys[KeyIndex] = AnimationDataInComponentSpace[BoneIndex][KeyIndex].GetRotation();
+							}
+						}
+
+						// apply translation - for now only one
+						const TArray<FRigTransformConstraint> & TranslationTransformConstraint = RigConstraint->Constraints[EControlConstraint::Type::Translation].TransformConstraints;
+
+						if (TranslationTransformConstraint.Num() > 0)
+						{
+							const FName & ParentSpace = TranslationTransformConstraint[0].ParentSpace;
+							const FName & ParentBoneName = Skeleton->GetRigBoneMapping(ParentSpace);
+							const int32 & ParentBoneIndex = Skeleton->GetReferenceSkeleton().FindBoneIndex(ParentBoneName);
+							if (ParentBoneIndex != INDEX_NONE)
+							{
+								// if no rig control, component space is used
+								for (int32 KeyIndex = 0; KeyIndex < NumFrames; ++KeyIndex)
+								{
+									FTransform ParentTransform = AnimationDataInComponentSpace[ParentBoneIndex][KeyIndex];
+									FTransform RelativeTransform = AnimationDataInComponentSpace[BoneIndex][KeyIndex].GetRelativeTransform(ParentTransform);
+									Track.PosKeys[KeyIndex] = RelativeTransform.GetTranslation();
+									Track.ScaleKeys[KeyIndex] = RelativeTransform.GetScale3D();
+								}
+							}
+							else
+							{
+								for (int32 KeyIndex = 0; KeyIndex < NumFrames; ++KeyIndex)
+								{
+									Track.PosKeys[KeyIndex] = AnimationDataInComponentSpace[BoneIndex][KeyIndex].GetTranslation();
+									Track.ScaleKeys[KeyIndex] = AnimationDataInComponentSpace[BoneIndex][KeyIndex].GetScale3D();
+								}
+							}
+						}
+						else
+						{
+							for (int32 KeyIndex = 0; KeyIndex < NumFrames; ++KeyIndex)
+							{
+								Track.PosKeys[KeyIndex] = AnimationDataInComponentSpace[BoneIndex][KeyIndex].GetTranslation();
+								Track.ScaleKeys[KeyIndex] = AnimationDataInComponentSpace[BoneIndex][KeyIndex].GetScale3D();
+							}
+						}
+					}
+					else
+					{
+						// if no rig control, component space is used
+						for (int32 KeyIndex = 0; KeyIndex < NumFrames; ++KeyIndex)
+						{
+							Track.PosKeys[KeyIndex] = AnimationDataInComponentSpace[BoneIndex][KeyIndex].GetTranslation();
+							Track.RotKeys[KeyIndex] = AnimationDataInComponentSpace[BoneIndex][KeyIndex].GetRotation();
+							Track.ScaleKeys[KeyIndex] = AnimationDataInComponentSpace[BoneIndex][KeyIndex].GetScale3D();
+						}
+					}
+				}
+				else
+				{
+					// if no bone index, we still put identity
+					for (int32 KeyIndex = 0; KeyIndex < NumFrames; ++KeyIndex)
+					{
+						Track.PosKeys[KeyIndex] = FVector::ZeroVector;
+						Track.RotKeys[KeyIndex] = FQuat::Identity;
+						Track.ScaleKeys[KeyIndex] = FVector(1.f);
+					}
+				}
+			}
+		}
+
+		return true;
+	}
+
+	return false;
+}
+
+bool UAnimSequence::ConvertRiggingDataToAnimationData(TArray<struct FRawAnimSequenceTrack> & RawRiggingAnimationData)
+{
+	if (RawRiggingAnimationData.Num() > 0)
+	{
+		TArray< TArray<FTransform> > AnimationDataInComponentSpace;
+		int32 NumBones = GetSpaceBasedAnimationData(AnimationDataInComponentSpace, &RawRiggingAnimationData);
+
+		TArray<FRawAnimSequenceTrack> OldAnimationData = RawAnimationData;
+		TArray<FName> OldAnimationTrackNames = AnimationTrackNames;
+		int32 NumNodes = RawRiggingAnimationData.Num();
+		// get local spaces
+		// add all tracks?
+		AnimationTrackNames.Empty(NumNodes);
+		AnimationTrackNames.AddUninitialized(NumNodes);
+		RawAnimationData.Empty(NumNodes);
+		RawAnimationData.AddZeroed(NumNodes);
+
+		USkeleton * Skeleton = GetSkeleton();
+		const FReferenceSkeleton & RefSkeleton = Skeleton->GetReferenceSkeleton();
+		const URig * Rig = Skeleton->GetRig();
+		for (int32 NodeIndex = 0; NodeIndex < NumNodes; ++NodeIndex)
+		{
+			FName BoneName = Skeleton->GetRigBoneMapping(Rig->GetNodeName(NodeIndex));
+			int32 BoneIndex = RefSkeleton.FindBoneIndex(BoneName);
+
+			if (BoneIndex != INDEX_NONE)
+			{
+				// add track names
+				AnimationTrackNames[NodeIndex] = BoneName;
+
+				// update bone trasfnrom
+				FRawAnimSequenceTrack & Track = RawAnimationData[NodeIndex];
+
+				Track.PosKeys.Empty();
+				Track.RotKeys.Empty();
+				Track.ScaleKeys.Empty();
+				Track.PosKeys.AddUninitialized(NumFrames);
+				Track.RotKeys.AddUninitialized(NumFrames);
+				Track.ScaleKeys.AddUninitialized(NumFrames);
+
+				const int32 & ParentBoneIndex = RefSkeleton.GetParentIndex(BoneIndex);
+
+				if(ParentBoneIndex != INDEX_NONE)
+				{
+					for(int32 KeyIndex = 0; KeyIndex < NumFrames; ++KeyIndex)
+					{
+						FTransform LocalTransform = AnimationDataInComponentSpace[BoneIndex][KeyIndex].GetRelativeTransform(AnimationDataInComponentSpace[ParentBoneIndex][KeyIndex]);
+
+						Track.PosKeys[KeyIndex] = LocalTransform.GetTranslation();
+						Track.RotKeys[KeyIndex] = LocalTransform.GetRotation();
+						Track.ScaleKeys[KeyIndex] = LocalTransform.GetScale3D();
+					}
+				}
+				else
+				{
+					for(int32 KeyIndex = 0; KeyIndex < NumFrames; ++KeyIndex)
+					{
+						FTransform LocalTransform = AnimationDataInComponentSpace[BoneIndex][KeyIndex];
+
+						Track.PosKeys[KeyIndex] = LocalTransform.GetTranslation();
+						Track.RotKeys[KeyIndex] = LocalTransform.GetRotation();
+						Track.ScaleKeys[KeyIndex] = LocalTransform.GetScale3D();
+					}
+				}
+			}
+		}
+
+		// recreate track map
+		TrackToSkeletonMapTable.Empty(AnimationTrackNames.Num());
+		TrackToSkeletonMapTable.AddUninitialized(AnimationTrackNames.Num());
+		int32 TrackIdx = 0;
+		for (auto TrackName : AnimationTrackNames)
+		{
+			TrackToSkeletonMapTable[TrackIdx++].BoneTreeIndex = Skeleton->GetReferenceSkeleton().FindBoneIndex(TrackName);
+		}
+		PostProcessSequence();
+
+		return true;
+	}
+
+	return false;
+}
+
 #endif
 /*-----------------------------------------------------------------------------
 	AnimNotify & subclasses
diff --git a/Engine/Source/Runtime/Engine/Private/Animation/AnimSingleNodeInstance.cpp b/Engine/Source/Runtime/Engine/Private/Animation/AnimSingleNodeInstance.cpp
--- a/Engine/Source/Runtime/Engine/Private/Animation/AnimSingleNodeInstance.cpp
+++ b/Engine/Source/Runtime/Engine/Private/Animation/AnimSingleNodeInstance.cpp
@@ -23,11 +23,11 @@
 
 UAnimSingleNodeInstance::UAnimSingleNodeInstance(const class FPostConstructInitializeProperties& PCIP)
 	: Super(PCIP)
+	, PlayRate(1.f)
+	, bLooping(true)
+	, bPlaying(true)
+	, bReverse(false)
 {
-	PlayRate = 1.0f;
-	bLooping = true;
-	bPlaying = true;
-	bReverse = false;
 }
 
 void UAnimSingleNodeInstance::SetAnimationAsset(class UAnimationAsset* NewAsset,bool bIsLooping,float InPlayRate)
diff --git a/Engine/Source/Runtime/Engine/Private/Animation/AnimationRuntime.cpp b/Engine/Source/Runtime/Engine/Private/Animation/AnimationRuntime.cpp
--- a/Engine/Source/Runtime/Engine/Private/Animation/AnimationRuntime.cpp
+++ b/Engine/Source/Runtime/Engine/Private/Animation/AnimationRuntime.cpp
@@ -700,6 +700,25 @@ void FAnimationRuntime::FillWithRefPose(TArray<FTransform> & OutAtoms, const FBo
 	}
 }
 
+void FAnimationRuntime::FillWithRetargetBaseRefPose( TArray<FTransform> & OutAtoms, const USkeletalMesh * Mesh, const FBoneContainer & RequiredBones )
+{
+	// Copy Target Asset's ref pose.
+	if (Mesh)
+	{
+		const TArray<FBoneIndexType> & BoneIndices = RequiredBones.GetBoneIndicesArray();
+		OutAtoms.Empty(BoneIndices.Num());
+		OutAtoms.AddUninitialized(BoneIndices.Num());
+		for (int32 Idx = 0; Idx < BoneIndices.Num(); ++Idx)
+		{
+			OutAtoms[Idx] = Mesh->RetargetBasePose[BoneIndices[Idx]];
+		}
+	}
+	else
+	{
+		OutAtoms.Empty();
+	}
+}
+
 void FAnimationRuntime::ConvertPoseToMeshSpace(const TArray<FTransform> & LocalTransforms, TArray<FTransform> & MeshSpaceTransforms, const FBoneContainer & RequiredBones)
 {
 	const int32 NumBones = RequiredBones.GetNumBones();
@@ -1181,6 +1200,62 @@ bool FAnimationRuntime::ContainsNaN(TArray<FBoneIndexType> & RequiredBoneIndices
 }
 #endif
 
+void FAnimationRuntime::FillUpSpaceBasesRefPose(const USkeleton * Skeleton, TArray<FTransform> &SpaceBaseRefPose)
+{
+	check(Skeleton);
+
+	const TArray<FTransform> & ReferencePose = Skeleton->GetReferenceSkeleton().GetRefBonePose();
+	SpaceBaseRefPose.Empty(ReferencePose.Num());
+	SpaceBaseRefPose.AddUninitialized(ReferencePose.Num());
+
+	// initialize to identity since some of them don't have tracks
+	for(int Index=0; Index <SpaceBaseRefPose.Num(); ++Index)
+	{
+		int32 ParentIndex = Skeleton->GetReferenceSkeleton().GetParentIndex(Index);
+		if(ParentIndex != INDEX_NONE)
+		{
+			SpaceBaseRefPose[Index] = ReferencePose[Index] * SpaceBaseRefPose[ParentIndex];
+		}
+		else
+		{
+			SpaceBaseRefPose[Index] = ReferencePose[Index];
+		}
+	}
+}
+
+void FAnimationRuntime::FillUpSpaceBasesRetargetBasePose(const USkeleton * Skeleton, TArray<FTransform> &SpaceBaseRefPose)
+{
+	check(Skeleton);
+
+	// @Todo fixme: this has to get preview mesh instead of skeleton
+	
+	const USkeletalMesh * PreviewMesh = Skeleton->GetPreviewMesh();
+	if (PreviewMesh)
+	{
+		const TArray<FTransform> & ReferencePose = PreviewMesh->RetargetBasePose;
+		SpaceBaseRefPose.Empty(ReferencePose.Num());
+		SpaceBaseRefPose.AddUninitialized(ReferencePose.Num());
+
+		// initialize to identity since some of them don't have tracks
+		for(int Index=0; Index <SpaceBaseRefPose.Num(); ++Index)
+		{
+			int32 ParentIndex = Skeleton->GetReferenceSkeleton().GetParentIndex(Index);
+			if(ParentIndex != INDEX_NONE)
+			{
+				SpaceBaseRefPose[Index] = ReferencePose[Index] * SpaceBaseRefPose[ParentIndex];
+			}
+			else
+			{
+				SpaceBaseRefPose[Index] = ReferencePose[Index];
+			}
+		}
+	}
+	else
+	{
+		FAnimationRuntime::FillUpSpaceBasesRefPose(Skeleton, SpaceBaseRefPose);
+	}
+}
+
 /////////////////////////////////////////////////////////////////////////////////////////
 // FA2CSPose
 /////////////////////////////////////////////////////////////////////////////////////////
diff --git a/Engine/Source/Runtime/Engine/Private/Animation/Rig.cpp b/Engine/Source/Runtime/Engine/Private/Animation/Rig.cpp
new file mode 100644
--- /dev/null
+++ b/Engine/Source/Runtime/Engine/Private/Animation/Rig.cpp
@@ -0,0 +1,447 @@
+// Copyright 1998-2014 Epic Games, Inc. All Rights Reserved.
+
+/*=============================================================================
+	URig.cpp: Rig functionality for sharing animations
+=============================================================================*/ 
+
+#include "EnginePrivate.h"
+#include "Animation/Rig.h"
+#include "AnimationRuntime.h"
+//@todo should move all this window stuff somewhere else. Persona?
+
+#define LOCTEXT_NAMESPACE "Rig"
+
+FName URig::WorldNodeName(TEXT("World"));
+
+URig::URig(const class FPostConstructInitializeProperties& PCIP)
+	: Super(PCIP)
+{
+}
+
+void URig::PostLoad()
+{
+	Super::PostLoad();
+	// find 
+	// display name is empty, copy from Name - this is for old data
+	for (auto &Node : Nodes )
+	{
+		if (Node.DisplayName.Len() == 0)
+		{
+			Node.DisplayName = Node.Name.ToString();
+		}
+	}
+}
+
+#if WITH_EDITOR
+void URig::PostEditChangeProperty(FPropertyChangedEvent& PropertyChangedEvent)
+{
+	Super::PostEditChangeProperty(PropertyChangedEvent);
+
+	// verify lots of things
+	// make sure it's unique name
+	// 
+}
+
+// for now these are privates since we don't have good control yet
+bool URig::AddNode(FName Name, FName ParentNode, FTransform Transform)
+{
+	int32 NodeIndex = FindNode(Name);
+
+	if (NodeIndex == INDEX_NONE)
+	{
+		Nodes.Add(FNode(Name, ParentNode, Transform));
+
+		return true;
+	}
+
+	return false;
+}
+
+bool URig::DeleteNode(FName Name)
+{
+	int32 NodeIndex = FindNode(Name);
+
+	if (NodeIndex != INDEX_NONE)
+	{
+		Nodes.RemoveAt(NodeIndex);
+		return true;
+	}
+
+	return false;
+}
+
+int32 URig::FindNode(const FName & NodeName) const
+{
+	int32 Index=0;
+	for (auto Node: Nodes)
+	{
+		if (Node.Name == NodeName)
+		{
+			return Index;
+		}
+
+		++Index;
+	}
+
+	return INDEX_NONE;
+}
+
+const FNode * URig::GetNode(int32 NodeIndex) const
+{
+	if ( Nodes.IsValidIndex(NodeIndex) )
+	{
+		return &Nodes[NodeIndex];
+	}
+
+	return NULL;
+}
+
+FName URig::GetNodeName(int32 NodeIndex) const
+{
+	if(Nodes.IsValidIndex(NodeIndex))
+	{
+		return Nodes[NodeIndex].Name;
+	}
+
+	return NAME_None;
+}
+
+FName URig::GetParentNode(FName & NodeName) const
+{
+	int32 NodeIndex = FindNode(NodeName);
+
+	if (NodeIndex != INDEX_NONE)
+	{
+		if (Nodes[NodeIndex].ParentName != NAME_None)
+		{
+			return Nodes[NodeIndex].ParentName;
+		}
+	}
+
+	return WorldNodeName;
+}
+
+bool URig::AddRigConstraint(FName NodeName, EControlConstraint::Type ConstraintType, EConstraintTransform::Type	TransformType, FName ParentSpace, float Weight /*= 1.f*/)
+{
+	if (ConstraintType == EControlConstraint::Type::Max)
+	{
+		// invalid type
+		return false;
+	}
+
+	// make sure the ParentSpace is valid
+	int32 ParentIndex = FindNode(ParentSpace);
+	if (ParentIndex == INDEX_NONE)
+	{
+		// if parent is invalid, set it to World
+		ParentSpace = WorldNodeName;
+	}
+
+	int32 Index = FindRigControlByNodeName(NodeName);
+
+	if (Index == INDEX_NONE)
+	{
+		FRigControl NewRigControl;
+		NewRigControl.Node = NodeName;
+
+		if (ConstraintType < EControlConstraint::Max)
+		{
+			FRigTransformConstraint NewTransformConstraint;
+			NewTransformConstraint.TranformType = TransformType;
+			NewTransformConstraint.ParentSpace = ParentSpace;
+			NewTransformConstraint.Weight = Weight;
+
+			NewRigControl.Constraints[ConstraintType].TransformConstraints.Add(NewTransformConstraint);
+
+			RigControls.Add(NewRigControl);
+		}
+	}
+	else
+	{
+		// it exists already, need to make sure we don't have different constraint types
+		FRigControl & RigControl = RigControls[Index];
+
+		if (ConstraintType < EControlConstraint::Max)
+		{
+			FRigControlConstraint & ControlConstraint = RigControl.Constraints[ConstraintType];
+
+			FRigTransformConstraint NewTransformConstraint;
+			NewTransformConstraint.TranformType = TransformType;
+			NewTransformConstraint.ParentSpace = ParentSpace;
+			NewTransformConstraint.Weight = Weight;
+
+			// add new transform constraint
+			ControlConstraint.TransformConstraints.Add(NewTransformConstraint);
+		}
+	}
+
+	return true;
+}
+
+int32 URig::GetNodeNum() const
+{
+	return Nodes.Num();
+}
+
+int32 URig::GetRigControlNum() const
+{
+	return RigControls.Num();
+}
+
+const FRigControl * URig::GetRigControl(int32 ControlIndex) const
+{
+	if (RigControls.IsValidIndex(ControlIndex))
+	{
+		return &RigControls[ControlIndex];
+	}
+
+	return NULL;
+}
+
+const FRigControl * URig::GetRigControlByNodeName(FName NodeName) const
+{
+	int32 ControlIndex = FindRigControlByNodeName(NodeName);
+
+	if (ControlIndex != INDEX_NONE)
+	{
+		return &RigControls[ControlIndex];
+	}
+
+	return NULL;
+}
+
+int32 URig::FindTransformParentNode(int32 NodeIndex, bool bTranslate, int32 Index/*=0*/) const
+{
+	if (Nodes.IsValidIndex(NodeIndex))
+	{
+		const FRigControl * RigControl = GetRigControlByNodeName(Nodes[NodeIndex].Name);
+
+		if(RigControl)
+		{
+			FName ParentNodeName = NAME_None;
+			if(bTranslate)
+			{
+				ParentNodeName = RigControl->Constraints[EControlConstraint::Type::Translation].TransformConstraints[Index].ParentSpace;
+			}
+			else
+			{
+				ParentNodeName = RigControl->Constraints[EControlConstraint::Type::Orientation].TransformConstraints[Index].ParentSpace;
+			}
+
+			if(ParentNodeName != NAME_None)
+			{
+				return  FindNode(ParentNodeName);
+			}
+		}
+	}
+
+	return INDEX_NONE;
+}
+
+int32 URig::FindRigControlByNodeName(FName NodeName) const
+{
+	int32 Index=0;
+
+	for (auto RigControl : RigControls)
+	{
+		if (RigControl.Node == NodeName)
+		{
+			return Index;
+		}
+
+		++Index;
+	}
+
+	return INDEX_NONE;
+}
+
+void URig::CreateFromSkeleton(const USkeleton * Skeleton, const TMap<int32, int32> & RequiredBones)
+{
+	// show dialog to choose which bone to add
+	if(RequiredBones.Num() >0)
+	{
+		const FReferenceSkeleton & RefSkeleton = Skeleton->GetReferenceSkeleton();
+		TArray<FTransform> SpaceBaseRefPose;
+
+		FAnimationRuntime::FillUpSpaceBasesRefPose(Skeleton, SpaceBaseRefPose);
+
+		// once selected, add node to the rig
+		for (auto It = RequiredBones.CreateConstIterator(); It; ++It)
+		{
+			int32 BoneIndex = (int32)(It.Key());
+			int32 ParentIndex = (int32)(It.Value());
+
+			check (BoneIndex != INDEX_NONE);
+
+			FName BoneName = RefSkeleton.GetBoneName(BoneIndex);
+			FName ParentBoneName = WorldNodeName;
+			if (ParentIndex != INDEX_NONE)
+			{
+				ParentBoneName = RefSkeleton.GetBoneName(ParentIndex);
+			}
+
+			const FTransform & Transform = SpaceBaseRefPose[BoneIndex];
+
+			AddNode(BoneName, ParentBoneName, Transform);
+		}
+
+		// add constraint to parent space and relative transform 
+		for(auto It = RequiredBones.CreateConstIterator(); It; ++It)
+		{
+			int32 BoneIndex = (int32)(It.Key());
+			int32 ParentIndex = (int32)(It.Value());
+
+			FName BoneName = RefSkeleton.GetBoneName(BoneIndex);
+
+			if (ParentIndex == INDEX_NONE)
+			{
+				AddRigConstraint(BoneName, EControlConstraint::Translation, EConstraintTransform::Absoluate, WorldNodeName, 1.f);
+				AddRigConstraint(BoneName, EControlConstraint::Orientation, EConstraintTransform::Absoluate, WorldNodeName, 1.f);
+			}
+			else
+			{
+				FName ParentBoneName = RefSkeleton.GetBoneName(ParentIndex);
+				AddRigConstraint(BoneName, EControlConstraint::Translation, EConstraintTransform::Absoluate, ParentBoneName, 1.f);
+				AddRigConstraint(BoneName, EControlConstraint::Orientation, EConstraintTransform::Absoluate, ParentBoneName, 1.f);
+			}
+		}
+	}
+}
+
+void URig::CalculateComponentSpace(int32 NodeIndex, const FTransform & LocalTransform, const TArray<FTransform> & TransformBuffer, const FGetParentIndex& DelegateToGetParentIndex, FTransform & OutComponentSpaceTransform) const
+{
+	int32 ConstraintIndex = FindRigControlByNodeName(Nodes[NodeIndex].Name);
+
+	if (ConstraintIndex != INDEX_NONE)
+	{
+		// now find transform constraint data
+		{
+			const FRigControlConstraint & Constraints = RigControls[ConstraintIndex].Constraints[EControlConstraint::Type::Orientation];
+
+			// for now we only care for the first one
+			const FName ParentName = Constraints.TransformConstraints[0].ParentSpace;
+			int32 ParentIndex = INDEX_NONE;
+
+			if(DelegateToGetParentIndex.IsBound())
+			{
+				ParentIndex = DelegateToGetParentIndex.Execute(ParentName);
+			}
+			else
+			{
+				ParentIndex = FindTransformParentNode(NodeIndex, false);
+			}
+
+			if(ParentIndex != INDEX_NONE && TransformBuffer.IsValidIndex(ParentIndex))
+			{
+				FQuat ParentRotation = TransformBuffer[ParentIndex].GetRotation();
+				OutComponentSpaceTransform.SetRotation(ParentRotation * LocalTransform.GetRotation());
+			}
+		}
+
+		// same thing for translation
+		{
+			const FRigControlConstraint & Constraints = RigControls[ConstraintIndex].Constraints[EControlConstraint::Type::Translation];
+
+			// for now we only care for the first one
+			const FName ParentName = Constraints.TransformConstraints[0].ParentSpace;
+			int32 ParentIndex = INDEX_NONE;
+
+			if(DelegateToGetParentIndex.IsBound())
+			{
+				ParentIndex = DelegateToGetParentIndex.Execute(ParentName);
+			}
+			else
+			{
+				ParentIndex = FindTransformParentNode(NodeIndex, false);
+			}
+
+			if(ParentIndex != INDEX_NONE && TransformBuffer.IsValidIndex(ParentIndex))
+			{
+				FTransform ParentTransform = TransformBuffer[ParentIndex];
+				OutComponentSpaceTransform.SetTranslation(ParentTransform.TransformVector(LocalTransform.GetTranslation()));
+			}
+		}
+		// @todo fix this
+		OutComponentSpaceTransform.SetScale3D(LocalTransform.GetScale3D());
+	}
+}
+
+void URig::CalculateLocalSpace(int32 NodeIndex, const FTransform & ComponentTransform, const TArray<FTransform> & TransformBuffer, const FGetParentIndex& DelegateToGetParentIndex, FTransform & OutLocalSpaceTransform) const
+{
+	int32 ConstraintIndex = FindRigControlByNodeName(Nodes[NodeIndex].Name);
+
+	if(ConstraintIndex != INDEX_NONE)
+	{
+		// now find transform constraint data
+		{
+			const FRigControlConstraint & Constraints = RigControls[ConstraintIndex].Constraints[EControlConstraint::Type::Orientation];
+
+			// for now we only care for the first one
+			const FName ParentName = Constraints.TransformConstraints[0].ParentSpace;
+			int32 ParentIndex = INDEX_NONE;
+
+			if(DelegateToGetParentIndex.IsBound())
+			{
+				ParentIndex = DelegateToGetParentIndex.Execute(ParentName);
+			}
+			else
+			{
+				ParentIndex = FindTransformParentNode(NodeIndex, false);
+			}
+
+			if(ParentIndex != INDEX_NONE && TransformBuffer.IsValidIndex(ParentIndex))
+			{
+				FQuat ParentRotation = TransformBuffer[ParentIndex].GetRotation();
+				OutLocalSpaceTransform.SetRotation(ParentRotation.Inverse() * ComponentTransform.GetRotation());
+			}
+		}
+
+		// same thing for translation
+		{
+			const FRigControlConstraint & Constraints = RigControls[ConstraintIndex].Constraints[EControlConstraint::Type::Translation];
+
+			// for now we only care for the first one
+			const FName ParentName = Constraints.TransformConstraints[0].ParentSpace;
+			int32 ParentIndex = INDEX_NONE;
+
+			if(DelegateToGetParentIndex.IsBound())
+			{
+				ParentIndex = DelegateToGetParentIndex.Execute(ParentName);
+			}
+			else
+			{
+				ParentIndex = FindTransformParentNode(NodeIndex, false);
+			}
+
+			if(ParentIndex != INDEX_NONE && TransformBuffer.IsValidIndex(ParentIndex))
+			{
+				FTransform LocalTransform = ComponentTransform.GetRelativeTransform(TransformBuffer[ParentIndex]);
+				OutLocalSpaceTransform.SetTranslation(LocalTransform.GetTranslation());
+			}
+		}
+
+		// @todo fix this
+		OutLocalSpaceTransform.SetScale3D(ComponentTransform.GetScale3D());
+	}
+}
+
+void URig::SetAllConstraintsToParents()
+{
+	for(auto & Control : RigControls)
+	{
+		FName ParentNode = GetParentNode(Control.Node);
+
+		Control.Constraints[EControlConstraint::Type::Translation].TransformConstraints[0].ParentSpace = ParentNode;
+		Control.Constraints[EControlConstraint::Type::Orientation].TransformConstraints[0].ParentSpace = ParentNode;
+	}
+}
+void URig::SetAllConstraintsToWorld()
+{
+	for (auto & Control : RigControls)
+	{
+		Control.Constraints[EControlConstraint::Type::Translation].TransformConstraints[0].ParentSpace = WorldNodeName;
+		Control.Constraints[EControlConstraint::Type::Orientation].TransformConstraints[0].ParentSpace = WorldNodeName;
+	}
+}
+#endif // WITH_EDITOR
+#undef LOCTEXT_NAMESPACE 
\ No newline at end of file
diff --git a/Engine/Source/Runtime/Engine/Private/Animation/Skeleton.cpp b/Engine/Source/Runtime/Engine/Private/Animation/Skeleton.cpp
--- a/Engine/Source/Runtime/Engine/Private/Animation/Skeleton.cpp
+++ b/Engine/Source/Runtime/Engine/Private/Animation/Skeleton.cpp
@@ -9,6 +9,7 @@
 #include "AnimationRuntime.h"
 #include "AssetRegistryModule.h"
 #include "Animation/AnimSequence.h"
+#include "Animation/Rig.h"
 
 #define LOCTEXT_NAMESPACE "Skeleton"
 #define ROOT_BONE_PARENT	INDEX_NONE
@@ -755,6 +756,11 @@ USkeletalMesh* USkeleton::GetPreviewMesh(bool bFindIfNotSet)
 	return PreviewMesh;
 }
 
+USkeletalMesh* USkeleton::GetPreviewMesh() const
+{
+	return PreviewSkeletalMesh.Get();
+}
+
 void USkeleton::SetPreviewMesh(USkeletalMesh* PreviewMesh, bool bMarkAsDirty/*=true*/)
 {
 	if (bMarkAsDirty)
@@ -977,4 +983,127 @@ void USkeleton::RegenerateGuid()
 	check(Guid.IsValid());
 }
 
+void USkeleton::SetRigConfig(URig * Rig)
+{
+	if (RigConfig.Rig != Rig)
+	{
+		RigConfig.Rig = Rig;
+		RigConfig.BoneMappingTable.Empty();
+
+		if (Rig)
+		{
+			const FReferenceSkeleton & RefSkeleton = GetReferenceSkeleton();
+			const TArray<FNode> & Nodes = Rig->GetNodes();
+			// now add bone mapping table
+			for (auto Node: Nodes)
+			{
+				// if find same bone, use that bone for mapping
+				if (RefSkeleton.FindBoneIndex(Node.Name) != INDEX_NONE)
+				{
+					RigConfig.BoneMappingTable.Add(FNameMapping(Node.Name, Node.Name));
+				}
+				else
+				{
+					RigConfig.BoneMappingTable.Add(FNameMapping(Node.Name));
+				}
+			}
+		}
+	}
+}
+
+int32 USkeleton::FindRigBoneMapping(const FName & NodeName) const
+{
+	int32 Index=0;
+	for(const auto & NameMap : RigConfig.BoneMappingTable)
+	{
+		if(NameMap.NodeName == NodeName)
+		{
+			return Index;
+		}
+
+		++Index;
+	}
+
+	return INDEX_NONE;
+}
+
+FName USkeleton::GetRigBoneMapping(const FName & NodeName) const
+{
+	int32 Index = FindRigBoneMapping(NodeName);
+
+	if (Index != INDEX_NONE)
+	{
+		return RigConfig.BoneMappingTable[Index].BoneName;
+	}
+
+	return NAME_None;
+}
+
+bool USkeleton::SetRigBoneMapping(const FName & NodeName, FName BoneName)
+{
+	// make sure it's valid
+	int32 BoneIndex = ReferenceSkeleton.FindBoneIndex(BoneName);
+
+	// @todo we need to have validation phase where you can't set same bone for different nodes
+	// but it might be annoying to do that right now since the tool is ugly
+	// so for now it lets you set everything, but in the future
+	// we'll have to add verification
+	if ( BoneIndex != INDEX_NONE )
+	{
+		int32 Index = FindRigBoneMapping(NodeName);
+
+		if(Index != INDEX_NONE)
+		{
+			RigConfig.BoneMappingTable[Index].BoneName = BoneName;
+			return true;
+		}
+	}
+
+
+	return false;
+}
+
+void USkeleton::RefreshRigConfig()
+{
+	if (RigConfig.Rig != NULL)
+	{
+		if (RigConfig.BoneMappingTable.Num() > 0)
+		{
+			// verify if any missing bones or anything
+			// remove if removed
+			for ( int32 TableId=0; TableId<RigConfig.BoneMappingTable.Num(); ++TableId )
+			{
+				auto & BoneMapping = RigConfig.BoneMappingTable[TableId];
+
+				if ( RigConfig.Rig->FindNode(BoneMapping.NodeName) == INDEX_NONE)
+				{
+					// if not contains, remove it
+					RigConfig.BoneMappingTable.RemoveAt(TableId);
+					--TableId;
+				}
+			}
+
+			// if the count doesn't match, there is missing nodes. 
+			if (RigConfig.Rig->GetNodeNum() != RigConfig.BoneMappingTable.Num())
+			{
+				int32 NodeNum = RigConfig.Rig->GetNodeNum();
+				for(int32 NodeId=0; NodeId<NodeNum; ++NodeId)
+				{
+					const auto * Node = RigConfig.Rig->GetNode(NodeId);
+
+					if (FindRigBoneMapping(Node->Name) == INDEX_NONE)
+					{
+						RigConfig.BoneMappingTable.Add(FNameMapping(Node->Name));
+					}
+				}
+			}
+		}
+	}
+}
+
+URig * USkeleton::GetRig() const
+{
+	return RigConfig.Rig;
+}
+
 #undef LOCTEXT_NAMESPACE 
\ No newline at end of file
diff --git a/Engine/Source/Runtime/Engine/Private/SkeletalMesh.cpp b/Engine/Source/Runtime/Engine/Private/SkeletalMesh.cpp
--- a/Engine/Source/Runtime/Engine/Private/SkeletalMesh.cpp
+++ b/Engine/Source/Runtime/Engine/Private/SkeletalMesh.cpp
@@ -2490,6 +2490,13 @@ void USkeletalMesh::CalculateInvRefMatrices()
 			// Precompute inverse so we can use from-refpose-skin vertices.
 			RefBasesInvMatrix[b] = RefBases[b].InverseFast(); 
 		}
+
+#if WITH_EDITORONLY_DATA
+		if(RetargetBasePose.Num() == 0)
+		{
+			RetargetBasePose = RefSkeleton.GetRefBonePose();
+		}
+#endif // WITH_EDITORONLY_DATA
 	}
 }
 
@@ -2685,6 +2692,13 @@ void USkeletalMesh::PostLoad()
 		// Rebuild name table.
 		RefSkeleton.RebuildNameToIndexMap();
 	}
+
+#if WITH_EDITORONLY_DATA
+	if (RetargetBasePose.Num() == 0)
+	{
+		RetargetBasePose = RefSkeleton.GetRefBonePose();
+	}
+#endif
 }
 
 void USkeletalMesh::GetAssetRegistryTags(TArray<FAssetRegistryTag>& OutTags) const
diff --git a/Engine/Source/Runtime/Engine/Private/SkeletalMeshComponent.cpp b/Engine/Source/Runtime/Engine/Private/SkeletalMeshComponent.cpp
--- a/Engine/Source/Runtime/Engine/Private/SkeletalMeshComponent.cpp
+++ b/Engine/Source/Runtime/Engine/Private/SkeletalMeshComponent.cpp
@@ -1026,7 +1026,7 @@ void USkeletalMeshComponent::PostAnimEvaluation(FAnimationEvaluationContext& Eva
 		SCOPE_CYCLE_COUNTER(STAT_UpdateLocalToWorldAndOverlaps);
 
 		// New bone positions need to be sent to render thread
-		UpdateComponentToWorld();
+		UpdateChildTransforms();
 
 		// animation often change overlap. 
 		UpdateOverlaps();
diff --git a/Engine/Source/Runtime/Engine/Public/AnimationRuntime.h b/Engine/Source/Runtime/Engine/Public/AnimationRuntime.h
--- a/Engine/Source/Runtime/Engine/Public/AnimationRuntime.h
+++ b/Engine/Source/Runtime/Engine/Public/AnimationRuntime.h
@@ -144,6 +144,9 @@ public:
 	/** Fill ref pose **/
 	static void FillWithRefPose(TArray<FTransform> & OutAtoms, const FBoneContainer & RequiredBones);
 
+	/** fill with retarget base ref pose but this isn't used during run-time, so it always copies all of them */
+	static void FillWithRetargetBaseRefPose( TArray<FTransform> & OutAtoms, const USkeletalMesh * Mesh, const FBoneContainer & RequiredBones );
+
 	/** Convert LocalTransforms into MeshSpaceTransforms over RequiredBones. */
 	static void ConvertPoseToMeshSpace(const TArray<FTransform> & LocalTransforms, TArray<FTransform> & MeshSpaceTransforms, const FBoneContainer & RequiredBones);
 
@@ -259,6 +262,10 @@ public:
 	static void SetSpaceTransform(FA2Pose& Pose, int32 Index, FTransform & NewTransform);
 	static void SetSpaceTransform(FA2CSPose& Pose, int32 Index, FTransform & NewTransform);
 
+	// space bases
+	static void FillUpSpaceBasesRefPose(const USkeleton * Skeleton, TArray<FTransform> &SpaceBaseRefPose);
+	static void FillUpSpaceBasesRetargetBasePose(const USkeleton * Skeleton, TArray<FTransform> &SpaceBaseRefPose);
+
 private:
 	/** 
 	* Blend Poses per bone weights : The BasePoses + BlendPoses(SourceIndex) * Blend Weights(BoneIndex)
