Camera Actors can now tick.
(TTP #342261)

diff --git a/Engine/Source/Runtime/Engine/Private/Camera/CameraActor.cpp b/Engine/Source/Runtime/Engine/Private/Camera/CameraActor.cpp
--- a/Engine/Source/Runtime/Engine/Private/Camera/CameraActor.cpp
+++ b/Engine/Source/Runtime/Engine/Private/Camera/CameraActor.cpp
@@ -30,6 +30,8 @@ ACameraActor::ACameraActor(const class FPostConstructInitializeProperties& PCIP)
 	AspectRatio_DEPRECATED = 1.777778f;
 	PostProcessBlendWeight_DEPRECATED = 1.0f;
 	// End of deprecated property initialization
+
+	PrimaryActorTick.bCanEverTick = true;
 }
 
 void ACameraActor::Serialize(FArchive& Ar)
