Fix non-editor build

diff --git a/Engine/Source/Runtime/Engine/Classes/Animation/Rig.h b/Engine/Source/Runtime/Engine/Classes/Animation/Rig.h
--- a/Engine/Source/Runtime/Engine/Classes/Animation/Rig.h
+++ b/Engine/Source/Runtime/Engine/Classes/Animation/Rig.h
@@ -149,7 +149,6 @@ public:
 #if WITH_EDITOR
 	// Begin UObject Interface
 	virtual void PostEditChangeProperty(FPropertyChangedEvent& PropertyChangedEvent) override;
-	virtual void PostLoad() override;
 	// End UObject Interface
 
 	// node related functions
diff --git a/Engine/Source/Runtime/Engine/Private/Animation/AnimationRuntime.cpp b/Engine/Source/Runtime/Engine/Private/Animation/AnimationRuntime.cpp
--- a/Engine/Source/Runtime/Engine/Private/Animation/AnimationRuntime.cpp
+++ b/Engine/Source/Runtime/Engine/Private/Animation/AnimationRuntime.cpp
@@ -700,6 +700,7 @@ void FAnimationRuntime::FillWithRefPose(TArray<FTransform> & OutAtoms, const FBo
 	}
 }
 
+#if WITH_EDITOR
 void FAnimationRuntime::FillWithRetargetBaseRefPose( TArray<FTransform> & OutAtoms, const USkeletalMesh * Mesh, const FBoneContainer & RequiredBones )
 {
 	// Copy Target Asset's ref pose.
@@ -718,6 +719,7 @@ void FAnimationRuntime::FillWithRetargetBaseRefPose( TArray<FTransform> & OutAto
 		OutAtoms.Empty();
 	}
 }
+#endif // WITH_EDITOR
 
 void FAnimationRuntime::ConvertPoseToMeshSpace(const TArray<FTransform> & LocalTransforms, TArray<FTransform> & MeshSpaceTransforms, const FBoneContainer & RequiredBones)
 {
@@ -1223,6 +1225,7 @@ void FAnimationRuntime::FillUpSpaceBasesRefPose(const USkeleton * Skeleton, TArr
 	}
 }
 
+#if WITH_EDITOR
 void FAnimationRuntime::FillUpSpaceBasesRetargetBasePose(const USkeleton * Skeleton, TArray<FTransform> &SpaceBaseRefPose)
 {
 	check(Skeleton);
@@ -1255,6 +1258,7 @@ void FAnimationRuntime::FillUpSpaceBasesRetargetBasePose(const USkeleton * Skele
 		FAnimationRuntime::FillUpSpaceBasesRefPose(Skeleton, SpaceBaseRefPose);
 	}
 }
+#endif // WITH_EDITOR
 
 /////////////////////////////////////////////////////////////////////////////////////////
 // FA2CSPose
diff --git a/Engine/Source/Runtime/Engine/Private/Animation/Rig.cpp b/Engine/Source/Runtime/Engine/Private/Animation/Rig.cpp
--- a/Engine/Source/Runtime/Engine/Private/Animation/Rig.cpp
+++ b/Engine/Source/Runtime/Engine/Private/Animation/Rig.cpp
@@ -18,20 +18,6 @@ URig::URig(const class FPostConstructInitializeProperties& PCIP)
 {
 }
 
-void URig::PostLoad()
-{
-	Super::PostLoad();
-	// find 
-	// display name is empty, copy from Name - this is for old data
-	for (auto &Node : Nodes )
-	{
-		if (Node.DisplayName.Len() == 0)
-		{
-			Node.DisplayName = Node.Name.ToString();
-		}
-	}
-}
-
 #if WITH_EDITOR
 void URig::PostEditChangeProperty(FPropertyChangedEvent& PropertyChangedEvent)
 {
diff --git a/Engine/Source/Runtime/Engine/Private/Animation/Skeleton.cpp b/Engine/Source/Runtime/Engine/Private/Animation/Skeleton.cpp
--- a/Engine/Source/Runtime/Engine/Private/Animation/Skeleton.cpp
+++ b/Engine/Source/Runtime/Engine/Private/Animation/Skeleton.cpp
@@ -983,6 +983,7 @@ void USkeleton::RegenerateGuid()
 	check(Guid.IsValid());
 }
 
+#if WITH_EDITOR
 void USkeleton::SetRigConfig(URig * Rig)
 {
 	if (RigConfig.Rig != Rig)
@@ -1105,5 +1106,5 @@ URig * USkeleton::GetRig() const
 {
 	return RigConfig.Rig;
 }
-
+#endif //WITH_EDITOR
 #undef LOCTEXT_NAMESPACE 
\ No newline at end of file
diff --git a/Engine/Source/Runtime/Engine/Public/AnimationRuntime.h b/Engine/Source/Runtime/Engine/Public/AnimationRuntime.h
--- a/Engine/Source/Runtime/Engine/Public/AnimationRuntime.h
+++ b/Engine/Source/Runtime/Engine/Public/AnimationRuntime.h
@@ -144,8 +144,10 @@ public:
 	/** Fill ref pose **/
 	static void FillWithRefPose(TArray<FTransform> & OutAtoms, const FBoneContainer & RequiredBones);
 
+#if WITH_EDITOR
 	/** fill with retarget base ref pose but this isn't used during run-time, so it always copies all of them */
 	static void FillWithRetargetBaseRefPose( TArray<FTransform> & OutAtoms, const USkeletalMesh * Mesh, const FBoneContainer & RequiredBones );
+#endif
 
 	/** Convert LocalTransforms into MeshSpaceTransforms over RequiredBones. */
 	static void ConvertPoseToMeshSpace(const TArray<FTransform> & LocalTransforms, TArray<FTransform> & MeshSpaceTransforms, const FBoneContainer & RequiredBones);
@@ -264,8 +266,9 @@ public:
 
 	// space bases
 	static void FillUpSpaceBasesRefPose(const USkeleton * Skeleton, TArray<FTransform> &SpaceBaseRefPose);
+#if WITH_EDITOR
 	static void FillUpSpaceBasesRetargetBasePose(const USkeleton * Skeleton, TArray<FTransform> &SpaceBaseRefPose);
-
+#endif
 private:
 	/** 
 	* Blend Poses per bone weights : The BasePoses + BlendPoses(SourceIndex) * Blend Weights(BoneIndex)
