UE4 - Canvas
- Fixes upside rendering on GL (at a higher level)
- Add canvas draw events

diff --git a/Engine/Source/Runtime/Engine/Private/TileRendering.cpp b/Engine/Source/Runtime/Engine/Private/TileRendering.cpp
--- a/Engine/Source/Runtime/Engine/Private/TileRendering.cpp
+++ b/Engine/Source/Runtime/Engine/Private/TileRendering.cpp
@@ -134,12 +134,12 @@ public:
 };
 TGlobalResource<FTileMesh> GTileMesh;
 
-void FTileRenderer::DrawTile(FRHICommandListImmediate& RHICmdList, const class FSceneView& View, const FMaterialRenderProxy* MaterialRenderProxy, float X, float Y, float SizeX, float SizeY, float U, float V, float SizeU, float SizeV, bool bIsHitTesting, const FHitProxyId HitProxyId)
+void FTileRenderer::DrawTile(FRHICommandListImmediate& RHICmdList, const class FSceneView& View, const FMaterialRenderProxy* MaterialRenderProxy, bool bNeedsToSwitchVerticalAxis, float X, float Y, float SizeX, float SizeY, float U, float V, float SizeU, float SizeV, bool bIsHitTesting, const FHitProxyId HitProxyId)
 {
 	FMaterialTileVertex DestVertex[4];
 
 	// create verts
-	if (RHINeedsToSwitchVerticalAxis(GRHIShaderPlatform))
+	if (bNeedsToSwitchVerticalAxis)
 	{
 		DestVertex[0].Initialize(X + SizeX, View.ViewRect.Height() - (Y + SizeY), U + SizeU, V + SizeV);
 		DestVertex[1].Initialize(X, View.ViewRect.Height() - (Y + SizeY), U, V + SizeV);
diff --git a/Engine/Source/Runtime/Engine/Private/UnrealClient.cpp b/Engine/Source/Runtime/Engine/Private/UnrealClient.cpp
--- a/Engine/Source/Runtime/Engine/Private/UnrealClient.cpp
+++ b/Engine/Source/Runtime/Engine/Private/UnrealClient.cpp
@@ -1074,7 +1074,10 @@ void FViewport::Draw( bool bShouldPresent /*= true */)
 				auto World = ViewportClient->GetWorld();
 				FCanvas Canvas(this, NULL, World, World ? World->FeatureLevel : GRHIFeatureLevel);
 				{
+					// Make sure the Canvas is not rendered upside down
+					Canvas.SetAllowSwitchVerticalAxis(false);
 					ViewportClient->Draw(this, &Canvas);
+					Canvas.SetAllowSwitchVerticalAxis(true);
 				}
 				Canvas.Flush_GameThread();
 				ViewportClient->ProcessScreenShots(this);
diff --git a/Engine/Source/Runtime/Engine/Private/UserInterface/Canvas.cpp b/Engine/Source/Runtime/Engine/Private/UserInterface/Canvas.cpp
--- a/Engine/Source/Runtime/Engine/Private/UserInterface/Canvas.cpp
+++ b/Engine/Source/Runtime/Engine/Private/UserInterface/Canvas.cpp
@@ -65,6 +65,7 @@ void FCanvas::Construct()
 	check(RenderTarget);
 
 	bScaledToRenderTarget = false;
+	bAllowsToSwitchVerticalAxis = true;
 
 	// Push the viewport transform onto the stack.  Default to using a 2D projection. 
 	new(TransformStack) FTransformEntry( 
@@ -173,7 +174,7 @@ bool FCanvasBatchedElementRenderItem::Render_RenderThread(FRHICommandListImmedia
 			Gamma = 1.0f;
 		}
 
-		const bool bNeedsToSwitchVerticalAxis = RHINeedsToSwitchVerticalAxis(GRHIShaderPlatform);
+		const bool bNeedsToSwitchVerticalAxis = RHINeedsToSwitchVerticalAxis(GRHIShaderPlatform) || !Canvas->GetAllowSwitchVerticalAxis();
 
 
 		// draw batched items
@@ -217,11 +218,14 @@ bool FCanvasBatchedElementRenderItem::Render_GameThread(const FCanvas* Canvas)
 			Gamma = 1.0f;
 		}
 
+		const bool bNeedsToSwitchVerticalAxis = RHINeedsToSwitchVerticalAxis(Canvas->GetShaderPlatform()) || !Canvas->GetAllowSwitchVerticalAxis();
+
 		// Render the batched elements.
 		struct FBatchedDrawParameters
 		{
 			FRenderData* RenderData;
 			uint32 bHitTesting : 1;
+			uint32 bNeedsToSwitchVerticalAxis : 1;
 			uint32 ViewportSizeX;
 			uint32 ViewportSizeY;
 			float DisplayGamma;
@@ -234,6 +238,7 @@ bool FCanvasBatchedElementRenderItem::Render_GameThread(const FCanvas* Canvas)
 		{
 			Data,
 			Canvas->IsHitTesting(),
+			bNeedsToSwitchVerticalAxis,
 			(uint32)CanvasRenderTarget->GetSizeXY().X,
 			(uint32)CanvasRenderTarget->GetSizeXY().Y,
 			Gamma,
@@ -245,13 +250,13 @@ bool FCanvasBatchedElementRenderItem::Render_GameThread(const FCanvas* Canvas)
 			BatchedDrawCommand,
 			FBatchedDrawParameters,Parameters,DrawParameters,
 		{
-			const bool bNeedsToSwitchVerticalAxis = RHINeedsToSwitchVerticalAxis(Parameters.ShaderPlatform);
-				
+			SCOPED_DRAW_EVENT(RHICmdList, CanvasBatchedElements, DEC_CANVAS);
+
 			// draw batched items
 			Parameters.RenderData->BatchedElements.Draw(
 				RHICmdList,
 				Parameters.FeatureLevel,
-				bNeedsToSwitchVerticalAxis,
+				Parameters.bNeedsToSwitchVerticalAxis,
 				Parameters.RenderData->Transform.GetMatrix(),
 				Parameters.ViewportSizeX,
 				Parameters.ViewportSizeY,
@@ -306,6 +311,8 @@ bool FCanvasTileRendererItem::Render_RenderThread(FRHICommandListImmediate& RHIC
 	ViewInitOptions.BackgroundColor = FLinearColor::Black;
 	ViewInitOptions.OverlayColor = FLinearColor::White;
 
+	const bool bNeedsToSwitchVerticalAxis = RHINeedsToSwitchVerticalAxis(Canvas->GetShaderPlatform()) || !Canvas->GetAllowSwitchVerticalAxis();
+
 	FSceneView* View = new FSceneView(ViewInitOptions);
 
 	for( int32 TileIdx=0; TileIdx < Data->Tiles.Num(); TileIdx++ )
@@ -315,6 +322,7 @@ bool FCanvasTileRendererItem::Render_RenderThread(FRHICommandListImmediate& RHIC
 			RHICmdList,
 			*View, 
 			Data->MaterialRenderProxy, 
+			bNeedsToSwitchVerticalAxis,
 			Tile.X, Tile.Y, Tile.SizeX, Tile.SizeY, 
 			Tile.U, Tile.V, Tile.SizeU, Tile.SizeV,
 			Canvas->IsHitTesting(), Tile.HitProxyId
@@ -370,24 +378,28 @@ bool FCanvasTileRendererItem::Render_GameThread(const FCanvas* Canvas)
 
 	FSceneView* View = new FSceneView(ViewInitOptions);
 
+	const bool bNeedsToSwitchVerticalAxis = RHINeedsToSwitchVerticalAxis(Canvas->GetShaderPlatform()) || !Canvas->GetAllowSwitchVerticalAxis();
 	struct FDrawTileParameters
 	{
 		FSceneView* View;
 		FRenderData* RenderData;
 		uint32 bIsHitTesting : 1;
+		uint32 bNeedsToSwitchVerticalAxis : 1;
 		uint32 AllowedCanvasModes;
 	};
 	FDrawTileParameters DrawTileParameters =
 	{
 		View,
 		Data,
+		bNeedsToSwitchVerticalAxis,
 		Canvas->IsHitTesting(),
 		Canvas->GetAllowedModes()
 	};
 	ENQUEUE_UNIQUE_RENDER_COMMAND_ONEPARAMETER(
 		DrawTileCommand,
 		FDrawTileParameters, Parameters, DrawTileParameters,
-		{
+	{
+		SCOPED_DRAW_EVENT(RHICmdList, CanvasDrawTile, DEC_CANVAS);
 		for (int32 TileIdx = 0; TileIdx < Parameters.RenderData->Tiles.Num(); TileIdx++)
 		{
 			const FRenderData::FTileInst& Tile = Parameters.RenderData->Tiles[TileIdx];
@@ -395,6 +407,7 @@ bool FCanvasTileRendererItem::Render_GameThread(const FCanvas* Canvas)
 				RHICmdList,
 				*Parameters.View,
 				Parameters.RenderData->MaterialRenderProxy,
+				Parameters.bNeedsToSwitchVerticalAxis,
 				Tile.X, Tile.Y, Tile.SizeX, Tile.SizeY,
 				Tile.U, Tile.V, Tile.SizeU, Tile.SizeV,
 				Parameters.bIsHitTesting, Tile.HitProxyId
@@ -552,7 +565,7 @@ void FCanvas::Flush_RenderThread(FRHICommandListImmediate& RHICmdList, bool bFor
 	// sort the array of FCanvasSortElement entries so that higher sort keys render first (back-to-front)
 	SortedElements.Sort(FCompareFCanvasSortElement());
 
-	SCOPED_DRAW_EVENT(RHICmdList, CanvasFlush, DEC_SCENE_ITEMS);
+	SCOPED_DRAW_EVENT(RHICmdList, CanvasFlush, DEC_CANVAS);
 	const FTexture2DRHIRef& RenderTargetTexture = RenderTarget->GetRenderTargetTexture();
 
 	check(IsValidRef(RenderTargetTexture));
@@ -644,11 +657,13 @@ void FCanvas::Flush_GameThread(bool bForce)
 		ViewRect,
 		RenderTarget
 	};
+	bool bEmitCanvasDrawEvents = GEmitDrawEvents;
+
 	ENQUEUE_UNIQUE_RENDER_COMMAND_ONEPARAMETER(
 		CanvasFlushSetupCommand,
 		FCanvasFlushParameters,Parameters,FlushParameters,
 	{
-		SCOPED_DRAW_EVENT(RHICmdList, CanvasFlush, DEC_SCENE_ITEMS);
+		SCOPED_DRAW_EVENT(RHICmdList, CanvasFlush, DEC_CANVAS);
 
 		// Set the RHI render target.
 		::SetRenderTarget(RHICmdList, Parameters.CanvasRenderTarget->GetRenderTargetTexture(), FTextureRHIRef());
@@ -794,14 +809,15 @@ void FCanvas::Clear(const FLinearColor& Color)
 		ClearCommand,
 		FColor,Color,GammaCorrectedColor,
 		FRenderTarget*,CanvasRenderTarget,GetRenderTarget(),
+	{
+		SCOPED_DRAW_EVENT(RHICmdList, CanvasClear, DEC_CANVAS);
+		if( CanvasRenderTarget )
 		{
-			if( CanvasRenderTarget )
-			{
-				::SetRenderTarget(RHICmdList, CanvasRenderTarget->GetRenderTargetTexture(),FTextureRHIRef());
-				RHICmdList.SetViewport(0,0,0.0f,CanvasRenderTarget->GetSizeXY().X,CanvasRenderTarget->GetSizeXY().Y,1.0f);
-			}
-			RHICmdList.Clear(true,Color,false,0.0f,false,0, FIntRect());
-		});
+			::SetRenderTarget(RHICmdList, CanvasRenderTarget->GetRenderTargetTexture(),FTextureRHIRef());
+			RHICmdList.SetViewport(0,0,0.0f,CanvasRenderTarget->GetSizeXY().X,CanvasRenderTarget->GetSizeXY().Y,1.0f);
+		}
+		RHICmdList.Clear(true,Color,false,0.0f,false,0, FIntRect());
+	});
 }
 
 void FCanvas::DrawTile( float X, float Y, float SizeX,	float SizeY, float U, float V, float SizeU, float SizeV, const FLinearColor& Color,	const FTexture* Texture, bool AlphaBlend )
diff --git a/Engine/Source/Runtime/Engine/Public/CanvasTypes.h b/Engine/Source/Runtime/Engine/Public/CanvasTypes.h
--- a/Engine/Source/Runtime/Engine/Public/CanvasTypes.h
+++ b/Engine/Source/Runtime/Engine/Public/CanvasTypes.h
@@ -208,7 +208,7 @@ public:
 	 * Sets a rect that should be used to offset rendering into the viewport render target
 	 * If not set the canvas will render to the full target
 	 *
- 	 * @param ViewRect The rect to use
+	 * @param ViewRect The rect to use
 	 */
 	ENGINE_API void SetRenderTargetRect( const FIntRect& ViewRect );
 
@@ -312,6 +312,11 @@ public:
 	*/
 	EShaderPlatform GetShaderPlatform() const { return ShaderPlatform; }
 
+	// Get/Set if this Canvas allows its batched elements to switch vertical axis (e.g., rendering to back buffer should never flip)
+	bool GetAllowSwitchVerticalAxis() const { return bAllowsToSwitchVerticalAxis; }
+
+	void SetAllowSwitchVerticalAxis(bool bInAllowsToSwitchVerticalAxis) { bAllowsToSwitchVerticalAxis = bInAllowsToSwitchVerticalAxis; }
+
 public:
 	float AlphaModulate;
 
@@ -382,6 +387,8 @@ private:
 	float CurrentDeltaWorldTime;
 	/** true, if Canvas should be scaled to whole render target */
 	bool bScaledToRenderTarget;
+	// True if canvas allows switching vertical axis; false will ignore any flip
+	bool bAllowsToSwitchVerticalAxis;
 	/** Feature level that we are currently rendering with */
 	ERHIFeatureLevel::Type FeatureLevel;
 	/** Shader platform that we are currently rendering with */
diff --git a/Engine/Source/Runtime/Engine/Public/TileRendering.h b/Engine/Source/Runtime/Engine/Public/TileRendering.h
--- a/Engine/Source/Runtime/Engine/Public/TileRendering.h
+++ b/Engine/Source/Runtime/Engine/Public/TileRendering.h
@@ -15,7 +15,7 @@ public:
 	 * Draw a tile at the given location and size, using the given UVs
 	 * (UV = [0..1]
 	 */
-	ENGINE_API static void DrawTile(FRHICommandListImmediate& RHICmdList, const class FSceneView& View, const FMaterialRenderProxy* MaterialRenderProxy, float X, float Y, float SizeX, float SizeY, float U, float V, float SizeU, float SizeV, bool bIsHitTesting = false, const FHitProxyId HitProxyId = FHitProxyId());
+	ENGINE_API static void DrawTile(FRHICommandListImmediate& RHICmdList, const class FSceneView& View, const FMaterialRenderProxy* MaterialRenderProxy, bool bNeedsToSwitchVerticalAxis, float X, float Y, float SizeX, float SizeY, float U, float V, float SizeU, float SizeV, bool bIsHitTesting = false, const FHitProxyId HitProxyId = FHitProxyId());
 
 private:
 
