#ue4 - Added new virtual functions and variables for CharacterMovementComponent to extend falling movement and moving on bases.

- Added new function OnUnableToFollowBaseMove() to notify when UpdateBasedMovement() is unable to move the character the entire move distance.
- Added new functions to break out some logic in PhysFalling(), specifically dealing with AirControl: GetFallingLateralAcceleration(), GetAirControl(), BoostAirControl(), FindAirControlImpact(), LimitAirControl().
- Added new variables to control air control boosting, replacing old hard-coded values: AirControlBoostMultiplier, AirControlBoostVelocityThreshold.
- Deprecated PerformAirControl() in favor of PerformAirControlForPathFollowing() to remove confusion as to its intent.

diff --git a/Engine/Source/Runtime/Engine/Classes/GameFramework/CharacterMovementComponent.h b/Engine/Source/Runtime/Engine/Classes/GameFramework/CharacterMovementComponent.h
--- a/Engine/Source/Runtime/Engine/Classes/GameFramework/CharacterMovementComponent.h
+++ b/Engine/Source/Runtime/Engine/Classes/GameFramework/CharacterMovementComponent.h
@@ -186,7 +186,7 @@ public:
 	/**
 	 * Coefficient of friction. This property allows you to control how much friction is applied when moving across the ground, applying an opposing force that scales with current velocity.
 	 * This can be used to simulate slippery surfaces such as ice or oil by changing the value (possibly based on the material pawn is standing on).
-	 * See also: BrakingDecelerationWalking
+	 * @see BrakingDecelerationWalking
 	 */
 	UPROPERTY(Category="Character Movement", EditAnywhere, BlueprintReadWrite, meta=(ClampMin="0", UIMin="0"))
 	float GroundFriction;
@@ -216,8 +216,8 @@ public:
 	float MaxAcceleration;
 
 	/**
-	 * Deceleration when walking and not applying acceleration.
-	 * See also: GroundFriction
+	 * Deceleration when walking and not applying acceleration. This is a constant opposing force that directly lowers velocity by a constant value.
+	 * @see GroundFriction
 	 */
 	UPROPERTY(Category="Character Movement", EditAnywhere, BlueprintReadWrite, meta=(ClampMin="0", UIMin="0"))
 	float BrakingDecelerationWalking;
@@ -234,10 +234,18 @@ public:
 	UPROPERTY(Category="Character Movement", EditAnywhere, BlueprintReadWrite, meta=(ClampMin="0", UIMin="0"))
 	float BrakingDecelerationFlying;
 
-	/** Amount of lateral movement control available to the pawn when falling. 0 = no control, 1 = full control at max speed of MaxWalkSpeed. */
+	/** When falling, amount of lateral movement control available to the character. 0 = no control, 1 = full control at max speed of MaxWalkSpeed. */
 	UPROPERTY(Category="Character Movement", EditAnywhere, BlueprintReadWrite, meta=(ClampMin="0", UIMin="0"))
 	float AirControl;
 
+	/** When falling, multiplier applied to AirControl when lateral velocity is less than AirControlBoostVelocityThreshold. Setting this to zero will disable air control boosting. */
+	UPROPERTY(Category="Character Movement", EditAnywhere, BlueprintReadWrite, meta=(ClampMin="0", UIMin="0"))
+	float AirControlBoostMultiplier;
+
+	/** When falling, if lateral velocity magnitude is less than this value, AirControl is multiplied by AirControlBoostMultiplier. Setting this to zero will disable air control  boosting. */
+	UPROPERTY(Category="Character Movement", EditAnywhere, BlueprintReadWrite, meta=(ClampMin="0", UIMin="0"))
+	float AirControlBoostVelocityThreshold;
+
 	/** Friction to apply to lateral air movement when falling. */
 	UPROPERTY(Category="Character Movement", EditAnywhere, BlueprintReadWrite, meta=(ClampMin="0", UIMin="0"))
 	float FallingLateralFriction;
@@ -426,8 +434,10 @@ public:
 	 * Compute remaining time step given remaining time and current iterations.
 	 * The last iteration (limited by MaxSimulationIterations) always returns the remaining time, which may violate MaxSimulationTimeStep.
 	 *
-	 * @see MaxSimulationTimeStep
-	 * @see MaxSimulationIterations
+	 * @param RemainingTime		Remaining time in the tick.
+	 * @param Iterations		Current iteration of the tick (starting at 1).
+	 * @return The remaining time step to use for the next sub-step of iteration.
+	 * @see MaxSimulationTimeStep, MaxSimulationIterations
 	 */
 	float GetSimulationTimeStep(float RemainingTime, int32 Iterations) const;
 
@@ -708,9 +718,15 @@ public:
 	/** Return true if we have a valid CharacterOwner and UpdatedComponent. */
 	virtual bool HasValidData() const;
 
-	/** Update Velocity and Acceleration to air control in the desired Direction for AIControlled Pawn 
-		@PARAM Direction is the desired direction of movement
-		@Param ZDiff is the height difference between the destination and the Pawn's current position */
+	/**
+	 * Update Velocity and Acceleration to air control in the desired Direction for character using path following.
+	 * @param Direction is the desired direction of movement
+	 * @param ZDiff is the height difference between the destination and the Pawn's current position
+	 * @see RequestDirectMove()
+	*/
+	virtual void PerformAirControlForPathFollowing(FVector Direction, float ZDiff);
+
+	DEPRECATED(4.5, "PerformAirControl has been deprecated, use PerformAirControlForPathFollowing instead.")
 	virtual void PerformAirControl(FVector Direction, float ZDiff);
 
 	/** Transition from walking to falling */
@@ -905,9 +921,78 @@ public:
 	/** Get as close to waterline as possible, staying on same side as currently. */
 	FVector FindWaterLine(FVector Start, FVector End);
 
-	/** Handle changing physics to falling over deltatime and iterations */
+	/** Handle falling movement. */
 	virtual void PhysFalling(float deltaTime, int32 Iterations);
 
+	// Helpers for PhysFalling
+
+	/**
+	 * Get the acceleration to use during falling movement. Base implementation simply returns current Acceleration.
+	 * The Z component of the result is ignored, and the actual value of acceleration used while falling will be modified by the result of GetAirControl().
+	 * This function is used internally by PhysFalling().
+	 *
+	 * @param DeltaTime Time step for the current update.
+	 * @return Acceleration to use during falling movement.
+	 */
+	virtual FVector GetFallingLateralAcceleration(float DeltaTime);
+	
+	/**
+	 * Get the air control to use during falling movement. This is a multiplier applied to the acceleration used when falling.
+	 * Given an initial air control (TickAirControl), applies the result of BoostAirControl(). If non-zero and FindAirControlImpact()
+	 * returns true, we then apply the result of LimitAirControl().
+	 * This function is used internally by PhysFalling().
+	 *
+	 * @param DeltaTime			Time step for the current update.
+	 * @param TickAirControl	Current air control value.
+	 * @param FallAcceleration	Acceleration used during movement.
+	 * @return Air control to use during falling movement.
+	 * @see AirControl, BoostAirControl(), LimitAirControl()
+	 */
+	virtual float GetAirControl(float DeltaTime, float TickAirControl, const FVector& FallAcceleration);
+
+protected:
+
+	/**
+	 * Increase air control if conditions of AirControlBoostMultiplier and AirControlBoostVelocityThreshold are met.
+	 * This function is used internally by GetAirControl().
+	 *
+	 * @param DeltaTime			Time step for the current update.
+	 * @param TickAirControl	Current air control value.
+	 * @param FallAcceleration	Acceleration used during movement.
+	 * @return Modified air control to use during falling movement
+	 * @see GetAirControl()
+	 */
+	virtual float BoostAirControl(float DeltaTime, float TickAirControl, const FVector& FallAcceleration);
+
+	/**
+	 * Checks if air control will cause the player collision shape to hit something given current conditions.
+	 * This function is used internally by GetAirControl().
+	 *
+	 * @param DeltaTime			Time step for the current update.
+	 * @param TickAirControl	Current air control value.
+	 * @param FallAcceleration	Acceleration used during movement.
+	 * @param OutHitResult		Result of impact, valid if this function returns true.
+	 * @return True if there is an impact, in which case OutHitResult contains the result of that impact.
+	 * @see GetAirControl()
+	 */
+	virtual bool FindAirControlImpact(float DeltaTime, float TickAirControl, const FVector& FallAcceleration, FHitResult& OutHitResult);
+
+	/**
+	 * Limits the air control to use during falling movement, given an impact from FindAirControlImpact().
+	 * This function is used internally by GetAirControl().
+	 *
+	 * @param DeltaTime			Time step for the current update.
+	 * @param TickAirControl	Current air control value.
+	 * @param FallAcceleration	Acceleration used during movement.
+	 * @param HitResult			Result of impact from FindAirControlImpact().
+	 * @return Modified air control to use during falling movement
+	 * @see FindAirControlImpact()
+	 */
+	virtual float LimitAirControl(float DeltaTime, float TickAirControl, const FVector& FallAcceleration, const FHitResult& HitResult);
+	
+
+public:
+
 	/** Handle landing against Hit surface over remaingTime and iterations */
 	virtual void ProcessLanded(const FHitResult& Hit, float remainingTime, int32 Iterations);
 
@@ -992,7 +1077,7 @@ public:
 	 * Add force to character. Forces are accumulated each tick and applied together
 	 * so multiple calls to this function will accumulate.
 	 * Forces are scaled depending on timestep, so they can be applied each frame. If you want an
-	 * instantaneous force, use AdddImpulse.
+	 * instantaneous force, use AddImpulse.
 	 * Adding a force always takes the actor's mass into account.
 	 * Note that changing the momentum of characters like this can change the movement mode.
 	 * 
@@ -1261,6 +1346,15 @@ protected:
 	/** Internal function to call OnMovementUpdated delegate on CharacterOwner. */
 	virtual void CallMovementUpdateDelegate(float DeltaSeconds, const FVector& OldLocation, const FVector& OldVelocity);
 
+	/**
+	 * Event triggered when we are moving on a base but we are not able to move the full DeltaPosition because something has blocked us.
+	 * Note: MoveComponentFlags includes the flag to ignore the movement base while this event is fired.
+	 * @param DeltaPosition		How far we tried to move with the base.
+	 * @param OldLocation		Location before we tried to move with the base.
+	 * @param MoveOnBaseHit		Hit result for the object we hit when trying to move with the base.
+	 */
+	virtual void OnUnableToFollowBaseMove(const FVector& DeltaPosition, const FVector& OldLocation, const FHitResult& MoveOnBaseHit);
+
 public:
 	// Movement functions broken out based on owner's network Role.
 	// TickComponent calls the correct version based on the Role.
diff --git a/Engine/Source/Runtime/Engine/Private/CharacterMovementComponent.cpp b/Engine/Source/Runtime/Engine/Private/CharacterMovementComponent.cpp
--- a/Engine/Source/Runtime/Engine/Private/CharacterMovementComponent.cpp
+++ b/Engine/Source/Runtime/Engine/Private/CharacterMovementComponent.cpp
@@ -182,6 +182,8 @@ UCharacterMovementComponent::UCharacterMovementComponent(const class FPostConstr
 	MaxOutOfWaterStepHeight = 40.0f;
 	OutofWaterZ = 420.0f;
 	AirControl = 0.05f;
+	AirControlBoostMultiplier = 2.f;
+	AirControlBoostVelocityThreshold = 25.f;
 	FallingLateralFriction = 0.f;
 	MaxAcceleration = 2048.0f;
 	BrakingDecelerationWalking = MaxAcceleration;
@@ -710,8 +712,14 @@ void UCharacterMovementComponent::ApplyNetworkMovementMode(const uint8 ReceivedM
 	SetMovementMode(NetMovementMode, NetCustomMode);
 }
 
-
+// TODO: deprecated, remove
 void UCharacterMovementComponent::PerformAirControl(FVector Direction, float ZDiff)
+{
+	PerformAirControlForPathFollowing(Direction, ZDiff);
+}
+
+
+void UCharacterMovementComponent::PerformAirControlForPathFollowing(FVector Direction, float ZDiff)
 {
 	// use air control if low grav or above destination and falling towards it
 	if ( CharacterOwner && Velocity.Z < 0.f && (ZDiff < 0.f || GetGravityZ() > 0.9f * GetWorld()->GetDefaultGravityZ()))
@@ -1293,7 +1301,13 @@ void UCharacterMovementComponent::UpdateBasedMovement(float DeltaSeconds)
 			}
 			else
 			{
-				MoveUpdatedComponent( DeltaPosition, FinalQuat.Rotator(), true );
+				FHitResult MoveOnBaseHit(1.f);
+				const FVector OldLocation = UpdatedComponent->GetComponentLocation();
+				MoveUpdatedComponent(DeltaPosition, FinalQuat.Rotator(), true, &MoveOnBaseHit);
+				if ((UpdatedComponent->GetComponentLocation() - (OldLocation + DeltaPosition)).IsNearlyZero() == false)
+				{
+					OnUnableToFollowBaseMove(DeltaPosition, OldLocation, MoveOnBaseHit);
+				}
 			}
 		}
 
@@ -1305,6 +1319,12 @@ void UCharacterMovementComponent::UpdateBasedMovement(float DeltaSeconds)
 }
 
 
+void UCharacterMovementComponent::OnUnableToFollowBaseMove(const FVector& DeltaPosition, const FVector& OldLocation, const FHitResult& MoveOnBaseHit)
+{
+	// no default implementation, left for subclasses to override.
+}
+
+
 void UCharacterMovementComponent::UpdateBasedRotation(FRotator &FinalRotation, const FRotator& ReducedRotation)
 {
 	AController* Controller = CharacterOwner ? CharacterOwner->Controller : NULL;
@@ -2214,7 +2234,7 @@ void UCharacterMovementComponent::RequestDirectMove(const FVector& MoveVelocity,
 	if (IsFalling())
 	{
 		const FVector FallVelocity = MoveVelocity.ClampMaxSize(GetMaxSpeed());
-		PerformAirControl(FallVelocity, FallVelocity.Z);
+		PerformAirControlForPathFollowing(FallVelocity, FallVelocity.Z);
 		return;
 	}
 
@@ -2747,6 +2767,47 @@ void UCharacterMovementComponent::NotifyJumpApex()
 }
 
 
+FVector UCharacterMovementComponent::GetFallingLateralAcceleration(float DeltaTime)
+{
+	// No acceleration in Z
+	return FVector(Acceleration.X, Acceleration.Y, 0.f);
+}
+
+
+float UCharacterMovementComponent::GetAirControl(float DeltaTime, float TickAirControl, const FVector& FallAcceleration)
+{
+	// Boost
+	if (TickAirControl > 0.f)
+	{
+		TickAirControl = BoostAirControl(DeltaTime, TickAirControl, FallAcceleration);
+	}
+
+	// Limit
+	if (TickAirControl > 0.f)
+	{
+		FHitResult HitResult(1.f);
+		if (FindAirControlImpact(DeltaTime, TickAirControl, FallAcceleration, HitResult))
+		{
+			TickAirControl = LimitAirControl(DeltaTime, TickAirControl, FallAcceleration, HitResult);
+		}
+	}
+
+	return TickAirControl;
+}
+
+
+float UCharacterMovementComponent::BoostAirControl(float DeltaTime, float TickAirControl, const FVector& FallAcceleration)
+{
+	// Allow a burst of initial acceleration
+	if (AirControlBoostMultiplier > 0.f && Velocity.SizeSquared2D() < FMath::Square(AirControlBoostVelocityThreshold))
+	{
+		TickAirControl = FMath::Min(1.f, AirControlBoostMultiplier * TickAirControl);
+	}
+
+	return TickAirControl;
+}
+
+
 void UCharacterMovementComponent::PhysFalling(float deltaTime, int32 Iterations)
 {
 	if (deltaTime < MIN_TICK_TIME)
@@ -2758,66 +2819,19 @@ void UCharacterMovementComponent::PhysFalling(float deltaTime, int32 Iterations)
 	const float Speed2d = Velocity.Size2D();
 	const float BoundSpeed2d = FMath::Max(Speed2d, GetMaxSpeed() * AnalogInputModifier);
 
-	//bound acceleration, falling object has minimal ability to impact acceleration
-	FVector FallAcceleration = Acceleration;
+	FVector FallAcceleration = GetFallingLateralAcceleration(deltaTime);
 	FallAcceleration.Z = 0.f;
-	
-	FHitResult Hit(1.f);
-	if( !HasRootMotion() )
-	{
-		// test for slope to avoid using air control to climb walls
-		float TickAirControl = AirControl;
-		if ( TickAirControl > 0.0f && FallAcceleration.SizeSquared() > 0.f )
-		{
-			const float TestWalkTime = FMath::Max(deltaTime, 0.05f);
-			const FVector TestWalk = ((TickAirControl * GetMaxAcceleration() * FallAcceleration.SafeNormal() + FVector(0.f,0.f,GetGravityZ())) * TestWalkTime + Velocity) * TestWalkTime;
-			if(!TestWalk.IsZero())
-			{
-				static const FName FallingTraceParamsTag = FName(TEXT("PhysFalling"));
-				FHitResult Result(1.f);
-				FCollisionQueryParams CapsuleQuery(FallingTraceParamsTag, false, CharacterOwner);
-				FCollisionResponseParams ResponseParam;
-				InitCollisionParams(CapsuleQuery, ResponseParam);
-				const FVector PawnLocation = CharacterOwner->GetActorLocation();
-				const ECollisionChannel CollisionChannel = UpdatedComponent->GetCollisionObjectType();
-
-				bool bHit = false;
-				const FCollisionShape CapsuleShape = GetPawnCapsuleCollisionShape(SHRINK_None);
-				if(bUseFlatBaseForFloorChecks)
-				{
-					const FCollisionShape BoxShape = FCollisionShape::MakeBox(FVector(CapsuleShape.GetCapsuleRadius(), CapsuleShape.GetCapsuleRadius(), CapsuleShape.GetCapsuleHalfHeight()));
-					bHit = GetWorld()->SweepSingle( Result, PawnLocation, PawnLocation + TestWalk, FQuat::Identity, CollisionChannel, BoxShape, CapsuleQuery, ResponseParam );
-				}
-				else
-				{
-					bHit = GetWorld()->SweepSingle( Result, PawnLocation, PawnLocation + TestWalk, FQuat::Identity, CollisionChannel, CapsuleShape, CapsuleQuery, ResponseParam );
-				}
 
-				if (bHit)
-				{
-					// Only matters if we can't walk there
-					if (!IsValidLandingSpot(Result.Location, Result))
-					{
-						TickAirControl = 0.f;
-					}
-				}
-			}
-		}
-
-		float MaxAccel = GetMaxAcceleration() * TickAirControl;
+	// bound acceleration, falling object has minimal ability to impact acceleration
+	if (!HasRootMotion() && FallAcceleration.SizeSquared() > 0.f)
+	{
+		const float TickAirControl = GetAirControl(deltaTime, AirControl, FallAcceleration);
 
-		// Boost maxAccel to increase player's control when falling
-		if( (Speed2d < 10.f) && (TickAirControl > 0.f) && (TickAirControl <= 0.05f) ) //allow initial burst
-		{
-			MaxAccel = MaxAccel + (10.f - Speed2d)/deltaTime;
-		}
-		
+		const float MaxAccel = GetMaxAcceleration() * TickAirControl;
 		FallAcceleration = FallAcceleration.ClampMaxSize(MaxAccel);
 	}
 
 	float remainingTime = deltaTime;
-	float timeTick = 0.1f;
-
 	while( (remainingTime >= MIN_TICK_TIME) && (Iterations < MaxSimulationIterations) )
 	{
 		Iterations++;
@@ -2855,6 +2869,7 @@ void UCharacterMovementComponent::PhysFalling(float deltaTime, int32 Iterations)
 		}
 
 		// Move
+		FHitResult Hit(1.f);
 		FVector Adjusted = 0.5f*(OldVelocity + Velocity) * timeTick;  
 		SafeMoveUpdatedComponent( Adjusted, PawnRotation, true, Hit);
 		
@@ -3008,6 +3023,54 @@ void UCharacterMovementComponent::PhysFalling(float deltaTime, int32 Iterations)
 }
 
 
+bool UCharacterMovementComponent::FindAirControlImpact(float DeltaTime, float TickAirControl, const FVector& FallAcceleration, FHitResult& OutHitResult)
+{
+	// test for slope to avoid using air control to climb walls
+	const float TestWalkTime = FMath::Max(DeltaTime, 0.05f);
+	const FVector TestWalk = ((TickAirControl * GetMaxAcceleration() * FallAcceleration.SafeNormal() + FVector(0.f,0.f,GetGravityZ())) * TestWalkTime + Velocity) * TestWalkTime;
+	if (!TestWalk.IsZero())
+	{
+		static const FName FallingTraceParamsTag = FName(TEXT("PhysFalling"));
+		FCollisionQueryParams CapsuleQuery(FallingTraceParamsTag, false, CharacterOwner);
+		FCollisionResponseParams ResponseParam;
+		InitCollisionParams(CapsuleQuery, ResponseParam);
+		const FVector PawnLocation = CharacterOwner->GetActorLocation();
+		const ECollisionChannel CollisionChannel = UpdatedComponent->GetCollisionObjectType();
+
+		const FCollisionShape CapsuleShape = GetPawnCapsuleCollisionShape(SHRINK_None);
+		if (bUseFlatBaseForFloorChecks)
+		{
+			const FCollisionShape BoxShape = FCollisionShape::MakeBox(FVector(CapsuleShape.GetCapsuleRadius(), CapsuleShape.GetCapsuleRadius(), CapsuleShape.GetCapsuleHalfHeight()));
+			GetWorld()->SweepSingle( OutHitResult, PawnLocation, PawnLocation + TestWalk, FQuat::Identity, CollisionChannel, BoxShape, CapsuleQuery, ResponseParam );
+		}
+		else
+		{
+			GetWorld()->SweepSingle( OutHitResult, PawnLocation, PawnLocation + TestWalk, FQuat::Identity, CollisionChannel, CapsuleShape, CapsuleQuery, ResponseParam );
+		}
+
+		if (OutHitResult.bBlockingHit)
+		{
+			return true;
+		}
+	}
+
+	return false;
+}
+
+
+float UCharacterMovementComponent::LimitAirControl(float DeltaTime, float TickAirControl, const FVector& FallAcceleration, const FHitResult& HitResult)
+{
+	if (HitResult.bBlockingHit)
+	{
+		if (!IsValidLandingSpot(HitResult.Location, HitResult))
+		{
+			TickAirControl = 0.f;
+		}
+	}
+
+	return TickAirControl;
+}
+
 bool UCharacterMovementComponent::CheckLedgeDirection(const FVector& OldLocation, const FVector& SideStep, const FVector& GravDir)
 {
 	FVector SideDest = OldLocation + SideStep;
