#ue4 - Custom movement in MoveSmooth should still use a FScopedMovementUpdate.

- Also fixed bad logic on HitResult in CheckLedgeDirection().

diff --git a/Engine/Source/Runtime/Engine/Private/CharacterMovementComponent.cpp b/Engine/Source/Runtime/Engine/Private/CharacterMovementComponent.cpp
--- a/Engine/Source/Runtime/Engine/Private/CharacterMovementComponent.cpp
+++ b/Engine/Source/Runtime/Engine/Private/CharacterMovementComponent.cpp
@@ -3080,12 +3080,12 @@ bool UCharacterMovementComponent::CheckLedgeDirection(const FVector& OldLocation
 	InitCollisionParams(CapsuleParams, ResponseParam);
 	FCollisionShape CapsuleShape = GetPawnCapsuleCollisionShape(SHRINK_None);
 	const ECollisionChannel CollisionChannel = UpdatedComponent->GetCollisionObjectType();
-	FHitResult Result(ForceInit);
+	FHitResult Result(1.f);
 	GetWorld()->SweepSingle(Result, OldLocation, SideDest, FQuat::Identity, CollisionChannel, CapsuleShape, CapsuleParams, ResponseParam);
 
-	if ( (Result.Time == 1.f) || IsWalkable(Result) )
+	if ( !Result.bBlockingHit || IsWalkable(Result) )
 	{
-		if ( Result.Time == 1.f )
+		if ( !Result.bBlockingHit )
 		{
 			GetWorld()->SweepSingle(Result, SideDest, SideDest + GravDir * (MaxStepHeight + LedgeCheckThreshold), FQuat::Identity, CollisionChannel, CapsuleShape, CapsuleParams, ResponseParam);
 		}
@@ -3948,14 +3948,22 @@ void UCharacterMovementComponent::AddForce( FVector Force )
 
 void UCharacterMovementComponent::MoveSmooth(const FVector& InVelocity, const float DeltaSeconds, FStepDownResult* OutStepDownResult)
 {
+	if (!HasValidData())
+	{
+		return;
+	}
+
+	// Custom movement mode.
+	// Custom movement may need an update even if there is zero velocity.
 	if (MovementMode == MOVE_Custom)
 	{
-		PhysCustom(DeltaSeconds, 1);
+		FScopedMovementUpdate ScopedMovementUpdate(UpdatedComponent, bEnableScopedMovementUpdates ? EScopedUpdate::DeferredUpdates : EScopedUpdate::ImmediateUpdates);
+		PhysCustom(DeltaSeconds, 0);
 		return;
 	}
 
 	FVector Delta = InVelocity * DeltaSeconds;
-	if (!HasValidData() || Delta.IsZero())
+	if (Delta.IsZero())
 	{
 		return;
 	}
