Alignment issue fix for win32 mono

diff --git a/Engine/Source/Runtime/Engine/Classes/Animation/Rig.h b/Engine/Source/Runtime/Engine/Classes/Animation/Rig.h
--- a/Engine/Source/Runtime/Engine/Classes/Animation/Rig.h
+++ b/Engine/Source/Runtime/Engine/Classes/Animation/Rig.h
@@ -19,10 +19,6 @@ struct FNode
 {
 	GENERATED_USTRUCT_BODY()
 
-	/** This is Display Name where it will be used to display in Retarget Manager. This name has to be unique. */
-	UPROPERTY(EditAnywhere, Category="FNode")
-	FString DisplayName;
-
 	/** Name of the original node. We don't allow to change this. This is used for identity.**/
 	UPROPERTY(VisibleAnywhere, Category="FNode")
 	FName Name;
@@ -35,17 +31,21 @@ struct FNode
 	UPROPERTY()
 	FTransform Transform;
 
+	/** This is Display Name where it will be used to display in Retarget Manager. This name has to be unique. */
+	UPROPERTY(EditAnywhere, Category="FNode")
+	FString DisplayName;
+
 	FNode()
 		: Name(NAME_None)
 		, ParentName(NAME_None)
 	{
 	}
 
-	FNode(FName InNodeName, FName InParentName, FTransform InTransform)
-		: DisplayName(InNodeName.ToString())
-		, Name(InNodeName)
+	FNode(FName InNodeName, FName InParentName, const FTransform & InTransform)
+		: Name(InNodeName)
 		, ParentName(InParentName)
 		, Transform(InTransform)
+		, DisplayName(InNodeName.ToString())
 
 	{
 	}
