#ue4 - Documentation and header cleanup for more framework classes.

diff --git a/Engine/Source/Runtime/AIModule/Classes/AIController.h b/Engine/Source/Runtime/AIModule/Classes/AIController.h
--- a/Engine/Source/Runtime/AIModule/Classes/AIController.h
+++ b/Engine/Source/Runtime/AIModule/Classes/AIController.h
@@ -1,13 +1,5 @@
 // Copyright 1998-2014 Epic Games, Inc. All Rights Reserved.
 
-//=============================================================================
-// AIController, the base class of AI.
-//
-// Controllers are non-physical actors that can be attached to a pawn to control
-// its actions.  AIControllers implement the artificial intelligence for the pawns they control.
-//
-//=============================================================================
-
 #pragma once
 
 #include "AITypes.h"
@@ -62,6 +54,18 @@ struct FFocusKnowledge
 	TArray<struct FFocusItem> Priorities;
 };
 
+//=============================================================================
+/**
+ * AIController is the base class of controllers for AI-controlled Pawns.
+ * 
+ * Controllers are non-physical actors that can be attached to a pawn to control its actions.
+ * AIControllers manage the artificial intelligence for the pawns they control.
+ * In networked games, they only exist on the server.
+ *
+ * @see https://docs.unrealengine.com/latest/INT/Gameplay/Framework/Controller/
+ */
+//=============================================================================
+
 UCLASS(BlueprintType, Blueprintable)
 class AIMODULE_API AAIController : public AController, public IAIPerceptionListenerInterface
 {
@@ -107,6 +111,7 @@ public:
 	UPROPERTY(BlueprintReadOnly, Category = AI)
 	TSubobjectPtr<UPawnActionsComponent> ActionsComp;
 
+	/** Event called when PossessedPawn is possesed by this controller. */
 	UFUNCTION(BlueprintImplementableEvent, Category = "AI")
 	void OnPossess(APawn* PossessedPawn);
 
@@ -208,10 +213,11 @@ public:
 	UFUNCTION(BlueprintCallable, Category="AI", meta=(FriendlyName="SetFocalPoint"))
 	void K2_SetFocalPoint(FVector FP, bool bOffsetFromBase = false);
 
-	/* Set Focus for actor, will set FocalPoint as a result. */
+	/** Set Focus for actor, will set FocalPoint as a result. */
 	UFUNCTION(BlueprintCallable, Category="AI", meta=(FriendlyName="SetFocus"))
 	void K2_SetFocus(AActor* NewFocus);
 
+	/** Get the focused actor. */
 	UFUNCTION(BlueprintCallable, Category="AI")
 	AActor* GetFocusActor() const;
 
diff --git a/Engine/Source/Runtime/Engine/Classes/Components/PrimitiveComponent.h b/Engine/Source/Runtime/Engine/Classes/Components/PrimitiveComponent.h
--- a/Engine/Source/Runtime/Engine/Classes/Components/PrimitiveComponent.h
+++ b/Engine/Source/Runtime/Engine/Classes/Components/PrimitiveComponent.h
@@ -23,14 +23,18 @@ struct FStreamingTexturePrimitiveInfo
 	float TexelFactor;
 };
 
+/** Determines whether a Character can attempt to step up onto a component when they walk in to it. */
 UENUM()
 enum ECanBeCharacterBase
 {
-	// Character cannot step up onto this Component.
+	/** Character cannot step up onto this Component. */
 	ECB_No,
-	// Character can step up onto this Component.
+	/** Character can step up onto this Component. */
 	ECB_Yes,
-	// Owning actor determines whether character can step up onto this Component (default true unless overridden in code).
+	/**
+	 * Owning actor determines whether character can step up onto this Component (default true unless overridden in code).
+	 * @see AActor::CanBeBaseForCharacter()
+	 */
 	ECB_Owner,
 	ECB_MAX,
 };
@@ -70,13 +74,14 @@ struct FSpriteCategoryInfo
 	FText Description;
 };
 
-/** Delegate for notification of blocking collision against a specific component.  
+/**
+ * Delegate for notification of blocking collision against a specific component.  
  * NormalImpulse will be filled in for physics-simulating bodies, but will be zero for swept-component blocking collisions. 
  */
 DECLARE_DYNAMIC_MULTICAST_DELEGATE_FourParams( FComponentHitSignature, class AActor*, OtherActor, class UPrimitiveComponent*, OtherComp, FVector, NormalImpulse, const FHitResult&, Hit );
-/** Delegate for notification of start of overlap of a specific component */
+/** Delegate for notification of start of overlap with a specific component */
 DECLARE_DYNAMIC_MULTICAST_DELEGATE_FiveParams( FComponentBeginOverlapSignature,class AActor*, OtherActor, class UPrimitiveComponent*, OtherComp, int32, OtherBodyIndex, bool, bFromSweep, const FHitResult &, SweepResult);
-/** Delegate for notification of end of overlap of a specific component */
+/** Delegate for notification of end of overlap with a specific component */
 DECLARE_DYNAMIC_MULTICAST_DELEGATE_ThreeParams( FComponentEndOverlapSignature, class AActor*, OtherActor, class UPrimitiveComponent*, OtherComp, int32, OtherBodyIndex);
 
 DECLARE_DYNAMIC_MULTICAST_DELEGATE_OneParam( FComponentBeginCursorOverSignature, UPrimitiveComponent*, TouchedComponent );
@@ -88,46 +93,16 @@ DECLARE_DYNAMIC_MULTICAST_DELEGATE_TwoParams( FComponentOnInputTouchEndSignature
 DECLARE_DYNAMIC_MULTICAST_DELEGATE_TwoParams( FComponentBeginTouchOverSignature, ETouchIndex::Type, FingerIndex, UPrimitiveComponent*, TouchedComponent );
 DECLARE_DYNAMIC_MULTICAST_DELEGATE_TwoParams( FComponentEndTouchOverSignature, ETouchIndex::Type, FingerIndex, UPrimitiveComponent*, TouchedComponent );
 
+/**
+ * PrimitiveComponents are SceneComponents that contain or generate some sort of geometry, generally to be rendered or used as collision data.
+ * There are several subclasses for the various types of geometry, but the most common by far are the ShapeComponents (Capsule, Sphere, Box), StaticMeshComponent, and SkeletalMeshComponent.
+ * ShapeComponents generate geometry that is used for collision detection but are not rendered, while StaticMeshComponents and SkeletalMeshComponents contain pre-built geometry that is rendered, but can also be used for collision detection.
+ */
 UCLASS(abstract, HideCategories=(Mobility), ShowCategories=(PhysicsVolume))
 class ENGINE_API UPrimitiveComponent : public USceneComponent, public INavRelevantInterface
 {
 	GENERATED_UCLASS_BODY()
 
-	/** Tick function for physics ticking **/
-	UPROPERTY()
-	struct FPrimitiveComponentPostPhysicsTickFunction PostPhysicsComponentTick;
-
-	/** A fence to track when the primitive is detached from the scene in the rendering thread. */
-	FRenderCommandFence DetachFence;
-
-	/**
-	 * Incremented by the main thread before being attached to the scene, decremented
-	 * by the rendering thread after removal. This counter exists to assert that 
-	 * operations are safe in order to help avoid race conditions.
-	 *
-	 *           *** Runtime logic should NEVER rely on this value. ***
-	 *
-	 * The only safe assertions to make are:
-	 *
-	 *     AttachmentCounter == 0: The primitive is not exposed to the rendering
-	 *                             thread, it is safe to modify shared members.
-	 *                             This assertion is valid ONLY from the main thread.
-	 *
-	 *     AttachmentCounter >= 1: The primitive IS exposed to the rendering
-	 *                             thread and therefore shared members must not
-	 *                             be modified. This assertion may be made from
-	 *                             any thread. Note that it is valid and expected
-	 *                             for AttachmentCounter to be larger than 1, e.g.
-	 *                             during reattachment.
-	 */
-	FThreadSafeCounter AttachmentCounter;
-
-	// Scene data
-
-	/** Replacement primitive to draw instead of this one (multiple UPrim's will point to the same Replacement) */
-	UPROPERTY()
-	TLazyObjectPtr<class UPrimitiveComponent> ReplacementPrimitive;
-
 	// Rendering
 	
 	/**
@@ -169,30 +144,46 @@ public:
 	 * Indicates if we'd like to create physics state all the time (for collision and simulation). 
 	 * If you set this to false, it still will create physics state if collision or simulation activated. 
 	 * This can help performance if you'd like to avoid overhead of creating physics state when triggers 
-	 **/
+	 */
 	UPROPERTY(EditAnywhere, AdvancedDisplay, BlueprintReadOnly, Category=Collision)
 	uint32 bAlwaysCreatePhysicsState:1;
 
-	/** If true, this component will generate overlap events when it is overlapping other components (e.g. Begin Overlap) */
+	/**
+	 * If true, this component will generate overlap events when it is overlapping other components (e.g. Begin Overlap)
+	 * @see [Overlap Events] https://docs.unrealengine.com/latest/INT/Engine/Physics/Collision/index.html#overlapandgenerateoverlapevents
+	 * @see UpdateOverlaps(), BeginComponentOverlap(), EndComponentOverlap()
+	 */
 	UPROPERTY(EditAnywhere, BlueprintReadWrite, Category=Collision)
 	uint32 bGenerateOverlapEvents:1;
 
-	/** If true, this component will generate individual overlaps for each overlapping physics body if it is a multi-body component. When false, this component will
-		generate only one overlap, regardless of how many physics bodies it has and how many of them are overlapping another component/body. This flag has no
-		influence on single body components. */
+	/**
+	 * If true, this component will generate individual overlaps for each overlapping physics body if it is a multi-body component. When false, this component will
+	 * generate only one overlap, regardless of how many physics bodies it has and how many of them are overlapping another component/body. This flag has no
+	 * influence on single body components.
+	 */
 	UPROPERTY(EditAnywhere, AdvancedDisplay, BlueprintReadWrite, Category=Collision)
 	uint32 bMultiBodyOverlap:1;
 
-
-	/** If true, this component will look for collisions on both physic scenes during movement */
+	/**
+	 * If true, this component will look for collisions on both physic scenes during movement.
+	 * Only required if the asynchronous physics scene is enabled and has geometry in it, and you wish to test for collisions with objects in that scene.
+	 * @see MoveComponent()
+	 */
 	UPROPERTY(EditAnywhere, AdvancedDisplay, BlueprintReadWrite, Category=Collision)
 	uint32 bCheckAsyncSceneOnMove:1;
 
-	/** If true, component sweeps with this component should trace against complex collision during movement. */
+	/**
+	 * If true, component sweeps with this component should trace against complex collision during movement (for example, each triangle of a mesh).
+	 * If false, collision will be resolved against simple collision bounds instead.
+	 * @see MoveComponent()
+	 */
 	UPROPERTY(EditAnywhere, AdvancedDisplay, BlueprintReadWrite, Category=Collision)
 	uint32 bTraceComplexOnMove:1;
 
-	/** If true, component sweeps will return the material in their hit-info. */
+	/**
+	 * If true, component sweeps will return the material in their hit result.
+	 * @see MoveComponent(), FHitResult
+	 */
 	UPROPERTY(EditAnywhere, AdvancedDisplay, BlueprintReadWrite, Category=Collision)
 	uint32 bReturnMaterialOnMove:1;
 
@@ -355,6 +346,7 @@ public:
 	/** Composite the drawing of this component onto the scene after post processing (only applies to editor drawing) */
 	UPROPERTY()
 	uint32 bUseEditorCompositing:1;
+
 	/**
 	 * Translucent objects with a lower sort priority draw behind objects with a higher priority.
 	 * Translucent objects with the same priority are rendered from back-to-front based on their bounds origin.
@@ -362,7 +354,7 @@ public:
 	 * Ignored if the object is not translucent.  The default priority is zero.
 	 * Warning: This should never be set to a non-default value unless you know what you are doing, as it will prevent the renderer from sorting correctly.  
 	 * It is especially problematic on dynamic gameplay effects.
-	 **/
+	 */
 	UPROPERTY(EditAnywhere, BlueprintReadOnly, AdvancedDisplay, Category=Rendering)
 	int32 TranslucencySortPriority;
 
@@ -373,10 +365,11 @@ public:
 	/** Used by the renderer, to identify a component across re-registers. */
 	FPrimitiveComponentId ComponentId;
 
-	/** Multiplier used to scale the Light Propagation Volume light injection bias, to reduce light bleeding. 
-	  * Set to 0 for no bias, 1 for default or higher for increased biasing (e.g. for 
-	  * thin geometry such as walls)
-	  **/
+	/**
+	 * Multiplier used to scale the Light Propagation Volume light injection bias, to reduce light bleeding. 
+	 * Set to 0 for no bias, 1 for default or higher for increased biasing (e.g. for 
+	 * thin geometry such as walls)
+	 */
 	UPROPERTY(EditAnywhere, BlueprintReadOnly, AdvancedDisplay, Category=Rendering, meta=(UIMin = "0.0", UIMax = "3.0"))
 	float LpvBiasMultiplier;
 
@@ -386,11 +379,11 @@ public:
 	UPROPERTY(EditAnywhere, BlueprintReadOnly, Category=Collision, meta=(ShowOnlyInnerProperties))
 	FBodyInstance BodyInstance;
 
-	/** can this component potentially influence navigation */
+	/** Whether this component can potentially influence navigation */
 	UPROPERTY(EditAnywhere, Category=Collision, AdvancedDisplay)
 	uint32 bCanEverAffectNavigation:1;
 
-	/** cached navigation relevancy flag for collision updates */
+	/** Cached navigation relevancy flag for collision updates */
 	uint32 bNavigationRelevant : 1;
 
 protected:
@@ -404,8 +397,7 @@ protected:
 	/** Last time we checked AreAllCollideableDescendantsRelative(), so we can throttle those tests since it rarely changes once false. */
 	float LastCheckedAllCollideableDescendantsTime;
 
-	/** if set to true then DoCustomNavigableGeometryExport will be called to collect navigable 
-	 *	geometry of this component. */
+	/** If true then DoCustomNavigableGeometryExport will be called to collect navigable geometry of this component. */
 	TEnumAsByte<EHasCustomNavigableGeometry::Type> bHasCustomNavigableGeometry;
 
 	/** Next id to be used by a component. */
@@ -429,7 +421,6 @@ public:
 	 * The value of WorldSettings->TimeSeconds for the frame when this component was last rendered.  This is written
 	 * from the render thread, which is up to a frame behind the game thread, so you should allow this time to
 	 * be at least a frame behind the game thread's world time before you consider the actor non-visible.
-	 * There's an equivalent variable in PrimitiveComponent.
 	 */
 	UPROPERTY(transient)
 	float LastRenderTime;
@@ -449,7 +440,7 @@ public:
 	 * Set of actors to ignore during component sweeps in MoveComponent().
 	 * All components owned by these actors will be ignored when this component moves or updates overlaps.
 	 * Components on the other Actor may also need to be told to do the same when they move.
-	 * @see IgnoreActorWhenMoving
+	 * @see IgnoreActorWhenMoving()
 	 */
 	TArray<TWeakObjectPtr<AActor> > MoveIgnoreActors;
 
@@ -488,6 +479,7 @@ public:
 	 * Begin tracking an overlap interaction with the component specified.
 	 * @param OtherComp - The component of the other actor that this component is now overlapping
 	 * @param bDoNotifies - True to dispatch appropriate begin/end overlap notifications when these events occur.
+	 * @see [Overlap Events] https://docs.unrealengine.com/latest/INT/Engine/Physics/Collision/index.html#overlapandgenerateoverlapevents
 	 */
 	void BeginComponentOverlap(const FOverlapInfo& OtherOverlap, bool bDoNotifies);
 	
@@ -496,11 +488,13 @@ public:
 	 * @param OtherComp - The component of the other actor to stop overlapping
 	 * @param bDoNotifies - True to dispatch appropriate begin/end overlap notifications when these events occur.
 	 * @param bNoNotifySelf	- True to skip end overlap notifications to this component's.  Does not affect notifications to OtherComp's actor.
+	 * @see [Overlap Events] https://docs.unrealengine.com/latest/INT/Engine/Physics/Collision/index.html#overlapandgenerateoverlapevents
 	 */
 	void EndComponentOverlap(const FOverlapInfo& OtherOverlap, bool bDoNotifies=true, bool bNoNotifySelf=false);
 
 	/** Returns true if this component is overlapping OtherComp, false otherwise. */
 	bool IsOverlappingComponent(UPrimitiveComponent const* OtherComp) const;
+	
 	/** Returns true if this component is overlapping OtherComp, false otherwise. */
 	bool IsOverlappingComponent(const FOverlapInfo& Overlap) const;
 
@@ -536,62 +530,62 @@ public:
 	 */
 	virtual void UpdateOverlaps(TArray<FOverlapInfo> const* PendingOverlaps=NULL, bool bDoNotifies=true, const TArray<FOverlapInfo>* OverlapsAtEndLocation=NULL) override;
 
-	/** Overridden to use the overlaps to find the physics volume. */
+	/** Update current physics volume for this component, if bShouldUpdatePhysicsVolume is true. Overridden to use the overlaps to find the physics volume. */
 	virtual void UpdatePhysicsVolume( bool bTriggerNotifiers ) override;
 
 	/**
-	 *  Test the collision of the supplied component at the supplied location/rotation, and determine the set of components that it overlaps
+	 *  Test the collision of the supplied component at the supplied location/rotation, and determine the set of components that it overlaps.
 	 *  @param  OutOverlaps     Array of overlaps found between this component in specified pose and the world
 	 *  @param  World			World to use for overlap test
 	 *  @param  Pos             Location to place the component's geometry at to test against the world
 	 *  @param  Rot             Rotation to place components' geometry at to test against the world
 	 *  @param  TestChannel		The 'channel' that this ray is in, used to determine which components to hit
 	 *	@param	ObjectQueryParams	List of object types it's looking for. When this enters, we do object query with component shape
-	 *  @return TRUE if OutOverlaps contains any blocking results
+	 *  @return true if OutOverlaps contains any blocking results
 	 */
 	virtual bool ComponentOverlapMulti(TArray<struct FOverlapResult>& OutOverlaps, const class UWorld* World, const FVector& Pos, const FRotator& Rot, ECollisionChannel TestChannel, const struct FComponentQueryParams& Params, const struct FCollisionObjectQueryParams& ObjectQueryParams = FCollisionObjectQueryParams::DefaultObjectQueryParam) const;
 
-	/** Called when a component is touched */
+	/** Event called when a component is touched */
 	UPROPERTY(BlueprintAssignable, Category="Collision")
 	FComponentHitSignature OnComponentHit;
 
-	/** Called when something overlaps this component */
+	/** Event called when something overlaps this component */
 	UPROPERTY(BlueprintAssignable, Category="Collision")
 	FComponentBeginOverlapSignature OnComponentBeginOverlap;
 
-	/** Called when something ends overlapping this component */
+	/** Event called when something ends overlapping this component */
 	UPROPERTY(BlueprintAssignable, Category="Collision")
 	FComponentEndOverlapSignature OnComponentEndOverlap;
 
-	/** Called when the mouse cursor is moved over this component and mouse over events are enabled in the player controller */
+	/** Event called when the mouse cursor is moved over this component and mouse over events are enabled in the player controller */
 	UPROPERTY(BlueprintAssignable, Category="Input|Mouse Input")
 	FComponentBeginCursorOverSignature OnBeginCursorOver;
 		 
-	/** Called when the mouse cursor is moved off this component and mouse over events are enabled in the player controller */
+	/** Event called when the mouse cursor is moved off this component and mouse over events are enabled in the player controller */
 	UPROPERTY(BlueprintAssignable, Category="Input|Mouse Input")
 	FComponentEndCursorOverSignature OnEndCursorOver;
 
-	/** Called when the left mouse button is clicked while the mouse is over this component and click events are enabled in the player controller */
+	/** Event called when the left mouse button is clicked while the mouse is over this component and click events are enabled in the player controller */
 	UPROPERTY(BlueprintAssignable, Category="Input|Mouse Input")
 	FComponentOnClickedSignature OnClicked;
 
-	/** Called when the left mouse button is released while the mouse is over this component click events are enabled in the player controller */
+	/** Event called when the left mouse button is released while the mouse is over this component click events are enabled in the player controller */
 	UPROPERTY(BlueprintAssignable, Category="Input|Mouse Input")
 	FComponentOnReleasedSignature OnReleased;
 		 
-	/** Called when a touch input is received over this component when touch events are enabled in the player controller */
+	/** Event called when a touch input is received over this component when touch events are enabled in the player controller */
 	UPROPERTY(BlueprintAssignable, Category="Input|Touch Input")
 	FComponentOnInputTouchBeginSignature OnInputTouchBegin;
 
-	/** Called when a touch input is released over this component when touch events are enabled in the player controller */
+	/** Event called when a touch input is released over this component when touch events are enabled in the player controller */
 	UPROPERTY(BlueprintAssignable, Category="Input|Touch Input")
 	FComponentOnInputTouchEndSignature OnInputTouchEnd;
 
-	/** Called when a finger is moved over this component when touch over events are enabled in the player controller */
+	/** Event called when a finger is moved over this component when touch over events are enabled in the player controller */
 	UPROPERTY(BlueprintAssignable, Category="Input|Touch Input")
 	FComponentBeginTouchOverSignature OnInputTouchEnter;
 
-	/** Called when a finger is moved off this component when touch over events are enabled in the player controller */
+	/** Event called when a finger is moved off this component when touch over events are enabled in the player controller */
 	UPROPERTY(BlueprintAssignable, Category="Input|Touch Input")
 	FComponentEndTouchOverSignature OnInputTouchLeave;
 
@@ -851,6 +845,37 @@ public:
 	/** The primitive's scene info. */
 	FPrimitiveSceneProxy* SceneProxy;
 	
+	/** A fence to track when the primitive is detached from the scene in the rendering thread. */
+	FRenderCommandFence DetachFence;
+
+	/**
+	 * Incremented by the main thread before being attached to the scene, decremented
+	 * by the rendering thread after removal. This counter exists to assert that 
+	 * operations are safe in order to help avoid race conditions.
+	 *
+	 *           *** Runtime logic should NEVER rely on this value. ***
+	 *
+	 * The only safe assertions to make are:
+	 *
+	 *     AttachmentCounter == 0: The primitive is not exposed to the rendering
+	 *                             thread, it is safe to modify shared members.
+	 *                             This assertion is valid ONLY from the main thread.
+	 *
+	 *     AttachmentCounter >= 1: The primitive IS exposed to the rendering
+	 *                             thread and therefore shared members must not
+	 *                             be modified. This assertion may be made from
+	 *                             any thread. Note that it is valid and expected
+	 *                             for AttachmentCounter to be larger than 1, e.g.
+	 *                             during reattachment.
+	 */
+	FThreadSafeCounter AttachmentCounter;
+
+	// Scene data
+
+	/** Replacement primitive to draw instead of this one (multiple UPrim's will point to the same Replacement) */
+	UPROPERTY()
+	TLazyObjectPtr<class UPrimitiveComponent> ReplacementPrimitive;
+
 #if WITH_EDITOR
 	virtual const int32 GetNumUncachedStaticLightingInteractions() const; // recursive function
 #endif
@@ -953,6 +978,11 @@ public:
 	 */
 	virtual void GetUsedTextures(TArray<UTexture*>& OutTextures, EMaterialQualityLevel::Type QualityLevel);
 
+
+	/** Tick function for physics ticking **/
+	UPROPERTY()
+	struct FPrimitiveComponentPostPhysicsTickFunction PostPhysicsComponentTick;
+
 	/** Controls if we get a post physics tick or not. If set during ticking, will take effect next frame **/
 	void SetPostPhysicsComponentTickEnabled(bool bEnable);
 
@@ -1032,7 +1062,7 @@ public:
 		return false;
 	}
 
-	/** Called when a component is 'damaged', allowing for component class specific behaviour */
+	/** Event called when a component is 'damaged', allowing for component class specific behaviour */
 	virtual void ReceiveComponentDamage(float DamageAmount, FDamageEvent const& DamageEvent, AController* EventInstigator, AActor* DamageCauser);
 
 	/**
@@ -1067,7 +1097,7 @@ protected:
 	// Begin USceneComponent Interface
 	virtual void OnUpdateTransform(bool bSkipPhysicsMove) override;
 
-	/** Called when AttachParent changes, to allow the scene to update its attachment state. */
+	/** Event called when AttachParent changes, to allow the scene to update its attachment state. */
 	virtual void OnAttachmentChanged() override;
 
 	/**
@@ -1111,6 +1141,7 @@ protected:
 
 	/** Called to send a transform update for this component to the physics engine */
 	void SendPhysicsTransform(bool bTeleport);
+
 	/** Ensure physics state created **/
 	void EnsurePhysicsStateCreated();
 public:
diff --git a/Engine/Source/Runtime/Engine/Classes/Components/SceneComponent.h b/Engine/Source/Runtime/Engine/Classes/Components/SceneComponent.h
--- a/Engine/Source/Runtime/Engine/Classes/Components/SceneComponent.h
+++ b/Engine/Source/Runtime/Engine/Classes/Components/SceneComponent.h
@@ -73,7 +73,6 @@ DECLARE_DYNAMIC_MULTICAST_DELEGATE_OneParam(FPhysicsVolumeChanged, class APhysic
 /**
  * A SceneComponent has a transform and supports attachment, but has no rendering or collision capabilities.
  * Useful as a 'dummy' component in the hierarchy to offset others.
- * @see https://docs.unrealengine.com/latest/INT/API/Runtime/Engine/Components/USceneComponent/index.html
  * @see [Scene Components](https://docs.unrealengine.com/latest/INT/Programming/UnrealArchitecture/Actors/Components/index.html#scenecomponents)
  */
 UCLASS(ClassGroup=Utility, BlueprintType, HideCategories=(Trigger, PhysicsVolume), meta=(BlueprintSpawnableComponent))
@@ -497,13 +496,9 @@ public:
 	virtual void UpdateOverlaps(TArray<FOverlapInfo> const* PendingOverlaps=NULL, bool bDoNotifies=true, const TArray<FOverlapInfo>* OverlapsAtEndLocation=NULL);
 
 	/**
-	 * Tries to move the component by a movement vector.
-	 * 
-	 * Assumes that the component's current Location is valid and that the component does fit in its current Location. 
-	 * Assumes that the level's Dynamics member is locked, which will always be the case during
-	 * a call to UWorld::Tick; if not locked, no actor-actor collision checking is performed.
-	 *
-	 * Updates actor's PhysicsVolume and touching interactions (if bShouldUpdatePhysicsVolume is true).
+	 * Tries to move the component by a movement vector (Delta) and sets rotation to NewRotation.
+	 * Assumes that the component's current location is valid and that the component does fit in its current Location.
+	 * Dispatches blocking hit notifications (if bSweep is true), and calls UpdateOverlaps() after movement to update overlap state.
 	 * 
 	 * @param Delta			The desired location change in world space.
 	 * @param NewRotation	The new desired rotation in world space.
@@ -608,7 +603,7 @@ public:
 	}
 
 	/** 
-	 * Updates the PhysicsVolume of this SceneComponent
+	 * Updates the PhysicsVolume of this SceneComponent, if bShouldUpdatePhysicsVolume is true.
 	 * 
 	 * @param bTriggerNotifiers		if true, send zone/volume change events
 	 */
diff --git a/Engine/Source/Runtime/Engine/Classes/Engine/EngineBaseTypes.h b/Engine/Source/Runtime/Engine/Classes/Engine/EngineBaseTypes.h
--- a/Engine/Source/Runtime/Engine/Classes/Engine/EngineBaseTypes.h
+++ b/Engine/Source/Runtime/Engine/Classes/Engine/EngineBaseTypes.h
@@ -554,12 +554,26 @@ struct ENGINE_API FURL
 	bool operator==( const FURL& Other ) const;
 };
 
+/**
+ * The network mode the game is currently running.
+ */
 enum ENetMode
 {
+	/** Standalone: a game without networking, with one or more local players. Still considered a server because it has all server functionality. */
 	NM_Standalone,
+
+	/** Dedicated server: server with no local players. */
 	NM_DedicatedServer,
+
+	/** Listen server: a server that also has a local player who is hosting the game, available to other players on the network. */
 	NM_ListenServer,
-	NM_Client,	// note that everything below this value is a kind of server
+
+	/**
+	 * Network client: client connected to a remote server.
+	 * Note that every mode less than this value is a kind of server, so checking NetMode < NM_Client is always some variety of server.
+	 */
+	NM_Client,
+
 	NM_MAX,
 };
 
diff --git a/Engine/Source/Runtime/Engine/Classes/Engine/EngineTypes.h b/Engine/Source/Runtime/Engine/Classes/Engine/EngineTypes.h
--- a/Engine/Source/Runtime/Engine/Classes/Engine/EngineTypes.h
+++ b/Engine/Source/Runtime/Engine/Classes/Engine/EngineTypes.h
@@ -200,16 +200,28 @@ enum ETriangleSortAxis
 	TSA_MAX,
 };
 
-/** Movement modes for Characters.  */
+/** Movement modes for Characters. */
 UENUM(BlueprintType)
 enum EMovementMode
 {
+	/** None (movement is disabled). */
 	MOVE_None		UMETA(DisplayName="None"),
+
+	/** Walking on a surface. */
 	MOVE_Walking	UMETA(DisplayName="Walking"),
+
+	/** Falling under the effects of gravity, such as after jumping or walking off the edge of a surface. */
 	MOVE_Falling	UMETA(DisplayName="Falling"),
+
+	/** Swimming through a fluid volume, under the effects of gravity and buoyancy. */
 	MOVE_Swimming	UMETA(DisplayName="Swimming"),
+
+	/** Flying, ignoring the effects of gravity. Affected by the current physics volume's fluid friction. */
 	MOVE_Flying		UMETA(DisplayName="Flying"),
+
+	/** User-defined custom movement mode, including many possible sub-modes. */
 	MOVE_Custom		UMETA(DisplayName="Custom"),
+
 	MOVE_MAX		UMETA(Hidden),
 };
 
@@ -1312,13 +1324,15 @@ struct FPrimitiveMaterialRef
 	}
 };
 
-/** Structure containing information about one hit of the trace */
+/**
+ * Structure containing information about one hit of a trace, such as point of impact and surface normal at that point.
+ */
 USTRUCT(BlueprintType, meta=(HasNativeBreak="Engine.GameplayStatics.BreakHitResult"))
 struct ENGINE_API FHitResult
 {
 	GENERATED_USTRUCT_BODY()
 
-	/** Indicates if this hit was requesting a block - if false, was requesting a touch instead */
+	/** Indicates if this hit was a result of blocking collision. If false, there was no hit or it was an overlap/touch instead. */
 	UPROPERTY()
 	uint32 bBlockingHit:1;
 
@@ -1331,43 +1345,67 @@ struct ENGINE_API FHitResult
 	UPROPERTY()
 	uint32 bStartPenetrating:1;
 
-	/* 'Time' of hit along ray (between 0.0 and 1.0), between TraceStart and TraceEnd. */
+	/**
+	 * 'Time' of impact along trace direction (ranging from 0.0 to 1.0) if there is a hit, indicating time between TraceStart and TraceEnd.
+	 * For swept movement (but not queries) this may be pulled back slightly from the actual time of impact, to prevent precision problems with adjacent geometry.
+	 */
 	UPROPERTY()
 	float Time;
 	
-	/** Location of the hit in world space.  If this was a Swept Shape test, this is the location where we can place the shape in the world where it will not inter-penetrate. */
+	/**
+	 * The location in world space where the shape would end up against the impacted object, if there is a hit. Equal to the point of impact for line tests.
+	 * Example: for a sphere trace test, this is the point where the center of the sphere would be located when it touched the other object.
+	 * For swept movement (but not queries) this may not equal the final location of the shape since hits are pulled back slightly to prevent precision issues.
+	 */
 	UPROPERTY()
 	FVector_NetQuantize Location;
 
-	/** Normal of the hit in world space, for the object that was swept. This is computed for capsules and spheres, otherwise it will be the same as ImpactNormal.*/
+	/**
+	 * Normal of the hit in world space, for the object that was swept. Equal to ImpactNormal for line tests.
+	 * This is computed for capsules and spheres, otherwise it will be the same as ImpactNormal.
+	 * Example: for a sphere trace test, this is a normalized vector pointing in towards the center of the sphere at the point of impact.
+	 */
 	UPROPERTY()
 	FVector_NetQuantizeNormal Normal;
 
-	/** Location of the actual contact of the trace shape (box, sphere, etc.) with the world. (e.g. A box was swept.  This is the point where a side of the box touched something.) */
+	/**
+	 * Location in world space of the actual contact of the trace shape (box, sphere, ray, etc) with the impacted object.
+	 * Example: for a sphere trace test, this is the point where the surface of the sphere touches the other object.
+	 */
 	UPROPERTY()
 	FVector_NetQuantize ImpactPoint;
 
-	/** Normal of the hit in world space, for the object that was hit by the sweep. */
+	/**
+	 * Normal of the hit in world space, for the object that was hit by the sweep, if any.
+	 * For example if a box hits a flat plane, this is a normalized vector pointing out from the plane.
+	 * In the case of impact with a corner or edge of a surface, usually the "most opposing" normal (opposed to the query direction) is chosen.
+	 */
 	UPROPERTY()
 	FVector_NetQuantizeNormal ImpactNormal;
 
-	/** Start location of the trace. */
+	/**
+	 * Start location of the trace.
+	 * For example if a sphere is swept against the world, this is the starting location of the center of the sphere.
+	 */
 	UPROPERTY()
 	FVector_NetQuantize TraceStart;
 
-	/** End location of the trace. This is NOT where the impact occurred (if any), but the furthest point in the attempted sweep. */
+	/**
+	 * End location of the trace; this is NOT where the impact occurred (if any), but the furthest point in the attempted sweep.
+	 * For example if a sphere is swept against the world, this would be the center of the sphere if there was no blocking hit.
+	 */
 	UPROPERTY()
 	FVector_NetQuantize TraceEnd;
 
 	/**
-	  * If this test started in penetration (bStartPenetrating is true), and a depenetration vector can be computed,
+	  * If this test started in penetration (bStartPenetrating is true) and a depenetration vector can be computed,
 	  * this value is the distance along Normal that will result in moving out of penetration.
 	  * If the distance cannot be computed, this distance will be zero.
 	  */
 	UPROPERTY()
 	float PenetrationDepth;
 
-	/** Extra data about item that was hit (hit primitive specific) */
+	/** Extra data about item that was hit (hit primitive specific). */
 	UPROPERTY()
 	int32 Item;
 
@@ -1375,11 +1413,11 @@ struct ENGINE_API FHitResult
 	UPROPERTY()
 	TWeakObjectPtr<class UPhysicalMaterial> PhysMaterial;
 
-	/** Actor that the check hit. */
+	/** Actor hit by the trace. */
 	UPROPERTY()
 	TWeakObjectPtr<class AActor> Actor;
 
-	/** PrimitiveComponent that the check hit. */
+	/** PrimitiveComponent hit by the trace. */
 	UPROPERTY(NotReplicated)
 	TWeakObjectPtr<class UPrimitiveComponent> Component;
 
@@ -1387,10 +1425,11 @@ struct ENGINE_API FHitResult
 	UPROPERTY()
 	FName BoneName;
 
-	/** Face index we hit (for complex hits with triangle meshes) */
+	/** Face index we hit (for complex hits with triangle meshes). */
 	UPROPERTY()
 	int32 FaceIndex;
 
+
 	FHitResult()
 	{
 		FMemory::Memzero(this, sizeof(FHitResult));
@@ -1423,10 +1462,10 @@ struct ENGINE_API FHitResult
 		}
 	}
 
-	/** Utility to return the Actor that owns the Component that was hit */
+	/** Utility to return the Actor that owns the Component that was hit. */
 	AActor* GetActor() const;
 
-	/** Utility to return the Component that was hit */
+	/** Utility to return the Component that was hit. */
 	UPrimitiveComponent* GetComponent() const;
 
 	/** Return true if there was a blocking hit that was not caused by starting in penetration. */
@@ -1435,7 +1474,7 @@ struct ENGINE_API FHitResult
 		return bBlockingHit && !bStartPenetrating;
 	}
 
-	/** Returns the first 'blocking' hit in an array of results */
+	/** Static utility function that returns the first 'blocking' hit in an array of results. */
 	static FHitResult* GetFirstBlockingHit(TArray<FHitResult>& InHits)
 	{
 		for(int32 HitIdx=0; HitIdx<InHits.Num(); HitIdx++)
@@ -1448,7 +1487,7 @@ struct ENGINE_API FHitResult
 		return NULL;
 	}
 
-	/** Returns number of blocking hits in array */
+	/** Static utility function that returns the number of blocking hits in array. */
 	static int32 GetNumBlockingHits(const TArray<FHitResult>& InHits)
 	{
 		int32 NumBlocks = 0;
@@ -1462,7 +1501,7 @@ struct ENGINE_API FHitResult
 		return NumBlocks;
 	}
 
-	/** Returns number of overlapping hits in array */
+	/** Static utility function that returns the number of overlapping hits in array. */
 	static int32 GetNumOverlapHits(const TArray<FHitResult>& InHits)
 	{
 		return (InHits.Num() - GetNumBlockingHits(InHits));
@@ -1504,17 +1543,17 @@ struct ENGINE_API FOverlapResult
 	}
 };
 
-/** Structure containing information about minimum translation directon (MTD) */
+/** Structure containing information about minimum translation direction (MTD) */
 USTRUCT()
 struct ENGINE_API FMTDResult
 {
 	GENERATED_USTRUCT_BODY()
 
-	/** Direction of the minimum translation required to fix penetration. */
+	/** Normalized direction of the minimum translation required to fix penetration. */
 	UPROPERTY()
 	FVector Direction;
 
-	/** Distance required to move along the MTD. */
+	/** Distance required to move along the MTD vector (Direction). */
 	UPROPERTY()
 	float Distance;
 
diff --git a/Engine/Source/Runtime/Engine/Classes/GameFramework/Character.h b/Engine/Source/Runtime/Engine/Classes/GameFramework/Character.h
--- a/Engine/Source/Runtime/Engine/Classes/GameFramework/Character.h
+++ b/Engine/Source/Runtime/Engine/Classes/GameFramework/Character.h
@@ -186,12 +186,15 @@ DEPRECATED(4.4, "FRepRelativeMovement has been renamed to FBasedMovementInfo")
 typedef FBasedMovementInfo FRepRelativeMovement;
 
 
-//=============================================================================
-// Characters are Pawns that have a mesh, collision, and physics and are responsible for all
-// physical interaction between the player or AI and the world. They are strongly tied
-// to a CharacterMovementComponent that handles movement of the collision capsule, and they
-// also have implementations of basic networking and input models.
-//=============================================================================
+
+/**
+ * Characters are Pawns that have a mesh, collision, and default movement logic. They are responsible for all
+ * physical interaction between the player or AI and the world, and also have implementations of basic networking and input models.
+ * They are designed for a vertically-oriented player representation that can walk, run, jump, fly, and swim through the world using CharacterMovementComponent.
+ *
+ * @see APawn, UCharacterMovementComponent
+ * @see https://docs.unrealengine.com/latest/INT/Gameplay/Framework/Pawn/Character/
+ */ 
 
 UCLASS(abstract, config=Game, BlueprintType, hidecategories=("Pawn|Character|InternalEvents"))
 class ENGINE_API ACharacter : public APawn
diff --git a/Engine/Source/Runtime/Engine/Classes/GameFramework/CharacterMovementComponent.h b/Engine/Source/Runtime/Engine/Classes/GameFramework/CharacterMovementComponent.h
--- a/Engine/Source/Runtime/Engine/Classes/GameFramework/CharacterMovementComponent.h
+++ b/Engine/Source/Runtime/Engine/Classes/GameFramework/CharacterMovementComponent.h
@@ -1,10 +1,5 @@
 // Copyright 1998-2014 Epic Games, Inc. All Rights Reserved.
 
-/**
- * Movement component updates position of associated PrimitiveComponent (owned by a Character) during its tick.
- * Assumes/requires the character's root component to be a CapsuleComponent.
- */
-
 #pragma once
 #include "AI/Navigation/NavigationAvoidanceTypes.h"
 #include "Animation/AnimationAsset.h"
@@ -19,7 +14,7 @@ struct FVector_NetQuantize100;
 struct FVector_NetQuantizeNormal;
 class FDebugDisplayInfo;
 
-// Data about the floor for walking movement.
+/** Data about the floor for walking movement, used by CharacterMovementComponent. */
 USTRUCT(BlueprintType)
 struct ENGINE_API FFindFloorResult
 {
@@ -45,7 +40,7 @@ struct ENGINE_API FFindFloorResult
 	UPROPERTY(VisibleInstanceOnly, BlueprintReadOnly, Category=CharacterFloor)
 	float LineDist;
 
-	/** Hit result of the test that found a floor. */
+	/** Hit result of the test that found a floor. Includes more specific data about the point of impact and surface normal at that point. */
 	UPROPERTY(VisibleInstanceOnly, BlueprintReadOnly, Category=CharacterFloor)
 	FHitResult HitResult;
 
@@ -105,15 +100,33 @@ struct FCharacterMovementComponentPreClothTickFunction : public FTickFunction
 	virtual FString DiagnosticMessage() override;
 };
 
+
+//=============================================================================
+/**
+ * CharacterMovementComponent handles movement logic for the associated Character owner.
+ *
+ * It supports various movement modes including:
+ *    - walking:  Walking on a surface, under the effects of friction, and able to "step up" barriers.
+ *    - falling:  Falling under the effects of gravity. Usually this is the mode after jumping or walking off the edge of a surface.
+ *    - flying:   Flying, ignoring the effects of gravity.
+ *    - swimming: Swimming through a fluid volume, under the effects of gravity and buoyancy.
+ *    - custom:   User-defined custom movement mode, including many possible sub-modes. See UCharacterMovementComponent::CustomMovementMode and SetMovementMode().
+ *
+ * Movement is affected primarily by current Velocity and Acceleration. Acceleration is updated each frame
+ * based on the input vector accumulated thus far (see UPawnMovementComponent::GetPendingInputVector()).
+ *
+ * Networking is fully implemented, with server-client correction and prediction included.
+ *
+ * @see ACharacter, UPawnMovementComponent
+ * @see https://docs.unrealengine.com/latest/INT/Gameplay/Framework/Pawn/Character/
+ */
+//=============================================================================
+
 UCLASS()
 class ENGINE_API UCharacterMovementComponent : public UPawnMovementComponent, public INetworkPredictionInterface
 {
 	GENERATED_UCLASS_BODY()
 
-	/** Post-physics tick function for this character */
-	UPROPERTY()
-	struct FCharacterMovementComponentPreClothTickFunction PreClothComponentTick;
-
 protected:
 
 	/** Character movement component belongs to */
@@ -151,32 +164,23 @@ private:
 public:
 	
 	/**
-	 * Don't allow the character to perch on the edge of a surface if the contact is this close to the edge of the capsule.
-	 * Note that characters will not fall off if they are within MaxStepHeight of a walkable surface below.
+	 * Actor's current movement mode (walking, falling, etc). This is automatically replicated through the Character owner and for client-server movement functions.
+	 * @see SetMovementMode()
 	 */
-	UPROPERTY(Category="Character Movement", EditAnywhere, BlueprintReadWrite, AdvancedDisplay, meta=(ClampMin="0", UIMin="0"))
-	float PerchRadiusThreshold;
-
-	/**
-	 * When perching on a ledge, add this additional distance to MaxStepHeight when determining how high above a walkable floor we can perch.
-	 * Note that we still enforce MaxStepHeight to start the step up; this just allows the character to hang off the edge or step slightly higher off the floor.
-	 * (@see PerchRadiusThreshold)
-	 */
-	UPROPERTY(Category="Character Movement", EditAnywhere, BlueprintReadWrite, AdvancedDisplay, meta=(ClampMin="0", UIMin="0"))
-	float PerchAdditionalHeight;
-
-	/** Actor's current physics mode. This is automatically replicated through the CharacterOwner and for client-server movement functions. */
 	UPROPERTY(Category=MovementMode, BlueprintReadOnly)
 	TEnumAsByte<enum EMovementMode> MovementMode;
 
-	/** Current custom sub-mode if MovementMode is set to Custom. This is automatically replicated through the CharacterOwner and for client-server movement functions. */
+	/**
+	 * Current custom sub-mode if MovementMode is set to Custom. This is automatically replicated through the Character owner and for client-server movement functions.
+	 * @see SetMovementMode()
+	 */
 	UPROPERTY(Category=MovementMode, BlueprintReadOnly)
 	uint8 CustomMovementMode;
 
-	/** Save Base location to check in UpdateBasedMovement() whether base moved in the last frame, and therefore pawn needs an update. */
+	/** Saved location of object we are standing on, for UpdateBasedMovement() to determine if base moved in the last frame, and therefore pawn needs an update. */
 	FVector OldBaseLocation;
 
-	/** Save Base rotation to check in UpdateBasedMovement() whether base moved in the last frame, and therefore pawn needs an update. */
+	/** Saved location of object we are standing on, for UpdateBasedMovement() to determine if base moved in the last frame, and therefore pawn needs an update. */
 	FQuat OldBaseQuat;
 
 	/** Custom gravity scale. Gravity is multiplied by this amount for the character. */
@@ -258,6 +262,21 @@ public:
 	UPROPERTY(Category="Character Movement", EditAnywhere, BlueprintReadWrite)
 	float Buoyancy;
 
+	/**
+	 * Don't allow the character to perch on the edge of a surface if the contact is this close to the edge of the capsule.
+	 * Note that characters will not fall off if they are within MaxStepHeight of a walkable surface below.
+	 */
+	UPROPERTY(Category="Character Movement", EditAnywhere, BlueprintReadWrite, AdvancedDisplay, meta=(ClampMin="0", UIMin="0"))
+	float PerchRadiusThreshold;
+
+	/**
+	 * When perching on a ledge, add this additional distance to MaxStepHeight when determining how high above a walkable floor we can perch.
+	 * Note that we still enforce MaxStepHeight to start the step up; this just allows the character to hang off the edge or step slightly higher off the floor.
+	 * (@see PerchRadiusThreshold)
+	 */
+	UPROPERTY(Category="Character Movement", EditAnywhere, BlueprintReadWrite, AdvancedDisplay, meta=(ClampMin="0", UIMin="0"))
+	float PerchAdditionalHeight;
+
 	/** Change in rotation per second, used when UseControllerDesiredRotation or OrientRotationToMovement are true. */
 	UPROPERTY(Category="Character Movement", EditAnywhere, BlueprintReadWrite)
 	FRotator RotationRate;
@@ -273,10 +292,18 @@ public:
 	UPROPERTY(Category="Character Movement", EditAnywhere, BlueprintReadWrite)
 	uint32 bOrientRotationToMovement:1;
 
-	/** true during movement update. Attempts to change CharacterOwner and UpdatedComponent are deferred until after an update. */
+protected:
+
+	/**
+	 * True during movement update.
+	 * Used internally so that attempts to change CharacterOwner and UpdatedComponent are deferred until after an update.
+	 * @see IsMovementInProgress()
+	 */
 	UPROPERTY()
 	uint32 bMovementInProgress:1;
 
+public:
+
 	/**
 	 * If true, high-level movement updates will be wrapped in a movement scope that accumulates updates and defers a bulk of the work until the end.
 	 * When enabled, touch and hit events will not be triggered until the end of multiple moves within an update, which can improve performance.
@@ -411,11 +438,11 @@ protected:
 	UPROPERTY()
 	FVector LastUpdateLocation;
 
-	/** Accumulated impulse added this tick */
+	/** Accumulated impulse to be added next tick. */
 	UPROPERTY()
 	FVector PendingImpulseToApply;
 
-	/** Accumulated force added this tick */
+	/** Accumulated force to be added next tick. */
 	UPROPERTY()
 	FVector PendingForceToApply;
 
@@ -458,6 +485,7 @@ public:
 	 * Increasing this value can address issues with fast-moving objects or complex collision scenarios, at the cost of performance.
 	 *
 	 * WARNING: if (MaxSimulationTimeStep * MaxSimulationIterations) is too low for the min framerate, the last simulation step may exceed MaxSimulationTimeStep to complete the simulation.
+	 * @see MaxSimulationTimeStep
 	 */
 	UPROPERTY(Category="Character Movement", EditAnywhere, BlueprintReadWrite, AdvancedDisplay, meta=(ClampMin="1", ClampMax="25", UIMin="1", UIMax="25"))
 	int32 MaxSimulationIterations;
@@ -471,21 +499,22 @@ public:
 	float JumpOutOfWaterPitch;
 
 	/**
-	 * Impacts on the upper hemisphere will have the vertical velocity scaled by this factor of the Normal's Z component (in AdjustUpperHemisphereImpact).
-	 * A value of 0 indicates that this behavior is disabled.
+	 * Impacts on the upper hemisphere will have the vertical velocity scaled by this factor of the normal's Z component (in AdjustUpperHemisphereImpact()).
+	 * This can dampen such impacts to feel less like a capsule. A value of 0 indicates that this behavior is disabled.
+	 * @see AdjustUpperHemisphereImpact()
 	 */
 	UPROPERTY(Category="Character Movement", EditAnywhere, BlueprintReadWrite, AdvancedDisplay)
 	float UpperImpactNormalScale;
 
-	/** Information about the floor the Character is standing on (only used by MOVE_Walking) */
+	/** Information about the floor the Character is standing on (updated only during walking movement). */
 	UPROPERTY(Category="Character Movement", VisibleInstanceOnly, BlueprintReadOnly)
 	FFindFloorResult CurrentFloor;
 
-	/** Default movement mode when not in water */
+	/** Default movement mode when not in water. Used at player startup or when teleported. */
 	UPROPERTY(Category="Character Movement", EditAnywhere, BlueprintReadWrite)
 	TEnumAsByte<enum EMovementMode> DefaultLandMovementMode;
 
-	/** Default movement mode when in water */
+	/** Default movement mode when in water. Used at player startup or when teleported. */
 	UPROPERTY(Category="Character Movement", EditAnywhere, BlueprintReadWrite)
 	TEnumAsByte<enum EMovementMode> DefaultWaterMovementMode;
 
@@ -508,8 +537,11 @@ public:
 	UPROPERTY(Category="Character Movement", EditAnywhere, BlueprintReadWrite)
 	uint32 bImpartBaseVelocityZ:1;
 
-	/** If true, impart the base component's tangential components of angular velocity when jumping or falling off it.
-	  * Only those components of the velocity allowed by the separate component settings (bImpartBaseVelocity[XYZ]) will be applied. */
+	/**
+	 * If true, impart the base component's tangential components of angular velocity when jumping or falling off it.
+	 * Only those components of the velocity allowed by the separate component settings (bImpartBaseVelocityX etc) will be applied.
+	 * @see bImpartBaseVelocityX, bImpartBaseVelocityY, bImpartBaseVelocityZ
+	 */
 	UPROPERTY(Category="Character Movement", EditAnywhere, BlueprintReadWrite)
 	uint32 bImpartBaseAngularVelocity:1;
 
@@ -517,11 +549,11 @@ public:
 	UPROPERTY(Transient)
 	uint32 bJustTeleported:1;
 
-	/** True when a network replication update is received for simulated proxies */
+	/** True when a network replication update is received for simulated proxies. */
 	UPROPERTY(Transient)
 	uint32 bNetworkUpdateReceived:1;
 
-	/** True when the networked movement mode has been replicated */
+	/** True when the networked movement mode has been replicated. */
 	UPROPERTY(Transient)
 	uint32 bNetworkMovementModeChanged:1;
 
@@ -529,7 +561,7 @@ public:
 	UPROPERTY()
 	uint32 bNotifyApex:1;
 
-	/** instantly stop in flying mode */
+	/** Instantly stop when in flying mode and no acceleration is being applied. */
 	UPROPERTY()
 	uint32 bCheatFlying:1;
 
@@ -550,8 +582,8 @@ public:
 	uint32 bIgnoreBaseRotation:1;
 
 	/** 
-	 *  Set this to true if riding on a moving base that you know is clear from non-moving world obstructions.
-	 *  This can solve move-order dependencies when riding an interp actor, and it's faster.
+	 * Set this to true if riding on a moving base that you know is clear from non-moving world obstructions.
+	 * Optimization to avoid sweeps during based movement, use with care.
 	 */
 	UPROPERTY()
 	uint32 bFastAttachedMove:1;
@@ -586,6 +618,8 @@ public:
 
 protected:
 
+	// AI PATH FOLLOWING
+
 	/** Was velocity requested by path following? */
 	UPROPERTY(Transient)
 	uint32 bHasRequestedVelocity:1;
@@ -612,7 +646,10 @@ protected:
 
 public:
 
-	/** Velocity requested by path following */
+	/**
+	 * Velocity requested by path following.
+	 * @see RequestDirectMove()
+	 */
 	UPROPERTY(Transient)
 	FVector RequestedVelocity;
 
@@ -699,6 +736,12 @@ public:
 	UFUNCTION(BlueprintCallable, Category="Pawn|Components|CharacterMovement")
 	bool IsWalking() const;
 
+	/**
+	 * @return true if currently performing a movement update.
+	 * @see bMovementInProgress
+	 */
+	bool IsMovementInProgress() const { return bMovementInProgress; }
+
 	//BEGIN UNavMovementComponent Interface
 	virtual void RequestDirectMove(const FVector& MoveVelocity, bool bForceMaxSpeed) override;
 	virtual bool CanStopPathFollowing() const override;
@@ -873,7 +916,7 @@ public:
 	/** @return true if we can step up on the actor in the given FHitResult. */
 	virtual bool CanStepUp(const FHitResult& Hit) const;
 
-	/**  Struct updated by StepUp() method to return result of final step down, if applicable. */
+	/** Struct updated by StepUp() to return result of final step down, if applicable. */
 	struct FStepDownResult
 	{
 		uint32 bComputedFloor:1;		// True if the floor was computed as a result of the step down.
@@ -888,21 +931,21 @@ public:
 	/** 
 	 * Move up steps or slope. Does nothing and returns false if CanStepUp(Hit) returns false.
 	 *
-	 * @param GravDir:			Gravity vector
-	 * @param Delta:			Requested move
-	 * @param Hit:				[In] The hit before the step up.
-	 * @param OutStepDownResult:[Out] If non-null, a floor check will be performed if possible as part of the final step down, and it will be updated to reflect this result.
+	 * @param GravDir			Gravity vector
+	 * @param Delta				Requested move
+	 * @param Hit				[In] The hit before the step up.
+	 * @param OutStepDownResult	[Out] If non-null, a floor check will be performed if possible as part of the final step down, and it will be updated to reflect this result.
 	 * @return true if the step up was successful.
 	 */
 	virtual bool StepUp(const FVector& GravDir, const FVector& Delta, const FHitResult &Hit, struct UCharacterMovementComponent::FStepDownResult* OutStepDownResult = NULL);
 
-	/** Update the base of the actor */
+	/** Update the base of the character, which is the obejct we are standing on. */
 	virtual void SetBase(UPrimitiveComponent* NewBase, const FName BoneName = NAME_None, bool bNotifyActor=true);
 
-	/** Applies repulsion force to all touched components */
+	/** Applies repulsion force to all touched components. */
 	void ApplyRepulsionForce(float DeltaTime);
 	
-	/** Applies momentum accumulated through AddImpulse()/AddForce() */
+	/** Applies momentum accumulated through AddImpulse() and AddForce(). */
 	void ApplyAccumulatedForces(float DeltaSeconds);	
 
 	/** 
@@ -991,15 +1034,16 @@ protected:
 	virtual float LimitAirControl(float DeltaTime, float TickAirControl, const FVector& FallAcceleration, const FHitResult& HitResult);
 	
 
-public:
-
-	/** Handle landing against Hit surface over remaingTime and iterations */
+	/** Handle landing against Hit surface over remaingTime and iterations, calling SetPostLandedPhysics() and starting the new movement mode. */
 	virtual void ProcessLanded(const FHitResult& Hit, float remainingTime, int32 Iterations);
 
-	/** Set new physics after landing */
+	/** Use new physics after landing. Defaults to swimming if in water, walking otherwise. */
 	virtual void SetPostLandedPhysics(const FHitResult& Hit);
 
-	/** Called by owning Character on successful teleport. */
+
+public:
+
+	/** Called by owning Character upon successful teleport from AActor::TeleportTo(). */
 	virtual void OnTeleported();
 
 	/**
@@ -1015,14 +1059,14 @@ public:
 	 */
 	virtual void UnCrouch(bool bClientSimulation = false);
 
-	// @return true if the character is allowed to crouch in the current state. By default it is allowed when walking or falling, if CanEverCrouch() is true.
+	/** @return true if the character is allowed to crouch in the current state. By default it is allowed when walking or falling, if CanEverCrouch() is true. */
 	virtual bool CanCrouchInCurrentState() const;
 	
-	/** @return true if there is a suitable floor SideStep from current position  */
+	/** @return true if there is a suitable floor SideStep from current position. */
 	bool CheckLedgeDirection(const FVector& OldLocation, const FVector& SideStep, const FVector& GravDir);
 
 	/** 
-	 * @param Delta is the current move delta (which ended up going over a ledge.
+	 * @param Delta is the current move delta (which ended up going over a ledge).
 	 * @return new delta which moves along the ledge
 	 */
 	FVector GetLedgeMove(const FVector& OldLocation, const FVector& Delta, const FVector& GravDir);
@@ -1128,7 +1172,6 @@ public:
 	UFUNCTION(BlueprintCallable, Category="Pawn|Components|CharacterMovement")
 	void SetWalkableFloorAngle(float InWalkableFloorAngle);
 
-
 	/** Get the Z component of the normal of the steepest walkable surface for the character. Any lower than this and it is not walkable. */
 	FORCEINLINE float GetWalkableFloorZ() const { return WalkableFloorZ; }
 
@@ -1140,6 +1183,10 @@ public:
 	UFUNCTION(BlueprintCallable, Category="Pawn|Components|CharacterMovement")
 	void SetWalkableFloorZ(float InWalkableFloorZ);
 
+	/** Post-physics tick function for this character */
+	UPROPERTY()
+	struct FCharacterMovementComponentPreClothTickFunction PreClothComponentTick;
+
 	/** Tick function called after physics (sync scene) has finished simulation, before cloth */
 	virtual void PreClothTick(float DeltaTime, FCharacterMovementComponentPreClothTickFunction& ThisTickFunction);
 
@@ -1217,8 +1264,9 @@ protected:
 	virtual FVector ComputeSlideVector(const FVector& Delta, const float Time, const FVector& Normal, const FHitResult& Hit) const override;
 
 	/**
-	 * Given an upward impact on the top of the capsule, allows calculation of a different movement delta.
+	 * Given an upward impact on the top of the capsule, allows calculation of a different movement delta. This can dampen such impacts to feel less like a capsule.
 	 * Default implementation scales the Delta Z value by: 1 - (Abs(Normal.Z) * UpperImpactNormalScale), clamped to [0..1]
+	 * @see UpperImpactNormalScale
 	 */
 	virtual FVector AdjustUpperHemisphereImpact(const FVector& Delta, const FHitResult& Hit) const;
 
diff --git a/Engine/Source/Runtime/Engine/Classes/GameFramework/Controller.h b/Engine/Source/Runtime/Engine/Classes/GameFramework/Controller.h
--- a/Engine/Source/Runtime/Engine/Classes/GameFramework/Controller.h
+++ b/Engine/Source/Runtime/Engine/Classes/GameFramework/Controller.h
@@ -1,24 +1,5 @@
 // Copyright 1998-2014 Epic Games, Inc. All Rights Reserved.
 
-//=============================================================================
-// Controller, the base class of players or AI.
-//
-// Controllers are non-physical actors that can possess a Pawn to control
-// its actions.  PlayerControllers are used by human players to control pawns, while
-// AIControllers implement the artificial intelligence for the pawns they control.
-// Controllers take control of a pawn using their Possess() method, and relinquish
-// control of the pawn by calling UnPossess().
-//
-// Controllers receive notifications for many of the events occurring for the Pawn they
-// are controlling.  This gives the controller the opportunity to implement the behavior
-// in response to this event, intercepting the event and superseding the Pawn's default
-// behavior.
-//
-// The control rotation (accessed via GetControlRotation()), determines the aiming
-// orientation of the controlled Pawn.
-//
-//=============================================================================
-
 #pragma once
 #include "AI/Navigation/NavAgentInterface.h"
 #include "Controller.generated.h"
@@ -29,6 +10,26 @@ class UPathFollowingComponent;
 UDELEGATE(BlueprintAuthorityOnly)
 DECLARE_DYNAMIC_MULTICAST_DELEGATE_FourParams( FInstigatedAnyDamageSignature, float, Damage, const class UDamageType*, DamageType, class AActor*, DamagedActor, class AActor*, DamageCauser );
 
+//=============================================================================
+/**
+ * Controllers are non-physical actors that can possess a Pawn to control
+ * its actions.  PlayerControllers are used by human players to control pawns, while
+ * AIControllers implement the artificial intelligence for the pawns they control.
+ * Controllers take control of a pawn using their Possess() method, and relinquish
+ * control of the pawn by calling UnPossess().
+ *
+ * Controllers receive notifications for many of the events occurring for the Pawn they
+ * are controlling.  This gives the controller the opportunity to implement the behavior
+ * in response to this event, intercepting the event and superseding the Pawn's default
+ * behavior.
+ *
+ * ControlRotation (accessed via GetControlRotation()), determines the viewing/aiming
+ * direction of the controlled Pawn and is affected by input such as from a mouse or gamepad.
+ * 
+ * @see https://docs.unrealengine.com/latest/INT/Gameplay/Framework/Controller/
+ */
+//=============================================================================
+
 UCLASS(abstract, notplaceable, NotBlueprintable, HideCategories=(Collision,Rendering,"Utilities|Transformation")) 
 class ENGINE_API AController : public AActor, public INavAgentInterface
 {
diff --git a/Engine/Source/Runtime/Engine/Classes/GameFramework/DefaultPawn.h b/Engine/Source/Runtime/Engine/Classes/GameFramework/DefaultPawn.h
--- a/Engine/Source/Runtime/Engine/Classes/GameFramework/DefaultPawn.h
+++ b/Engine/Source/Runtime/Engine/Classes/GameFramework/DefaultPawn.h
@@ -13,6 +13,9 @@ class UFloatingPawnMovement;
 class USphereComponent;
 class UStaticMeshComponent;
 
+/**
+ * 
+ */
 UCLASS(config=Game, Blueprintable, BlueprintType)
 class ENGINE_API ADefaultPawn : public APawn
 {
diff --git a/Engine/Source/Runtime/Engine/Classes/GameFramework/FloatingPawnMovement.h b/Engine/Source/Runtime/Engine/Classes/GameFramework/FloatingPawnMovement.h
--- a/Engine/Source/Runtime/Engine/Classes/GameFramework/FloatingPawnMovement.h
+++ b/Engine/Source/Runtime/Engine/Classes/GameFramework/FloatingPawnMovement.h
@@ -8,6 +8,9 @@
 #include "GameFramework/PawnMovementComponent.h"
 #include "FloatingPawnMovement.generated.h"
 
+/**
+ * 
+ */
 UCLASS()
 class ENGINE_API UFloatingPawnMovement : public UPawnMovementComponent
 {
diff --git a/Engine/Source/Runtime/Engine/Classes/GameFramework/MovementComponent.h b/Engine/Source/Runtime/Engine/Classes/GameFramework/MovementComponent.h
--- a/Engine/Source/Runtime/Engine/Classes/GameFramework/MovementComponent.h
+++ b/Engine/Source/Runtime/Engine/Classes/GameFramework/MovementComponent.h
@@ -13,12 +13,25 @@ struct FCollisionQueryParams;
 struct FCollisionResponseParams;
 struct FCollisionShape;
 
+
+/**
+ * MovementComponent is an abstract component class that defines functionality for moving a PrimitiveComponent (our UpdatedComponent) each tick.
+ * Base functionality includes:
+ *    - Restricting movement to a plane or axis.
+ *    - Utility functions for special handling of collision results (SlideAlongSurface(), ComputeSlideVector(), TwoWallAdjust()).
+ *    - Utility functions for moving when there may be initial penetration (SafeMoveUpdatedComponent(), ResolvePenetration()).
+ *    - Automatically registering the component tick and finding a component to move on the owning Actor.
+ */
 UCLASS(ClassGroup=Movement, abstract, BlueprintType)
 class ENGINE_API UMovementComponent : public UActorComponent
 {
 	GENERATED_UCLASS_BODY()
 
-	/** The component we move and update. **/
+	/**
+	 * The component we move and update.
+	 * If this is null at startup and bAutoRegisterUpdatedComponent is true, the owning Actor's root component will automatically be set as our UpdatedComponent at startup.
+	 * @see bAutoRegisterUpdatedComponent, SetUpdatedComponent()
+	 */
 	UPROPERTY(BlueprintReadOnly, Category=MovementComponent)
 	UPrimitiveComponent* UpdatedComponent;
 
@@ -28,7 +41,7 @@ class ENGINE_API UMovementComponent : public UActorComponent
 	 */
 	EMoveComponentFlags MoveComponentFlags;
 
-	/** Current velocity of moved component, replicated to clients using Actor ReplicatedMovement property. */
+	/** Current velocity of updated component. */
 	UPROPERTY(EditAnywhere, BlueprintReadWrite, Category=Velocity)
 	FVector Velocity;
 
@@ -50,15 +63,17 @@ protected:
 	FVector PlaneConstraintOrigin;
 
 public:
-	/** If true, movement will be constrained to a plane.	 */
+	/**
+	 * If true, movement will be constrained to a plane.
+	 * @see PlaneConstraintNormal, PlaneConstraintOrigin
+	 */
 	UPROPERTY(EditDefaultsOnly, BlueprintReadOnly, Category=PlanarMovement)
 	uint32 bConstrainToPlane:1;
 
-	/** If true, and if plane constraints are enabled, then the component updated by this MovementComponent will be snapped to the plane when first attached (in SetUpdatedComponent()).	 */
+	/** If true and plane constraints are enabled, then the updated component will be snapped to the plane when first attached. */
 	UPROPERTY(EditDefaultsOnly, BlueprintReadOnly, Category=PlanarMovement)
 	uint32 bSnapToPlaneAtStart:1;
 
-
 	/** If true, skips TickComponent() if UpdatedComponent was not recently rendered. */
 	UPROPERTY(EditAnywhere, BlueprintReadWrite, Category=MovementComponent)
 	uint32 bUpdateOnlyIfRendered:1;
@@ -159,7 +174,7 @@ public:
 	/** Update ComponentVelocity of UpdatedComponent. This needs to be called by derived classes at the end of an update whenever Velocity has changed.	 */
 	virtual void UpdateComponentVelocity();
 
-	/** Initialize collision params appropriately based on our collision settings. Use this before any Line or Sweep tests. */
+	/** Initialize collision params appropriately based on our collision settings. Use this before any Line, Overlap, or Sweep tests. */
 	virtual void InitCollisionParams(FCollisionQueryParams &OutParams, FCollisionResponseParams& OutResponseParam) const;
 
 	/** Return true if the given collision shape overlaps other geometry at the given location and rotation. The collision params are set by InitCollisionParams(). */
@@ -168,6 +183,7 @@ public:
 	/**
 	 * Moves our UpdatedComponent by the given Delta, and sets rotation to NewRotation.
 	 * Respects the plane constraint, if enabled.
+	 * @return True if some movement occurred, false if no movement occurred. Result of any impact will be stored in OutHit.
 	 */
 	virtual bool MoveUpdatedComponent(const FVector& Delta, const FRotator& NewRotation, bool bSweep, FHitResult* OutHit = NULL);
 
@@ -234,6 +250,7 @@ public:
 
 	/**
 	 * Adds force from radial force components.
+	 * Intended to be overridden by subclasses; default implementation does nothing.
 	 * @param	Origin		The origin of the force
 	 * @param	Radius		The radius in which the force will be applied
 	 * @param	Strength	The strength of the force
@@ -243,6 +260,7 @@ public:
 
 	/**
 	 * Adds impulse from radial force components.
+	 * Intended to be overridden by subclasses; default implementation does nothing.
 	 * @param	Origin		The origin of the force
 	 * @param	Radius		The radius in which the force will be applied
 	 * @param	Strength	The strength of the force
@@ -252,40 +270,43 @@ public:
 	virtual void AddRadialImpulse(const FVector& Origin, float Radius, float Strength, ERadialImpulseFalloff Falloff, bool bVelChange);
 
 	/**
-	 * Sets the normal of the plane that constrains movement, if plane constraint is enabled.
-	 * @param PlaneNormal	- The normal of the plane. If non-zero, it will be normalized.
+	 * Sets the normal of the plane that constrains movement, enforced if the plane constraint is enabled.
+	 * @param PlaneNormal	The normal of the plane. If non-zero in length, it will be normalized.
 	 */
 	UFUNCTION(BlueprintCallable, Category="Components|Movement|Planar")
 	virtual void SetPlaneConstraintNormal(FVector PlaneNormal);
 
-	/** Uses the Forward and Up vectors to compute the plane that constrains movement, if plane constraint is enabled.	 */
+	/** Uses the Forward and Up vectors to compute the plane that constrains movement, enforced if the plane constraint is enabled. */
 	UFUNCTION(BlueprintCallable, Category="Components|Movement|Planar")
 	virtual void SetPlaneConstraintFromVectors(FVector Forward, FVector Up);
 
-	/** Sets the origin of the plane that constrains movement, if plane constraint is enabled.	 */
+	/** Sets the origin of the plane that constrains movement, enforced if the plane constraint is enabled. */
 	UFUNCTION(BlueprintCallable, Category="Components|Movement|Planar")
 	virtual void SetPlaneConstraintOrigin(FVector PlaneOrigin);
 
-	/** @return The normal of the plane that constrains movement, if plane constraint is enabled.	 */
+	/** @return The normal of the plane that constrains movement, enforced if the plane constraint is enabled. */
 	UFUNCTION(BlueprintCallable, Category="Components|Movement|Planar")
 	const FVector& GetPlaneConstraintNormal() const;
 
 	/**
-	 * @return The origin of the plane that constrains movement, if plane constraint is enabled.
-	 * This defines the behavior of snapping a position to the plane, such as by SnapUpdatedComponentToPlane().
+	 * Get the plane constraint origin. This defines the behavior of snapping a position to the plane, such as by SnapUpdatedComponentToPlane().
+	 * @return The origin of the plane that constrains movement, if the plane constraint is enabled.
 	 */
 	UFUNCTION(BlueprintCallable, Category="Components|Movement|Planar")
 	const FVector& GetPlaneConstraintOrigin() const;
 
-	/** Constrain a direction vector to the plane constraint, if enabled.	 */
+	/**
+	 * Constrain a direction vector to the plane constraint, if enabled.
+	 * @see SetPlaneConstraint
+	 */
 	UFUNCTION(BlueprintCallable, Category="Components|Movement|Planar")
 	virtual FVector ConstrainDirectionToPlane(FVector Direction) const;
 
-	/** Constrain a position vector to the plane constraint, if enabled.	 */
+	/** Constrain a position vector to the plane constraint, if enabled. */
 	UFUNCTION(BlueprintCallable, Category="Components|Movement|Planar")
 	virtual FVector ConstrainLocationToPlane(FVector Location) const;
 
-	/** Snap the updated component to the plane constraint, if enabled.	 */
+	/** Snap the updated component to the plane constraint, if enabled. */
 	UFUNCTION(BlueprintCallable, Category="Components|Movement|Planar")
 	virtual void SnapUpdatedComponentToPlane();
 
diff --git a/Engine/Source/Runtime/Engine/Classes/GameFramework/NavMovementComponent.h b/Engine/Source/Runtime/Engine/Classes/GameFramework/NavMovementComponent.h
--- a/Engine/Source/Runtime/Engine/Classes/GameFramework/NavMovementComponent.h
+++ b/Engine/Source/Runtime/Engine/Classes/GameFramework/NavMovementComponent.h
@@ -13,6 +13,10 @@
 #include "Components/PrimitiveComponent.h"
 #include "NavMovementComponent.generated.h"
 
+/**
+ * NavMovementComponent defines base functionality for MovementComponents that move any 'agent' that may be involved in AI pathfinding.
+
+ */
 UCLASS(abstract)
 class ENGINE_API UNavMovementComponent : public UMovementComponent
 {
diff --git a/Engine/Source/Runtime/Engine/Classes/GameFramework/Pawn.h b/Engine/Source/Runtime/Engine/Classes/GameFramework/Pawn.h
--- a/Engine/Source/Runtime/Engine/Classes/GameFramework/Pawn.h
+++ b/Engine/Source/Runtime/Engine/Classes/GameFramework/Pawn.h
@@ -21,9 +21,10 @@ class UPlayer;
 class UPrimitiveComponent;
 
 /** 
- *	Pawn is the base class of all actors that can be possessed by players or AI.
- *	Pawns are the physical representations of players and creatures in a level.
- *	They are responsible for all physical interaction between the player or AI and the world.
+ * Pawn is the base class of all actors that can be possessed by players or AI.
+ * Pawns are the physical representations of players and creatures in a level.
+ * They are responsible for all physical interaction between the player or AI and the world.
+ * @see https://docs.unrealengine.com/latest/INT/Gameplay/Framework/Pawn/
  */
 UCLASS(abstract, config=Game, BlueprintType, Blueprintable, hidecategories=(Navigation, "AI|Navigation"))
 class ENGINE_API APawn : public AActor, public INavAgentInterface
diff --git a/Engine/Source/Runtime/Engine/Classes/GameFramework/PawnMovementComponent.h b/Engine/Source/Runtime/Engine/Classes/GameFramework/PawnMovementComponent.h
--- a/Engine/Source/Runtime/Engine/Classes/GameFramework/PawnMovementComponent.h
+++ b/Engine/Source/Runtime/Engine/Classes/GameFramework/PawnMovementComponent.h
@@ -8,6 +8,11 @@
 #include "GameFramework/NavMovementComponent.h"
 #include "PawnMovementComponent.generated.h"
 
+/** 
+ * PawnMovementComponent can be used to update movement for an associated Pawn.
+ * It also provides ways to accumulate and read directional input in a generic way (with AddInputVector(), ConsumeInputVector(), etc).
+ * @see APawn
+ */
 UCLASS(abstract)
 class ENGINE_API UPawnMovementComponent : public UNavMovementComponent
 {
@@ -15,13 +20,13 @@ class ENGINE_API UPawnMovementComponent : public UNavMovementComponent
 
 public:
 
-	// Overridden to only allow registration with components owned by a Pawn.
+	/** Overridden to only allow registration with components owned by a Pawn. */
 	virtual void SetUpdatedComponent(class UPrimitiveComponent* NewUpdatedComponent) override;
 
 	/**
 	 * Adds the given vector to the accumulated input in world space. Input vectors are usually between 0 and 1 in magnitude. 
 	 * They are accumulated during a frame then applied as acceleration during the movement update.
-	 * @see Pawn.AddMovementInput()
+	 * @see APawn::AddMovementInput()
 	 *
 	 * @param WorldDirection:	Direction in world space to apply input
 	 * @param ScaleValue:		Scale to apply to input. This can be used for analog input, ie a value of 0.5 applies half the normal value.
@@ -69,7 +74,7 @@ public:
 
 protected:
 
-	/** Pawn which owns this component. */
+	/** Pawn that owns this component. */
 	UPROPERTY()
 	class APawn* PawnOwner;
 
diff --git a/Engine/Source/Runtime/Engine/Classes/GameFramework/PlayerController.h b/Engine/Source/Runtime/Engine/Classes/GameFramework/PlayerController.h
--- a/Engine/Source/Runtime/Engine/Classes/GameFramework/PlayerController.h
+++ b/Engine/Source/Runtime/Engine/Classes/GameFramework/PlayerController.h
@@ -59,11 +59,14 @@ struct FDynamicForceFeedbackDetails
 /**
  * PlayerControllers are used by human players to control Pawns.
  *
- * The control rotation (accessed via GetControlRotation()), determines the aiming
+ * ControlRotation (accessed via GetControlRotation()), determines the aiming
  * orientation of the controlled Pawn.
  *
- * PlayerControllers exist on the server for every player-controlled pawn,
- * and also on the controlling client's machine.
+ * In networked games, PlayerControllers exist on the server for every player-controlled pawn,
+ * and also on the controlling client's machine. They do *not* exist on a client's
+ * machine for pawns controlled by remote players elsewhere on the network.
+ *
+ * @see https://docs.unrealengine.com/latest/INT/Gameplay/Framework/Controller/PlayerController/
  */
 //=============================================================================
 
diff --git a/Engine/Source/Runtime/Engine/Classes/GameFramework/ProjectileMovementComponent.h b/Engine/Source/Runtime/Engine/Classes/GameFramework/ProjectileMovementComponent.h
--- a/Engine/Source/Runtime/Engine/Classes/GameFramework/ProjectileMovementComponent.h
+++ b/Engine/Source/Runtime/Engine/Classes/GameFramework/ProjectileMovementComponent.h
@@ -9,6 +9,7 @@
  * Behavior such as bouncing after impacts and homing toward a target are supported.
  * If the updated component is simulating physics, only the initial launch parameters (when initial velocity is non-zero) will affect the projectile,
  * and the physics sim will take over from there.
+ * @see UMovementComponent
  */
 UCLASS(ClassGroup=Movement, meta=(BlueprintSpawnableComponent), ShowCategories=(Velocity))
 class ENGINE_API UProjectileMovementComponent : public UMovementComponent
diff --git a/Engine/Source/Runtime/Engine/Classes/GameFramework/RotatingMovementComponent.h b/Engine/Source/Runtime/Engine/Classes/GameFramework/RotatingMovementComponent.h
--- a/Engine/Source/Runtime/Engine/Classes/GameFramework/RotatingMovementComponent.h
+++ b/Engine/Source/Runtime/Engine/Classes/GameFramework/RotatingMovementComponent.h
@@ -5,18 +5,25 @@
 #include "RotatingMovementComponent.generated.h"
 
 /**
- * Movement component updates position of associated PrimitiveComponent during its tick.
+ * Performs continuous rotation of a component at a specific rotation rate.
+ * Rotation can optionally be offset around a pivot point.
+ * Collision testing is not performed during movement.
  */
 UCLASS(ClassGroup=Movement, meta=(BlueprintSpawnableComponent), HideCategories=(Velocity))
 class ENGINE_API URotatingMovementComponent : public UMovementComponent
 {
 	GENERATED_UCLASS_BODY()
 
-	/** How fast to update roll/pitch/yaw of UpdatedComponent */
+	/** How fast to update roll/pitch/yaw of the component we update. */
 	UPROPERTY(EditAnywhere, BlueprintReadWrite, Category=RotatingComponent)
 	FRotator RotationRate;
 
-	/** Translation of pivot point. Always relative to current rotation. */
+	/**
+	 * Translation of pivot point around which we rotate, relative to current rotation.
+	 * For instance, with PivotTranslation set to (X=+100, Y=0, Z=0), rotation will occur
+	 * around the point +100 units along the local X axis from the center of the object,
+	 * rather than around the object's origin (the default).
+	 */
 	UPROPERTY(EditAnywhere, BlueprintReadWrite, Category=RotatingComponent)
 	FVector PivotTranslation;
 
@@ -25,6 +32,7 @@ class ENGINE_API URotatingMovementComponent : public UMovementComponent
 	uint32 bRotationInLocalSpace:1;
 
 	//Begin UActorComponent Interface
+	/** Applies rotation to UpdatedComponent. */
 	virtual void TickComponent(float DeltaTime, enum ELevelTick TickType, FActorComponentTickFunction *ThisTickFunction) override;
 	//End UActorComponent Interface
 };
