WIP Mobile Preview
 - Store global bitfield of active feature levels
 - Generate shaders for all active feature levels
 - Add option to experimental settings to enable dynamic feature level
 - Redirect GRHIFeatureLevel to GMaxRHIFeatureLevel - Effectively const now, ready to be removed

diff --git a/Engine/Plugins/Experimental/Paper2D/Source/Paper2D/Private/PaperFlipbookSceneProxy.cpp b/Engine/Plugins/Experimental/Paper2D/Source/Paper2D/Private/PaperFlipbookSceneProxy.cpp
--- a/Engine/Plugins/Experimental/Paper2D/Source/Paper2D/Private/PaperFlipbookSceneProxy.cpp
+++ b/Engine/Plugins/Experimental/Paper2D/Source/Paper2D/Private/PaperFlipbookSceneProxy.cpp
@@ -12,5 +12,5 @@ FPaperFlipbookSceneProxy::FPaperFlipbookSceneProxy(const UPaperFlipbookComponent
 	//@TODO: PAPER2D: WireframeColor = RenderComp->GetWireframeColor();
 
 	Material = InComponent->GetMaterial(0);
-	MaterialRelevance = InComponent->GetMaterialRelevance();
+	MaterialRelevance = InComponent->GetMaterialRelevance(GetScene()->GetFeatureLevel());
 }
diff --git a/Engine/Plugins/Experimental/Paper2D/Source/Paper2D/Private/PaperSpriteSceneProxy.cpp b/Engine/Plugins/Experimental/Paper2D/Source/Paper2D/Private/PaperSpriteSceneProxy.cpp
--- a/Engine/Plugins/Experimental/Paper2D/Source/Paper2D/Private/PaperSpriteSceneProxy.cpp
+++ b/Engine/Plugins/Experimental/Paper2D/Source/Paper2D/Private/PaperSpriteSceneProxy.cpp
@@ -15,7 +15,7 @@ FPaperSpriteSceneProxy::FPaperSpriteSceneProxy(const UPaperSpriteComponent* InCo
 	Material = InComponent->GetMaterial(0);
 	AlternateMaterial = InComponent->GetMaterial(1);
 	MaterialSplitIndex = INDEX_NONE;
-	MaterialRelevance = InComponent->GetMaterialRelevance();
+	MaterialRelevance = InComponent->GetMaterialRelevance(GetScene()->GetFeatureLevel());
 
 	SourceSprite = InComponent->SourceSprite; //@TODO: This is totally not threadsafe, and won't keep up to date if the actor's sprite changes, etc....
 }
diff --git a/Engine/Plugins/Experimental/Paper2D/Source/Paper2D/Private/PaperTileMapRenderSceneProxy.cpp b/Engine/Plugins/Experimental/Paper2D/Source/Paper2D/Private/PaperTileMapRenderSceneProxy.cpp
--- a/Engine/Plugins/Experimental/Paper2D/Source/Paper2D/Private/PaperTileMapRenderSceneProxy.cpp
+++ b/Engine/Plugins/Experimental/Paper2D/Source/Paper2D/Private/PaperTileMapRenderSceneProxy.cpp
@@ -19,7 +19,7 @@ FPaperTileMapRenderSceneProxy::FPaperTileMapRenderSceneProxy(const UPaperTileMap
 
 		if (Material)
 		{
-			MaterialRelevance = Material->GetRelevance();
+			MaterialRelevance = Material->GetRelevance(GetScene()->GetFeatureLevel());
 		}
 	}
 }
diff --git a/Engine/Plugins/Experimental/Paper2D/Source/Paper2D/Private/Terrain/PaperTerrainComponent.cpp b/Engine/Plugins/Experimental/Paper2D/Source/Paper2D/Private/Terrain/PaperTerrainComponent.cpp
--- a/Engine/Plugins/Experimental/Paper2D/Source/Paper2D/Private/Terrain/PaperTerrainComponent.cpp
+++ b/Engine/Plugins/Experimental/Paper2D/Source/Paper2D/Private/Terrain/PaperTerrainComponent.cpp
@@ -50,7 +50,7 @@ FPaperTerrainSceneProxy::FPaperTerrainSceneProxy(const UPaperTerrainComponent* I
 	for (const FPaperTerrainMaterialPair& Batch : DrawingData)
 	{
 		const UMaterialInterface* MaterialInterface = (Batch.Material != nullptr) ? Batch.Material : UMaterial::GetDefaultMaterial(MD_Surface);
-		MaterialRelevance |= MaterialInterface->GetRelevance_Concurrent();
+		MaterialRelevance |= MaterialInterface->GetRelevance_Concurrent(GetScene()->GetFeatureLevel());
 	}
 }
 
diff --git a/Engine/Plugins/Runtime/CableComponent/Source/CableComponent/Private/CableComponent.cpp b/Engine/Plugins/Runtime/CableComponent/Source/CableComponent/Private/CableComponent.cpp
--- a/Engine/Plugins/Runtime/CableComponent/Source/CableComponent/Private/CableComponent.cpp
+++ b/Engine/Plugins/Runtime/CableComponent/Source/CableComponent/Private/CableComponent.cpp
@@ -93,7 +93,7 @@ public:
 		: FPrimitiveSceneProxy(Component)
 		, Material(NULL)
 		, DynamicData(NULL)
-		, MaterialRelevance(Component->GetMaterialRelevance())
+		, MaterialRelevance(Component->GetMaterialRelevance(GetScene()->GetFeatureLevel()))
 		, NumSegments(Component->NumSegments)
 		, CableWidth(Component->CableWidth)
 		, NumSides(Component->NumSides)
diff --git a/Engine/Plugins/Runtime/CustomMeshComponent/Source/CustomMeshComponent/Private/CustomMeshComponent.cpp b/Engine/Plugins/Runtime/CustomMeshComponent/Source/CustomMeshComponent/Private/CustomMeshComponent.cpp
--- a/Engine/Plugins/Runtime/CustomMeshComponent/Source/CustomMeshComponent/Private/CustomMeshComponent.cpp
+++ b/Engine/Plugins/Runtime/CustomMeshComponent/Source/CustomMeshComponent/Private/CustomMeshComponent.cpp
@@ -81,7 +81,7 @@ public:
 
 	FCustomMeshSceneProxy(UCustomMeshComponent* Component)
 		: FPrimitiveSceneProxy(Component)
-		, MaterialRelevance(Component->GetMaterialRelevance())
+		, MaterialRelevance(Component->GetMaterialRelevance(GetScene()->GetFeatureLevel()))
 	{
 		const FColor VertexColor(255,255,255);
 
diff --git a/Engine/Source/Editor/LevelEditor/Private/LevelEditorActions.cpp b/Engine/Source/Editor/LevelEditor/Private/LevelEditorActions.cpp
--- a/Engine/Source/Editor/LevelEditor/Private/LevelEditorActions.cpp
+++ b/Engine/Source/Editor/LevelEditor/Private/LevelEditorActions.cpp
@@ -455,8 +455,7 @@ bool FLevelEditorActionCallbacks::IsMaterialQualityLevelChecked( EMaterialQualit
 
 void FLevelEditorActionCallbacks::SetFeatureLevelPreview(ERHIFeatureLevel::Type InPreviewFeatureLevel)
 {
-	UWorld::ChangeAllWorldFeatureLevels(InPreviewFeatureLevel);
-//	GetWorld()->ChangeFeatureLevel(InPreviewFeatureLevel); // The ultimate goal, but needs removal of the global feature level first!
+	GetWorld()->ChangeFeatureLevel(InPreviewFeatureLevel);
 }
 
 bool FLevelEditorActionCallbacks::IsFeatureLevelPreviewChecked(ERHIFeatureLevel::Type InPreviewFeatureLevel)
diff --git a/Engine/Source/Editor/LevelEditor/Private/LevelEditorToolBar.cpp b/Engine/Source/Editor/LevelEditor/Private/LevelEditorToolBar.cpp
--- a/Engine/Source/Editor/LevelEditor/Private/LevelEditorToolBar.cpp
+++ b/Engine/Source/Editor/LevelEditor/Private/LevelEditorToolBar.cpp
@@ -697,8 +697,7 @@ TSharedRef< SWidget > FLevelEditorToolBar::GenerateQuickSettingsMenu( TSharedRef
 			LOCTEXT( "MaterialQualityLevelSubMenu_ToolTip", "Sets the value of the CVar \"r.MaterialQualityLevel\" (low=0, high=1). This affects materials via the QualitySwitch material expression." ),
 			FNewMenuDelegate::CreateStatic( &MakeMaterialQualityLevelMenu ) );
 
-		static const auto CVarFeatureLevelPreview = IConsoleManager::Get().FindTConsoleVariableDataInt(TEXT("r.FeatureLevelPreview"));
-		if (CVarFeatureLevelPreview->GetValueOnGameThread() != 0)
+		if (GetDefault<UEditorExperimentalSettings>()->bFeatureLevelPreview)
 		{
 			MenuBuilder.AddSubMenu(
 				LOCTEXT("FeatureLevelPreviewSubMenu", "Feature Level Preview"),
diff --git a/Engine/Source/Editor/UnrealEd/Classes/Settings/EditorExperimentalSettings.h b/Engine/Source/Editor/UnrealEd/Classes/Settings/EditorExperimentalSettings.h
--- a/Engine/Source/Editor/UnrealEd/Classes/Settings/EditorExperimentalSettings.h
+++ b/Engine/Source/Editor/UnrealEd/Classes/Settings/EditorExperimentalSettings.h
@@ -64,6 +64,10 @@ public:
 	UPROPERTY(EditAnywhere,config,Category=UserInterface,meta=(DisplayName="Allow ChunkID Assignments"))
 	bool bContextMenuChunkAssignments;
 
+	/** Enables the dynamic feature level switching functionality */
+	UPROPERTY(EditAnywhere, config, Category = Rendering, meta = (DisplayName = "Feature Level Preview"))
+	bool bFeatureLevelPreview;
+
 	/**
 	 * Returns an event delegate that is executed when a setting has changed.
 	 *
diff --git a/Engine/Source/Editor/UnrealEd/Private/EditorComponents.cpp b/Engine/Source/Editor/UnrealEd/Private/EditorComponents.cpp
--- a/Engine/Source/Editor/UnrealEd/Private/EditorComponents.cpp
+++ b/Engine/Source/Editor/UnrealEd/Private/EditorComponents.cpp
@@ -111,7 +111,7 @@ static void GetAxisColors(FLinearColor Out[3], bool b3D)
 
 void FGridWidget::DrawNewGrid(const FSceneView* View, FPrimitiveDrawInterface* PDI)
 {
-	if(LevelGridMaterial->IsCompilingOrHadCompileError() || LevelGridMaterial2->IsCompilingOrHadCompileError())
+	if (LevelGridMaterial->IsCompilingOrHadCompileError(View->GetFeatureLevel()) || LevelGridMaterial2->IsCompilingOrHadCompileError(View->GetFeatureLevel()))
 	{
 		// The material would appear to be black (because we don't use a MaterialDomain but a UsageFlag - we should change that).
 		// Here we rather want to hide it.
diff --git a/Engine/Source/Editor/UnrealEd/Private/Lightmass/Lightmass.cpp b/Engine/Source/Editor/UnrealEd/Private/Lightmass/Lightmass.cpp
--- a/Engine/Source/Editor/UnrealEd/Private/Lightmass/Lightmass.cpp
+++ b/Engine/Source/Editor/UnrealEd/Private/Lightmass/Lightmass.cpp
@@ -4036,7 +4036,7 @@ bool FLightmassProcessor::ImportTextureMapping(int32 Channel, FTextureMappingImp
 	float BytesPerPixel = 1.0f;
 
 	float LightMapTypeModifier = NUM_HQ_LIGHTMAP_COEF;
-	if (!AllowHighQualityLightmaps())
+	if (!AllowHighQualityLightmaps(GMaxRHIFeatureLevel))
 	{
 		LightMapTypeModifier = NUM_LQ_LIGHTMAP_COEF;
 	}
diff --git a/Engine/Source/Runtime/Engine/Classes/Components/MeshComponent.h b/Engine/Source/Runtime/Engine/Classes/Components/MeshComponent.h
--- a/Engine/Source/Runtime/Engine/Classes/Components/MeshComponent.h
+++ b/Engine/Source/Runtime/Engine/Classes/Components/MeshComponent.h
@@ -3,6 +3,7 @@
 
 #pragma once
 #include "Components/PrimitiveComponent.h"
+#include "RHIDefinitions.h"
 #include "MeshComponent.generated.h"
 
 struct FMaterialRelevance;
@@ -29,7 +30,7 @@ class ENGINE_API UMeshComponent : public UPrimitiveComponent
 	// End UPrimitiveComponent Interface
 
 	/** Accesses the scene relevance information for the materials applied to the mesh. */
-	FMaterialRelevance GetMaterialRelevance() const;
+	FMaterialRelevance GetMaterialRelevance(ERHIFeatureLevel::Type InFeatureLevel) const;
 
 	/**
 	 *	Tell the streaming system whether or not all mip levels of all textures used by this component should be loaded and remain loaded.
diff --git a/Engine/Source/Runtime/Engine/Classes/Engine/NiagaraEffectRenderer.h b/Engine/Source/Runtime/Engine/Classes/Engine/NiagaraEffectRenderer.h
--- a/Engine/Source/Runtime/Engine/Classes/Engine/NiagaraEffectRenderer.h
+++ b/Engine/Source/Runtime/Engine/Classes/Engine/NiagaraEffectRenderer.h
@@ -147,7 +147,7 @@ protected:
 			Material = UMaterial::GetDefaultMaterial(MD_Surface);
 		}
 		check(Material);
-		MaterialRelevance = Material->GetRelevance();
+		MaterialRelevance = Material->GetRelevance(SceneProxy->GetScene()->GetFeatureLevel());
 	}
 
 	const FNiagaraSceneProxy *SceneProxy;
diff --git a/Engine/Source/Runtime/Engine/Classes/Engine/World.h b/Engine/Source/Runtime/Engine/Classes/Engine/World.h
--- a/Engine/Source/Runtime/Engine/Classes/Engine/World.h
+++ b/Engine/Source/Runtime/Engine/Classes/Engine/World.h
@@ -635,11 +635,8 @@ public:
 
 	/** Change the feature level that this world is current rendering with */
 	void ChangeFeatureLevel(ERHIFeatureLevel::Type InFeatureLevel);
-	static void ChangeAllWorldFeatureLevels(ERHIFeatureLevel::Type InFeatureLevel);
 
 private:
-	/** Change the feature level that this world is current rendering with */
-	void ChangeWorldFeatureLevel(ERHIFeatureLevel::Type InFeatureLevel);
 
 	/** List of all the controllers in the world. */
 	TArray<TAutoWeakObjectPtr<class AController> >	ControllerList;
diff --git a/Engine/Source/Runtime/Engine/Classes/Materials/Material.h b/Engine/Source/Runtime/Engine/Classes/Materials/Material.h
--- a/Engine/Source/Runtime/Engine/Classes/Materials/Material.h
+++ b/Engine/Source/Runtime/Engine/Classes/Materials/Material.h
@@ -837,7 +837,7 @@ public:
 	ENGINE_API bool IsUsageFlagDirty(EMaterialUsage Usage);
 	
 	/** Useful to customize rendering if that case (e.g. hide the object) */
-	ENGINE_API bool IsCompilingOrHadCompileError();
+	ENGINE_API bool IsCompilingOrHadCompileError(ERHIFeatureLevel::Type InFeatureLevel);
 
 private:
 	void BackwardsCompatibilityInputConversion();
diff --git a/Engine/Source/Runtime/Engine/Classes/Materials/MaterialInterface.h b/Engine/Source/Runtime/Engine/Classes/Materials/MaterialInterface.h
--- a/Engine/Source/Runtime/Engine/Classes/Materials/MaterialInterface.h
+++ b/Engine/Source/Runtime/Engine/Classes/Materials/MaterialInterface.h
@@ -204,10 +204,15 @@ private:
 	/** Feature levels to force to compile. */
 	uint32 FeatureLevelsToForceCompile;
 
+	/** Feature level bitfield to compile for all materials */
+	static uint32 FeatureLevelsForAllMaterials;
 public:
-	/** Set which feature levels the preview material should compile. GRHIFeatureLevel is always compiled! */
+	/** Set which feature levels this material instance should compile. GMaxRHIFeatureLevel is always compiled! */
 	ENGINE_API void SetFeatureLevelToCompile(ERHIFeatureLevel::Type FeatureLevel, bool bShouldCompile);
 
+	/** Set which feature levels _all_ materials should compile to. GMaxRHIFeatureLevel is always compiled. */
+	ENGINE_API static void SetGlobalRequiredFeatureLevel(ERHIFeatureLevel::Type FeatureLevel, bool bShouldCompile);
+
 	// Begin UObject interface.
 	ENGINE_API virtual void BeginDestroy() override;
 	ENGINE_API virtual bool IsReadyForFinishDestroy() override;
@@ -331,11 +336,11 @@ public:
 		PURE_VIRTUAL(UMaterialInterface::GetTerrainLayerWeightParameterValue,return false;);
 
 	/** @return The material's relevance. */
-	ENGINE_API FMaterialRelevance GetRelevance() const;
+	ENGINE_API FMaterialRelevance GetRelevance(ERHIFeatureLevel::Type InFeatureLevel) const;
 	/** @return The material's relevance, from concurrent render thread updates. */
-	ENGINE_API FMaterialRelevance GetRelevance_Concurrent() const;
+	ENGINE_API FMaterialRelevance GetRelevance_Concurrent(ERHIFeatureLevel::Type InFeatureLevel) const;
 private:
-	ENGINE_API FMaterialRelevance GetRelevance_Internal(const UMaterial* Material) const;
+	ENGINE_API FMaterialRelevance GetRelevance_Internal(const UMaterial* Material, ERHIFeatureLevel::Type InFeatureLevel) const;
 public:
 
 	int32 GetWidth() const;
@@ -579,8 +584,12 @@ public:
 	/** Allows material properties to be compiled with the option of being overridden by the material attributes input. */
 	ENGINE_API virtual int32 CompilePropertyEx( class FMaterialCompiler* Compiler, EMaterialProperty Property );
 
+	/** Get bitfield indicating which feature levels should be compiled by default */
+	ENGINE_API static uint32 GetFeatureLevelsToCompileForAllMaterials() { return FeatureLevelsForAllMaterials | (1 << GMaxRHIFeatureLevel); }
+
 protected:
-	/** Returns a bitfield indicating which feature levels should be compiled for rendering. */
+
+	/** Returns a bitfield indicating which feature levels should be compiled for rendering. GMaxRHIFeatureLevel is always present */
 	ENGINE_API uint32 GetFeatureLevelsToCompileForRendering() const;
 
 	void UpdateMaterialRenderProxy(FMaterialRenderProxy& Proxy);
diff --git a/Engine/Source/Runtime/Engine/Private/GlobalShader.cpp b/Engine/Source/Runtime/Engine/Private/GlobalShader.cpp
--- a/Engine/Source/Runtime/Engine/Private/GlobalShader.cpp
+++ b/Engine/Source/Runtime/Engine/Private/GlobalShader.cpp
@@ -413,59 +413,34 @@ TShaderMap<FGlobalShaderType>* GetGlobalShaderMap(EShaderPlatform Platform, bool
 	return GGlobalShaderMap[Platform];
 }
 
-bool IsGlobalShaderMapComplete()
+bool IsGlobalShaderMapComplete(const TCHAR* TypeNameSubstring)
 {
-	EShaderPlatform Platform = GRHIShaderPlatform;
-
-	TShaderMap<FGlobalShaderType>* GlobalShaderMap = GGlobalShaderMap[Platform];
-
-	if (GlobalShaderMap)
+	for (int32 i = 0; i < SP_NumPlatforms; ++i)
 	{
-		for(TLinkedList<FShaderType*>::TIterator ShaderTypeIt(FShaderType::GetTypeList());ShaderTypeIt;ShaderTypeIt.Next())
-		{
-			FGlobalShaderType* GlobalShaderType = ShaderTypeIt->GetGlobalShaderType();
-
-			if (GlobalShaderType && GlobalShaderType->ShouldCache(Platform))
-			{
-				if (!GlobalShaderMap->HasShader(GlobalShaderType))
-				{
-					return false;
-				}
-			}
-		}
-
-		return true;
-	}
-	
-	return false;
-}
-
-bool AreGlobalShadersComplete(const TCHAR* TypeNameSubstring)
-{
-	EShaderPlatform Platform = GRHIShaderPlatform;
-	TShaderMap<FGlobalShaderType>* GlobalShaderMap = GGlobalShaderMap[Platform];
+		EShaderPlatform Platform = (EShaderPlatform)i;
+		
+		TShaderMap<FGlobalShaderType>* GlobalShaderMap = GGlobalShaderMap[Platform];
 
-	if (GlobalShaderMap)
-	{
-		for(TLinkedList<FShaderType*>::TIterator ShaderTypeIt(FShaderType::GetTypeList());ShaderTypeIt;ShaderTypeIt.Next())
+		if (GlobalShaderMap)
 		{
-			FGlobalShaderType* GlobalShaderType = ShaderTypeIt->GetGlobalShaderType();
-
-			if (GlobalShaderType 
-				&& FPlatformString::Strstr(GlobalShaderType->GetName(), TypeNameSubstring) != NULL
-				&& GlobalShaderType->ShouldCache(Platform))
+			for (TLinkedList<FShaderType*>::TIterator ShaderTypeIt(FShaderType::GetTypeList()); ShaderTypeIt; ShaderTypeIt.Next())
 			{
-				if (!GlobalShaderMap->HasShader(GlobalShaderType))
+				FGlobalShaderType* GlobalShaderType = ShaderTypeIt->GetGlobalShaderType();
+
+				if (GlobalShaderType
+					&& (TypeNameSubstring == nullptr || (FPlatformString::Strstr(GlobalShaderType->GetName(), TypeNameSubstring) != nullptr))
+					&& GlobalShaderType->ShouldCache(Platform))
 				{
-					return false;
+					if (!GlobalShaderMap->HasShader(GlobalShaderType))
+					{
+						return false;
+					}
 				}
 			}
 		}
-
-		return true;
 	}
-
-	return false;
+	
+	return true;
 }
 
 /**
diff --git a/Engine/Source/Runtime/Engine/Private/Landscape/LandscapeLight.cpp b/Engine/Source/Runtime/Engine/Private/Landscape/LandscapeLight.cpp
--- a/Engine/Source/Runtime/Engine/Private/Landscape/LandscapeLight.cpp
+++ b/Engine/Source/Runtime/Engine/Private/Landscape/LandscapeLight.cpp
@@ -631,7 +631,9 @@ void ULandscapeComponent::GetLightAndShadowMapMemoryUsage( int32& LightMapMemory
 	int32 Width, Height;
 	GetLightMapResolution(Width, Height);
 
-	if(AllowHighQualityLightmaps())
+	const auto FeatureLevel = GetWorld() ? GetWorld()->FeatureLevel : GMaxRHIFeatureLevel;
+
+	if(AllowHighQualityLightmaps(FeatureLevel))
 	{
 		LightMapMemoryUsage  = NUM_HQ_LIGHTMAP_COEF * (Width * Height * 4 / 3); // assuming DXT5
 	}
diff --git a/Engine/Source/Runtime/Engine/Private/Landscape/LandscapeRender.cpp b/Engine/Source/Runtime/Engine/Private/Landscape/LandscapeRender.cpp
--- a/Engine/Source/Runtime/Engine/Private/Landscape/LandscapeRender.cpp
+++ b/Engine/Source/Runtime/Engine/Private/Landscape/LandscapeRender.cpp
@@ -635,7 +635,7 @@ FLandscapeComponentSceneProxy::FLandscapeComponentSceneProxy(ULandscapeComponent
 		MaterialInterface = UMaterial::GetDefaultMaterial(MD_Surface);
 	}
 
-	MaterialRelevance = MaterialInterface->GetRelevance();
+	MaterialRelevance = MaterialInterface->GetRelevance(FeatureLevel);
 
 #if !(UE_BUILD_SHIPPING || UE_BUILD_TEST) || (UE_BUILD_SHIPPING && WITH_EDITOR)
 	if( GIsEditor )
@@ -804,6 +804,8 @@ FPrimitiveViewRelevance FLandscapeComponentSceneProxy::GetViewRelevance(const FS
 	FPrimitiveViewRelevance Result;
 	Result.bDrawRelevance = IsShown(View) && View->Family->EngineShowFlags.Landscape;
 
+	auto FeatureLevel = View->GetFeatureLevel();
+
 #if WITH_EDITOR
 	if( !GLandscapeEditModeActive )
 	{
@@ -823,13 +825,13 @@ FPrimitiveViewRelevance FLandscapeComponentSceneProxy::GetViewRelevance(const FS
 			if( EditToolRenderData->ToolMaterial )
 			{
 				Result.bDynamicRelevance = true;
-				ToolRelevance |= EditToolRenderData->ToolMaterial->GetRelevance_Concurrent();
+				ToolRelevance |= EditToolRenderData->ToolMaterial->GetRelevance_Concurrent(FeatureLevel);
 			}
 
 			if( EditToolRenderData->GizmoMaterial )
 			{
 				Result.bDynamicRelevance = true;
-				ToolRelevance |= EditToolRenderData->GizmoMaterial->GetRelevance_Concurrent();
+				ToolRelevance |= EditToolRenderData->GizmoMaterial->GetRelevance_Concurrent(FeatureLevel);
 			}
 		}
 
@@ -840,12 +842,12 @@ FPrimitiveViewRelevance FLandscapeComponentSceneProxy::GetViewRelevance(const FS
 				&& !(GLandscapeEditRenderMode & ELandscapeEditRenderMode::Mask) && GSelectionRegionMaterial)
 			{
 				Result.bDynamicRelevance = true;
-				ToolRelevance |= GSelectionRegionMaterial->GetRelevance_Concurrent();
+				ToolRelevance |= GSelectionRegionMaterial->GetRelevance_Concurrent(FeatureLevel);
 			}
 			if ((GLandscapeEditRenderMode & ELandscapeEditRenderMode::SelectComponent) && (EditToolRenderData->SelectedType & FLandscapeEditToolRenderData::ST_COMPONENT) && GSelectionColorMaterial)
 			{
 				Result.bDynamicRelevance = true;
-				ToolRelevance |= GSelectionColorMaterial->GetRelevance_Concurrent();
+				ToolRelevance |= GSelectionColorMaterial->GetRelevance_Concurrent(FeatureLevel);
 			}
 		}
 
@@ -854,7 +856,7 @@ FPrimitiveViewRelevance FLandscapeComponentSceneProxy::GetViewRelevance(const FS
 			((EditToolRenderData && (EditToolRenderData->SelectedType & FLandscapeEditToolRenderData::ST_REGION)) || (!(GLandscapeEditRenderMode & ELandscapeEditRenderMode::InvertedMask))) )
 		{
 			Result.bDynamicRelevance = true;
-			ToolRelevance |= GMaskRegionMaterial->GetRelevance_Concurrent();
+			ToolRelevance |= GMaskRegionMaterial->GetRelevance_Concurrent(FeatureLevel);
 		}
 
 		ToolRelevance.SetPrimitiveViewRelevance(Result);
@@ -1041,7 +1043,7 @@ namespace
 {
 	inline bool RequiresAdjacencyInformation(FMaterialRenderProxy* MaterialRenderProxy, ERHIFeatureLevel::Type InFeatureLevel) // Assumes VertexFactory supports tessellation, and rendering thread with this function
 	{
-		if (RHISupportsTessellation(GRHIShaderPlatform) && MaterialRenderProxy )
+		if (RHISupportsTessellation(GShaderPlatformForFeatureLevel[InFeatureLevel]) && MaterialRenderProxy)
 		{
 			check ( IsInRenderingThread() );
 			const FMaterial* MaterialResource = MaterialRenderProxy->GetMaterial(InFeatureLevel);
@@ -2748,8 +2750,10 @@ void ULandscapeComponent::GetStreamingTextureInfo(TArray<FStreamingTexturePrimit
 		}
 
 		// Lightmap
+		const auto FeatureLevel = GetWorld() ? GetWorld()->FeatureLevel : GMaxRHIFeatureLevel;
+
 		FLightMap2D* Lightmap = LightMap ? LightMap->GetLightMap2D() : nullptr;
-		uint32 LightmapIndex = AllowHighQualityLightmaps() ? 0 : 1;
+		uint32 LightmapIndex = AllowHighQualityLightmaps(FeatureLevel) ? 0 : 1;
 		if (Lightmap && Lightmap->IsValid(LightmapIndex))
 		{
 			const FVector2D& Scale = Lightmap->GetCoordinateScale();
diff --git a/Engine/Source/Runtime/Engine/Private/LightMap.cpp b/Engine/Source/Runtime/Engine/Private/LightMap.cpp
--- a/Engine/Source/Runtime/Engine/Private/LightMap.cpp
+++ b/Engine/Source/Runtime/Engine/Private/LightMap.cpp
@@ -68,9 +68,9 @@ ENGINE_API FColor GTexelSelectionColor(255, 50, 0);
 
 FLightMap::FLightMap()
 	: bAllowHighQualityLightMaps(true)
-	, NumRefs(0) 
+	, NumRefs(0)
 {
-	bAllowHighQualityLightMaps = !IsMobilePlatform(GRHIShaderPlatform) && AllowHighQualityLightmaps();
+	bAllowHighQualityLightMaps = AllowHighQualityLightmaps();
 #if !PLATFORM_DESKTOP 
 	checkf(bAllowHighQualityLightMaps || IsMobilePlatform(GRHIShaderPlatform), TEXT("Low quality lightmaps are not currently supported on consoles. Make sure console variable r.HighQualityLightMaps is true for this platform"));
 #endif
@@ -1265,7 +1265,7 @@ UTexture2D* FLightMap2D::GetTexture(uint32 BasisIndex)
  */
 bool FLightMap2D::IsValid(uint32 BasisIndex) const
 {
-	return AllowHighQualityLightmaps() ? BasisIndex == 0 : BasisIndex == 1;
+	return Textures[BasisIndex] != nullptr;
 }
 
 struct FLegacyLightMapTextureInfo
@@ -1419,16 +1419,18 @@ void FLightMap2D::Serialize(FArchive& Ar)
 	}
 }
 
-FLightMapInteraction FLightMap2D::GetInteraction() const
+FLightMapInteraction FLightMap2D::GetInteraction(ERHIFeatureLevel::Type InFeatureLevel) const
 {
-	int32 LightmapIndex = AllowHighQualityLightmaps() ? 0 : 1;
+	bool bHighQuality = AllowHighQualityLightmaps(InFeatureLevel);
+
+	int32 LightmapIndex = bHighQuality ? 0 : 1;
 
 	bool bValidTextures = Textures[ LightmapIndex ] && Textures[ LightmapIndex ]->Resource;
 
 	// When the FLightMap2D is first created, the textures aren't set, so that case needs to be handled.
 	if(bValidTextures)
 	{
-		return FLightMapInteraction::Texture(Textures, SkyOcclusionTexture, ScaleVectors, AddVectors, CoordinateScale, CoordinateBias, AllowHighQualityLightmaps());
+		return FLightMapInteraction::Texture(Textures, SkyOcclusionTexture, ScaleVectors, AddVectors, CoordinateScale, CoordinateBias, bHighQuality);
 	}
 
 	return FLightMapInteraction::None();
diff --git a/Engine/Source/Runtime/Engine/Private/MaterialBillboardComponent.cpp b/Engine/Source/Runtime/Engine/Private/MaterialBillboardComponent.cpp
--- a/Engine/Source/Runtime/Engine/Private/MaterialBillboardComponent.cpp
+++ b/Engine/Source/Runtime/Engine/Private/MaterialBillboardComponent.cpp
@@ -108,7 +108,7 @@ public:
 			UMaterialInterface* Material = Elements[ElementIndex].Material;
 			if (Material)
 			{
-				MaterialRelevance |= Material->GetRelevance();
+				MaterialRelevance |= Material->GetRelevance(GetScene()->GetFeatureLevel());
 			}
 		}
 
diff --git a/Engine/Source/Runtime/Engine/Private/Materials/Material.cpp b/Engine/Source/Runtime/Engine/Private/Materials/Material.cpp
--- a/Engine/Source/Runtime/Engine/Private/Materials/Material.cpp
+++ b/Engine/Source/Runtime/Engine/Private/Materials/Material.cpp
@@ -715,9 +715,9 @@ bool UMaterial::IsUsageFlagDirty(EMaterialUsage Usage)
 	return false;
 }
 
-bool UMaterial::IsCompilingOrHadCompileError()
+bool UMaterial::IsCompilingOrHadCompileError(ERHIFeatureLevel::Type InFeatureLevel)
 {
-	FMaterialResource* Res = GetMaterialResource(GRHIFeatureLevel);
+	FMaterialResource* Res = GetMaterialResource(InFeatureLevel);
 
 	// should never be the case
 	check(Res);
@@ -3023,7 +3023,7 @@ bool UMaterial::UpdateLightmassTextureTracking()
 #if WITH_EDITORONLY_DATA
 	TArray<UTexture*> UsedTextures;
 	
-	GetUsedTextures(UsedTextures, EMaterialQualityLevel::Num, true, GRHIFeatureLevel, true);
+	GetUsedTextures(UsedTextures, EMaterialQualityLevel::Num, true, GMaxRHIFeatureLevel, true);
 	if (UsedTextures.Num() != ReferencedTextureGuids.Num())
 	{
 		bTexturesHaveChanged = true;
diff --git a/Engine/Source/Runtime/Engine/Private/Materials/MaterialInterface.cpp b/Engine/Source/Runtime/Engine/Private/Materials/MaterialInterface.cpp
--- a/Engine/Source/Runtime/Engine/Private/Materials/MaterialInterface.cpp
+++ b/Engine/Source/Runtime/Engine/Private/Materials/MaterialInterface.cpp
@@ -38,11 +38,11 @@ void UMaterialInterface::PostLoad()
 	PostLoadDefaultMaterials();
 }
 
-FMaterialRelevance UMaterialInterface::GetRelevance_Internal(const UMaterial* Material) const
+FMaterialRelevance UMaterialInterface::GetRelevance_Internal(const UMaterial* Material, ERHIFeatureLevel::Type InFeatureLevel) const
 {
 	if(Material)
 	{
-		const FMaterialResource* MaterialResource = Material->GetMaterialResource(GRHIFeatureLevel);
+		const FMaterialResource* MaterialResource = Material->GetMaterialResource(InFeatureLevel);
 		const bool bIsTranslucent = IsTranslucentBlendMode((EBlendMode)GetBlendMode());
 		const bool bIsLit = GetShadingModel() != MSM_Unlit;
 		// Determine the material's view relevance.
@@ -62,19 +62,19 @@ FMaterialRelevance UMaterialInterface::GetRelevance_Internal(const UMaterial* Ma
 	}
 }
 
-FMaterialRelevance UMaterialInterface::GetRelevance() const
+FMaterialRelevance UMaterialInterface::GetRelevance(ERHIFeatureLevel::Type InFeatureLevel) const
 {
 	// Find the interface's concrete material.
 	const UMaterial* Material = GetMaterial();
-	return GetRelevance_Internal(Material);
+	return GetRelevance_Internal(Material, InFeatureLevel);
 }
 
-FMaterialRelevance UMaterialInterface::GetRelevance_Concurrent() const
+FMaterialRelevance UMaterialInterface::GetRelevance_Concurrent(ERHIFeatureLevel::Type InFeatureLevel) const
 {
 	// Find the interface's concrete material.
 	TMicRecursionGuard RecursionGuard;
 	const UMaterial* Material = GetMaterial_Concurrent(RecursionGuard);
-	return GetRelevance_Internal(Material);
+	return GetRelevance_Internal(Material, InFeatureLevel);
 }
 
 int32 UMaterialInterface::GetWidth() const
@@ -324,9 +324,25 @@ void UMaterialInterface::SetFeatureLevelToCompile(ERHIFeatureLevel::Type Feature
 	}
 }
 
+uint32 UMaterialInterface::FeatureLevelsForAllMaterials = 0;
+
+void UMaterialInterface::SetGlobalRequiredFeatureLevel(ERHIFeatureLevel::Type FeatureLevel, bool bShouldCompile)
+{
+	uint32 FeatureLevelBit = (1 << FeatureLevel);
+	if (bShouldCompile)
+	{
+		FeatureLevelsForAllMaterials |= FeatureLevelBit;
+	}
+	else
+	{
+		FeatureLevelsForAllMaterials &= (~FeatureLevelBit);
+	}
+}
+
+
 uint32 UMaterialInterface::GetFeatureLevelsToCompileForRendering() const
 {
-	return FeatureLevelsToForceCompile | (1 << GRHIFeatureLevel);
+	return FeatureLevelsToForceCompile | GetFeatureLevelsToCompileForAllMaterials();
 }
 
 
diff --git a/Engine/Source/Runtime/Engine/Private/Materials/MaterialShader.cpp b/Engine/Source/Runtime/Engine/Private/Materials/MaterialShader.cpp
--- a/Engine/Source/Runtime/Engine/Private/Materials/MaterialShader.cpp
+++ b/Engine/Source/Runtime/Engine/Private/Materials/MaterialShader.cpp
@@ -784,7 +784,7 @@ void FMaterialShaderMap::LoadFromDerivedDataCache(const FMaterial* Material, con
 				checkSlow(InOutShaderMap->GetShaderMapId() == ShaderMapId);
 
 				// Register in the global map
-				InOutShaderMap->Register();
+				InOutShaderMap->Register(Platform);
 			}
 			else
 			{
@@ -985,7 +985,7 @@ void FMaterialShaderMap::LoadForRemoteRecompile(FArchive& Ar, EShaderPlatform Sh
 				ShaderMap->Serialize(Ar, false);
 
 				// Register in the global map
-				ShaderMap->Register();
+				ShaderMap->Register(ShaderPlatform);
 
 				LoadedShaderMaps.Add(ShaderMap);
 			}
@@ -1215,7 +1215,7 @@ void FMaterialShaderMap::Compile(
 			}
 
 			// Register this shader map in the global map with the material's ID.
-			Register();
+			Register(InPlatform);
   
 			// Mark the shader map as not having been finalized with ProcessCompilationResults
 			bCompilationFinalized = false;
@@ -1448,9 +1448,9 @@ void FMaterialShaderMap::GetShaderList(TMap<FShaderId,FShader*>& OutShaders) con
 /**
  * Registers a material shader map in the global map so it can be used by materials.
  */
-void FMaterialShaderMap::Register() 
+void FMaterialShaderMap::Register(EShaderPlatform InShaderPlatform)
 {
-	if (GCreateShadersOnLoad && Platform == GRHIShaderPlatform)
+	if (GCreateShadersOnLoad && Platform == InShaderPlatform)
 	{
 		for (TMap<FShaderType*, TRefCountPtr<FShader> >::TConstIterator ShaderIt(GetShaders()); ShaderIt; ++ShaderIt)
 		{
diff --git a/Engine/Source/Runtime/Engine/Private/Materials/MaterialShared.cpp b/Engine/Source/Runtime/Engine/Private/Materials/MaterialShared.cpp
--- a/Engine/Source/Runtime/Engine/Private/Materials/MaterialShared.cpp
+++ b/Engine/Source/Runtime/Engine/Private/Materials/MaterialShared.cpp
@@ -1247,7 +1247,7 @@ bool FMaterial::CacheShaders(const FMaterialShaderMapId& ShaderMapId, EShaderPla
 		else if (GameThreadShaderMap)
 		{
 			// We are going to use the inlined shader map, register it so it can be re-used by other materials
-			GameThreadShaderMap->Register();
+			GameThreadShaderMap->Register(Platform);
 		}
 	}
 	else
@@ -1535,13 +1535,14 @@ void FMaterialRenderProxy::CacheUniformExpressions()
 	InitResource();
 
 	check(UMaterial::GetDefaultMaterial(MD_Surface));
-	const bool bES2Preview = false;
-	ERHIFeatureLevel::Type FeatureLevelsToCache[2] = { GRHIFeatureLevel, ERHIFeatureLevel::ES2 };
-	int32 NumFeatureLevelsToCache = bES2Preview ? 2 : 1;
 
-	for (int32 i = 0; i < NumFeatureLevelsToCache; ++i)
+	uint32 FeatureLevelsToCompile = UMaterialInterface::GetFeatureLevelsToCompileForAllMaterials();
+	TArray<FMaterialResource*> ResourcesToCache;
+
+	while (FeatureLevelsToCompile != 0)
 	{
-		ERHIFeatureLevel::Type FeatureLevel = FeatureLevelsToCache[i];
+		ERHIFeatureLevel::Type FeatureLevel = (ERHIFeatureLevel::Type)FBitSet::GetAndClearNextBit(FeatureLevelsToCompile);
+
 		const FMaterial* MaterialNoFallback = GetMaterialNoFallback(FeatureLevel);
 
 		if (MaterialNoFallback && MaterialNoFallback->GetRenderingThreadShaderMap())
diff --git a/Engine/Source/Runtime/Engine/Private/MeshComponent.cpp b/Engine/Source/Runtime/Engine/Private/MeshComponent.cpp
--- a/Engine/Source/Runtime/Engine/Private/MeshComponent.cpp
+++ b/Engine/Source/Runtime/Engine/Private/MeshComponent.cpp
@@ -59,7 +59,7 @@ void UMeshComponent::SetMaterial(int32 ElementIndex, UMaterialInterface* Materia
 	}
 }
 
-FMaterialRelevance UMeshComponent::GetMaterialRelevance() const
+FMaterialRelevance UMeshComponent::GetMaterialRelevance(ERHIFeatureLevel::Type InFeatureLevel) const
 {
 	// Combine the material relevance for all materials.
 	FMaterialRelevance Result;
@@ -70,7 +70,7 @@ FMaterialRelevance UMeshComponent::GetMaterialRelevance() const
 		{
 			MaterialInterface = UMaterial::GetDefaultMaterial(MD_Surface);
 		}
-		Result |= MaterialInterface->GetRelevance_Concurrent();
+		Result |= MaterialInterface->GetRelevance_Concurrent(InFeatureLevel);
 	}
 	return Result;
 }
diff --git a/Engine/Source/Runtime/Engine/Private/ModelRender.cpp b/Engine/Source/Runtime/Engine/Private/ModelRender.cpp
--- a/Engine/Source/Runtime/Engine/Private/ModelRender.cpp
+++ b/Engine/Source/Runtime/Engine/Private/ModelRender.cpp
@@ -163,7 +163,7 @@ public:
 		{
 			const FModelElement& SourceElement = SourceElements[ElementIndex];
 			FElementInfo* Element = new(Elements) FElementInfo(SourceElement);
-			MaterialRelevance |= Element->GetMaterial()->GetRelevance();
+			MaterialRelevance |= Element->GetMaterial()->GetRelevance(GetScene()->GetFeatureLevel());
 		}
 
 		// Try to find a color for level coloration.
@@ -813,9 +813,9 @@ private:
 			}
 		}
 
-		virtual FLightMapInteraction GetLightMapInteraction() const
+		virtual FLightMapInteraction GetLightMapInteraction(ERHIFeatureLevel::Type InFeatureLevel) const
 		{
-			return LightMap ? LightMap->GetInteraction() : FLightMapInteraction();
+			return LightMap ? LightMap->GetInteraction(InFeatureLevel) : FLightMapInteraction();
 		}
 
 		virtual FShadowMapInteraction GetShadowMapInteraction() const
diff --git a/Engine/Source/Runtime/Engine/Private/Particles/ParticleComponents.cpp b/Engine/Source/Runtime/Engine/Private/Particles/ParticleComponents.cpp
--- a/Engine/Source/Runtime/Engine/Private/Particles/ParticleComponents.cpp
+++ b/Engine/Source/Runtime/Engine/Private/Particles/ParticleComponents.cpp
@@ -5044,7 +5044,8 @@ void UParticleSystemComponent::CacheViewRelevanceFlags(UParticleSystem* Template
 
 				if (EmitterLODLevel->bEnabled == true)
 				{
-					EmitterInst->GatherMaterialRelevance(&LODViewRel, EmitterLODLevel);
+					auto World = GetWorld();
+					EmitterInst->GatherMaterialRelevance(&LODViewRel, EmitterLODLevel, World ? World->FeatureLevel : GMaxRHIFeatureLevel);
 				}
 			}
 		}
diff --git a/Engine/Source/Runtime/Engine/Private/Particles/ParticleEmitterInstances.cpp b/Engine/Source/Runtime/Engine/Private/Particles/ParticleEmitterInstances.cpp
--- a/Engine/Source/Runtime/Engine/Private/Particles/ParticleEmitterInstances.cpp
+++ b/Engine/Source/Runtime/Engine/Private/Particles/ParticleEmitterInstances.cpp
@@ -2503,21 +2503,21 @@ bool FParticleEmitterInstance::FillReplayData( FDynamicEmitterReplayDataBase& Ou
  * @param OutMaterialRelevance - Pointer to where material relevance flags will be stored.
  * @param LODLevel - The LOD level for which to compute material relevance flags.
  */
-void FParticleEmitterInstance::GatherMaterialRelevance( FMaterialRelevance* OutMaterialRelevance, const UParticleLODLevel* LODLevel ) const
+void FParticleEmitterInstance::GatherMaterialRelevance(FMaterialRelevance* OutMaterialRelevance, const UParticleLODLevel* LODLevel, ERHIFeatureLevel::Type InFeatureLevel) const
 {
 	// These will catch the sprite cases...
 	if (CurrentMaterial)
 	{
-		(*OutMaterialRelevance) |= CurrentMaterial->GetRelevance();
+		(*OutMaterialRelevance) |= CurrentMaterial->GetRelevance(InFeatureLevel);
 	}
 	else if (LODLevel->RequiredModule->Material)
 	{
-		(*OutMaterialRelevance) |= LODLevel->RequiredModule->Material->GetRelevance();
+		(*OutMaterialRelevance) |= LODLevel->RequiredModule->Material->GetRelevance(InFeatureLevel);
 	}
 	else
 	{
 		check( UMaterial::GetDefaultMaterial(MD_Surface) );
-		(*OutMaterialRelevance) |= UMaterial::GetDefaultMaterial(MD_Surface)->GetRelevance();
+		(*OutMaterialRelevance) |= UMaterial::GetDefaultMaterial(MD_Surface)->GetRelevance(InFeatureLevel);
 	}
 }
 
@@ -3337,13 +3337,13 @@ void FParticleMeshEmitterInstance::SetMeshMaterials( const TArray<UMaterialInter
  * @param OutMaterialRelevance - Pointer to where material relevance flags will be stored.
  * @param LODLevel - The LOD level for which to compute material relevance flags.
  */
-void FParticleMeshEmitterInstance::GatherMaterialRelevance( FMaterialRelevance* OutMaterialRelevance, const UParticleLODLevel* LODLevel ) const
+void FParticleMeshEmitterInstance::GatherMaterialRelevance( FMaterialRelevance* OutMaterialRelevance, const UParticleLODLevel* LODLevel, ERHIFeatureLevel::Type InFeatureLevel ) const
 {
 	TArray<UMaterialInterface*,TInlineAllocator<2> > Materials;
 	GetMeshMaterials(Materials, LODLevel);
 	for (int32 MaterialIndex = 0; MaterialIndex < Materials.Num(); ++MaterialIndex)
 	{
-		(*OutMaterialRelevance) |= Materials[MaterialIndex]->GetRelevance();
+		(*OutMaterialRelevance) |= Materials[MaterialIndex]->GetRelevance(InFeatureLevel);
 	}
 }
 
diff --git a/Engine/Source/Runtime/Engine/Private/PrimitiveDrawingUtils.cpp b/Engine/Source/Runtime/Engine/Private/PrimitiveDrawingUtils.cpp
--- a/Engine/Source/Runtime/Engine/Private/PrimitiveDrawingUtils.cpp
+++ b/Engine/Source/Runtime/Engine/Private/PrimitiveDrawingUtils.cpp
@@ -1495,14 +1495,17 @@ void ApplyViewModeOverrides(
 			bool bTextureMapped = false;
 			FVector2D LMResolution;
 
-			if (EngineShowFlags.LightMapDensity &&
-				Mesh.LCI &&
-				(Mesh.LCI->GetLightMapInteraction().GetType() == LMIT_Texture) &&
-				Mesh.LCI->GetLightMapInteraction().GetTexture())
+			if (EngineShowFlags.LightMapDensity && Mesh.LCI)
 			{
-				LMResolution.X = Mesh.LCI->GetLightMapInteraction().GetTexture()->GetSizeX();
-				LMResolution.Y = Mesh.LCI->GetLightMapInteraction().GetTexture()->GetSizeY();
-				bTextureMapped = true;
+				auto Interaction = Mesh.LCI->GetLightMapInteraction(FeatureLevel);
+				auto Texture = Interaction.GetTexture(AllowHighQualityLightmaps(FeatureLevel));
+
+				if (Interaction.GetType() == LMIT_Texture && Texture)
+				{
+					LMResolution.X = Texture->GetSizeX();
+					LMResolution.Y = Texture->GetSizeY();
+					bTextureMapped = true;
+				}
 			}
 
 			if (bTextureMapped == false)
@@ -1746,14 +1749,17 @@ int32 DrawRichMesh(
 			bool bTextureMapped = false;
 			FVector2D LMResolution;
 
-			if (PDI->View->Family->EngineShowFlags.LightMapDensity &&
-				ModifiedMesh.LCI &&
-				(ModifiedMesh.LCI->GetLightMapInteraction().GetType() == LMIT_Texture) &&
-				ModifiedMesh.LCI->GetLightMapInteraction().GetTexture())
+			if (EngineShowFlags.LightMapDensity && Mesh.LCI)
 			{
-				LMResolution.X = Mesh.LCI->GetLightMapInteraction().GetTexture()->GetSizeX();
-				LMResolution.Y = Mesh.LCI->GetLightMapInteraction().GetTexture()->GetSizeY();
-				bTextureMapped = true;
+				auto Interaction = Mesh.LCI->GetLightMapInteraction(FeatureLevel);
+				auto Texture = Interaction.GetTexture(AllowHighQualityLightmaps(FeatureLevel));
+
+				if (Interaction.GetType() == LMIT_Texture && Texture)
+				{
+					LMResolution.X = Texture->GetSizeX();
+					LMResolution.Y = Texture->GetSizeY();
+					bTextureMapped = true;
+				}
 			}
 
 			if (bTextureMapped == false)
diff --git a/Engine/Source/Runtime/Engine/Private/ReflectionCaptureComponent.cpp b/Engine/Source/Runtime/Engine/Private/ReflectionCaptureComponent.cpp
--- a/Engine/Source/Runtime/Engine/Private/ReflectionCaptureComponent.cpp
+++ b/Engine/Source/Runtime/Engine/Private/ReflectionCaptureComponent.cpp
@@ -960,8 +960,7 @@ void UReflectionCaptureComponent::PostLoad()
 {
 	Super::PostLoad();
 
-	static const auto CVar = IConsoleManager::Get().FindTConsoleVariableDataInt(TEXT("r.FeatureLevelPreview"));
-	bool bRetainAllFeatureLevelData = CVar->GetValueOnGameThread() == 1 && GIsEditor && GMaxRHIFeatureLevel >= ERHIFeatureLevel::SM4;
+	bool bRetainAllFeatureLevelData = GIsEditor && GMaxRHIFeatureLevel >= ERHIFeatureLevel::SM4;
 
 	// If we're loading on a platform that doesn't require cooked data, attempt to load missing data from the DDC
 	if (!FPlatformProperties::RequiresCookedData())
@@ -995,19 +994,19 @@ void UReflectionCaptureComponent::PostLoad()
 		// If we have full HDR data but not encoded HDR data, generate the encoded data now
 		if (FullHDRDerivedData 
 			&& !EncodedHDRDerivedData 
-			&& (GRHIFeatureLevel == ERHIFeatureLevel::ES2 || bRetainAllFeatureLevelData))
+			&& (GMaxRHIFeatureLevel == ERHIFeatureLevel::ES2 || bRetainAllFeatureLevelData))
 		{
 			EncodedHDRDerivedData = FReflectionCaptureEncodedHDRDerivedData::GenerateEncodedHDRData(*FullHDRDerivedData, StateId, Brightness);
 		}
 	}
 	
 	// Initialize rendering resources for the current feature level, and toss data only needed by other feature levels
-	if (FullHDRDerivedData && (GRHIFeatureLevel >= ERHIFeatureLevel::SM4 || bRetainAllFeatureLevelData))
+	if (FullHDRDerivedData && (GMaxRHIFeatureLevel >= ERHIFeatureLevel::SM4 || bRetainAllFeatureLevelData))
 	{
 		// Don't need encoded HDR data for rendering on this feature level
 		INC_MEMORY_STAT_BY(STAT_ReflectionCaptureMemory, FullHDRDerivedData->CompressedCapturedData.GetAllocatedSize());
 
-		if ((GRHIFeatureLevel == ERHIFeatureLevel::SM4 || bRetainAllFeatureLevelData))
+		if ((GMaxRHIFeatureLevel == ERHIFeatureLevel::SM4 || bRetainAllFeatureLevelData))
 		{
 			SM4FullHDRCubemapTexture = new FReflectionTextureCubeResource();
 			SM4FullHDRCubemapTexture->SetupParameters(GReflectionCaptureSize, FMath::CeilLogTwo(GReflectionCaptureSize) + 1, PF_FloatRGBA, &FullHDRDerivedData->GetCapturedDataForSM4Load());
@@ -1020,7 +1019,7 @@ void UReflectionCaptureComponent::PostLoad()
 		}
 	}
 
-	if (EncodedHDRDerivedData && (GRHIFeatureLevel == ERHIFeatureLevel::ES2 || bRetainAllFeatureLevelData))
+	if (EncodedHDRDerivedData && (GMaxRHIFeatureLevel == ERHIFeatureLevel::ES2 || bRetainAllFeatureLevelData))
 	{
 		if (FPlatformProperties::RequiresCookedData())
 		{
diff --git a/Engine/Source/Runtime/Engine/Private/SkeletalMesh.cpp b/Engine/Source/Runtime/Engine/Private/SkeletalMesh.cpp
--- a/Engine/Source/Runtime/Engine/Private/SkeletalMesh.cpp
+++ b/Engine/Source/Runtime/Engine/Private/SkeletalMesh.cpp
@@ -3795,7 +3795,7 @@ FSkeletalMeshSceneProxy::FSkeletalMeshSceneProxy(const USkinnedMeshComponent* Co
 		,	PhysicsAssetForDebug(Component->GetPhysicsAsset())
 		,	bForceWireframe(Component->bForceWireframe)
 		,	bCanHighlightSelectedSections(Component->bCanHighlightSelectedSections)
-		,	MaterialRelevance(Component->GetMaterialRelevance())
+		,	MaterialRelevance(Component->GetMaterialRelevance(GetScene()->GetFeatureLevel()))
 		,	bMaterialsNeedMorphUsage_GameThread(false)
 {
 	check(MeshObject);
@@ -3810,6 +3810,8 @@ FSkeletalMeshSceneProxy::FSkeletalMeshSceneProxy(const USkinnedMeshComponent* Co
 		bAlwaysHasVelocity = true;
 	}
 
+	const auto FeatureLevel = GetScene()->GetFeatureLevel();
+
 	// setup materials and performance classification for each LOD.
 	extern bool GForceDefaultMaterial;
 	bool bCastShadow = Component->CastShadow;
@@ -3851,7 +3853,7 @@ FSkeletalMeshSceneProxy::FSkeletalMeshSceneProxy(const USkinnedMeshComponent* Co
 			if (GForceDefaultMaterial && Material && !IsTranslucentBlendMode(Material->GetBlendMode()))
 			{
 				Material = UMaterial::GetDefaultMaterial(MD_Surface);
-				MaterialRelevance |= Material->GetRelevance();
+				MaterialRelevance |= Material->GetRelevance(FeatureLevel);
 			}
 
 			// if this is a clothing section, then enabled and will be drawn but the corresponding original section should be disabled
@@ -3867,7 +3869,7 @@ FSkeletalMeshSceneProxy::FSkeletalMeshSceneProxy(const USkinnedMeshComponent* Co
 			  (bClothSection && !Material->CheckMaterialUsage(MATUSAGE_Clothing)))
 			{
 				Material = UMaterial::GetDefaultMaterial(MD_Surface);
-				MaterialRelevance |= Material->GetRelevance();
+				MaterialRelevance |= Material->GetRelevance(FeatureLevel);
 			}
 
 			const bool bRequiresAdjacencyInformation = RequiresAdjacencyInformation( Material, &TGPUSkinVertexFactory<false>::StaticType, GetScene()->GetFeatureLevel() );
@@ -3878,7 +3880,7 @@ FSkeletalMeshSceneProxy::FSkeletalMeshSceneProxy(const USkinnedMeshComponent* Co
 					*Material->GetPathName(), 
 					*Component->SkeletalMesh->GetPathName() )
 				Material = UMaterial::GetDefaultMaterial(MD_Surface);
-				MaterialRelevance |= UMaterial::GetDefaultMaterial(MD_Surface)->GetRelevance();
+				MaterialRelevance |= UMaterial::GetDefaultMaterial(MD_Surface)->GetRelevance(FeatureLevel);
 			}
 
 			bool bSectionCastsShadow = !bSectionHidden && bCastShadow &&
@@ -4605,10 +4607,11 @@ void FSkeletalMeshSceneProxy::UpdateMorphMaterialUsage_GameThread(bool bNeedsMor
 		// update the new LODSections on the render thread proxy
 		if (MaterialsToSwap.Num())
 		{
-			ENQUEUE_UNIQUE_RENDER_COMMAND_THREEPARAMETER(
+			ENQUEUE_UNIQUE_RENDER_COMMAND_FOURPARAMETER(
 			UpdateSkelProxyLODSectionElementsCmd,
 				TSet<UMaterialInterface*>,MaterialsToSwap,MaterialsToSwap,
 				UMaterialInterface*,DefaultMaterial,UMaterial::GetDefaultMaterial(MD_Surface),
+				ERHIFeatureLevel::Type, FeatureLevel, GetScene()->GetFeatureLevel(),
 			FSkeletalMeshSceneProxy*,SkelMeshSceneProxy,this,
 			{
 					for( int32 LodIdx=0; LodIdx < SkelMeshSceneProxy->LODSections.Num(); LodIdx++ )
@@ -4624,7 +4627,7 @@ void FSkeletalMeshSceneProxy::UpdateMorphMaterialUsage_GameThread(bool bNeedsMor
 							}
 						}
 					}
-					SkelMeshSceneProxy->MaterialRelevance |= DefaultMaterial->GetRelevance();
+					SkelMeshSceneProxy->MaterialRelevance |= DefaultMaterial->GetRelevance(FeatureLevel);
 			});
 		}
 	}
diff --git a/Engine/Source/Runtime/Engine/Private/StaticMeshComponent.cpp b/Engine/Source/Runtime/Engine/Private/StaticMeshComponent.cpp
--- a/Engine/Source/Runtime/Engine/Private/StaticMeshComponent.cpp
+++ b/Engine/Source/Runtime/Engine/Private/StaticMeshComponent.cpp
@@ -475,6 +475,8 @@ void UStaticMeshComponent::GetStreamingTextureInfo(TArray<FStreamingTexturePrimi
 {
 	if ( StaticMesh && !bIgnoreInstanceForTextureStreaming )
 	{
+		const auto FeatureLevel = GetWorld() ? GetWorld()->FeatureLevel : GMaxRHIFeatureLevel;
+
 		bool bHasValidLightmapCoordinates = ((StaticMesh->LightMapCoordinateIndex >= 0)
 			&& StaticMesh->RenderData
 			&& StaticMesh->RenderData->LODResources.Num() > 0
@@ -525,7 +527,7 @@ void UStaticMeshComponent::GetStreamingTextureInfo(TArray<FStreamingTexturePrimi
 				// Enumerate the textures used by the material.
 				TArray<UTexture*> Textures;
 
-				Material->GetUsedTextures(Textures, EMaterialQualityLevel::Num, false, GRHIFeatureLevel, false);
+				Material->GetUsedTextures(Textures, EMaterialQualityLevel::Num, false, FeatureLevel, false);
 
 				// Add each texture to the output with the appropriate parameters.
 				// TODO: Take into account which UVIndex is being used.
@@ -543,7 +545,7 @@ void UStaticMeshComponent::GetStreamingTextureInfo(TArray<FStreamingTexturePrimi
 			{
 				const FStaticMeshComponentLODInfo& LODInfo = LODData[LODIndex];
 				FLightMap2D* Lightmap = LODInfo.LightMap ? LODInfo.LightMap->GetLightMap2D() : NULL;
-				uint32 LightmapIndex = AllowHighQualityLightmaps() ? 0 : 1;
+				uint32 LightmapIndex = AllowHighQualityLightmaps(FeatureLevel) ? 0 : 1;
 				if ( Lightmap != NULL && Lightmap->IsValid(LightmapIndex) )
 				{
 					const FVector2D& Scale = Lightmap->GetCoordinateScale();
diff --git a/Engine/Source/Runtime/Engine/Private/StaticMeshRender.cpp b/Engine/Source/Runtime/Engine/Private/StaticMeshRender.cpp
--- a/Engine/Source/Runtime/Engine/Private/StaticMeshRender.cpp
+++ b/Engine/Source/Runtime/Engine/Private/StaticMeshRender.cpp
@@ -55,7 +55,7 @@ FStaticMeshSceneProxy::FStaticMeshSceneProxy(UStaticMeshComponent* InComponent):
 	ForcedLodModel(InComponent->ForcedLodModel),
 	bCastShadow(InComponent->CastShadow),
 	CollisionTraceFlag(ECollisionTraceFlag::CTF_UseDefault),
-	MaterialRelevance(InComponent->GetMaterialRelevance()),
+	MaterialRelevance(InComponent->GetMaterialRelevance(GetScene()->GetFeatureLevel())),
 	CollisionResponse(InComponent->GetCollisionResponseToChannels())
 {
 	check(RenderData);
@@ -64,12 +64,14 @@ FStaticMeshSceneProxy::FStaticMeshSceneProxy(UStaticMeshComponent* InComponent):
 	LevelColor = FLinearColor(1,1,1);
 	PropertyColor = FLinearColor(1,1,1);
 
+	const auto FeatureLevel = GetScene()->GetFeatureLevel();
+
 	// Copy the pointer to the volume data, async building of the data may modify the one on FStaticMeshLODResources while we are rendering
 	DistanceFieldData = RenderData->LODResources[0].DistanceFieldData;
 
 	if (GForceDefaultMaterial)
 	{
-		MaterialRelevance |= UMaterial::GetDefaultMaterial(MD_Surface)->GetRelevance();
+		MaterialRelevance |= UMaterial::GetDefaultMaterial(MD_Surface)->GetRelevance(FeatureLevel);
 	}
 
 	// Build the proxy's LOD data.
@@ -88,7 +90,7 @@ FStaticMeshSceneProxy::FStaticMeshSceneProxy(UStaticMeshComponent* InComponent):
 			bAnySectionCastsShadows |= RenderData->LODResources[LODIndex].Sections[SectionIndex].bCastShadow;
 			if (SectionInfo.Material == UMaterial::GetDefaultMaterial(MD_Surface))
 			{
-				MaterialRelevance |= UMaterial::GetDefaultMaterial(MD_Surface)->GetRelevance();
+				MaterialRelevance |= UMaterial::GetDefaultMaterial(MD_Surface)->GetRelevance(FeatureLevel);
 			}
 		}
 	}
@@ -390,9 +392,9 @@ HHitProxy* FStaticMeshSceneProxy::CreateHitProxies(UPrimitiveComponent* Componen
 #endif // WITH_EDITOR
 
 // use for render thread only
-bool UseLightPropagationVolumeRT2()
+bool UseLightPropagationVolumeRT2(ERHIFeatureLevel::Type InFeatureLevel)
 {
-	if(!IsFeatureLevelSupported(GRHIShaderPlatform, ERHIFeatureLevel::SM5))
+	if (InFeatureLevel < ERHIFeatureLevel::SM5)
 	{
 		return false;
 	}
@@ -406,10 +408,10 @@ bool UseLightPropagationVolumeRT2()
 	return Value != 0;
 }
 
-inline bool AllowShadowOnlyMesh()
+inline bool AllowShadowOnlyMesh(ERHIFeatureLevel::Type InFeatureLevel)
 {
 	// todo: later we should refine that (only if occlusion feature in LPV is on, only if inside a cascade, if shadow casting is disabled it should look at bUseEmissiveForDynamicAreaLighting)
-	return !UseLightPropagationVolumeRT2();
+	return !UseLightPropagationVolumeRT2(InFeatureLevel);
 }
 
 void FStaticMeshSceneProxy::DrawStaticElements(FStaticPrimitiveDrawInterface* PDI)
@@ -423,6 +425,7 @@ void FStaticMeshSceneProxy::DrawStaticElements(FStaticPrimitiveDrawInterface* PD
 		//Never use the dynamic path in this path, because only unselected elements will use DrawStaticElements
 		bool bUseSelectedMaterial = false;
 		const bool bUseHoveredMaterial = false;
+		const auto FeatureLevel = GetScene()->GetFeatureLevel();
 
 		//check if a LOD is being forced
 		if (ForcedLodModel > 0) 
@@ -471,12 +474,12 @@ void FStaticMeshSceneProxy::DrawStaticElements(FStaticPrimitiveDrawInterface* PD
 
 					// The shadow-only mesh can be used only if all elements cast shadows and use opaque materials with no vertex modification.
 					// In some cases (e.g. LPV) we don't want the optimization
-					bool bSafeToUseShadowOnlyMesh = AllowShadowOnlyMesh();
+					bool bSafeToUseShadowOnlyMesh = AllowShadowOnlyMesh(FeatureLevel);
 
 					bool bAnySectionCastsShadow = false;
 					for (int32 SectionIndex = 0; bSafeToUseShadowOnlyMesh && SectionIndex < LODModel.Sections.Num(); SectionIndex++)
 					{
-						const FMaterial* Material = ProxyLODInfo.Sections[SectionIndex].Material->GetRenderProxy(false)->GetMaterial(GetScene()->GetFeatureLevel());
+						const FMaterial* Material = ProxyLODInfo.Sections[SectionIndex].Material->GetRenderProxy(false)->GetMaterial(FeatureLevel);
 						const FStaticMeshSection& Section = LODModel.Sections[SectionIndex];
 						bSafeToUseShadowOnlyMesh =
 							Section.bCastShadow
@@ -1377,9 +1380,9 @@ FLightInteraction FStaticMeshSceneProxy::FLODInfo::GetInteraction(const FLightSc
 	return FLightInteraction::Dynamic();
 }
 
-FLightMapInteraction FStaticMeshSceneProxy::FLODInfo::GetLightMapInteraction() const
+FLightMapInteraction FStaticMeshSceneProxy::FLODInfo::GetLightMapInteraction(ERHIFeatureLevel::Type InFeatureLevel) const
 {
-	return LightMap ? LightMap->GetInteraction() : FLightMapInteraction();
+	return LightMap ? LightMap->GetInteraction(InFeatureLevel) : FLightMapInteraction();
 }
 
 FShadowMapInteraction FStaticMeshSceneProxy::FLODInfo::GetShadowMapInteraction() const
diff --git a/Engine/Source/Runtime/Engine/Private/TextRenderComponent.cpp b/Engine/Source/Runtime/Engine/Private/TextRenderComponent.cpp
--- a/Engine/Source/Runtime/Engine/Private/TextRenderComponent.cpp
+++ b/Engine/Source/Runtime/Engine/Private/TextRenderComponent.cpp
@@ -420,7 +420,7 @@ FTextRenderSceneProxy::FTextRenderSceneProxy( UTextRenderComponent* Component) :
 	}
 
 	TextMaterial = EffectiveMaterial;
-	MaterialRelevance |= TextMaterial->GetMaterial()->GetRelevance();
+	MaterialRelevance |= TextMaterial->GetMaterial()->GetRelevance(GetScene()->GetFeatureLevel());
 }
 
 FTextRenderSceneProxy::~FTextRenderSceneProxy()
diff --git a/Engine/Source/Runtime/Engine/Private/UnrealEngine.cpp b/Engine/Source/Runtime/Engine/Private/UnrealEngine.cpp
--- a/Engine/Source/Runtime/Engine/Private/UnrealEngine.cpp
+++ b/Engine/Source/Runtime/Engine/Private/UnrealEngine.cpp
@@ -10043,7 +10043,8 @@ bool AllowHighQualityLightmaps(ERHIFeatureLevel::Type FeatureLevel)
 {
 	return FPlatformProperties::SupportsHighQualityLightmaps()
 		&& (FeatureLevel > ERHIFeatureLevel::ES3_1)
-		&& (CVarAllowHighQualityLightMaps.GetValueOnAnyThread() != 0);
+		&& (CVarAllowHighQualityLightMaps.GetValueOnAnyThread() != 0)
+		&& !IsMobilePlatform(GShaderPlatformForFeatureLevel[FeatureLevel]);
 }
 
 bool AllowHighQualityLightmaps()
diff --git a/Engine/Source/Runtime/Engine/Private/World.cpp b/Engine/Source/Runtime/Engine/Private/World.cpp
--- a/Engine/Source/Runtime/Engine/Private/World.cpp
+++ b/Engine/Source/Runtime/Engine/Private/World.cpp
@@ -77,7 +77,7 @@ FWorldDelegates::FOnLevelChanged FWorldDelegates::LevelRemovedFromWorld;
 
 UWorld::UWorld( const class FPostConstructInitializeProperties& PCIP )
 :	UObject(PCIP)
-,	FeatureLevel(GRHIFeatureLevel)
+,	FeatureLevel(GMaxRHIFeatureLevel)
 ,	TickTaskLevel(FTickTaskManagerInterface::Get().AllocateTickTaskLevel())
 ,   bIsBuilt(false)
 ,	NextTravelType(TRAVEL_Relative)
@@ -91,7 +91,7 @@ UWorld::UWorld( const class FPostConstructInitializeProperties& PCIP )
 
 UWorld::UWorld( const class FPostConstructInitializeProperties& PCIP,const FURL& InURL )
 :	UObject(PCIP)
-,	FeatureLevel(GRHIFeatureLevel)
+,	FeatureLevel(GMaxRHIFeatureLevel)
 ,	URL(InURL)
 ,	TickTaskLevel(FTickTaskManagerInterface::Get().AllocateTickTaskLevel())
 ,   bIsBuilt(false)
@@ -5139,26 +5139,6 @@ void UWorld::GetLandscapeTexturesAndMaterials(ULevel* Level, TArray<UObject*>& O
 	}
 }
 
-void UWorld::ChangeWorldFeatureLevel(ERHIFeatureLevel::Type InFeatureLevel)
-{
-	if (InFeatureLevel != FeatureLevel)
-	{
-		FeatureLevel = InFeatureLevel;
-
-		if (Scene)
-		{
-			Scene->Release();
-			GetRendererModule().RemoveScene(Scene);
-
-			Scene = GetRendererModule().AllocateScene(this, bRequiresHitProxies, InFeatureLevel);
-
-			FFXSystemInterface::Destroy(FXSystem);
-			FXSystem = FFXSystemInterface::Create(InFeatureLevel);
-			Scene->SetFXSystem(FXSystem);
-		}
-	}
-}
-
 void UWorld::ChangeFeatureLevel(ERHIFeatureLevel::Type InFeatureLevel)
 {
 	if (InFeatureLevel != FeatureLevel)
@@ -5166,10 +5146,9 @@ void UWorld::ChangeFeatureLevel(ERHIFeatureLevel::Type InFeatureLevel)
 		FlushRenderingCommands();
 
 		FGlobalComponentReregisterContext RecreateComponents;
-
-		// TODO: Remove these
-		GRHIFeatureLevelValue = InFeatureLevel;
-		GRHIShaderPlatformValue = GShaderPlatformForFeatureLevel[InFeatureLevel];
+		
+		// Decrement refcount on old feature level
+		UMaterialInterface::SetGlobalRequiredFeatureLevel(InFeatureLevel, true);
 
 		UMaterial::AllMaterialsCacheResourceShadersForRendering();
 		UMaterialInstance::AllMaterialsCacheResourceShadersForRendering();
@@ -5183,58 +5162,21 @@ void UWorld::ChangeFeatureLevel(ERHIFeatureLevel::Type InFeatureLevel)
 		}
 
 		// try iterating over all worlds here
-		ChangeWorldFeatureLevel(InFeatureLevel);
-
-		FOutputDeviceNull Ar;
-		RecompileShaders(TEXT("CHANGED"), Ar);
-	}
-}
-
-void UWorld::ChangeAllWorldFeatureLevels(ERHIFeatureLevel::Type InFeatureLevel)
-{
-	bool bSomethingChanged = false;
-
-	// try iterating over all worlds here
-	for (TObjectIterator<UWorld> It; It; ++It)
-	{
-		UWorld* World = *It;
+		FeatureLevel = InFeatureLevel;
 
-		if (World->FeatureLevel != InFeatureLevel)
+		if (Scene)
 		{
-			bSomethingChanged = true;
-		}
-	}
-
-	if (bSomethingChanged)
-	{
-		FlushRenderingCommands();
-
-		FGlobalComponentReregisterContext RecreateComponents;
+			Scene->Release();
+			GetRendererModule().RemoveScene(Scene);
 
-		GRHIFeatureLevelValue = InFeatureLevel;
-		GRHIShaderPlatformValue = GShaderPlatformForFeatureLevel[InFeatureLevel];
+			Scene = GetRendererModule().AllocateScene(this, bRequiresHitProxies, InFeatureLevel);
 
-		UMaterial::AllMaterialsCacheResourceShadersForRendering();
-		UMaterialInstance::AllMaterialsCacheResourceShadersForRendering();
-		GetGlobalShaderMap(InFeatureLevel, false);
-		GShaderCompilingManager->ProcessAsyncResults(false, true);
-		
-		//invalidate global bound shader states so they will be created with the new shaders the next time they are set (in SetGlobalBoundShaderState)
-		for (TLinkedList<FGlobalBoundShaderStateResource*>::TIterator It(FGlobalBoundShaderStateResource::GetGlobalBoundShaderStateList()); It; It.Next())
-		{
-			BeginUpdateResourceRHI(*It);
+			FFXSystemInterface::Destroy(FXSystem);
+			FXSystem = FFXSystemInterface::Create(InFeatureLevel);
+			Scene->SetFXSystem(FXSystem);
 		}
 
-		// try iterating over all worlds here
-		for (TObjectIterator<UWorld> It; It; ++It)
-		{
-			UWorld* World = *It;
-
-			if (World->FeatureLevel != InFeatureLevel)
-			{
-				World->ChangeWorldFeatureLevel(InFeatureLevel);
-			}
-		}
+		TriggerStreamingDataRebuild();
 
 		FOutputDeviceNull Ar;
 		RecompileShaders(TEXT("CHANGED"), Ar);
diff --git a/Engine/Source/Runtime/Engine/Public/Landscape/LandscapeRender.h b/Engine/Source/Runtime/Engine/Public/Landscape/LandscapeRender.h
--- a/Engine/Source/Runtime/Engine/Public/Landscape/LandscapeRender.h
+++ b/Engine/Source/Runtime/Engine/Public/Landscape/LandscapeRender.h
@@ -378,9 +378,9 @@ class FLandscapeComponentSceneProxy : public FPrimitiveSceneProxy
 		// FLightCacheInterface
 		virtual FLightInteraction GetInteraction(const class FLightSceneProxy* LightSceneProxy) const;
 
-		virtual FLightMapInteraction GetLightMapInteraction() const
+		virtual FLightMapInteraction GetLightMapInteraction(ERHIFeatureLevel::Type InFeatureLevel) const
 		{
-			return LightMap ? LightMap->GetInteraction() : FLightMapInteraction();
+			return LightMap ? LightMap->GetInteraction(InFeatureLevel) : FLightMapInteraction();
 		}
 
 		virtual FShadowMapInteraction GetShadowMapInteraction() const
diff --git a/Engine/Source/Runtime/Engine/Public/LightMap.h b/Engine/Source/Runtime/Engine/Public/LightMap.h
--- a/Engine/Source/Runtime/Engine/Public/LightMap.h
+++ b/Engine/Source/Runtime/Engine/Public/LightMap.h
@@ -71,7 +71,7 @@ public:
 	// FLightMap interface.
 	virtual void AddReferencedObjects( FReferenceCollector& Collector ) {}
 	virtual void Serialize(FArchive& Ar);
-	virtual FLightMapInteraction GetInteraction() const = 0;
+	virtual FLightMapInteraction GetInteraction(ERHIFeatureLevel::Type InFeatureLevel) const = 0;
 
 	// Runtime type casting.
 	virtual class FLightMap2D* GetLightMap2D() { return NULL; }
@@ -228,7 +228,7 @@ public:
 	virtual void AddReferencedObjects( FReferenceCollector& Collector );
 
 	virtual void Serialize(FArchive& Ar);
-	virtual FLightMapInteraction GetInteraction() const;
+	virtual FLightMapInteraction GetInteraction(ERHIFeatureLevel::Type InFeatureLevel) const;
 
 	// Runtime type casting.
 	virtual const FLightMap2D* GetLightMap2D() const { return this; }
@@ -372,7 +372,7 @@ public:
 	// FLightMap interface.
 	virtual void Serialize(FArchive& Ar);
 
-	virtual FLightMapInteraction GetInteraction() const { return FLightMapInteraction::None(); }
+	virtual FLightMapInteraction GetInteraction(ERHIFeatureLevel::Type InFeatureLevel) const { return FLightMapInteraction::None(); }
 };
 
 /** Stores debug information for a lightmap sample. */
diff --git a/Engine/Source/Runtime/Engine/Public/MaterialShared.h b/Engine/Source/Runtime/Engine/Public/MaterialShared.h
--- a/Engine/Source/Runtime/Engine/Public/MaterialShared.h
+++ b/Engine/Source/Runtime/Engine/Public/MaterialShared.h
@@ -553,7 +553,7 @@ public:
 	ENGINE_API void GetShaderList(TMap<FShaderId,FShader*>& OutShaders) const;
 
 	/** Registers a material shader map in the global map so it can be used by materials. */
-	void Register();
+	void Register(EShaderPlatform InShaderPlatform);
 
 	// Reference counting.
 	ENGINE_API void AddRef();
diff --git a/Engine/Source/Runtime/Engine/Public/ParticleEmitterInstances.h b/Engine/Source/Runtime/Engine/Public/ParticleEmitterInstances.h
--- a/Engine/Source/Runtime/Engine/Public/ParticleEmitterInstances.h
+++ b/Engine/Source/Runtime/Engine/Public/ParticleEmitterInstances.h
@@ -618,7 +618,7 @@ public:
 	 * @param OutMaterialRelevance - Pointer to where material relevance flags will be stored.
 	 * @param LODLevel - The LOD level for which to compute material relevance flags.
 	 */
-	virtual void GatherMaterialRelevance( FMaterialRelevance* OutMaterialRelevance, const UParticleLODLevel* LODLevel ) const;
+	virtual void GatherMaterialRelevance( FMaterialRelevance* OutMaterialRelevance, const UParticleLODLevel* LODLevel, ERHIFeatureLevel::Type InFeatureLevel ) const;
 
 	/**
 	 * When an emitter is killed, this will check other emitters and clean up anything pointing to this one
@@ -839,7 +839,7 @@ struct ENGINE_API FParticleMeshEmitterInstance : public FParticleEmitterInstance
 	 * @param OutMaterialRelevance - Pointer to where material relevance flags will be stored.
 	 * @param LODLevel - The LOD level for which to compute material relevance flags.
 	 */
-	virtual void GatherMaterialRelevance( FMaterialRelevance* OutMaterialRelevance, const UParticleLODLevel* LODLevel ) const override;
+	virtual void GatherMaterialRelevance(FMaterialRelevance* OutMaterialRelevance, const UParticleLODLevel* LODLevel, ERHIFeatureLevel::Type InFeatureLevel) const override;
 
 	/**
 	 * Gets the materials applied to each section of a mesh.
diff --git a/Engine/Source/Runtime/Engine/Public/SceneManagement.h b/Engine/Source/Runtime/Engine/Public/SceneManagement.h
--- a/Engine/Source/Runtime/Engine/Public/SceneManagement.h
+++ b/Engine/Source/Runtime/Engine/Public/SceneManagement.h
@@ -219,11 +219,11 @@ public:
 	// Accessors.
 	ELightMapInteractionType GetType() const { return Type; }
 	
-	const ULightMapTexture2D* GetTexture() const
+	const ULightMapTexture2D* GetTexture(bool bHighQuality) const
 	{
 		check(Type == LMIT_Texture);
 #if ALLOW_LQ_LIGHTMAPS && ALLOW_HQ_LIGHTMAPS
-		return AllowsHighQualityLightmaps() ? HighQualityTexture : LowQualityTexture;
+		return bHighQuality ? HighQualityTexture : LowQualityTexture;
 #elif ALLOW_HQ_LIGHTMAPS
 		return HighQualityTexture;
 #else
@@ -446,7 +446,7 @@ class FLightCacheInterface
 {
 public:
 	virtual FLightInteraction GetInteraction(const class FLightSceneProxy* LightSceneProxy) const = 0;
-	virtual FLightMapInteraction GetLightMapInteraction() const = 0;
+	virtual FLightMapInteraction GetLightMapInteraction(ERHIFeatureLevel::Type InFeatureLevel) const = 0;
 	virtual FShadowMapInteraction GetShadowMapInteraction() const { return FShadowMapInteraction::None(); }
 };
 
diff --git a/Engine/Source/Runtime/Engine/Public/ShaderCompiler.h b/Engine/Source/Runtime/Engine/Public/ShaderCompiler.h
--- a/Engine/Source/Runtime/Engine/Public/ShaderCompiler.h
+++ b/Engine/Source/Runtime/Engine/Public/ShaderCompiler.h
@@ -338,8 +338,5 @@ extern void GlobalBeginCompileShader(
 /** Implementation of the 'recompileshaders' console command.  Recompiles shaders at runtime based on various criteria. */
 extern bool RecompileShaders(const TCHAR* Cmd, FOutputDevice& Ar);
 
-/** Returns whether the global shader map contains all global shader types. */
-extern ENGINE_API bool IsGlobalShaderMapComplete();
-
-/** Returns whether all global shader types containing the substring are complete and ready for rendering. */
-extern ENGINE_API bool AreGlobalShadersComplete(const TCHAR* TypeNameSubstring);
\ No newline at end of file
+/** Returns whether all global shader types containing the substring are complete and ready for rendering. if type name is null, check everything */
+extern ENGINE_API bool IsGlobalShaderMapComplete(const TCHAR* TypeNameSubstring = nullptr);
\ No newline at end of file
diff --git a/Engine/Source/Runtime/Engine/Public/StaticMeshResources.h b/Engine/Source/Runtime/Engine/Public/StaticMeshResources.h
--- a/Engine/Source/Runtime/Engine/Public/StaticMeshResources.h
+++ b/Engine/Source/Runtime/Engine/Public/StaticMeshResources.h
@@ -838,7 +838,7 @@ protected:
 		// FLightCacheInterface.
 		virtual FLightInteraction GetInteraction(const FLightSceneProxy* LightSceneProxy) const;
 
-		virtual FLightMapInteraction GetLightMapInteraction() const;
+		virtual FLightMapInteraction GetLightMapInteraction(ERHIFeatureLevel::Type InFeatureLevel) const;
 
 		virtual FShadowMapInteraction GetShadowMapInteraction() const;
 
diff --git a/Engine/Source/Runtime/NullDrv/Private/NullRHI.cpp b/Engine/Source/Runtime/NullDrv/Private/NullRHI.cpp
--- a/Engine/Source/Runtime/NullDrv/Private/NullRHI.cpp
+++ b/Engine/Source/Runtime/NullDrv/Private/NullRHI.cpp
@@ -14,7 +14,7 @@
 
 FNullDynamicRHI::FNullDynamicRHI()
 {
-	GRHIShaderPlatformValue = ShaderFormatToLegacyShaderPlatform(FName(FPlatformMisc::GetNullRHIShaderFormat()));
+	GMaxRHIShaderPlatformValue = ShaderFormatToLegacyShaderPlatform(FName(FPlatformMisc::GetNullRHIShaderFormat()));
 }
 
 void FNullDynamicRHI::Init()
diff --git a/Engine/Source/Runtime/OpenGLDrv/Private/OpenGLDevice.cpp b/Engine/Source/Runtime/OpenGLDrv/Private/OpenGLDevice.cpp
--- a/Engine/Source/Runtime/OpenGLDrv/Private/OpenGLDevice.cpp
+++ b/Engine/Source/Runtime/OpenGLDrv/Private/OpenGLDevice.cpp
@@ -609,15 +609,15 @@ static void InitRHICapabilitiesForGL()
 	const GLint MinorVersion = FOpenGL::GetMinorVersion();
 
 	// Shader platform & RHI feature level
-	GMaxRHIFeatureLevelValue = GRHIFeatureLevelValue = FOpenGL::GetFeatureLevel();
-	GRHIShaderPlatformValue = FOpenGL::GetShaderPlatform();
+	GMaxRHIFeatureLevelValue = FOpenGL::GetFeatureLevel();
+	GMaxRHIShaderPlatformValue = FOpenGL::GetShaderPlatform();
 
 	// Emulate uniform buffers on ES2, unless we're on a desktop platform emulating ES2.
-	GUseEmulatedUniformBuffers = IsES2Platform(GRHIShaderPlatform) && !IsPCPlatform(GRHIShaderPlatform);
+	GUseEmulatedUniformBuffers = IsES2Platform(GMaxRHIShaderPlatform) && !IsPCPlatform(GMaxRHIShaderPlatform);
 
 	FString FeatureLevelName;
 	GetFeatureLevelName(GMaxRHIFeatureLevel, FeatureLevelName);
-	FString ShaderPlatformName = LegacyShaderPlatformToShaderFormat(GRHIShaderPlatform).ToString();
+	FString ShaderPlatformName = LegacyShaderPlatformToShaderFormat(GMaxRHIShaderPlatform).ToString();
 
 	UE_LOG(LogRHI, Log, TEXT("OpenGL MajorVersion = %d, MinorVersion = %d, ShaderPlatform = %s, FeatureLevel = %s"), MajorVersion, MinorVersion, *ShaderPlatformName, *FeatureLevelName);
 
diff --git a/Engine/Source/Runtime/RHI/Private/RHI.cpp b/Engine/Source/Runtime/RHI/Private/RHI.cpp
--- a/Engine/Source/Runtime/RHI/Private/RHI.cpp
+++ b/Engine/Source/Runtime/RHI/Private/RHI.cpp
@@ -221,28 +221,19 @@ void RHIPrivateBeginFrame()
 // The current shader platform.
 //
 
-RHI_API EShaderPlatform GRHIShaderPlatformValue = SP_PCD3D_SM5;
+RHI_API EShaderPlatform GMaxRHIShaderPlatformValue = SP_PCD3D_SM5;
 
 /** The maximum feature level supported on this machine */
 RHI_API ERHIFeatureLevel::Type GMaxRHIFeatureLevelValue = ERHIFeatureLevel::SM5;
 
-/** The current feature level being used on this machine - This is only temporary and will go away when feature level becomes dynamic */
-RHI_API ERHIFeatureLevel::Type GRHIFeatureLevelValue = ERHIFeatureLevel::SM5;
-
-RHI_API ERHIFeatureLevel::Type GetCurrentRHIFeatureLevel()
-{
-	//check(!IsInRenderingThread());
-	return GRHIFeatureLevelValue;
-}
-
 RHI_API ERHIFeatureLevel::Type GetMaxRHIFeatureLevel()
 {
 	return GMaxRHIFeatureLevelValue;
 }
 
-RHI_API EShaderPlatform GetCurrentRHIShaderPlatform()
+RHI_API EShaderPlatform GetMaxRHIShaderPlatform()
 {
-	return GRHIShaderPlatformValue;
+	return GMaxRHIShaderPlatformValue;
 }
 
 FName FeatureLevelNames[] = 
diff --git a/Engine/Source/Runtime/RHI/Public/RHI.h b/Engine/Source/Runtime/RHI/Public/RHI.h
--- a/Engine/Source/Runtime/RHI/Public/RHI.h
+++ b/Engine/Source/Runtime/RHI/Public/RHI.h
@@ -237,31 +237,35 @@ extern RHI_API void GetFeatureLevelName(ERHIFeatureLevel::Type InFeatureLevel, F
 /** Creates an FName for the given feature level. */
 extern RHI_API void GetFeatureLevelName(ERHIFeatureLevel::Type InFeatureLevel, FName& OutName);
 
-extern RHI_API ERHIFeatureLevel::Type GRHIFeatureLevelValue;
 extern RHI_API ERHIFeatureLevel::Type GMaxRHIFeatureLevelValue;
-extern RHI_API EShaderPlatform GRHIShaderPlatformValue;
+extern RHI_API EShaderPlatform GMaxRHIShaderPlatformValue;
 
 /** Function for retrieving the current feature level. Only exists to provide a convient way of tracking access to this global during mobile preview work */
-extern RHI_API ERHIFeatureLevel::Type GetCurrentRHIFeatureLevel();
 extern RHI_API ERHIFeatureLevel::Type GetMaxRHIFeatureLevel();
-extern RHI_API EShaderPlatform GetCurrentRHIShaderPlatform();
+extern RHI_API EShaderPlatform GetMaxRHIShaderPlatform();
 
 /** treating GRHIFeatureLevel as a function allows for better usage tracking at small cost to performance. */
 #define RHI_FEATURE_LEVEL_AS_FUNCTION 0
 
 #if RHI_FEATURE_LEVEL_AS_FUNCTION
-	#define GRHIFeatureLevel GetCurrentRHIFeatureLevel()
 	#define GMaxRHIFeatureLevel GetMaxRHIFeatureLevel()
-	#define GRHIShaderPlatform GetCurrentRHIShaderPlatform()
+	#define GMaxRHIShaderPlatform GetMaxRHIShaderPlatform()
 #else
-	#define GRHIFeatureLevel GRHIFeatureLevelValue
 	#define GMaxRHIFeatureLevel GMaxRHIFeatureLevelValue
-	#define GRHIShaderPlatform GRHIShaderPlatformValue
+	#define GMaxRHIShaderPlatform GMaxRHIShaderPlatformValue
 #endif
 
+#define GRHIShaderPlatform GMaxRHIShaderPlatform
+#define GRHIFeatureLevel GMaxRHIFeatureLevel
+
 /** Table for finding out which shader platform corresponds to a given feature level for this RHI. */
 extern RHI_API EShaderPlatform GShaderPlatformForFeatureLevel[ERHIFeatureLevel::Num];
 
+/** Get the shader platform associated with the supplied feature level on this machine */
+inline EShaderPlatform GetFeatureLevelShaderPlatform(ERHIFeatureLevel::Type InFeatureLevel)
+{
+	return GShaderPlatformForFeatureLevel[InFeatureLevel];
+}
 
 inline FArchive& operator <<(FArchive& Ar, EResourceLockMode& LockMode)
 {
diff --git a/Engine/Source/Runtime/Renderer/Private/BasePassRendering.h b/Engine/Source/Runtime/Renderer/Private/BasePassRendering.h
--- a/Engine/Source/Runtime/Renderer/Private/BasePassRendering.h
+++ b/Engine/Source/Runtime/Renderer/Private/BasePassRendering.h
@@ -879,7 +879,7 @@ void ProcessBasePassMesh(
 		else
 		{
 			const FLightMapInteraction LightMapInteraction = (bAllowStaticLighting && Parameters.Mesh.LCI && bIsLitMaterial) 
-				? Parameters.Mesh.LCI->GetLightMapInteraction() 
+				? Parameters.Mesh.LCI->GetLightMapInteraction(Parameters.FeatureLevel) 
 				: FLightMapInteraction();
 
 			// force LQ lightmaps based on system settings
diff --git a/Engine/Source/Runtime/Renderer/Private/ForwardBasePassRendering.cpp b/Engine/Source/Runtime/Renderer/Private/ForwardBasePassRendering.cpp
--- a/Engine/Source/Runtime/Renderer/Private/ForwardBasePassRendering.cpp
+++ b/Engine/Source/Runtime/Renderer/Private/ForwardBasePassRendering.cpp
@@ -421,7 +421,7 @@ void FForwardShadingSceneRenderer::RenderForwardShadingBasePass(FRHICommandListI
 				}
 			}
 
-			const bool bNeedToSwitchVerticalAxis = RHINeedsToSwitchVerticalAxis(GRHIShaderPlatform);
+			const bool bNeedToSwitchVerticalAxis = RHINeedsToSwitchVerticalAxis(GShaderPlatformForFeatureLevel[FeatureLevel]);
 
 			// Draw the base pass for the view's batched mesh elements.
 			DrawViewElements<FBasePassForwardOpaqueDrawingPolicyFactory>(RHICmdList, View, FBasePassForwardOpaqueDrawingPolicyFactory::ContextType(ESceneRenderTargetsMode::DontSet), SDPG_World, true);
diff --git a/Engine/Source/Runtime/Renderer/Private/ForwardBasePassRendering.h b/Engine/Source/Runtime/Renderer/Private/ForwardBasePassRendering.h
--- a/Engine/Source/Runtime/Renderer/Private/ForwardBasePassRendering.h
+++ b/Engine/Source/Runtime/Renderer/Private/ForwardBasePassRendering.h
@@ -516,7 +516,7 @@ void ProcessBasePassMeshForForwardShading(
 	const bool bIsLitMaterial = Parameters.ShadingModel != MSM_Unlit;
 
 	const FLightMapInteraction LightMapInteraction = (Parameters.Mesh.LCI && bIsLitMaterial) 
-		? Parameters.Mesh.LCI->GetLightMapInteraction() 
+		? Parameters.Mesh.LCI->GetLightMapInteraction(Parameters.FeatureLevel) 
 		: FLightMapInteraction();
 
 	if (bIsLitMaterial)
diff --git a/Engine/Source/Runtime/Renderer/Private/LightMapDensityRendering.cpp b/Engine/Source/Runtime/Renderer/Private/LightMapDensityRendering.cpp
--- a/Engine/Source/Runtime/Renderer/Private/LightMapDensityRendering.cpp
+++ b/Engine/Source/Runtime/Renderer/Private/LightMapDensityRendering.cpp
@@ -124,6 +124,7 @@ bool FLightMapDensityDrawingPolicyFactory::DrawDynamicMesh(
 	const FMaterialRenderProxy* MaterialRenderProxy = Mesh.MaterialRenderProxy;
 	const FMaterial* Material = MaterialRenderProxy->GetMaterial(View.GetFeatureLevel());
 	const EBlendMode BlendMode = Material->GetBlendMode();
+	const auto FeatureLevel = View.GetFeatureLevel();
 
 	const bool bMaterialMasked = Material->IsMasked();
 	const bool bMaterialModifiesMesh = Material->MaterialModifiesMeshPosition_RenderThread();
@@ -134,14 +135,14 @@ bool FLightMapDensityDrawingPolicyFactory::DrawDynamicMesh(
 	}
 
 	bool bIsLitMaterial = (Material->GetShadingModel() != MSM_Unlit);
-	/*const */FLightMapInteraction LightMapInteraction = (Mesh.LCI && bIsLitMaterial) ? Mesh.LCI->GetLightMapInteraction() : FLightMapInteraction();
+	/*const */FLightMapInteraction LightMapInteraction = (Mesh.LCI && bIsLitMaterial) ? Mesh.LCI->GetLightMapInteraction(FeatureLevel) : FLightMapInteraction();
 	// force simple lightmaps based on system settings
-	bool bAllowHighQualityLightMaps = AllowHighQualityLightmaps(View.GetFeatureLevel()) && LightMapInteraction.AllowsHighQualityLightmaps();
+	bool bAllowHighQualityLightMaps = AllowHighQualityLightmaps(FeatureLevel) && LightMapInteraction.AllowsHighQualityLightmaps();
 	if (bIsLitMaterial && PrimitiveSceneProxy && (LightMapInteraction.GetType() == LMIT_Texture))
 	{
 		// Should this object be texture lightmapped? Ie, is lighting not built for it??
 		bool bUseDummyLightMapPolicy = false;
-		if ((Mesh.LCI == NULL) || (Mesh.LCI->GetLightMapInteraction().GetType() != LMIT_Texture))
+		if ((Mesh.LCI == NULL) || (Mesh.LCI->GetLightMapInteraction(FeatureLevel).GetType() != LMIT_Texture))
 		{
 			bUseDummyLightMapPolicy = true;
 
diff --git a/Engine/Source/Runtime/Renderer/Private/LightMapDensityRendering.h b/Engine/Source/Runtime/Renderer/Private/LightMapDensityRendering.h
--- a/Engine/Source/Runtime/Renderer/Private/LightMapDensityRendering.h
+++ b/Engine/Source/Runtime/Renderer/Private/LightMapDensityRendering.h
@@ -343,13 +343,17 @@ public:
 		// LMResolutionScale is the physical resolution of the lightmap texture
 		FVector2D LMResolutionScale(1.0f,1.0f);
 
+		const auto FeatureLevel = View.GetFeatureLevel();
+
+		bool bHighQualityLightMaps = AllowHighQualityLightmaps(FeatureLevel);
+
 		bool bTextureMapped = false;
 		if (Mesh.LCI &&
-			(Mesh.LCI->GetLightMapInteraction().GetType() == LMIT_Texture) &&
-			Mesh.LCI->GetLightMapInteraction().GetTexture())
+			(Mesh.LCI->GetLightMapInteraction(FeatureLevel).GetType() == LMIT_Texture) &&
+			Mesh.LCI->GetLightMapInteraction(FeatureLevel).GetTexture(bHighQualityLightMaps))
 		{
-			LMResolutionScale.X = Mesh.LCI->GetLightMapInteraction().GetTexture()->GetSizeX();
-			LMResolutionScale.Y = Mesh.LCI->GetLightMapInteraction().GetTexture()->GetSizeY();
+			LMResolutionScale.X = Mesh.LCI->GetLightMapInteraction(FeatureLevel).GetTexture(bHighQualityLightMaps)->GetSizeX();
+			LMResolutionScale.Y = Mesh.LCI->GetLightMapInteraction(FeatureLevel).GetTexture(bHighQualityLightMaps)->GetSizeY();
 			bTextureMapped = true;
 
 			BuiltLightingAndSelectedFlags.X = 1.0f;
diff --git a/Engine/Source/Runtime/Renderer/Private/LightMapRendering.h b/Engine/Source/Runtime/Renderer/Private/LightMapRendering.h
--- a/Engine/Source/Runtime/Renderer/Private/LightMapRendering.h
+++ b/Engine/Source/Runtime/Renderer/Private/LightMapRendering.h
@@ -267,7 +267,7 @@ public:
 		if(PixelShaderParameters)
 		{
 			PixelShaderParameters->SetLightMapScale(RHICmdList, PixelShader,LightMapInteraction);
-			PixelShaderParameters->SetLightMapTexture(RHICmdList, PixelShader, LightMapInteraction.GetTexture());
+			PixelShaderParameters->SetLightMapTexture(RHICmdList, PixelShader, LightMapInteraction.GetTexture(AllowHighQualityLightmaps(View.GetFeatureLevel())));
 			PixelShaderParameters->SetSkyOcclusionTexture(RHICmdList, PixelShader, LightMapInteraction.GetSkyOcclusionTexture());
 		}
 	}
diff --git a/Engine/Source/Runtime/Renderer/Private/PostProcess/PostProcessMobile.cpp b/Engine/Source/Runtime/Renderer/Private/PostProcess/PostProcessMobile.cpp
--- a/Engine/Source/Runtime/Renderer/Private/PostProcess/PostProcessMobile.cpp
+++ b/Engine/Source/Runtime/Renderer/Private/PostProcess/PostProcessMobile.cpp
@@ -154,7 +154,7 @@ void FRCPassPostProcessBloomSetupES2::SetShader(const FRenderingCompositePassCon
 	uint32 UseSunDof = (UseSun << 1) + UseDof;
 
 	static const auto CVarMobileMSAA = IConsoleManager::Get().FindTConsoleVariableDataInt(TEXT("r.MobileMSAA"));
-	UseSunDof += (GRHIShaderPlatform == SP_METAL) ? (CVarMobileMSAA ? (CVarMobileMSAA->GetValueOnRenderThread() > 1 ? 4 : 0) : 0) : 0;
+	UseSunDof += (GShaderPlatformForFeatureLevel[Context.GetFeatureLevel()] == SP_METAL) ? (CVarMobileMSAA ? (CVarMobileMSAA->GetValueOnRenderThread() > 1 ? 4 : 0) : 0) : 0;
 
 	switch(UseSunDof)
 	{
diff --git a/Engine/Source/Runtime/Renderer/Private/PostProcess/PostProcessTonemap.cpp b/Engine/Source/Runtime/Renderer/Private/PostProcess/PostProcessTonemap.cpp
--- a/Engine/Source/Runtime/Renderer/Private/PostProcess/PostProcessTonemap.cpp
+++ b/Engine/Source/Runtime/Renderer/Private/PostProcess/PostProcessTonemap.cpp
@@ -729,7 +729,7 @@ static uint32 TonemapperGenerateBitmaskMobile(const FViewInfo* RESTRICT View, bo
 	}
 
 	static const auto CVarMobileMSAA = IConsoleManager::Get().FindTConsoleVariableDataInt(TEXT("r.MobileMSAA"));
-	if((GRHIShaderPlatform == SP_METAL) && (CVarMobileMSAA ? CVarMobileMSAA->GetValueOnAnyThread() > 1 : false))
+	if ((GShaderPlatformForFeatureLevel[View->GetFeatureLevel()] == SP_METAL) && (CVarMobileMSAA ? CVarMobileMSAA->GetValueOnAnyThread() > 1 : false))
 	{
 		Bitmask += TonemapperMsaa;
 	}
diff --git a/Engine/Source/Runtime/Renderer/Private/PostProcess/SceneRenderTargets.h b/Engine/Source/Runtime/Renderer/Private/PostProcess/SceneRenderTargets.h
--- a/Engine/Source/Runtime/Renderer/Private/PostProcess/SceneRenderTargets.h
+++ b/Engine/Source/Runtime/Renderer/Private/PostProcess/SceneRenderTargets.h
@@ -263,7 +263,7 @@ public:
 	const FTexture2DRHIRef* GetActualDepthTexture() const
 	{
 		const FTexture2DRHIRef* DepthTexture = NULL;
-		if((CurrentFeatureLevel >= ERHIFeatureLevel::SM4) || IsPCPlatform(GRHIShaderPlatform))
+		if((CurrentFeatureLevel >= ERHIFeatureLevel::SM4) || IsPCPlatform(GShaderPlatformForFeatureLevel[CurrentFeatureLevel]))
 		{
 			if(GSupportsDepthFetchDuringDepthTest)
 			{
diff --git a/Engine/Source/Runtime/Renderer/Private/RendererScene.cpp b/Engine/Source/Runtime/Renderer/Private/RendererScene.cpp
--- a/Engine/Source/Runtime/Renderer/Private/RendererScene.cpp
+++ b/Engine/Source/Runtime/Renderer/Private/RendererScene.cpp
@@ -217,6 +217,8 @@ FScene::FScene(UWorld* InWorld, bool bInRequiresHitProxies, bool bInIsEditorScen
 {
 	check(World);
 
+	FeatureLevel = World->FeatureLevel;
+
 	static auto* MobileHDRCvar = IConsoleManager::Get().FindTConsoleVariableDataInt(TEXT("r.MobileHDR"));
 	static auto* MobileHDR32bppCvar = IConsoleManager::Get().FindTConsoleVariableDataInt(TEXT("r.MobileHDR32bpp"));
 	bStaticDrawListsMobileHDR = MobileHDRCvar->GetValueOnAnyThread() == 1;
diff --git a/Engine/Source/Runtime/Renderer/Private/SceneHitProxyRendering.cpp b/Engine/Source/Runtime/Renderer/Private/SceneHitProxyRendering.cpp
--- a/Engine/Source/Runtime/Renderer/Private/SceneHitProxyRendering.cpp
+++ b/Engine/Source/Runtime/Renderer/Private/SceneHitProxyRendering.cpp
@@ -161,7 +161,7 @@ FHitProxyDrawingPolicy::FHitProxyDrawingPolicy(
 	DomainShader = NULL;
 
 	EMaterialTessellationMode MaterialTessellationMode = MaterialResource->GetTessellationMode();
-	if(RHISupportsTessellation(GRHIShaderPlatform)
+	if (RHISupportsTessellation(GShaderPlatformForFeatureLevel[InFeatureLevel])
 		&& InVertexFactory->GetType()->SupportsTessellationShaders() 
 		&& MaterialTessellationMode != MTM_NoTessellation)
 	{
@@ -379,7 +379,7 @@ void RenderHitProxies(FRHICommandListImmediate& RHICmdList, const FSceneRenderer
 	RHICmdList.SetDepthStencilState(TStaticDepthStencilState<true, CF_GreaterEqual>::GetRHI());
 	RHICmdList.SetBlendState(TStaticBlendState<>::GetRHI());
 
-	const bool bNeedToSwitchVerticalAxis = RHINeedsToSwitchVerticalAxis(GRHIShaderPlatform);
+	const bool bNeedToSwitchVerticalAxis = RHINeedsToSwitchVerticalAxis(GShaderPlatformForFeatureLevel[SceneRenderer->FeatureLevel]);
 
 	for (int32 ViewIndex = 0; ViewIndex < Views.Num(); ViewIndex++)
 	{
diff --git a/Engine/Source/Runtime/Renderer/Private/ScenePrivate.h b/Engine/Source/Runtime/Renderer/Private/ScenePrivate.h
--- a/Engine/Source/Runtime/Renderer/Private/ScenePrivate.h
+++ b/Engine/Source/Runtime/Renderer/Private/ScenePrivate.h
@@ -1461,7 +1461,7 @@ public:
 
 	virtual bool IsEditorScene() const override { return bIsEditorScene; }
 
-	virtual ERHIFeatureLevel::Type GetFeatureLevel() const override { return GetWorld()->FeatureLevel; }
+	virtual ERHIFeatureLevel::Type GetFeatureLevel() const override { return FeatureLevel; }
 
 	bool ShouldRenderSkylight() const
 	{
@@ -1579,6 +1579,9 @@ private:
 	 * Note: This is tracked on the game thread!
 	 */
 	bool bHasSkyLight;
+
+	/** This scene's feature level */
+	ERHIFeatureLevel::Type FeatureLevel;
 };
 
 #endif // __SCENEPRIVATE_H__
diff --git a/Engine/Source/Runtime/ShaderCore/Private/Shader.cpp b/Engine/Source/Runtime/ShaderCore/Private/Shader.cpp
--- a/Engine/Source/Runtime/ShaderCore/Private/Shader.cpp
+++ b/Engine/Source/Runtime/ShaderCore/Private/Shader.cpp
@@ -392,21 +392,64 @@ void FShaderResource::FinishCleanup()
 	delete this;
 }
 
+bool ArePlatformsCompatible(EShaderPlatform CurrentPlatform, EShaderPlatform TargetPlatform)
+{
+	bool bFeatureLevelCompatible = CurrentPlatform == TargetPlatform;
+	
+	if (!bFeatureLevelCompatible && IsPCPlatform(CurrentPlatform) && IsPCPlatform(TargetPlatform) )
+	{
+		if (CurrentPlatform == SP_OPENGL_SM4_MAC || TargetPlatform == SP_OPENGL_SM4_MAC)
+		{
+			// prevent SP_OPENGL_SM4 == SP_OPENGL_SM4_MAC, allow SP_OPENGL_SM4_MAC == SP_OPENGL_SM4_MAC,
+			// allow lesser feature levels on SP_OPENGL_SM4_MAC device.
+			// do not allow MAC targets to work on non MAC devices.
+			bFeatureLevelCompatible = CurrentPlatform == SP_OPENGL_SM4_MAC && 
+				GetMaxSupportedFeatureLevel(CurrentPlatform) >= GetMaxSupportedFeatureLevel(TargetPlatform);
+		}
+		else
+		{
+			bFeatureLevelCompatible = GetMaxSupportedFeatureLevel(CurrentPlatform) >= GetMaxSupportedFeatureLevel(TargetPlatform);
+		}
+
+		bool bIsTargetD3D = TargetPlatform == SP_PCD3D_SM5 ||
+								TargetPlatform == SP_PCD3D_SM4 ||
+								TargetPlatform == SP_PCD3D_ES2;
+
+		bool bIsCurrentPlatformD3D = CurrentPlatform == SP_PCD3D_SM5 ||
+								CurrentPlatform == SP_PCD3D_SM4 ||
+								CurrentPlatform == SP_PCD3D_ES2;
+
+		bFeatureLevelCompatible = bFeatureLevelCompatible && bIsCurrentPlatformD3D && bIsTargetD3D;
+	}
+
+	return bFeatureLevelCompatible;
+}
 
 void FShaderResource::InitRHI()
 {
 	checkf(Code.Num() > 0, TEXT("FShaderResource::InitRHI was called with empty bytecode, which can happen if the resource is initialized multiple times on platforms with no editor data."));
 
 	// we can't have this called on the wrong platform's shaders
-	if (Target.Platform != GRHIShaderPlatform)
-	{
-		if (FPlatformProperties::RequiresCookedData())
+	if (!ArePlatformsCompatible(GMaxRHIShaderPlatform, (EShaderPlatform)Target.Platform))
+ 	{
+		for (int i = 0; i < EShaderPlatform::SP_NumPlatforms; i++)
 		{
-			UE_LOG(LogShaders, Fatal, TEXT("FShaderResource::InitRHI got platform %s but expected %s"), 
-				*LegacyShaderPlatformToShaderFormat((EShaderPlatform)Target.Platform).ToString(), *LegacyShaderPlatformToShaderFormat(GRHIShaderPlatform).ToString() );
+			for (int j = 0; j < EShaderPlatform::SP_NumPlatforms; j++)
+			{
+				bool bOK = ArePlatformsCompatible((EShaderPlatform)i, (EShaderPlatform)j);
+
+				UE_LOG(LogShaders, Warning, TEXT("platform %s %s %s"),
+					*LegacyShaderPlatformToShaderFormat((EShaderPlatform)i).ToString(), (bOK ? *FString("IS compatible with") : *FString("ISNOT compatible with")), *LegacyShaderPlatformToShaderFormat((EShaderPlatform)j).ToString());
+			}
 		}
-		return;
-	}
+
+ 		if (FPlatformProperties::RequiresCookedData())
+ 		{
+ 			UE_LOG(LogShaders, Fatal, TEXT("FShaderResource::InitRHI got platform %s but expected %s"), 
+ 				*LegacyShaderPlatformToShaderFormat((EShaderPlatform)Target.Platform).ToString(), *LegacyShaderPlatformToShaderFormat(GRHIShaderPlatform).ToString() );
+ 		}
+ 		return;
+ 	}
 
 	INC_DWORD_STAT_BY(STAT_Shaders_NumShadersUsedForRendering, 1);
 	SCOPE_CYCLE_COUNTER(STAT_Shaders_RTShaderLoadTime);
@@ -1087,13 +1130,6 @@ const TArray<FName>& GetTargetShaderFormats()
 	return Results;
 }
 
-
-FName GetRuntimeShaderFormat()
-{
-	static FName RuntimeFormat = LegacyShaderPlatformToShaderFormat(GRHIShaderPlatform);
-	return RuntimeFormat;
-}
-
 void ShaderMapAppendKeyString(EShaderPlatform Platform, FString& KeyString)
 {
 	// Globals that should cause all shaders to recompile when changed must be appended to the key here
diff --git a/Engine/Source/Runtime/ShaderCore/Public/Shader.h b/Engine/Source/Runtime/ShaderCore/Public/Shader.h
--- a/Engine/Source/Runtime/ShaderCore/Public/Shader.h
+++ b/Engine/Source/Runtime/ShaderCore/Public/Shader.h
@@ -1336,8 +1336,5 @@ extern SHADERCORE_API void DispatchIndirectComputeShader(
 /** Returns an array of all target shader formats, possibly from multiple target platforms. */
 extern SHADERCORE_API const TArray<FName>& GetTargetShaderFormats();
 
-/** Returns the runtime shader format. */
-extern SHADERCORE_API FName GetRuntimeShaderFormat();
-
 /** Appends to KeyString for all shaders. */
 extern SHADERCORE_API void ShaderMapAppendKeyString(EShaderPlatform Platform, FString& KeyString);
\ No newline at end of file
diff --git a/Engine/Source/Runtime/SlateRHIRenderer/Private/SlateRHIRenderer.cpp b/Engine/Source/Runtime/SlateRHIRenderer/Private/SlateRHIRenderer.cpp
--- a/Engine/Source/Runtime/SlateRHIRenderer/Private/SlateRHIRenderer.cpp
+++ b/Engine/Source/Runtime/SlateRHIRenderer/Private/SlateRHIRenderer.cpp
@@ -1146,7 +1146,7 @@ void FSlateRHIRenderer::SetColorVisionDeficiencyType( uint32 Type )
 bool FSlateRHIRenderer::AreShadersInitialized() const
 {
 #if WITH_EDITORONLY_DATA
-	return AreGlobalShadersComplete(TEXT("SlateElement"));
+	return IsGlobalShaderMapComplete(TEXT("SlateElement"));
 #else
 	return true;
 #endif
diff --git a/Engine/Source/Runtime/SlateRHIRenderer/Private/SlateRHIRenderingPolicy.cpp b/Engine/Source/Runtime/SlateRHIRenderer/Private/SlateRHIRenderingPolicy.cpp
--- a/Engine/Source/Runtime/SlateRHIRenderer/Private/SlateRHIRenderingPolicy.cpp
+++ b/Engine/Source/Runtime/SlateRHIRenderer/Private/SlateRHIRenderingPolicy.cpp
@@ -283,7 +283,7 @@ void FSlateRHIRenderingPolicy::DrawElements(FRHICommandListImmediate& RHICmdList
 
 	FSceneView* SceneView = NULL;
 
-	TShaderMapRef<FSlateElementVS> VertexShader(GetGlobalShaderMap(GRHIFeatureLevel));
+	TShaderMapRef<FSlateElementVS> VertexShader(GetGlobalShaderMap(GMaxRHIFeatureLevel));
 
 	// Disabled stencil test state
 	FDepthStencilStateRHIRef DSOff = TStaticDepthStencilState<false,CF_Always>::GetRHI();
@@ -466,7 +466,7 @@ void FSlateRHIRenderingPolicy::DrawElements(FRHICommandListImmediate& RHICmdList
 FSlateElementPS* FSlateRHIRenderingPolicy::GetTexturePixelShader( ESlateShader::Type ShaderType, ESlateDrawEffect::Type DrawEffects )
 {
 	FSlateElementPS* PixelShader = NULL;
-	const auto FeatureLevel = GRHIFeatureLevel;
+	const auto FeatureLevel = GMaxRHIFeatureLevel;
 	auto ShaderMap = GetGlobalShaderMap(FeatureLevel);
 
 #if !DEBUG_OVERDRAW
diff --git a/Engine/Source/Runtime/Windows/D3D11RHI/Private/D3D11Device.cpp b/Engine/Source/Runtime/Windows/D3D11RHI/Private/D3D11Device.cpp
--- a/Engine/Source/Runtime/Windows/D3D11RHI/Private/D3D11Device.cpp
+++ b/Engine/Source/Runtime/Windows/D3D11RHI/Private/D3D11Device.cpp
@@ -75,18 +75,18 @@ FD3D11DynamicRHI::FD3D11DynamicRHI(IDXGIFactory* InDXGIFactory,D3D_FEATURE_LEVEL
 	// ES2 feature level emulation in D3D11
 	if (FParse::Param(FCommandLine::Get(), TEXT("FeatureLevelES2")) && !GIsEditor)
 	{
-		GMaxRHIFeatureLevelValue = GRHIFeatureLevelValue = ERHIFeatureLevel::ES2;
-		GRHIShaderPlatformValue = SP_PCD3D_ES2;
+		GMaxRHIFeatureLevelValue = ERHIFeatureLevel::ES2;
+		GMaxRHIShaderPlatformValue = SP_PCD3D_ES2;
 	}
 	else if(FeatureLevel == D3D_FEATURE_LEVEL_11_0)
 	{
-		GMaxRHIFeatureLevelValue = GRHIFeatureLevelValue = ERHIFeatureLevel::SM5;
-		GRHIShaderPlatformValue = SP_PCD3D_SM5;
+		GMaxRHIFeatureLevelValue = ERHIFeatureLevel::SM5;
+		GMaxRHIShaderPlatformValue = SP_PCD3D_SM5;
 	}
 	else if(FeatureLevel == D3D_FEATURE_LEVEL_10_0)
 	{
-		GMaxRHIFeatureLevelValue = GRHIFeatureLevelValue = ERHIFeatureLevel::SM4;
-		GRHIShaderPlatformValue = SP_PCD3D_SM4;
+		GMaxRHIFeatureLevelValue = ERHIFeatureLevel::SM4;
+		GMaxRHIShaderPlatformValue = SP_PCD3D_SM4;
 	}
 
 	GPixelCenterOffset = 0.0f;	// Note that in D3D11, there is no half-texel offset (ala DX9)	
diff --git a/Engine/Source/Runtime/Windows/D3D11RHI/Private/WinRT/WinRTD3D11Device.cpp b/Engine/Source/Runtime/Windows/D3D11RHI/Private/WinRT/WinRTD3D11Device.cpp
--- a/Engine/Source/Runtime/Windows/D3D11RHI/Private/WinRT/WinRTD3D11Device.cpp
+++ b/Engine/Source/Runtime/Windows/D3D11RHI/Private/WinRT/WinRTD3D11Device.cpp
@@ -14,7 +14,7 @@ FDynamicRHI* FD3D11DynamicRHIModule::CreateRHI()
 	if (DynamicRHI)
 	{
 		// Initialize the RHI capabilities.
-		GRHIShaderPlatform = SP_PCD3D_SM5;
+		GMaxRHIShaderPlatformValue = SP_PCD3D_SM5;
 
 		// This will be set in the D3D11DynamicRHI::Init call!
 		//	GIsRHIInitialized = true;
