This was hiding in my default changelist missed from submitted changelist 2296372,
#TTP 342661 - UDN: LIVE: BP: Inactive/disconnected Timelines still take resources in editor
#Branch UE4
#Proj KismetCompiler

#Change bValidatedAsWired property initialisation missed from previous commit

diff --git a/Engine/Source/Runtime/Engine/Private/TimelineTemplate.cpp b/Engine/Source/Runtime/Engine/Private/TimelineTemplate.cpp
--- a/Engine/Source/Runtime/Engine/Private/TimelineTemplate.cpp
+++ b/Engine/Source/Runtime/Engine/Private/TimelineTemplate.cpp
@@ -35,6 +35,7 @@ UTimelineTemplate::UTimelineTemplate(const class FPostConstructInitializePropert
 	TimelineLength = 5.0f;
 	TimelineGuid = FGuid::NewGuid();
 	bReplicated = false;
+	bValidatedAsWired = false;
 }
 
 FName UTimelineTemplate::GetDirectionPropertyName() const
