Fixed ANavigationData::TickActor not calling Super::TickActor #UE4

diff --git a/Engine/Source/Runtime/Engine/Private/AI/Navigation/NavigationData.cpp b/Engine/Source/Runtime/Engine/Private/AI/Navigation/NavigationData.cpp
--- a/Engine/Source/Runtime/Engine/Private/AI/Navigation/NavigationData.cpp
+++ b/Engine/Source/Runtime/Engine/Private/AI/Navigation/NavigationData.cpp
@@ -183,6 +183,8 @@ void ANavigationData::PostLoad()
 
 void ANavigationData::TickActor(float DeltaTime, enum ELevelTick TickType, FActorTickFunction& ThisTickFunction)
 {
+	Super::TickActor(DeltaTime, TickType, ThisTickFunction);
+
 	PurgeUnusedPaths();
 	if (NextObservedPathsTickInSeconds >= 0.f)
 	{
