Added ENGINE_API to CurveFloat class

diff --git a/Engine/Source/Runtime/Engine/Classes/Curves/CurveFloat.h b/Engine/Source/Runtime/Engine/Classes/Curves/CurveFloat.h
--- a/Engine/Source/Runtime/Engine/Classes/Curves/CurveFloat.h
+++ b/Engine/Source/Runtime/Engine/Classes/Curves/CurveFloat.h
@@ -6,7 +6,7 @@
 #include "CurveFloat.generated.h"
 
 USTRUCT()
-struct FRuntimeFloatCurve
+struct ENGINE_API FRuntimeFloatCurve
 {
 	GENERATED_USTRUCT_BODY()
 
