FMath::LineBoxIntersection() parameter simplification

GitHub pull request #430

https://github.com/EpicGames/UnrealEngine/pull/430/

Added overloaded verison of LineBoxIntersection that calculates Direction.Reciprocal() itself.

diff --git a/Engine/Source/Developer/MeshUtilities/Private/MeshUtilities.cpp b/Engine/Source/Developer/MeshUtilities/Private/MeshUtilities.cpp
--- a/Engine/Source/Developer/MeshUtilities/Private/MeshUtilities.cpp
+++ b/Engine/Source/Developer/MeshUtilities/Private/MeshUtilities.cpp
@@ -284,7 +284,7 @@ void FMeshDistanceFieldAsyncTask::DoWork()
 				const FVector RayDirection = (*SampleDirections)[SampleIndex];
 				const FVector ReflectedRayDirection = RayDirection;
 
-				if (FMath::LineBoxIntersection(VolumeBounds, VoxelPosition, VoxelPosition + ReflectedRayDirection * VolumeMaxDistance, ReflectedRayDirection, FVector(1.0f) / ReflectedRayDirection))
+				if (FMath::LineBoxIntersection(VolumeBounds, VoxelPosition, VoxelPosition + ReflectedRayDirection * VolumeMaxDistance, ReflectedRayDirection))
 				{
 					FkHitResult Result;
 
diff --git a/Engine/Source/Runtime/Core/Public/Math/UnrealMath.h b/Engine/Source/Runtime/Core/Public/Math/UnrealMath.h
--- a/Engine/Source/Runtime/Core/Public/Math/UnrealMath.h
+++ b/Engine/Source/Runtime/Core/Public/Math/UnrealMath.h
@@ -177,6 +177,17 @@ inline bool FMath::PointBoxIntersection
 		return 0;
 }
 
+inline bool FMath::LineBoxIntersection
+	( 
+	const FBox& Box, 
+	const FVector& Start, 
+	const FVector& End, 
+	const FVector& Direction
+	)
+{
+	return LineBoxIntersection(Box, Start, End, Direction, Direction.Reciprocal());
+}
+
 inline bool FMath::LineBoxIntersection
 	(
 	const FBox&		Box,
diff --git a/Engine/Source/Runtime/Core/Public/Math/UnrealMathUtility.h b/Engine/Source/Runtime/Core/Public/Math/UnrealMathUtility.h
--- a/Engine/Source/Runtime/Core/Public/Math/UnrealMathUtility.h
+++ b/Engine/Source/Runtime/Core/Public/Math/UnrealMathUtility.h
@@ -770,6 +770,9 @@ struct FMath : public FPlatformMath
 	static bool PointBoxIntersection( const FVector& Point, const FBox& Box );
 
 	/** Determines whether a line intersects a box. */
+	static bool LineBoxIntersection( const FBox& Box, const FVector& Start, const FVector& End, const FVector& Direction );
+
+	/** Determines whether a line intersects a box. This overload avoids the need to do the reciprocal every time. */
 	static bool LineBoxIntersection( const FBox& Box, const FVector& Start, const FVector& End, const FVector& Direction, const FVector& OneOverDirection );
 
 	/* Swept-Box vs Box test */
diff --git a/Engine/Source/Runtime/Engine/Private/AI/Navigation/NavigationPath.cpp b/Engine/Source/Runtime/Engine/Private/AI/Navigation/NavigationPath.cpp
--- a/Engine/Source/Runtime/Engine/Private/AI/Navigation/NavigationPath.cpp
+++ b/Engine/Source/Runtime/Engine/Private/AI/Navigation/NavigationPath.cpp
@@ -243,7 +243,7 @@ bool FNavigationPath::DoesIntersectBox(const FBox& Box, uint32 StartingIndex, in
 			{
 				const FVector Direction = (End - Start);
 
-				if (FMath::LineBoxIntersection(Box, Start, End, Direction, Direction.Reciprocal()))
+				if (FMath::LineBoxIntersection(Box, Start, End, Direction))
 				{
 					bIntersects = true;
 					break;
@@ -862,7 +862,7 @@ bool FNavMeshPath::DoesIntersectBox(const FBox& Box, uint32 StartingIndex, int32
 			if (FVector::DistSquared(Start, End) > SMALL_NUMBER)
 			{
 				const FVector Direction = (End - Start);
-				bIntersects = FMath::LineBoxIntersection(Box, Start, End, Direction, Direction.Reciprocal());
+				bIntersects = FMath::LineBoxIntersection(Box, Start, End, Direction);
 			}
 
 			Start = End;
@@ -876,7 +876,7 @@ bool FNavMeshPath::DoesIntersectBox(const FBox& Box, uint32 StartingIndex, int32
 			if (FVector::DistSquared(Start, End) > SMALL_NUMBER)
 			{
 				const FVector Direction = (End - Start);
-				bIntersects = FMath::LineBoxIntersection(Box, Start, End, Direction, Direction.Reciprocal());
+				bIntersects = FMath::LineBoxIntersection(Box, Start, End, Direction);
 			}
 		}
 	}
