Prevent CallHandleDisconnectForFailure from crashing on travel or network failure

#ttp 342945 - Assert in CallHandleDisconnectForFailure

#branch UE4

#change Unable to repro synced to CL#2297935, but modified checks to ensures to catch if it does happen again in future, but allow the user to continue

diff --git a/Engine/Source/Runtime/Engine/Private/UnrealEngine.cpp b/Engine/Source/Runtime/Engine/Private/UnrealEngine.cpp
--- a/Engine/Source/Runtime/Engine/Private/UnrealEngine.cpp
+++ b/Engine/Source/Runtime/Engine/Private/UnrealEngine.cpp
@@ -7426,8 +7426,10 @@ static inline void CallHandleDisconnectForFailure(UWorld* InWorld, UNetDriver* N
 	if (InWorld)
 	{
 		ULocalPlayer* const LP = GEngine->GetFirstGamePlayer(InWorld);
-		check(LP);
-		LP->HandleDisconnect(InWorld, NetDriver);
+		if(ensure(LP))
+		{
+			LP->HandleDisconnect(InWorld, NetDriver);
+		}
 	}
 	else if(NetDriver && NetDriver->NetDriverName == NAME_PendingNetDriver)
 	{
@@ -7436,8 +7438,10 @@ static inline void CallHandleDisconnectForFailure(UWorld* InWorld, UNetDriver* N
 		check(Context.OwningGameInstance != NULL && Context.OwningGameInstance->GetFirstGamePlayer() != NULL);
 
 		ULocalPlayer* const LP = Context.OwningGameInstance->GetFirstGamePlayer();
-		check(LP);
-		LP->HandleDisconnect(InWorld, NetDriver);
+		if(ensure(LP))
+		{
+			LP->HandleDisconnect(InWorld, NetDriver);
+		}
 	}
 	else
 	{
