Rearrange and make advanced several properties of WorldSettings to clean up UI

diff --git a/Engine/Source/Runtime/Engine/Classes/GameFramework/WorldSettings.h b/Engine/Source/Runtime/Engine/Classes/GameFramework/WorldSettings.h
--- a/Engine/Source/Runtime/Engine/Classes/GameFramework/WorldSettings.h
+++ b/Engine/Source/Runtime/Engine/Classes/GameFramework/WorldSettings.h
@@ -183,7 +183,7 @@ struct ENGINE_API FNetViewer
 /**
  * Actor containing all script accessible world properties.
  */
-UCLASS(config=game, showcategories=(Rendering), hidecategories=(Actor, Advanced, Display, Events, Object, Attachment, Info, Input, Blueprint, Layers, Tags, Replication), showcategories=("Input|MouseInput", "Input|TouchInput"), notplaceable)
+UCLASS(config=game, hidecategories=(Actor, Advanced, Display, Events, Object, Attachment, Info, Input, Blueprint, Layers, Tags, Replication), showcategories=("Input|MouseInput", "Input|TouchInput"), notplaceable)
 class ENGINE_API AWorldSettings : public AInfo, public IInterface_AssetUserData
 {
 	GENERATED_UCLASS_BODY()
@@ -191,11 +191,11 @@ class ENGINE_API AWorldSettings : public AInfo, public IInterface_AssetUserData
 	/** DEFAULT BASIC PHYSICS SETTINGS **/
 
 	/** If true, enables CheckStillInWorld checks */
-	UPROPERTY(EditAnywhere, BlueprintReadOnly, Category=Config)
+	UPROPERTY(EditAnywhere, BlueprintReadOnly, Category=World, AdvancedDisplay)
 	uint32 bEnableWorldBoundsChecks:1;
 
 	/** if set to false navigation system will not get created (and all navigation functionality won't be accesible)*/
-	UPROPERTY(EditAnywhere, BlueprintReadOnly, Category=Config)
+	UPROPERTY(EditAnywhere, BlueprintReadOnly, Category=World, AdvancedDisplay)
 	uint32 bEnableNavigationSystem:1;
 
 	/** 
@@ -203,11 +203,11 @@ class ENGINE_API AWorldSettings : public AInfo, public IInterface_AssetUserData
 	 * Level has to be saved to enable this option.
 	 * Warning: Enabling this option will remove all streaming levels from your persistent level.
 	 */
-	UPROPERTY(EditAnywhere, BlueprintReadOnly, Category=Config)
+	UPROPERTY(EditAnywhere, BlueprintReadOnly, Category=World)
 	uint32 bEnableWorldComposition:1;
 
 	/** World origin will shift to a camera position when camera goes far away from current origin */
-	UPROPERTY(EditAnywhere, BlueprintReadOnly, Category=Config, meta=(editcondition = "bEnableWorldComposition"))
+	UPROPERTY(EditAnywhere, BlueprintReadOnly, Category=World, AdvancedDisplay, meta=(editcondition = "bEnableWorldComposition"))
 	uint32 bEnableWorldOriginRebasing:1;
 		
 	/** if set to true, when we call GetGravityZ we assume WorldGravityZ has already been initialized and skip the lookup of DefaultGravityZ and GlobalGravityZ */
@@ -218,12 +218,16 @@ class ENGINE_API AWorldSettings : public AInfo, public IInterface_AssetUserData
 	UPROPERTY(EditAnywhere, BlueprintReadOnly, meta=(DisplayName = "Override World Gravity"), Category = Physics)
 	uint32 bGlobalGravitySet:1;
 
+	/** title of the map displayed in the UI */
+	UPROPERTY(EditAnywhere, BlueprintReadOnly, localized, Category=World)
+	FString Title;
+
 	// any actor falling below this level gets destroyed
-	UPROPERTY(EditAnywhere, BlueprintReadOnly, Category=Damage, meta=(editcondition = "bEnableWorldBoundsChecks"))
+	UPROPERTY(EditAnywhere, BlueprintReadOnly, Category=World, meta=(editcondition = "bEnableWorldBoundsChecks"))
 	float KillZ;
 
 	// The type of damage inflicted when a actor falls below KillZ
-	UPROPERTY(EditAnywhere, BlueprintReadOnly, Category=Damage, meta=(editcondition = "bEnableWorldBoundsChecks"))
+	UPROPERTY(EditAnywhere, BlueprintReadOnly, Category=World, AdvancedDisplay, meta=(editcondition = "bEnableWorldBoundsChecks"))
 	TSubclassOf<UDamageType> KillZDamageType;
 
 	// current gravity actually being used
@@ -235,7 +239,7 @@ class ENGINE_API AWorldSettings : public AInfo, public IInterface_AssetUserData
 	float GlobalGravityZ;
 
 	// optional level specific collision handler
-	UPROPERTY(EditAnywhere, BlueprintReadOnly, Category=Physics)
+	UPROPERTY(EditAnywhere, BlueprintReadOnly, Category=Physics, AdvancedDisplay)
 	TSubclassOf<class UPhysicsCollisionHandler>	PhysicsCollisionHandlerClass;
 
 	/************************************/
@@ -258,7 +262,7 @@ class ENGINE_API AWorldSettings : public AInfo, public IInterface_AssetUserData
 	
 	/** RENDERING SETTINGS **/
 	/** Maximum size of textures for packed light and shadow maps */
-	UPROPERTY(EditAnywhere, Category=Rendering)
+	UPROPERTY(EditAnywhere, Category=Lightmass, AdvancedDisplay)
 	int32 PackedLightAndShadowMapTextureSize;
 
 	/** 
@@ -266,11 +270,11 @@ class ENGINE_API AWorldSettings : public AInfo, public IInterface_AssetUserData
 	 * This is useful when you need to reduce draw calls but can reduce texture streaming efficiency and effective lightmap resolution.
 	 * Note - changes require a rebuild to propagate.  Also, be sure to select all surfaces and make sure they all have the same flags to minimize section count.
 	 */
-	UPROPERTY(EditAnywhere, Category=Rendering)
+	UPROPERTY(EditAnywhere, Category=World, AdvancedDisplay)
 	uint32 bMinimizeBSPSections:1;
 
 	/** Default color scale for the level */
-	UPROPERTY(EditAnywhere, BlueprintReadOnly, Category=Rendering)
+	UPROPERTY(EditAnywhere, BlueprintReadOnly, Category=World, AdvancedDisplay)
 	FVector DefaultColorScale;
 
 	/************************************/
@@ -286,21 +290,21 @@ class ENGINE_API AWorldSettings : public AInfo, public IInterface_AssetUserData
 	/** 
 	 * Whether to place visibility cells only along camera tracks or only above shadow casting surfaces.
 	 */
-	UPROPERTY(EditAnywhere, Category=PrecomputedVisibility)
+	UPROPERTY(EditAnywhere, Category=PrecomputedVisibility, AdvancedDisplay)
 	uint32 bPlaceCellsOnlyAlongCameraTracks:1;
 
 	/** 
 	 * World space size of precomputed visibility cells in x and y.
 	 * Smaller sizes produce more effective occlusion culling at the cost of increased runtime memory usage and lighting build times.
 	 */
-	UPROPERTY(EditAnywhere, Category=PrecomputedVisibility)
+	UPROPERTY(EditAnywhere, Category=PrecomputedVisibility, AdvancedDisplay)
 	int32 VisibilityCellSize;
 
 	/** 
 	 * Determines how aggressive precomputed visibility should be.
 	 * More aggressive settings cull more objects but also cause more visibility errors like popping.
 	 */
-	UPROPERTY(EditAnywhere, Category=PrecomputedVisibility)
+	UPROPERTY(EditAnywhere, Category=PrecomputedVisibility, AdvancedDisplay)
 	TEnumAsByte<enum EVisibilityAggressiveness> VisibilityAggressiveness;
 
 	/************************************/
@@ -312,7 +316,7 @@ class ENGINE_API AWorldSettings : public AInfo, public IInterface_AssetUserData
 	 * This is useful for improving iteration in levels with fully dynamic lighting and shadowing.
 	 * Note that any lighting and shadowing interactions that are usually precomputed will be lost if this is enabled.
 	 */
-	UPROPERTY(EditAnywhere, Category=Lightmass)
+	UPROPERTY(EditAnywhere, Category=Lightmass, AdvancedDisplay)
 	uint32 bForceNoPrecomputedLighting:1;
 
 	UPROPERTY(EditAnywhere, Category=Lightmass)
@@ -339,12 +343,8 @@ class ENGINE_API AWorldSettings : public AInfo, public IInterface_AssetUserData
 	/************************************/
 	/** DEFAULT SETTINGS **/
 
-	/** title of the map displayed in the UI */
-	UPROPERTY(EditAnywhere, BlueprintReadOnly, localized, Category=WorldSettings)
-	FString Title;
-
 	/** scale of 1uu to 1m in real world measurements, for HMD and other physically tracked devices (e.g. 1uu = 1cm would be 100.0) */
-	UPROPERTY(EditAnywhere, BlueprintReadOnly, Category=WorldSettings)
+	UPROPERTY(EditAnywhere, BlueprintReadOnly, Category=VR)
 	float WorldToMeters;
 
 #if WITH_EDITORONLY_DATA
