UE4 - Fixes upside down flip on OpenGL

diff --git a/Engine/Source/Runtime/Engine/Private/Slate/DebugCanvas.cpp b/Engine/Source/Runtime/Engine/Private/Slate/DebugCanvas.cpp
--- a/Engine/Source/Runtime/Engine/Private/Slate/DebugCanvas.cpp
+++ b/Engine/Source/Runtime/Engine/Private/Slate/DebugCanvas.cpp
@@ -150,8 +150,8 @@ void FDebugCanvasDrawer::InitDebugCanvas(UWorld* InWorld)
 {
 	// If the canvas is not null there is more than one viewport draw call before slate draws.  This can happen on resizes. 
 	// We need to delete the old canvas
-        // This can also happen if we are debugging a HUD blueprint and in that case we need to continue using
-        // the same canvas
+		// This can also happen if we are debugging a HUD blueprint and in that case we need to continue using
+		// the same canvas
 	if (FSlateApplication::Get().IsNormalExecution())
 	{
 		if( GameThreadCanvas != NULL )
@@ -171,6 +171,9 @@ void FDebugCanvasDrawer::DrawRenderThread(FRHICommandListImmediate& RHICmdList,
 	{
 		FTexture2DRHIRef& RT = *(FTexture2DRHIRef*)InWindowBackBuffer;
 		RenderTarget->SetRenderTargetTexture( RT );
+		bool bNeedToFlipVertical = RenderThreadCanvas->Canvas.GetAllowSwitchVerticalAxis();
+		// Do not flip when rendering to the back buffer
+		RenderThreadCanvas->Canvas.SetAllowSwitchVerticalAxis(false);
 		if (RenderThreadCanvas->Canvas.IsScaledToRenderTarget() && IsValidRef(RT)) 
 		{
 			RenderThreadCanvas->Canvas.SetRenderTargetRect( FIntRect(0, 0, RT->GetSizeX(), RT->GetSizeY()) );
@@ -181,6 +184,7 @@ void FDebugCanvasDrawer::DrawRenderThread(FRHICommandListImmediate& RHICmdList,
 			RenderThreadCanvas->Canvas.Flush_RenderThread(RHICmdList, true);
 		}
 		RenderThreadCanvas->Canvas.Flush_RenderThread(RHICmdList, true);
+		RenderThreadCanvas->Canvas.SetAllowSwitchVerticalAxis(bNeedToFlipVertical);
 		RenderTarget->ClearRenderTargetTexture();
 	}
 }
diff --git a/Engine/Source/Runtime/Engine/Private/UnrealClient.cpp b/Engine/Source/Runtime/Engine/Private/UnrealClient.cpp
--- a/Engine/Source/Runtime/Engine/Private/UnrealClient.cpp
+++ b/Engine/Source/Runtime/Engine/Private/UnrealClient.cpp
@@ -1077,7 +1077,6 @@ void FViewport::Draw( bool bShouldPresent /*= true */)
 					// Make sure the Canvas is not rendered upside down
 					Canvas.SetAllowSwitchVerticalAxis(false);
 					ViewportClient->Draw(this, &Canvas);
-					Canvas.SetAllowSwitchVerticalAxis(true);
 				}
 				Canvas.Flush_GameThread();
 				ViewportClient->ProcessScreenShots(this);
diff --git a/Engine/Source/Runtime/Engine/Private/UserInterface/Canvas.cpp b/Engine/Source/Runtime/Engine/Private/UserInterface/Canvas.cpp
--- a/Engine/Source/Runtime/Engine/Private/UserInterface/Canvas.cpp
+++ b/Engine/Source/Runtime/Engine/Private/UserInterface/Canvas.cpp
@@ -35,6 +35,7 @@ FCanvas::FCanvas(FRenderTarget* InRenderTarget, FHitProxyConsumer* InHitProxyCon
 ,	CurrentWorldTime(0)
 ,	CurrentDeltaWorldTime(0)
 ,	FeatureLevel(InFeatureLevel)
+,	ShaderPlatform(GRHIShaderPlatform)
 {
 	Construct();
 
@@ -56,6 +57,7 @@ FCanvas::FCanvas(FRenderTarget* InRenderTarget,FHitProxyConsumer* InHitProxyCons
 ,	CurrentWorldTime(InWorldTime)
 ,	CurrentDeltaWorldTime(InWorldDeltaTime)
 ,	FeatureLevel(InFeatureLevel)
+,	ShaderPlatform(GRHIShaderPlatform)
 {
 	Construct();
 }
@@ -174,8 +176,7 @@ bool FCanvasBatchedElementRenderItem::Render_RenderThread(FRHICommandListImmedia
 			Gamma = 1.0f;
 		}
 
-		const bool bNeedsToSwitchVerticalAxis = RHINeedsToSwitchVerticalAxis(GRHIShaderPlatform) || !Canvas->GetAllowSwitchVerticalAxis();
-
+		bool bNeedsToSwitchVerticalAxis = RHINeedsToSwitchVerticalAxis(Canvas->GetShaderPlatform()) && !Canvas->GetAllowSwitchVerticalAxis();
 
 		// draw batched items
 		Data->BatchedElements.Draw(
@@ -205,7 +206,7 @@ bool FCanvasBatchedElementRenderItem::Render_RenderThread(FRHICommandListImmedia
 bool FCanvasBatchedElementRenderItem::Render_GameThread(const FCanvas* Canvas)
 {	
 	checkSlow(Data);
-	bool bDirty=false;		
+	bool bDirty=false;
 	if( Data->BatchedElements.HasPrimsToDraw() )
 	{
 		bDirty = true;
@@ -218,7 +219,7 @@ bool FCanvasBatchedElementRenderItem::Render_GameThread(const FCanvas* Canvas)
 			Gamma = 1.0f;
 		}
 
-		const bool bNeedsToSwitchVerticalAxis = RHINeedsToSwitchVerticalAxis(Canvas->GetShaderPlatform()) || !Canvas->GetAllowSwitchVerticalAxis();
+		bool bNeedsToSwitchVerticalAxis = RHINeedsToSwitchVerticalAxis(Canvas->GetShaderPlatform()) && !Canvas->GetAllowSwitchVerticalAxis();
 
 		// Render the batched elements.
 		struct FBatchedDrawParameters
@@ -311,7 +312,7 @@ bool FCanvasTileRendererItem::Render_RenderThread(FRHICommandListImmediate& RHIC
 	ViewInitOptions.BackgroundColor = FLinearColor::Black;
 	ViewInitOptions.OverlayColor = FLinearColor::White;
 
-	const bool bNeedsToSwitchVerticalAxis = RHINeedsToSwitchVerticalAxis(Canvas->GetShaderPlatform()) || !Canvas->GetAllowSwitchVerticalAxis();
+	bool bNeedsToSwitchVerticalAxis = RHINeedsToSwitchVerticalAxis(Canvas->GetShaderPlatform()) && !Canvas->GetAllowSwitchVerticalAxis();
 
 	FSceneView* View = new FSceneView(ViewInitOptions);
 
@@ -378,7 +379,7 @@ bool FCanvasTileRendererItem::Render_GameThread(const FCanvas* Canvas)
 
 	FSceneView* View = new FSceneView(ViewInitOptions);
 
-	const bool bNeedsToSwitchVerticalAxis = RHINeedsToSwitchVerticalAxis(Canvas->GetShaderPlatform()) || !Canvas->GetAllowSwitchVerticalAxis();
+	bool bNeedsToSwitchVerticalAxis = RHINeedsToSwitchVerticalAxis(Canvas->GetShaderPlatform()) && !Canvas->GetAllowSwitchVerticalAxis();
 	struct FDrawTileParameters
 	{
 		FSceneView* View;
