Fix non-monolithic game builds running into link errors

diff --git a/Engine/Source/Runtime/Engine/Classes/Engine/BlueprintCore.h b/Engine/Source/Runtime/Engine/Classes/Engine/BlueprintCore.h
--- a/Engine/Source/Runtime/Engine/Classes/Engine/BlueprintCore.h
+++ b/Engine/Source/Runtime/Engine/Classes/Engine/BlueprintCore.h
@@ -44,8 +44,10 @@ public:
 
 	virtual void Serialize( FArchive& Ar ) override;
 
+#if WITH_EDITORONLY_DATA
 	/** Gets asset registry tags */
 	virtual void GetAssetRegistryTags(TArray<FAssetRegistryTag>& OutTags) const override;
+#endif
 
 	/** Generates a new blueprint Guid, used when creating new blueprints */
 	void GenerateNewGuid()
diff --git a/Engine/Source/Runtime/Engine/Private/Blueprint.cpp b/Engine/Source/Runtime/Engine/Private/Blueprint.cpp
--- a/Engine/Source/Runtime/Engine/Private/Blueprint.cpp
+++ b/Engine/Source/Runtime/Engine/Private/Blueprint.cpp
@@ -232,6 +232,7 @@ void UBlueprintCore::Serialize(FArchive& Ar)
 	}
 }
 
+#if WITH_EDITORONLY_DATA
 void UBlueprintCore::GetAssetRegistryTags(TArray<FAssetRegistryTag>& OutTags) const
 {
 	Super::GetAssetRegistryTags(OutTags);
@@ -248,6 +249,7 @@ void UBlueprintCore::GetAssetRegistryTags(TArray<FAssetRegistryTag>& OutTags) co
 
 	OutTags.Add( FAssetRegistryTag("GeneratedClass", GeneratedClassVal, FAssetRegistryTag::TT_Hidden) );
 }
+#endif
 
 void UBlueprintCore::GenerateDeterministicGuid()
 {
