Merged from John Barrett: Fixed stat net ping display, and fixed minor divide-by-zero, prior to first ping.

diff --git a/Engine/Source/Runtime/Engine/Private/NetworkDriver.cpp b/Engine/Source/Runtime/Engine/Private/NetworkDriver.cpp
--- a/Engine/Source/Runtime/Engine/Private/NetworkDriver.cpp
+++ b/Engine/Source/Runtime/Engine/Private/NetworkDriver.cpp
@@ -204,7 +204,7 @@ void UNetDriver::TickFlush(float DeltaSeconds)
 			
 			if (ServerConnection != NULL && ServerConnection->PlayerController != NULL && ServerConnection->PlayerController->PlayerState != NULL)
 			{
-				Ping = FMath::TruncToInt(1000.0f * ServerConnection->PlayerController->PlayerState->ExactPing);
+				Ping = FMath::TruncToInt(ServerConnection->PlayerController->PlayerState->ExactPing);
 			}
 
 			if (ServerConnection != NULL)
diff --git a/Engine/Source/Runtime/Engine/Private/PlayerState.cpp b/Engine/Source/Runtime/Engine/Private/PlayerState.cpp
--- a/Engine/Source/Runtime/Engine/Private/PlayerState.cpp
+++ b/Engine/Source/Runtime/Engine/Private/PlayerState.cpp
@@ -68,7 +68,7 @@ void APlayerState::RecalculateAvgPing()
 	}
 
 	// Calculate the average, and divide it by 4 to optimize replication
-	ExactPing = (float)Sum / (float)Count;
+	ExactPing = (Count > 0 ? ((float)Sum / (float)Count) : 0.f);
 	Ping = FMath::Min(255, (int32)(ExactPing * 0.25f));
 }
 
