#ue4 - Framework header comment cleanup.

diff --git a/Engine/Source/Runtime/AIModule/Classes/AIController.h b/Engine/Source/Runtime/AIModule/Classes/AIController.h
--- a/Engine/Source/Runtime/AIModule/Classes/AIController.h
+++ b/Engine/Source/Runtime/AIModule/Classes/AIController.h
@@ -64,7 +64,6 @@ struct FFocusKnowledge
  *
  * @see https://docs.unrealengine.com/latest/INT/Gameplay/Framework/Controller/
  */
-//=============================================================================
 
 UCLASS(BlueprintType, Blueprintable)
 class AIMODULE_API AAIController : public AController, public IAIPerceptionListenerInterface
diff --git a/Engine/Source/Runtime/Engine/Classes/GameFramework/CharacterMovementComponent.h b/Engine/Source/Runtime/Engine/Classes/GameFramework/CharacterMovementComponent.h
--- a/Engine/Source/Runtime/Engine/Classes/GameFramework/CharacterMovementComponent.h
+++ b/Engine/Source/Runtime/Engine/Classes/GameFramework/CharacterMovementComponent.h
@@ -120,7 +120,6 @@ struct FCharacterMovementComponentPreClothTickFunction : public FTickFunction
  * @see ACharacter, UPawnMovementComponent
  * @see https://docs.unrealengine.com/latest/INT/Gameplay/Framework/Pawn/Character/
  */
-//=============================================================================
 
 UCLASS()
 class ENGINE_API UCharacterMovementComponent : public UPawnMovementComponent, public INetworkPredictionInterface
@@ -1311,7 +1310,6 @@ protected:
 	 * @param OutFloorResult:	Result of the floor check.
 	 * @param SweepRadius:		The radius to use for sweep tests. Should be <= capsule radius.
 	 * @param DownwardSweepResult:	If non-null and it contains valid blocking hit info, this will be used as the result of a downward sweep test instead of doing it as part of the update.
-	 *
 	 */
 	virtual void ComputeFloorDist(const FVector& CapsuleLocation, float LineDistance, float SweepDistance, FFindFloorResult& OutFloorResult, float SweepRadius, const FHitResult* DownwardSweepResult = NULL) const;
 
diff --git a/Engine/Source/Runtime/Engine/Classes/GameFramework/Controller.h b/Engine/Source/Runtime/Engine/Classes/GameFramework/Controller.h
--- a/Engine/Source/Runtime/Engine/Classes/GameFramework/Controller.h
+++ b/Engine/Source/Runtime/Engine/Classes/GameFramework/Controller.h
@@ -28,7 +28,6 @@ DECLARE_DYNAMIC_MULTICAST_DELEGATE_FourParams( FInstigatedAnyDamageSignature, fl
  * 
  * @see https://docs.unrealengine.com/latest/INT/Gameplay/Framework/Controller/
  */
-//=============================================================================
 
 UCLASS(abstract, notplaceable, NotBlueprintable, HideCategories=(Collision,Rendering,"Utilities|Transformation")) 
 class ENGINE_API AController : public AActor, public INavAgentInterface
diff --git a/Engine/Source/Runtime/Engine/Classes/GameFramework/DefaultPawn.h b/Engine/Source/Runtime/Engine/Classes/GameFramework/DefaultPawn.h
--- a/Engine/Source/Runtime/Engine/Classes/GameFramework/DefaultPawn.h
+++ b/Engine/Source/Runtime/Engine/Classes/GameFramework/DefaultPawn.h
@@ -14,7 +14,8 @@ class USphereComponent;
 class UStaticMeshComponent;
 
 /**
- * 
+ * DefaultPawn implements a simple Pawn with spherical collision and built-in flying movement.
+ * @see UFloatingPawnMovement
  */
 UCLASS(config=Game, Blueprintable, BlueprintType)
 class ENGINE_API ADefaultPawn : public APawn
@@ -27,15 +28,27 @@ class ENGINE_API ADefaultPawn : public APawn
 	virtual void UpdateNavigationRelevance() override;
 	// End Pawn overrides
 
-	/** Input callback to move forward. */
+	/**
+	 * Input callback to move forward in local space (or backward if Val is negative).
+	 * @param Val Amount of movement in the forward direction (or backward if negative).
+	 * @see APawn::AddMovementInput()
+	 */
 	UFUNCTION(BlueprintCallable, Category="Pawn")
 	virtual void MoveForward(float Val);
 
-	/** Input callback to strafe right. */
+	/**
+	 * Input callback to strafe right in local space (or left if Val is negative).
+	 * @param Val Amount of movement in the right direction (or left if negative).
+	 * @see APawn::AddMovementInput()
+	 */
 	UFUNCTION(BlueprintCallable, Category="Pawn")
 	virtual void MoveRight(float Val);
 
-	/** Input callback to move up in world space. */
+	/**
+	 * Input callback to move up in world space (or down if Val is negative).
+	 * @param Val Amount of movement in the world up direction (or down if negative).
+	 * @see APawn::AddMovementInput()
+	 */
 	UFUNCTION(BlueprintCallable, Category="Pawn")
 	virtual void MoveUp_World(float Val);
 
@@ -47,7 +60,7 @@ class ENGINE_API ADefaultPawn : public APawn
 	void TurnAtRate(float Rate);
 
 	/**
-	* Called via input to lookup at a given rate.
+	* Called via input to look up at a given rate (or down if Rate is negative).
 	* @param Rate	This is a normalized rate, i.e. 1.0 means 100% of desired turn rate
 	*/
 	UFUNCTION(BlueprintCallable, Category="Pawn")
@@ -92,7 +105,7 @@ public:
 	UPROPERTY(Category=Pawn, VisibleAnywhere, BlueprintReadOnly)
 	TSubobjectPtr<UStaticMeshComponent> MeshComponent;
 
-	/** If true, adds default bindings for movement and camera look. */
+	/** If true, adds default input bindings for movement and camera look. */
 	UPROPERTY(Category=Pawn, EditAnywhere, BlueprintReadOnly)
 	uint32 bAddDefaultMovementBindings:1;
 };
diff --git a/Engine/Source/Runtime/Engine/Classes/GameFramework/FloatingPawnMovement.h b/Engine/Source/Runtime/Engine/Classes/GameFramework/FloatingPawnMovement.h
--- a/Engine/Source/Runtime/Engine/Classes/GameFramework/FloatingPawnMovement.h
+++ b/Engine/Source/Runtime/Engine/Classes/GameFramework/FloatingPawnMovement.h
@@ -9,7 +9,8 @@
 #include "FloatingPawnMovement.generated.h"
 
 /**
- * 
+ * FloatingPawnMovement is a movement component that provides simple movement for any Pawn class.
+ * Limits on speed and acceleration are provided, while gravity is not implemented.
  */
 UCLASS()
 class ENGINE_API UFloatingPawnMovement : public UPawnMovementComponent
diff --git a/Engine/Source/Runtime/Engine/Classes/GameFramework/MovementComponent.h b/Engine/Source/Runtime/Engine/Classes/GameFramework/MovementComponent.h
--- a/Engine/Source/Runtime/Engine/Classes/GameFramework/MovementComponent.h
+++ b/Engine/Source/Runtime/Engine/Classes/GameFramework/MovementComponent.h
@@ -21,6 +21,7 @@ struct FCollisionShape;
  *    - Utility functions for special handling of collision results (SlideAlongSurface(), ComputeSlideVector(), TwoWallAdjust()).
  *    - Utility functions for moving when there may be initial penetration (SafeMoveUpdatedComponent(), ResolvePenetration()).
  *    - Automatically registering the component tick and finding a component to move on the owning Actor.
+ * Normally the root component of the owning actor is moved, however another component may be selected (see SetUpdatedComponent()).
  */
 UCLASS(ClassGroup=Movement, abstract, BlueprintType)
 class ENGINE_API UMovementComponent : public UActorComponent
diff --git a/Engine/Source/Runtime/Engine/Classes/GameFramework/PlayerController.h b/Engine/Source/Runtime/Engine/Classes/GameFramework/PlayerController.h
--- a/Engine/Source/Runtime/Engine/Classes/GameFramework/PlayerController.h
+++ b/Engine/Source/Runtime/Engine/Classes/GameFramework/PlayerController.h
@@ -68,7 +68,6 @@ struct FDynamicForceFeedbackDetails
  *
  * @see https://docs.unrealengine.com/latest/INT/Gameplay/Framework/Controller/PlayerController/
  */
-//=============================================================================
 
 UCLASS(config=Game, BlueprintType, Blueprintable)
 class ENGINE_API APlayerController : public AController
diff --git a/Engine/Source/Runtime/Engine/Classes/GameFramework/ProjectileMovementComponent.h b/Engine/Source/Runtime/Engine/Classes/GameFramework/ProjectileMovementComponent.h
--- a/Engine/Source/Runtime/Engine/Classes/GameFramework/ProjectileMovementComponent.h
+++ b/Engine/Source/Runtime/Engine/Classes/GameFramework/ProjectileMovementComponent.h
@@ -5,10 +5,14 @@
 #include "ProjectileMovementComponent.generated.h"
 
 /**
- * ProjectileMovementComponent updates the position of the associated UpdatedComponent during its tick.
+ * ProjectileMovementComponent updates the position of another component during its tick.
+ *
  * Behavior such as bouncing after impacts and homing toward a target are supported.
- * If the updated component is simulating physics, only the initial launch parameters (when initial velocity is non-zero) will affect the projectile,
- * and the physics sim will take over from there.
+ *
+ * Normally the root component of the owning actor is moved, however another component may be selected (see SetUpdatedComponent()).
+ * If the updated component is simulating physics, only the initial launch parameters (when initial velocity is non-zero)
+ * will affect the projectile, and the physics sim will take over from there.
+ *
  * @see UMovementComponent
  */
 UCLASS(ClassGroup=Movement, meta=(BlueprintSpawnableComponent), ShowCategories=(Velocity))
