#ue4 - Clean up some comments in CharacterMovementComponent.

diff --git a/Engine/Source/Runtime/Engine/Classes/GameFramework/CharacterMovementComponent.h b/Engine/Source/Runtime/Engine/Classes/GameFramework/CharacterMovementComponent.h
--- a/Engine/Source/Runtime/Engine/Classes/GameFramework/CharacterMovementComponent.h
+++ b/Engine/Source/Runtime/Engine/Classes/GameFramework/CharacterMovementComponent.h
@@ -237,15 +237,24 @@ public:
 	UPROPERTY(Category="Character Movement", EditAnywhere, BlueprintReadWrite, meta=(ClampMin="0", UIMin="0"))
 	float BrakingDecelerationFlying;
 
-	/** When falling, amount of lateral movement control available to the character. 0 = no control, 1 = full control at max speed of MaxWalkSpeed. */
+	/**
+	 * When falling, amount of lateral movement control available to the character.
+	 * 0 = no control, 1 = full control at max speed of MaxWalkSpeed.
+	 */
 	UPROPERTY(Category="Character Movement", EditAnywhere, BlueprintReadWrite, meta=(ClampMin="0", UIMin="0"))
 	float AirControl;
 
-	/** When falling, multiplier applied to AirControl when lateral velocity is less than AirControlBoostVelocityThreshold. Setting this to zero will disable air control boosting. */
+	/**
+	 * When falling, multiplier applied to AirControl when lateral velocity is less than AirControlBoostVelocityThreshold.
+	 * Setting this to zero will disable air control boosting. Final result is clamped at 1.
+	 */
 	UPROPERTY(Category="Character Movement", EditAnywhere, BlueprintReadWrite, meta=(ClampMin="0", UIMin="0"))
 	float AirControlBoostMultiplier;
 
-	/** When falling, if lateral velocity magnitude is less than this value, AirControl is multiplied by AirControlBoostMultiplier. Setting this to zero will disable air control  boosting. */
+	/**
+	 * When falling, if lateral velocity magnitude is less than this value, AirControl is multiplied by AirControlBoostMultiplier.
+	 * Setting this to zero will disable air control boosting.
+	 */
 	UPROPERTY(Category="Character Movement", EditAnywhere, BlueprintReadWrite, meta=(ClampMin="0", UIMin="0"))
 	float AirControlBoostVelocityThreshold;
 
