#ue4 - Improved comments for ProjectileMovement homing projectiles.

diff --git a/Engine/Source/Runtime/Engine/Classes/GameFramework/ProjectileMovementComponent.h b/Engine/Source/Runtime/Engine/Classes/GameFramework/ProjectileMovementComponent.h
--- a/Engine/Source/Runtime/Engine/Classes/GameFramework/ProjectileMovementComponent.h
+++ b/Engine/Source/Runtime/Engine/Classes/GameFramework/ProjectileMovementComponent.h
@@ -56,17 +56,17 @@ class ENGINE_API UProjectileMovementComponent : public UMovementComponent
 	uint32 bForceSubStepping:1;
 
 	/**
-	 * If true, we will accelerate towards our homing target.
+	 * If true, we will accelerate toward our homing target. HomingTargetComponent must be set after the projectile is spawned.
 	 * @see HomingTargetComponent, HomingAccelerationMagnitude
 	 */
-	UPROPERTY(EditAnywhere, BlueprintReadWrite, Category = Homing)
+	UPROPERTY(EditAnywhere, BlueprintReadWrite, Category=Homing)
 	uint32 bIsHomingProjectile:1;
 
-	/** Custom gravity scale for this projectile.  Set to 0 for no gravity. */
+	/** Custom gravity scale for this projectile. Set to 0 for no gravity. */
 	UPROPERTY(EditAnywhere, BlueprintReadWrite, Category=Projectile)
 	float ProjectileGravityScale;
 
-	/** Buoyancy of UpdatedComponent in fluid.  0.0=sinks as fast as in air, 1.0=neutral buoyancy */
+	/** Buoyancy of UpdatedComponent in fluid. 0.0=sinks as fast as in air, 1.0=neutral buoyancy */
 	UPROPERTY()
 	float Buoyancy;
 
@@ -91,14 +91,14 @@ class ENGINE_API UProjectileMovementComponent : public UMovementComponent
 	FOnProjectileStopDelegate OnProjectileStop;
 
 	/** The magnitude of our acceleration towards the homing target. Overall velocity magnitude will still be limited by MaxSpeed. */
-	UPROPERTY(EditAnywhere, BlueprintReadWrite, Category = Homing)
+	UPROPERTY(EditAnywhere, BlueprintReadWrite, Category=Homing)
 	float HomingAccelerationMagnitude;
 
 	/**
 	 * The current target we are homing towards. Can only be set at runtime (when projectile is spawned or updating).
 	 * @see bIsHomingProjectile
 	 */
-	UPROPERTY(VisibleInstanceOnly, BlueprintReadWrite, Category = Homing)
+	UPROPERTY(VisibleInstanceOnly, BlueprintReadWrite, Category=Homing)
 	TWeakObjectPtr<USceneComponent> HomingTargetComponent;
 
 	/** Sets the velocity to the new value, rotated into Actor space. */
