#ue4 - Improved comments for ProjectileMovementComponent

diff --git a/Engine/Source/Runtime/Engine/Classes/GameFramework/ProjectileMovementComponent.h b/Engine/Source/Runtime/Engine/Classes/GameFramework/ProjectileMovementComponent.h
--- a/Engine/Source/Runtime/Engine/Classes/GameFramework/ProjectileMovementComponent.h
+++ b/Engine/Source/Runtime/Engine/Classes/GameFramework/ProjectileMovementComponent.h
@@ -31,11 +31,11 @@ class ENGINE_API UProjectileMovementComponent : public UMovementComponent
 	UPROPERTY(EditAnywhere, BlueprintReadWrite, Category=Projectile)
 	float MaxSpeed;
 
-	/** If true, this projectile will have its rotation updated each frame to match its velocity */
+	/** If true, this projectile will have its rotation updated each frame to match the direction of its velocity. */
 	UPROPERTY(EditAnywhere, BlueprintReadWrite, Category=Projectile)
 	uint32 bRotationFollowsVelocity:1;
 
-	/** If true, simple bounces will be simulated.  Set to false to implement custom contact behavior (via HitActor) or to stop simulating on contact. */
+	/** If true, simple bounces will be simulated. Set this to false to stop simulating on contact. */
 	UPROPERTY(EditAnywhere, BlueprintReadWrite, Category=ProjectileBounces)
 	uint32 bShouldBounce:1;
 
@@ -70,15 +70,24 @@ class ENGINE_API UProjectileMovementComponent : public UMovementComponent
 	UPROPERTY()
 	float Buoyancy;
 
-	/** Percentage of velocity maintained after the bounce in the direction of the normal of impact (coefficient of restitution). 1.0 = no velocity lost, 0.0 = no bounce. Ignored if bShouldBounce=false. */
+	/**
+	 * Percentage of velocity maintained after the bounce in the direction of the normal of impact (coefficient of restitution).
+	 * 1.0 = no velocity lost, 0.0 = no bounce. Ignored if bShouldBounce is false.
+	 */
 	UPROPERTY(EditAnywhere, BlueprintReadWrite, Category=ProjectileBounces)
 	float Bounciness;
 
-	/** Percentage of velocity maintained after the bounce in the direction tangent to the normal of impact (coefficient of friction). 0.0 = no velocity lost, 1.0 = no sliding. Ignored if bShouldBounce=false. */
+	/**
+	 * Percentage of velocity maintained after the bounce in the direction tangent to the normal of impact (coefficient of friction).
+	 * 0.0 = no velocity lost, 1.0 = no sliding. Ignored if bShouldBounce is false.
+	 */
 	UPROPERTY(EditAnywhere, BlueprintReadWrite, Category=ProjectileBounces)
 	float Friction;
 
-	/** If velocity is below this threshold after a bounce, stops simulating and triggers the OnProjectileStop event. Ignored if bShouldBounce=false, in which case the projectile stops simulating on the first impact. */
+	/**
+	 * If velocity is below this threshold after a bounce, stops simulating and triggers the OnProjectileStop event.
+	 * Ignored if bShouldBounce is false, in which case the projectile stops simulating on the first impact.
+	 */
 	UPROPERTY(EditAnywhere, BlueprintReadWrite, Category=ProjectileBounces)
 	float BounceVelocityStopSimulatingThreshold;
 
