Added a way to select single EQS query in LogVisualizer tool, to show selected EQS data in a world. #ue4

diff --git a/Engine/Source/Developer/LogVisualizer/Private/SLogVisualizer.cpp b/Engine/Source/Developer/LogVisualizer/Private/SLogVisualizer.cpp
--- a/Engine/Source/Developer/LogVisualizer/Private/SLogVisualizer.cpp
+++ b/Engine/Source/Developer/LogVisualizer/Private/SLogVisualizer.cpp
@@ -514,7 +514,8 @@ void SLogVisualizer::Construct(const FArguments& InArgs, FLogVisualizer* InLogVi
 									SAssignNew(LogsLinesWidget, SListView<TSharedPtr<FLogEntryItem> >)
 									.ItemHeight(20)
 									.ListItemsSource(&LogEntryLines)
-									.SelectionMode(ESelectionMode::Multi)
+									.SelectionMode(ESelectionMode::Single)
+									.OnSelectionChanged(this, &SLogVisualizer::LogEntryLineSelectionChanged)
 									.OnGenerateRow(this, &SLogVisualizer::LogEntryLinesGenerateRow)
 								]
 							]
@@ -1510,12 +1511,27 @@ TSharedRef<ITableRow> SLogVisualizer::LogEntryLinesGenerateRow(TSharedPtr<FLogEn
 			]			
 		];
 }
-/*
-void SLogVisualizer::LogEntryLineSelectionChanged(TSharedPtr<FLogsListItem> SelectedItem, ESelectInfo::Type SelectInfo)
-{
 
+void SLogVisualizer::LogEntryLineSelectionChanged(TSharedPtr<FLogEntryItem> SelectedItem, ESelectInfo::Type SelectInfo)
+{
+	TMap<FName, FVisualLogExtensionInterface*>& AllExtensions = FVisualLog::Get().GetAllExtensions();
+	for (auto Iterator = AllExtensions.CreateIterator(); Iterator; ++Iterator)
+	{
+		FVisualLogExtensionInterface* Extension = (*Iterator).Value;
+		if (Extension != NULL)
+		{
+			if (SelectedItem.IsValid() == true)
+			{
+				Extension->LogEntryLineSelectionChanged(SelectedItem, SelectedItem->UserData, SelectedItem->TagName);
+			}
+			else
+			{
+				Extension->LogEntryLineSelectionChanged(SelectedItem, 0, NAME_None);
+			}
+		}
+	}
 }
-*/
+
 bool SLogVisualizer::ShouldListLog(const TSharedPtr<FActorsVisLog>& Log)
 {
 	//// Check log name filter
@@ -1742,6 +1758,8 @@ void SLogVisualizer::ShowEntry(const FVisLogEntry* LogEntry)
 
 			EntryItem.Verbosity = LogLine->Verbosity;
 			EntryItem.Line = LogLine->Line;
+			EntryItem.UserData = LogLine->UserData;
+			EntryItem.TagName = LogLine->TagName;
 
 			LogEntryLines.Add(MakeShareable(new FLogEntryItem(EntryItem)));
 		}
diff --git a/Engine/Source/Developer/LogVisualizer/Private/SLogVisualizer.h b/Engine/Source/Developer/LogVisualizer/Private/SLogVisualizer.h
--- a/Engine/Source/Developer/LogVisualizer/Private/SLogVisualizer.h
+++ b/Engine/Source/Developer/LogVisualizer/Private/SLogVisualizer.h
@@ -32,6 +32,8 @@ struct FLogEntryItem
 	FLinearColor CategoryColor;
 	ELogVerbosity::Type Verbosity;
 	FString Line;
+	int64 UserData;
+	FName TagName;
 };
 
 namespace ELogsSortMode
@@ -106,7 +108,8 @@ public:
 	TSharedRef<ITableRow> LogsListGenerateRow(TSharedPtr<FLogsListItem> InItem, const TSharedRef<STableViewBase>& OwnerTable);
 	void LogsListSelectionChanged(TSharedPtr<FLogsListItem> SelectedItem, ESelectInfo::Type SelectInfo);
 	TSharedRef<ITableRow> LogEntryLinesGenerateRow(TSharedPtr<FLogEntryItem> Item, const TSharedRef<STableViewBase>& OwnerTable);
-
+	void LogEntryLineSelectionChanged(TSharedPtr<FLogEntryItem> SelectedItem, ESelectInfo::Type SelectInfo);
+	
 	FOnZoomChanged& OnZoomChanged() { return ZoomChangedNotify; }
 
 	FOnHistogramWindowChanged& OnHistogramWindowChanged() { return HistogramWindowChangedNotify; }
diff --git a/Engine/Source/Runtime/AIModule/Classes/EnvironmentQuery/EnvQueryDebugHelpers.h b/Engine/Source/Runtime/AIModule/Classes/EnvironmentQuery/EnvQueryDebugHelpers.h
--- a/Engine/Source/Runtime/AIModule/Classes/EnvironmentQuery/EnvQueryDebugHelpers.h
+++ b/Engine/Source/Runtime/AIModule/Classes/EnvironmentQuery/EnvQueryDebugHelpers.h
@@ -149,8 +149,8 @@ FArchive& operator<<(FArchive& Ar, EQSDebug::FQueryData& Data)
 		ensure(OwnerActor != NULL); \
 		if (OwnerActor) \
 		{  \
-			UE_VLOG(OwnerActor, CategoryName, Verbosity, TEXT("Executed EQS: \n - Name: '%s' (id=%d, option=%d),\n - All Items: %d,\n - ValidItems: %d"), *Query->QueryName, Query->QueryID, Query->OptionIndex, Query->ItemDetails.Num(), Query->NumValidItems); \
-			UEnvQueryDebugHelpers::LogQuery(OwnerActor, Query, CategoryName.GetCategoryName(), ELogVerbosity::Verbosity); \
+			UEnvQueryDebugHelpers::LogQuery(OwnerActor, Query, CategoryName.GetCategoryName(), ELogVerbosity::Verbosity, \
+				FString::Printf(TEXT("Executed EQS: \n - Name: '%s' (id=%d, option=%d),\n - All Items: %d,\n - ValidItems: %d"), *Query->QueryName, Query->QueryID, Query->OptionIndex, Query->ItemDetails.Num(), Query->NumValidItems)); \
 		} \
 	} \
 }
@@ -170,6 +170,6 @@ class AIMODULE_API UEnvQueryDebugHelpers : public UObject
 	static void QueryToDebugData(struct FEnvQueryInstance* Query, EQSDebug::FQueryData& EQSLocalData);
 	static void QueryToBlobArray(struct FEnvQueryInstance* Query, TArray<uint8>& BlobArray, bool bUseCompression = false);
 	static void BlobArrayToDebugData(const TArray<uint8>& BlobArray, EQSDebug::FQueryData& EQSLocalData, bool bUseCompression = false);
-	static void LogQuery(const class AActor* LogOwnerActor, struct FEnvQueryInstance* Query, const FName& CategoryName, ELogVerbosity::Type);
+	static void LogQuery(const class AActor* LogOwnerActor, struct FEnvQueryInstance* Query, const FName& CategoryName, ELogVerbosity::Type, const FString& AdditionalLogInfo);
 #endif
 };
diff --git a/Engine/Source/Runtime/AIModule/Classes/VisualLoggerExtension.h b/Engine/Source/Runtime/AIModule/Classes/VisualLoggerExtension.h
--- a/Engine/Source/Runtime/AIModule/Classes/VisualLoggerExtension.h
+++ b/Engine/Source/Runtime/AIModule/Classes/VisualLoggerExtension.h
@@ -15,15 +15,17 @@ class FVisualLoggerExtension : public FVisualLogExtensionInterface
 {
 public:
 	FVisualLoggerExtension();
-	virtual void OnTimestampChange(float Timestamp, class UWorld* InWorld, class AActor* HelperActor);
-	virtual void DrawData(class UWorld* InWorld, class UCanvas* Canvas, class AActor* HelperActor, const FName& TagName, const struct FVisLogEntry::FDataBlock& DataBlock, float Timestamp);
-	virtual void DisableDrawingForData(class UWorld* InWorld, class UCanvas* Canvas, class AActor* HelperActor, const FName& TagName, const FVisLogEntry::FDataBlock& DataBlock, float Timestamp);
+	virtual void OnTimestampChange(float Timestamp, class UWorld* InWorld, class AActor* HelperActor) override;
+	virtual void DrawData(class UWorld* InWorld, class UCanvas* Canvas, class AActor* HelperActor, const FName& TagName, const struct FVisLogEntry::FDataBlock& DataBlock, float Timestamp) override;
+	virtual void DisableDrawingForData(class UWorld* InWorld, class UCanvas* Canvas, class AActor* HelperActor, const FName& TagName, const FVisLogEntry::FDataBlock& DataBlock, float Timestamp) override;
+	virtual void LogEntryLineSelectionChanged(TSharedPtr<struct FLogEntryItem> SelectedItem, int64 UserData, FName TagName) override;
 
 private:
 	void DisableEQSRendering(class AActor* HelperActor);
 
 protected:
 	uint32 CachedEQSId;
+	uint32 SelectedEQSId;
 	float CurrentTimestamp;
 };
 #endif //ENABLE_VISUAL_LOG
diff --git a/Engine/Source/Runtime/AIModule/Private/EnvironmentQuery/EnvQueryDebugHelpers.cpp b/Engine/Source/Runtime/AIModule/Private/EnvironmentQuery/EnvQueryDebugHelpers.cpp
--- a/Engine/Source/Runtime/AIModule/Private/EnvironmentQuery/EnvQueryDebugHelpers.cpp
+++ b/Engine/Source/Runtime/AIModule/Private/EnvironmentQuery/EnvQueryDebugHelpers.cpp
@@ -129,11 +129,16 @@ void  UEnvQueryDebugHelpers::BlobArrayToDebugData(const TArray<uint8>& BlobArray
 	}
 }
 
-void UEnvQueryDebugHelpers::LogQuery(const class AActor* LogOwnerActor, struct FEnvQueryInstance* Query, const FName& CategoryName, ELogVerbosity::Type Type)
+void UEnvQueryDebugHelpers::LogQuery(const class AActor* LogOwnerActor, struct FEnvQueryInstance* Query, const FName& CategoryName, ELogVerbosity::Type Type, const FString& AdditionalLogInfo)
 {
 	TArray<uint8> BlobArray;
 	UEnvQueryDebugHelpers::QueryToBlobArray(Query, BlobArray);
-	FVisualLog::Get().GetEntryToWrite(LogOwnerActor)->AddDataBlock(EVisLogTags::TAG_EQS, BlobArray, CategoryName);
+	FVisLogEntry* EntryToWrite = FVisualLog::Get().GetEntryToWrite(LogOwnerActor);
+	if (EntryToWrite)
+	{
+		FVisualLog::Get().LogLine(LogOwnerActor, CategoryName, Type, AdditionalLogInfo, Query->QueryID, *EVisLogTags::TAG_EQS);
+		EntryToWrite->AddDataBlock(EVisLogTags::TAG_EQS, BlobArray, CategoryName);
+	}
 }
 
 #endif //ENABLE_VISUAL_LOG
diff --git a/Engine/Source/Runtime/AIModule/Private/VisualLoggerExtension.cpp b/Engine/Source/Runtime/AIModule/Private/VisualLoggerExtension.cpp
--- a/Engine/Source/Runtime/AIModule/Private/VisualLoggerExtension.cpp
+++ b/Engine/Source/Runtime/AIModule/Private/VisualLoggerExtension.cpp
@@ -7,6 +7,7 @@
 FVisualLoggerExtension::FVisualLoggerExtension()
 	:  CachedEQSId(INDEX_NONE)
 	, CurrentTimestamp(FLT_MIN)
+	, SelectedEQSId(INDEX_NONE)
 {
 
 }
@@ -38,6 +39,7 @@ void FVisualLoggerExtension::DisableEQSRendering(class AActor* HelperActor)
 #if USE_EQS_DEBUGGER
 	if (HelperActor)
 	{
+		SelectedEQSId = INDEX_NONE;
 		UEQSRenderingComponent* EQSRenderComp = HelperActor->FindComponentByClass<UEQSRenderingComponent>();
 		if (EQSRenderComp)
 		{
@@ -50,6 +52,14 @@ void FVisualLoggerExtension::DisableEQSRendering(class AActor* HelperActor)
 #endif
 }
 
+void FVisualLoggerExtension::LogEntryLineSelectionChanged(TSharedPtr<struct FLogEntryItem> SelectedItem, int64 UserData, FName TagName)
+{
+	if (TagName == *EVisLogTags::TAG_EQS)
+	{
+		SelectedEQSId = (int32)UserData;
+	}
+}
+
 void FVisualLoggerExtension::DrawData(class UWorld* InWorld, class UCanvas* Canvas, class AActor* HelperActor, const FName& TagName, const FVisLogEntry::FDataBlock& DataBlock, float Timestamp)
 {
 #if USE_EQS_DEBUGGER
@@ -68,11 +78,18 @@ void FVisualLoggerExtension::DrawData(class UWorld* InWorld, class UCanvas* Canv
 		UEnvQueryDebugHelpers::BlobArrayToDebugData(DataBlock.Data, DebugData, false);
 		if (DebugData.Id != CachedEQSId || (EQSRenderComp && EQSRenderComp->bHiddenInGame))
 		{
-			CachedEQSId = DebugData.Id;
-			EQSRenderComp->DebugData = DebugData;
-			EQSRenderComp->Activate();
-			EQSRenderComp->SetHiddenInGame(false);
-			EQSRenderComp->MarkRenderStateDirty();
+				if (SelectedEQSId == INDEX_NONE)
+				{
+					SelectedEQSId = DebugData.Id;
+				}
+				if (DebugData.Id == SelectedEQSId)
+				{
+					CachedEQSId = DebugData.Id;
+					EQSRenderComp->DebugData = DebugData;
+					EQSRenderComp->Activate();
+					EQSRenderComp->SetHiddenInGame(false);
+					EQSRenderComp->MarkRenderStateDirty();
+				}
 		}
 	}
 #endif
diff --git a/Engine/Source/Runtime/Engine/Private/VisualLog.cpp b/Engine/Source/Runtime/Engine/Private/VisualLog.cpp
--- a/Engine/Source/Runtime/Engine/Private/VisualLog.cpp
+++ b/Engine/Source/Runtime/Engine/Private/VisualLog.cpp
@@ -86,6 +86,9 @@ FVisLogEntry::FVisLogEntry(TSharedPtr<FJsonValue> FromJson)
 				LogLines[LogLineIndex].Category = FName(*(JsonLogLine->GetStringField(VisualLogJson::TAG_CATEGORY)));
 				LogLines[LogLineIndex].Verbosity = TEnumAsByte<ELogVerbosity::Type>((uint8)FMath::TruncToInt(JsonLogLine->GetNumberField(VisualLogJson::TAG_VERBOSITY)));
 				LogLines[LogLineIndex].Line = JsonLogLine->GetStringField(VisualLogJson::TAG_LINE);
+				LogLines[LogLineIndex].TagName = FName(*(JsonLogLine->GetStringField(VisualLogJson::TAG_TAGNAME)));
+				LogLines[LogLineIndex].UserData = (int64)(JsonLogLine->GetNumberField(VisualLogJson::TAG_USERDATA));
+
 			}
 		}
 	}
@@ -155,7 +158,7 @@ FVisLogEntry::FVisLogEntry(TSharedPtr<FJsonValue> FromJson)
 			if (JsonSampleObject.IsValid())
 			{
 				FDataBlock& Sample = DataBlocks[DataBlocks.Add(FDataBlock())];
-				Sample.TagName = FName(*(JsonSampleObject->GetStringField(VisualLogJson::TAG_DATABLOCK_TAGNAME)));
+				Sample.TagName = FName(*(JsonSampleObject->GetStringField(VisualLogJson::TAG_TAGNAME)));
 				Sample.Category = FName(*(JsonSampleObject->GetStringField(VisualLogJson::TAG_CATEGORY)));
 				Sample.Verbosity = TEnumAsByte<ELogVerbosity::Type>((uint8)FMath::TruncToInt(JsonSampleObject->GetNumberField(VisualLogJson::TAG_VERBOSITY)));
 
@@ -220,6 +223,8 @@ TSharedPtr<FJsonValue> FVisLogEntry::ToJson() const
 		JsonLogLineObject->SetStringField(VisualLogJson::TAG_CATEGORY, LogLine->Category.ToString());
 		JsonLogLineObject->SetNumberField(VisualLogJson::TAG_VERBOSITY, LogLine->Verbosity);
 		JsonLogLineObject->SetStringField(VisualLogJson::TAG_LINE, LogLine->Line);
+		JsonLogLineObject->SetStringField(VisualLogJson::TAG_TAGNAME, LogLine->TagName.ToString());
+		JsonLogLineObject->SetNumberField(VisualLogJson::TAG_USERDATA, LogLine->UserData);
 		JsonLines[LogLineIndex] = MakeShareable(new FJsonValueObject(JsonLogLineObject));
 	}
 
@@ -298,7 +303,7 @@ TSharedPtr<FJsonValue> FVisLogEntry::ToJson() const
 		const FString CurrentDataAsString = FBase64::Encode(CompressedData);
 
 		JsonSampleObject->SetStringField(VisualLogJson::TAG_CATEGORY, CurrentData.Category.ToString());
-		JsonSampleObject->SetStringField(VisualLogJson::TAG_DATABLOCK_TAGNAME, CurrentData.TagName.ToString());
+		JsonSampleObject->SetStringField(VisualLogJson::TAG_TAGNAME, CurrentData.TagName.ToString());
 		JsonSampleObject->SetStringField(VisualLogJson::TAG_DATABLOCK_DATA, CurrentDataAsString);
 		JsonSampleObject->SetNumberField(VisualLogJson::TAG_VERBOSITY, CurrentData.Verbosity);
 
@@ -654,7 +659,7 @@ void FVisualLog::Redirect(UObject* Source, const AActor* NewRedirection)
 	}
 }
 
-void FVisualLog::LogLine(const AActor* Actor, const FName& CategoryName, ELogVerbosity::Type Verbosity, const FString& Line)
+void FVisualLog::LogLine(const AActor* Actor, const FName& CategoryName, ELogVerbosity::Type Verbosity, const FString& Line, int64 UserData, FName TagName)
 {
 	if (IsRecording() == false || Actor == NULL || Actor->IsPendingKill() || (IsAllBlocked() && !InWhitelist(CategoryName)))
 	{
@@ -667,7 +672,8 @@ void FVisualLog::LogLine(const AActor* Actor, const FName& CategoryName, ELogVer
 	{
 		// @todo will have to store CategoryName separately, and create a map of names 
 		// used in log to have saved logs independent from FNames index changes
-		Entry->LogLines.Add(FVisLogEntry::FLogLine(CategoryName, Verbosity, Line));
+		int32 LineIndex = Entry->LogLines.Add(FVisLogEntry::FLogLine(CategoryName, Verbosity, Line, UserData));
+		Entry->LogLines[LineIndex].TagName = TagName;
 	}
 }
 #endif // ENABLE_VISUAL_LOG
diff --git a/Engine/Source/Runtime/Engine/Public/VisualLog.h b/Engine/Source/Runtime/Engine/Public/VisualLog.h
--- a/Engine/Source/Runtime/Engine/Public/VisualLog.h
+++ b/Engine/Source/Runtime/Engine/Public/VisualLog.h
@@ -32,9 +32,15 @@ struct ENGINE_API FVisLogEntry
 		FString Line;
 		FName Category;
 		TEnumAsByte<ELogVerbosity::Type> Verbosity;
+		int64 UserData;
+		FName TagName;
 
 		FLogLine(const FName& InCategory, ELogVerbosity::Type InVerbosity, const FString& InLine)
-			: Line(InLine), Category(InCategory), Verbosity(InVerbosity)
+			: Line(InLine), Category(InCategory), Verbosity(InVerbosity), UserData(0)
+		{}
+
+		FLogLine(const FName& InCategory, ELogVerbosity::Type InVerbosity, const FString& InLine, int64 InUserData)
+			: Line(InLine), Category(InCategory), Verbosity(InVerbosity), UserData(InUserData)
 		{}
 	};
 	TArray<FLogLine> LogLines;
@@ -314,6 +320,8 @@ namespace VisualLogJson
 	static const FString TAG_TYPECOLORSIZE = TEXT("TypeColorSize");
 	static const FString TAG_POINTS = TEXT("Points");
 	static const FString TAG_ELEMENTSTODRAW = TEXT("ElementsToDraw");
+	static const FString TAG_TAGNAME = TEXT("DataBlockTagName");
+	static const FString TAG_USERDATA = TEXT("UserData");
 
 	static const FString TAG_HISTOGRAMSAMPLES = TEXT("HistogramSamples");
 	static const FString TAG_HISTOGRAMSAMPLE = TEXT("Sample");
@@ -322,7 +330,6 @@ namespace VisualLogJson
 
 	static const FString TAG_DATABLOCK = TEXT("DataBlock");
 	static const FString TAG_DATABLOCK_DATA = TEXT("DataBlockData");
-	static const FString TAG_DATABLOCK_TAGNAME = TEXT("DataBlockTagName");
 
 	static const FString TAG_LOGS = TEXT("Logs");
 }
@@ -348,6 +355,7 @@ struct FVisualLogExtensionInterface
 	virtual void OnTimestampChange(float Timestamp, class UWorld* InWorld, class AActor* HelperActor) = 0;
 	virtual void DrawData(class UWorld* InWorld, class UCanvas* Canvas, class AActor* HelperActor, const FName& TagName, const FVisLogEntry::FDataBlock& DataBlock, float Timestamp) = 0;
 	virtual void DisableDrawingForData(class UWorld* InWorld, class UCanvas* Canvas, class AActor* HelperActor, const FName& TagName, const FVisLogEntry::FDataBlock& DataBlock, float Timestamp) = 0;
+	virtual void LogEntryLineSelectionChanged(TSharedPtr<struct FLogEntryItem> SelectedItem, int64 UserData, FName TagName) = 0;
 };
 
 
@@ -377,7 +385,7 @@ public:
 
 	void RedirectToVisualLog(const class UObject* Src, const class AActor* Dest);
 
-	void LogLine(const class AActor* Actor, const FName& CategoryName, ELogVerbosity::Type Verbosity, const FString& Line);
+	void LogLine(const class AActor* Actor, const FName& CategoryName, ELogVerbosity::Type Verbosity, const FString& Line, int64 UserData = 0, FName TagName = NAME_Name);
 
 	const FLogsMap* GetLogs() const { return &LogsMap; }
 
