Landscape : Added LightingLODBias to Compopnent to adjust Landscape LOD for Lightmass

diff --git a/Engine/Source/Runtime/Engine/Classes/Landscape/LandscapeComponent.h b/Engine/Source/Runtime/Engine/Classes/Landscape/LandscapeComponent.h
--- a/Engine/Source/Runtime/Engine/Classes/Landscape/LandscapeComponent.h
+++ b/Engine/Source/Runtime/Engine/Classes/Landscape/LandscapeComponent.h
@@ -211,6 +211,10 @@ public:
 	FGuid StateId;
 
 #if WITH_EDITORONLY_DATA
+	/** LOD level Bias to use when lighting buidling via lightmass, -1 Means automatic LOD calculation based on ForcedLOD + LODBias */
+	UPROPERTY(EditAnywhere, Category=LandscapeComponent)
+	int32 LightingLODBias;
+
 	UPROPERTY(Transient, DuplicateTransient)
 	UTexture2D* SelectDataTexture; // Data texture used for selection mask
 
diff --git a/Engine/Source/Runtime/Engine/Private/Landscape/Landscape.cpp b/Engine/Source/Runtime/Engine/Private/Landscape/Landscape.cpp
--- a/Engine/Source/Runtime/Engine/Private/Landscape/Landscape.cpp
+++ b/Engine/Source/Runtime/Engine/Private/Landscape/Landscape.cpp
@@ -85,6 +85,9 @@ ULandscapeComponent::ULandscapeComponent(const class FPostConstructInitializePro
 	NeighborLODBias[5] = 128;
 	NeighborLODBias[6] = 128;
 	NeighborLODBias[7] = 128;
+#if WITH_EDITORONLY_DATA
+	LightingLODBias = -1; // -1 Means automatic LOD calculation based on ForcedLOD + LODBias
+#endif
 
 	Mobility = EComponentMobility::Static;
 
diff --git a/Engine/Source/Runtime/Engine/Private/Landscape/LandscapeEdit.cpp b/Engine/Source/Runtime/Engine/Private/Landscape/LandscapeEdit.cpp
--- a/Engine/Source/Runtime/Engine/Private/Landscape/LandscapeEdit.cpp
+++ b/Engine/Source/Runtime/Engine/Private/Landscape/LandscapeEdit.cpp
@@ -3946,6 +3946,13 @@ void ULandscapeComponent::PostEditChangeProperty(FPropertyChangedEvent& Property
 		{
 			StaticLightingResolution = 0;
 		}
+		InvalidateLightingCache();
+	}
+	else if (GIsEditor && PropertyName == FName(TEXT("LightingLODBias")))
+	{
+		int32 MaxLOD = FMath::CeilLogTwo(SubsectionSizeQuads + 1) - 1;
+		LightingLODBias = FMath::Clamp<int32>(LightingLODBias, -1, MaxLOD);
+		InvalidateLightingCache();
 	}
 	else if (GIsEditor && (PropertyName == FName(TEXT("CollisionMipLevel"))))
 	{
diff --git a/Engine/Source/Runtime/Engine/Private/Landscape/LandscapeLight.cpp b/Engine/Source/Runtime/Engine/Private/Landscape/LandscapeLight.cpp
--- a/Engine/Source/Runtime/Engine/Private/Landscape/LandscapeLight.cpp
+++ b/Engine/Source/Runtime/Engine/Private/Landscape/LandscapeLight.cpp
@@ -97,6 +97,22 @@ void FLandscapeStaticLightingTextureMapping::Apply(FQuantizedLightmapData* Quant
 	LandscapeComponent->MarkPackageDirty();
 }
 
+namespace
+{
+	// Calculate Geometric LOD for lighting
+	int32 GetLightingLOD(const ULandscapeComponent* InComponent)
+	{
+		if (InComponent->LightingLODBias < 0)
+		{
+			return FMath::Clamp<int32>(InComponent->ForcedLOD >= 0 ? InComponent->ForcedLOD : InComponent->LODBias, 0, FMath::CeilLogTwo(InComponent->SubsectionSizeQuads + 1) - 1);
+		}
+		else
+		{
+			return InComponent->LightingLODBias;
+		}
+	}
+};
+
 /** Initialization constructor. */
 FLandscapeStaticLightingMesh::FLandscapeStaticLightingMesh(ULandscapeComponent* InComponent, const TArray<ULightComponent*>& InRelevantLights, int32 InExpandQuadsX, int32 InExpandQuadsY, float InLightMapRatio, int32 InLOD)
 	: FStaticLightingMesh(
@@ -125,7 +141,7 @@ FLandscapeStaticLightingMesh::FLandscapeStaticLightingMesh(ULandscapeComponent*
 	UVFactor = LightMapRatio / NumVertices;
 	bReverseWinding = (LocalToWorld.GetDeterminant() < 0.0f);
 
-	int32 GeometricLOD = FMath::Clamp<int32>(InComponent->ForcedLOD >= 0 ? InComponent->ForcedLOD : InComponent->LODBias, 0, FMath::CeilLogTwo(InComponent->SubsectionSizeQuads+1) - 1);
+	int32 GeometricLOD = ::GetLightingLOD(InComponent);
 	GetHeightmapData(InLOD, FMath::Max(GeometricLOD, InLOD));
 }
 
@@ -213,7 +229,7 @@ namespace
 
 				if (Neighbor)
 				{
-					NeighborLOD = FMath::Clamp<int32>(Neighbor->ForcedLOD >= 0 ? Neighbor->ForcedLOD : Neighbor->LODBias, 0, MaxLOD);
+					NeighborLOD = ::GetLightingLOD(Neighbor);
 				}
 				else
 				{
@@ -228,10 +244,10 @@ namespace
 								continue;
 							}
 
-							ULandscapeComponent* Neighbor = Info->XYtoComponentMap.FindRef(ComponentBase + FIntPoint(xx, yy));
+							ULandscapeComponent* Neighbor = Info->XYtoComponentMap.FindRef(ComponentBase + FIntPoint(x+xx, y+yy));
 							if (Neighbor)
 							{
-								NeighborLOD = FMath::Max(FMath::Clamp<int32>(Neighbor->ForcedLOD >= 0 ? Neighbor->ForcedLOD : Neighbor->LODBias, 0, MaxLOD), NeighborLOD);
+								NeighborLOD = FMath::Max(::GetLightingLOD(Neighbor), NeighborLOD);
 							}
 						}
 					}
@@ -453,7 +469,7 @@ void FLandscapeStaticLightingMesh::GetHeightmapData(int32 InLOD, int32 GeometryL
 				// Data array for upscaling case
 				TArray<FColor> CompHeightData;
 				TArray<FColor> CompXYOffsetData;
-				int32 NeighborGeometricLOD = FMath::Clamp<int32>(Neighbor->ForcedLOD >= 0 ? Neighbor->ForcedLOD : Neighbor->LODBias, 0, MaxLOD);
+				int32 NeighborGeometricLOD = ::GetLightingLOD(Neighbor);
 				FLandscapeComponentDataInterface DataInterface(Neighbor, InLOD);
 				::InternalUpscaling(DataInterface, Neighbor, InLOD, NeighborGeometricLOD, CompHeightData, CompXYOffsetData);
 				for (int32 Y = 0; Y < YNum; Y++)
