Made SetSuppressTransitionMessage bp callable

diff --git a/Engine/Source/Runtime/Engine/Classes/Engine/GameViewportClient.h b/Engine/Source/Runtime/Engine/Classes/Engine/GameViewportClient.h
--- a/Engine/Source/Runtime/Engine/Classes/Engine/GameViewportClient.h
+++ b/Engine/Source/Runtime/Engine/Classes/Engine/GameViewportClient.h
@@ -570,9 +570,9 @@ public:
 	FViewportFrame* ViewportFrame;
 
 	/**
- 	 * Controls supression of the blue transition text messages 
+ 	 * Controls suppression of the blue transition text messages 
 	 * 
-	 * @param bSuppress	Pass true to supress messages
+	 * @param bSuppress	Pass true to suppress messages
 	 */
 	void SetSuppressTransitionMessage( bool bSuppress )
 	{
diff --git a/Engine/Source/Runtime/Engine/Classes/Kismet/KismetSystemLibrary.h b/Engine/Source/Runtime/Engine/Classes/Kismet/KismetSystemLibrary.h
--- a/Engine/Source/Runtime/Engine/Classes/Kismet/KismetSystemLibrary.h
+++ b/Engine/Source/Runtime/Engine/Classes/Kismet/KismetSystemLibrary.h
@@ -1163,4 +1163,13 @@ class UKismetSystemLibrary : public UBlueprintFunctionLibrary
 	 */
 	UFUNCTION(BlueprintCallable, Category = "Utilities|Platform")
 	static void ShowPlatformSpecificAchievementsScreen(class APlayerController* SpecificPlayer);
+
+	/** 
+	 * Sets the state of the transition message rendered by the viewport. (The blue text displayed when the game is paused and so forth.)
+	 *
+	 * @param WorldContextObject	World context
+	 * @param State					set true to supress transition message
+	 */
+	UFUNCTION(BlueprintCallable, Category = "Utilities", meta = (HidePin = "WorldContextObject", DefaultToSelf = "WorldContextObject"))
+	static void SetSupressViewportTransitionMessage(UObject* WorldContextObject, bool bState);
 };
diff --git a/Engine/Source/Runtime/Engine/Private/GameViewportClient.cpp b/Engine/Source/Runtime/Engine/Private/GameViewportClient.cpp
--- a/Engine/Source/Runtime/Engine/Private/GameViewportClient.cpp
+++ b/Engine/Source/Runtime/Engine/Private/GameViewportClient.cpp
@@ -117,7 +117,7 @@ UGameViewportClient::UGameViewportClient(const class FPostConstructInitializePro
 	SplitscreenInfo[ESplitScreenType::FourPlayer].PlayerData.Add(FPerPlayerSplitscreenData(0.5f, 0.5f, 0.5f, 0.5f));
 
 	MaxSplitscreenPlayers = 4;
-	bSuppressTransitionMessage = false;
+	bSuppressTransitionMessage = true;
 
 	if (HasAnyFlags(RF_ClassDefaultObject) == false)
 	{
diff --git a/Engine/Source/Runtime/Engine/Private/KismetSystemLibrary.cpp b/Engine/Source/Runtime/Engine/Private/KismetSystemLibrary.cpp
--- a/Engine/Source/Runtime/Engine/Private/KismetSystemLibrary.cpp
+++ b/Engine/Source/Runtime/Engine/Private/KismetSystemLibrary.cpp
@@ -2669,8 +2669,19 @@ void UKismetSystemLibrary::ShowPlatformSpecificAchievementsScreen(class APlayerC
 		ExternalUI->ShowAchievementsUI(LocalUserNum);
 	}
 }
+
 void UKismetSystemLibrary::SetStructurePropertyByName(UObject* Object, FName PropertyName, const FGenericStruct& Value)
 {
 	// We should never hit these!  They're stubs to avoid NoExport on the class.
 	check(0);
-}
\ No newline at end of file
+}
+
+void UKismetSystemLibrary::SetSupressViewportTransitionMessage(UObject* WorldContextObject, bool bState)
+{
+	UWorld* World = GEngine->GetWorldFromContextObject(WorldContextObject);
+	if (World && World->GetFirstLocalPlayerFromController() != nullptr && World->GetFirstLocalPlayerFromController()->ViewportClient != nullptr )
+	{
+		World->GetFirstLocalPlayerFromController()->ViewportClient->SetSuppressTransitionMessage(bState);
+	}
+}
+
