Use optimized actor iterator for GetAllActorsOfClass

diff --git a/Engine/Source/Runtime/Engine/Private/GameplayStatics.cpp b/Engine/Source/Runtime/Engine/Private/GameplayStatics.cpp
--- a/Engine/Source/Runtime/Engine/Private/GameplayStatics.cpp
+++ b/Engine/Source/Runtime/Engine/Private/GameplayStatics.cpp
@@ -452,28 +452,12 @@ void UGameplayStatics::GetAllActorsOfClass(UObject* WorldContextObject, TSubclas
 	// We do nothing if not class provided, rather than giving ALL actors!
 	if(ActorClass != NULL)
 	{
-		// If it's a pawn class, use pawn iterator, that is more efficient
-		if(ActorClass->IsChildOf(APawn::StaticClass()))
+		for(TActorIterator<AActor> It(World, ActorClass); It; ++It)
 		{
-			for(FConstPawnIterator Iterator = World->GetPawnIterator(); Iterator; ++Iterator)
-			{
-				APawn* Pawn = *Iterator;
-				if(Pawn != NULL && !Pawn->IsPendingKill() && Pawn->IsA(ActorClass))
-				{
-					OutActors.Add(Pawn);
-				}
-			}
-		}
-		// Non-pawn, so use actor iterator in the world
-		else
-		{
-			for(FActorIterator It(World); It; ++It)
+			AActor* Actor = *It;
+			if(!Actor->IsPendingKill())
 			{
-				AActor* Actor = *It;
-				if(Actor != NULL && !Actor->IsPendingKill() && Actor->IsA(ActorClass))
-				{
-					OutActors.Add(Actor);
-				}
+				OutActors.Add(Actor);
 			}
 		}
 	}
