Fixed landscape UV mismatch between proxies in world composition

diff --git a/Engine/Source/Runtime/Engine/Private/Landscape/Landscape.cpp b/Engine/Source/Runtime/Engine/Private/Landscape/Landscape.cpp
--- a/Engine/Source/Runtime/Engine/Private/Landscape/Landscape.cpp
+++ b/Engine/Source/Runtime/Engine/Private/Landscape/Landscape.cpp
@@ -1339,6 +1339,7 @@ void ALandscapeProxy::SetAbsoluteSectionBase(FIntPoint InSectionBase)
 		{
 			FIntPoint AbsoluteSectionBase = Comp->GetSectionBase() + Difference;
 			Comp->SetSectionBase(AbsoluteSectionBase);
+			Comp->RecreateRenderState_Concurrent();
 		}
 	}
 
