#ue4 - Camera pitch and roll limits changed from 89.9 to 89.99 to allow nearly-vertical views, and improve behavior when people expect projectiles to spawn with vertical velocities.

diff --git a/Engine/Source/Runtime/Engine/Private/PlayerCameraManager.cpp b/Engine/Source/Runtime/Engine/Private/PlayerCameraManager.cpp
--- a/Engine/Source/Runtime/Engine/Private/PlayerCameraManager.cpp
+++ b/Engine/Source/Runtime/Engine/Private/PlayerCameraManager.cpp
@@ -26,12 +26,12 @@ APlayerCameraManager::APlayerCameraManager(const class FPostConstructInitializeP
 	bReplicates = false;
 	FreeCamDistance = 256.0f;
 	bDebugClientSideCamera = false;
-	ViewPitchMin = -89.9f;
-	ViewPitchMax = 89.9f;
+	ViewPitchMin = -89.99f;
+	ViewPitchMax = 89.99f;
 	ViewYawMin = 0.f;
 	ViewYawMax = 359.999f;
-	ViewRollMin = -89.9f;
-	ViewRollMax = 89.9f;
+	ViewRollMin = -89.99f;
+	ViewRollMax = 89.99f;
 	CameraShakeCamModClass = UCameraModifier_CameraShake::StaticClass();
 	bUseClientSideCameraUpdates = true;
 	CameraStyle = NAME_Default;
