#ue4 - CharacterMovementComponent::ApplyRepulsionForce() and ApplyAccumulatedForces() are now virtual functions.

diff --git a/Engine/Source/Runtime/Engine/Classes/GameFramework/CharacterMovementComponent.h b/Engine/Source/Runtime/Engine/Classes/GameFramework/CharacterMovementComponent.h
--- a/Engine/Source/Runtime/Engine/Classes/GameFramework/CharacterMovementComponent.h
+++ b/Engine/Source/Runtime/Engine/Classes/GameFramework/CharacterMovementComponent.h
@@ -951,10 +951,10 @@ public:
 	virtual void SetBase(UPrimitiveComponent* NewBase, const FName BoneName = NAME_None, bool bNotifyActor=true);
 
 	/** Applies repulsion force to all touched components. */
-	void ApplyRepulsionForce(float DeltaTime);
+	virtual void ApplyRepulsionForce(float DeltaSeconds);
 	
 	/** Applies momentum accumulated through AddImpulse() and AddForce(). */
-	void ApplyAccumulatedForces(float DeltaSeconds);	
+	virtual void ApplyAccumulatedForces(float DeltaSeconds);
 
 	/** 
 	 * Handle start swimming functionality
diff --git a/Engine/Source/Runtime/Engine/Private/CharacterMovementComponent.cpp b/Engine/Source/Runtime/Engine/Private/CharacterMovementComponent.cpp
--- a/Engine/Source/Runtime/Engine/Private/CharacterMovementComponent.cpp
+++ b/Engine/Source/Runtime/Engine/Private/CharacterMovementComponent.cpp
@@ -6092,7 +6092,7 @@ void UCharacterMovementComponent::SetAvoidanceEnabled(bool bEnable)
 	}
 }
 
-void UCharacterMovementComponent::ApplyRepulsionForce( float DeltaTime )
+void UCharacterMovementComponent::ApplyRepulsionForce( float DeltaSeconds )
 {
 	FCollisionQueryParams QueryParams;
 	QueryParams.bReturnFaceIndex = false;
@@ -6164,7 +6164,7 @@ void UCharacterMovementComponent::ApplyRepulsionForce( float DeltaTime )
 			}
 			
 			const float DistanceNow = (HitLoc - BodyLocation).SizeSquared2D();
-			const float DistanceLater = (HitLoc - (BodyLocation + BodyVelocity * DeltaTime)).SizeSquared2D();
+			const float DistanceLater = (HitLoc - (BodyLocation + BodyVelocity * DeltaSeconds)).SizeSquared2D();
 
 			if (BodyLocation.SizeSquared() > 0.1f && bHasHit && DistanceNow < StopBodyDistance && !bIsPenetrating)
 			{
