Add GetGlobalTimeDilation blueprint function

diff --git a/Engine/Source/Runtime/Engine/Classes/Kismet/GameplayStatics.h b/Engine/Source/Runtime/Engine/Classes/Kismet/GameplayStatics.h
--- a/Engine/Source/Runtime/Engine/Classes/Kismet/GameplayStatics.h
+++ b/Engine/Source/Runtime/Engine/Classes/Kismet/GameplayStatics.h
@@ -137,6 +137,13 @@ class ENGINE_API UGameplayStatics : public UBlueprintFunctionLibrary
 	UFUNCTION(BlueprintPure, meta=(FriendlyName = "GetClass"), Category="Utilities")
 	static class UClass *GetObjectClass(const UObject *Object);
 
+	/**
+	 * Sets the global time dilation
+	 * @param	TimeDilation	value to set the global time dilation to
+	 */
+	UFUNCTION(BlueprintPure, Category="Utilities|Time", meta=(WorldContext="WorldContextObject") )
+	static float GetGlobalTimeDilation(UObject* WorldContextObject);
+
 	/**
 	 * Sets the global time dilation
 	 * @param	TimeDilation	value to set the global time dilation to
diff --git a/Engine/Source/Runtime/Engine/Private/GameplayStatics.cpp b/Engine/Source/Runtime/Engine/Private/GameplayStatics.cpp
--- a/Engine/Source/Runtime/Engine/Private/GameplayStatics.cpp
+++ b/Engine/Source/Runtime/Engine/Private/GameplayStatics.cpp
@@ -103,6 +103,12 @@ class UClass* UGameplayStatics::GetObjectClass(const UObject* Object)
 	return Object ? Object->GetClass() : NULL;
 }
 
+float UGameplayStatics::GetGlobalTimeDilation(UObject* WorldContextObject)
+{
+	UWorld* const World = GEngine->GetWorldFromContextObject( WorldContextObject );
+	return World->GetWorldSettings()->TimeDilation;
+}
+
 void UGameplayStatics::SetGlobalTimeDilation(UObject* WorldContextObject, float TimeDilation)
 {
 	UWorld* const World = GEngine->GetWorldFromContextObject( WorldContextObject );
