Fix crash in destructible mesh when using negative number for CellSiteCount

TTP#346667

diff --git a/Engine/Source/Runtime/Engine/Classes/Engine/DestructibleFractureSettings.h b/Engine/Source/Runtime/Engine/Classes/Engine/DestructibleFractureSettings.h
--- a/Engine/Source/Runtime/Engine/Classes/Engine/DestructibleFractureSettings.h
+++ b/Engine/Source/Runtime/Engine/Classes/Engine/DestructibleFractureSettings.h
@@ -153,7 +153,7 @@ class UDestructibleFractureSettings : public UObject
 	GENERATED_UCLASS_BODY()
 
 	/** The number of voronoi cell sites. */
-	UPROPERTY(EditAnywhere, Category=Voronoi)
+	UPROPERTY(EditAnywhere, Category = Voronoi, meta = (ClampMin = "1", UIMin = "1"))
 	int32									CellSiteCount;
 
 	/** Stored interior material data.  Just need one as we only support Voronoi splitting. */
