New Niagara effect class, factory, asset type actions and surrounding infrastructure plus a skeleton editor; bit of cleanup and removal of unnecessary opcodes from the VM. Moved particle data into its own class for separate and easier buffer management.

diff --git a/Engine/Source/Developer/AssetTools/Private/AssetTools.cpp b/Engine/Source/Developer/AssetTools/Private/AssetTools.cpp
--- a/Engine/Source/Developer/AssetTools/Private/AssetTools.cpp
+++ b/Engine/Source/Developer/AssetTools/Private/AssetTools.cpp
@@ -97,6 +97,7 @@ FAssetTools::FAssetTools()
 	if ( bEnableNiagara )
 	{
 		RegisterAssetTypeActions( MakeShareable(new FAssetTypeActions_NiagaraScript) );
+		RegisterAssetTypeActions( MakeShareable(new FAssetTypeActions_NiagaraEffect) );
 	}
 }
 
diff --git a/Engine/Source/Developer/AssetTools/Private/AssetToolsPrivatePCH.h b/Engine/Source/Developer/AssetTools/Private/AssetToolsPrivatePCH.h
--- a/Engine/Source/Developer/AssetTools/Private/AssetToolsPrivatePCH.h
+++ b/Engine/Source/Developer/AssetTools/Private/AssetToolsPrivatePCH.h
@@ -56,6 +56,7 @@ DECLARE_LOG_CATEGORY_EXTERN(LogAssetTools, Log, All);
 #include "AssetTypeActions/AssetTypeActions_MaterialParameterCollection.h"
 #include "AssetTypeActions/AssetTypeActions_MorphTarget.h"
 #include "AssetTypeActions/AssetTypeActions_NiagaraScript.h"
+#include "AssetTypeActions/AssetTypeActions_NiagaraEffect.h"
 #include "AssetTypeActions/AssetTypeActions_ObjectLibrary.h"
 #include "AssetTypeActions/AssetTypeActions_ParticleSystem.h"
 #include "AssetTypeActions/AssetTypeActions_PhysicalMaterial.h"
diff --git a/Engine/Source/Developer/AssetTools/Private/AssetTypeActions/AssetTypeActions_NiagaraEffect.cpp b/Engine/Source/Developer/AssetTools/Private/AssetTypeActions/AssetTypeActions_NiagaraEffect.cpp
new file mode 100644
--- /dev/null
+++ b/Engine/Source/Developer/AssetTools/Private/AssetTypeActions/AssetTypeActions_NiagaraEffect.cpp
@@ -0,0 +1,61 @@
+// Copyright 1998-2014 Epic Games, Inc. All Rights Reserved.
+
+#include "AssetToolsPrivatePCH.h"
+#include "Engine/NiagaraEffect.h"
+
+#include "Toolkits/AssetEditorManager.h"
+
+#include "Editor/NiagaraEditor/Public/NiagaraEditorModule.h"
+#include "Editor/NiagaraEditor/Public/INiagaraEffectEditor.h"
+
+#define LOCTEXT_NAMESPACE "AssetTypeActions"
+
+void FAssetTypeActions_NiagaraEffect::GetActions(const TArray<UObject*>& InObjects, FMenuBuilder& MenuBuilder)
+{
+	auto Scripts = GetTypedWeakObjectPtrs<UNiagaraEffect>(InObjects);
+
+	MenuBuilder.AddMenuEntry(
+		LOCTEXT("NiagaraEffect_Edit", "Edit"),
+		LOCTEXT("NiagaraEffect_EditTooltip", "Opens the selected effect in editor."),
+		FSlateIcon(),
+		FUIAction(
+		FExecuteAction::CreateSP(this, &FAssetTypeActions_NiagaraEffect::ExecuteEdit, Scripts),
+		FCanExecuteAction()
+		)
+		);
+}
+
+void FAssetTypeActions_NiagaraEffect::OpenAssetEditor(const TArray<UObject*>& InObjects, TSharedPtr<IToolkitHost> EditWithinLevelEditor)
+{
+	EToolkitMode::Type Mode = EditWithinLevelEditor.IsValid() ? EToolkitMode::WorldCentric : EToolkitMode::Standalone;
+
+	for (auto ObjIt = InObjects.CreateConstIterator(); ObjIt; ++ObjIt)
+	{
+		auto Script = Cast<UNiagaraEffect>(*ObjIt);
+		if (Script != NULL)
+		{
+			FNiagaraEditorModule& NiagaraEditorModule = FModuleManager::LoadModuleChecked<FNiagaraEditorModule>("NiagaraEditor");
+			TSharedRef< INiagaraEffectEditor > NewNiagaraEditor = NiagaraEditorModule.CreateNiagaraEffectEditor(EToolkitMode::Standalone, EditWithinLevelEditor, Script);
+		}
+	}
+}
+
+void FAssetTypeActions_NiagaraEffect::ExecuteEdit(TArray<TWeakObjectPtr<UNiagaraEffect>> Objects)
+{
+	for (auto ObjIt = Objects.CreateConstIterator(); ObjIt; ++ObjIt)
+	{
+		auto Object = (*ObjIt).Get();
+		if (Object)
+		{
+			FAssetEditorManager::Get().OpenEditorForAsset(Object);
+		}
+	}
+}
+
+UClass* FAssetTypeActions_NiagaraEffect::GetSupportedClass() const
+{
+	return UNiagaraEffect::StaticClass();
+}
+
+
+#undef LOCTEXT_NAMESPACE
diff --git a/Engine/Source/Developer/AssetTools/Private/AssetTypeActions/AssetTypeActions_NiagaraEffect.h b/Engine/Source/Developer/AssetTools/Private/AssetTypeActions/AssetTypeActions_NiagaraEffect.h
new file mode 100644
--- /dev/null
+++ b/Engine/Source/Developer/AssetTools/Private/AssetTypeActions/AssetTypeActions_NiagaraEffect.h
@@ -0,0 +1,20 @@
+// Copyright 1998-2014 Epic Games, Inc. All Rights Reserved.
+
+#pragma once
+
+class FAssetTypeActions_NiagaraEffect: public FAssetTypeActions_Base
+{
+public:
+	// IAssetTypeActions Implementation
+	virtual FText GetName() const override { return NSLOCTEXT("AssetTypeActions", "AssetTypeActions_NiagaraEffect", "Niagara Effect"); }
+	virtual FColor GetTypeColor() const override { return FColor(255, 194, 0); }
+	virtual UClass* GetSupportedClass() const override;
+	virtual bool HasActions(const TArray<UObject*>& InObjects) const override { return true; }
+	virtual void GetActions(const TArray<UObject*>& InObjects, FMenuBuilder& MenuBuilder) override;
+	virtual void OpenAssetEditor(const TArray<UObject*>& InObjects, TSharedPtr<class IToolkitHost> EditWithinLevelEditor = TSharedPtr<IToolkitHost>()) override;
+	virtual uint32 GetCategories() override { return EAssetTypeCategories::Misc; }
+
+private:
+	/** Handler for when Edit is selected */
+	void ExecuteEdit(TArray<TWeakObjectPtr<class UNiagaraEffect>> Objects);
+};
\ No newline at end of file
diff --git a/Engine/Source/Editor/NiagaraEditor/Private/NiagaraEditorModule.cpp b/Engine/Source/Editor/NiagaraEditor/Private/NiagaraEditorModule.cpp
--- a/Engine/Source/Editor/NiagaraEditor/Private/NiagaraEditorModule.cpp
+++ b/Engine/Source/Editor/NiagaraEditor/Private/NiagaraEditorModule.cpp
@@ -30,3 +30,9 @@ TSharedRef<INiagaraEditor> FNiagaraEditorModule::CreateNiagaraEditor( const EToo
 	return NewNiagaraEditor;
 }
 
+TSharedRef<INiagaraEffectEditor> FNiagaraEditorModule::CreateNiagaraEffectEditor(const EToolkitMode::Type Mode, const TSharedPtr< class IToolkitHost >& InitToolkitHost, UNiagaraEffect* Effect)
+{
+	TSharedRef< FNiagaraEffectEditor > NewNiagaraEditor(new FNiagaraEffectEditor());
+	NewNiagaraEditor->InitNiagaraEffectEditor(Mode, InitToolkitHost, Effect);
+	return NewNiagaraEditor;
+}
diff --git a/Engine/Source/Editor/NiagaraEditor/Private/NiagaraEditorPrivatePCH.h b/Engine/Source/Editor/NiagaraEditor/Private/NiagaraEditorPrivatePCH.h
--- a/Engine/Source/Editor/NiagaraEditor/Private/NiagaraEditorPrivatePCH.h
+++ b/Engine/Source/Editor/NiagaraEditor/Private/NiagaraEditorPrivatePCH.h
@@ -6,6 +6,7 @@
 
 #include "NiagaraEditorModule.h"
 #include "NiagaraEditor.h"
+#include "NiagaraEffectEditor.h"
 
 #include "GraphEditor.h"
 
diff --git a/Engine/Source/Editor/NiagaraEditor/Private/NiagaraEffectEditor.cpp b/Engine/Source/Editor/NiagaraEditor/Private/NiagaraEffectEditor.cpp
new file mode 100644
--- /dev/null
+++ b/Engine/Source/Editor/NiagaraEditor/Private/NiagaraEffectEditor.cpp
@@ -0,0 +1,221 @@
+// Copyright 1998-2014 Epic Games, Inc. All Rights Reserved.
+
+#include "NiagaraEditorPrivatePCH.h"
+#include "Engine/NiagaraEffect.h"
+
+#include "Toolkits/IToolkitHost.h"
+#include "Editor/WorkspaceMenuStructure/Public/WorkspaceMenuStructureModule.h"
+#include "SNiagaraEffectEditorWidget.h"
+
+#define LOCTEXT_NAMESPACE "NiagaraEffectEditor"
+
+const FName FNiagaraEffectEditor::UpdateTabId(TEXT("NiagaraEditor_Effect"));
+
+void FNiagaraEffectEditor::RegisterTabSpawners(const TSharedRef<class FTabManager>& TabManager)
+{
+	FAssetEditorToolkit::RegisterTabSpawners(TabManager);
+
+	const IWorkspaceMenuStructure& MenuStructure = WorkspaceMenu::GetMenuStructure();
+
+	TabManager->RegisterTabSpawner(UpdateTabId, FOnSpawnTab::CreateSP(this, &FNiagaraEffectEditor::SpawnTab))
+		.SetDisplayName(LOCTEXT("NiagaraEffect", "Niagara Effect"))
+		.SetGroup(MenuStructure.GetAssetEditorCategory());
+}
+
+void FNiagaraEffectEditor::UnregisterTabSpawners(const TSharedRef<class FTabManager>& TabManager)
+{
+	FAssetEditorToolkit::UnregisterTabSpawners(TabManager);
+
+	TabManager->UnregisterTabSpawner(UpdateTabId);
+}
+
+FNiagaraEffectEditor::~FNiagaraEffectEditor()
+{
+
+}
+
+
+void FNiagaraEffectEditor::InitNiagaraEffectEditor(const EToolkitMode::Type Mode, const TSharedPtr< class IToolkitHost >& InitToolkitHost, UNiagaraEffect* InEffect)
+{
+	Effect = InEffect;
+	check(Effect != NULL);
+
+	TSharedRef<FTabManager::FLayout> StandaloneDefaultLayout = FTabManager::NewLayout("Standalone_Niagara_Effect_Layout_v2")
+		->AddArea
+		(
+		FTabManager::NewPrimaryArea()->SetOrientation(Orient_Vertical)
+		->Split
+		(
+		FTabManager::NewStack()
+		->SetSizeCoefficient(0.1f)
+		->AddTab(GetToolbarTabId(), ETabState::OpenedTab)
+		->SetHideTabWell(true)
+		)
+		->Split
+		(
+		FTabManager::NewStack()
+		->AddTab(UpdateTabId, ETabState::OpenedTab)
+		)
+		);
+
+	const bool bCreateDefaultStandaloneMenu = true;
+	const bool bCreateDefaultToolbar = true;
+	FAssetEditorToolkit::InitAssetEditor(Mode, InitToolkitHost, FNiagaraEditorModule::NiagaraEditorAppIdentifier, StandaloneDefaultLayout, bCreateDefaultStandaloneMenu, bCreateDefaultToolbar, Effect);
+
+	FNiagaraEditorModule& NiagaraEditorModule = FModuleManager::LoadModuleChecked<FNiagaraEditorModule>("NiagaraEditor");
+	AddMenuExtender(NiagaraEditorModule.GetMenuExtensibilityManager()->GetAllExtenders(GetToolkitCommands(), GetEditingObjects()));
+
+	ExtendToolbar();
+	RegenerateMenusAndToolbars();
+
+	// @todo toolkit world centric editing
+	/*// Setup our tool's layout
+	if( IsWorldCentricAssetEditor() )
+	{
+	const FString TabInitializationPayload(TEXT(""));		// NOTE: Payload not currently used for table properties
+	SpawnToolkitTab( UpdateGraphTabId, TabInitializationPayload, EToolkitTabSpot::Details );
+	}*/
+}
+
+FName FNiagaraEffectEditor::GetToolkitFName() const
+{
+	return FName("Niagara");
+}
+
+FText FNiagaraEffectEditor::GetBaseToolkitName() const
+{
+	return LOCTEXT("AppLabel", "Niagara");
+}
+
+FString FNiagaraEffectEditor::GetWorldCentricTabPrefix() const
+{
+	return LOCTEXT("WorldCentricTabPrefix", "Niagara ").ToString();
+}
+
+
+FLinearColor FNiagaraEffectEditor::GetWorldCentricTabColorScale() const
+{
+	return FLinearColor(0.0f, 0.0f, 0.2f, 0.5f);
+}
+
+
+/** Create new tab for the supplied graph - don't call this directly, call SExplorer->FindTabForGraph.*/
+TSharedRef<SNiagaraEffectEditorWidget> FNiagaraEffectEditor::CreateEditorWidget(UNiagaraEffect* InEffect)
+{
+	check(InEffect != NULL);
+
+	if (!EditorCommands.IsValid())
+	{
+		EditorCommands = MakeShareable(new FUICommandList);
+		/*
+		// Editing commands
+		EditorCommands->MapAction(FGenericCommands::Get().Delete,
+			FExecuteAction::CreateSP(this, &FNiagaraEffectEditor::DeleteSelectedNodes),
+			FCanExecuteAction::CreateSP(this, &FNiagaraEffectEditor::CanDeleteNodes)
+			);*/
+	}
+
+	// Create the appearance info
+	FGraphAppearanceInfo AppearanceInfo;
+	AppearanceInfo.CornerText = LOCTEXT("AppearanceCornerText", "NIAGARA").ToString();
+
+	//SGraphEditor::FGraphEditorEvents InEvents;
+
+	// Make title bar
+	TSharedRef<SWidget> TitleBarWidget =
+		SNew(SBorder)
+		.BorderImage(FEditorStyle::GetBrush(TEXT("Graph.TitleBackground")))
+		.HAlign(HAlign_Fill)
+		[
+			SNew(SHorizontalBox)
+			+ SHorizontalBox::Slot()
+			.HAlign(HAlign_Center)
+			.FillWidth(1.f)
+			[
+				SNew(STextBlock)
+				.Text(LOCTEXT("EffectLabel", "Niagara Effect"))
+				.TextStyle(FEditorStyle::Get(), TEXT("GraphBreadcrumbButtonText"))
+			]
+		];
+
+	// Make full graph editor
+	return SNew(SNiagaraEffectEditorWidget)
+		//.AdditionalCommands(EditorCommands)
+		//.Appearance(AppearanceInfo)
+		.TitleBar(TitleBarWidget)
+		.EffectObj(InEffect);
+}
+
+
+TSharedRef<SDockTab> FNiagaraEffectEditor::SpawnTab(const FSpawnTabArgs& Args)
+{
+	check(Args.GetTabId().TabType == UpdateTabId);
+
+	TSharedRef<SNiagaraEffectEditorWidget> Editor = CreateEditorWidget(Effect);
+
+	UpdateEditorPtr = Editor; // Keep pointer to editor
+
+	return SNew(SDockTab)
+		.Label(LOCTEXT("NiagaraEffect", "Niagara Effect"))
+		.TabColorScale(GetTabColorScale())
+		[
+			Editor
+		];
+}
+
+void FNiagaraEffectEditor::ExtendToolbar()
+{
+	/*
+	struct Local
+	{
+		static void FillToolbar(FToolBarBuilder& ToolbarBuilder, TSharedRef<SWidget> CompileBox)
+		{
+			ToolbarBuilder.BeginSection("Compile");
+			{
+				ToolbarBuilder.AddWidget(CompileBox);
+			}
+			ToolbarBuilder.EndSection();
+		}
+	};
+
+	TSharedPtr<FExtender> ToolbarExtender = MakeShareable(new FExtender);
+	TSharedRef<SWidget> Compilebox = SNew(SHorizontalBox)
+		+ SHorizontalBox::Slot()
+		.AutoWidth()
+		.Padding(4)
+		[
+			SNew(SButton)
+			.OnClicked(this, &FNiagaraEditor::OnCompileClicked)
+			.Content()
+			[
+				SNew(SVerticalBox)
+				+ SVerticalBox::Slot()
+				.AutoHeight()
+				[
+					SNew(SImage)
+					.Image(FEditorStyle::GetBrush("LevelEditor.Recompile"))
+				]
+				+ SVerticalBox::Slot()
+					.AutoHeight()
+					[
+						SNew(STextBlock)
+						.Text(LOCTEXT("NiagaraToolbar_Compile", "Compile"))
+					]
+			]
+		];
+
+	ToolbarExtender->AddToolBarExtension(
+		"Asset",
+		EExtensionHook::After,
+		GetToolkitCommands(),
+		FToolBarExtensionDelegate::CreateStatic(&Local::FillToolbar, Compilebox)
+		);
+
+	AddToolbarExtender(ToolbarExtender);
+	*/
+	FNiagaraEditorModule& NiagaraEditorModule = FModuleManager::LoadModuleChecked<FNiagaraEditorModule>("NiagaraEditor");
+	AddToolbarExtender(NiagaraEditorModule.GetToolBarExtensibilityManager()->GetAllExtenders(GetToolkitCommands(), GetEditingObjects()));
+}
+
+
+#undef LOCTEXT_NAMESPACE
diff --git a/Engine/Source/Editor/NiagaraEditor/Private/NiagaraEffectEditor.h b/Engine/Source/Editor/NiagaraEditor/Private/NiagaraEffectEditor.h
new file mode 100644
--- /dev/null
+++ b/Engine/Source/Editor/NiagaraEditor/Private/NiagaraEffectEditor.h
@@ -0,0 +1,54 @@
+// Copyright 1998-2014 Epic Games, Inc. All Rights Reserved.
+
+#pragma once
+
+#include "INiagaraEffectEditor.h"
+#include "Toolkits/AssetEditorToolkit.h"
+
+
+
+
+/** Viewer/editor for a DataTable */
+class FNiagaraEffectEditor : public INiagaraEffectEditor
+{
+
+public:
+
+	virtual void RegisterTabSpawners(const TSharedRef<class FTabManager>& TabManager) override;
+	virtual void UnregisterTabSpawners(const TSharedRef<class FTabManager>& TabManager) override;
+
+	/** Edits the specified Niagara Script */
+	void InitNiagaraEffectEditor(const EToolkitMode::Type Mode, const TSharedPtr< class IToolkitHost >& InitToolkitHost, UNiagaraEffect* Effect);
+
+	/** Destructor */
+	virtual ~FNiagaraEffectEditor();
+
+	// Begin IToolkit interface
+	virtual FName GetToolkitFName() const override;
+	virtual FText GetBaseToolkitName() const override;
+	virtual FString GetWorldCentricTabPrefix() const override;
+	virtual FLinearColor GetWorldCentricTabColorScale() const override;
+	// End IToolkit interface
+
+private:
+	/** Create widget for graph editing */
+	TSharedRef<class SNiagaraEffectEditorWidget> CreateEditorWidget(UNiagaraEffect* InEffect);
+
+	/** Spawns the tab with the update graph inside */
+	TSharedRef<SDockTab> SpawnTab(const FSpawnTabArgs& Args);
+
+	/** Builds the toolbar widget */
+	void ExtendToolbar();
+private:
+
+	/* The Effect being edited */
+	UNiagaraEffect	*Effect;
+	
+	/** The command list for this editor */
+	TSharedPtr<FUICommandList> EditorCommands;
+
+	TWeakPtr<SNiagaraEffectEditorWidget>		UpdateEditorPtr;
+
+	/**	Graph editor tab */
+	static const FName UpdateTabId;
+};
diff --git a/Engine/Source/Editor/NiagaraEditor/Private/SNiagaraEffectEditorWidget.h b/Engine/Source/Editor/NiagaraEditor/Private/SNiagaraEffectEditorWidget.h
new file mode 100644
--- /dev/null
+++ b/Engine/Source/Editor/NiagaraEditor/Private/SNiagaraEffectEditorWidget.h
@@ -0,0 +1,184 @@
+// Copyright 1998-2014 Epic Games, Inc. All Rights Reserved.
+
+#include "Visibility.h"
+#include "Slate.h"
+#include "Editor/PropertyEditor/Public/PropertyEditorModule.h"
+#include "Editor/PropertyEditor/Public/IDetailsView.h"
+
+
+FString NiagaraEffectText(TEXT("Niagara Effect"));
+FString ReadOnlyText(TEXT("Read Only"));
+
+class SNiagaraEffectEditorWidget : public SCompoundWidget, public FNotifyHook
+{
+private:
+	UNiagaraEffect *EffectObj;
+
+	/** Notification list to pass messages to editor users  */
+	TSharedPtr<SNotificationList> NotificationListPtr;
+	TSharedPtr<class IDetailsView> EmitterDetails;
+	SOverlay::FOverlaySlot* DetailsViewSlot;
+
+public:
+	SLATE_BEGIN_ARGS(SNiagaraEffectEditorWidget)
+		: _EffectObj(nullptr)
+	{}
+		SLATE_ARGUMENT(UNiagaraEffect*, EffectObj)
+		SLATE_ARGUMENT(TSharedPtr<SWidget>, TitleBar)
+	SLATE_END_ARGS()
+
+	/** Function to check whether we should show read-only text in the panel */
+	EVisibility ReadOnlyVisibility() const
+	{
+		return EVisibility::Hidden;
+	}
+
+	void Construct(const FArguments& InArgs)
+	{
+		/*
+		Commands = MakeShareable(new FUICommandList());
+		IsEditable = InArgs._IsEditable;
+		Appearance = InArgs._Appearance;
+		TitleBarEnabledOnly = InArgs._TitleBarEnabledOnly;
+		TitleBar = InArgs._TitleBar;
+		bAutoExpandActionMenu = InArgs._AutoExpandActionMenu;
+		ShowGraphStateOverlay = InArgs._ShowGraphStateOverlay;
+		OnNavigateHistoryBack = InArgs._OnNavigateHistoryBack;
+		OnNavigateHistoryForward = InArgs._OnNavigateHistoryForward;
+		OnNodeSpawnedByKeymap = InArgs._GraphEvents.OnNodeSpawnedByKeymap;
+		*/
+
+		/*
+		// Make sure that the editor knows about what kinds
+		// of commands GraphEditor can do.
+		FGraphEditorCommands::Register();
+
+		// Tell GraphEditor how to handle all the known commands
+		{
+		Commands->MapAction(FGraphEditorCommands::Get().ReconstructNodes,
+		FExecuteAction::CreateSP(this, &SGraphEditorImpl::ReconstructNodes),
+		FCanExecuteAction::CreateSP(this, &SGraphEditorImpl::CanReconstructNodes)
+		);
+
+		Commands->MapAction(FGraphEditorCommands::Get().BreakNodeLinks,
+		FExecuteAction::CreateSP(this, &SGraphEditorImpl::BreakNodeLinks),
+		FCanExecuteAction::CreateSP(this, &SGraphEditorImpl::CanBreakNodeLinks)
+		);
+
+		Commands->MapAction(FGraphEditorCommands::Get().BreakPinLinks,
+		FExecuteAction::CreateSP(this, &SGraphEditorImpl::BreakPinLinks, true),
+		FCanExecuteAction::CreateSP(this, &SGraphEditorImpl::CanBreakPinLinks)
+		);
+
+		// Append any additional commands that a consumer of GraphEditor wants us to be aware of.
+		const TSharedPtr<FUICommandList>& AdditionalCommands = InArgs._AdditionalCommands;
+		if (AdditionalCommands.IsValid())
+		{
+		Commands->Append(AdditionalCommands.ToSharedRef());
+		}
+
+		}
+		*/
+
+		EffectObj = InArgs._EffectObj;
+		//bNeedsRefresh = false;
+
+
+		FPropertyEditorModule& PropertyEditorModule = FModuleManager::GetModuleChecked<FPropertyEditorModule>("PropertyEditor");
+		const FDetailsViewArgs DetailsViewArgs(false, false, true, false, true, this);
+		EmitterDetails = PropertyEditorModule.CreateDetailView(DetailsViewArgs);
+		//FOnGetDetailCustomizationInstance LayoutMICDetails = FOnGetDetailCustomizationInstance::CreateStatic(&FMaterialInstanceParameterDetails::MakeInstance, MaterialEditorInstance, FGetShowHiddenParameters::CreateSP(this, &FMaterialInstanceEditor::GetShowHiddenParameters));
+		//EmitterDetails->RegisterInstancedCustomPropertyLayout(UMaterialEditorInstanceConstant::StaticClass(), LayoutMICDetails);
+
+		//EmitterDetails->SetObjects(EmitterObj, true);
+
+
+		TSharedPtr<SOverlay> OverlayWidget;
+		this->ChildSlot
+			[
+				SAssignNew(OverlayWidget, SOverlay)
+
+				// details view on the left
+				+ SOverlay::Slot()
+				.Expose(DetailsViewSlot)
+				.HAlign(HAlign_Left)
+				.VAlign(VAlign_Top)
+				[
+					SNew(SBorder)
+					.Padding(4)
+					[
+						EmitterDetails.ToSharedRef()
+					]
+				]
+
+				/*
+				// The graph panel
+				+ SOverlay::Slot()
+				.Expose(GraphPanelSlot)
+				[
+					SAssignNew(GraphPanel, SGraphPanel)
+					.EffectObj(EffectObj)
+					.GraphObjToDiff(InArgs._GraphToDiff)
+					//					.OnGetContextMenuFor(this, &SGraphEditorImpl::GraphEd_OnGetContextMenuFor)
+					.OnSelectionChanged(InArgs._GraphEvents.OnSelectionChanged)
+					.OnNodeDoubleClicked(InArgs._GraphEvents.OnNodeDoubleClicked)
+					//					.IsEditable(this, &SGraphEditorImpl::IsGraphEditable)
+					.OnDropActor(InArgs._GraphEvents.OnDropActor)
+					.OnDropStreamingLevel(InArgs._GraphEvents.OnDropStreamingLevel)
+					//					.IsEnabled(this, &SGraphEditorImpl::GraphEd_OnGetGraphEnabled)
+					.OnVerifyTextCommit(InArgs._GraphEvents.OnVerifyTextCommit)
+					.OnTextCommitted(InArgs._GraphEvents.OnTextCommitted)
+					.OnSpawnNodeByShortcut(InArgs._GraphEvents.OnSpawnNodeByShortcut)
+					//					.OnUpdateGraphPanel(this, &SGraphEditorImpl::GraphEd_OnPanelUpdated)
+					.OnDisallowedPinConnection(InArgs._GraphEvents.OnDisallowedPinConnection)
+					.ShowGraphStateOverlay(InArgs._ShowGraphStateOverlay)
+				]*/
+
+				// Title bar - optional
+				+ SOverlay::Slot()
+					.VAlign(VAlign_Top)
+					[
+						InArgs._TitleBar.IsValid() ? InArgs._TitleBar.ToSharedRef() : (TSharedRef<SWidget>)SNullWidget::NullWidget
+					]
+
+				// Bottom-right corner text indicating the type of tool
+				+ SOverlay::Slot()
+					.Padding(10)
+					.VAlign(VAlign_Bottom)
+					.HAlign(HAlign_Right)
+					[
+						SNew(STextBlock)
+						.Visibility(EVisibility::HitTestInvisible)
+						.TextStyle(FEditorStyle::Get(), "Graph.CornerText")
+						.Text( NiagaraEffectText )
+					]
+
+				// Top-right corner text indicating read only when not simulating
+				+ SOverlay::Slot()
+					.Padding(20)
+					.VAlign(VAlign_Top)
+					.HAlign(HAlign_Right)
+					[
+						SNew(STextBlock)
+						.Visibility(this, &SNiagaraEffectEditorWidget::ReadOnlyVisibility)
+						.TextStyle(FEditorStyle::Get(), "Graph.CornerText")
+						.Text(ReadOnlyText)
+					]
+
+				// Bottom-right corner text for notification list position
+				+ SOverlay::Slot()
+					.Padding(15.f)
+					.VAlign(VAlign_Bottom)
+					.HAlign(HAlign_Right)
+					[
+						SAssignNew(NotificationListPtr, SNotificationList)
+						.Visibility(EVisibility::HitTestInvisible)
+					]
+			];
+
+		//GraphPanel->RestoreViewSettings(FVector2D::ZeroVector, -1);
+
+		//NotifyGraphChanged();
+	}
+};
+
diff --git a/Engine/Source/Editor/NiagaraEditor/Public/NiagaraEditorModule.h b/Engine/Source/Editor/NiagaraEditor/Public/NiagaraEditorModule.h
--- a/Engine/Source/Editor/NiagaraEditor/Public/NiagaraEditorModule.h
+++ b/Engine/Source/Editor/NiagaraEditor/Public/NiagaraEditorModule.h
@@ -7,6 +7,7 @@
 #include "Toolkits/AssetEditorToolkit.h" // For FExtensibilityManager
 
 class INiagaraEditor;
+class INiagaraEffectEditor;
 
 /** DataTable Editor module */
 class FNiagaraEditorModule : public IModuleInterface,
@@ -23,6 +24,7 @@ public:
 
 	/** Creates an instance of Niagara editor.  Only virtual so that it can be called across the DLL boundary. */
 	virtual TSharedRef<INiagaraEditor> CreateNiagaraEditor( const EToolkitMode::Type Mode, const TSharedPtr< class IToolkitHost >& InitToolkitHost, class UNiagaraScript* Script );
+	virtual TSharedRef<INiagaraEffectEditor> CreateNiagaraEffectEditor(const EToolkitMode::Type Mode, const TSharedPtr< class IToolkitHost >& InitToolkitHost, class UNiagaraEffect* Effect);
 
 	/** Gets the extensibility managers for outside entities to extend static mesh editor's menus and toolbars */
 	virtual TSharedPtr<FExtensibilityManager> GetMenuExtensibilityManager() {return MenuExtensibilityManager;}
diff --git a/Engine/Source/Editor/UnrealEd/Classes/Factories/NiagaraEffectFactoryNew.h b/Engine/Source/Editor/UnrealEd/Classes/Factories/NiagaraEffectFactoryNew.h
new file mode 100644
--- /dev/null
+++ b/Engine/Source/Editor/UnrealEd/Classes/Factories/NiagaraEffectFactoryNew.h
@@ -0,0 +1,17 @@
+// Copyright 1998-2014 Epic Games, Inc. All Rights Reserved.
+
+#pragma once
+#include "NiagaraEffectFactoryNew.generated.h"
+
+UCLASS(hidecategories = Object)
+class UNiagaraEffectFactoryNew : public UFactory
+{
+	GENERATED_UCLASS_BODY()
+
+	// Begin UFactory Interface
+	virtual UObject* FactoryCreateNew(UClass* Class, UObject* InParent, FName Name, EObjectFlags Flags, UObject* Context, FFeedbackContext* Warn) override;
+	// Begin UFactory Interface	
+};
+
+
+
diff --git a/Engine/Source/Editor/UnrealEd/Private/Factories/NiagaraEffectFactory.cpp b/Engine/Source/Editor/UnrealEd/Private/Factories/NiagaraEffectFactory.cpp
new file mode 100644
--- /dev/null
+++ b/Engine/Source/Editor/UnrealEd/Private/Factories/NiagaraEffectFactory.cpp
@@ -0,0 +1,39 @@
+// Copyright 1998-2014 Epic Games, Inc. All Rights Reserved.
+
+#include "UnrealEd.h"
+#include "Engine/NiagaraEffect.h"
+#include "NiagaraEditorModule.h"
+#include "VectorVM.h"
+
+#define LOCTEXT_NAMESPACE "NiagaraEffectFactory"
+
+UNiagaraEffectFactoryNew::UNiagaraEffectFactoryNew(const class FPostConstructInitializeProperties& PCIP)
+: Super(PCIP)
+{
+
+	SupportedClass = UNiagaraEffect::StaticClass();
+	bCreateNew = false;
+	bEditAfterNew = true;
+
+	// Check ini to see if we should enable creation
+	bool bEnableNiagara = false;
+	GConfig->GetBool(TEXT("Niagara"), TEXT("EnableNiagara"), bEnableNiagara, GEngineIni);
+	bCreateNew = bEnableNiagara;
+}
+
+UObject* UNiagaraEffectFactoryNew::FactoryCreateNew(UClass* Class, UObject* InParent, FName Name, EObjectFlags Flags, UObject* Context, FFeedbackContext* Warn)
+{
+	check(Class->IsChildOf(UNiagaraEffect::StaticClass()));
+
+	// First allocate runtime script 
+	UNiagaraEffect* NewEffect = ConstructObject<UNiagaraEffect>(Class, InParent, Name, Flags);
+	if (NewEffect != NULL)
+	{
+		// Then allocate editor-only objects
+		//FNiagaraEditorModule& NiagaraEditorModule = FModuleManager::Get().LoadModuleChecked<FNiagaraEditorModule>(TEXT("NiagaraEditor"));
+	}
+
+	return NewEffect;
+}
+
+#undef LOCTEXT_NAMESPACE
\ No newline at end of file
diff --git a/Engine/Source/Runtime/Engine/Classes/Components/NiagaraComponent.h b/Engine/Source/Runtime/Engine/Classes/Components/NiagaraComponent.h
--- a/Engine/Source/Runtime/Engine/Classes/Components/NiagaraComponent.h
+++ b/Engine/Source/Runtime/Engine/Classes/Components/NiagaraComponent.h
@@ -11,6 +11,74 @@ enum ERenderModuleType
 };
 
 
+/* Predefined particle attribute names */
+const FName ParticleAttr_Position = "Position";
+const FName ParticleAttr_Velocity = "Velocity";
+const FName ParticleAttr_Color = "Color";
+
+class FNiagaraEmitterParticleData
+{
+public:
+	FNiagaraEmitterParticleData()
+		: CurrentBuffer(0),
+		NumParticles(0)
+	{
+		AddAttribute("Position");
+		AddAttribute("Velocity");
+		AddAttribute("Color");
+		AddAttribute("Rotation");
+		AddAttribute("Age");
+	}
+
+	~FNiagaraEmitterParticleData() {}
+
+	void Allocate(uint32 NumExpectedParticles)
+	{
+		ParticleBuffers[CurrentBuffer].Reset(NumExpectedParticles * AttrMap.Num());
+		ParticleBuffers[CurrentBuffer].AddUninitialized(NumExpectedParticles * AttrMap.Num());
+		ParticleAllocation = NumExpectedParticles;
+	}
+
+	const FVector4 *GetAttributeData(FName Name) const
+	{
+		int32 Offset = AttrMap[Name] * ParticleAllocation;
+		return ParticleBuffers[CurrentBuffer].GetTypedData() + Offset;
+	}
+
+	FVector4 *GetAttributeDataWrite(FName Name) 
+	{
+		int32 Offset = AttrMap[Name] * ParticleAllocation;
+		return ParticleBuffers[CurrentBuffer].GetTypedData() + Offset;
+	}
+
+	
+
+
+	void AddAttribute(FName NewAttrName)
+	{
+		uint32 Idx = AttrMap.Num();
+		AttrMap.Add(NewAttrName, Idx);
+	}
+
+	int GetNumAttributes()				{ return AttrMap.Num(); }
+
+	uint32 GetNumParticles() const		{ return NumParticles; }
+	uint32 GetParticleAllocation() const { return ParticleAllocation; }
+	void SetNumParticles(uint32 Num)	{ NumParticles = Num;  }
+	void SwapBuffers()					{ CurrentBuffer ^= 0x1; }
+
+	FVector4 *GetCurrentBuffer()		{ return ParticleBuffers[CurrentBuffer].GetTypedData(); }
+	FVector4 *GetPreviousBuffer()		{ return ParticleBuffers[CurrentBuffer^0x1].GetTypedData(); }
+
+private:
+	uint32 CurrentBuffer, NumParticles, ParticleAllocation;
+	TArray<FVector4> ParticleBuffers[2];
+	TArray<FVector4*> Attributes;
+	TMap<FName, uint32> AttrMap;
+};
+
+
+
 UCLASS()
 class ENGINE_API UNiagaraComponent : public UPrimitiveComponent
 {
diff --git a/Engine/Source/Runtime/Engine/Classes/Engine/NiagaraEffect.h b/Engine/Source/Runtime/Engine/Classes/Engine/NiagaraEffect.h
new file mode 100644
--- /dev/null
+++ b/Engine/Source/Runtime/Engine/Classes/Engine/NiagaraEffect.h
@@ -0,0 +1,13 @@
+// Copyright 1998-2014 Epic Games, Inc. All Rights Reserved.
+
+#pragma once
+#include "NiagaraEffect.generated.h"
+
+UCLASS(MinimalAPI)
+class UNiagaraEffect : public UObject
+{
+	GENERATED_UCLASS_BODY()
+
+	UPROPERTY()
+	TArray<FName> Emitters;
+};
diff --git a/Engine/Source/Runtime/Engine/Classes/Engine/NiagaraEffectRenderer.h b/Engine/Source/Runtime/Engine/Classes/Engine/NiagaraEffectRenderer.h
--- a/Engine/Source/Runtime/Engine/Classes/Engine/NiagaraEffectRenderer.h
+++ b/Engine/Source/Runtime/Engine/Classes/Engine/NiagaraEffectRenderer.h
@@ -10,46 +10,10 @@ DECLARE_CYCLE_STAT(TEXT("PreRenderView"), STAT_NiagaraPreRenderView, STATGROUP_N
 DECLARE_CYCLE_STAT(TEXT("Render Total"), STAT_NiagaraRender, STATGROUP_Niagara);
 DECLARE_CYCLE_STAT(TEXT("Render Sprites"), STAT_NiagaraRenderSprites, STATGROUP_Niagara);
 DECLARE_CYCLE_STAT(TEXT("Render Ribbons"), STAT_NiagaraRenderRibbons, STATGROUP_Niagara);
+DECLARE_CYCLE_STAT(TEXT("Generate Sprite Vertex Data"), STAT_NiagaraGenSpriteVertexData, STATGROUP_Niagara);
+DECLARE_CYCLE_STAT(TEXT("Generate Ribbon Vertex Data"), STAT_NiagaraGenRibbonVertexData, STATGROUP_Niagara);
 
 
-/** TODO: This should not be declared here but in its own file */
-struct FNiagaraParticleAttribute
-{
-	FName Name;
-	FVector4 *DataPtr;
-};
-
-class FNiagaraEmitterParticleData
-{
-public:
-	FNiagaraEmitterParticleData() 
-	{
-		AddAttribute("Position");
-		AddAttribute("Velocity");
-		AddAttribute("Color");
-	}
-
-	~FNiagaraEmitterParticleData() {}
-	
-	void Allocate(int NumParticles)
-	{
-	}
-
-	FVector4 *GetAttributeData(FName Name)
-	{
-		return nullptr;
-	}
-
-	void AddAttribute(FName NewAttrName)
-	{
-	}
-
-private:
-	TArray<FVector4> ParticleBuffers[2];
-	TArray<FNiagaraParticleAttribute> Attributes;
-	TMap<FName, FNiagaraParticleAttribute*> AttrMap;
-};
-
 
 
 /** Struct used to pass dynamic data from game thread to render thread */
@@ -110,7 +74,7 @@ public:
 	virtual void DrawDynamicElements(FPrimitiveDrawInterface* PDI, const FSceneView* View) = 0;
 	virtual void CreateRenderThreadResources() = 0;
 	virtual void ReleaseRenderThreadResources() = 0;
-	virtual FNiagaraDynamicDataBase *GenerateVertexData(int NumParticles, const TArray<FVector4> &Particles) = 0;
+	virtual FNiagaraDynamicDataBase *GenerateVertexData(const FNiagaraEmitterParticleData &Data) = 0;
 	virtual int GetDynamicDataSize() = 0;
 
 	virtual bool HasDynamicData() = 0;
@@ -415,30 +379,35 @@ public:
 
 
 	/** Update render data buffer from attributes */
-	FNiagaraDynamicDataBase *GenerateVertexData(int NumParticles, const TArray<FVector4> &ParticleArray) override
+	FNiagaraDynamicDataBase *GenerateVertexData(const FNiagaraEmitterParticleData &Data) override
 	{
+		SCOPE_CYCLE_COUNTER(STAT_NiagaraGenSpriteVertexData);
 		FNiagaraDynamicDataSprites *DynamicData = new FNiagaraDynamicDataSprites;
 		TArray<FParticleSpriteVertex>& RenderData = DynamicData->VertexData;
 
-		RenderData.Reset(NumParticles);
+		RenderData.Reset(Data.GetNumParticles());
 		//CachedBounds.Init();
 
-		const FVector4* Particles = ParticleArray.GetTypedData();
-		int32 AttrStride = NumParticles;
-		for (int32 ParticleIndex = 0; ParticleIndex < NumParticles; ParticleIndex++)
-		{
-			const FVector4* Particle = Particles + ParticleIndex;
+		const FVector4 *PosPtr = Data.GetAttributeData("Position");
+		const FVector4 *ColPtr = Data.GetAttributeData("Color");
+		const FVector4 *AgePtr = Data.GetAttributeData("Age");
+		const FVector4 *RotPtr = Data.GetAttributeData("Rotation");
 
+		float ParticleId = 0.0f, IdInc = 1.0f / Data.GetNumParticles();
+		for (uint32 ParticleIndex = 0; ParticleIndex < Data.GetNumParticles(); ParticleIndex++)
+		{
 			FParticleSpriteVertex& NewVertex = *new(RenderData)FParticleSpriteVertex;
-			NewVertex.Position = Particle[AttrStride * 0];
+			NewVertex.Position = PosPtr[ParticleIndex];
 			NewVertex.OldPosition = NewVertex.Position;
-			//CachedBounds += NewVertex.Position;
-			NewVertex.Color = FLinearColor(Particle[AttrStride * 2]);
-			NewVertex.ParticleId = static_cast<float>(ParticleIndex) / NumParticles;
-			NewVertex.RelativeTime = Particle[AttrStride*4].X;
+			NewVertex.Color = FLinearColor(ColPtr[ParticleIndex]);
+			NewVertex.ParticleId = ParticleId;
+			ParticleId += IdInc;
+			NewVertex.RelativeTime = AgePtr[ParticleIndex].X;
 			NewVertex.Size = FVector2D(1.0f, 1.0f);
-			NewVertex.Rotation = Particle[AttrStride*3].X;
+			NewVertex.Rotation = RotPtr[ParticleIndex].X;
 			NewVertex.SubImageIndex = 0.f;
+
+			//CachedBounds += NewVertex.Position;
 		}
 		//CachedBounds.ExpandBy(MaxSize);
 
@@ -447,6 +416,7 @@ public:
 
 
 
+
 	virtual void SetDynamicData_RenderThread(FNiagaraDynamicDataBase* NewDynamicData) override
 	{
 		check(IsInRenderingThread());
@@ -729,44 +699,47 @@ public:
 
 
 	/** Update render data buffer from attributes */
-	FNiagaraDynamicDataBase *GenerateVertexData(int NumParticles, const TArray<FVector4> &ParticleArray) override
+	FNiagaraDynamicDataBase *GenerateVertexData(const FNiagaraEmitterParticleData &Data)
 	{
+		SCOPE_CYCLE_COUNTER(STAT_NiagaraGenRibbonVertexData);
 		FNiagaraDynamicDataRibbon *DynamicData = new FNiagaraDynamicDataRibbon;
 		TArray<FParticleBeamTrailVertex>& RenderData = DynamicData->VertexData;
 
-		RenderData.Reset(NumParticles*2);
+		RenderData.Reset(Data.GetNumParticles() * 2);
 		//CachedBounds.Init();
 
-		const FVector4* Particles = ParticleArray.GetTypedData();
-		int32 AttrStride = NumParticles;
-
-
 		// build a sorted list by age, so we always get particles in order 
 		// regardless of them being moved around due to dieing and spawning
 		TArray<int32> SortedIndices;
-		for (int32 Idx = 0; Idx < NumParticles; Idx++)
+		for (uint32 Idx = 0; Idx < Data.GetNumParticles(); Idx++)
 		{
 			SortedIndices.Add(Idx);
 		}
 
+		const FVector4 *AgeData = Data.GetAttributeData("Age");
 		SortedIndices.Sort(
-			[&Particles, AttrStride](const int32& A, const int32& B) {
-				return Particles[A + AttrStride * 4].X < Particles[B + AttrStride * 4].X;
-				}
+			[&AgeData](const int32& A, const int32& B) {
+			return AgeData[A].X < AgeData[B].X;
+		}
 		);
 
 
 		FVector2D UVs[4] = { FVector2D(0.0f, 0.0f), FVector2D(1.0f, 0.0f), FVector2D(1.0f, 1.0f), FVector2D(0.0f, 1.0f) };
-		
+
+		const FVector4 *PosPtr = Data.GetAttributeData("Position");
+		const FVector4 *ColorPtr = Data.GetAttributeData("Color");
+		const FVector4 *AgePtr = Data.GetAttributeData("Age");
+		const FVector4 *RotPtr = Data.GetAttributeData("Rotation");
+
 		FVector PrevPos, PrevPos2, PrevDir(0.0f, 0.0f, 0.1f);
-		for (int32 i = 0; i<SortedIndices.Num() - 1; i++)
+		for (int32 i = 0; i < SortedIndices.Num() - 1; i++)
 		{
-			const FVector4* Particle = Particles + SortedIndices[i];
-			const FVector4* Particle2 = Particles + SortedIndices[i+1];
+			uint32 Index1 = SortedIndices[i];
+			uint32 Index2 = SortedIndices[i + 1];
 
-			const FVector ParticlePos = Particle[AttrStride * 0];
-			FVector ParticleDir = Particle2[AttrStride * 0] - ParticlePos;
-			if (ParticleDir.Size()<=SMALL_NUMBER)
+			const FVector ParticlePos = PosPtr[Index1];
+			FVector ParticleDir = PosPtr[Index2] - ParticlePos;
+			if (ParticleDir.Size() <= SMALL_NUMBER)
 			{
 				ParticleDir = PrevDir*0.1f;
 			}
@@ -776,17 +749,17 @@ public:
 
 			if (i == 0)
 			{
-				AddRibbonVert(RenderData, ParticlePos + ParticleRight, Particle, AttrStride, UVs[0]);
-				AddRibbonVert(RenderData, ParticlePos - ParticleRight, Particle, AttrStride, UVs[1]);
+				AddRibbonVert(RenderData, ParticlePos + ParticleRight, Data, UVs[0], ColorPtr[i], AgePtr[i], RotPtr[i]);
+				AddRibbonVert(RenderData, ParticlePos - ParticleRight, Data, UVs[1], ColorPtr[i], AgePtr[i], RotPtr[i]);
 			}
 			else
 			{
-				AddRibbonVert(RenderData, PrevPos2, Particle, AttrStride, UVs[0]);
-				AddRibbonVert(RenderData, PrevPos, Particle, AttrStride, UVs[1]);
+				AddRibbonVert(RenderData, PrevPos2, Data, UVs[0], ColorPtr[i], AgePtr[i], RotPtr[i]);
+				AddRibbonVert(RenderData, PrevPos, Data, UVs[1], ColorPtr[i], AgePtr[i], RotPtr[i]);
 			}
 
-			AddRibbonVert(RenderData, ParticlePos - ParticleRight + ParticleDir, Particle, AttrStride, UVs[2]);
-			AddRibbonVert(RenderData, ParticlePos + ParticleRight + ParticleDir, Particle, AttrStride, UVs[3]);
+			AddRibbonVert(RenderData, ParticlePos - ParticleRight + ParticleDir, Data, UVs[2], ColorPtr[i], AgePtr[i], RotPtr[i]);
+			AddRibbonVert(RenderData, ParticlePos + ParticleRight + ParticleDir, Data, UVs[3], ColorPtr[i], AgePtr[i], RotPtr[i]);
 			PrevPos = ParticlePos - ParticleRight + ParticleDir;
 			PrevPos2 = ParticlePos + ParticleRight + ParticleDir;
 			PrevDir = ParticleDir;
@@ -795,17 +768,17 @@ public:
 		return DynamicData;
 	}
 
-
-	void AddRibbonVert(TArray<FParticleBeamTrailVertex>& RenderData, FVector ParticlePos, const FVector4 *Particle, int32 AttrStride, FVector2D UV1)
+	void AddRibbonVert(TArray<FParticleBeamTrailVertex>& RenderData, FVector ParticlePos, const FNiagaraEmitterParticleData &Data, FVector2D UV1,
+		const FVector4 Color, const FVector4 Age, const FVector4 Rotation)
 	{
 		FParticleBeamTrailVertex& NewVertex = *new(RenderData)FParticleBeamTrailVertex;
 		NewVertex.Position = ParticlePos;
 		NewVertex.OldPosition = NewVertex.Position;
-		NewVertex.Color = FLinearColor(Particle[AttrStride * 2]);
+		NewVertex.Color = FLinearColor(Color);
 		NewVertex.ParticleId = 0;
-		NewVertex.RelativeTime = Particle[AttrStride * 4].X;
+		NewVertex.RelativeTime = Age.X;
 		NewVertex.Size = FVector2D(1.0f, 1.0f);
-		NewVertex.Rotation = Particle[AttrStride * 3].X;
+		NewVertex.Rotation = Rotation.X;
 		NewVertex.SubImageIndex = 0.f;
 		NewVertex.Tex_U = UV1.X;
 		NewVertex.Tex_V = UV1.Y;
@@ -814,6 +787,7 @@ public:
 	}
 
 
+
 	virtual void SetDynamicData_RenderThread(FNiagaraDynamicDataBase* NewDynamicData) override
 	{
 		check(IsInRenderingThread());
diff --git a/Engine/Source/Runtime/Engine/Private/NiagaraComponent.cpp b/Engine/Source/Runtime/Engine/Private/NiagaraComponent.cpp
--- a/Engine/Source/Runtime/Engine/Private/NiagaraComponent.cpp
+++ b/Engine/Source/Runtime/Engine/Private/NiagaraComponent.cpp
@@ -141,9 +141,6 @@ public:
 		, UpdateScript(*InComponent->UpdateScript)
 		, SpawnScript(InComponent->SpawnScript)
 		, ConstantTable(InComponent->UpdateScript->ConstantTable)
-		, BufferIndex(0)
-		, NumVectorsPerAttribute(0)
-		, NumParticles(0)
 		, SpawnRemainder(0.0f)
 		, CachedBounds(ForceInit)
 	{
@@ -165,60 +162,39 @@ public:
 		// Cache the ComponentToWorld transform.
 		CachedComponentToWorld = Component.GetComponentToWorld();
 		
-		// Flip buffers.
-		BufferIndex ^= 0x1;
-
-		// Get prev + next data.
-		TArray<FVector4>& PrevParticles = ParticleBuffers[BufferIndex ^ 0x1];
-		TArray<FVector4>& Particles = ParticleBuffers[BufferIndex];
+		Data.SwapBuffers();
 
 		// Figure out how many we will spawn.
 		int32 NumToSpawn = CalcNumToSpawn(DeltaSeconds);
+		int32 MaxNewParticles = Data.GetNumParticles() + NumToSpawn;
 
 		// Remember the stride of the original data.
-		int32 PrevNumVectorsPerAttribute = NumVectorsPerAttribute;
-
-		// The script updates relative time so we don't know yet which will die
-		// without simulation. Allocate for the worst case.
-		int32 MaxNewParticles = NumParticles + NumToSpawn;
-		NumVectorsPerAttribute = MaxNewParticles;
-		Particles.Reset(NumAttributes * NumVectorsPerAttribute);
-		Particles.AddUninitialized(NumAttributes * NumVectorsPerAttribute);
+		int32 PrevNumVectorsPerAttribute = Data.GetParticleAllocation();
 
+		Data.Allocate( MaxNewParticles );
 		// Simulate particles forward by DeltaSeconds.
-		{
-			SCOPE_CYCLE_COUNTER(STAT_NiagaraSimulate);
-			UpdateParticles(
-				DeltaSeconds,
-				PrevParticles.GetTypedData(),
-				PrevNumVectorsPerAttribute,
-				Particles.GetTypedData(),
-				NumVectorsPerAttribute,
-				NumParticles
-				);
-		}
+		UpdateParticles(
+			DeltaSeconds,
+			Data.GetPreviousBuffer(),
+			PrevNumVectorsPerAttribute,
+			Data.GetCurrentBuffer(),
+			MaxNewParticles,
+			Data.GetNumParticles()
+			);
 
-		// Spawn new particles overwriting dead particles and compact the final buffer.
-		{
-			SCOPE_CYCLE_COUNTER(STAT_NiagaraSpawnAndKill);
-			NumParticles = SpawnAndKillParticles(
-				Particles.GetTypedData(),
-				NumParticles,
-				NumToSpawn,
-				NumVectorsPerAttribute
-				);
-		}
+		SpawnAndKillParticles(
+			NumToSpawn
+			);
+
+		DECLARE_DWORD_COUNTER_STAT(TEXT("NumParticles"), STAT_NiagaraNumParticles, STATGROUP_Niagara);
+		INC_DWORD_STAT_BY(STAT_NiagaraNumParticles, Data.GetNumParticles());
 
-		DECLARE_DWORD_COUNTER_STAT(TEXT("NumParticles"),STAT_NiagaraNumParticles,STATGROUP_Niagara);
-		INC_DWORD_STAT_BY(STAT_NiagaraNumParticles, NumParticles);
 	}
 
 	FBox GetBounds() const { return CachedBounds; }
 
 
-	int GetNumParticles() const { return NumParticles; }
-	int GetNumVectorsPerAttribute() const { return NumVectorsPerAttribute; }
-	const TArray<FVector4> &GetParticles() { return ParticleBuffers[BufferIndex]; }
+	FNiagaraEmitterParticleData &GetData()	{ return Data; }
 
 	void SetConstant(int Idx, const FVector4 &Value)
 	{
@@ -238,14 +214,8 @@ private:
 	UNiagaraScript *SpawnScript;
 	/** Local constant table. */
 	TArray<FVector4> ConstantTable;
-	/** 16-byte aligned set of attribute data used by vector VM for simulation */
-	TArray<FVector4> ParticleBuffers[2];
-	/** Current buffer. */
-	int32 BufferIndex;
-	/** Stride between attributes for the current particle buffer. In sizeof(FVector4). */
-	int32 NumVectorsPerAttribute;
-	/** Total number of alive particles */
-	int32 NumParticles;
+	/** particle simulation data */
+	FNiagaraEmitterParticleData Data;
 	/** Keep partial particle spawns from last frame */
 	float SpawnRemainder;
 	/** The cached ComponentToWorld transform. */
@@ -288,10 +258,11 @@ private:
 		int32 NumParticles
 		)
 	{
-		VectorRegister* InputRegisters[VectorVM::MaxInputRegisters] = {0};
+		SCOPE_CYCLE_COUNTER(STAT_NiagaraSimulate);
+
+		VectorRegister* InputRegisters[VectorVM::MaxInputRegisters] = { 0 };
 		VectorRegister* OutputRegisters[VectorVM::MaxOutputRegisters] = {0};
 		const int32 NumAttr = UpdateScript.Attributes.Num();
-		const int32 NumVectors = NumParticles;
 
 		check(NumAttr < VectorVM::MaxInputRegisters);
 		check(NumAttr < VectorVM::MaxOutputRegisters);
@@ -313,27 +284,16 @@ private:
 			OutputRegisters,
 			NumAttr,
 			ConstantTable.GetData(),
-			NumVectors
+			NumParticles
 			);
 	}
 
-	/** Init new particles, and write over dead ones */
-	int32 SpawnAndKillParticles(
-		FVector4* Particles,
-		int32 NumParticles,
-		int32 NumToSpawn,
-		int32 NumVectorsPerAttribute
-		)
+	int32 SpawnAndKillParticles(int32 NumToSpawn)
 	{
-		int32 OrigNumParticles = NumParticles;
-		FVector4 *NewParticles = Particles + NumParticles;
-
-		// Spawn new Particles at the end of the buffer
-		for (int32 i = 0; i < NumToSpawn; i++)
-		{
-			SpawnParticleAtIndex(NewParticles, NumVectorsPerAttribute, i);
-			NumParticles++;
-		}
+		SCOPE_CYCLE_COUNTER(STAT_NiagaraSpawnAndKill);
+		int32 OrigNumParticles = Data.GetNumParticles();
+		int32 CurNumParticles = OrigNumParticles;
+		CurNumParticles = SpawnParticles(NumToSpawn);
 
 		// run the spawn graph over all new particles
 		if (SpawnScript && SpawnScript->ByteCode.Num())
@@ -341,16 +301,18 @@ private:
 			VectorRegister* InputRegisters[VectorVM::MaxInputRegisters] = { 0 };
 			VectorRegister* OutputRegisters[VectorVM::MaxOutputRegisters] = { 0 };
 			const int32 NumAttr = SpawnScript->Attributes.Num();
-			const int32 NumVectors = NumToSpawn;// ((NumParticles + 0x3) & ~0x3);
+			const int32 NumVectors = NumToSpawn;
 
 			check(NumAttr < VectorVM::MaxInputRegisters);
 			check(NumAttr < VectorVM::MaxOutputRegisters);
 
+			FVector4 *NewParticlesStart = Data.GetCurrentBuffer() + OrigNumParticles;
+
 			// Setup input and output registers.
 			for (int32 AttrIndex = 0; AttrIndex < NumAttr; ++AttrIndex)
 			{
-				InputRegisters[AttrIndex] = (VectorRegister*)(NewParticles + AttrIndex * NumVectorsPerAttribute);
-				OutputRegisters[AttrIndex] = (VectorRegister*)(NewParticles + AttrIndex * NumVectorsPerAttribute);
+				InputRegisters[AttrIndex] = (VectorRegister*)(NewParticlesStart + AttrIndex * Data.GetParticleAllocation());
+				OutputRegisters[AttrIndex] = (VectorRegister*)(NewParticlesStart + AttrIndex * Data.GetParticleAllocation());
 			}
 
 			// copy script specific constants into the constant table
@@ -367,50 +329,62 @@ private:
 				);
 		}
 
-
 		// Iterate over looking for dead particles and move from the end of the list to the dead location, compacting in the process
 		int32 ParticleIndex = 0;
-		FVector4* ParticleRelativeTimes = (Particles + ENiagaraVectorAttr::RelativeTime * NumVectorsPerAttribute);
+		const FVector4* ParticleRelativeTimes = Data.GetAttributeData("Age");
 		while (ParticleIndex < OrigNumParticles)
 		{
 			if (ParticleRelativeTimes[ParticleIndex].X > 1.0f)
 			{
 				// Particle is dead, move one from the end here.
-				MoveParticleToIndex(Particles, NumVectorsPerAttribute, --NumParticles, ParticleIndex);
+				MoveParticleToIndex(--CurNumParticles, ParticleIndex);
 			}
 			ParticleIndex++;
 		}
 
-		return NumParticles;
+		Data.SetNumParticles(CurNumParticles);
+		return CurNumParticles;
 	}
 
 	/** Spawn a new particle at this index */
-	void SpawnParticleAtIndex(FVector4* Particles, int32 NumVectorsPerAttribute, int32 ParticleIndex)
+	int32 SpawnParticles(int32 NumToSpawn)
 	{
-		FVector SpawnLocation = CachedComponentToWorld.GetLocation();
-		SpawnLocation.X += FMath::FRandRange(-20.f, 20.f);
-		SpawnLocation.Y += FMath::FRandRange(-20.f, 20.f);
-		SpawnLocation.Z += FMath::FRandRange(-20.f, 20.f);
-
-		FVector4* Attr = Particles + ParticleIndex;
-		Attr[NumVectorsPerAttribute*ENiagaraVectorAttr::Position] = SpawnLocation;
-		Attr[NumVectorsPerAttribute*ENiagaraVectorAttr::Color] = FVector4(0.0f, 0.0f, 0.0f, 0.0f);
-		Attr[NumVectorsPerAttribute*ENiagaraVectorAttr::RelativeTime] = FVector4(0.0f, 0.0f, 0.0f, 0.0f);
-		Attr[NumVectorsPerAttribute*ENiagaraVectorAttr::Rotation] = FVector4(0.0f, 0.0f, 0.0f, 0.0f);
-		Attr[NumVectorsPerAttribute*ENiagaraVectorAttr::Velocity] = FVector4(0.0f, 0.0f, 2.0f, 0.0f);
+		FVector4 *PosPtr = Data.GetAttributeDataWrite("Position");
+		FVector4 *VelPtr = Data.GetAttributeDataWrite("Velocity");
+		FVector4 *ColPtr = Data.GetAttributeDataWrite("Color");
+		FVector4 *RotPtr = Data.GetAttributeDataWrite("Rotation");
+		FVector4 *AgePtr = Data.GetAttributeDataWrite("Age");
+
+		// Spawn new Particles at the end of the buffer
+		int32 ParticleIndex = Data.GetNumParticles();
+		for (int32 i = 0; i < NumToSpawn; i++)
+		{
+			FVector SpawnLocation = CachedComponentToWorld.GetLocation();
+			SpawnLocation.X += FMath::FRandRange(-20.f, 20.f);
+			SpawnLocation.Y += FMath::FRandRange(-20.f, 20.f);
+			SpawnLocation.Z += FMath::FRandRange(-20.f, 20.f);
+
+			PosPtr[ParticleIndex] = SpawnLocation;
+			ColPtr[ParticleIndex] = FVector4(0.0f, 0.0f, 0.0f, 0.0f);
+			VelPtr[ParticleIndex] = FVector4(0.0f, 0.0f, 2.0f, 0.0f);
+			RotPtr[ParticleIndex] = FVector4(0.0f, 0.0f, 0.0f, 0.0f);
+			AgePtr[ParticleIndex] = FVector4(0.0f, 0.0f, 0.0f, 0.0f);
+			ParticleIndex++;
+		}
+		return ParticleIndex;
 	}
 
 	/** Util to move a particle */
-	void MoveParticleToIndex(FVector4* Particles, int32 NumVectorsPerAttribute, int32 SrcIndex, int32 DestIndex)
+	void MoveParticleToIndex(int32 SrcIndex, int32 DestIndex)
 	{
-		FVector4 *SrcPtr = Particles + SrcIndex;
-		FVector4 *DestPtr = Particles + DestIndex;
+		FVector4 *SrcPtr = Data.GetCurrentBuffer() + SrcIndex;
+		FVector4 *DestPtr = Data.GetCurrentBuffer() + DestIndex;
 
-		for (int32 AttrIndex = 0; AttrIndex < ENiagaraVectorAttr::MaxVectorAttribs; AttrIndex++)
+		for (int32 AttrIndex = 0; AttrIndex < Data.GetNumAttributes(); AttrIndex++)
 		{
 			*DestPtr = *SrcPtr;
-			DestPtr += NumVectorsPerAttribute;
-			SrcPtr += NumVectorsPerAttribute;
+			DestPtr += Data.GetParticleAllocation();
+			SrcPtr += Data.GetParticleAllocation();
 		}
 	}
 
@@ -478,8 +452,7 @@ void UNiagaraComponent::SendRenderDynamicData_Concurrent()
 	if (Simulation && SceneProxy)
 	{
 		FNiagaraSceneProxy *NiagaraProxy = static_cast<FNiagaraSceneProxy*>(SceneProxy);
-		FNiagaraDynamicDataBase* DynamicData = NiagaraProxy->GetEffectRenderer()->GenerateVertexData(Simulation->GetNumVectorsPerAttribute(), Simulation->GetParticles());
-		//FNiagaraDynamicDataBase* DynamicData = Simulation->GetDynamicData();
+		FNiagaraDynamicDataBase* DynamicData = NiagaraProxy->GetEffectRenderer()->GenerateVertexData(Simulation->GetData());
 
 		ENQUEUE_UNIQUE_RENDER_COMMAND_TWOPARAMETER(
 			FSendNiagaraDynamicData,
diff --git a/Engine/Source/Runtime/Engine/Private/NiagaraEffect.cpp b/Engine/Source/Runtime/Engine/Private/NiagaraEffect.cpp
new file mode 100644
--- /dev/null
+++ b/Engine/Source/Runtime/Engine/Private/NiagaraEffect.cpp
@@ -0,0 +1,12 @@
+// Copyright 1998-2014 Epic Games, Inc. All Rights Reserved.
+
+#include "EnginePrivate.h"
+#include "Engine/NiagaraEffect.h"
+
+
+
+UNiagaraEffect::UNiagaraEffect(const class FPostConstructInitializeProperties& PCIP)
+: Super(PCIP)
+{
+
+}
diff --git a/Engine/Source/Runtime/VectorVM/Private/VectorVM.cpp b/Engine/Source/Runtime/VectorVM/Private/VectorVM.cpp
--- a/Engine/Source/Runtime/VectorVM/Private/VectorVM.cpp
+++ b/Engine/Source/Runtime/VectorVM/Private/VectorVM.cpp
@@ -503,10 +503,10 @@ void VectorVM::Init()
 			{
 				for (int x = 0; x < 10; x++)
 				{
-					float f1 = ((float)FMath::Rand() / (float)RAND_MAX) * 2 - 1;
-					float f2 = ((float)FMath::Rand() / (float)RAND_MAX) * 2 - 1;
-					float f3 = ((float)FMath::Rand() / (float)RAND_MAX) * 2 - 1;
-					float f4 = ((float)FMath::Rand() / (float)RAND_MAX) * 2 - 1;
+					float f1 = (float)FMath::FRandRange(-1.0f, 1.0f);
+					float f2 = (float)FMath::FRandRange(-1.0f, 1.0f);
+					float f3 = (float)FMath::FRandRange(-1.0f, 1.0f);
+					float f4 = (float)FMath::FRandRange(-1.0f, 1.0f);
 
 					FVectorKernelNoise::RandomTable[x][y][z] = MakeVectorRegister(f1, f2, f3, f4);
 				}
@@ -607,26 +607,26 @@ namespace VectorVM
 		FVectorVMOpInfo(EOp::done, EOpFlags::None, EOpSrc::Invalid, EOpSrc::Invalid, EOpSrc::Invalid, TEXT("done")),
 
 		FVectorVMOpInfo(EOp::add, EOpFlags::Implemented | EOpFlags::Commutative, EOpSrc::Register, EOpSrc::Register, EOpSrc::Invalid, TEXT("Add")),
-		FVectorVMOpInfo(EOp::add, EOpFlags::None, EOpSrc::Register, EOpSrc::Const, EOpSrc::Invalid, TEXT("addi")),
+		FVectorVMOpInfo(),
 
 		FVectorVMOpInfo(EOp::sub, EOpFlags::Implemented, EOpSrc::Register, EOpSrc::Register, EOpSrc::Invalid, TEXT("Sub")),
-		FVectorVMOpInfo(EOp::sub, EOpFlags::None, EOpSrc::Register, EOpSrc::Const, EOpSrc::Invalid, TEXT("subi")),
+		FVectorVMOpInfo(),
 
 		FVectorVMOpInfo(EOp::mul, EOpFlags::Implemented | EOpFlags::Commutative, EOpSrc::Register, EOpSrc::Register, EOpSrc::Invalid, TEXT("Multiply")),
-		FVectorVMOpInfo(EOp::mul,EOpFlags::None,EOpSrc::Register,EOpSrc::Const,EOpSrc::Invalid,TEXT("muli")),
+		FVectorVMOpInfo(),
 
 		FVectorVMOpInfo(EOp::mad,EOpFlags::Implemented|EOpFlags::Commutative,EOpSrc::Register,EOpSrc::Register,EOpSrc::Register,TEXT("Multply-Add")),
-		FVectorVMOpInfo(EOp::mad,EOpFlags::None,EOpSrc::Register,EOpSrc::Register,EOpSrc::Const,TEXT("madrri")),
-		FVectorVMOpInfo(EOp::mad,EOpFlags::None,EOpSrc::Register,EOpSrc::Const,EOpSrc::Register,TEXT("madrir")),
-		FVectorVMOpInfo(EOp::mad,EOpFlags::None,EOpSrc::Register,EOpSrc::Const,EOpSrc::Const,TEXT("madrii")),
-		FVectorVMOpInfo(EOp::mad,EOpFlags::None,EOpSrc::Const,EOpSrc::Const,EOpSrc::Register,TEXT("madiir")),
-		FVectorVMOpInfo(EOp::mad,EOpFlags::None,EOpSrc::Const,EOpSrc::Const,EOpSrc::Const,TEXT("madiii")),
+		FVectorVMOpInfo(),
+		FVectorVMOpInfo(),
+		FVectorVMOpInfo(),
+		FVectorVMOpInfo(),
+		FVectorVMOpInfo(),
 
 		FVectorVMOpInfo(EOp::lerp,EOpFlags::Implemented,EOpSrc::Register,EOpSrc::Register,EOpSrc::Register,TEXT("Lerp")),
-		FVectorVMOpInfo(EOp::lerp,EOpFlags::None,EOpSrc::Const,EOpSrc::Register,EOpSrc::Register,TEXT("lerpirr")),
-		FVectorVMOpInfo(EOp::lerp,EOpFlags::None,EOpSrc::Register,EOpSrc::Const,EOpSrc::Register,TEXT("lerprir")), 
-		FVectorVMOpInfo(EOp::lerp,EOpFlags::None,EOpSrc::Register,EOpSrc::Register,EOpSrc::Const,TEXT("lerprri")),
-		FVectorVMOpInfo(EOp::lerp,EOpFlags::None,EOpSrc::Const,EOpSrc::Const,EOpSrc::Register,TEXT("lerpiir")),
+		FVectorVMOpInfo(),
+		FVectorVMOpInfo(), 
+		FVectorVMOpInfo(),
+		FVectorVMOpInfo(),
 
 		FVectorVMOpInfo(EOp::rcp,EOpFlags::Implemented,EOpSrc::Register,EOpSrc::Invalid,EOpSrc::Invalid,TEXT("Reciprocal")),
 		FVectorVMOpInfo(EOp::rsq,EOpFlags::Implemented,EOpSrc::Register,EOpSrc::Invalid,EOpSrc::Invalid,TEXT("Reciprocal Sqrt")),
@@ -638,7 +638,7 @@ namespace VectorVM
 		FVectorVMOpInfo(EOp::log,EOpFlags::None,EOpSrc::Register,EOpSrc::Invalid,EOpSrc::Invalid,TEXT("Log")),
 		FVectorVMOpInfo(EOp::log2,EOpFlags::None,EOpSrc::Register,EOpSrc::Invalid,EOpSrc::Invalid,TEXT("Log base 2")),
 		FVectorVMOpInfo(EOp::sin,EOpFlags::Implemented,EOpSrc::Register,EOpSrc::Invalid,EOpSrc::Invalid,TEXT("Sin")),
-		FVectorVMOpInfo(EOp::sin, EOpFlags::None, EOpSrc::Const, EOpSrc::Invalid, EOpSrc::Invalid, TEXT("sini")),
+		FVectorVMOpInfo(),
 		FVectorVMOpInfo(EOp::cos, EOpFlags::None, EOpSrc::Register, EOpSrc::Invalid, EOpSrc::Invalid, TEXT("Cos")),
 		FVectorVMOpInfo(EOp::tan,EOpFlags::None,EOpSrc::Register,EOpSrc::Invalid,EOpSrc::Invalid,TEXT("Tan")),
 		FVectorVMOpInfo(EOp::asin,EOpFlags::None,EOpSrc::Register,EOpSrc::Invalid,EOpSrc::Invalid,TEXT("Arcsin")),
@@ -649,45 +649,45 @@ namespace VectorVM
 		FVectorVMOpInfo(EOp::floor,EOpFlags::None,EOpSrc::Register,EOpSrc::Invalid,EOpSrc::Invalid,TEXT("Round down")),
 		FVectorVMOpInfo(EOp::fmod,EOpFlags::None,EOpSrc::Register,EOpSrc::Register,EOpSrc::Invalid,TEXT("Modulo")),
 		FVectorVMOpInfo(EOp::frac,EOpFlags::None,EOpSrc::Register,EOpSrc::Invalid,EOpSrc::Invalid,TEXT("Fractional")),
-		FVectorVMOpInfo(EOp::trunc,EOpFlags::None,EOpSrc::Register,EOpSrc::Invalid,EOpSrc::Invalid,TEXT("Trunc")),
+		FVectorVMOpInfo(EOp::trunc,EOpFlags::None,EOpSrc::Register,EOpSrc::Invalid,EOpSrc::Invalid,TEXT("Truncate")),
 
 		FVectorVMOpInfo(EOp::clamp,EOpFlags::Implemented,EOpSrc::Register,EOpSrc::Register,EOpSrc::Register,TEXT("Clamp")),
-		FVectorVMOpInfo(EOp::clamp,EOpFlags::None,EOpSrc::Register,EOpSrc::Const,EOpSrc::Register,TEXT("clampir")),
-		FVectorVMOpInfo(EOp::clamp,EOpFlags::None,EOpSrc::Register,EOpSrc::Register,EOpSrc::Const,TEXT("clampri")),
-		FVectorVMOpInfo(EOp::clamp,EOpFlags::None,EOpSrc::Register,EOpSrc::Const,EOpSrc::Const,TEXT("clampii")),
+		FVectorVMOpInfo(),
+		FVectorVMOpInfo(),
+		FVectorVMOpInfo(),
 
 		FVectorVMOpInfo(EOp::min,EOpFlags::Implemented,EOpSrc::Register,EOpSrc::Register,EOpSrc::Invalid,TEXT("Min")),
-		FVectorVMOpInfo(EOp::min,EOpFlags::None,EOpSrc::Register,EOpSrc::Const,EOpSrc::Invalid,TEXT("mini")),
+		FVectorVMOpInfo(),
 
 		FVectorVMOpInfo(EOp::max,EOpFlags::Implemented,EOpSrc::Register,EOpSrc::Register,EOpSrc::Invalid,TEXT("Max")),
-		FVectorVMOpInfo(EOp::max,EOpFlags::None,EOpSrc::Register,EOpSrc::Const,EOpSrc::Invalid,TEXT("maxi")),
+		FVectorVMOpInfo(),
 
 		FVectorVMOpInfo(EOp::pow,EOpFlags::Implemented,EOpSrc::Register,EOpSrc::Register,EOpSrc::Invalid,TEXT("Pow")),
-		FVectorVMOpInfo(EOp::pow,EOpFlags::None,EOpSrc::Register,EOpSrc::Const,EOpSrc::Invalid,TEXT("powi")),
+		FVectorVMOpInfo(),
 
 		FVectorVMOpInfo(EOp::add,EOpFlags::None,EOpSrc::Register,EOpSrc::Invalid,EOpSrc::Invalid,TEXT("Sign")),
 
 		FVectorVMOpInfo(EOp::step,EOpFlags::None,EOpSrc::Register,EOpSrc::Register,EOpSrc::Invalid,TEXT("Step")),
-		FVectorVMOpInfo(EOp::step,EOpFlags::None,EOpSrc::Register,EOpSrc::Const,EOpSrc::Invalid,TEXT("stepi")),
+		FVectorVMOpInfo(),
 
-		FVectorVMOpInfo(EOp::add,EOpFlags::None,EOpSrc::Invalid,EOpSrc::Invalid,EOpSrc::Invalid,TEXT("tex1d")),
+		FVectorVMOpInfo(),
 
 		FVectorVMOpInfo(EOp::dot, EOpFlags::Implemented, EOpSrc::Register, EOpSrc::Register, EOpSrc::Invalid, TEXT("Dot Product")),
 		FVectorVMOpInfo(EOp::cross, EOpFlags::Implemented|EOpFlags::Commutative, EOpSrc::Register, EOpSrc::Register, EOpSrc::Invalid, TEXT("Cross Product")),
-		FVectorVMOpInfo(EOp::cross, EOpFlags::None, EOpSrc::Register, EOpSrc::Const, EOpSrc::Invalid, TEXT("Cross Product with const")),
+		FVectorVMOpInfo(),
 
 		FVectorVMOpInfo(EOp::normalize, EOpFlags::Implemented, EOpSrc::Register, EOpSrc::Invalid, EOpSrc::Invalid, TEXT("Normalize")),
 		FVectorVMOpInfo(EOp::random, EOpFlags::Implemented, EOpSrc::Const, EOpSrc::Invalid, EOpSrc::Invalid, TEXT("Random")),
 
 		FVectorVMOpInfo(EOp::length, EOpFlags::Implemented, EOpSrc::Register, EOpSrc::Invalid, EOpSrc::Invalid, TEXT("Vector Length")),
-		FVectorVMOpInfo(EOp::length, EOpFlags::None, EOpSrc::Const, EOpSrc::Invalid, EOpSrc::Invalid, TEXT("Vector Length (const)")),
+		FVectorVMOpInfo(),
 
 		FVectorVMOpInfo(EOp::sin4, EOpFlags::Implemented, EOpSrc::Register, EOpSrc::Invalid, EOpSrc::Invalid, TEXT("Sin4")),
-		FVectorVMOpInfo(EOp::sin4, EOpFlags::None, EOpSrc::Const, EOpSrc::Invalid, EOpSrc::Invalid, TEXT("sin4i")),
+		FVectorVMOpInfo(),
 
-		FVectorVMOpInfo(EOp::add, EOpFlags::None, EOpSrc::Const, EOpSrc::Const, EOpSrc::Invalid, TEXT("addii")),
-		FVectorVMOpInfo(EOp::sub, EOpFlags::None, EOpSrc::Const, EOpSrc::Const, EOpSrc::Invalid, TEXT("subii")),
-		FVectorVMOpInfo(EOp::mul, EOpFlags::None, EOpSrc::Const, EOpSrc::Const, EOpSrc::Invalid, TEXT("mulii")),
+		FVectorVMOpInfo(),
+		FVectorVMOpInfo(),
+		FVectorVMOpInfo(),
 		FVectorVMOpInfo(EOp::noise, EOpFlags::Implemented, EOpSrc::Register, EOpSrc::Invalid, EOpSrc::Invalid, TEXT("Noise")),
 
 		FVectorVMOpInfo(EOp::add, EOpFlags::None, EOpSrc::Invalid, EOpSrc::Invalid, EOpSrc::Invalid, TEXT("invalid"))
diff --git a/Engine/Source/Runtime/VectorVM/Private/VectorVMPrivate.h b/Engine/Source/Runtime/VectorVM/Private/VectorVMPrivate.h
--- a/Engine/Source/Runtime/VectorVM/Private/VectorVMPrivate.h
+++ b/Engine/Source/Runtime/VectorVM/Private/VectorVMPrivate.h
@@ -12,7 +12,7 @@ namespace VectorVM
 	/** Constants. */
 	enum
 	{
-		ChunkSize = 64,
+		ChunkSize = 128,
 		ElementsPerVector = 4,
 		VectorsPerChunk = ChunkSize / ElementsPerVector,
 	};
diff --git a/Engine/Source/Runtime/VectorVM/Public/VectorVM.h b/Engine/Source/Runtime/VectorVM/Public/VectorVM.h
--- a/Engine/Source/Runtime/VectorVM/Public/VectorVM.h
+++ b/Engine/Source/Runtime/VectorVM/Public/VectorVM.h
@@ -14,8 +14,8 @@ namespace VectorVM
 	/** Constants. */
 	enum
 	{
-		NumTempRegisters = 64,
-		MaxInputRegisters = 64,
+		NumTempRegisters = 100,
+		MaxInputRegisters = 100,
 		MaxOutputRegisters = MaxInputRegisters,
 		MaxConstants = 256,
 		FirstInputRegister = NumTempRegisters,
@@ -31,72 +31,45 @@ namespace VectorVM
 		enum Type
 		{
 			done = 0,
-			add,
-			addi,
-			sub,
-			subi,
-			mul,
-			muli,
-			mad,
-			madrri,
-			madrir,
-			madrii,
-			madiir, // Remove me when a constant expressions can be evaluated at a coarser granularity.
-			madiii, // Remove me when a constant expressions can be evaluated at a coarser granularity.
-			lerp,
-			lerpirr, 
-			lerprir, 
-			lerprri,
-			lerpiir,
-			rcp,
-			rsq,
-			sqrt,
-			neg,
-			abs,
-			exp,
-			exp2,
-			log,
-			log2,
-			sin,
-			sini,
-			cos,
-			tan,
-			asin,
-			acos,
-			atan,
-			atan2,
-			ceil,
-			floor,
-			fmod,
-			frac,
-			trunc,
-			clamp,
-			clampir,
-			clampri,
-			clampii,
-			min,
-			mini,
-			max,
-			maxi,
-			pow,
-			powi,
-			sign,
-			step,
-			stepi,
-			tex1d,
-			dot,
-			cross,
-			crossi,
-			normalize,
-			random,
-			length,
-			lengthi,
-			sin4, 
-			sin4i,
-			addii,
-			subii,
-			mulii,
-			noise,
+			add = 1,
+			sub = 3,
+			mul = 5,
+			mad = 7,
+			lerp = 13,
+			rcp = 18,
+			rsq = 19,
+			sqrt = 20,
+			neg = 21,
+			abs = 22,
+			exp = 23,
+			exp2 = 24,
+			log = 25,
+			log2 = 26,
+			sin = 27,
+			cos = 29,
+			tan = 30,
+			asin = 31,
+			acos = 32,
+			atan = 33,
+			atan2 = 34,
+			ceil = 35,
+			floor = 36,
+			fmod = 37,
+			frac = 38,
+			trunc = 39,
+			clamp = 40,
+			min = 44,
+			max = 46,
+			pow = 48,
+			sign = 50,
+			step = 51,
+			dot = 54,
+			cross = 55,
+			normalize = 57,
+			random = 58,
+			length = 59,
+			sin4 = 61, 
+			noise = 66,
 			NumOpcodes
 		};
 	} // namespace EOp
@@ -136,12 +109,12 @@ namespace VectorVM
 
 		/** Constructor */
 		FVectorVMOpInfo(
-			EOp::Type InBaseOpcode,
-			uint8 InFlags,
-			EOpSrc::Type Src0,
-			EOpSrc::Type Src1,
-			EOpSrc::Type Src2,
-			const FString& InFriendlyName
+			EOp::Type InBaseOpcode = EOp::done,
+			uint8 InFlags = EOpFlags::None,
+			EOpSrc::Type Src0 = EOpSrc::Type::Invalid,
+			EOpSrc::Type Src1 = EOpSrc::Type::Invalid,
+			EOpSrc::Type Src2 = EOpSrc::Type::Invalid,
+			const FString& InFriendlyName = TEXT("Undefined Opcode")
 			)
 			: BaseOpcode(InBaseOpcode)
 			, Flags(InFlags)
