#ue4 - Avoid continuous updates of CameraComponent rotation when attached to a Pawn unless the rotation changes by a measurable amount.

- Rotations in this case often had different normalization (0..360 versus -180..180)
- Quat->Rotator conversion could also cause comparison to always report inequality.

diff --git a/Engine/Source/Runtime/Engine/Private/Camera/CameraComponent.cpp b/Engine/Source/Runtime/Engine/Private/Camera/CameraComponent.cpp
--- a/Engine/Source/Runtime/Engine/Private/Camera/CameraComponent.cpp
+++ b/Engine/Source/Runtime/Engine/Private/Camera/CameraComponent.cpp
@@ -149,7 +149,7 @@ void UCameraComponent::GetCameraView(float DeltaTime, FMinimalViewInfo& DesiredV
 		if (APawn* OwningPawn = Cast<APawn>(GetOwner()))
 		{
 			const FRotator PawnViewRotation = OwningPawn->GetViewRotation();
-			if (PawnViewRotation != GetComponentRotation())
+			if (!PawnViewRotation.Equals(GetComponentRotation()))
 			{
 				SetWorldRotation(PawnViewRotation);
 			}
