Fix stale data crash when setting the destructible mesh without recreating physics state.

TTP#345239

diff --git a/Engine/Source/Runtime/Engine/Private/DestructibleComponent.cpp b/Engine/Source/Runtime/Engine/Private/DestructibleComponent.cpp
--- a/Engine/Source/Runtime/Engine/Private/DestructibleComponent.cpp
+++ b/Engine/Source/Runtime/Engine/Private/DestructibleComponent.cpp
@@ -689,7 +689,7 @@ void UDestructibleComponent::SetDestructibleMesh(class UDestructibleMesh* NewMes
 	this->DestructibleMesh = GetDestructibleMesh();
 #endif // WITH_EDITORONLY_DATA
 	
-	SetChunkVisible(0, true);
+	RecreatePhysicsState();
 }
 
 class UDestructibleMesh * UDestructibleComponent::GetDestructibleMesh()
