#ue4 - Remove bounding of 2D speed in CharacterMovementComponent::PhysFalling(), it is redundant with that done in CalcVelocity().

- This caused issues if custom CalcVelocity() implementations wanted to use their own implementation with different speed constraints.

diff --git a/Engine/Source/Runtime/Engine/Private/CharacterMovementComponent.cpp b/Engine/Source/Runtime/Engine/Private/CharacterMovementComponent.cpp
--- a/Engine/Source/Runtime/Engine/Private/CharacterMovementComponent.cpp
+++ b/Engine/Source/Runtime/Engine/Private/CharacterMovementComponent.cpp
@@ -2815,10 +2815,6 @@ void UCharacterMovementComponent::PhysFalling(float deltaTime, int32 Iterations)
 		return;
 	}
 
-	// Bound final 2d portion of velocity
-	const float Speed2d = Velocity.Size2D();
-	const float BoundSpeed2d = FMath::Max(Speed2d, GetMaxSpeed() * AnalogInputModifier);
-
 	FVector FallAcceleration = GetFallingLateralAcceleration(deltaTime);
 	FallAcceleration.Z = 0.f;
 
@@ -2853,9 +2849,6 @@ void UCharacterMovementComponent::PhysFalling(float deltaTime, int32 Iterations)
 			Velocity.Z = 0.f;
 			CalcVelocity(timeTick, FallingLateralFriction, false, BrakingDecelerationFalling);
 			Velocity.Z = SavedVelZ;
-
-			// make sure we're not exceeding acceptable speed
-			Velocity = Velocity.ClampMaxSize2D(BoundSpeed2d);
 		}
 
 		// Apply gravity
