Add 'raycast' and 'sweep' keywords to BP query functions
Stop Kismet library headers being globally included

diff --git a/Engine/Plugins/Experimental/LiveEditor/Source/LiveEditor/Private/K2Node_LiveEditObject.cpp b/Engine/Plugins/Experimental/LiveEditor/Source/LiveEditor/Private/K2Node_LiveEditObject.cpp
--- a/Engine/Plugins/Experimental/LiveEditor/Source/LiveEditor/Private/K2Node_LiveEditObject.cpp
+++ b/Engine/Plugins/Experimental/LiveEditor/Source/LiveEditor/Private/K2Node_LiveEditObject.cpp
@@ -1,6 +1,7 @@
 
 #include "LiveEditorPrivatePCH.h"
 
+#include "Kismet/KismetMathLibrary.h"
 #include "BlueprintGraphDefinitions.h"
 #include "BlueprintUtilities.h"
 #include "Kismet2/BlueprintEditorUtils.h"
diff --git a/Engine/Source/Editor/BlueprintGraph/Private/EdGraphSchema_K2.cpp b/Engine/Source/Editor/BlueprintGraph/Private/EdGraphSchema_K2.cpp
--- a/Engine/Source/Editor/BlueprintGraph/Private/EdGraphSchema_K2.cpp
+++ b/Engine/Source/Editor/BlueprintGraph/Private/EdGraphSchema_K2.cpp
@@ -8,6 +8,8 @@
 #include "BlueprintGraphPrivatePCH.h"
 
 #include "Engine/LevelScriptBlueprint.h"
+#include "Kismet/KismetSystemLibrary.h"
+#include "Kismet/KismetArrayLibrary.h"
 #include "GraphEditorActions.h"
 #include "GraphEditorSettings.h"
 #include "ScopedTransaction.h"
diff --git a/Engine/Source/Editor/BlueprintGraph/Private/K2Node_CallFunction.cpp b/Engine/Source/Editor/BlueprintGraph/Private/K2Node_CallFunction.cpp
--- a/Engine/Source/Editor/BlueprintGraph/Private/K2Node_CallFunction.cpp
+++ b/Engine/Source/Editor/BlueprintGraph/Private/K2Node_CallFunction.cpp
@@ -6,6 +6,8 @@
 #include "CompilerResultsLog.h"
 #include "CallFunctionHandler.h"
 #include "K2Node_SwitchEnum.h"
+#include "Kismet/KismetMathLibrary.h"
+#include "Kismet/KismetArrayLibrary.h"
 
 #define LOCTEXT_NAMESPACE "K2Node"
 
diff --git a/Engine/Source/Editor/BlueprintGraph/Private/K2Node_EaseFunction.cpp b/Engine/Source/Editor/BlueprintGraph/Private/K2Node_EaseFunction.cpp
--- a/Engine/Source/Editor/BlueprintGraph/Private/K2Node_EaseFunction.cpp
+++ b/Engine/Source/Editor/BlueprintGraph/Private/K2Node_EaseFunction.cpp
@@ -12,6 +12,7 @@
 #include "KismetCompiler.h"
 #include "ScopedTransaction.h"
 #include "EdGraph/EdGraphNodeUtils.h" // for FNodeTextCache
+#include "Kismet/KismetMathLibrary.h"
 
 #define LOCTEXT_NAMESPACE "K2Node_EaseFunction"
 
diff --git a/Engine/Source/Editor/BlueprintGraph/Private/K2Node_FormatText.cpp b/Engine/Source/Editor/BlueprintGraph/Private/K2Node_FormatText.cpp
--- a/Engine/Source/Editor/BlueprintGraph/Private/K2Node_FormatText.cpp
+++ b/Engine/Source/Editor/BlueprintGraph/Private/K2Node_FormatText.cpp
@@ -2,6 +2,7 @@
 
 
 #include "BlueprintGraphPrivatePCH.h"
+#include "Kismet/KismetTextLibrary.h"
 #include "Editor/GraphEditor/Public/DiffResults.h"
 #include "Kismet2NameValidators.h"
 #include "KismetCompiler.h"
diff --git a/Engine/Source/Editor/BlueprintGraph/Private/K2Node_Message.cpp b/Engine/Source/Editor/BlueprintGraph/Private/K2Node_Message.cpp
--- a/Engine/Source/Editor/BlueprintGraph/Private/K2Node_Message.cpp
+++ b/Engine/Source/Editor/BlueprintGraph/Private/K2Node_Message.cpp
@@ -4,6 +4,8 @@
 #include "KismetCompiler.h"
 #include "K2Node_CallArrayFunction.h"
 #include "K2Node_Message.h"
+#include "Kismet/KismetArrayLibrary.h"
+
 #define LOCTEXT_NAMESPACE "K2Node_Message"
 
 UK2Node_Message::UK2Node_Message(const class FPostConstructInitializeProperties& PCIP)
diff --git a/Engine/Source/Editor/BlueprintGraph/Private/K2Node_SwitchString.cpp b/Engine/Source/Editor/BlueprintGraph/Private/K2Node_SwitchString.cpp
--- a/Engine/Source/Editor/BlueprintGraph/Private/K2Node_SwitchString.cpp
+++ b/Engine/Source/Editor/BlueprintGraph/Private/K2Node_SwitchString.cpp
@@ -2,6 +2,7 @@
 
 #include "BlueprintGraphPrivatePCH.h"
 #include "Kismet/KismetMathLibrary.h"
+#include "Kismet/KismetStringLibrary.h"
 #include "K2Node_SwitchString.h"
 #include "BlueprintNodeSpawner.h"
 #include "BlueprintActionDatabaseRegistrar.h"
diff --git a/Engine/Source/Editor/GameplayAbilitiesEditor/Private/K2Node_LatentAbilityCall.cpp b/Engine/Source/Editor/GameplayAbilitiesEditor/Private/K2Node_LatentAbilityCall.cpp
--- a/Engine/Source/Editor/GameplayAbilitiesEditor/Private/K2Node_LatentAbilityCall.cpp
+++ b/Engine/Source/Editor/GameplayAbilitiesEditor/Private/K2Node_LatentAbilityCall.cpp
@@ -2,6 +2,8 @@
 
 
 #include "AbilitySystemEditorPrivatePCH.h"
+#include "Kismet/KismetMathLibrary.h"
+#include "Kismet/KismetArrayLibrary.h"
 #include "AbilityTask.h"
 #include "KismetCompiler.h"
 #include "BlueprintEditorUtils.h"
diff --git a/Engine/Source/Editor/Kismet/Private/SKismetInspector.cpp b/Engine/Source/Editor/Kismet/Private/SKismetInspector.cpp
--- a/Engine/Source/Editor/Kismet/Private/SKismetInspector.cpp
+++ b/Engine/Source/Editor/Kismet/Private/SKismetInspector.cpp
@@ -2,6 +2,8 @@
 
 
 #include "BlueprintEditorPrivatePCH.h"
+#include "Kismet/KismetStringLibrary.h"
+#include "Kismet/KismetSystemLibrary.h"
 #include "GraphEditor.h"
 #include "BlueprintUtilities.h"
 #include "AnimGraphDefinitions.h"
diff --git a/Engine/Source/Editor/KismetCompiler/Private/KismetCompiler.cpp b/Engine/Source/Editor/KismetCompiler/Private/KismetCompiler.cpp
--- a/Engine/Source/Editor/KismetCompiler/Private/KismetCompiler.cpp
+++ b/Engine/Source/Editor/KismetCompiler/Private/KismetCompiler.cpp
@@ -20,6 +20,7 @@
 #include "K2Node_EnumLiteral.h"
 #include "EdGraph/EdGraphNode_Documentation.h"
 #include "Engine/DynamicBlueprintBinding.h"
+#include "Kismet/KismetMathLibrary.h"
 
 static bool bDebugPropertyPropagation = false;
 
diff --git a/Engine/Source/Editor/UnrealEd/Private/Tests/EditorBuildPromotionTests.cpp b/Engine/Source/Editor/UnrealEd/Private/Tests/EditorBuildPromotionTests.cpp
--- a/Engine/Source/Editor/UnrealEd/Private/Tests/EditorBuildPromotionTests.cpp
+++ b/Engine/Source/Editor/UnrealEd/Private/Tests/EditorBuildPromotionTests.cpp
@@ -36,6 +36,8 @@
 #include "KismetDebugUtilities.h"
 #include "Engine/Breakpoint.h"
 #include "Engine/LevelScriptBlueprint.h"
+#include "Kismet/KismetStringLibrary.h"
+#include "Kismet/KismetSystemLibrary.h"
 #include "BlueprintEditorUtils.h"
 
 //Utils
diff --git a/Engine/Source/Runtime/Engine/Classes/Kismet/KismetSystemLibrary.h b/Engine/Source/Runtime/Engine/Classes/Kismet/KismetSystemLibrary.h
--- a/Engine/Source/Runtime/Engine/Classes/Kismet/KismetSystemLibrary.h
+++ b/Engine/Source/Runtime/Engine/Classes/Kismet/KismetSystemLibrary.h
@@ -558,7 +558,7 @@ class UKismetSystemLibrary : public UBlueprintFunctionLibrary
 	 * @param OutHit		Properties of the trace hit.
 	 * @return				True if there was a hit, false otherwise.
 	 */
-	UFUNCTION(BlueprintCallable, Category="Collision", meta=(bIgnoreSelf="true", WorldContext="WorldContextObject", AutoCreateRefTerm="ActorsToIgnore", FriendlyName = "SingleLineTraceByChannel"))
+	UFUNCTION(BlueprintCallable, Category="Collision", meta=(bIgnoreSelf="true", WorldContext="WorldContextObject", AutoCreateRefTerm="ActorsToIgnore", FriendlyName="SingleLineTraceByChannel", Keywords="raycast"))
 	static ENGINE_API bool LineTraceSingle_NEW(UObject* WorldContextObject, const FVector Start, const FVector End, ETraceTypeQuery TraceChannel, bool bTraceComplex, const TArray<AActor*>& ActorsToIgnore, EDrawDebugTrace::Type DrawDebugType, FHitResult& OutHit, bool bIgnoreSelf);
 	
 	/**
@@ -573,7 +573,7 @@ class UKismetSystemLibrary : public UBlueprintFunctionLibrary
 	 * @param OutHit		Properties of the trace hit.
 	 * @return				True if there was a hit, false otherwise.
 	 */
-	UFUNCTION(BlueprintCallable, Category="Collision", meta=(bIgnoreSelf="true", WorldContext="WorldContextObject", AutoCreateRefTerm="ActorsToIgnore", FriendlyName = "MultiLineTraceByChannel"))
+	UFUNCTION(BlueprintCallable, Category="Collision", meta=(bIgnoreSelf="true", WorldContext="WorldContextObject", AutoCreateRefTerm="ActorsToIgnore", FriendlyName = "MultiLineTraceByChannel", Keywords="raycast"))
 	static ENGINE_API bool LineTraceMulti_NEW(UObject* WorldContextObject, const FVector Start, const FVector End, ETraceTypeQuery TraceChannel, bool bTraceComplex, const TArray<AActor*>& ActorsToIgnore, EDrawDebugTrace::Type DrawDebugType, TArray<FHitResult>& OutHits, bool bIgnoreSelf);
 
 	/**
@@ -588,7 +588,7 @@ class UKismetSystemLibrary : public UBlueprintFunctionLibrary
 	 * @param OutHit		Properties of the trace hit.
 	 * @return				True if there was a hit, false otherwise.
 	 */
-	UFUNCTION(BlueprintCallable, Category="Collision", meta=(bIgnoreSelf="true", WorldContext="WorldContextObject", AutoCreateRefTerm="ActorsToIgnore", FriendlyName = "SingleSphereTraceByChannel"))
+	UFUNCTION(BlueprintCallable, Category="Collision", meta=(bIgnoreSelf="true", WorldContext="WorldContextObject", AutoCreateRefTerm="ActorsToIgnore", FriendlyName = "SingleSphereTraceByChannel", Keywords="sweep"))
 	static ENGINE_API bool SphereTraceSingle_NEW(UObject* WorldContextObject, const FVector Start, const FVector End, float Radius, ETraceTypeQuery TraceChannel, bool bTraceComplex, const TArray<AActor*>& ActorsToIgnore, EDrawDebugTrace::Type DrawDebugType, FHitResult& OutHit, bool bIgnoreSelf);
 
 	/**
@@ -604,7 +604,7 @@ class UKismetSystemLibrary : public UBlueprintFunctionLibrary
 	 * @param OutHits		A list of hits, sorted along the trace from start to finish.  The blocking hit will be the last hit, if there was one.
 	 * @return				True if there was a hit, false otherwise.
 	 */
-	UFUNCTION(BlueprintCallable, Category="Collision", meta=(bIgnoreSelf="true", WorldContext="WorldContextObject", AutoCreateRefTerm="ActorsToIgnore", FriendlyName = "MultiSphereTraceByChannel"))
+	UFUNCTION(BlueprintCallable, Category="Collision", meta=(bIgnoreSelf="true", WorldContext="WorldContextObject", AutoCreateRefTerm="ActorsToIgnore", FriendlyName = "MultiSphereTraceByChannel", Keywords="sweep"))
 	static ENGINE_API bool SphereTraceMulti_NEW(UObject* WorldContextObject, const FVector Start, const FVector End, float Radius, ETraceTypeQuery TraceChannel, bool bTraceComplex, const TArray<AActor*>& ActorsToIgnore, EDrawDebugTrace::Type DrawDebugType, TArray<FHitResult>& OutHits, bool bIgnoreSelf);
 
 	/**
@@ -620,7 +620,7 @@ class UKismetSystemLibrary : public UBlueprintFunctionLibrary
 	* @param OutHit			Properties of the trace hit.
 	* @return				True if there was a hit, false otherwise.
 	*/
-	UFUNCTION(BlueprintCallable, Category = "Collision", meta = (bIgnoreSelf = "true", HidePin = "WorldContextObject", DefaultToSelf = "WorldContextObject", AutoCreateRefTerm = "ActorsToIgnore", FriendlyName = "SingleBoxTraceByChannel"))
+	UFUNCTION(BlueprintCallable, Category = "Collision", meta = (bIgnoreSelf = "true", HidePin = "WorldContextObject", DefaultToSelf = "WorldContextObject", AutoCreateRefTerm = "ActorsToIgnore", FriendlyName = "SingleBoxTraceByChannel", Keywords="sweep"))
 	static ENGINE_API bool BoxTraceSingle(UObject* WorldContextObject, const FVector Start, const FVector End, const FVector HalfSize, const FRotator Orientation, ETraceTypeQuery TraceChannel, bool bTraceComplex, const TArray<AActor*>& ActorsToIgnore, EDrawDebugTrace::Type DrawDebugType, FHitResult& OutHit, bool bIgnoreSelf);
 
 	/**
@@ -636,7 +636,7 @@ class UKismetSystemLibrary : public UBlueprintFunctionLibrary
 	* @param OutHits		A list of hits, sorted along the trace from start to finish. The blocking hit will be the last hit, if there was one.
 	* @return				True if there was a hit, false otherwise.
 	*/
-	UFUNCTION(BlueprintCallable, Category = "Collision", meta = (bIgnoreSelf = "true", HidePin = "WorldContextObject", DefaultToSelf = "WorldContextObject", AutoCreateRefTerm = "ActorsToIgnore", FriendlyName = "MultiBoxTraceByChannel"))
+	UFUNCTION(BlueprintCallable, Category = "Collision", meta = (bIgnoreSelf = "true", HidePin = "WorldContextObject", DefaultToSelf = "WorldContextObject", AutoCreateRefTerm = "ActorsToIgnore", FriendlyName = "MultiBoxTraceByChannel", Keywords="sweep"))
 	static ENGINE_API bool BoxTraceMulti(UObject* WorldContextObject, const FVector Start, const FVector End, FVector HalfSize, const FRotator Orientation, ETraceTypeQuery TraceChannel, bool bTraceComplex, const TArray<AActor*>& ActorsToIgnore, EDrawDebugTrace::Type DrawDebugType, TArray<FHitResult>& OutHits, bool bIgnoreSelf);
 
 	
@@ -654,7 +654,7 @@ class UKismetSystemLibrary : public UBlueprintFunctionLibrary
 	 * @param OutHit		Properties of the trace hit.
 	 * @return				True if there was a hit, false otherwise.
 	 */
-	UFUNCTION(BlueprintCallable, Category="Collision", meta=(bIgnoreSelf="true", WorldContext="WorldContextObject", AutoCreateRefTerm="ActorsToIgnore", FriendlyName = "SingleCapsuleTraceByChannel"))
+	UFUNCTION(BlueprintCallable, Category="Collision", meta=(bIgnoreSelf="true", WorldContext="WorldContextObject", AutoCreateRefTerm="ActorsToIgnore", FriendlyName = "SingleCapsuleTraceByChannel", Keywords="sweep"))
 	static ENGINE_API bool CapsuleTraceSingle_NEW(UObject* WorldContextObject, const FVector Start, const FVector End, float Radius, float HalfHeight, ETraceTypeQuery TraceChannel, bool bTraceComplex, const TArray<AActor*>& ActorsToIgnore, EDrawDebugTrace::Type DrawDebugType, FHitResult& OutHit, bool bIgnoreSelf);
 
 	/**
@@ -671,7 +671,7 @@ class UKismetSystemLibrary : public UBlueprintFunctionLibrary
 	 * @param OutHits		A list of hits, sorted along the trace from start to finish.  The blocking hit will be the last hit, if there was one.
 	 * @return				True if there was a hit, false otherwise.
 	 */
-	UFUNCTION(BlueprintCallable, Category="Collision", meta=(bIgnoreSelf="true", WorldContext="WorldContextObject", AutoCreateRefTerm="ActorsToIgnore", FriendlyName = "MultiCapsuleTraceByChannel"))
+	UFUNCTION(BlueprintCallable, Category="Collision", meta=(bIgnoreSelf="true", WorldContext="WorldContextObject", AutoCreateRefTerm="ActorsToIgnore", FriendlyName = "MultiCapsuleTraceByChannel", Keywords="sweep"))
 	static ENGINE_API bool CapsuleTraceMulti_NEW(UObject* WorldContextObject, const FVector Start, const FVector End, float Radius, float HalfHeight, ETraceTypeQuery TraceChannel, bool bTraceComplex, const TArray<AActor*>& ActorsToIgnore, EDrawDebugTrace::Type DrawDebugType, TArray<FHitResult>& OutHits, bool bIgnoreSelf);
 
 	/**
@@ -686,7 +686,7 @@ class UKismetSystemLibrary : public UBlueprintFunctionLibrary
 	 * @param OutHit		Properties of the trace hit.
 	 * @return				True if there was a hit, false otherwise.
 	 */
-	UFUNCTION(BlueprintCallable, Category="Collision", meta=(bIgnoreSelf="true", WorldContext="WorldContextObject", AutoCreateRefTerm="ActorsToIgnore", FriendlyName = "SingleLineTraceForObjects"))
+	UFUNCTION(BlueprintCallable, Category="Collision", meta=(bIgnoreSelf="true", WorldContext="WorldContextObject", AutoCreateRefTerm="ActorsToIgnore", FriendlyName = "SingleLineTraceForObjects", Keywords="raycast"))
 	static ENGINE_API bool LineTraceSingleForObjects(UObject* WorldContextObject, const FVector Start, const FVector End, const TArray<TEnumAsByte<EObjectTypeQuery> > & ObjectTypes, bool bTraceComplex, const TArray<AActor*>& ActorsToIgnore, EDrawDebugTrace::Type DrawDebugType, FHitResult& OutHit, bool bIgnoreSelf);
 	
 	/**
@@ -701,7 +701,7 @@ class UKismetSystemLibrary : public UBlueprintFunctionLibrary
 	 * @param OutHit		Properties of the trace hit.
 	 * @return				True if there was a hit, false otherwise.
 	 */
-	UFUNCTION(BlueprintCallable, Category="Collision", meta=(bIgnoreSelf="true", WorldContext="WorldContextObject", AutoCreateRefTerm="ActorsToIgnore", FriendlyName = "MultiLineTraceForObjects"))
+	UFUNCTION(BlueprintCallable, Category="Collision", meta=(bIgnoreSelf="true", WorldContext="WorldContextObject", AutoCreateRefTerm="ActorsToIgnore", FriendlyName = "MultiLineTraceForObjects", Keywords="raycast"))
 	static ENGINE_API bool LineTraceMultiForObjects(UObject* WorldContextObject, const FVector Start, const FVector End, const TArray<TEnumAsByte<EObjectTypeQuery> > & ObjectTypes, bool bTraceComplex, const TArray<AActor*>& ActorsToIgnore, EDrawDebugTrace::Type DrawDebugType, TArray<FHitResult>& OutHits, bool bIgnoreSelf);
 
 	/**
@@ -716,7 +716,7 @@ class UKismetSystemLibrary : public UBlueprintFunctionLibrary
 	 * @param OutHit		Properties of the trace hit.
 	 * @return				True if there was a hit, false otherwise.
 	 */
-	UFUNCTION(BlueprintCallable, Category="Collision", meta=(bIgnoreSelf="true", WorldContext="WorldContextObject", AutoCreateRefTerm="ActorsToIgnore", FriendlyName = "SingleSphereTraceForObjects"))
+	UFUNCTION(BlueprintCallable, Category="Collision", meta=(bIgnoreSelf="true", WorldContext="WorldContextObject", AutoCreateRefTerm="ActorsToIgnore", FriendlyName = "SingleSphereTraceForObjects", Keywords="sweep"))
 	static ENGINE_API bool SphereTraceSingleForObjects(UObject* WorldContextObject, const FVector Start, const FVector End, float Radius, const TArray<TEnumAsByte<EObjectTypeQuery> > & ObjectTypes, bool bTraceComplex, const TArray<AActor*>& ActorsToIgnore, EDrawDebugTrace::Type DrawDebugType, FHitResult& OutHit, bool bIgnoreSelf);
 
 	/**
@@ -732,7 +732,7 @@ class UKismetSystemLibrary : public UBlueprintFunctionLibrary
 	 * @param OutHits		A list of hits, sorted along the trace from start to finish.  The blocking hit will be the last hit, if there was one.
 	 * @return				True if there was a hit, false otherwise.
 	 */
-	UFUNCTION(BlueprintCallable, Category="Collision", meta=(bIgnoreSelf="true", WorldContext="WorldContextObject", AutoCreateRefTerm="ActorsToIgnore", FriendlyName = "MultiSphereTraceForObjects"))
+	UFUNCTION(BlueprintCallable, Category="Collision", meta=(bIgnoreSelf="true", WorldContext="WorldContextObject", AutoCreateRefTerm="ActorsToIgnore", FriendlyName = "MultiSphereTraceForObjects", Keywords="sweep"))
 	static ENGINE_API bool SphereTraceMultiForObjects(UObject* WorldContextObject, const FVector Start, const FVector End, float Radius, const TArray<TEnumAsByte<EObjectTypeQuery> > & ObjectTypes, bool bTraceComplex, const TArray<AActor*>& ActorsToIgnore, EDrawDebugTrace::Type DrawDebugType, TArray<FHitResult>& OutHits, bool bIgnoreSelf);
 	
 
@@ -749,7 +749,7 @@ class UKismetSystemLibrary : public UBlueprintFunctionLibrary
 	* @param OutHit			Properties of the trace hit.
 	* @return				True if there was a hit, false otherwise.
 	*/
-	UFUNCTION(BlueprintCallable, Category = "Collision", meta = (bIgnoreSelf = "true", HidePin = "WorldContextObject", DefaultToSelf = "WorldContextObject", AutoCreateRefTerm = "ActorsToIgnore", FriendlyName = "SingleBoxTraceForObjects"))
+	UFUNCTION(BlueprintCallable, Category = "Collision", meta = (bIgnoreSelf = "true", HidePin = "WorldContextObject", DefaultToSelf = "WorldContextObject", AutoCreateRefTerm = "ActorsToIgnore", FriendlyName = "SingleBoxTraceForObjects", Keywords="sweep"))
 	static ENGINE_API bool BoxTraceSingleForObjects(UObject* WorldContextObject, const FVector Start, const FVector End, const FVector HalfSize, const FRotator Orientation, const TArray<TEnumAsByte<EObjectTypeQuery> > & ObjectTypes, bool bTraceComplex, const TArray<AActor*>& ActorsToIgnore, EDrawDebugTrace::Type DrawDebugType, FHitResult& OutHit, bool bIgnoreSelf);
 
 
@@ -766,7 +766,7 @@ class UKismetSystemLibrary : public UBlueprintFunctionLibrary
 	* @param OutHits		A list of hits, sorted along the trace from start to finish.  The blocking hit will be the last hit, if there was one.
 	* @return				True if there was a hit, false otherwise.
 	*/
-	UFUNCTION(BlueprintCallable, Category = "Collision", meta = (bIgnoreSelf = "true", HidePin = "WorldContextObject", DefaultToSelf = "WorldContextObject", AutoCreateRefTerm = "ActorsToIgnore", FriendlyName = "MultiBoxTraceForObjects"))
+	UFUNCTION(BlueprintCallable, Category = "Collision", meta = (bIgnoreSelf = "true", HidePin = "WorldContextObject", DefaultToSelf = "WorldContextObject", AutoCreateRefTerm = "ActorsToIgnore", FriendlyName = "MultiBoxTraceForObjects", Keywords="sweep"))
 	static ENGINE_API bool BoxTraceMultiForObjects(UObject* WorldContextObject, const FVector Start, const FVector End, const FVector HalfSize, const FRotator Orientation, const TArray<TEnumAsByte<EObjectTypeQuery> > & ObjectTypes, bool bTraceComplex, const TArray<AActor*>& ActorsToIgnore, EDrawDebugTrace::Type DrawDebugType, TArray<FHitResult>& OutHits, bool bIgnoreSelf);
 
 	/**
@@ -783,7 +783,7 @@ class UKismetSystemLibrary : public UBlueprintFunctionLibrary
 	 * @param OutHit		Properties of the trace hit.
 	 * @return				True if there was a hit, false otherwise.
 	 */
-	UFUNCTION(BlueprintCallable, Category="Collision", meta=(bIgnoreSelf="true", WorldContext="WorldContextObject", AutoCreateRefTerm="ActorsToIgnore", FriendlyName = "SingleCapsuleTraceForObjects"))
+	UFUNCTION(BlueprintCallable, Category="Collision", meta=(bIgnoreSelf="true", WorldContext="WorldContextObject", AutoCreateRefTerm="ActorsToIgnore", FriendlyName = "SingleCapsuleTraceForObjects", Keywords="sweep"))
 	static ENGINE_API bool CapsuleTraceSingleForObjects(UObject* WorldContextObject, const FVector Start, const FVector End, float Radius, float HalfHeight, const TArray<TEnumAsByte<EObjectTypeQuery> > & ObjectTypes, bool bTraceComplex, const TArray<AActor*>& ActorsToIgnore, EDrawDebugTrace::Type DrawDebugType, FHitResult& OutHit, bool bIgnoreSelf);
 
 	/**
@@ -800,7 +800,7 @@ class UKismetSystemLibrary : public UBlueprintFunctionLibrary
 	 * @param OutHits		A list of hits, sorted along the trace from start to finish.  The blocking hit will be the last hit, if there was one.
 	 * @return				True if there was a hit, false otherwise.
 	 */
-	UFUNCTION(BlueprintCallable, Category="Collision", meta=(bIgnoreSelf="true", WorldContext="WorldContextObject", AutoCreateRefTerm="ActorsToIgnore", FriendlyName = "MultiCapsuleTraceForObjects"))
+	UFUNCTION(BlueprintCallable, Category="Collision", meta=(bIgnoreSelf="true", WorldContext="WorldContextObject", AutoCreateRefTerm="ActorsToIgnore", FriendlyName = "MultiCapsuleTraceForObjects", Keywords="sweep"))
 	static ENGINE_API bool CapsuleTraceMultiForObjects(UObject* WorldContextObject, const FVector Start, const FVector End, float Radius, float HalfHeight, const TArray<TEnumAsByte<EObjectTypeQuery> > & ObjectTypes, bool bTraceComplex, const TArray<AActor*>& ActorsToIgnore, EDrawDebugTrace::Type DrawDebugType, TArray<FHitResult>& OutHits, bool bIgnoreSelf);
 
 	/**
diff --git a/Engine/Source/Runtime/Engine/Private/EnginePrivate.h b/Engine/Source/Runtime/Engine/Private/EnginePrivate.h
--- a/Engine/Source/Runtime/Engine/Private/EnginePrivate.h
+++ b/Engine/Source/Runtime/Engine/Private/EnginePrivate.h
@@ -141,14 +141,6 @@
 #include "Sound/DialogueTypes.h"
 #include "Kismet/GameplayStatics.h"
 #include "Kismet/HeadMountedDisplayFunctionLibrary.h"
-#include "Kismet/KismetArrayLibrary.h"
-#include "Kismet/KismetInputLibrary.h"
-#include "Kismet/KismetMaterialLibrary.h"
-#include "Kismet/KismetMathLibrary.h"
-#include "Kismet/KismetNodeHelperLibrary.h"
-#include "Kismet/KismetStringLibrary.h"
-#include "Kismet/KismetSystemLibrary.h"
-#include "Kismet/KismetTextLibrary.h"
 #include "Engine/BlueprintGeneratedClass.h"
 #include "PhysicsEngine/BodySetup.h"
 #include "Camera/CameraAnim.h"
diff --git a/Engine/Source/Runtime/Engine/Private/KismetArrayLibrary.cpp b/Engine/Source/Runtime/Engine/Private/KismetArrayLibrary.cpp
--- a/Engine/Source/Runtime/Engine/Private/KismetArrayLibrary.cpp
+++ b/Engine/Source/Runtime/Engine/Private/KismetArrayLibrary.cpp
@@ -1,6 +1,7 @@
 // Copyright 1998-2014 Epic Games, Inc. All Rights Reserved.
 
 #include "EnginePrivate.h"
+#include "Kismet/KismetArrayLibrary.h"
 
 DEFINE_LOG_CATEGORY_STATIC(LogArray, Warning, All);
 
diff --git a/Engine/Source/Runtime/Engine/Private/KismetInputLibrary.cpp b/Engine/Source/Runtime/Engine/Private/KismetInputLibrary.cpp
--- a/Engine/Source/Runtime/Engine/Private/KismetInputLibrary.cpp
+++ b/Engine/Source/Runtime/Engine/Private/KismetInputLibrary.cpp
@@ -1,6 +1,7 @@
 // Copyright 1998-2014 Epic Games, Inc. All Rights Reserved.
 
 #include "EnginePrivate.h"
+#include "Kismet/KismetInputLibrary.h"
 #include "MessageLog.h"
 #include "UObjectToken.h"
 
diff --git a/Engine/Source/Runtime/Engine/Private/KismetMaterialLibrary.cpp b/Engine/Source/Runtime/Engine/Private/KismetMaterialLibrary.cpp
--- a/Engine/Source/Runtime/Engine/Private/KismetMaterialLibrary.cpp
+++ b/Engine/Source/Runtime/Engine/Private/KismetMaterialLibrary.cpp
@@ -1,6 +1,7 @@
 // Copyright 1998-2014 Epic Games, Inc. All Rights Reserved.
 
 #include "EnginePrivate.h"
+#include "Kismet/KismetMaterialLibrary.h"
 #include "MessageLog.h"
 #include "UObjectToken.h"
 
diff --git a/Engine/Source/Runtime/Engine/Private/KismetMathLibrary.cpp b/Engine/Source/Runtime/Engine/Private/KismetMathLibrary.cpp
--- a/Engine/Source/Runtime/Engine/Private/KismetMathLibrary.cpp
+++ b/Engine/Source/Runtime/Engine/Private/KismetMathLibrary.cpp
@@ -1,6 +1,7 @@
 // Copyright 1998-2014 Epic Games, Inc. All Rights Reserved.
 
 #include "EnginePrivate.h"
+#include "Kismet/KismetMathLibrary.h"
 #include "CoreStats.h"
 
 /** Interpolate a linear alpha value using an ease mode and function, BlendExp used in cases where the easing is exponential */
diff --git a/Engine/Source/Runtime/Engine/Private/KismetNodeHelperLibrary.cpp b/Engine/Source/Runtime/Engine/Private/KismetNodeHelperLibrary.cpp
--- a/Engine/Source/Runtime/Engine/Private/KismetNodeHelperLibrary.cpp
+++ b/Engine/Source/Runtime/Engine/Private/KismetNodeHelperLibrary.cpp
@@ -1,6 +1,7 @@
 // Copyright 1998-2014 Epic Games, Inc. All Rights Reserved.
 
 #include "EnginePrivate.h"
+#include "Kismet/KismetNodeHelperLibrary.h"
 
 UKismetNodeHelperLibrary::UKismetNodeHelperLibrary(const class FPostConstructInitializeProperties& PCIP)
 	: Super(PCIP)
diff --git a/Engine/Source/Runtime/Engine/Private/KismetStringLibrary.cpp b/Engine/Source/Runtime/Engine/Private/KismetStringLibrary.cpp
--- a/Engine/Source/Runtime/Engine/Private/KismetStringLibrary.cpp
+++ b/Engine/Source/Runtime/Engine/Private/KismetStringLibrary.cpp
@@ -1,6 +1,7 @@
 // Copyright 1998-2014 Epic Games, Inc. All Rights Reserved.
 
 #include "EnginePrivate.h"
+#include "Kismet/KismetStringLibrary.h"
 
 //////////////////////////////////////////////////////////////////////////
 // UKismetStringLibrary
diff --git a/Engine/Source/Runtime/Engine/Private/KismetSystemLibrary.cpp b/Engine/Source/Runtime/Engine/Private/KismetSystemLibrary.cpp
--- a/Engine/Source/Runtime/Engine/Private/KismetSystemLibrary.cpp
+++ b/Engine/Source/Runtime/Engine/Private/KismetSystemLibrary.cpp
@@ -1,6 +1,7 @@
 // Copyright 1998-2014 Epic Games, Inc. All Rights Reserved.
 
 #include "EnginePrivate.h"
+#include "Kismet/KismetSystemLibrary.h"
 #include "Engine/Console.h"
 #include "LatentActions.h"
 #include "DelayAction.h"
diff --git a/Engine/Source/Runtime/Engine/Private/KismetTextLibrary.cpp b/Engine/Source/Runtime/Engine/Private/KismetTextLibrary.cpp
--- a/Engine/Source/Runtime/Engine/Private/KismetTextLibrary.cpp
+++ b/Engine/Source/Runtime/Engine/Private/KismetTextLibrary.cpp
@@ -1,6 +1,7 @@
 // Copyright 1998-2014 Epic Games, Inc. All Rights Reserved.
 
 #include "EnginePrivate.h"
+#include "Kismet/KismetTextLibrary.h"
 
 #define LOCTEXT_NAMESPACE "Kismet"
 
diff --git a/Engine/Source/Runtime/Engine/Public/Engine.h b/Engine/Source/Runtime/Engine/Public/Engine.h
--- a/Engine/Source/Runtime/Engine/Public/Engine.h
+++ b/Engine/Source/Runtime/Engine/Public/Engine.h
@@ -194,14 +194,6 @@
 #include "Animation/AnimBlueprint.h"
 #include "Sound/DialogueTypes.h"
 #include "Kismet/HeadMountedDisplayFunctionLibrary.h"
-#include "Kismet/KismetArrayLibrary.h"
-#include "Kismet/KismetInputLibrary.h"
-#include "Kismet/KismetMaterialLibrary.h"
-#include "Kismet/KismetMathLibrary.h"
-#include "Kismet/KismetNodeHelperLibrary.h"
-#include "Kismet/KismetStringLibrary.h"
-#include "Kismet/KismetSystemLibrary.h"
-#include "Kismet/KismetTextLibrary.h"
 #include "Engine/BlueprintGeneratedClass.h"
 #include "Animation/AnimBlueprintGeneratedClass.h"
 #include "PhysicsEngine/BodySetup.h"
