BP: CustomThunk functions were refactored and unified, so their implementation can be used later in native code.
FIX: BoolProperty issues with arrays of bools.

diff --git a/Engine/Source/Runtime/Engine/Classes/Kismet/DataTableFunctionLibrary.h b/Engine/Source/Runtime/Engine/Classes/Kismet/DataTableFunctionLibrary.h
--- a/Engine/Source/Runtime/Engine/Classes/Kismet/DataTableFunctionLibrary.h
+++ b/Engine/Source/Runtime/Engine/Classes/Kismet/DataTableFunctionLibrary.h
@@ -30,6 +30,8 @@ class UDataTableFunctionLibrary : public UBlueprintFunctionLibrary
     UFUNCTION(BlueprintCallable, CustomThunk, Category = "DataTable", meta=(CustomStructureParam = "OutRow", BlueprintInternalUseOnly="true"))
     static bool GetDataTableRowFromName(UDataTable* Table, FName RowName, FTableRowBase& OutRow);
     
+	ENGINE_API static bool Generic_GetDataTableRowFromName(UDataTable* Table, FName RowName, void* OutRowPtr);
+
     /** Based on UDataTableFunctionLibrary::GetDataTableRow */
     DECLARE_FUNCTION(execGetDataTableRowFromName)
     {
@@ -41,24 +43,6 @@ class UDataTableFunctionLibrary : public UBlueprintFunctionLibrary
         
         P_FINISH;
         
-        bool foundRow = false;
-        
-        if (OutRowPtr != NULL)
-        {
-            void* RowPtr = Table->FindRowUnchecked(RowName);
-            
-            if (RowPtr != NULL)
-            {
-                UScriptStruct* StructType = Table->RowStruct;
-                
-                if (StructType != NULL)
-                {
-                    StructType->CopyScriptStruct(OutRowPtr, RowPtr);
-                    foundRow = true;
-                }
-            }
-        }
-        
-        *(bool*)Result = foundRow;
+		*(bool*)Result = Generic_GetDataTableRowFromName(Table, RowName, OutRowPtr);
     }
 };
diff --git a/Engine/Source/Runtime/Engine/Classes/Kismet/KismetArrayLibrary.h b/Engine/Source/Runtime/Engine/Classes/Kismet/KismetArrayLibrary.h
--- a/Engine/Source/Runtime/Engine/Classes/Kismet/KismetArrayLibrary.h
+++ b/Engine/Source/Runtime/Engine/Classes/Kismet/KismetArrayLibrary.h
@@ -157,20 +157,25 @@ class UKismetArrayLibrary : public UBlueprintFunctionLibrary
 	static void SetArrayPropertyByName(UObject* Object, FName PropertyName, const TArray<int32>& Value);
 
 	// Native functions that will be called by the below custom thunk layers, which read off the property address, and call the appropriate native handler
-	int32 GenericArray_Add(void* TargetArray, const UArrayProperty* ArrayProp, const int32& NewItem);
-	void GenericArray_Shuffle(void* TargetArray, const UArrayProperty* ArrayProp);
-	void GenericArray_Append(void* TargetArray, const UArrayProperty* TargetArrayProp, void* SourceArray, const UArrayProperty* SourceArrayProperty);
-	void GenericArray_Insert(void* TargetArray, const UArrayProperty* ArrayProp, const int32& NewItem, int32 Index);
-	void GenericArray_Remove(void* TargetArray, const UArrayProperty* ArrayProp, int32 IndexToRemove);
-	bool GenericArray_RemoveItem(void* TargetArray, const UArrayProperty* ArrayProp,const int32& Item);
-	void GenericArray_Clear(void* TargetArray, const UArrayProperty* ArrayProp);
-	void GenericArray_Resize(void* TargetArray, const UArrayProperty* ArrayProp, int32 Size);
-	int32 GenericArray_Length(void* TargetArray, const UArrayProperty* ArrayProp);
-	int32 GenericArray_LastIndex(void* TargetArray, const UArrayProperty* ArrayProp);
-	void GenericArray_Get(void* TargetArray, const UArrayProperty* ArrayProp, int32 Index, int32& Item);
-	void GenericArray_Set(void* TargetArray, const UArrayProperty* ArrayProp, int32 Index, const int32& NewItem, bool bSizeToFit);
-	int32 GenericArray_Find(void* TargetArray, const UArrayProperty* ArrayProperty, const int32& ItemToFind);
-
+	ENGINE_API static int32 GenericArray_Add(void* TargetArray, const UArrayProperty* ArrayProp, const void* NewItem);
+	ENGINE_API static void GenericArray_Shuffle(void* TargetArray, const UArrayProperty* ArrayProp);
+	ENGINE_API static void GenericArray_Append(void* TargetArray, const UArrayProperty* TargetArrayProp, void* SourceArray, const UArrayProperty* SourceArrayProperty);
+	ENGINE_API static void GenericArray_Insert(void* TargetArray, const UArrayProperty* ArrayProp, const void* NewItem, int32 Index);
+	ENGINE_API static void GenericArray_Remove(void* TargetArray, const UArrayProperty* ArrayProp, int32 IndexToRemove);
+	ENGINE_API static bool GenericArray_RemoveItem(void* TargetArray, const UArrayProperty* ArrayProp, const void* Item);
+	ENGINE_API static void GenericArray_Clear(void* TargetArray, const UArrayProperty* ArrayProp);
+	ENGINE_API static void GenericArray_Resize(void* TargetArray, const UArrayProperty* ArrayProp, int32 Size);
+	ENGINE_API static int32 GenericArray_Length(void* TargetArray, const UArrayProperty* ArrayProp);
+	ENGINE_API static int32 GenericArray_LastIndex(void* TargetArray, const UArrayProperty* ArrayProp);
+	ENGINE_API static void GenericArray_Get(void* TargetArray, const UArrayProperty* ArrayProp, int32 Index, void* Item);
+	ENGINE_API static void GenericArray_Set(void* TargetArray, const UArrayProperty* ArrayProp, int32 Index, const void* NewItem, bool bSizeToFit);
+	ENGINE_API static int32 GenericArray_Find(void* TargetArray, const UArrayProperty* ArrayProperty, const void* ItemToFind);
+	ENGINE_API static void GenericArray_SetArrayPropertyByName(UObject* OwnerObject, FName ArrayPropertyName, const void* SrcArrayAddr);
+
+private:
+	static void GenericArray_HandleBool(const UProperty* Property, void* ItemPtr);
+
+public:
 	// Helper function to get the last valid index of the array for error reporting, or 0 if the array is empty
 	static int32 GetLastIndex(const FScriptArrayHelper& ArrayHelper)
 	{
@@ -193,11 +198,11 @@ class UKismetArrayLibrary : public UBlueprintFunctionLibrary
  
  		Stack.MostRecentPropertyAddress = NULL;
  		Stack.StepCompiledIn<UProperty>(StorageSpace);
- 		int32& NewItem = (Stack.MostRecentPropertyAddress != NULL) ? *(int32*)(Stack.MostRecentPropertyAddress) : (*(int32*)StorageSpace);
+		void* NewItemPtr = (Stack.MostRecentPropertyAddress != NULL) ? Stack.MostRecentPropertyAddress : StorageSpace;
  
  		P_FINISH;
  
- 		*(int32*)Result = GenericArray_Add(ArrayAddr, ArrayProperty, NewItem);
+		*(int32*)Result = GenericArray_Add(ArrayAddr, ArrayProperty, NewItemPtr);
  
 		InnerProp->DestroyValue(StorageSpace);
 	}
@@ -247,12 +252,12 @@ class UKismetArrayLibrary : public UBlueprintFunctionLibrary
 
 		Stack.MostRecentPropertyAddress = NULL;
 		Stack.StepCompiledIn<UProperty>(StorageSpace);
-		int32& NewItem = (Stack.MostRecentPropertyAddress != NULL) ? *(int32*)(Stack.MostRecentPropertyAddress) : (*(int32*)StorageSpace);
+		void* NewItemPtr = (Stack.MostRecentPropertyAddress != NULL) ? Stack.MostRecentPropertyAddress : StorageSpace;
 
 		P_GET_PROPERTY(UIntProperty, Index);
 		P_FINISH;
 
-		GenericArray_Insert(ArrayAddr, ArrayProperty, NewItem, Index);
+		GenericArray_Insert(ArrayAddr, ArrayProperty, NewItemPtr, Index);
 
 		InnerProp->DestroyValue(StorageSpace);
 	}
@@ -284,23 +289,18 @@ class UKismetArrayLibrary : public UBlueprintFunctionLibrary
 
 		Stack.MostRecentPropertyAddress = NULL;
 		Stack.StepCompiledIn<UProperty>(StorageSpace);
-		int32& Item = (Stack.MostRecentPropertyAddress != NULL) ? *(int32*)(Stack.MostRecentPropertyAddress) : (*(int32*)StorageSpace);
+		void* ItemPtr = StorageSpace;
 
 		P_FINISH;
 
 		// Bools need to be processed internally by the property so that C++ bool value is properly set.
-		const UBoolProperty* BoolProperty = Cast<const UBoolProperty>(InnerProp);
-		if (BoolProperty != NULL)
-		{
-			BoolProperty->SetPropertyValue(&Item, !!Item);
-		}
+		GenericArray_HandleBool(InnerProp, ItemPtr);
 		  
-		bool WasRemoved = GenericArray_RemoveItem(ArrayAddr, ArrayProperty, Item);
+		bool WasRemoved = GenericArray_RemoveItem(ArrayAddr, ArrayProperty, ItemPtr);
 		*(bool*)Result = WasRemoved; 
 
 		InnerProp->DestroyValue(StorageSpace);
 	}
-
 	
 	DECLARE_FUNCTION(execArray_Clear)
 	{
@@ -363,11 +363,11 @@ class UKismetArrayLibrary : public UBlueprintFunctionLibrary
 
 		Stack.MostRecentPropertyAddress = NULL;
 		Stack.StepCompiledIn<UProperty>(StorageSpace);
-		int32& Item = (Stack.MostRecentPropertyAddress != NULL) ? *(int32*)(Stack.MostRecentPropertyAddress) : (*(int32*)StorageSpace);
+		void* ItemPtr = (Stack.MostRecentPropertyAddress != NULL) ? Stack.MostRecentPropertyAddress : StorageSpace;
 
 		P_FINISH;
 
-		GenericArray_Get(ArrayAddr, ArrayProperty, Index, Item);
+		GenericArray_Get(ArrayAddr, ArrayProperty, Index, ItemPtr);
 
 		InnerProp->DestroyValue(StorageSpace);
 	}
@@ -388,13 +388,13 @@ class UKismetArrayLibrary : public UBlueprintFunctionLibrary
 
 		Stack.MostRecentPropertyAddress = NULL;
 		Stack.StepCompiledIn<UProperty>(StorageSpace);
-		int32& NewItem = (Stack.MostRecentPropertyAddress != NULL) ? *(int32*)(Stack.MostRecentPropertyAddress) : (*(int32*)StorageSpace);
+		void* NewItemPtr = (Stack.MostRecentPropertyAddress != NULL) ? Stack.MostRecentPropertyAddress : StorageSpace;
 
 		P_GET_UBOOL(bSizeToFit);
 
 		P_FINISH;
 
-		GenericArray_Set(ArrayAddr, ArrayProperty, Index, NewItem, bSizeToFit);
+		GenericArray_Set(ArrayAddr, ArrayProperty, Index, NewItemPtr, bSizeToFit);
 
 		InnerProp->DestroyValue(StorageSpace);
 	}
@@ -414,19 +414,15 @@ class UKismetArrayLibrary : public UBlueprintFunctionLibrary
 
 		Stack.MostRecentPropertyAddress = NULL;
 		Stack.StepCompiledIn<UProperty>(StorageSpace);
-		int32& ItemToFind = (Stack.MostRecentPropertyAddress != NULL) ? *(int32*)(Stack.MostRecentPropertyAddress) : (*(int32*)StorageSpace);
+		void* ItemToFindPtr = StorageSpace;
 
 		P_FINISH;
 
 		// Bools need to be processed internally by the property so that C++ bool value is properly set.
-		const UBoolProperty* BoolProperty = Cast<const UBoolProperty>(InnerProp);
-		if (BoolProperty != NULL)
-		{
-			BoolProperty->SetPropertyValue(&ItemToFind, !!ItemToFind);
-		}
+		GenericArray_HandleBool(InnerProp, ItemToFindPtr);
 
 		// Perform the search
-		int32 FoundIndex = GenericArray_Find(ArrayAddr, ArrayProperty, ItemToFind);
+		int32 FoundIndex = GenericArray_Find(ArrayAddr, ArrayProperty, ItemToFindPtr);
 		*(int32*)Result = FoundIndex;
 
 		InnerProp->DestroyValue(StorageSpace);
@@ -447,19 +443,15 @@ class UKismetArrayLibrary : public UBlueprintFunctionLibrary
 
 		Stack.MostRecentPropertyAddress = NULL;
 		Stack.StepCompiledIn<UProperty>(StorageSpace);
-		int32& ItemToFind = (Stack.MostRecentPropertyAddress != NULL) ? *(int32*)(Stack.MostRecentPropertyAddress) : (*(int32*)StorageSpace);
+		void* ItemToFindPtr = StorageSpace;
 
 		P_FINISH;
 
 		// Bools need to be processed internally by the property so that C++ bool value is properly set.
-		const UBoolProperty* BoolProperty = Cast<const UBoolProperty>(InnerProp);
-		if (BoolProperty != NULL)
-		{
-			BoolProperty->SetPropertyValue(&ItemToFind, !!ItemToFind);
-		}
+		GenericArray_HandleBool(InnerProp, ItemToFindPtr);
 
 		// Perform the search
-		int32 FoundIndex = GenericArray_Find(ArrayAddr, ArrayProperty, ItemToFind);
+		int32 FoundIndex = GenericArray_Find(ArrayAddr, ArrayProperty, ItemToFindPtr);
 		*(bool*)Result = (FoundIndex >= 0);
 
 		InnerProp->DestroyValue(StorageSpace);
@@ -475,14 +467,6 @@ class UKismetArrayLibrary : public UBlueprintFunctionLibrary
 
 		P_FINISH;
 
-		if(OwnerObject != NULL)
-		{
-			UArrayProperty* ArrayProp = FindField<UArrayProperty>(OwnerObject->GetClass(), ArrayPropertyName);
-			if(ArrayProp != NULL)
-			{
-				void* Dest = ArrayProp->ContainerPtrToValuePtr<void>(OwnerObject);
-				ArrayProp->CopyValuesInternal(Dest, SrcArrayAddr, 1);
-			}
-		}
+		GenericArray_SetArrayPropertyByName(OwnerObject, ArrayPropertyName, SrcArrayAddr);
 	}
 };
diff --git a/Engine/Source/Runtime/Engine/Classes/Kismet/KismetSystemLibrary.h b/Engine/Source/Runtime/Engine/Classes/Kismet/KismetSystemLibrary.h
--- a/Engine/Source/Runtime/Engine/Classes/Kismet/KismetSystemLibrary.h
+++ b/Engine/Source/Runtime/Engine/Classes/Kismet/KismetSystemLibrary.h
@@ -382,6 +382,8 @@ class UKismetSystemLibrary : public UBlueprintFunctionLibrary
 	UFUNCTION(BlueprintCallable, CustomThunk, meta = (BlueprintInternalUseOnly = "true", CustomStructureParam = "Value"))
 	static void SetStructurePropertyByName(UObject* Object, FName PropertyName, const FGenericStruct& Value);
 
+	ENGINE_API static void Generic_SetStructurePropertyByName(UObject* OwnerObject, FName StructPropertyName, const void* SrcStructAddr);
+
 	/** Based on UKismetArrayLibrary::execSetArrayPropertyByName */
 	DECLARE_FUNCTION(execSetStructurePropertyByName)
 	{
@@ -393,15 +395,7 @@ class UKismetSystemLibrary : public UBlueprintFunctionLibrary
 
 		P_FINISH;
 
-		if (OwnerObject != NULL)
-		{
-			UStructProperty* StructProp = FindField<UStructProperty>(OwnerObject->GetClass(), StructPropertyName);
-			if (StructProp != NULL)
-			{
-				void* Dest = StructProp->ContainerPtrToValuePtr<void>(OwnerObject);
-				StructProp->CopyValuesInternal(Dest, SrcStructAddr, 1);
-			}
-		}
+		Generic_SetStructurePropertyByName(OwnerObject, StructPropertyName, SrcStructAddr);
 	}
 
 	// --- Collision functions ------------------------------
diff --git a/Engine/Source/Runtime/Engine/Private/DataTableFunctionLibrary.cpp b/Engine/Source/Runtime/Engine/Private/DataTableFunctionLibrary.cpp
--- a/Engine/Source/Runtime/Engine/Private/DataTableFunctionLibrary.cpp
+++ b/Engine/Source/Runtime/Engine/Private/DataTableFunctionLibrary.cpp
@@ -26,6 +26,29 @@ void UDataTableFunctionLibrary::EvaluateCurveTableRow(UCurveTable* CurveTable, F
     }
 }
 
+bool UDataTableFunctionLibrary::Generic_GetDataTableRowFromName(UDataTable* Table, FName RowName, void* OutRowPtr)
+{
+	bool bFoundRow = false;
+
+	if (OutRowPtr && Table)
+	{
+		void* RowPtr = Table->FindRowUnchecked(RowName);
+
+		if (RowPtr != NULL)
+		{
+			UScriptStruct* StructType = Table->RowStruct;
+
+			if (StructType != NULL)
+			{
+				StructType->CopyScriptStruct(OutRowPtr, RowPtr);
+				bFoundRow = true;
+			}
+		}
+	}
+
+	return bFoundRow;
+}
+
 bool UDataTableFunctionLibrary::GetDataTableRowFromName(UDataTable* Table, FName RowName, FTableRowBase& OutRow)
 {
 	// We should never hit this!  stubs to avoid NoExport on the class.
diff --git a/Engine/Source/Runtime/Engine/Private/KismetArrayLibrary.cpp b/Engine/Source/Runtime/Engine/Private/KismetArrayLibrary.cpp
--- a/Engine/Source/Runtime/Engine/Private/KismetArrayLibrary.cpp
+++ b/Engine/Source/Runtime/Engine/Private/KismetArrayLibrary.cpp
@@ -26,7 +26,7 @@ void UKismetArrayLibrary::FilterArray(const TArray<AActor*>& TargetArray, TSubcl
 	}
 }
 
-int32 UKismetArrayLibrary::GenericArray_Add(void* TargetArray, const UArrayProperty* ArrayProp, const int32& NewItem)
+int32 UKismetArrayLibrary::GenericArray_Add(void* TargetArray, const UArrayProperty* ArrayProp, const void* NewItem)
 {
 	int32 NewIndex = 0;
 	if( TargetArray )
@@ -35,7 +35,7 @@ int32 UKismetArrayLibrary::GenericArray_Add(void* TargetArray, const UArrayPrope
 		UProperty* InnerProp = ArrayProp->Inner;
 
 		NewIndex = ArrayHelper.AddValue();
-		InnerProp->CopySingleValueToScriptVM(ArrayHelper.GetRawPtr(NewIndex), (void*)(&NewItem));
+		InnerProp->CopySingleValueToScriptVM(ArrayHelper.GetRawPtr(NewIndex), NewItem);
 	}
 	return NewIndex;
 }
@@ -60,7 +60,7 @@ void UKismetArrayLibrary::GenericArray_Append(void* TargetArray, const UArrayPro
 	}
 }
 
-void UKismetArrayLibrary::GenericArray_Insert(void* TargetArray, const UArrayProperty* ArrayProp, const int32& NewItem, int32 Index)
+void UKismetArrayLibrary::GenericArray_Insert(void* TargetArray, const UArrayProperty* ArrayProp, const void* NewItem, int32 Index)
 {
 	if( TargetArray )
 	{
@@ -71,7 +71,7 @@ void UKismetArrayLibrary::GenericArray_Insert(void* TargetArray, const UArrayPro
 			|| (Index >= 0 && Index <= ArrayHelper.Num()) )
 		{
 			ArrayHelper.InsertValues(Index, 1);
-			InnerProp->CopySingleValueToScriptVM(ArrayHelper.GetRawPtr(Index), (void*)(&NewItem));
+			InnerProp->CopySingleValueToScriptVM(ArrayHelper.GetRawPtr(Index), NewItem);
 		}
 		else
 		{
@@ -97,7 +97,7 @@ void UKismetArrayLibrary::GenericArray_Remove(void* TargetArray, const UArrayPro
 }
 
 
-bool UKismetArrayLibrary::GenericArray_RemoveItem( void* TargetArray, const UArrayProperty* ArrayProp, const int32& Item )
+bool UKismetArrayLibrary::GenericArray_RemoveItem(void* TargetArray, const UArrayProperty* ArrayProp, const void* Item)
 {
 	bool bRemoved = false;
 
@@ -184,7 +184,7 @@ int32 UKismetArrayLibrary::GenericArray_LastIndex(void* TargetArray, const UArra
 	return INDEX_NONE;
 }
 
-void UKismetArrayLibrary::GenericArray_Get(void* TargetArray, const UArrayProperty* ArrayProp, int32 Index, int32& Item)
+void UKismetArrayLibrary::GenericArray_Get(void* TargetArray, const UArrayProperty* ArrayProp, int32 Index, void* Item)
 {
 	if( TargetArray )
 	{
@@ -193,18 +193,18 @@ void UKismetArrayLibrary::GenericArray_Get(void* TargetArray, const UArrayProper
 		UProperty* InnerProp = ArrayProp->Inner;
 		if( ArrayHelper.IsValidIndex(Index) )
 		{
-			InnerProp->CopyCompleteValueFromScriptVM(&Item, ArrayHelper.GetRawPtr(Index));	
+			InnerProp->CopyCompleteValueFromScriptVM(Item, ArrayHelper.GetRawPtr(Index));	
 		}
 		else
 		{
 			UE_LOG(LogArray, Warning, TEXT("Attempted to get an item from array %s out of bounds [%d/%d]!"), *ArrayProp->GetName(), Index, GetLastIndex(ArrayHelper));
-			InnerProp->InitializeValue(&Item);
+			InnerProp->InitializeValue(Item);
 		}
 	}
 }
 
 
-void UKismetArrayLibrary::GenericArray_Set(void* TargetArray, const UArrayProperty* ArrayProp, int32 Index, const int32& NewItem, bool bSizeToFit)
+void UKismetArrayLibrary::GenericArray_Set(void* TargetArray, const UArrayProperty* ArrayProp, int32 Index, const void* NewItem, bool bSizeToFit)
 {
 	if( TargetArray )
 	{
@@ -219,7 +219,7 @@ void UKismetArrayLibrary::GenericArray_Set(void* TargetArray, const UArrayProper
 
 		if (ArrayHelper.IsValidIndex(Index))
 		{
-			InnerProp->CopySingleValueToScriptVM(ArrayHelper.GetRawPtr(Index), (void*)(&NewItem));
+			InnerProp->CopySingleValueToScriptVM(ArrayHelper.GetRawPtr(Index), NewItem);
 		}
 		else
 		{
@@ -228,7 +228,7 @@ void UKismetArrayLibrary::GenericArray_Set(void* TargetArray, const UArrayProper
 	}
 }
 
-int32 UKismetArrayLibrary::GenericArray_Find(void* TargetArray, const UArrayProperty* ArrayProperty, const int32& ItemToFind)
+int32 UKismetArrayLibrary::GenericArray_Find(void* TargetArray, const UArrayProperty* ArrayProperty, const void* ItemToFind)
 {
 	int32 ResultIndex = INDEX_NONE;
 
@@ -240,7 +240,7 @@ int32 UKismetArrayLibrary::GenericArray_Find(void* TargetArray, const UArrayProp
 		// compare against each element in the array
 		for (int32 Idx = 0; Idx < ArrayHelper.Num() && ResultIndex == INDEX_NONE; Idx++)
 		{
-			if (InnerProp->Identical(&ItemToFind,ArrayHelper.GetRawPtr(Idx)))
+			if (InnerProp->Identical(ItemToFind,ArrayHelper.GetRawPtr(Idx)))
 			{
 				ResultIndex = Idx;
 			}
@@ -251,6 +251,29 @@ int32 UKismetArrayLibrary::GenericArray_Find(void* TargetArray, const UArrayProp
 	return ResultIndex;
 }
 
+void UKismetArrayLibrary::GenericArray_SetArrayPropertyByName(UObject* OwnerObject, FName ArrayPropertyName, const void* SrcArrayAddr)
+{
+	if (OwnerObject != NULL)
+	{
+		UArrayProperty* ArrayProp = FindField<UArrayProperty>(OwnerObject->GetClass(), ArrayPropertyName);
+		if (ArrayProp != NULL)
+		{
+			void* Dest = ArrayProp->ContainerPtrToValuePtr<void>(OwnerObject);
+			ArrayProp->CopyValuesInternal(Dest, SrcArrayAddr, 1);
+		}
+	}
+}
+
+void UKismetArrayLibrary::GenericArray_HandleBool(const UProperty* Property, void* ItemPtr)
+{
+	const UBoolProperty* BoolProperty = Cast<const UBoolProperty>(Property);
+	if (BoolProperty)
+	{
+		ensure((BoolProperty->ElementSize * BoolProperty->ArrayDim) == sizeof(uint8));
+		BoolProperty->SetPropertyValue(ItemPtr, 0 != *(reinterpret_cast<uint8*>(ItemPtr)));
+	}
+}
+
 //////////////////////////////////////////////////////////////////////////
 // Stubs for the UFunctions declared as kismet callable.  These are never actually called...the CustomThunk code calls the appropriate native function with a void* reference to the array
 
diff --git a/Engine/Source/Runtime/Engine/Private/KismetSystemLibrary.cpp b/Engine/Source/Runtime/Engine/Private/KismetSystemLibrary.cpp
--- a/Engine/Source/Runtime/Engine/Private/KismetSystemLibrary.cpp
+++ b/Engine/Source/Runtime/Engine/Private/KismetSystemLibrary.cpp
@@ -582,6 +582,18 @@ void UKismetSystemLibrary::SetTransformPropertyByName(UObject* Object, FName Pro
 	}
 }
 
+void UKismetSystemLibrary::Generic_SetStructurePropertyByName(UObject* OwnerObject, FName StructPropertyName, const void* SrcStructAddr)
+{
+	if (OwnerObject != NULL)
+	{
+		UStructProperty* StructProp = FindField<UStructProperty>(OwnerObject->GetClass(), StructPropertyName);
+		if (StructProp != NULL)
+		{
+			void* Dest = StructProp->ContainerPtrToValuePtr<void>(OwnerObject);
+			StructProp->CopyValuesInternal(Dest, SrcStructAddr, 1);
+		}
+	}
+}
 
 void UKismetSystemLibrary::GetActorListFromComponentList(const TArray<UPrimitiveComponent*>& ComponentList, UClass* ActorClassFilter, TArray<class AActor*>& OutActorList)
 {
