Const fix for FAnimNotifyEvent::operator==.

diff --git a/Engine/Source/Runtime/Engine/Classes/Animation/AnimSequenceBase.h b/Engine/Source/Runtime/Engine/Classes/Animation/AnimSequenceBase.h
--- a/Engine/Source/Runtime/Engine/Classes/Animation/AnimSequenceBase.h
+++ b/Engine/Source/Runtime/Engine/Classes/Animation/AnimSequenceBase.h
@@ -109,7 +109,7 @@ struct FAnimNotifyEvent
 		return Notify != NULL || NotifyStateClass != NULL;
 	}
 
-	bool operator ==(const FAnimNotifyEvent& Other)
+	bool operator ==(const FAnimNotifyEvent& Other) const
 	{
 		return(
 			(Notify && Notify == Other.Notify) || 
