tpp341307 BP:  Add "!IsPendingKill() check to the IsValid macro

diff --git a/Engine/Source/Runtime/Engine/Classes/Kismet/KismetSystemLibrary.h b/Engine/Source/Runtime/Engine/Classes/Kismet/KismetSystemLibrary.h
--- a/Engine/Source/Runtime/Engine/Classes/Kismet/KismetSystemLibrary.h
+++ b/Engine/Source/Runtime/Engine/Classes/Kismet/KismetSystemLibrary.h
@@ -72,6 +72,10 @@ class UKismetSystemLibrary : public UBlueprintFunctionLibrary
 		StackTraceImpl(Stack);
 	}
 
+	// Return true if the object is usable : non-null and not pending kill
+	UFUNCTION(BlueprintPure, Category = "Utilities")
+	static ENGINE_API bool IsValid(const UObject* Object);
+
 	// Returns the display name (or actor label), for displaying to end users.  Note:  In editor builds, this is the actor label.  In non-editor builds, this is the actual object name.  This function should not be used to uniquely identify actors!
 	UFUNCTION(BlueprintPure, Category="Utilities")
 	static ENGINE_API FString GetDisplayName(const UObject* Object);
diff --git a/Engine/Source/Runtime/Engine/Private/KismetSystemLibrary.cpp b/Engine/Source/Runtime/Engine/Private/KismetSystemLibrary.cpp
--- a/Engine/Source/Runtime/Engine/Private/KismetSystemLibrary.cpp
+++ b/Engine/Source/Runtime/Engine/Private/KismetSystemLibrary.cpp
@@ -24,6 +24,11 @@ void UKismetSystemLibrary::StackTraceImpl(const FFrame& StackFrame)
 	UE_LOG(LogBlueprintUserMessages, Log, TEXT("\n%s"), *Trace);
 }
 
+bool UKismetSystemLibrary::IsValid(const UObject* Object)
+{
+	return ::IsValid(Object);
+}
+
 FString UKismetSystemLibrary::GetDisplayName(const UObject* Object)
 {
 #if WITH_EDITOR
