UE4 - Propagate extra bone influences (fixes assert)

diff --git a/Engine/Source/Runtime/Engine/Private/SkeletalMeshMerge.cpp b/Engine/Source/Runtime/Engine/Private/SkeletalMeshMerge.cpp
--- a/Engine/Source/Runtime/Engine/Private/SkeletalMeshMerge.cpp
+++ b/Engine/Source/Runtime/Engine/Private/SkeletalMeshMerge.cpp
@@ -583,6 +583,7 @@ void FSkeletalMeshMerge::GenerateLODModel( int32 LODIdx )
 	MergeLODModel.VertexBufferGPUSkin.SetUseFullPrecisionUVs(MergeMesh->bUseFullPrecisionUVs);
 	// set CPU skinning on vertex buffer since it affects the type of TResourceArray needed
 	MergeLODModel.VertexBufferGPUSkin.SetNeedsCPUAccess(MergeResource->RequiresCPUSkinning(GRHIFeatureLevel));
+	MergeLODModel.VertexBufferGPUSkin.SetHasExtraBoneInfluences(MergeResource->HasExtraBoneInfluences());
 	// Set the number of tex coords on this vertex buffer
 	MergeLODModel.VertexBufferGPUSkin.SetNumTexCoords(TotalNumUVs);
 	MergeLODModel.NumTexCoords = TotalNumUVs;
