Fixed Win32 build failure.

diff --git a/Engine/Source/Runtime/Engine/Classes/Engine/NiagaraEffectRenderer.h b/Engine/Source/Runtime/Engine/Classes/Engine/NiagaraEffectRenderer.h
--- a/Engine/Source/Runtime/Engine/Classes/Engine/NiagaraEffectRenderer.h
+++ b/Engine/Source/Runtime/Engine/Classes/Engine/NiagaraEffectRenderer.h
@@ -769,7 +769,7 @@ public:
 	}
 
 	void AddRibbonVert(TArray<FParticleBeamTrailVertex>& RenderData, FVector ParticlePos, const FNiagaraEmitterParticleData &Data, FVector2D UV1,
-		const FVector4 Color, const FVector4 Age, const FVector4 Rotation)
+		const FVector4 &Color, const FVector4 &Age, const FVector4 &Rotation)
 	{
 		FParticleBeamTrailVertex& NewVertex = *new(RenderData)FParticleBeamTrailVertex;
 		NewVertex.Position = ParticlePos;
