Changing intial position on a details panel of a skeletal mesh now updated the pose

#TTP 346556 - EDITOR: ORION: Animation Sequence does not move when scrubbing the Initial Position bar

diff --git a/Engine/Source/Runtime/Engine/Private/SkeletalMeshComponent.cpp b/Engine/Source/Runtime/Engine/Private/SkeletalMeshComponent.cpp
--- a/Engine/Source/Runtime/Engine/Private/SkeletalMeshComponent.cpp
+++ b/Engine/Source/Runtime/Engine/Private/SkeletalMeshComponent.cpp
@@ -395,6 +395,7 @@ void USkeletalMeshComponent::PostEditChangeProperty(FPropertyChangedEvent& Prope
 		if ( PropertyThatChanged->GetFName() == GET_MEMBER_NAME_CHECKED( FSingleAnimationPlayData, SavedPosition ))
 		{
 			AnimationData.ValidatePosition();
+			SetPosition(AnimationData.SavedPosition, false);
 		}
 	}
 }
