#ue4 - Header documentation for StaticMesh (actor, component, and mesh), SkeletalMesh (actor, component, and mesh), Mobility.

diff --git a/Engine/Source/Runtime/Engine/Classes/Animation/SkeletalMeshActor.h b/Engine/Source/Runtime/Engine/Classes/Animation/SkeletalMeshActor.h
--- a/Engine/Source/Runtime/Engine/Classes/Animation/SkeletalMeshActor.h
+++ b/Engine/Source/Runtime/Engine/Classes/Animation/SkeletalMeshActor.h
@@ -4,6 +4,14 @@
 #include "Matinee/MatineeAnimInterface.h"
 #include "SkeletalMeshActor.generated.h"
 
+/**
+ * SkeletalMeshActor is an instance of a USkeletalMesh in the world.
+ * Skeletal meshes are deformable meshes that can be animated and change their geometry at run-time.
+ * Skeletal meshes dragged into the level from the Content Browser are automatically converted to StaticMeshActors.
+ * 
+ * @see https://docs.unrealengine.com/latest/INT/Engine/Content/Types/SkeletalMeshes/index.html
+ * @see USkeletalMesh
+ */
 UCLASS(ClassGroup=ISkeletalMeshes, Blueprintable, ConversionRoot, meta=(ChildCanTick))
 class ENGINE_API ASkeletalMeshActor : public AActor, public IMatineeAnimInterface
 {
diff --git a/Engine/Source/Runtime/Engine/Classes/Components/SkeletalMeshComponent.h b/Engine/Source/Runtime/Engine/Classes/Components/SkeletalMeshComponent.h
--- a/Engine/Source/Runtime/Engine/Classes/Components/SkeletalMeshComponent.h
+++ b/Engine/Source/Runtime/Engine/Classes/Components/SkeletalMeshComponent.h
@@ -260,7 +260,10 @@ struct FSkeletalMeshComponentPreClothTickFunction : public FTickFunction
 
 
 /**
- * SkeletalMeshComponent is a mesh that supports skeletal animation and physics.
+ * SkeletalMeshComponent is used to create an instance of a USkeletalMesh.
+ *
+ * @see https://docs.unrealengine.com/latest/INT/Engine/Content/Types/SkeletalMeshes/index.html
+ * @see USkeletalMesh
  */
 UCLASS(ClassGroup=(Rendering, Common), hidecategories=Object, config=Engine, editinlinenew, meta=(BlueprintSpawnableComponent))
 class ENGINE_API USkeletalMeshComponent : public USkinnedMeshComponent, public IInterface_CollisionDataProvider
diff --git a/Engine/Source/Runtime/Engine/Classes/Components/StaticMeshComponent.h b/Engine/Source/Runtime/Engine/Classes/Components/StaticMeshComponent.h
--- a/Engine/Source/Runtime/Engine/Classes/Components/StaticMeshComponent.h
+++ b/Engine/Source/Runtime/Engine/Classes/Components/StaticMeshComponent.h
@@ -106,7 +106,12 @@ struct TStructOpsTypeTraits<FStaticMeshComponentLODInfo> : public TStructOpsType
 	};
 };
 
-/** A StaticMeshComponent is a mesh that does not animate. */
+/**
+ * StaticMeshComponent is used to create an instance of a UStaticMesh.
+ *
+ * @see https://docs.unrealengine.com/latest/INT/Engine/Content/Types/StaticMeshes/index.html
+ * @see UStaticMesh
+ */
 UCLASS(ClassGroup=(Rendering, Common), hidecategories=(Object,Activation,"Components|Activation"), ShowCategories=(Mobility), editinlinenew, meta=(BlueprintSpawnableComponent))
 class ENGINE_API UStaticMeshComponent : public UMeshComponent
 {
diff --git a/Engine/Source/Runtime/Engine/Classes/Engine/EngineTypes.h b/Engine/Source/Runtime/Engine/Classes/Engine/EngineTypes.h
--- a/Engine/Source/Runtime/Engine/Classes/Engine/EngineTypes.h
+++ b/Engine/Source/Runtime/Engine/Classes/Engine/EngineTypes.h
@@ -2540,30 +2540,36 @@ enum EPhysicalSurface
 	SurfaceType_Max UMETA(Hidden)
 };
 
-/** Enum describing how often this component is allowed to move */
+/** Describes how often this component is allowed to move. */
 UENUM()
 namespace EComponentMobility
 {
 	enum Type
 	{
-		// A static object can't be changed in game.
-		// - Allows Baked Lighting
-		// - Fastest Rendering
+		/**
+		 * Static objects cannot be moved or changed in game.
+		 * - Allows baked lighting
+		 * - Fastest rendering
+		 */
 		Static,
 
-		// A stationary light will only have its shadowing and bounced lighting from static geometry baked by Lightmass, all other lighting will be dynamic.
-		// - Stationary only makes sense for light components
-		// - It can change color and intensity in game.
-		// - Can't Move
-		// - Allows Partial Baked Lighting
-		// - Dynamic Shadows
+		/**
+		 * A stationary light will only have its shadowing and bounced lighting from static geometry baked by Lightmass, all other lighting will be dynamic.
+		 * - Stationary only makes sense for light components
+		 * - It can change color and intensity in game.
+		 * - Can't move
+		 * - Allows partial baked lighting
+		 * - Dynamic shadows
+		 */
 		Stationary,
 
-		// Movable objects can be moved and changed in game
-		// - Totally Dynamic
-		// - Allows Dynamic Shadows
-		// - Slowest Rendering
-		Movable		
+		/**
+		 * Movable objects can be moved and changed in game.
+		 * - Totally dynamic
+		 * - Can cast dynamic shadows
+		 * - Slowest rendering
+		 */
+		Movable
 	};
 }
 
diff --git a/Engine/Source/Runtime/Engine/Classes/Engine/SkeletalMesh.h b/Engine/Source/Runtime/Engine/Classes/Engine/SkeletalMesh.h
--- a/Engine/Source/Runtime/Engine/Classes/Engine/SkeletalMesh.h
+++ b/Engine/Source/Runtime/Engine/Classes/Engine/SkeletalMesh.h
@@ -555,6 +555,13 @@ struct FSkeletalMaterial
 
 class FSkeletalMeshResource;
 
+/**
+ * SkeletalMesh is geometry bound to a hierarchical skeleton of bones which can be animated for the purpose of deforming the mesh.
+ * Skeletal Meshes are built up of two parts; a set of polygons composed to make up the surface of the mesh, and a hierarchical skeleton which can be used to animate the polygons.
+ * The 3D models, rigging, and animations are created in an external modeling and animation application (3DSMax, Maya, Softimage, etc).
+ *
+ * @see https://docs.unrealengine.com/latest/INT/Engine/Content/Types/SkeletalMeshes/index.html
+ */
 UCLASS(hidecategories=Object, MinimalAPI, BlueprintType)
 class USkeletalMesh : public UObject
 {
diff --git a/Engine/Source/Runtime/Engine/Classes/Engine/StaticMesh.h b/Engine/Source/Runtime/Engine/Classes/Engine/StaticMesh.h
--- a/Engine/Source/Runtime/Engine/Classes/Engine/StaticMesh.h
+++ b/Engine/Source/Runtime/Engine/Classes/Engine/StaticMesh.h
@@ -313,6 +313,14 @@ struct FAssetEditorOrbitCameraPosition
 	FRotator CamOrbitRotation;
 };
 
+/**
+ * A StaticMesh is a piece of geometry that consists of a static set of polygons.
+ * Static Meshes can be translated, rotated, and scaled, but they cannot have their vertices animated in any way. As such, they are more efficient
+ * to render than other types of geometry such as USkeletalMesh, and they are often the basic building block of levels created in the engine.
+ *
+ * @see https://docs.unrealengine.com/latest/INT/Engine/Content/Types/StaticMeshes/index.html
+ * @see AStaticMeshActor, UStaticMeshComponent
+ */
 UCLASS(collapsecategories, hidecategories=Object, customconstructor, MinimalAPI, BlueprintType, config=Engine)
 class UStaticMesh : public UObject, public IInterface_CollisionDataProvider, public IInterface_AssetUserData
 {
diff --git a/Engine/Source/Runtime/Engine/Classes/Engine/StaticMeshActor.h b/Engine/Source/Runtime/Engine/Classes/Engine/StaticMeshActor.h
--- a/Engine/Source/Runtime/Engine/Classes/Engine/StaticMeshActor.h
+++ b/Engine/Source/Runtime/Engine/Classes/Engine/StaticMeshActor.h
@@ -4,8 +4,12 @@
 #include "StaticMeshActor.generated.h"
 
 /**
- * An instance of a StaticMesh in a level
- * Note that PostInitializeComponents() is not called for StaticMeshActors
+ * StaticMeshActor is an instance of a UStaticMesh in the world.
+ * Static meshes are geometry that do not animate or otherwise deform, and are more efficient to render than other types of geometry.
+ * Static meshes dragged into the level from the Content Browser are automatically converted to StaticMeshActors.
+ *
+ * @see https://docs.unrealengine.com/latest/INT/Engine/Actors/StaticMeshActor/index.html
+ * @see UStaticMesh
  */
 UCLASS(hidecategories=(Input), showcategories=("Input|MouseInput", "Input|TouchInput"), ConversionRoot, meta=(ChildCanTick))
 class ENGINE_API AStaticMeshActor : public AActor
@@ -21,7 +25,7 @@ class ENGINE_API AStaticMeshActor : public AActor
 	UPROPERTY(Category=Actor, EditAnywhere, AdvancedDisplay)
 	bool bStaticMeshReplicateMovement;
 
-	/** Function to change mobility type of light */
+	/** Function to change mobility type */
 	void SetMobility(EComponentMobility::Type InMobility);
 
 #if WITH_EDITOR
