Make trail emitters work in local space, moving the trail with the emitter if use local space is checked.

diff --git a/Engine/Source/Runtime/Engine/Private/Particles/ParticleSystemRender.cpp b/Engine/Source/Runtime/Engine/Private/Particles/ParticleSystemRender.cpp
--- a/Engine/Source/Runtime/Engine/Private/Particles/ParticleSystemRender.cpp
+++ b/Engine/Source/Runtime/Engine/Private/Particles/ParticleSystemRender.cpp
@@ -6604,8 +6604,12 @@ int32 FDynamicRibbonEmitterData::FillVertexData(struct FAsyncBufferFillData& Dat
 					}
 
 					FVector FinalPos = InterpPos + InterpUp * InterpSize;
+					if (Source.bUseLocalSpace)
+					{
+						FinalPos += Data.LocalToWorld.GetOrigin();
+					}
 					Vertex = (FParticleBeamTrailVertex*)(TempVertexData);
-					Vertex->Position = FinalPos;
+					Vertex->Position = FinalPos + 
 					Vertex->OldPosition = FinalPos;
 					Vertex->ParticleId	= 0;
 					Vertex->Size.X = InterpSize;
