#ue4 - Add keywords for PrimitiveComponent::IgnoreActorWhenMoving() to make it easier to find in the context menu.

diff --git a/Engine/Source/Runtime/Engine/Classes/Components/PrimitiveComponent.h b/Engine/Source/Runtime/Engine/Classes/Components/PrimitiveComponent.h
--- a/Engine/Source/Runtime/Engine/Classes/Components/PrimitiveComponent.h
+++ b/Engine/Source/Runtime/Engine/Classes/Components/PrimitiveComponent.h
@@ -448,7 +448,7 @@ public:
 	 * Tells this component whether to ignore collision with all components of a specific Actor when this component is moved.
 	 * Components on the other Actor may also need to be told to do the same when they move.
 	 */
-	UFUNCTION(BlueprintCallable, Category = "Collision")
+	UFUNCTION(BlueprintCallable, Category = "Collision", meta=(Keywords="Move MoveIgnore"))
 	void IgnoreActorWhenMoving(AActor* Actor, bool bShouldIgnore);
 
 	/**
