Updated the Cinematic FPS Perf Capture test.
--Running with -MatineeName="The Name of the matinee actor (Make sure to put _# if there's a number at the end of the name)" will run the test based on when the matinee starts and stops itself and not by time.
--User can still specify a time lenght inplace of the matinee name.
--MatineeName trumps RunTime and CE commaneds.
--User can now run the test from the UFE although they cannot specify a time lenght, defaults to 60 seconds, and no CE commands will be used.

diff --git a/Engine/Source/Runtime/Engine/Private/Tests/AutomationTestCommon.cpp b/Engine/Source/Runtime/Engine/Private/Tests/AutomationTestCommon.cpp
--- a/Engine/Source/Runtime/Engine/Private/Tests/AutomationTestCommon.cpp
+++ b/Engine/Source/Runtime/Engine/Private/Tests/AutomationTestCommon.cpp
@@ -14,6 +14,8 @@ bool FPlayMatineeLatentCommand::Update()
 {
 	if (MatineeActor)
 	{
+		UE_LOG(LogEngineAutomationLatentCommand, Log, TEXT("Triggering the matinee named: '%s'"), *MatineeActor->GetName())
+		
 		//force this matinee to not be looping so it doesn't infinitely loop
 		MatineeActor->bLooping = false;
 		MatineeActor->Play();
@@ -78,3 +80,42 @@ bool FEnqueuePerformanceCaptureCommands::Update()
 
 	return true;
 }
+
+
+bool FMatineePerformanceCaptureCommand::Update()
+{
+	//for every matinee actor in the level
+	for (TObjectIterator<AMatineeActor> It; It; ++It)
+	{
+		AMatineeActor* MatineeActor = *It;
+		FString MatineeFOOName = MatineeActor->GetName();
+		if (MatineeActor->GetName().Equals(MatineeName,ESearchCase::IgnoreCase))
+		{	
+			//add latent action to execute this matinee
+			ADD_LATENT_AUTOMATION_COMMAND(FPlayMatineeLatentCommand(MatineeActor));
+
+			//Run the Stat FPS Chart command
+			ADD_LATENT_AUTOMATION_COMMAND(FExecWorldStringLatentCommand(TEXT("StartFPSChart")));
+
+			//add action to wait until matinee is complete
+			ADD_LATENT_AUTOMATION_COMMAND(FWaitForMatineeToCompleteLatentCommand(MatineeActor));
+
+			//Stop the Stat FPS Chart command
+			ADD_LATENT_AUTOMATION_COMMAND(FExecWorldStringLatentCommand(TEXT("StopFPSChart")));
+		}
+	}
+
+	return true;
+
+}
+
+
+bool FExecWorldStringLatentCommand::Update()
+{
+	check(GEngine->GetWorldContexts().Num() == 1);
+	check(GEngine->GetWorldContexts()[0].WorldType == EWorldType::Game);
+
+	UE_LOG(LogEngineAutomationLatentCommand, Log, TEXT("Running Exec Command. '%s'"), *ExecCommand);
+	GEngine->Exec(GEngine->GetWorldContexts()[0].World(), *ExecCommand);
+	return true;
+}
\ No newline at end of file
diff --git a/Engine/Source/Runtime/Engine/Private/Tests/AutomationTestCommon.h b/Engine/Source/Runtime/Engine/Private/Tests/AutomationTestCommon.h
--- a/Engine/Source/Runtime/Engine/Private/Tests/AutomationTestCommon.h
+++ b/Engine/Source/Runtime/Engine/Private/Tests/AutomationTestCommon.h
@@ -37,3 +37,14 @@ DEFINE_LATENT_AUTOMATION_COMMAND_ONE_PARAMETER(FEngineWaitLatentCommand, float,
  * Enqueue performance capture commands after a map has been loaded
  */
 DEFINE_LATENT_AUTOMATION_COMMAND(FEnqueuePerformanceCaptureCommands);
+
+
+/**
+* Run FPS chart command using for the actual duration of the matinee.
+*/
+DEFINE_LATENT_AUTOMATION_COMMAND_ONE_PARAMETER(FMatineePerformanceCaptureCommand, FString, MatineeName);
+
+/**
+* Latent command to run an exec command that also requires a UWorld.
+*/
+DEFINE_LATENT_AUTOMATION_COMMAND_ONE_PARAMETER(FExecWorldStringLatentCommand, FString, ExecCommand);
\ No newline at end of file
diff --git a/Engine/Source/Runtime/Engine/Private/Tests/EngineAutomationTests.cpp b/Engine/Source/Runtime/Engine/Private/Tests/EngineAutomationTests.cpp
--- a/Engine/Source/Runtime/Engine/Private/Tests/EngineAutomationTests.cpp
+++ b/Engine/Source/Runtime/Engine/Private/Tests/EngineAutomationTests.cpp
@@ -351,25 +351,11 @@ bool FLoadGameMapCommand::Update()
 	return true;
 }
 
-/**
- * Latent command to run an exec command that also requires a UWorld.
- */
-DEFINE_LATENT_AUTOMATION_COMMAND_ONE_PARAMETER(FExecWorldStringLatentCommand, FString, ExecCommand);
-
-bool FExecWorldStringLatentCommand::Update()
-{
-	check(GEngine->GetWorldContexts().Num() == 1);
-	check(GEngine->GetWorldContexts()[0].WorldType == EWorldType::Game);
-
-	UE_LOG(LogEngineAutomationTests, Log, TEXT("Running Exec Command. '%s'"), *ExecCommand);
-	GEngine->Exec(GEngine->GetWorldContexts()[0].World(), *ExecCommand);
-	return true;
-}
 
 /**
  * Automation test to load a map and capture FPS performance charts
  */
-IMPLEMENT_COMPLEX_AUTOMATION_TEST(FCinematicFPSPerfTest, "Engine.Cinematic FPS Perf Capture", (EAutomationTestFlags::ATF_Game | EAutomationTestFlags::ATF_NonNullRHI));
+IMPLEMENT_COMPLEX_AUTOMATION_TEST(FCinematicFPSPerfTest, "Maps.Cinematic FPS Perf Capture", (EAutomationTestFlags::ATF_Game | EAutomationTestFlags::ATF_NonNullRHI));
 
 void FCinematicFPSPerfTest::GetTests(TArray<FString>& OutBeautifiedNames, TArray <FString>& OutTestCommands) const
 {
@@ -378,40 +364,82 @@ void FCinematicFPSPerfTest::GetTests(TArray<FString>& OutBeautifiedNames, TArray
 
 bool FCinematicFPSPerfTest::RunTest(const FString& Parameters)
 {
+	//Map to use for this test.
+	const FString MapName = Parameters;
+	
+	//The name of the matinee actor that will be triggered.
+	FString MatineeActorName;
+
 	//Check we are running from commandline
 	const FString CommandLine(FCommandLine::Get());
-	if( CommandLine.Contains(TEXT("AutomationTests")) )
+	if (CommandLine.Contains(TEXT("AutomationTests")))
 	{
-		//Get the name of the console event to trigger the cinematic
-		FString CinematicEventCommand;
-		if( !FParse::Value(*CommandLine, TEXT("CE="), CinematicEventCommand) )
+		//Get the name of the matinee to be used.
+		//If the game was not launched with the -MatineeName argument then this test will be ran based on time.
+		if (!FParse::Value(*CommandLine, TEXT("MatineeName="), MatineeActorName))
 		{
-			CinematicEventCommand = TEXT("CE Start");
-		}
+			UE_LOG(LogEngineAutomationTests, Log, TEXT("The matinee name was not specified.  Run the game with -MatineeName=\"Name of the matinee actor\"."));
+			
+			//Get the name of the console event to trigger the cinematic
+			FString CinematicEventCommand;
+			if (!FParse::Value(*CommandLine, TEXT("CE="), CinematicEventCommand))
+			{
+				UE_LOG(LogEngineAutomationTests, Log, TEXT("A console event command was not specified. Defaults to CE START.  Run the game with -CE=\"Command\"."));
+				CinematicEventCommand = TEXT("CE Start");
+			}
 
-		//Get the length of time to let the cinematic run
-		float RunTime;
-		if( !FParse::Value(*CommandLine, TEXT("RunTime="), RunTime) )
+			//Get the length of time the cinematic will run
+			float RunTime;
+			if (!FParse::Value(*CommandLine, TEXT("RunTime="), RunTime))
+			{
+				UE_LOG(LogEngineAutomationTests, Log, TEXT("A run time length in seconds was not specified. Defaults to 60 seconds. Run the game with -RunTime=###."));
+				RunTime = 60.f;
+			}
+
+			//Load map
+			ADD_LATENT_AUTOMATION_COMMAND(FEngineWaitLatentCommand(1.0f));
+			ADD_LATENT_AUTOMATION_COMMAND(FLoadGameMapCommand(MapName));
+			ADD_LATENT_AUTOMATION_COMMAND(FEngineWaitLatentCommand(1.0f));
+
+			//Start the matinee and perform the FPS Chart
+			ADD_LATENT_AUTOMATION_COMMAND(FExecWorldStringLatentCommand(CinematicEventCommand));
+			ADD_LATENT_AUTOMATION_COMMAND(FExecWorldStringLatentCommand(TEXT("StartFPSChart")));
+			ADD_LATENT_AUTOMATION_COMMAND(FEngineWaitLatentCommand(RunTime));
+			ADD_LATENT_AUTOMATION_COMMAND(FExecWorldStringLatentCommand(TEXT("StopFPSChart")));
+			ADD_LATENT_AUTOMATION_COMMAND(FEngineWaitLatentCommand(1.0f));
+
+			return true;
+		}
+		else
 		{
-			RunTime=5.f;
+			//Load map
+			ADD_LATENT_AUTOMATION_COMMAND(FEngineWaitLatentCommand(1.0f));
+			ADD_LATENT_AUTOMATION_COMMAND(FLoadGameMapCommand(MapName));
+			ADD_LATENT_AUTOMATION_COMMAND(FEngineWaitLatentCommand(1.0f));
+
+			//Start the matinee and perform the FPS Chart
+			ADD_LATENT_AUTOMATION_COMMAND(FMatineePerformanceCaptureCommand(MatineeActorName));
+			ADD_LATENT_AUTOMATION_COMMAND(FEngineWaitLatentCommand(1.0f));
+			
+			return true;
 		}
+	}
+	else
+	{
+		//If the user is running from the UFE then we'll use the default values.
+		//@todo Give the end user a way to specify the values for this test.
+
+		UE_LOG(LogEngineAutomationTests, Log, TEXT("Running the FPS chart performance capturing for 60 seconds while in '%s'.\nThe default CE command won't be used at this time."), *MapName);
 
 		//Load map
-		const FString MapName = Parameters;
-		ADD_LATENT_AUTOMATION_COMMAND(FEngineWaitLatentCommand(1.0f));
 		ADD_LATENT_AUTOMATION_COMMAND(FLoadGameMapCommand(MapName));
 		ADD_LATENT_AUTOMATION_COMMAND(FEngineWaitLatentCommand(1.0f));
 
 		//Start the matinee and perform the FPS Chart
-		ADD_LATENT_AUTOMATION_COMMAND(FExecWorldStringLatentCommand(CinematicEventCommand));
 		ADD_LATENT_AUTOMATION_COMMAND(FExecWorldStringLatentCommand(TEXT("StartFPSChart")));
-		ADD_LATENT_AUTOMATION_COMMAND(FEngineWaitLatentCommand(RunTime));
+		ADD_LATENT_AUTOMATION_COMMAND(FEngineWaitLatentCommand(60.0f));
 		ADD_LATENT_AUTOMATION_COMMAND(FExecWorldStringLatentCommand(TEXT("StopFPSChart")));
-	}
-	else
-	{
-		UE_LOG(LogEngineAutomationTests, Warning, TEXT("FCinematicFPSPerfTest is a Commandline test.  Please use -AutomationTests=\"Engine.Cinematic FPS Perf Capture\""));
-		return false;
+		ADD_LATENT_AUTOMATION_COMMAND(FEngineWaitLatentCommand(1.0f));
 	}
 
 	return true;
