Rearranged the Cinematic Perf Test.

diff --git a/Engine/Source/Runtime/Engine/Private/Tests/EngineAutomationTests.cpp b/Engine/Source/Runtime/Engine/Private/Tests/EngineAutomationTests.cpp
--- a/Engine/Source/Runtime/Engine/Private/Tests/EngineAutomationTests.cpp
+++ b/Engine/Source/Runtime/Engine/Private/Tests/EngineAutomationTests.cpp
@@ -366,7 +366,7 @@ bool FCinematicFPSPerfTest::RunTest(const FString& Parameters)
 {
 	//Map to use for this test.
 	const FString MapName = Parameters;
-	
+
 	//The name of the matinee actor that will be triggered.
 	FString MatineeActorName;
 
@@ -376,10 +376,23 @@ bool FCinematicFPSPerfTest::RunTest(const FString& Parameters)
 	{
 		//Get the name of the matinee to be used.
 		//If the game was not launched with the -MatineeName argument then this test will be ran based on time.
-		if (!FParse::Value(*CommandLine, TEXT("MatineeName="), MatineeActorName))
+		if (FParse::Value(*CommandLine, TEXT("MatineeName="), MatineeActorName))
+		{
+			//Load map
+			ADD_LATENT_AUTOMATION_COMMAND(FEngineWaitLatentCommand(1.0f));
+			ADD_LATENT_AUTOMATION_COMMAND(FLoadGameMapCommand(MapName));
+			ADD_LATENT_AUTOMATION_COMMAND(FEngineWaitLatentCommand(1.0f));
+
+			//Start the matinee and perform the FPS Chart
+			ADD_LATENT_AUTOMATION_COMMAND(FMatineePerformanceCaptureCommand(MatineeActorName));
+			ADD_LATENT_AUTOMATION_COMMAND(FEngineWaitLatentCommand(1.0f));
+
+			return true;
+		}
+		else
 		{
 			UE_LOG(LogEngineAutomationTests, Log, TEXT("The matinee name was not specified.  Run the game with -MatineeName=\"Name of the matinee actor\"."));
-			
+
 			//Get the name of the console event to trigger the cinematic
 			FString CinematicEventCommand;
 			if (!FParse::Value(*CommandLine, TEXT("CE="), CinematicEventCommand))
@@ -409,43 +422,27 @@ bool FCinematicFPSPerfTest::RunTest(const FString& Parameters)
 			ADD_LATENT_AUTOMATION_COMMAND(FEngineWaitLatentCommand(1.0f));
 
 			return true;
-		}
-		else
-		{
-			//Load map
-			ADD_LATENT_AUTOMATION_COMMAND(FEngineWaitLatentCommand(1.0f));
-			ADD_LATENT_AUTOMATION_COMMAND(FLoadGameMapCommand(MapName));
-			ADD_LATENT_AUTOMATION_COMMAND(FEngineWaitLatentCommand(1.0f));
 
-			//Start the matinee and perform the FPS Chart
-			ADD_LATENT_AUTOMATION_COMMAND(FMatineePerformanceCaptureCommand(MatineeActorName));
-			ADD_LATENT_AUTOMATION_COMMAND(FEngineWaitLatentCommand(1.0f));
-			
-			return true;
 		}
 	}
-	else
-	{
-		//If the user is running from the UFE then we'll use the default values.
-		//@todo Give the end user a way to specify the values for this test.
+	//If the user is running from the UFE then we'll use the default values.
+	//@todo Give the end user a way to specify the values for this test.
 
-		UE_LOG(LogEngineAutomationTests, Log, TEXT("Running the FPS chart performance capturing for 60 seconds while in '%s'.\nThe default CE command won't be used at this time."), *MapName);
+	UE_LOG(LogEngineAutomationTests, Log, TEXT("Running the FPS chart performance capturing for 60 seconds while in '%s'.\nThe default CE command won't be used at this time."), *MapName);
 
-		//Load map
-		ADD_LATENT_AUTOMATION_COMMAND(FLoadGameMapCommand(MapName));
-		ADD_LATENT_AUTOMATION_COMMAND(FEngineWaitLatentCommand(1.0f));
+	//Load map
+	ADD_LATENT_AUTOMATION_COMMAND(FLoadGameMapCommand(MapName));
+	ADD_LATENT_AUTOMATION_COMMAND(FEngineWaitLatentCommand(1.0f));
 
-		//Start the matinee and perform the FPS Chart
-		ADD_LATENT_AUTOMATION_COMMAND(FExecWorldStringLatentCommand(TEXT("StartFPSChart")));
-		ADD_LATENT_AUTOMATION_COMMAND(FEngineWaitLatentCommand(60.0f));
-		ADD_LATENT_AUTOMATION_COMMAND(FExecWorldStringLatentCommand(TEXT("StopFPSChart")));
-		ADD_LATENT_AUTOMATION_COMMAND(FEngineWaitLatentCommand(1.0f));
-	}
+	//Start the matinee and perform the FPS Chart
+	ADD_LATENT_AUTOMATION_COMMAND(FExecWorldStringLatentCommand(TEXT("StartFPSChart")));
+	ADD_LATENT_AUTOMATION_COMMAND(FEngineWaitLatentCommand(60.0f));
+	ADD_LATENT_AUTOMATION_COMMAND(FExecWorldStringLatentCommand(TEXT("StopFPSChart")));
+	ADD_LATENT_AUTOMATION_COMMAND(FEngineWaitLatentCommand(1.0f));
 
 	return true;
 }
 
-
 /* UAutomationTestSettings interface
  *****************************************************************************/
 
