#ue4 - Fix comment on MovementComponent::SetUpdatedComponent()

diff --git a/Engine/Source/Runtime/Engine/Classes/GameFramework/MovementComponent.h b/Engine/Source/Runtime/Engine/Classes/GameFramework/MovementComponent.h
--- a/Engine/Source/Runtime/Engine/Classes/GameFramework/MovementComponent.h
+++ b/Engine/Source/Runtime/Engine/Classes/GameFramework/MovementComponent.h
@@ -152,7 +152,7 @@ public:
 	UFUNCTION()
 	virtual void PhysicsVolumeChanged(class APhysicsVolume* NewVolume);
 
-	/** add updated component to physics volume */
+	/** Assign the component we move and update. */
 	UFUNCTION(BlueprintCallable, Category="Components|Movement")
 	virtual void SetUpdatedComponent(class UPrimitiveComponent* NewUpdatedComponent);
 
