[GitHub] - Various changes to CharacterMovementComponent.

- Add virtual function for swept collision used to find the floor (FloorSweepTest()). bUseFlatBaseForFloorChecks now uses 2 box tests enclosed by the capsule.
- Modify NewFallVelocity() to be more general in regards to gravity orientation. Note that other code still assumes vertical gravity for now.
- ApplyRepulsionForce() optimized if there is no repulsion force.

Pull request: https://github.com/EpicGames/UnrealEngine/pull/415

diff --git a/Engine/Source/Runtime/Engine/Classes/GameFramework/CharacterMovementComponent.h b/Engine/Source/Runtime/Engine/Classes/GameFramework/CharacterMovementComponent.h
--- a/Engine/Source/Runtime/Engine/Classes/GameFramework/CharacterMovementComponent.h
+++ b/Engine/Source/Runtime/Engine/Classes/GameFramework/CharacterMovementComponent.h
@@ -604,10 +604,8 @@ public:
 	uint32 bAlwaysCheckFloor:1;
 
 	/**
-	 * Performs floor checks as if the character is using a box collider.
-	 * This avoids the situation where characters slowly slide down a ledge (as their capsule 'balances' on the edge).
-	 * Note that sweep checks perfomed using the box collision are axis-aligned and use a bounding box that tightly 
-	 * encloses the capsule.
+	 * Performs floor checks as if the character is using a shape with a flat base.
+	 * This avoids the situation where characters slowly lower off the side of a ledge (as their capsule 'balances' on the edge).
 	 */
 	UPROPERTY(Category="Character Movement", EditAnywhere, BlueprintReadWrite, AdvancedDisplay)
 	uint32 bUseFlatBaseForFloorChecks:1;
@@ -876,9 +874,8 @@ public:
 
 	/**
 	 * Compute new falling velocity from given velocity and gravity. Applies the limits of the current Physics Volume's TerminalVelocity.
-	 * Note: currently assumes Gravity points downward when enforcing terminal velocity.
 	 */
-	virtual FVector NewFallVelocity(FVector InitialVelocity, FVector Gravity, float DeltaTime) const;
+	virtual FVector NewFallVelocity(const FVector& InitialVelocity, const FVector& Gravity, float DeltaTime) const;
 
 	/* Determine how deep in water the character is immersed.
 	 * @return float in range 0.0 = not in water, 1.0 = fully immersed
@@ -886,7 +883,7 @@ public:
 	virtual float ImmersionDepth();
 
 	/** 
-	 * Updates Velocity and Acceleration based on the current state, applying the effects of friction and acceleration or deceleration. Does *not* apply gravity.
+	 * Updates Velocity and Acceleration based on the current state, applying the effects of friction and acceleration or deceleration. Does not apply gravity.
 	 * This is used internally during movement updates. Normally you don't need to call this from outside code, but you might want to use it for custom movement modes.
 	 *
 	 * @param	DeltaTime						time elapsed since last frame.
@@ -1231,7 +1228,7 @@ protected:
 	 * @param DeltaSeconds:			Time over which movement occurs
 	 * @param OutStepDownResult:	[Out] If non-null, and a floor check is performed, this will be updated to reflect that result.
 	 */
-	virtual void MoveAlongFloor(const FVector& InVelocity, const float DeltaSeconds, FStepDownResult* OutStepDownResult = NULL);
+	virtual void MoveAlongFloor(const FVector& InVelocity, float DeltaSeconds, FStepDownResult* OutStepDownResult = NULL);
 
 	/**
 	 * Adjusts velocity when walking so that Z velocity is zero.
@@ -1322,6 +1319,29 @@ protected:
 	 */
 	virtual void ComputeFloorDist(const FVector& CapsuleLocation, float LineDistance, float SweepDistance, FFindFloorResult& OutFloorResult, float SweepRadius, const FHitResult* DownwardSweepResult = NULL) const;
 
+	/**
+	 * Sweep against the world and return the first blocking hit.
+	 * Intended for tests against the floor, because it may change the result of impacts on the lower area of the test (especially if bUseFlatBaseForFloorChecks is true).
+	 *
+	 * @param OutHit			First blocking hit found.
+	 * @param Start				Start location of the capsule.
+	 * @param End				End location of the capsule.
+	 * @param TraceChannel		The 'channel' that this trace is in, used to determine which components to hit.
+	 * @param CollisionShape	Capsule collision shape.
+	 * @param Params			Additional parameters used for the trace.
+	 * @param ResponseParam		ResponseContainer to be used for this trace.
+	 * @return True if OutHit contains a blocking hit entry.
+	 */
+	virtual bool FloorSweepTest(
+		struct FHitResult& OutHit,
+		const FVector& Start,
+		const FVector& End,
+		ECollisionChannel TraceChannel,
+		const struct FCollisionShape& CollisionShape,
+		const struct FCollisionQueryParams& Params,
+		const struct FCollisionResponseParams& ResponseParam
+		) const;
+
 	/** Verify that the supplied hit result is a valid landing spot when falling. */
 	virtual bool IsValidLandingSpot(const FVector& CapsuleLocation, const FHitResult& Hit) const;
 
@@ -1329,7 +1349,7 @@ protected:
 	 * Determine whether we should try to find a valid landing spot after an impact with an invalid one (based on the Hit result).
 	 * For example, landing on the lower portion of the capsule on the edge of geometry may be a walkable surface, but could have reported an unwalkable impact normal.
 	 */
-	virtual bool ShouldCheckForValidLandingSpot(const float DeltaTime, const FVector& Delta, const FHitResult& Hit) const;
+	virtual bool ShouldCheckForValidLandingSpot(float DeltaTime, const FVector& Delta, const FHitResult& Hit) const;
 
 	/**
 	 * Check if the result of a sweep test (passed in InHit) might be a valid location to perch, in which case we should use ComputePerchResult to validate the location.
diff --git a/Engine/Source/Runtime/Engine/Private/CharacterMovementComponent.cpp b/Engine/Source/Runtime/Engine/Private/CharacterMovementComponent.cpp
--- a/Engine/Source/Runtime/Engine/Private/CharacterMovementComponent.cpp
+++ b/Engine/Source/Runtime/Engine/Private/CharacterMovementComponent.cpp
@@ -2016,18 +2016,25 @@ FVector UCharacterMovementComponent::AdjustUpperHemisphereImpact(const FVector&
 }
 
 
-FVector UCharacterMovementComponent::NewFallVelocity(FVector InitialVelocity, FVector Gravity, float DeltaTime) const
+FVector UCharacterMovementComponent::NewFallVelocity(const FVector& InitialVelocity, const FVector& Gravity, float DeltaTime) const
 {
 	FVector Result = InitialVelocity;
-	const float TerminalLimit = FMath::Abs(GetPhysicsVolume()->TerminalVelocity);
 
-	// Only apply gravity if below terminal velocity, and don't exceed it if so.
-	if (InitialVelocity.Z > -TerminalLimit)
+	if (!Gravity.IsZero())
 	{
-		Result = InitialVelocity + (Gravity * DeltaTime);
-		Result.Z = FMath::Max(Result.Z, -TerminalLimit);
+		// Apply gravity.
+		Result += Gravity * DeltaTime;
+
+		const FVector GravityDir = Gravity.SafeNormal();
+		const float TerminalLimit = FMath::Abs(GetPhysicsVolume()->TerminalVelocity);
+
+		// Don't exceed terminal velocity.
+		if ((Result | GravityDir) > TerminalLimit)
+		{
+			Result = FVector::PointPlaneProject(Result, FVector::ZeroVector, GravityDir) + GravityDir * TerminalLimit;
+		}
 	}
-	
+
 	return Result;
 }
 
@@ -3027,21 +3034,10 @@ bool UCharacterMovementComponent::FindAirControlImpact(float DeltaTime, float Ti
 		FCollisionQueryParams CapsuleQuery(FallingTraceParamsTag, false, CharacterOwner);
 		FCollisionResponseParams ResponseParam;
 		InitCollisionParams(CapsuleQuery, ResponseParam);
-		const FVector PawnLocation = CharacterOwner->GetActorLocation();
-		const ECollisionChannel CollisionChannel = UpdatedComponent->GetCollisionObjectType();
-
+		const FVector CapsuleLocation = UpdatedComponent->GetComponentLocation();
 		const FCollisionShape CapsuleShape = GetPawnCapsuleCollisionShape(SHRINK_None);
-		if (bUseFlatBaseForFloorChecks)
-		{
-			const FCollisionShape BoxShape = FCollisionShape::MakeBox(FVector(CapsuleShape.GetCapsuleRadius(), CapsuleShape.GetCapsuleRadius(), CapsuleShape.GetCapsuleHalfHeight()));
-			GetWorld()->SweepSingle( OutHitResult, PawnLocation, PawnLocation + TestWalk, FQuat::Identity, CollisionChannel, BoxShape, CapsuleQuery, ResponseParam );
-		}
-		else
-		{
-			GetWorld()->SweepSingle( OutHitResult, PawnLocation, PawnLocation + TestWalk, FQuat::Identity, CollisionChannel, CapsuleShape, CapsuleQuery, ResponseParam );
-		}
 
-		if (OutHitResult.bBlockingHit)
+		if (FloorSweepTest(OutHitResult, CapsuleLocation, CapsuleLocation + TestWalk, UpdatedComponent->GetCollisionObjectType(), CapsuleShape, CapsuleQuery, ResponseParam))
 		{
 			return true;
 		}
@@ -3245,7 +3241,7 @@ FVector UCharacterMovementComponent::ComputeGroundMovementDelta(const FVector& D
 }
 
 
-void UCharacterMovementComponent::MoveAlongFloor(const FVector& InVelocity, const float DeltaSeconds, FStepDownResult* OutStepDownResult)
+void UCharacterMovementComponent::MoveAlongFloor(const FVector& InVelocity, float DeltaSeconds, FStepDownResult* OutStepDownResult)
 {
 	const FVector Delta = FVector(InVelocity.X, InVelocity.Y, 0.f) * DeltaSeconds;
 	
@@ -4136,15 +4132,7 @@ void UCharacterMovementComponent::ComputeFloorDist(const FVector& CapsuleLocatio
 		FCollisionShape CapsuleShape = FCollisionShape::MakeCapsule(SweepRadius, PawnHalfHeight - ShrinkHeight);
 
 		FHitResult Hit(1.f);
-		if(bUseFlatBaseForFloorChecks)
-		{
-			const FCollisionShape BoxShape = FCollisionShape::MakeBox(FVector(CapsuleShape.GetCapsuleRadius(), CapsuleShape.GetCapsuleRadius(), CapsuleShape.GetCapsuleHalfHeight()));
-			bBlockingHit = GetWorld()->SweepSingle(Hit, CapsuleLocation, CapsuleLocation + FVector(0.f,0.f,-TraceDist), FQuat::Identity, CollisionChannel, BoxShape, QueryParams, ResponseParam);
-		}
-		else
-		{
-			bBlockingHit = GetWorld()->SweepSingle(Hit, CapsuleLocation, CapsuleLocation + FVector(0.f,0.f,-TraceDist), FQuat::Identity, CollisionChannel, CapsuleShape, QueryParams, ResponseParam);
-		}
+		bBlockingHit = FloorSweepTest(Hit, CapsuleLocation, CapsuleLocation + FVector(0.f,0.f,-TraceDist), CollisionChannel, CapsuleShape, QueryParams, ResponseParam);
 
 		if (bBlockingHit)
 		{
@@ -4159,15 +4147,7 @@ void UCharacterMovementComponent::ComputeFloorDist(const FVector& CapsuleLocatio
 				CapsuleShape.Capsule.HalfHeight = FMath::Max(PawnHalfHeight - ShrinkHeight, CapsuleShape.Capsule.Radius);
 				Hit.Reset(1.f, false);
 
-				if(bUseFlatBaseForFloorChecks)
-				{
-					const FCollisionShape BoxShape = FCollisionShape::MakeBox(FVector(CapsuleShape.GetCapsuleRadius(), CapsuleShape.GetCapsuleRadius(), CapsuleShape.GetCapsuleHalfHeight()));
-					bBlockingHit = GetWorld()->SweepSingle(Hit, CapsuleLocation, CapsuleLocation + FVector(0.f,0.f,-TraceDist), FQuat::Identity, CollisionChannel, BoxShape, QueryParams, ResponseParam);
-				}
-				else
-				{
-					bBlockingHit = GetWorld()->SweepSingle(Hit, CapsuleLocation, CapsuleLocation + FVector(0.f,0.f,-TraceDist), FQuat::Identity, CollisionChannel, CapsuleShape, QueryParams, ResponseParam);
-				}
+				bBlockingHit = FloorSweepTest(Hit, CapsuleLocation, CapsuleLocation + FVector(0.f,0.f,-TraceDist), CollisionChannel, CapsuleShape, QueryParams, ResponseParam);
 			}
 
 			// Reduce hit distance by ShrinkHeight because we shrank the capsule for the trace.
@@ -4333,6 +4313,44 @@ void UCharacterMovementComponent::FindFloor(const FVector& CapsuleLocation, FFin
 }
 
 
+bool UCharacterMovementComponent::FloorSweepTest(
+	FHitResult& OutHit,
+	const FVector& Start,
+	const FVector& End,
+	ECollisionChannel TraceChannel,
+	const struct FCollisionShape& CollisionShape,
+	const struct FCollisionQueryParams& Params,
+	const struct FCollisionResponseParams& ResponseParam
+	) const
+{
+	bool bBlockingHit = false;
+
+	if (!bUseFlatBaseForFloorChecks)
+	{
+		bBlockingHit = GetWorld()->SweepSingle(OutHit, Start, End, FQuat::Identity, TraceChannel, CollisionShape, Params, ResponseParam);
+	}
+	else
+	{
+		// Test with a box that is enclosed by the capsule.
+		const float CapsuleRadius = CollisionShape.GetCapsuleRadius();
+		const float CapsuleHeight = CollisionShape.GetCapsuleHalfHeight();
+		const FCollisionShape BoxShape = FCollisionShape::MakeBox(FVector(CapsuleRadius * 0.707f, CapsuleRadius * 0.707f, CapsuleHeight));
+
+		// First test with the box rotated so the corners are along the major axes (ie rotated 45 degrees).
+		bBlockingHit = GetWorld()->SweepSingle(OutHit, Start, End, FQuat(FVector(0.f, 0.f, -1.f), PI * 0.25f), TraceChannel, BoxShape, Params, ResponseParam);
+
+		if (!bBlockingHit)
+		{
+			// Test again with the same box, not rotated.
+			OutHit.Reset(1.f, false);			
+			bBlockingHit = GetWorld()->SweepSingle(OutHit, Start, End, FQuat::Identity, TraceChannel, BoxShape, Params, ResponseParam);
+		}
+	}
+
+	return bBlockingHit;
+}
+
+
 bool UCharacterMovementComponent::IsValidLandingSpot(const FVector& CapsuleLocation, const FHitResult& Hit) const
 {
 	if (!Hit.bBlockingHit)
@@ -4378,7 +4396,7 @@ bool UCharacterMovementComponent::IsValidLandingSpot(const FVector& CapsuleLocat
 }
 
 
-bool UCharacterMovementComponent::ShouldCheckForValidLandingSpot(const float DeltaTime, const FVector& Delta, const FHitResult& Hit) const
+bool UCharacterMovementComponent::ShouldCheckForValidLandingSpot(float DeltaTime, const FVector& Delta, const FHitResult& Hit) const
 {
 	// See if we hit an edge of a surface on the lower portion of the capsule.
 	// In this case the normal will not equal the impact normal, and a downward sweep may find a walkable surface on top of the edge.
@@ -6085,24 +6103,24 @@ void UCharacterMovementComponent::SetAvoidanceEnabled(bool bEnable)
 	}
 }
 
-void UCharacterMovementComponent::ApplyRepulsionForce( float DeltaSeconds )
+void UCharacterMovementComponent::ApplyRepulsionForce(float DeltaSeconds)
 {
-	FCollisionQueryParams QueryParams;
-	QueryParams.bReturnFaceIndex = false;
-	QueryParams.bReturnPhysicalMaterial = false;
-
-	const float CapsuleRadius = UpdatedComponent->GetCollisionShape().GetCapsuleRadius();
-	const float CapsuleHalfHeight = UpdatedComponent->GetCollisionShape().GetCapsuleHalfHeight();
-	const float RepulsionForceRadius = CapsuleRadius * 1.2f;
-	const float StopBodyDistance = 2.5f;
-	
-	if (RepulsionForce > 0.0f)
+	if (UpdatedComponent && RepulsionForce > 0.0f)
 	{
+		FCollisionQueryParams QueryParams;
+		QueryParams.bReturnFaceIndex = false;
+		QueryParams.bReturnPhysicalMaterial = false;
+
+		const FCollisionShape CollisionShape = UpdatedComponent->GetCollisionShape();
+		const float CapsuleRadius = CollisionShape.GetCapsuleRadius();
+		const float CapsuleHalfHeight = CollisionShape.GetCapsuleHalfHeight();
+		const float RepulsionForceRadius = CapsuleRadius * 1.2f;
+		const float StopBodyDistance = 2.5f;
+
 		const TArray<FOverlapInfo>& Overlaps = UpdatedComponent->GetOverlapInfos();
-		
 		const FVector MyLocation = UpdatedComponent->GetComponentLocation();
 
-		for (int32 i=0; i < Overlaps.Num(); ++i)
+		for (int32 i=0; i < Overlaps.Num(); i++)
 		{
 			const FOverlapInfo& Overlap = Overlaps[i];
 
@@ -6134,9 +6152,8 @@ void UCharacterMovementComponent::ApplyRepulsionForce( float DeltaSeconds )
 				continue;
 			}
 
-
 			FTransform BodyTransform = OverlapBody->GetUnrealWorldTransform();
-			
+
 			FVector BodyVelocity = OverlapBody->GetUnrealWorldVelocity();
 			FVector BodyLocation = BodyTransform.GetLocation();
 
@@ -6148,14 +6165,14 @@ void UCharacterMovementComponent::ApplyRepulsionForce( float DeltaSeconds )
 
 			FVector HitLoc = Hit.ImpactPoint;
 			bool bIsPenetrating = Hit.bStartPenetrating || Hit.PenetrationDepth > 2.5f;
-			
+
 			// If we didn't hit the capsule, we're inside the capsule
 			if(!bHasHit) 
 			{ 
 				HitLoc = BodyLocation; 
 				bIsPenetrating = true;
 			}
-			
+
 			const float DistanceNow = (HitLoc - BodyLocation).SizeSquared2D();
 			const float DistanceLater = (HitLoc - (BodyLocation + BodyVelocity * DeltaSeconds)).SizeSquared2D();
 
