Fixed bounds issue.

Back out revision 227 from //depot/UE4/Engine/Source/Runtime/Engine/Private/SkeletalMeshComponent.cpp

346627	CRITICAL: Regression: ELEMENTAL: Some assets are not rendering
346786	Framework: Culling not considering Physics Asset Bounds

diff --git a/Engine/Source/Runtime/Engine/Private/SkeletalMeshComponent.cpp b/Engine/Source/Runtime/Engine/Private/SkeletalMeshComponent.cpp
--- a/Engine/Source/Runtime/Engine/Private/SkeletalMeshComponent.cpp
+++ b/Engine/Source/Runtime/Engine/Private/SkeletalMeshComponent.cpp
@@ -1027,7 +1027,7 @@ void USkeletalMeshComponent::PostAnimEvaluation(FAnimationEvaluationContext& Eva
 		SCOPE_CYCLE_COUNTER(STAT_UpdateLocalToWorldAndOverlaps);
 
 		// New bone positions need to be sent to render thread
-		UpdateChildTransforms();
+		UpdateComponentToWorld();
 
 		// animation often change overlap. 
 		UpdateOverlaps();
