Landscape : Enabled simulation mode for XY Offseted Landscape (TTP #346165)

diff --git a/Engine/Source/Runtime/Engine/Private/Landscape/LandscapeCollision.cpp b/Engine/Source/Runtime/Engine/Private/Landscape/LandscapeCollision.cpp
--- a/Engine/Source/Runtime/Engine/Private/Landscape/LandscapeCollision.cpp
+++ b/Engine/Source/Runtime/Engine/Private/Landscape/LandscapeCollision.cpp
@@ -148,6 +148,7 @@ void ULandscapeHeightfieldCollisionComponent::CreatePhysicsState()
 
 				// Heightfield is used for simple and complex collision
 				PQueryFilterData.word3 |= (EPDF_SimpleCollision | EPDF_ComplexCollision);
+				PSimFilterData.word3 |= (EPDF_SimpleCollision | EPDF_ComplexCollision);
 				HeightFieldShapeSync->setQueryFilterData(PQueryFilterData);
 				HeightFieldShapeSync->setSimulationFilterData(PSimFilterData);
 				HeightFieldShapeSync->setFlag(PxShapeFlag::eSCENE_QUERY_SHAPE, true);
@@ -682,8 +683,10 @@ void ULandscapeMeshCollisionComponent::CreatePhysicsState()
 
 				// Heightfield is used for simple and complex collision
 				PQueryFilterData.word3 |= (EPDF_SimpleCollision | EPDF_ComplexCollision);
+				PSimFilterData.word3 |= (EPDF_SimpleCollision | EPDF_ComplexCollision);
 				MeshShapeSync->setQueryFilterData(PQueryFilterData);
-
+				MeshShapeSync->setSimulationFilterData(PSimFilterData);
+				MeshShapeSync->setFlag(PxShapeFlag::eSCENE_QUERY_SHAPE, true);
 				MeshShapeSync->setFlag(PxShapeFlag::eSIMULATION_SHAPE, true);
 				MeshShapeSync->setFlag(PxShapeFlag::eVISUALIZATION, true);
 
@@ -697,7 +700,10 @@ void ULandscapeMeshCollisionComponent::CreatePhysicsState()
 					PxShape* MeshShapeAsync = MeshActorAsync->createShape(PTriMeshGeom, MeshRef->UsedPhysicalMaterialArray.GetTypedData(), MeshRef->UsedPhysicalMaterialArray.Num());
 					check(MeshShapeAsync);
 
-					// No need for query filter data.  That will all be take care of by the sync actor
+					MeshShapeAsync->setQueryFilterData(PQueryFilterData);
+					MeshShapeAsync->setSimulationFilterData(PSimFilterData);
+					// Only perform scene queries in the synchronous scene for static shapes
+					MeshShapeAsync->setFlag(PxShapeFlag::eSCENE_QUERY_SHAPE, false);
 					MeshShapeAsync->setFlag(PxShapeFlag::eSIMULATION_SHAPE, true);
 					MeshShapeAsync->setFlag(PxShapeFlag::eVISUALIZATION, true);	// Setting visualization flag, in case we visualize only the async scene
 				}
