UE4 - tweaks to prevent flushes for common render commands

diff --git a/Engine/Source/Runtime/Engine/Private/GPUSkinVertexFactory.cpp b/Engine/Source/Runtime/Engine/Private/GPUSkinVertexFactory.cpp
--- a/Engine/Source/Runtime/Engine/Private/GPUSkinVertexFactory.cpp
+++ b/Engine/Source/Runtime/Engine/Private/GPUSkinVertexFactory.cpp
@@ -134,10 +134,20 @@ void FGPUBaseSkinVertexFactory::ShaderDataType::UpdateBoneData()
 		}
 		if(NumBones)
 		{
-			float* Data = (float*)RHILockVertexBuffer(BoneBuffer.VertexBufferRHI, 0, VectorArraySize, RLM_WriteOnly);
-			checkSlow(Data);
-			FMemory::Memcpy(Data, BoneMatrices.GetTypedData(), NumBones * sizeof(BoneMatrices[0]));
-			RHIUnlockVertexBuffer(BoneBuffer.VertexBufferRHI);
+#if PLATFORM_SUPPORTS_RHI_THREAD
+			check(VectorArraySize == NumBones * sizeof(BoneMatrices[0]));
+			if (GRHIThread)
+			{
+				GRHICommandList.GetImmediateCommandList().UpdateVertexBuffer(BoneBuffer.VertexBufferRHI, BoneMatrices.GetTypedData(), NumBones * sizeof(BoneMatrices[0]));
+			}
+			else
+#endif
+			{
+				float* Data = (float*)RHILockVertexBuffer(BoneBuffer.VertexBufferRHI, 0, VectorArraySize, RLM_WriteOnly);
+				checkSlow(Data);
+				FMemory::Memcpy(Data, BoneMatrices.GetTypedData(), NumBones * sizeof(BoneMatrices[0]));
+				RHIUnlockVertexBuffer(BoneBuffer.VertexBufferRHI);
+			}
 		}
 	}
 	else
@@ -165,7 +175,7 @@ float* FBoneDataVertexBuffer::LockData()
 	check(IsInRenderingThread() && GetSizeX() && IsValidRef(BoneBuffer));
 	check(!AllocSize && !AllocBlock);
 	AllocSize = ComputeMemorySize();
-	AllocBlock = FMemory::Malloc(((AllocSize + 15) / 16) * 16, 16);
+	AllocBlock = FMemory::Malloc(AllocSize + 15, 16);
 	float* Data = (float*)AllocBlock;
 	check(AllocSize && AllocBlock);
 	check(Data);
@@ -177,7 +187,7 @@ void FBoneDataVertexBuffer::UnlockData()
 	check(IsValidRef(BoneBuffer));
 	check(IsInRenderingThread() && AllocBlock && AllocSize);
 	FRHICommandListExecutor::GetImmediateCommandList().UpdateVertexBuffer(BoneBuffer.VertexBufferRHI, AllocBlock, AllocSize);
-	// UpdateVertexBuffer will free the block after the update
+	FMemory::Free(AllocBlock);
 	AllocBlock = nullptr;
 	AllocSize = 0;
 }
diff --git a/Engine/Source/Runtime/Engine/Private/Particles/ParticleEmitterInstances.cpp b/Engine/Source/Runtime/Engine/Private/Particles/ParticleEmitterInstances.cpp
--- a/Engine/Source/Runtime/Engine/Private/Particles/ParticleEmitterInstances.cpp
+++ b/Engine/Source/Runtime/Engine/Private/Particles/ParticleEmitterInstances.cpp
@@ -3513,7 +3513,7 @@ bool FParticleMeshEmitterInstance::FillReplayData( FDynamicEmitterReplayDataBase
 
 FDynamicEmitterDataBase::FDynamicEmitterDataBase(const UParticleModuleRequired* RequiredModule)
 	: bSelected(false)
-	, VertexFactory(NULL)
+	, ParticleVertexFactory(nullptr)
 {
 }
 
diff --git a/Engine/Source/Runtime/Engine/Private/Particles/ParticleGpuSimulation.cpp b/Engine/Source/Runtime/Engine/Private/Particles/ParticleGpuSimulation.cpp
--- a/Engine/Source/Runtime/Engine/Private/Particles/ParticleGpuSimulation.cpp
+++ b/Engine/Source/Runtime/Engine/Private/Particles/ParticleGpuSimulation.cpp
@@ -1531,10 +1531,20 @@ void InjectNewParticles(FRHICommandList& RHICmdList, ERHIFeatureLevel::Type Feat
 	{
 		// Copy new particles in to the vertex buffer.
 		const int32 ParticlesThisDrawCall = FMath::Min<int32>( ParticleCount, MaxParticlesPerDrawCall );
-		void* Dest = RHILockVertexBuffer( ScratchVertexBufferRHI, 0, ParticlesThisDrawCall * sizeof(FNewParticle), RLM_WriteOnly );
 		const void* Src = NewParticles.GetTypedData() + FirstParticle;
-		FMemory::Memcpy( Dest, Src, ParticlesThisDrawCall * sizeof(FNewParticle) );
-		RHIUnlockVertexBuffer( ScratchVertexBufferRHI );
+#if PLATFORM_SUPPORTS_RHI_THREAD
+		if (GRHIThread)
+		{
+			RHICmdList.UpdateVertexBuffer(ScratchVertexBufferRHI, Src, ParticlesThisDrawCall * sizeof(FNewParticle));
+		}
+		else
+#endif
+		{
+			void* Dest = RHILockVertexBuffer( ScratchVertexBufferRHI, 0, ParticlesThisDrawCall * sizeof(FNewParticle), RLM_WriteOnly );
+			FMemory::Memcpy( Dest, Src, ParticlesThisDrawCall * sizeof(FNewParticle) );
+			RHIUnlockVertexBuffer( ScratchVertexBufferRHI );
+
+		}
 		ParticleCount -= ParticlesThisDrawCall;
 		FirstParticle += ParticlesThisDrawCall;
 
diff --git a/Engine/Source/Runtime/Engine/Private/Particles/ParticleSystemRender.cpp b/Engine/Source/Runtime/Engine/Private/Particles/ParticleSystemRender.cpp
--- a/Engine/Source/Runtime/Engine/Private/Particles/ParticleSystemRender.cpp
+++ b/Engine/Source/Runtime/Engine/Private/Particles/ParticleSystemRender.cpp
@@ -1453,7 +1453,7 @@ int32 FDynamicSpriteEmitterData::RenderEmitter(FParticleSystemSceneProxy* Proxy,
 
 	const FDynamicSpriteEmitterReplayDataBase* SourceData = GetSourceData();
 	check(SourceData);
-	FParticleSpriteVertexFactory* SpriteVertexFactory = (FParticleSpriteVertexFactory*)VertexFactory;
+	FParticleSpriteVertexFactory* SpriteVertexFactory = (FParticleSpriteVertexFactory*)GetVertexFactory(Proxy);
 
 	int32 NumDraws = 0;
 	if (SourceData->EmitterRenderMode == ERM_Normal)
@@ -1561,23 +1561,13 @@ int32 FDynamicSpriteEmitterData::RenderEmitter(FParticleSystemSceneProxy* Proxy,
 }
 
 
+FParticleVertexFactoryBase* FDynamicSpriteEmitterData::BuildVertexFactory(const FParticleSystemSceneProxy* InOwnerProxy)
+{
+	return GParticleVertexFactoryPool.GetParticleVertexFactory(PVFT_Sprite, InOwnerProxy->GetScene()->GetFeatureLevel());
+}
 
-/**
- *	Create the render thread resources for this emitter data
- *
- *	@param	InOwnerProxy	The proxy that owns this dynamic emitter data
- *
- *	@return	bool			true if successful, false if failed
- */
 void FDynamicSpriteEmitterData::UpdateRenderThreadResourcesEmitter(const FParticleSystemSceneProxy* InOwnerProxy)
 {
-	// Create the vertex factory...
-	if (VertexFactory == NULL)
-	{
-		VertexFactory = GParticleVertexFactoryPool.GetParticleVertexFactory(PVFT_Sprite, InOwnerProxy->GetScene()->GetFeatureLevel());
-		check(VertexFactory);
-	}
-
 	// Generate the uniform buffer.
 	const FDynamicSpriteEmitterReplayDataBase* SourceData = GetSourceData();
 	if( SourceData )
@@ -1784,23 +1774,15 @@ void FDynamicMeshEmitterData::Init( bool bInSelected,
 	Source.MaterialInterface = NULL;
 }
 
-/**
- *	Create the render thread resources for this emitter data
- *
- *	@param	InOwnerProxy	The proxy that owns this dynamic emitter data
- *
- *	@return	bool			true if successful, false if failed
- */
-void FDynamicMeshEmitterData::UpdateRenderThreadResourcesEmitter(const FParticleSystemSceneProxy* InOwnerProxy)
+FParticleVertexFactoryBase* FDynamicMeshEmitterData::BuildVertexFactory(const FParticleSystemSceneProxy* InOwnerProxy)
 {
-	// Create the vertex factory
-	if (VertexFactory == NULL)
-	{
-		VertexFactory = GParticleVertexFactoryPool.GetParticleVertexFactory(PVFT_Mesh, InOwnerProxy->GetScene()->GetFeatureLevel());
-		check(VertexFactory);
-		SetupVertexFactory((FMeshParticleVertexFactory*)VertexFactory, StaticMesh->RenderData->LODResources[0]);
-	}
+	FParticleVertexFactoryBase* VertexFactory = GParticleVertexFactoryPool.GetParticleVertexFactory(PVFT_Mesh, InOwnerProxy->GetScene()->GetFeatureLevel());
+	SetupVertexFactory((FMeshParticleVertexFactory*)VertexFactory, StaticMesh->RenderData->LODResources[0]);
+	return VertexFactory;
+}
 
+void FDynamicMeshEmitterData::UpdateRenderThreadResourcesEmitter(const FParticleSystemSceneProxy* InOwnerProxy)
+{
 	// Generate the uniform buffer.
 	const FDynamicSpriteEmitterReplayDataBase* SourceData = GetSourceData();
 	if( SourceData )
@@ -2043,7 +2025,7 @@ int32 FDynamicMeshEmitterData::RenderEmitter(FParticleSystemSceneProxy* Proxy, F
 		// Only unsorted instance data will exist for unsorted emitters and must be used for all views being rendered. 
 		const int32 InstanceBufferIndex = ((Source.SortMode == PSORTMODE_None) ? 0 : ViewIndex);
 
-		FMeshParticleVertexFactory* MeshVertexFactory = (FMeshParticleVertexFactory*)VertexFactory;
+		FMeshParticleVertexFactory* MeshVertexFactory = (FMeshParticleVertexFactory*)GetVertexFactory(Proxy);
 
 		if(bInstanced)
 		{
@@ -2167,7 +2149,7 @@ void FDynamicMeshEmitterData::PreRenderView(FParticleSystemSceneProxy* Proxy, co
 	}
 
 	// Setup the vertex factory.
-	FMeshParticleVertexFactory* MeshVertexFactory = (FMeshParticleVertexFactory*)VertexFactory;
+	FMeshParticleVertexFactory* MeshVertexFactory = (FMeshParticleVertexFactory*)GetVertexFactory(Proxy);
 	MeshVertexFactory->SetUniformBuffer( UniformBuffer );
 
 	// Update the primitive uniform buffer
@@ -2200,7 +2182,7 @@ void FDynamicMeshEmitterData::PreRenderView(FParticleSystemSceneProxy* Proxy, co
 					FMeshBatch& Mesh = *MeshPtr;
 					Mesh.Elements.Reset();
 					FMeshBatchElement& BatchElement = *new(Mesh.Elements) FMeshBatchElement();
-					Mesh.VertexFactory = VertexFactory;
+					Mesh.VertexFactory = GetVertexFactory(Proxy);
 					Mesh.DynamicVertexData = NULL;
 					Mesh.LCI = NULL;
 					Mesh.UseDynamicData = false;
@@ -2893,22 +2875,9 @@ void FDynamicBeam2EmitterData::Init( bool bInSelected )
 	Source.MaterialInterface = NULL;
 }
 
-/**
- *	Create the render thread resources for this emitter data
- *
- *	@param	InOwnerProxy	The proxy that owns this dynamic emitter data
- *
- *	@return	bool			true if successful, false if failed
- */
-void FDynamicBeam2EmitterData::UpdateRenderThreadResourcesEmitter(const FParticleSystemSceneProxy* InOwnerProxy)
+FParticleVertexFactoryBase* FDynamicBeam2EmitterData::BuildVertexFactory(const FParticleSystemSceneProxy* InOwnerProxy)
 {
-	// Create the vertex factory...
-	//@todo. Cache these??
-	if (VertexFactory == NULL)
-	{
-		VertexFactory = (FParticleBeamTrailVertexFactory*)(GParticleVertexFactoryPool.GetParticleVertexFactory(PVFT_BeamTrail, InOwnerProxy->GetScene()->GetFeatureLevel()));
-		check(VertexFactory);
-	}
+	return (FParticleBeamTrailVertexFactory*)(GParticleVertexFactoryPool.GetParticleVertexFactory(PVFT_BeamTrail, InOwnerProxy->GetScene()->GetFeatureLevel()));
 }
 
 /** Perform the actual work of filling the buffer, often called from another thread 
@@ -3122,7 +3091,7 @@ int32 FDynamicBeam2EmitterData::RenderEmitter(FParticleSystemSceneProxy* Proxy,
 			}
 
 			// Create and set the uniform buffer for this emitter.
-			FParticleBeamTrailVertexFactory* BeamTrailVertexFactory = (FParticleBeamTrailVertexFactory*)VertexFactory;
+			FParticleBeamTrailVertexFactory* BeamTrailVertexFactory = (FParticleBeamTrailVertexFactory*)GetVertexFactory(Proxy);
 			BeamTrailVertexFactory->SetBeamTrailUniformBuffer( CreateBeamTrailUniformBuffer( Proxy, &Source, View ) );			
 
 			if( Allocation && IndexAllocation && Allocation->IsValid() && IndexAllocation->IsValid() )
@@ -5849,21 +5818,9 @@ void FDynamicTrailsEmitterData::Init(bool bInSelected)
 	SourcePointer->MaterialInterface = NULL;
 }
 
-/**
- *	Create the render thread resources for this emitter data
- *
- *	@param	InOwnerProxy	The proxy that owns this dynamic emitter data
- *
- *	@return	bool			true if successful, false if failed
- */
-void FDynamicTrailsEmitterData::UpdateRenderThreadResourcesEmitter(const FParticleSystemSceneProxy* InOwnerProxy)
+FParticleVertexFactoryBase* FDynamicTrailsEmitterData::BuildVertexFactory(const FParticleSystemSceneProxy* InOwnerProxy)
 {
-	// Create the vertex factory...
-	if (VertexFactory == NULL)
-	{
-		VertexFactory = (FParticleBeamTrailVertexFactory*)(GParticleVertexFactoryPool.GetParticleVertexFactory(PVFT_BeamTrail, InOwnerProxy->GetScene()->GetFeatureLevel()));
-		check(VertexFactory);
-	}
+	return (FParticleBeamTrailVertexFactory*)(GParticleVertexFactoryPool.GetParticleVertexFactory(PVFT_BeamTrail, InOwnerProxy->GetScene()->GetFeatureLevel()));
 }
 
 void FDynamicTrailsEmitterData::GetDynamicMeshElementsEmitter(const FParticleSystemSceneProxy* Proxy, const FSceneView* View, const FSceneViewFamily& ViewFamily, int32 ViewIndex, FMeshElementCollector& Collector) const
@@ -6038,7 +5995,7 @@ int32 FDynamicTrailsEmitterData::RenderEmitter(FParticleSystemSceneProxy* Proxy,
 		if ( Allocation && IndexAllocation && Allocation->IsValid() && IndexAllocation->IsValid() && (!bUsesDynamicParameter || (DynamicParameterAllocation && DynamicParameterAllocation->IsValid())) )
 		{
 			// Create and set the uniform buffer for this emitter.
-			FParticleBeamTrailVertexFactory* BeamTrailVertexFactory = (FParticleBeamTrailVertexFactory*)VertexFactory;
+			FParticleBeamTrailVertexFactory* BeamTrailVertexFactory = (FParticleBeamTrailVertexFactory*)GetVertexFactory(Proxy);
 			BeamTrailVertexFactory->SetBeamTrailUniformBuffer( CreateBeamTrailUniformBuffer( Proxy, SourcePointer, View ) );
 			BeamTrailVertexFactory->SetVertexBuffer( Allocation->VertexBuffer, Allocation->VertexOffset, GetDynamicVertexStride(View->GetFeatureLevel()) );
 			BeamTrailVertexFactory->SetDynamicParameterBuffer( DynamicParameterAllocation ? DynamicParameterAllocation->VertexBuffer : NULL, DynamicParameterAllocation ? DynamicParameterAllocation->VertexOffset : 0, GetDynamicParameterVertexStride() );
diff --git a/Engine/Source/Runtime/Engine/Public/ParticleHelper.h b/Engine/Source/Runtime/Engine/Public/ParticleHelper.h
--- a/Engine/Source/Runtime/Engine/Public/ParticleHelper.h
+++ b/Engine/Source/Runtime/Engine/Public/ParticleHelper.h
@@ -1292,19 +1292,28 @@ struct FDynamicEmitterDataBase
 	 *	Create the render thread resources for this emitter data
 	 *
 	 *	@param	InOwnerProxy	The proxy that owns this dynamic emitter data
-	 *
-	 *	@return	bool			true if successful, false if failed
 	 */
 	virtual void UpdateRenderThreadResourcesEmitter(const FParticleSystemSceneProxy* InOwnerProxy)
 	{
 	}
 
+	/**
+	 *	Get the vertex factory for this emitter data, possibly creating it
+	 *	@param	InOwnerProxy	The proxy that owns this dynamic emitter data
+	 */
+	FParticleVertexFactoryBase* GetVertexFactory(const FParticleSystemSceneProxy* InOwnerProxy)
+	{
+		if (!ParticleVertexFactory)
+		{
+			ParticleVertexFactory = BuildVertexFactory(InOwnerProxy);
+		}
+		return ParticleVertexFactory;
+	}
+
 	/**
 	 *	Release the render thread resources for this emitter data
 	 *
 	 *	@param	InOwnerProxy	The proxy that owns this dynamic emitter data
-	 *
-	 *	@return	bool			true if successful, false if failed
 	 */
 	virtual void ReleaseRenderThreadResources(const FParticleSystemSceneProxy* InOwnerProxy)
 	{
@@ -1353,8 +1362,18 @@ struct FDynamicEmitterDataBase
 	/** true if this emitter has valid rendering data */
 	uint32	bValid:1;
 
+protected:
+	/**
+	 *	Create the vertex factory for this emitter data
+	 *	@param	InOwnerProxy	The proxy that owns this dynamic emitter data
+	 */
+	virtual FParticleVertexFactoryBase* BuildVertexFactory(const FParticleSystemSceneProxy* InOwnerProxy)
+	{
+		return nullptr;
+	}
+
 	/** The vertex factory used for rendering */
-	FParticleVertexFactoryBase* VertexFactory; // RENDER-THREAD USAGE ONLY!!!
+	FParticleVertexFactoryBase* ParticleVertexFactory; // RENDER-THREAD USAGE ONLY!!!
 
 private:
 
@@ -1363,10 +1382,10 @@ private:
 	 */
 	void ReturnVertexFactory()
 	{
-		if (VertexFactory != NULL)
+		if (ParticleVertexFactory != NULL)
 		{
-			GParticleVertexFactoryPool.ReturnParticleVertexFactory(VertexFactory);
-			VertexFactory = NULL;
+			GParticleVertexFactoryPool.ReturnParticleVertexFactory(ParticleVertexFactory);
+			ParticleVertexFactory = NULL;
 		}
 	}
 };
@@ -1704,6 +1723,8 @@ struct FDynamicSpriteEmitterData : public FDynamicSpriteEmitterDataBase
 	 */
 	virtual void UpdateRenderThreadResourcesEmitter(const FParticleSystemSceneProxy* InOwnerProxy);
 
+	virtual FParticleVertexFactoryBase* BuildVertexFactory(const FParticleSystemSceneProxy* InOwnerProxy);
+
 	/**
 	 *	Release the render thread resources for this emitter data
 	 *
@@ -1800,6 +1821,8 @@ struct FDynamicMeshEmitterData : public FDynamicSpriteEmitterDataBase
 	 */
 	virtual void UpdateRenderThreadResourcesEmitter(const FParticleSystemSceneProxy* InOwnerProxy);
 
+	virtual FParticleVertexFactoryBase* BuildVertexFactory(const FParticleSystemSceneProxy* InOwnerProxy);
+
 	/**
 	 *	Release the render thread resources for this emitter data
 	 *
@@ -2105,16 +2128,10 @@ struct FDynamicBeam2EmitterData : public FDynamicSpriteEmitterDataBase
 	/** Initialize this emitter's dynamic rendering data, called after source data has been filled in */
 	void Init( bool bInSelected );
 
-	/**
-	 *	Create the render thread resources for this emitter data
-	 *
-	 *	@param	InOwnerProxy	The proxy that owns this dynamic emitter data
-	 *
-	 *	@return	bool			true if successful, false if failed
-	 */
-	virtual void UpdateRenderThreadResourcesEmitter(const FParticleSystemSceneProxy* InOwnerProxy);
 
 	virtual void GetDynamicMeshElementsEmitter(const FParticleSystemSceneProxy* Proxy, const FSceneView* View, const FSceneViewFamily& ViewFamily, int32 ViewIndex, FMeshElementCollector& Collector) const override;
+	
+	virtual FParticleVertexFactoryBase* BuildVertexFactory(const FParticleSystemSceneProxy* InOwnerProxy);
 
 	/**
 	 *	Render thread only draw call
@@ -2276,14 +2293,7 @@ struct FDynamicTrailsEmitterData : public FDynamicSpriteEmitterDataBase
 	/** Initialize this emitter's dynamic rendering data, called after source data has been filled in */
 	virtual void Init(bool bInSelected);
 
-	/**
-	 *	Create the render thread resources for this emitter data
-	 *
-	 *	@param	InOwnerProxy	The proxy that owns this dynamic emitter data
-	 *
-	 *	@return	bool			true if successful, false if failed
-	 */
-	virtual void UpdateRenderThreadResourcesEmitter(const FParticleSystemSceneProxy* InOwnerProxy);
+	virtual FParticleVertexFactoryBase* BuildVertexFactory(const FParticleSystemSceneProxy* InOwnerProxy);
 
 	virtual void GetDynamicMeshElementsEmitter(const FParticleSystemSceneProxy* Proxy, const FSceneView* View, const FSceneViewFamily& ViewFamily, int32 ViewIndex, FMeshElementCollector& Collector) const override;
 
diff --git a/Engine/Source/Runtime/RHI/Public/RHICommandList.h b/Engine/Source/Runtime/RHI/Public/RHICommandList.h
--- a/Engine/Source/Runtime/RHI/Public/RHICommandList.h
+++ b/Engine/Source/Runtime/RHI/Public/RHICommandList.h
@@ -1663,10 +1663,11 @@ public:
 			void* Data = LockVertexBuffer_Internal(VertexBuffer, 0, BufferSize, RLM_WriteOnly);
 			FMemory::Memcpy(Data, Buffer, BufferSize);
 			UnlockVertexBuffer_Internal(VertexBuffer);
-			FMemory::Free((void*)Buffer);
 			return;
 		}
-		new (AllocCommand<FRHICommandUpdateVertexBuffer>()) FRHICommandUpdateVertexBuffer(VertexBuffer, Buffer, BufferSize);
+		void* Data = Alloc(BufferSize, 16);
+		FMemory::Memcpy(Data, Buffer, BufferSize);
+		new (AllocCommand<FRHICommandUpdateVertexBuffer>()) FRHICommandUpdateVertexBuffer(VertexBuffer, Data, BufferSize);
 	}
 	FORCEINLINE_DEBUGGABLE void BeginRenderQuery(FRenderQueryRHIParamRef RenderQuery)
 	{
diff --git a/Engine/Source/Runtime/RHI/Public/RHICommandListCommandExecutes.inl b/Engine/Source/Runtime/RHI/Public/RHICommandListCommandExecutes.inl
--- a/Engine/Source/Runtime/RHI/Public/RHICommandListCommandExecutes.inl
+++ b/Engine/Source/Runtime/RHI/Public/RHICommandListCommandExecutes.inl
@@ -373,7 +373,6 @@ void FRHICommandUpdateVertexBuffer::Execute()
 	void* Data = INTERNAL_DECORATOR(LockVertexBuffer)(VertexBuffer, 0, BufferSize, RLM_WriteOnly);
 	FMemory::Memcpy(Data, Buffer, BufferSize);
 	INTERNAL_DECORATOR(UnlockVertexBuffer)(VertexBuffer);
-	FMemory::Free((void*)Buffer);
 }
 
 void FRHICommandBeginFrame::Execute()
diff --git a/Engine/Source/Runtime/RHI/Public/RHIMethods.h b/Engine/Source/Runtime/RHI/Public/RHIMethods.h
--- a/Engine/Source/Runtime/RHI/Public/RHIMethods.h
+++ b/Engine/Source/Runtime/RHI/Public/RHIMethods.h
@@ -1876,7 +1876,7 @@ DEFINE_RHIMETHOD_3(
 	uint32,ColorBitMask,
 	,
 	);
-DEFINE_RHIMETHOD_1(
+DEFINE_RHIMETHOD_CMDLIST_1(
 	void,SetRenderTargetsAndClear,
 	const FRHISetRenderTargetsInfo&, RenderTargetsInfo,
 	,
