#ttp 344794 COMPAT: Swing Ninja Music plays too fast on Android

Fixed issue where Android Target Platforms were not being compiled with Ogg Vorbis enabled thus causing encoding to default to ADPCM

Fixed issue with ADPCM decoder not dealing with stereo streams.

diff --git a/Engine/Source/Developer/Android/Android_ATCTargetPlatform/Android_ATCTargetPlatform.Build.cs b/Engine/Source/Developer/Android/Android_ATCTargetPlatform/Android_ATCTargetPlatform.Build.cs
--- a/Engine/Source/Developer/Android/Android_ATCTargetPlatform/Android_ATCTargetPlatform.Build.cs
+++ b/Engine/Source/Developer/Android/Android_ATCTargetPlatform/Android_ATCTargetPlatform.Build.cs
@@ -31,6 +31,8 @@ public class Android_ATCTargetPlatform : ModuleRules
 			PrivateIncludePathModuleNames.Add("TextureCompressor");		//@todo android: AndroidTargetPlatform.Build
 		}
 
+        Definitions.Add("WITH_OGGVORBIS=1");
+
 		PrivateIncludePaths.AddRange(
 			new string[]
 			{
diff --git a/Engine/Source/Developer/Android/Android_DXTTargetPlatform/Android_DXTTargetPlatform.Build.cs b/Engine/Source/Developer/Android/Android_DXTTargetPlatform/Android_DXTTargetPlatform.Build.cs
--- a/Engine/Source/Developer/Android/Android_DXTTargetPlatform/Android_DXTTargetPlatform.Build.cs
+++ b/Engine/Source/Developer/Android/Android_DXTTargetPlatform/Android_DXTTargetPlatform.Build.cs
@@ -31,6 +31,8 @@ public class Android_DXTTargetPlatform : ModuleRules
 			PrivateIncludePathModuleNames.Add("TextureCompressor");		//@todo android: AndroidTargetPlatform.Build
 		}
 
+        Definitions.Add("WITH_OGGVORBIS=1");
+
 		PrivateIncludePaths.AddRange(
 			new string[]
 			{
diff --git a/Engine/Source/Developer/Android/Android_ETC1TargetPlatform/Android_ETC1TargetPlatform.Build.cs b/Engine/Source/Developer/Android/Android_ETC1TargetPlatform/Android_ETC1TargetPlatform.Build.cs
--- a/Engine/Source/Developer/Android/Android_ETC1TargetPlatform/Android_ETC1TargetPlatform.Build.cs
+++ b/Engine/Source/Developer/Android/Android_ETC1TargetPlatform/Android_ETC1TargetPlatform.Build.cs
@@ -32,6 +32,8 @@ public class Android_ETC1TargetPlatform : ModuleRules
 			}
 		);
 
+        Definitions.Add("WITH_OGGVORBIS=1");
+
 		// compile with Engine
 		if (UEBuildConfiguration.bCompileAgainstEngine)
 		{
diff --git a/Engine/Source/Developer/Android/Android_ETC2TargetPlatform/Android_ETC2TargetPlatform.Build.cs b/Engine/Source/Developer/Android/Android_ETC2TargetPlatform/Android_ETC2TargetPlatform.Build.cs
--- a/Engine/Source/Developer/Android/Android_ETC2TargetPlatform/Android_ETC2TargetPlatform.Build.cs
+++ b/Engine/Source/Developer/Android/Android_ETC2TargetPlatform/Android_ETC2TargetPlatform.Build.cs
@@ -31,6 +31,8 @@ public class Android_ETC2TargetPlatform : ModuleRules
 			PrivateIncludePathModuleNames.Add("TextureCompressor");		//@todo android: AndroidTargetPlatform.Build
 		}
 
+        Definitions.Add("WITH_OGGVORBIS=1");
+
 		PrivateIncludePaths.AddRange(
 			new string[]
 			{
diff --git a/Engine/Source/Developer/Android/Android_PVRTCTargetPlatform/Android_PVRTCTargetPlatform.Build.cs b/Engine/Source/Developer/Android/Android_PVRTCTargetPlatform/Android_PVRTCTargetPlatform.Build.cs
--- a/Engine/Source/Developer/Android/Android_PVRTCTargetPlatform/Android_PVRTCTargetPlatform.Build.cs
+++ b/Engine/Source/Developer/Android/Android_PVRTCTargetPlatform/Android_PVRTCTargetPlatform.Build.cs
@@ -31,6 +31,8 @@ public class Android_PVRTCTargetPlatform : ModuleRules
 			PrivateIncludePathModuleNames.Add("TextureCompressor");		//@todo android: AndroidTargetPlatform.Build
 		}
 
+        Definitions.Add("WITH_OGGVORBIS=1");
+
 		PrivateIncludePaths.AddRange(
 			new string[]
 			{
diff --git a/Engine/Source/Runtime/Engine/Private/ADPCMAudioInfo.cpp b/Engine/Source/Runtime/Engine/Private/ADPCMAudioInfo.cpp
--- a/Engine/Source/Runtime/Engine/Private/ADPCMAudioInfo.cpp
+++ b/Engine/Source/Runtime/Engine/Private/ADPCMAudioInfo.cpp
@@ -65,6 +65,105 @@
 {
 	check(Destination);
 
+	return (*WaveInfo.pChannels == 1) ? DecodeMonoData(Destination, bLooping, BufferSize) : DecodeStereoData(Destination, bLooping, BufferSize);
+}
+
+void FADPCMAudioInfo::ExpandFile(uint8* DstBuffer, struct FSoundQualityInfo* QualityInfo)
+{
+	check(DstBuffer);
+
+	if (*WaveInfo.pChannels == 1)
+	{
+		const uint8* EncodedADPCMBlockStart = reinterpret_cast<uint8*>(WaveInfo.SampleDataStart);
+		uint8 * EncodedADPCMBlock = (uint8*)EncodedADPCMBlockStart;
+		while (EncodedADPCMBlock < EncodedADPCMBlockStart + (WaveInfo.SampleDataSize))
+		{
+			ADPCM::DecodeBlock(EncodedADPCMBlock, CompressedBlockSize, (int16*)DstBuffer);
+			EncodedADPCMBlock += CompressedBlockSize;
+			DstBuffer += UncompressedBlockSize;
+		}
+	}
+	else
+	{
+		const uint8* EncodedADPCMBlockLeftStart = reinterpret_cast<uint8*>(WaveInfo.SampleDataStart);
+		const uint8* EncodedADPCMBlockRightStart = reinterpret_cast<uint8*>(WaveInfo.SampleDataStart + (WaveInfo.SampleDataSize / 2));
+		uint8 * EncodedADPCMLeftBlock = (uint8*)EncodedADPCMBlockLeftStart;
+		uint8 * EncodedADPCMRightBlock = (uint8*)EncodedADPCMBlockRightStart;
+		while (EncodedADPCMLeftBlock < EncodedADPCMBlockLeftStart + (WaveInfo.SampleDataSize / 2))
+		{
+			ADPCM::DecodeBlockStereo(EncodedADPCMLeftBlock, EncodedADPCMRightBlock, CompressedBlockSize, (int16*)DstBuffer);
+			EncodedADPCMLeftBlock += CompressedBlockSize;
+			EncodedADPCMRightBlock += CompressedBlockSize;
+			DstBuffer += (UncompressedBlockSize * 2);
+		}
+	}
+}
+
+int FADPCMAudioInfo::GetStreamBufferSize() const
+{
+	return StreamBufferSize;
+}
+
+bool FADPCMAudioInfo::DecodeStereoData(uint8* Destination, bool bLooping, uint32 BufferSize)
+{
+	const uint8* EncodedADPCMBlockLeftStart = reinterpret_cast<uint8*>(WaveInfo.SampleDataStart);
+	const uint8* EncodedADPCMBlockRightStart = reinterpret_cast<uint8*>(WaveInfo.SampleDataStart + (WaveInfo.SampleDataSize / 2));
+	uint8 * EncodedADPCMLeftBlock = (uint8*)EncodedADPCMBlockLeftStart + SrcBufferOffset;
+	uint8 * EncodedADPCMRightBlock = (uint8*)EncodedADPCMBlockRightStart + SrcBufferOffset;
+	
+	// Check to see if we have enough data left to do a complete decode
+	if (EncodedADPCMLeftBlock + StreamBufferSizeInBlocks < EncodedADPCMBlockRightStart)
+	{
+		while (BufferSize)
+		{
+			ADPCM::DecodeBlockStereo(EncodedADPCMLeftBlock, EncodedADPCMRightBlock, CompressedBlockSize, (int16*)Destination);
+			EncodedADPCMLeftBlock += CompressedBlockSize;
+			EncodedADPCMRightBlock += CompressedBlockSize;
+			SrcBufferOffset += CompressedBlockSize;
+			Destination += (UncompressedBlockSize * 2);
+			BufferSize -= (UncompressedBlockSize * 2);
+		}
+	}
+	else
+	{
+		// otherwise we need to copy what we have left
+		while (EncodedADPCMLeftBlock < EncodedADPCMBlockRightStart)
+		{
+			ADPCM::DecodeBlockStereo(EncodedADPCMLeftBlock, EncodedADPCMRightBlock, CompressedBlockSize, (int16*)Destination);
+			EncodedADPCMLeftBlock += CompressedBlockSize;
+			EncodedADPCMRightBlock += CompressedBlockSize;
+			SrcBufferOffset += CompressedBlockSize;
+			Destination += (UncompressedBlockSize * 2);
+			BufferSize -= (UncompressedBlockSize * 2);
+		}
+
+		// and then decide if we zero pad or reset and loop
+		if (bLooping)
+		{
+			EncodedADPCMLeftBlock = (uint8*)EncodedADPCMBlockLeftStart;
+			EncodedADPCMRightBlock = (uint8*)EncodedADPCMBlockRightStart;
+			SrcBufferOffset = 0;
+			while (BufferSize)
+			{
+				ADPCM::DecodeBlockStereo(EncodedADPCMLeftBlock, EncodedADPCMRightBlock, CompressedBlockSize, (int16*)Destination);
+				EncodedADPCMLeftBlock += CompressedBlockSize;
+				EncodedADPCMRightBlock += CompressedBlockSize;
+				SrcBufferOffset += CompressedBlockSize;
+				Destination += (UncompressedBlockSize * 2);
+				BufferSize -= (UncompressedBlockSize * 2);
+			}
+		}
+		else
+		{
+			FMemory::Memzero(Destination, BufferSize);
+		}
+	}
+
+	return true;
+}
+
+bool FADPCMAudioInfo::DecodeMonoData(uint8* Destination, bool bLooping, uint32 BufferSize)
+{
 	const uint8* EncodedADPCMBlockStart = reinterpret_cast<uint8*>(WaveInfo.SampleDataStart);
 	uint8* EncodedADPCMBlock = (uint8*)EncodedADPCMBlockStart + SrcBufferOffset;
 	const uint8* EncodedDataEnd = EncodedADPCMBlockStart + WaveInfo.SampleDataSize;
@@ -116,39 +215,3 @@
 
 	return true;
 }
-
-void FADPCMAudioInfo::ExpandFile(uint8* DstBuffer, struct FSoundQualityInfo* QualityInfo)
-{
-	check(DstBuffer);
-
-	if (*WaveInfo.pChannels == 1)
-	{
-		const uint8* EncodedADPCMBlockStart = reinterpret_cast<uint8*>(WaveInfo.SampleDataStart);
-		uint8 * EncodedADPCMBlock = (uint8*)EncodedADPCMBlockStart;
-		while (EncodedADPCMBlock < EncodedADPCMBlockStart + (WaveInfo.SampleDataSize))
-		{
-			ADPCM::DecodeBlock(EncodedADPCMBlock, CompressedBlockSize, (int16*)DstBuffer);
-			EncodedADPCMBlock += CompressedBlockSize;
-			DstBuffer += UncompressedBlockSize;
-		}
-	}
-	else
-	{
-		const uint8* EncodedADPCMBlockLeftStart = reinterpret_cast<uint8*>(WaveInfo.SampleDataStart);
-		const uint8* EncodedADPCMBlockRightStart = reinterpret_cast<uint8*>(WaveInfo.SampleDataStart + (WaveInfo.SampleDataSize / 2));
-		uint8 * EncodedADPCMLeftBlock = (uint8*)EncodedADPCMBlockLeftStart;
-		uint8 * EncodedADPCMRightBlock = (uint8*)EncodedADPCMBlockRightStart;
-		while (EncodedADPCMLeftBlock < EncodedADPCMBlockLeftStart + (WaveInfo.SampleDataSize / 2))
-		{
-			ADPCM::DecodeBlockStereo(EncodedADPCMLeftBlock, EncodedADPCMRightBlock, CompressedBlockSize, (int16*)DstBuffer);
-			EncodedADPCMLeftBlock += CompressedBlockSize;
-			EncodedADPCMRightBlock += CompressedBlockSize;
-			DstBuffer += (UncompressedBlockSize * 2);
-		}
-	}
-}
-
-int FADPCMAudioInfo::GetStreamBufferSize() const
-{
-	return StreamBufferSize;
-}
diff --git a/Engine/Source/Runtime/Engine/Public/ADPCMAudioInfo.h b/Engine/Source/Runtime/Engine/Public/ADPCMAudioInfo.h
--- a/Engine/Source/Runtime/Engine/Public/ADPCMAudioInfo.h
+++ b/Engine/Source/Runtime/Engine/Public/ADPCMAudioInfo.h
@@ -35,4 +35,8 @@ public:
 	uint32 BlockSize;
 	int32 StreamBufferSize;
 	uint32 StreamBufferSizeInBlocks;
+
+private:
+	bool DecodeStereoData(uint8* Destination, bool bLooping, uint32 BufferSize);
+	bool DecodeMonoData(uint8* Destination, bool bLooping, uint32 BufferSize);
 };
