Log the guids when package map guids mismatch, which helps to determine if one or the other simply wasn't initialized

diff --git a/Engine/Source/Runtime/Engine/Private/PackageMapClient.cpp b/Engine/Source/Runtime/Engine/Private/PackageMapClient.cpp
--- a/Engine/Source/Runtime/Engine/Private/PackageMapClient.cpp
+++ b/Engine/Source/Runtime/Engine/Private/PackageMapClient.cpp
@@ -425,6 +425,8 @@ void UPackageMapClient::InternalWriteObject( FArchive & Ar, FNetworkGUID NetGUID
 		{
 			FGuid PackageGuid = CastChecked< const UPackage >( Object )->GetGuid();
 			Ar << PackageGuid;
+
+			UE_LOG( LogNetPackageMap, VeryVerbose, TEXT( "InternalWriteObject: Package: %s, GUID: %s" ), *ObjectPathName, *PackageGuid.ToString() );
 		}
 
 		if ( GuidCache->IsExportingNetGUIDBunch )
@@ -573,6 +575,7 @@ FNetworkGUID UPackageMapClient::InternalLoadObject( FArchive & Ar, UObject *& Ob
 		if ( bIsPackage )
 		{
 			Ar << PackageGuid;
+			UE_LOG( LogNetPackageMap, VeryVerbose, TEXT( "InternalLoadObject: Package: %s, GUID: %s" ), *PathName, *PackageGuid.ToString() );
 		}
 
 		if ( Ar.IsError() )
@@ -627,7 +630,7 @@ FNetworkGUID UPackageMapClient::InternalLoadObject( FArchive & Ar, UObject *& Ob
 
 				if ( Package->GetGuid() != PackageGuid )
 				{
-					UE_LOG( LogNetPackageMap, Error, TEXT( "UPackageMapClient::InternalLoadObject: Default object package guid mismatch! PathName: %s, ObjOuter: %s " ), *PathName, ObjOuter != NULL ? *ObjOuter->GetPathName() : TEXT( "NULL" ) );
+					UE_LOG( LogNetPackageMap, Error, TEXT( "UPackageMapClient::InternalLoadObject: Default object package guid mismatch! PathName: %s, ObjOuter: %s, GUID1: %s, GUID2: %s " ), *PathName, ObjOuter != NULL ? *ObjOuter->GetPathName() : TEXT( "NULL" ), *Package->GetGuid().ToString(), *PackageGuid.ToString() );
 					Object = NULL;
 					return NetGUID;
 				}
@@ -1663,7 +1666,7 @@ void FNetGUIDCache::AsyncPackageCallback( const FString & PackageName, UPackage
 	else if ( Package->GetGuid() != CacheObject->PackageGuid )
 	{
 		CacheObject->bIsBroken = true;
-		UE_LOG( LogNetPackageMap, Error, TEXT( "AsyncPackageCallback: Package GUID mismatch! Path: %s, NetGUID: %s" ), *PackageName, *NetGUID.ToString() );
+		UE_LOG( LogNetPackageMap, Error, TEXT( "AsyncPackageCallback: Package GUID mismatch! Path: %s, NetGUID: %s, GUID1: %s, GUID2: %s" ), *PackageName, *NetGUID.ToString(), *Package->GetGuid().ToString(), *CacheObject->PackageGuid.ToString() );
 	}
 }
 
@@ -1869,7 +1872,7 @@ UObject * FNetGUIDCache::GetObjectFromNetGUID( const FNetworkGUID & NetGUID, con
 		{
 			// If the package guid doesn't match, don't allow it to load
 			CacheObjectPtr->bIsBroken = true;
-			UE_LOG( LogNetPackageMap, Error, TEXT( "GetObjectFromNetGUID: Package GUID mismatch! Path: %s, NetGUID: %s" ), *CacheObjectPtr->PathName.ToString(), *NetGUID.ToString() );
+			UE_LOG( LogNetPackageMap, Error, TEXT( "GetObjectFromNetGUID: Package GUID mismatch! Path: %s, NetGUID: %s, GUID1: %s, GUID2: %s" ), *CacheObjectPtr->PathName.ToString(), *NetGUID.ToString(), *Package->GetGuid().ToString(), *CacheObjectPtr->PackageGuid.ToString() );
 			return NULL;
 		}
 	}
