Update outgoing acked packet it when we have InternalAck set

diff --git a/Engine/Source/Runtime/Engine/Private/NetConnection.cpp b/Engine/Source/Runtime/Engine/Private/NetConnection.cpp
--- a/Engine/Source/Runtime/Engine/Private/NetConnection.cpp
+++ b/Engine/Source/Runtime/Engine/Private/NetConnection.cpp
@@ -468,7 +468,7 @@ void UNetConnection::FlushNet(bool bIgnoreSimulation)
 	TimeSensitive = 0;
 
 	// If there is any pending data to send, send it.
-	if( SendBuffer.GetNumBits() || Driver->Time-LastSendTime>Driver->KeepAliveTime )
+	if ( SendBuffer.GetNumBits() || ( Driver->Time-LastSendTime > Driver->KeepAliveTime && !InternalAck ) )
 	{
 		// If sending keepalive packet, still write the packet id
 		if ( SendBuffer.GetNumBits() == 0 )
@@ -1300,6 +1300,8 @@ void UNetConnection::Tick()
 	// Pretend everything was acked, for 100% reliable connections or demo recording.
 	if( InternalAck )
 	{
+		OutAckPacketId = OutPacketId;
+
 		LastReceiveTime = Driver->Time;
 		for( int32 i=OpenChannels.Num()-1; i>=0; i-- )
 		{
