# retargeting clean up

- changed rigcontrol to trasnformbase as right now it only does transform base, and in the future we should add more controllers
- added clear bone name
- added "reset base pose" option for opting back to original mesh pose
- fixed so that rig nodes doesn't have to be mapped

diff --git a/Engine/Source/Editor/DetailCustomizations/Private/RigDetails.cpp b/Engine/Source/Editor/DetailCustomizations/Private/RigDetails.cpp
--- a/Engine/Source/Editor/DetailCustomizations/Private/RigDetails.cpp
+++ b/Engine/Source/Editor/DetailCustomizations/Private/RigDetails.cpp
@@ -34,9 +34,9 @@ void FRigDetails::CustomizeDetails(IDetailLayoutBuilder& DetailBuilder)
 	ItemBeingEdited = Objects[0];
 
 	IDetailCategoryBuilder& NodeCategory = DetailBuilder.EditCategory("Node");
-	IDetailCategoryBuilder& RigControlCategory = DetailBuilder.EditCategory("Constraint Setup");
+	IDetailCategoryBuilder& TransformBaseCategory = DetailBuilder.EditCategory("Constraint Setup");
 
-	RigControlsPropertyHandle = DetailBuilder.GetProperty("RigControls");
+	TransformBasesPropertyHandle = DetailBuilder.GetProperty("TransformBases");
 	NodesPropertyHandle = DetailBuilder.GetProperty("Nodes");
 
 	// since now we can't really resize the array, we'll just allocate everything here
@@ -49,7 +49,7 @@ void FRigDetails::CustomizeDetails(IDetailLayoutBuilder& DetailBuilder)
 		DisplayNameTextBoxes.AddZeroed(NumElement);
 	}
 
-	check (FPropertyAccess::Fail != RigControlsPropertyHandle->AsArray()->GetNumElements(NumElement));
+	check (FPropertyAccess::Fail != TransformBasesPropertyHandle->AsArray()->GetNumElements(NumElement));
 	if ( NumElement > 0 )
 	{
 		ParentSpaceOptionList.AddZeroed(NumElement);
@@ -62,13 +62,13 @@ void FRigDetails::CustomizeDetails(IDetailLayoutBuilder& DetailBuilder)
 
 	NodeCategory.AddCustomBuilder( NodeArrayBuilder, false );
 
-	TSharedRef<FDetailArrayBuilder> RigControlArrayBuilder = MakeShareable(new FDetailArrayBuilder(RigControlsPropertyHandle.ToSharedRef()));
-	RigControlArrayBuilder->OnGenerateArrayElementWidget(FOnGenerateArrayElementWidget::CreateSP(this, &FRigDetails::GenerateRigControlArrayElementWidget, &DetailBuilder));
+	TSharedRef<FDetailArrayBuilder> TransformBaseArrayBuilder = MakeShareable(new FDetailArrayBuilder(TransformBasesPropertyHandle.ToSharedRef()));
+	TransformBaseArrayBuilder->OnGenerateArrayElementWidget(FOnGenerateArrayElementWidget::CreateSP(this, &FRigDetails::GenerateTransformBaseArrayElementWidget, &DetailBuilder));
 
 	// add custom menu
 	// -> set all to world
 	// -> set all to default parent
-	RigControlCategory.AddCustomRow(TEXT(""))
+	TransformBaseCategory.AddCustomRow(TEXT(""))
 	[
 		// two button 1. view 2. save to base pose
 		SNew(SHorizontalBox)
@@ -98,7 +98,7 @@ void FRigDetails::CustomizeDetails(IDetailLayoutBuilder& DetailBuilder)
 		]
 	];
 
-	RigControlCategory.AddCustomBuilder( RigControlArrayBuilder, false );
+	TransformBaseCategory.AddCustomBuilder( TransformBaseArrayBuilder, false );
 }
 
 void FRigDetails::GenerateNodeArrayElementWidget(TSharedRef<IPropertyHandle> PropertyHandle, int32 ArrayIndex, IDetailChildrenBuilder& ChildrenBuilder, IDetailLayoutBuilder* DetailLayout)
@@ -106,6 +106,7 @@ void FRigDetails::GenerateNodeArrayElementWidget(TSharedRef<IPropertyHandle> Pro
 	TSharedRef<IPropertyHandle> DisplayNameProp = PropertyHandle->GetChildHandle(GET_MEMBER_NAME_CHECKED(FNode, DisplayName)).ToSharedRef();
 	TSharedRef<IPropertyHandle> NodeNameProp = PropertyHandle->GetChildHandle(GET_MEMBER_NAME_CHECKED(FNode, Name)).ToSharedRef();
 	TSharedRef<IPropertyHandle> ParentNameProp = PropertyHandle->GetChildHandle(GET_MEMBER_NAME_CHECKED(FNode, ParentName)).ToSharedRef();
+	TSharedRef<IPropertyHandle> AdvancedProp = PropertyHandle->GetChildHandle(GET_MEMBER_NAME_CHECKED(FNode, bAdvanced)).ToSharedRef();
 
 	TSharedPtr<SEditableTextBox> DisplayTextBox;
 
@@ -159,7 +160,7 @@ void FRigDetails::GenerateNodeArrayElementWidget(TSharedRef<IPropertyHandle> Pro
 
 		+SHorizontalBox::Slot()
 		.Padding(5, 2)
-		.FillWidth(1)
+		.AutoWidth()
 		.HAlign(HAlign_Left)
 		[
 			SNew(SBox)
@@ -174,6 +175,30 @@ void FRigDetails::GenerateNodeArrayElementWidget(TSharedRef<IPropertyHandle> Pro
 				.MinDesiredWidth(200)
 			]
 		]
+
+		+SHorizontalBox::Slot()
+		.Padding(5, 2)
+		.AutoWidth()
+		[
+			SNew(STextBlock)
+			.Text(FString(TEXT("Advanced")))
+			.Font(DetailLayout->GetDetailFontBold())
+		]
+
+		+SHorizontalBox::Slot()
+		.Padding(5, 2)
+		.FillWidth(1)
+		.HAlign(HAlign_Left)
+		[
+			SNew(SBox)
+			.WidthOverride(250)
+			.Content()
+			[
+				SNew(SCheckBox)
+				.IsChecked(this, &FRigDetails::AdvancedCheckBoxIsChecked, AdvancedProp)
+				.OnCheckStateChanged(this, &FRigDetails::OnAdvancedCheckBoxStateChanged, AdvancedProp)
+			]
+		]
 	];
 
 	DisplayNameTextBoxes[ArrayIndex] = DisplayTextBox;
@@ -233,17 +258,17 @@ void FRigDetails::OnDisplayNameCommitted(const FText& Text, ETextCommit::Type Co
 	DisplayNameProp->SetValueFromFormattedString(Text.ToString());
 }
 
-void FRigDetails::GenerateRigControlArrayElementWidget(TSharedRef<IPropertyHandle> PropertyHandle, int32 ArrayIndex, IDetailChildrenBuilder& ChildrenBuilder, IDetailLayoutBuilder* DetailLayout)
+void FRigDetails::GenerateTransformBaseArrayElementWidget(TSharedRef<IPropertyHandle> PropertyHandle, int32 ArrayIndex, IDetailChildrenBuilder& ChildrenBuilder, IDetailLayoutBuilder* DetailLayout)
 {
-	TSharedRef<IPropertyHandle> NodeNameProp = PropertyHandle->GetChildHandle(GET_MEMBER_NAME_CHECKED(FRigControl, Node)).ToSharedRef();
-	TSharedPtr<IPropertyHandleArray> ConstraintsProp = PropertyHandle->GetChildHandle(GET_MEMBER_NAME_CHECKED(FRigControl, Constraints))->AsArray();
+	TSharedRef<IPropertyHandle> NodeNameProp = PropertyHandle->GetChildHandle(GET_MEMBER_NAME_CHECKED(FTransformBase, Node)).ToSharedRef();
+	TSharedPtr<IPropertyHandleArray> ConstraintsProp = PropertyHandle->GetChildHandle(GET_MEMBER_NAME_CHECKED(FTransformBase, Constraints))->AsArray();
 
 	// translation
-	TSharedPtr<IPropertyHandleArray> TransformConstraintsProp_T = ConstraintsProp->GetElement(0)->GetChildHandle(GET_MEMBER_NAME_CHECKED(FRigControlConstraint, TransformConstraints))->AsArray();
+	TSharedPtr<IPropertyHandleArray> TransformConstraintsProp_T = ConstraintsProp->GetElement(0)->GetChildHandle(GET_MEMBER_NAME_CHECKED(FTransformBaseConstraint, TransformConstraints))->AsArray();
 	TSharedRef<IPropertyHandle> ParentNameProp_T = TransformConstraintsProp_T->GetElement(0)->GetChildHandle(GET_MEMBER_NAME_CHECKED(FRigTransformConstraint, ParentSpace)).ToSharedRef();
 
 	// orientation
-	TSharedPtr<IPropertyHandleArray> TransformConstraintsProp_R = ConstraintsProp->GetElement(1)->GetChildHandle(GET_MEMBER_NAME_CHECKED(FRigControlConstraint, TransformConstraints))->AsArray();
+	TSharedPtr<IPropertyHandleArray> TransformConstraintsProp_R = ConstraintsProp->GetElement(1)->GetChildHandle(GET_MEMBER_NAME_CHECKED(FTransformBaseConstraint, TransformConstraints))->AsArray();
 	TSharedRef<IPropertyHandle> ParentNameProp_R = TransformConstraintsProp_R->GetElement(0)->GetChildHandle(GET_MEMBER_NAME_CHECKED(FRigTransformConstraint, ParentSpace)).ToSharedRef();
 
 	// the interface will be node [display name] [parent node]
@@ -398,9 +423,12 @@ void FRigDetails::GenerateRigControlArrayElementWidget(TSharedRef<IPropertyHandl
 
 void FRigDetails::OnParentSpaceSelectionChanged(TSharedPtr<FString> SelectedItem, ESelectInfo::Type SelectInfo, TSharedRef<IPropertyHandle> ParentSpacePropertyHandle)
 {
-	if (SelectedItem.IsValid())
+	if (SelectInfo == ESelectInfo::OnKeyPress || SelectInfo == ESelectInfo::OnMouseClick)
 	{
-		check (ParentSpacePropertyHandle->SetValueFromFormattedString(*SelectedItem.Get()) != FPropertyAccess::Fail);
+		if(SelectedItem.IsValid())
+		{
+			check(ParentSpacePropertyHandle->SetValueFromFormattedString(*SelectedItem.Get()) != FPropertyAccess::Fail);
+		}
 	}
 }
 
@@ -470,4 +498,18 @@ void FRigDetails::OnComboBoxOopening(TSharedRef<IPropertyHandle> ParentSpaceProp
 	ComboBox->SetSelectedItem(SelectedItem);
 }
 
+void FRigDetails::OnAdvancedCheckBoxStateChanged(ESlateCheckBoxState::Type NewState, TSharedRef<IPropertyHandle> PropertyHandle)
+{
+	bool bValue = (NewState == ESlateCheckBoxState::Checked)? true : false;
+	check (PropertyHandle->SetValue(bValue) != FPropertyAccess::Fail);
+}
+
+ESlateCheckBoxState::Type FRigDetails::AdvancedCheckBoxIsChecked(TSharedRef<IPropertyHandle> PropertyHandle) const
+{
+	bool bValue;
+	// multi value doesn't work in array, so i'm not handling multi value
+	check(PropertyHandle->GetValue(bValue) != FPropertyAccess::Fail);
+
+	return (bValue)? ESlateCheckBoxState::Checked: ESlateCheckBoxState::Unchecked;
+}
 #undef LOCTEXT_NAMESPACE
diff --git a/Engine/Source/Editor/DetailCustomizations/Private/RigDetails.h b/Engine/Source/Editor/DetailCustomizations/Private/RigDetails.h
--- a/Engine/Source/Editor/DetailCustomizations/Private/RigDetails.h
+++ b/Engine/Source/Editor/DetailCustomizations/Private/RigDetails.h
@@ -17,10 +17,10 @@ private:
 
 	/** Generate a widget for a movie array element */
 	void GenerateNodeArrayElementWidget(TSharedRef<IPropertyHandle> PropertyHandle, int32 ArrayIndex, IDetailChildrenBuilder& ChildrenBuilder, IDetailLayoutBuilder* DetailLayout);
-	void GenerateRigControlArrayElementWidget(TSharedRef<IPropertyHandle> PropertyHandle, int32 ArrayIndex, IDetailChildrenBuilder& ChildrenBuilder, IDetailLayoutBuilder* DetailLayout);
+	void GenerateTransformBaseArrayElementWidget(TSharedRef<IPropertyHandle> PropertyHandle, int32 ArrayIndex, IDetailChildrenBuilder& ChildrenBuilder, IDetailLayoutBuilder* DetailLayout);
 
 	TSharedPtr<IPropertyHandle> NodesPropertyHandle;
-	TSharedPtr<IPropertyHandle> RigControlsPropertyHandle;
+	TSharedPtr<IPropertyHandle> TransformBasesPropertyHandle;
 
 	// array custom boxes - these will stay around as long as this window is up
 	TArray<TSharedPtr<SEditableTextBox>>	DisplayNameTextBoxes;
@@ -40,9 +40,13 @@ private:
 	void OnParentSpaceSelectionChanged(TSharedPtr<FString> SelectedItem, ESelectInfo::Type SelectInfo, TSharedRef<IPropertyHandle> ParentSpacePropertyHandle);
 	/** Called to create a widget for each string */
 	TSharedRef<SWidget> MakeItemWidget(TSharedPtr<FString> StringItem);
-	/** Helper method to get the text for a given item in the combo box */
+	// Helper method to get the text for a given item in the combo box 
 	FString GetSelectedTextLabel(TSharedRef<IPropertyHandle> ParentSpacePropertyHandle) const;
 	void OnComboBoxOopening(TSharedRef<IPropertyHandle> ParentSpacePropertyHandle, int32 ArrayIndex, bool bTranslation);
+	// check box handler
+	// Callback for changing this row's Share check box state.
+	void OnAdvancedCheckBoxStateChanged(ESlateCheckBoxState::Type NewState, TSharedRef<IPropertyHandle> PropertyHandle);
+	ESlateCheckBoxState::Type AdvancedCheckBoxIsChecked(TSharedRef<IPropertyHandle> PropertyHandle) const;
 
 	// button handlers
 	FReply OnSetAllToWorld();
diff --git a/Engine/Source/Editor/Persona/Private/SRetargetManager.cpp b/Engine/Source/Editor/Persona/Private/SRetargetManager.cpp
--- a/Engine/Source/Editor/Persona/Private/SRetargetManager.cpp
+++ b/Engine/Source/Editor/Persona/Private/SRetargetManager.cpp
@@ -164,6 +164,18 @@ void SRetargetManager::Construct(const FArguments& InArgs)
 			// two button 1. view 2. save to base pose
 			SNew(SHorizontalBox)
 
+			+SHorizontalBox::Slot()
+			.AutoWidth()
+			.HAlign(HAlign_Center)
+			[
+				SNew(SButton)
+				.OnClicked(FOnClicked::CreateSP(this, &SRetargetManager::OnResetRetargetBasePose))
+				.HAlign(HAlign_Center)
+				.VAlign(VAlign_Center)
+				.Text(LOCTEXT("ResetRetargetBasePose_Label", "Reset Pose"))
+				.ToolTipText(LOCTEXT("ResetRetargetBasePose_Tooltip", "Restore Retarget Base Pose to Mesh Refeference Pose"))
+			]
+
 			+SHorizontalBox::Slot()
 			.AutoWidth()
 			.HAlign(HAlign_Center)
@@ -173,6 +185,7 @@ void SRetargetManager::Construct(const FArguments& InArgs)
 				.HAlign(HAlign_Center)
 				.VAlign(VAlign_Center)
 				.Text(this, &SRetargetManager::GetToggleRetargetBasePose)
+				.ToolTipText(LOCTEXT("ViewRetargetBasePose_Tooltip", "Toggle to View/Edit Retarget Base Pose"))
 			]
 
 			+SHorizontalBox::Slot()
@@ -183,7 +196,8 @@ void SRetargetManager::Construct(const FArguments& InArgs)
 				.OnClicked(FOnClicked::CreateSP(this, &SRetargetManager::OnSaveRetargetBasePose))
 				.HAlign(HAlign_Center)
 				.VAlign(VAlign_Center)
-				.Text(LOCTEXT("SaveRetargetBasePose_Label", "Save Current Pose to Retarget Base Pose"))
+				.Text(LOCTEXT("SaveRetargetBasePose_Label", "Save Pose"))
+				.ToolTipText(LOCTEXT("SaveRetargetBasePose_Tooltip", "Save Current Pose to Retarget Base Pose"))
 			]
 		]
 	];
@@ -261,8 +275,38 @@ FReply SRetargetManager::OnSaveRetargetBasePose()
 
 			// Clear PreviewMeshComp bone modified, they're baked now
 			PreviewMeshComp->PreviewInstance->ResetModifiedBone();
+			// turn off the retarget base pose if you're looking at it
+			PreviewMeshComp->PreviewInstance->bForceRetargetBasePose = false;
+		}
+	}
+	return FReply::Handled();
+}
+
+FReply SRetargetManager::OnResetRetargetBasePose()
+{
+	const FText Message = LOCTEXT("ResetRetargetBasePose_Confirm", "This will reset current Retarget Base Pose to the Reference Pose of current preview mesh. Would you like to continue?");
+	EAppReturnType::Type Response = FMessageDialog::Open(EAppMsgType::OkCancel, Message);
+	if(Response == EAppReturnType::Ok)
+	{
+		UDebugSkelMeshComponent * PreviewMeshComp = PersonaPtr.Pin()->GetPreviewMeshComponent();
+		if(PreviewMeshComp && PreviewMeshComp->SkeletalMesh)
+		{
+			USkeletalMesh * PreviewMesh = PreviewMeshComp->SkeletalMesh;
+
+			check(Skeleton == PreviewMesh->Skeleton);
+
+			if(PreviewMesh)
+			{
+				const FScopedTransaction Transaction(LOCTEXT("ResetRetargetBasePose_Action", "Reset Retarget Base Pose"));
+				PreviewMesh->Modify();
+				// reset to original ref pose
+				PreviewMesh->RetargetBasePose = PreviewMesh->RefSkeleton.GetRefBonePose();
+				// turn off the retarget base pose if you're looking at it
+				PreviewMeshComp->PreviewInstance->bForceRetargetBasePose = true;
+			}
 		}
 	}
+
 	return FReply::Handled();
 }
 
@@ -277,15 +321,15 @@ FText SRetargetManager::GetToggleRetargetBasePose() const
 	{
 		if (PreviewMeshComp->PreviewInstance->bForceRetargetBasePose)
 		{
-			return LOCTEXT("HideRetargetBasePose_Label", "Hide Current Retarget Base Pose");
+			return LOCTEXT("HideRetargetBasePose_Label", "Hide Pose");
 		}
 		else
 		{
-			return LOCTEXT("ViewRetargetBasePose_Label", "View/Edit Current Retarget Base Pose");
+			return LOCTEXT("ViewRetargetBasePose_Label", "View Pose");
 		}
 	}
 
-	return LOCTEXT("InvalidRetargetBasePose_Label", "No Mesh for Retarget Base Pose");
+	return LOCTEXT("InvalidRetargetBasePose_Label", "No Mesh for Base Pose");
 }
 #undef LOCTEXT_NAMESPACE
 
diff --git a/Engine/Source/Editor/Persona/Private/SRetargetManager.h b/Engine/Source/Editor/Persona/Private/SRetargetManager.h
--- a/Engine/Source/Editor/Persona/Private/SRetargetManager.h
+++ b/Engine/Source/Editor/Persona/Private/SRetargetManager.h
@@ -35,6 +35,7 @@ public:
 
 private:
 	FReply OnViewRetargetBasePose();
+	FReply OnResetRetargetBasePose();
 	FReply OnSaveRetargetBasePose();
 
 	/** Pointer back to the Persona that owns us */
diff --git a/Engine/Source/Editor/Persona/Private/SRigWindow.cpp b/Engine/Source/Editor/Persona/Private/SRigWindow.cpp
--- a/Engine/Source/Editor/Persona/Private/SRigWindow.cpp
+++ b/Engine/Source/Editor/Persona/Private/SRigWindow.cpp
@@ -73,6 +73,7 @@ private:
 
 	// Bone tree widget delegates
 	void OnBoneSelectionChanged(FName Name);
+	FReply OnClearButtonClicked();
 	FName GetSelectedBone() const;
 
 	FOnBoneMappingChanged OnBoneMappingChanged;
@@ -130,24 +131,44 @@ TSharedRef< SWidget > SBoneMappingListRow::GenerateWidgetForColumn( const FName&
 			.Padding( 0.0f, 1.0f )
 			.VAlign( VAlign_Center )
 			[
-				SNew(SBoneSelectionWidget)
-				.Skeleton( Item->Skeleton )
-				.Tooltip( FText::Format( LOCTEXT( "BoneSelectinWidget",  "Select Bone for node {0}" ), FText::FromString( Item->GetDisplayName() ) ))
-				.OnBoneSelectionChanged(this, &SBoneMappingListRow::OnBoneSelectionChanged)
-				.OnGetSelectedBone(this, &SBoneMappingListRow::GetSelectedBone)
+				SNew(SHorizontalBox)
+
+				+SHorizontalBox::Slot()
+				[
+					SNew(SBoneSelectionWidget)
+					.Skeleton(Item->Skeleton)
+					.Tooltip(FText::Format(LOCTEXT("BoneSelectinWidget", "Select Bone for node {0}"), FText::FromString(Item->GetDisplayName())))
+					.OnBoneSelectionChanged(this, &SBoneMappingListRow::OnBoneSelectionChanged)
+					.OnGetSelectedBone(this, &SBoneMappingListRow::GetSelectedBone)
+				]
+
+				+SHorizontalBox::Slot()
+				.AutoWidth()
+				[
+					SNew(SButton)
+					.OnClicked(FOnClicked::CreateSP(this, &SBoneMappingListRow::OnClearButtonClicked))
+					.Text(FText::FromString(TEXT("x")))
+				]
 			];
 	}
 }
 
+FReply SBoneMappingListRow::OnClearButtonClicked()
+{
+	if(OnBoneMappingChanged.IsBound())
+	{
+		OnBoneMappingChanged.Execute(Item->GetNodeName(), NAME_None);
+	}
+
+	return FReply::Handled();
+}
+
 void SBoneMappingListRow::OnBoneSelectionChanged(FName Name)
 {
 	if (OnBoneMappingChanged.IsBound())
 	{
 		OnBoneMappingChanged.Execute(Item->GetNodeName(), Name);
 	}
-
-	// @todo delete?
-//	Item->BoneName = Name;
 }
 
 FName SBoneMappingListRow::GetSelectedBone() const
diff --git a/Engine/Source/Runtime/Engine/Classes/Animation/Rig.h b/Engine/Source/Runtime/Engine/Classes/Animation/Rig.h
--- a/Engine/Source/Runtime/Engine/Classes/Animation/Rig.h
+++ b/Engine/Source/Runtime/Engine/Classes/Animation/Rig.h
@@ -35,9 +35,13 @@ struct FNode
 	UPROPERTY(EditAnywhere, Category="FNode")
 	FString DisplayName;
 
+	UPROPERTY(EditAnywhere, Category="FNode")
+	bool bAdvanced;
+
 	FNode()
 		: Name(NAME_None)
 		, ParentName(NAME_None)
+		, bAdvanced(false)
 	{
 	}
 
@@ -46,7 +50,7 @@ struct FNode
 		, ParentName(InParentName)
 		, Transform(InTransform)
 		, DisplayName(InNodeName.ToString())
-
+		, bAdvanced(false)
 	{
 	}
 };
@@ -85,41 +89,41 @@ struct FRigTransformConstraint
 	GENERATED_USTRUCT_BODY()
 
 	/** What transform type **/
-	UPROPERTY(/*EditAnywhere, Category="FRigControlConstraint"*/)
+	UPROPERTY(/*EditAnywhere, Category="FTransformBaseConstraint"*/)
 	TEnumAsByte<EConstraintTransform::Type>	TranformType;
 
 	/** Parent space that are define **/
-	UPROPERTY(EditAnywhere, Category="FRigControlConstraint")
+	UPROPERTY(EditAnywhere, Category="FTransformBaseConstraint")
 	FName	ParentSpace;
 
 	/** Weight of the influence - for future*/
-	UPROPERTY(/*EditAnywhere, Category="FRigControlConstraint"*/)
+	UPROPERTY(/*EditAnywhere, Category="FTransformBaseConstraint"*/)
 	float	Weight;
 };
 
 /** This defines what constraint it is defined */
 USTRUCT()
-struct FRigControlConstraint
+struct FTransformBaseConstraint
 {
 	GENERATED_USTRUCT_BODY()
 
 	/** What transform type **/
-	UPROPERTY(EditAnywhere, Category="FRigControlConstraint")
+	UPROPERTY(EditAnywhere, Category="FTransformBaseConstraint")
 	TArray<FRigTransformConstraint>			TransformConstraints;
 };
 
 
 /** This is a mapping table between bone in a particular skeletal mesh and bone of this skeleton set. */
 USTRUCT()
-struct FRigControl
+struct FTransformBase
 {
 	GENERATED_USTRUCT_BODY()
 
-	UPROPERTY(EditAnywhere, Category="FRigControl")
+	UPROPERTY(EditAnywhere, Category="FTransformBase")
 	FName					Node;
 
-	UPROPERTY(EditAnywhere, Category="FRigControl")
-	FRigControlConstraint	Constraints[EControlConstraint::Type::Max];
+	UPROPERTY(EditAnywhere, Category="FTransformBase")
+	FTransformBaseConstraint	Constraints[EControlConstraint::Type::Max];
 };
 
 DECLARE_DELEGATE_RetVal_OneParam(int32, FGetParentIndex, FName);
@@ -138,7 +142,7 @@ private:
 
 	/** Skeleton bone tree - each contains name and parent index**/
 	UPROPERTY(EditAnywhere, Category = Rig, EditFixedSize)
-	TArray<FRigControl> RigControls;
+	TArray<FTransformBase> TransformBases;
 
 	/** Skeleton bone tree - each contains name and parent index**/
 	UPROPERTY(EditAnywhere, Category=Rig, EditFixedSize)
@@ -165,11 +169,11 @@ public:
 	ENGINE_API void SetAllConstraintsToWorld();
 
 	// rig control related
-	int32 GetRigControlNum() const;
-	const TArray<FRigControl> & GetRigControls() const { return RigControls; }
-	const FRigControl * GetRigControl(int32 ControlIndex) const;
-	const FRigControl * GetRigControlByNodeName(FName NodeName) const;
-	int32 FindRigControlByNodeName(FName NodeName) const;
+	int32 GetTransformBaseNum() const;
+	const TArray<FTransformBase> & GetTransformBases() const { return TransformBases; }
+	const FTransformBase * GetTransformBase(int32 ControlIndex) const;
+	const FTransformBase * GetTransformBaseByNodeName(FName NodeName) const;
+	int32 FindTransformBaseByNodeName(FName NodeName) const;
 	int32 FindTransformParentNode(int32 NodeIndex, bool bTranslate, int32 Index=0) const;
 #endif
 
@@ -188,7 +192,7 @@ private:
 #endif // WITH_EDITOR
 	
 	// not useful so far
-	void CalculateComponentSpace(int32 NodeIndex, const FTransform & LocalTransform, const TArray<FTransform> & TransformBuffer, const FGetParentIndex& DelegateToGetParentIndex, FTransform & OutComponentSpaceTransform) const;
-	void CalculateLocalSpace(int32 NodeIndex, const FTransform & ComponentTransform, const TArray<FTransform> & TransformBuffer, const FGetParentIndex& DelegateToGetParentIndex, FTransform & OutLocalSpaceTransform) const;
+//	void CalculateComponentSpace(int32 NodeIndex, const FTransform & LocalTransform, const TArray<FTransform> & TransformBuffer, const FGetParentIndex& DelegateToGetParentIndex, FTransform & OutComponentSpaceTransform) const;
+//	void CalculateLocalSpace(int32 NodeIndex, const FTransform & ComponentTransform, const TArray<FTransform> & TransformBuffer, const FGetParentIndex& DelegateToGetParentIndex, FTransform & OutLocalSpaceTransform) const;
 };
 
diff --git a/Engine/Source/Runtime/Engine/Private/Animation/AnimSequence.cpp b/Engine/Source/Runtime/Engine/Private/Animation/AnimSequence.cpp
--- a/Engine/Source/Runtime/Engine/Private/Animation/AnimSequence.cpp
+++ b/Engine/Source/Runtime/Engine/Private/Animation/AnimSequence.cpp
@@ -1578,10 +1578,12 @@ void FillUpTransformBasedOnRig(const USkeleton * Skeleton, TArray<FTransform> &
 
 	const URig * Rig = Skeleton->GetRig();
 
-	if (Rig && Rig->GetNodeNum() > 0)
-	{
-		int32 NodeNum = Rig->GetNodeNum();
+	// this one has to collect all Nodes in Rig data
+	// since we're comparing two of them together. 
+	int32 NodeNum = Rig->GetNodeNum();
 
+	if (Rig && NodeNum > 0)
+	{
 		NodeSpaceBases.Empty(NodeNum);
 		NodeSpaceBases.AddUninitialized(NodeNum);
 
@@ -1593,30 +1595,30 @@ void FillUpTransformBasedOnRig(const USkeleton * Skeleton, TArray<FTransform> &
 
 		for ( int32 Index = 0; Index < NodeNum; ++Index )
 		{
-			const FName & NodeName = Rig->GetNodeName(Index);
+			const FName NodeName = Rig->GetNodeName(Index);
 			const FName & BoneName = Skeleton->GetRigBoneMapping(NodeName);
 			const int32 & BoneIndex = Skeleton->GetReferenceSkeleton().FindBoneIndex(BoneName);
 
 			if (BoneIndex == INDEX_NONE)
 			{
-				// initialize with SpaceBases - assuming World Based
-				NodeSpaceBases[Index] = FTransform::Identity;
-				Rotations[Index] = FTransform::Identity;
-				Translations[Index] = FTransform::Identity;
+				// add identity
+				NodeSpaceBases[Index].SetIdentity();
+				Rotations[Index].SetIdentity();
+				Translations[Index].SetIdentity();
 			}
-			else 
+			else
 			{
 				// initialize with SpaceBases - assuming World Based
 				NodeSpaceBases[Index] = SpaceBases[BoneIndex];
 				Rotations[Index] = SpaceBases[BoneIndex];
 				Translations[Index] = SpaceBases[BoneIndex];
 
-				const FRigControl * RigControl = Rig->GetRigControlByNodeName(NodeName);
+				const FTransformBase * TransformBase = Rig->GetTransformBaseByNodeName(NodeName);
 
-				if (RigControl != NULL)
+				if (TransformBase != NULL)
 				{
 					// orientation constraint			
-					const auto & RotConstraint = RigControl->Constraints[EControlConstraint::Type::Orientation];
+					const auto & RotConstraint = TransformBase->Constraints[EControlConstraint::Type::Orientation];
 
 					if (RotConstraint.TransformConstraints.Num() > 0)
 					{
@@ -1630,7 +1632,7 @@ void FillUpTransformBasedOnRig(const USkeleton * Skeleton, TArray<FTransform> &
 					}
 
 					// translation constraint
-					const auto & TransConstraint = RigControl->Constraints[EControlConstraint::Type::Translation];
+					const auto & TransConstraint = TransformBase->Constraints[EControlConstraint::Type::Translation];
 
 					if (TransConstraint.TransformConstraints.Num() > 0)
 					{
@@ -1649,6 +1651,22 @@ void FillUpTransformBasedOnRig(const USkeleton * Skeleton, TArray<FTransform> &
 	}
 }
 
+int32 FindValidTransformParentTrack(const URig * Rig, int32 NodeIndex, bool bTranslate, const TArray<FName>& ValidNodeNames)
+{
+	int32 ParentIndex = Rig->FindTransformParentNode(NodeIndex, bTranslate);
+
+	// verify if it exists in ValidNodeNames
+	if (ParentIndex != INDEX_NONE)
+	{
+		FName NodeName = Rig->GetNodeName(ParentIndex);
+
+		return ValidNodeNames.Find(NodeName);
+	}
+
+	return INDEX_NONE;
+}
+
+
 void UAnimSequence::RemapTracksToNewSkeleton( USkeleton * NewSkeleton, bool bConvertSpaces )
 {
 	// this is not cheap, so make sure it only happens in editor
@@ -1685,6 +1703,7 @@ void UAnimSequence::RemapTracksToNewSkeleton( USkeleton * NewSkeleton, bool bCon
 				//			P2 - parent transform in component space for new skeleton
 				//			R2 - local space of the current bone for new skeleton
 				// what we're looking for is theta, so that we can apply that to animated transform
+				// this has to have all of nodes since comparing two skeletons, that might have different configuration
 				int32 NumNodes = Rig->GetNodeNum();
 				// saves the theta data per node
 				RelativeToNewSpaceBases.AddUninitialized(NumNodes);
@@ -1697,10 +1716,6 @@ void UAnimSequence::RemapTracksToNewSkeleton( USkeleton * NewSkeleton, bool bCon
 				{
 					// theta (RelativeToNewTransform) = (P1*R1)^(-1) * P2*R2 where theta => P1*R1*theta = P2*R2
 					RelativeToNewSpaceBases[NodeIndex] = NewSpaceBases[NodeIndex].GetRelativeTransform(OldSpaceBases[NodeIndex]); 
-// 					UE_LOG(LogAnimation, Warning, TEXT("NodeName : %s"), *Rig->GetNodeName(NodeIndex).ToString());
-// 					UE_LOG(LogAnimation, Warning, TEXT("Old Translation : %s"), *OldSpaceBases[NodeIndex].ToString());
-// 					UE_LOG(LogAnimation, Warning, TEXT("New Translation : %s"), *NewSpaceBases[NodeIndex].ToString());
-// 					UE_LOG(LogAnimation, Warning, TEXT("Diff Translation : %s"), *RelativeToNewSpaceBases[NodeIndex].ToString());
 
 					// also savees the translation difference between old to new
 					FVector OldTranslation = OldTranslations[NodeIndex].GetTranslation();
@@ -1717,103 +1732,113 @@ void UAnimSequence::RemapTracksToNewSkeleton( USkeleton * NewSkeleton, bool bCon
 			// now convert animation data to rig data
 			ConvertAnimationDataToRiggingData(RawRiggingAnimationData);
 
+			// here we have to watch out the index
+			// The RawRiggingAnimationData will contain only the nodes that are mapped to source skeleton
+			// and here we convert everything that is in RawRiggingAnimationData which means based on source data
+			// when mapped back to new skeleton, it will discard results that are not mapped to target skeleton
+			
+			TArray<FName> SrcValidNodeNames;
+			int32 SrcNumTracks = OldSkeleton->GetMappedValidNodes(SrcValidNodeNames);
+
 			// now convert to space bases animation 
-			int32 NumNodes = Rig->GetNodeNum();
 			TArray< TArray<FTransform> > ComponentSpaceAnimations, ConvertedLocalSpaceAnimations, ConvertedSpaceAnimations;
-			ComponentSpaceAnimations.AddZeroed(NumNodes);
-			ConvertedSpaceAnimations.AddZeroed(NumNodes);
-			ConvertedLocalSpaceAnimations.AddZeroed(NumNodes);
+			ComponentSpaceAnimations.AddZeroed(SrcNumTracks);
+			ConvertedSpaceAnimations.AddZeroed(SrcNumTracks);
+			ConvertedLocalSpaceAnimations.AddZeroed(SrcNumTracks);
 
 			int32 NumKeys = NumFrames;
 			float Interval = (NumKeys)? SequenceLength/NumKeys : 0.f;
 
 			// allocate arrays
-			for(int32 NodeIndex=0; NodeIndex<NumNodes; ++NodeIndex)
+			for(int32 SrcTrackIndex=0; SrcTrackIndex<SrcNumTracks; ++SrcTrackIndex)
 			{
-				ComponentSpaceAnimations[NodeIndex].AddUninitialized(NumKeys);
-				ConvertedLocalSpaceAnimations[NodeIndex].AddUninitialized(NumKeys);
-				ConvertedSpaceAnimations[NodeIndex].AddUninitialized(NumKeys);
+				ComponentSpaceAnimations[SrcTrackIndex].AddUninitialized(NumKeys);
+				ConvertedLocalSpaceAnimations[SrcTrackIndex].AddUninitialized(NumKeys);
+				ConvertedSpaceAnimations[SrcTrackIndex].AddUninitialized(NumKeys);
 			}
 
-
-			for (int32 NodeIndex=0; NodeIndex<NumNodes; ++NodeIndex)		
+			for (int32 SrcTrackIndex=0; SrcTrackIndex<SrcNumTracks; ++SrcTrackIndex)		
 			{
-				auto & RawAnimation = RawRiggingAnimationData[NodeIndex];
+				int32 NodeIndex = Rig->FindNode(SrcValidNodeNames[SrcTrackIndex]);
+				check (NodeIndex != INDEX_NONE);
+				auto & RawAnimation = RawRiggingAnimationData[SrcTrackIndex];
 
 				// find rotation parent node
-				int32 RotParentNodeIndex = Rig->FindTransformParentNode(NodeIndex, false);
-				int32 TransParentNodeIndex = Rig->FindTransformParentNode(NodeIndex, true);
+				int32 RotParentTrackIndex = FindValidTransformParentTrack(Rig, NodeIndex, false, SrcValidNodeNames);
+				int32 TransParentTrackIndex = FindValidTransformParentTrack(Rig, NodeIndex, true, SrcValidNodeNames);
 				// fill up keys - calculate PK1 * K1
 				for(int32 Key=0; Key<NumKeys; ++Key)
 				{
 					FTransform AnimatedLocalKey;
-					ExtractBoneTransform(RawRiggingAnimationData, AnimatedLocalKey, NodeIndex, Interval*Key);
+					ExtractBoneTransform(RawRiggingAnimationData, AnimatedLocalKey, SrcTrackIndex, Interval*Key);
 
 					AnimatedLocalKey.ScaleTranslation(OldToNewTranslationRatio[NodeIndex]);
 
-					if(RotParentNodeIndex != INDEX_NONE)
+					if(RotParentTrackIndex != INDEX_NONE)
 					{
-						FQuat ComponentSpaceRotation = ComponentSpaceAnimations[RotParentNodeIndex][Key].GetRotation() * AnimatedLocalKey.GetRotation();
-						ComponentSpaceAnimations[NodeIndex][Key].SetRotation(ComponentSpaceRotation);
+						FQuat ComponentSpaceRotation = ComponentSpaceAnimations[RotParentTrackIndex][Key].GetRotation() * AnimatedLocalKey.GetRotation();
+						ComponentSpaceAnimations[SrcTrackIndex][Key].SetRotation(ComponentSpaceRotation);
 					}
 					else
 					{
-						ComponentSpaceAnimations[NodeIndex][Key].SetRotation(AnimatedLocalKey.GetRotation());
+						ComponentSpaceAnimations[SrcTrackIndex][Key].SetRotation(AnimatedLocalKey.GetRotation());
 					}
 
-					if (TransParentNodeIndex != INDEX_NONE)
+					if (TransParentTrackIndex != INDEX_NONE)
 					{
-						FVector ComponentSpaceTranslation = ComponentSpaceAnimations[TransParentNodeIndex][Key].TransformPosition(AnimatedLocalKey.GetTranslation());
-						ComponentSpaceAnimations[NodeIndex][Key].SetTranslation(ComponentSpaceTranslation);
-						ComponentSpaceAnimations[NodeIndex][Key].SetScale3D(AnimatedLocalKey.GetScale3D());
+						FVector ComponentSpaceTranslation = ComponentSpaceAnimations[TransParentTrackIndex][Key].TransformPosition(AnimatedLocalKey.GetTranslation());
+						ComponentSpaceAnimations[SrcTrackIndex][Key].SetTranslation(ComponentSpaceTranslation);
+						ComponentSpaceAnimations[SrcTrackIndex][Key].SetScale3D(AnimatedLocalKey.GetScale3D());
 					}
 					else
 					{
-						ComponentSpaceAnimations[NodeIndex][Key].SetTranslation(AnimatedLocalKey.GetTranslation());
-						ComponentSpaceAnimations[NodeIndex][Key].SetScale3D(AnimatedLocalKey.GetScale3D());
+						ComponentSpaceAnimations[SrcTrackIndex][Key].SetTranslation(AnimatedLocalKey.GetTranslation());
+						ComponentSpaceAnimations[SrcTrackIndex][Key].SetScale3D(AnimatedLocalKey.GetScale3D());
 					}
 				}
 			}
 
 			// now animation is converted to component space
 			TArray<struct FRawAnimSequenceTrack> NewRawAnimationData = RawRiggingAnimationData;
-			for (int32 NodeIndex=0; NodeIndex<NumNodes; ++NodeIndex)
+			for (int32 SrcTrackIndex=0; SrcTrackIndex<SrcNumTracks; ++SrcTrackIndex)
 			{
+				int32 NodeIndex = Rig->FindNode(SrcValidNodeNames[SrcTrackIndex]);
 				// find rotation parent node
-				int32 RotParentNodeIndex = Rig->FindTransformParentNode(NodeIndex, false);
-				int32 TransParentNodeIndex = Rig->FindTransformParentNode(NodeIndex, true);
+				int32 RotParentTrackIndex = FindValidTransformParentTrack(Rig, NodeIndex, false, SrcValidNodeNames);
+				int32 TransParentTrackIndex = FindValidTransformParentTrack(Rig, NodeIndex, true, SrcValidNodeNames);
+
+				// clear translation;
+				RelativeToNewSpaceBases[NodeIndex].SetTranslation(FVector::ZeroVector);
 
 				for(int32 Key=0; Key<NumKeys; ++Key)
 				{
 					// now convert to the new space and save to local spaces
-					// clear translation;
-					RelativeToNewSpaceBases[NodeIndex].SetTranslation(FVector::ZeroVector);
-					ConvertedSpaceAnimations[NodeIndex][Key] = RelativeToNewSpaceBases[NodeIndex] * ComponentSpaceAnimations[NodeIndex][Key];
+					ConvertedSpaceAnimations[SrcTrackIndex][Key] = RelativeToNewSpaceBases[NodeIndex] * ComponentSpaceAnimations[SrcTrackIndex][Key];
 
-					if(RotParentNodeIndex != INDEX_NONE)
+					if(RotParentTrackIndex != INDEX_NONE)
 					{
-						FQuat LocalRotation = ConvertedSpaceAnimations[RotParentNodeIndex][Key].GetRotation().Inverse() * ConvertedSpaceAnimations[NodeIndex][Key].GetRotation();
-						ConvertedLocalSpaceAnimations[NodeIndex][Key].SetRotation(LocalRotation);
+						FQuat LocalRotation = ConvertedSpaceAnimations[RotParentTrackIndex][Key].GetRotation().Inverse() * ConvertedSpaceAnimations[SrcTrackIndex][Key].GetRotation();
+						ConvertedLocalSpaceAnimations[SrcTrackIndex][Key].SetRotation(LocalRotation);
 					}
 					else
 					{
-						ConvertedLocalSpaceAnimations[NodeIndex][Key].SetRotation(ConvertedSpaceAnimations[NodeIndex][Key].GetRotation());
+						ConvertedLocalSpaceAnimations[SrcTrackIndex][Key].SetRotation(ConvertedSpaceAnimations[SrcTrackIndex][Key].GetRotation());
 					}
 
-					if(TransParentNodeIndex != INDEX_NONE)
+					if(TransParentTrackIndex != INDEX_NONE)
 					{
-						FVector LocalTranslation = ConvertedSpaceAnimations[NodeIndex][Key].GetRelativeTransform(ConvertedSpaceAnimations[TransParentNodeIndex][Key]).GetTranslation();
-						ConvertedLocalSpaceAnimations[NodeIndex][Key].SetTranslation(LocalTranslation);
-						ConvertedLocalSpaceAnimations[NodeIndex][Key].SetScale3D(ConvertedSpaceAnimations[NodeIndex][Key].GetScale3D());
+						FVector LocalTranslation = ConvertedSpaceAnimations[SrcTrackIndex][Key].GetRelativeTransform(ConvertedSpaceAnimations[TransParentTrackIndex][Key]).GetTranslation();
+						ConvertedLocalSpaceAnimations[SrcTrackIndex][Key].SetTranslation(LocalTranslation);
+						ConvertedLocalSpaceAnimations[SrcTrackIndex][Key].SetScale3D(ConvertedSpaceAnimations[SrcTrackIndex][Key].GetScale3D());
 					}
 					else
 					{
-						ConvertedLocalSpaceAnimations[NodeIndex][Key].SetTranslation(ConvertedSpaceAnimations[NodeIndex][Key].GetTranslation());
-						ConvertedLocalSpaceAnimations[NodeIndex][Key].SetScale3D(ConvertedSpaceAnimations[NodeIndex][Key].GetScale3D());
+						ConvertedLocalSpaceAnimations[SrcTrackIndex][Key].SetTranslation(ConvertedSpaceAnimations[SrcTrackIndex][Key].GetTranslation());
+						ConvertedLocalSpaceAnimations[SrcTrackIndex][Key].SetScale3D(ConvertedSpaceAnimations[SrcTrackIndex][Key].GetScale3D());
 					}
 				}
 
-				auto & RawAnimation = NewRawAnimationData[NodeIndex];
+				auto & RawAnimation = NewRawAnimationData[SrcTrackIndex];
 				RawAnimation.PosKeys.Empty(NumKeys);
 				RawAnimation.PosKeys.AddUninitialized(NumKeys);
 				RawAnimation.RotKeys.Empty(NumKeys);
@@ -1823,13 +1848,48 @@ void UAnimSequence::RemapTracksToNewSkeleton( USkeleton * NewSkeleton, bool bCon
 
 				for(int32 Key=0; Key<NumKeys; ++Key)
 				{
-					RawAnimation.PosKeys[Key] = ConvertedLocalSpaceAnimations[NodeIndex][Key].GetLocation();
-					RawAnimation.RotKeys[Key] = ConvertedLocalSpaceAnimations[NodeIndex][Key].GetRotation();
-					RawAnimation.ScaleKeys[Key] = ConvertedLocalSpaceAnimations[NodeIndex][Key].GetScale3D();
+					RawAnimation.PosKeys[Key] = ConvertedLocalSpaceAnimations[SrcTrackIndex][Key].GetLocation();
+					RawAnimation.RotKeys[Key] = ConvertedLocalSpaceAnimations[SrcTrackIndex][Key].GetRotation();
+					RawAnimation.ScaleKeys[Key] = ConvertedLocalSpaceAnimations[SrcTrackIndex][Key].GetScale3D();
+				}
+			}
+
+			// now add missing tracsk, before converting back to rigging data, we need whole nodes to be mapped
+			// @Todo fix this to do this earlier, just like to avoid confusion of so many index
+			if (SrcNumTracks  != Rig->GetNodeNum())
+			{
+				// if not all entries are filled up
+				// insert nodes
+				int32 NumNodes = Rig->GetNodeNum();
+				for (int32 NodeIndex=0; NodeIndex < NumNodes; ++NodeIndex)
+				{
+					FName NodeName = Rig->GetNodeName(NodeIndex);
+					// if not found
+					if (SrcValidNodeNames.Contains(NodeName) == false)
+					{
+						NewRawAnimationData.InsertZeroed(NodeIndex);
+
+						auto & RawAnimation = NewRawAnimationData[NodeIndex];
+						RawAnimation.PosKeys.Empty(NumKeys);
+						RawAnimation.PosKeys.AddUninitialized(NumKeys);
+						RawAnimation.RotKeys.Empty(NumKeys);
+						RawAnimation.RotKeys.AddUninitialized(NumKeys);
+						RawAnimation.ScaleKeys.Empty(NumKeys);
+						RawAnimation.ScaleKeys.AddUninitialized(NumKeys);
+
+						for(int32 Key=0; Key<NumKeys; ++Key)
+						{
+							RawAnimation.PosKeys[Key] = FVector::ZeroVector;
+							RawAnimation.RotKeys[Key] = FQuat::Identity;
+							RawAnimation.ScaleKeys[Key] = FVector(1.f);
+						}
+					}
 				}
 			}
 
 			RawRiggingAnimationData = NewRawAnimationData;
+			// here RawRiggingAnimationData always should contain # of nodes of entries
+			check (RawRiggingAnimationData.Num() == Rig->GetNodeNum());
 
 			// set new skeleton
 			SetSkeleton(NewSkeleton);
@@ -2306,11 +2366,11 @@ int32 UAnimSequence::GetSpaceBasedAnimationData(TArray< TArray<FTransform> > & A
 		// @Todo sort rig control by parent -> child, so that we can calculate parent first before child
 		// but that depends on each skeleton
 
-		const TArray<FRigControl> & RigControls = Rig->GetRigControls();
+		const TArray<FTransformBase> & TransformBases = Rig->GetTransformBases();
 		// first put all the component space
-		for (const auto & RigControl : RigControls)
+		for (const auto & TransformBase : TransformBases)
 		{
-			const FName & NodeName = RigControl.Node;
+			const FName & NodeName = TransformBase.Node;
 			const int32 & NodeIndex = Rig->FindNode(NodeName);
 			const FName & BoneName = Skeleton->GetRigBoneMapping(NodeName);
 			const int32 & BoneIndex = RefSkeleton.FindBoneIndex(BoneName);
@@ -2318,7 +2378,7 @@ int32 UAnimSequence::GetSpaceBasedAnimationData(TArray< TArray<FTransform> > & A
 			if (BoneIndex != INDEX_NONE)
 			{
 				// now calculate the component space
-				const TArray<FRigTransformConstraint>	& RotTransformConstraints = RigControl.Constraints[EControlConstraint::Type::Orientation].TransformConstraints;
+				const TArray<FRigTransformConstraint>	& RotTransformConstraints = TransformBase.Constraints[EControlConstraint::Type::Orientation].TransformConstraints;
 
 				FQuat ComponentRotation;
 				FTransform ComponentTranslation;
@@ -2354,7 +2414,7 @@ int32 UAnimSequence::GetSpaceBasedAnimationData(TArray< TArray<FTransform> > & A
 					bCompleteComponentSpaces[BoneIndex] = true;
 				}
 
-				const TArray<FRigTransformConstraint>	& PosTransformConstraints = RigControl.Constraints[EControlConstraint::Type::Translation].TransformConstraints;
+				const TArray<FRigTransformConstraint>	& PosTransformConstraints = TransformBase.Constraints[EControlConstraint::Type::Translation].TransformConstraints;
 
 				//if (PosTransformConstraints.Num() > 0)
 				{
@@ -2471,7 +2531,8 @@ bool UAnimSequence::ConvertAnimationDataToRiggingData(TArray<struct FRawAnimSequ
 	if (Skeleton && Skeleton->GetRig())
 	{
 		const URig * Rig = Skeleton->GetRig();
-		const int32 & NumNodes = Rig->GetNodeNum();
+		TArray<FName> ValidNodeNames;
+		int32 NumNodes = Skeleton->GetMappedValidNodes(ValidNodeNames);
 		TArray< TArray<FTransform> > AnimationDataInComponentSpace;
 		int32 NumBones = GetSpaceBasedAnimationData(AnimationDataInComponentSpace, NULL);
 
@@ -2484,24 +2545,24 @@ bool UAnimSequence::ConvertAnimationDataToRiggingData(TArray<struct FRawAnimSequ
 			for (int32 NodeIndex = 0; NodeIndex < NumNodes; ++NodeIndex)
 			{
 				struct FRawAnimSequenceTrack & Track = RawRiggingAnimationData[NodeIndex];
-				const FName & NodeName = Rig->GetNodeName(NodeIndex);
+				const FName & NodeName = ValidNodeNames[NodeIndex];
 				const FName & BoneName = Skeleton->GetRigBoneMapping(NodeName);
 				const int32 & BoneIndex = Skeleton->GetReferenceSkeleton().FindBoneIndex(BoneName);
 
-				Track.PosKeys.Empty(NumFrames);
-				Track.RotKeys.Empty(NumFrames);
-				Track.ScaleKeys.Empty(NumFrames);
-				Track.PosKeys.AddUninitialized(NumFrames);
-				Track.RotKeys.AddUninitialized(NumFrames);
-				Track.ScaleKeys.AddUninitialized(NumFrames);
-
-				if (BoneIndex != INDEX_NONE)
+				if (ensure(BoneIndex != INDEX_NONE))
 				{
-					int32 RigConstraintIndex = Rig->FindRigControlByNodeName(NodeName);
+					Track.PosKeys.Empty(NumFrames);
+					Track.RotKeys.Empty(NumFrames);
+					Track.ScaleKeys.Empty(NumFrames);
+					Track.PosKeys.AddUninitialized(NumFrames);
+					Track.RotKeys.AddUninitialized(NumFrames);
+					Track.ScaleKeys.AddUninitialized(NumFrames);
+
+					int32 RigConstraintIndex = Rig->FindTransformBaseByNodeName(NodeName);
 
 					if (RigConstraintIndex != INDEX_NONE)
 					{
-						const auto * RigConstraint = Rig->GetRigControl(RigConstraintIndex);
+						const auto * RigConstraint = Rig->GetTransformBase(RigConstraintIndex);
 
 						// apply orientation - for now only one
 						const TArray<FRigTransformConstraint> & RotationTransformConstraint = RigConstraint->Constraints[EControlConstraint::Type::Orientation].TransformConstraints;
@@ -2587,16 +2648,6 @@ bool UAnimSequence::ConvertAnimationDataToRiggingData(TArray<struct FRawAnimSequ
 						}
 					}
 				}
-				else
-				{
-					// if no bone index, we still put identity
-					for (int32 KeyIndex = 0; KeyIndex < NumFrames; ++KeyIndex)
-					{
-						Track.PosKeys[KeyIndex] = FVector::ZeroVector;
-						Track.RotKeys[KeyIndex] = FQuat::Identity;
-						Track.ScaleKeys[KeyIndex] = FVector(1.f);
-					}
-				}
 			}
 		}
 
@@ -2613,22 +2664,23 @@ bool UAnimSequence::ConvertRiggingDataToAnimationData(TArray<struct FRawAnimSequ
 		TArray< TArray<FTransform> > AnimationDataInComponentSpace;
 		int32 NumBones = GetSpaceBasedAnimationData(AnimationDataInComponentSpace, &RawRiggingAnimationData);
 
+		USkeleton * Skeleton = GetSkeleton();
 		TArray<FRawAnimSequenceTrack> OldAnimationData = RawAnimationData;
 		TArray<FName> OldAnimationTrackNames = AnimationTrackNames;
-		int32 NumNodes = RawRiggingAnimationData.Num();
+		TArray<FName> ValidNodeNames;
+		int32 ValidNumNodes = Skeleton->GetMappedValidNodes(ValidNodeNames);
 		// get local spaces
 		// add all tracks?
-		AnimationTrackNames.Empty(NumNodes);
-		AnimationTrackNames.AddUninitialized(NumNodes);
-		RawAnimationData.Empty(NumNodes);
-		RawAnimationData.AddZeroed(NumNodes);
+		AnimationTrackNames.Empty(ValidNumNodes);
+		AnimationTrackNames.AddUninitialized(ValidNumNodes);
+		RawAnimationData.Empty(ValidNumNodes);
+		RawAnimationData.AddZeroed(ValidNumNodes);
 
-		USkeleton * Skeleton = GetSkeleton();
 		const FReferenceSkeleton & RefSkeleton = Skeleton->GetReferenceSkeleton();
 		const URig * Rig = Skeleton->GetRig();
-		for (int32 NodeIndex = 0; NodeIndex < NumNodes; ++NodeIndex)
+		for (int32 NodeIndex = 0; NodeIndex < ValidNumNodes; ++NodeIndex)
 		{
-			FName BoneName = Skeleton->GetRigBoneMapping(Rig->GetNodeName(NodeIndex));
+			FName BoneName = Skeleton->GetRigBoneMapping(ValidNodeNames[NodeIndex]);
 			int32 BoneIndex = RefSkeleton.FindBoneIndex(BoneName);
 
 			if (BoneIndex != INDEX_NONE)
diff --git a/Engine/Source/Runtime/Engine/Private/Animation/Rig.cpp b/Engine/Source/Runtime/Engine/Private/Animation/Rig.cpp
--- a/Engine/Source/Runtime/Engine/Private/Animation/Rig.cpp
+++ b/Engine/Source/Runtime/Engine/Private/Animation/Rig.cpp
@@ -123,12 +123,12 @@ bool URig::AddRigConstraint(FName NodeName, EControlConstraint::Type ConstraintT
 		ParentSpace = WorldNodeName;
 	}
 
-	int32 Index = FindRigControlByNodeName(NodeName);
+	int32 Index = FindTransformBaseByNodeName(NodeName);
 
 	if (Index == INDEX_NONE)
 	{
-		FRigControl NewRigControl;
-		NewRigControl.Node = NodeName;
+		FTransformBase NewTransformBase;
+		NewTransformBase.Node = NodeName;
 
 		if (ConstraintType < EControlConstraint::Max)
 		{
@@ -137,19 +137,19 @@ bool URig::AddRigConstraint(FName NodeName, EControlConstraint::Type ConstraintT
 			NewTransformConstraint.ParentSpace = ParentSpace;
 			NewTransformConstraint.Weight = Weight;
 
-			NewRigControl.Constraints[ConstraintType].TransformConstraints.Add(NewTransformConstraint);
+			NewTransformBase.Constraints[ConstraintType].TransformConstraints.Add(NewTransformConstraint);
 
-			RigControls.Add(NewRigControl);
+			TransformBases.Add(NewTransformBase);
 		}
 	}
 	else
 	{
 		// it exists already, need to make sure we don't have different constraint types
-		FRigControl & RigControl = RigControls[Index];
+		FTransformBase & TransformBase = TransformBases[Index];
 
 		if (ConstraintType < EControlConstraint::Max)
 		{
-			FRigControlConstraint & ControlConstraint = RigControl.Constraints[ConstraintType];
+			FTransformBaseConstraint & ControlConstraint = TransformBase.Constraints[ConstraintType];
 
 			FRigTransformConstraint NewTransformConstraint;
 			NewTransformConstraint.TranformType = TransformType;
@@ -169,28 +169,28 @@ int32 URig::GetNodeNum() const
 	return Nodes.Num();
 }
 
-int32 URig::GetRigControlNum() const
+int32 URig::GetTransformBaseNum() const
 {
-	return RigControls.Num();
+	return TransformBases.Num();
 }
 
-const FRigControl * URig::GetRigControl(int32 ControlIndex) const
+const FTransformBase * URig::GetTransformBase(int32 ControlIndex) const
 {
-	if (RigControls.IsValidIndex(ControlIndex))
+	if (TransformBases.IsValidIndex(ControlIndex))
 	{
-		return &RigControls[ControlIndex];
+		return &TransformBases[ControlIndex];
 	}
 
 	return NULL;
 }
 
-const FRigControl * URig::GetRigControlByNodeName(FName NodeName) const
+const FTransformBase * URig::GetTransformBaseByNodeName(FName NodeName) const
 {
-	int32 ControlIndex = FindRigControlByNodeName(NodeName);
+	int32 ControlIndex = FindTransformBaseByNodeName(NodeName);
 
 	if (ControlIndex != INDEX_NONE)
 	{
-		return &RigControls[ControlIndex];
+		return &TransformBases[ControlIndex];
 	}
 
 	return NULL;
@@ -200,18 +200,18 @@ int32 URig::FindTransformParentNode(int32 NodeIndex, bool bTranslate, int32 Inde
 {
 	if (Nodes.IsValidIndex(NodeIndex))
 	{
-		const FRigControl * RigControl = GetRigControlByNodeName(Nodes[NodeIndex].Name);
+		const FTransformBase * TransformBase = GetTransformBaseByNodeName(Nodes[NodeIndex].Name);
 
-		if(RigControl)
+		if(TransformBase)
 		{
 			FName ParentNodeName = NAME_None;
 			if(bTranslate)
 			{
-				ParentNodeName = RigControl->Constraints[EControlConstraint::Type::Translation].TransformConstraints[Index].ParentSpace;
+				ParentNodeName = TransformBase->Constraints[EControlConstraint::Type::Translation].TransformConstraints[Index].ParentSpace;
 			}
 			else
 			{
-				ParentNodeName = RigControl->Constraints[EControlConstraint::Type::Orientation].TransformConstraints[Index].ParentSpace;
+				ParentNodeName = TransformBase->Constraints[EControlConstraint::Type::Orientation].TransformConstraints[Index].ParentSpace;
 			}
 
 			if(ParentNodeName != NAME_None)
@@ -224,13 +224,13 @@ int32 URig::FindTransformParentNode(int32 NodeIndex, bool bTranslate, int32 Inde
 	return INDEX_NONE;
 }
 
-int32 URig::FindRigControlByNodeName(FName NodeName) const
+int32 URig::FindTransformBaseByNodeName(FName NodeName) const
 {
 	int32 Index=0;
 
-	for (auto RigControl : RigControls)
+	for (auto TransformBase : TransformBases)
 	{
-		if (RigControl.Node == NodeName)
+		if (TransformBase.Node == NodeName)
 		{
 			return Index;
 		}
@@ -297,13 +297,13 @@ void URig::CreateFromSkeleton(const USkeleton * Skeleton, const TMap<int32, int3
 /*
 void URig::CalculateComponentSpace(int32 NodeIndex, const FTransform & LocalTransform, const TArray<FTransform> & TransformBuffer, const FGetParentIndex& DelegateToGetParentIndex, FTransform & OutComponentSpaceTransform) const
 {
-	int32 ConstraintIndex = FindRigControlByNodeName(Nodes[NodeIndex].Name);
+	int32 ConstraintIndex = FindTransformBaseByNodeName(Nodes[NodeIndex].Name);
 
 	if (ConstraintIndex != INDEX_NONE)
 	{
 		// now find transform constraint data
 		{
-			const FRigControlConstraint & Constraints = RigControls[ConstraintIndex].Constraints[EControlConstraint::Type::Orientation];
+			const FTransformBaseConstraint & Constraints = TransformBases[ConstraintIndex].Constraints[EControlConstraint::Type::Orientation];
 
 			// for now we only care for the first one
 			const FName ParentName = Constraints.TransformConstraints[0].ParentSpace;
@@ -327,7 +327,7 @@ void URig::CalculateComponentSpace(int32 NodeIndex, const FTransform & LocalTran
 
 		// same thing for translation
 		{
-			const FRigControlConstraint & Constraints = RigControls[ConstraintIndex].Constraints[EControlConstraint::Type::Translation];
+			const FTransformBaseConstraint & Constraints = TransformBases[ConstraintIndex].Constraints[EControlConstraint::Type::Translation];
 
 			// for now we only care for the first one
 			const FName ParentName = Constraints.TransformConstraints[0].ParentSpace;
@@ -355,13 +355,13 @@ void URig::CalculateComponentSpace(int32 NodeIndex, const FTransform & LocalTran
 
 void URig::CalculateLocalSpace(int32 NodeIndex, const FTransform & ComponentTransform, const TArray<FTransform> & TransformBuffer, const FGetParentIndex& DelegateToGetParentIndex, FTransform & OutLocalSpaceTransform) const
 {
-	int32 ConstraintIndex = FindRigControlByNodeName(Nodes[NodeIndex].Name);
+	int32 ConstraintIndex = FindTransformBaseByNodeName(Nodes[NodeIndex].Name);
 
 	if(ConstraintIndex != INDEX_NONE)
 	{
 		// now find transform constraint data
 		{
-			const FRigControlConstraint & Constraints = RigControls[ConstraintIndex].Constraints[EControlConstraint::Type::Orientation];
+			const FTransformBaseConstraint & Constraints = TransformBases[ConstraintIndex].Constraints[EControlConstraint::Type::Orientation];
 
 			// for now we only care for the first one
 			const FName ParentName = Constraints.TransformConstraints[0].ParentSpace;
@@ -385,7 +385,7 @@ void URig::CalculateLocalSpace(int32 NodeIndex, const FTransform & ComponentTran
 
 		// same thing for translation
 		{
-			const FRigControlConstraint & Constraints = RigControls[ConstraintIndex].Constraints[EControlConstraint::Type::Translation];
+			const FTransformBaseConstraint & Constraints = TransformBases[ConstraintIndex].Constraints[EControlConstraint::Type::Translation];
 
 			// for now we only care for the first one
 			const FName ParentName = Constraints.TransformConstraints[0].ParentSpace;
@@ -414,7 +414,7 @@ void URig::CalculateLocalSpace(int32 NodeIndex, const FTransform & ComponentTran
 
 void URig::SetAllConstraintsToParents()
 {
-	for(auto & Control : RigControls)
+	for(auto & Control : TransformBases)
 	{
 		FName ParentNode = GetParentNode(Control.Node);
 
@@ -424,7 +424,7 @@ void URig::SetAllConstraintsToParents()
 }
 void URig::SetAllConstraintsToWorld()
 {
-	for (auto & Control : RigControls)
+	for (auto & Control : TransformBases)
 	{
 		Control.Constraints[EControlConstraint::Type::Translation].TransformConstraints[0].ParentSpace = WorldNodeName;
 		Control.Constraints[EControlConstraint::Type::Orientation].TransformConstraints[0].ParentSpace = WorldNodeName;
diff --git a/Engine/Source/Runtime/Engine/Private/Animation/Skeleton.cpp b/Engine/Source/Runtime/Engine/Private/Animation/Skeleton.cpp
--- a/Engine/Source/Runtime/Engine/Private/Animation/Skeleton.cpp
+++ b/Engine/Source/Runtime/Engine/Private/Animation/Skeleton.cpp
@@ -1066,17 +1066,18 @@ bool USkeleton::SetRigBoneMapping(const FName & NodeName, FName BoneName)
 	// but it might be annoying to do that right now since the tool is ugly
 	// so for now it lets you set everything, but in the future
 	// we'll have to add verification
-	if ( BoneIndex != INDEX_NONE )
+	if ( BoneIndex == INDEX_NONE )
 	{
-		int32 Index = FindRigBoneMapping(NodeName);
-
-		if(Index != INDEX_NONE)
-		{
-			RigConfig.BoneMappingTable[Index].BoneName = BoneName;
-			return true;
-		}
+		BoneName = NAME_None;
 	}
 
+	int32 Index = FindRigBoneMapping(NodeName);
+
+	if(Index != INDEX_NONE)
+	{
+		RigConfig.BoneMappingTable[Index].BoneName = BoneName;
+		return true;
+	}
 
 	return false;
 }
