Fix FBodyInstance copying natively instead of by property.
Fixes crash when converting actor to blueprint.

TTP#346217

diff --git a/Engine/Source/Runtime/Engine/Classes/PhysicsEngine/BodyInstance.h b/Engine/Source/Runtime/Engine/Classes/PhysicsEngine/BodyInstance.h
--- a/Engine/Source/Runtime/Engine/Classes/PhysicsEngine/BodyInstance.h
+++ b/Engine/Source/Runtime/Engine/Classes/PhysicsEngine/BodyInstance.h
@@ -103,6 +103,7 @@ struct ENGINE_API FBodyInstance
 {
 	GENERATED_USTRUCT_BODY()
 
+
 	/** 
 	 *	Index of this BodyInstance within the SkeletalMeshComponent/PhysicsAsset. 
 	 *	Is INDEX_NONE if a single body component
@@ -672,6 +673,15 @@ protected:
 #endif	//WITH_BOX2D
 };
 
+template<>
+struct TStructOpsTypeTraits<FBodyInstance> : public TStructOpsTypeTraitsBase
+{
+	enum
+	{
+		WithCopy = false
+	};
+};
+
 
 #if WITH_EDITOR
 
