#docs - Clean up links in framework headers.

diff --git a/Engine/Source/Runtime/Engine/Classes/Animation/SkeletalMeshActor.h b/Engine/Source/Runtime/Engine/Classes/Animation/SkeletalMeshActor.h
--- a/Engine/Source/Runtime/Engine/Classes/Animation/SkeletalMeshActor.h
+++ b/Engine/Source/Runtime/Engine/Classes/Animation/SkeletalMeshActor.h
@@ -9,7 +9,7 @@
  * Skeletal meshes are deformable meshes that can be animated and change their geometry at run-time.
  * Skeletal meshes dragged into the level from the Content Browser are automatically converted to StaticMeshActors.
  * 
- * @see https://docs.unrealengine.com/latest/INT/Engine/Content/Types/SkeletalMeshes/index.html
+ * @see https://docs.unrealengine.com/latest/INT/Engine/Content/Types/SkeletalMeshes/
  * @see USkeletalMesh
  */
 UCLASS(ClassGroup=ISkeletalMeshes, Blueprintable, ConversionRoot, meta=(ChildCanTick))
diff --git a/Engine/Source/Runtime/Engine/Classes/Components/SkeletalMeshComponent.h b/Engine/Source/Runtime/Engine/Classes/Components/SkeletalMeshComponent.h
--- a/Engine/Source/Runtime/Engine/Classes/Components/SkeletalMeshComponent.h
+++ b/Engine/Source/Runtime/Engine/Classes/Components/SkeletalMeshComponent.h
@@ -262,7 +262,7 @@ struct FSkeletalMeshComponentPreClothTickFunction : public FTickFunction
 /**
  * SkeletalMeshComponent is used to create an instance of a USkeletalMesh.
  *
- * @see https://docs.unrealengine.com/latest/INT/Engine/Content/Types/SkeletalMeshes/index.html
+ * @see https://docs.unrealengine.com/latest/INT/Engine/Content/Types/SkeletalMeshes/
  * @see USkeletalMesh
  */
 UCLASS(ClassGroup=(Rendering, Common), hidecategories=Object, config=Engine, editinlinenew, meta=(BlueprintSpawnableComponent))
diff --git a/Engine/Source/Runtime/Engine/Classes/Components/StaticMeshComponent.h b/Engine/Source/Runtime/Engine/Classes/Components/StaticMeshComponent.h
--- a/Engine/Source/Runtime/Engine/Classes/Components/StaticMeshComponent.h
+++ b/Engine/Source/Runtime/Engine/Classes/Components/StaticMeshComponent.h
@@ -109,7 +109,7 @@ struct TStructOpsTypeTraits<FStaticMeshComponentLODInfo> : public TStructOpsType
 /**
  * StaticMeshComponent is used to create an instance of a UStaticMesh.
  *
- * @see https://docs.unrealengine.com/latest/INT/Engine/Content/Types/StaticMeshes/index.html
+ * @see https://docs.unrealengine.com/latest/INT/Engine/Content/Types/StaticMeshes/
  * @see UStaticMesh
  */
 UCLASS(ClassGroup=(Rendering, Common), hidecategories=(Object,Activation,"Components|Activation"), ShowCategories=(Mobility), editinlinenew, meta=(BlueprintSpawnableComponent))
diff --git a/Engine/Source/Runtime/Engine/Classes/Engine/SkeletalMesh.h b/Engine/Source/Runtime/Engine/Classes/Engine/SkeletalMesh.h
--- a/Engine/Source/Runtime/Engine/Classes/Engine/SkeletalMesh.h
+++ b/Engine/Source/Runtime/Engine/Classes/Engine/SkeletalMesh.h
@@ -560,7 +560,7 @@ class FSkeletalMeshResource;
  * Skeletal Meshes are built up of two parts; a set of polygons composed to make up the surface of the mesh, and a hierarchical skeleton which can be used to animate the polygons.
  * The 3D models, rigging, and animations are created in an external modeling and animation application (3DSMax, Maya, Softimage, etc).
  *
- * @see https://docs.unrealengine.com/latest/INT/Engine/Content/Types/SkeletalMeshes/index.html
+ * @see https://docs.unrealengine.com/latest/INT/Engine/Content/Types/SkeletalMeshes/
  */
 UCLASS(hidecategories=Object, MinimalAPI, BlueprintType)
 class USkeletalMesh : public UObject
diff --git a/Engine/Source/Runtime/Engine/Classes/Engine/StaticMesh.h b/Engine/Source/Runtime/Engine/Classes/Engine/StaticMesh.h
--- a/Engine/Source/Runtime/Engine/Classes/Engine/StaticMesh.h
+++ b/Engine/Source/Runtime/Engine/Classes/Engine/StaticMesh.h
@@ -318,7 +318,7 @@ struct FAssetEditorOrbitCameraPosition
  * Static Meshes can be translated, rotated, and scaled, but they cannot have their vertices animated in any way. As such, they are more efficient
  * to render than other types of geometry such as USkeletalMesh, and they are often the basic building block of levels created in the engine.
  *
- * @see https://docs.unrealengine.com/latest/INT/Engine/Content/Types/StaticMeshes/index.html
+ * @see https://docs.unrealengine.com/latest/INT/Engine/Content/Types/StaticMeshes/
  * @see AStaticMeshActor, UStaticMeshComponent
  */
 UCLASS(collapsecategories, hidecategories=Object, customconstructor, MinimalAPI, BlueprintType, config=Engine)
diff --git a/Engine/Source/Runtime/Engine/Classes/Engine/StaticMeshActor.h b/Engine/Source/Runtime/Engine/Classes/Engine/StaticMeshActor.h
--- a/Engine/Source/Runtime/Engine/Classes/Engine/StaticMeshActor.h
+++ b/Engine/Source/Runtime/Engine/Classes/Engine/StaticMeshActor.h
@@ -8,7 +8,7 @@
  * Static meshes are geometry that do not animate or otherwise deform, and are more efficient to render than other types of geometry.
  * Static meshes dragged into the level from the Content Browser are automatically converted to StaticMeshActors.
  *
- * @see https://docs.unrealengine.com/latest/INT/Engine/Actors/StaticMeshActor/index.html
+ * @see https://docs.unrealengine.com/latest/INT/Engine/Actors/StaticMeshActor/
  * @see UStaticMesh
  */
 UCLASS(hidecategories=(Input), showcategories=("Input|MouseInput", "Input|TouchInput"), ConversionRoot, meta=(ChildCanTick))
diff --git a/Engine/Source/Runtime/Engine/Classes/GameFramework/Actor.h b/Engine/Source/Runtime/Engine/Classes/GameFramework/Actor.h
--- a/Engine/Source/Runtime/Engine/Classes/GameFramework/Actor.h
+++ b/Engine/Source/Runtime/Engine/Classes/GameFramework/Actor.h
@@ -46,11 +46,12 @@ DECLARE_DELEGATE_RetVal_ThreeParams(bool, FOnProcessEvent, AActor*, UFunction*,
 DECLARE_CYCLE_STAT_EXTERN(TEXT("GetComponentsTime"),STAT_GetComponentsTime,STATGROUP_Engine,ENGINE_API);
 
 /**
- * Actor is the base class for a UObject that can be placed or spawned in a level.
- * Actors may contain a collection of UActorComponents, which can be used to control how actors move, how they are rendered, etc.
+ * Actor is the base class for an Object that can be placed or spawned in a level.
+ * Actors may contain a collection of ActorComponents, which can be used to control how actors move, how they are rendered, etc.
  * The other main function of an Actor is the replication of properties and function calls across the network during play.
  * 
- * @see https://docs.unrealengine.com/latest/INT/Programming/UnrealArchitecture/Actors/index.html
+ * @see https://docs.unrealengine.com/latest/INT/Programming/UnrealArchitecture/Actors/
+ * @see UActorComponent
  */
 UCLASS(abstract, BlueprintType, Blueprintable, config=Engine)
 class ENGINE_API AActor : public UObject
@@ -72,7 +73,7 @@ public:
 	/**
 	 * Primary Actor tick function, which calls TickActor().
 	 * Tick functions can be configured to control whether ticking is enabled, at what time during a frame the update occurs, and to set up tick dependencies.
-	 * @see https://docs.unrealengine.com/latest/INT/API/Runtime/Engine/Engine/FTickFunction/index.html
+	 * @see https://docs.unrealengine.com/latest/INT/API/Runtime/Engine/Engine/FTickFunction/
 	 * @see AddTickPrerequisiteActor(), AddTickPrerequisiteComponent()
 	 */
 	UPROPERTY()
@@ -114,7 +115,7 @@ public:
 	 * If true, replicate movement/location related properties.
 	 * Actor must also be set to replicate.
 	 * @see SetReplicates()
-	 * @see https://docs.unrealengine.com/latest/INT/Gameplay/Networking/Replication/index.html
+	 * @see https://docs.unrealengine.com/latest/INT/Gameplay/Networking/Replication/
 	 */
 	UPROPERTY(Replicated, Category=Replication, EditDefaultsOnly)
 	uint32 bReplicateMovement:1;    
@@ -202,7 +203,7 @@ public:
 	 * Properties flagged for replication will update on clients if they change on the server.
 	 * Internally changes the RemoteRole property and handles the cases where the actor needs to be added to the network actor list.
 	 * @param bInReplicates Whether this Actor replicates to network clients.
-	 * @see https://docs.unrealengine.com/latest/INT/Gameplay/Networking/Replication/index.html
+	 * @see https://docs.unrealengine.com/latest/INT/Gameplay/Networking/Replication/
 	 */
 	UFUNCTION(BlueprintCallable, Category = "Replication")
 	void SetReplicates(bool bInReplicates);
diff --git a/Engine/Source/Runtime/Engine/Classes/GameFramework/Character.h b/Engine/Source/Runtime/Engine/Classes/GameFramework/Character.h
--- a/Engine/Source/Runtime/Engine/Classes/GameFramework/Character.h
+++ b/Engine/Source/Runtime/Engine/Classes/GameFramework/Character.h
@@ -188,9 +188,9 @@ typedef FBasedMovementInfo FRepRelativeMovement;
 
 
 /**
- * Characters are Pawns that have a mesh, collision, and default movement logic. They are responsible for all
- * physical interaction between the player or AI and the world, and also have implementations of basic networking and input models.
- * They are designed for a vertically-oriented player representation that can walk, run, jump, fly, and swim through the world using CharacterMovementComponent.
+ * Characters are Pawns that have a mesh, collision, and built-in movement logic.
+ * They are responsible for all physical interaction between the player or AI and the world, and also implement basic networking and input models.
+ * They are designed for a vertically-oriented player representation that can walk, jump, fly, and swim through the world using CharacterMovementComponent.
  *
  * @see APawn, UCharacterMovementComponent
  * @see https://docs.unrealengine.com/latest/INT/Gameplay/Framework/Pawn/Character/
@@ -210,7 +210,7 @@ class ENGINE_API ACharacter : public APawn
 	TSubobjectPtr<class UArrowComponent> ArrowComponent;
 #endif
 
-	/** CharacterMovement component used by walking/running/flying avatars not using rigid body physics */
+	/** Movement component used for movement logic in various movement modes (walking, falling, etc), containing relevant settings and functions to control movement. */
 	UPROPERTY(Category=Character, VisibleAnywhere, BlueprintReadOnly)
 	TSubobjectPtr<class UCharacterMovementComponent> CharacterMovement;
 
@@ -227,7 +227,7 @@ class ENGINE_API ACharacter : public APawn
 	/** Name of the CapsuleComponent. */
 	static FName CapsuleComponentName;
 
-	/** Sets the MovementBase used by CharacterMovement walking movement. */
+	/** Sets the component the Character is walking on, used by CharacterMovement walking movement to be able to follow dynamic objects. */
 	virtual void SetBase(UPrimitiveComponent* NewBase, const FName BoneName = NAME_None, bool bNotifyActor=true);
 	 
 protected:
diff --git a/Engine/Source/Runtime/Engine/Classes/GameFramework/CharacterMovementComponent.h b/Engine/Source/Runtime/Engine/Classes/GameFramework/CharacterMovementComponent.h
--- a/Engine/Source/Runtime/Engine/Classes/GameFramework/CharacterMovementComponent.h
+++ b/Engine/Source/Runtime/Engine/Classes/GameFramework/CharacterMovementComponent.h
@@ -163,14 +163,16 @@ private:
 public:
 	
 	/**
-	 * Actor's current movement mode (walking, falling, etc). This is automatically replicated through the Character owner and for client-server movement functions.
+	 * Actor's current movement mode (walking, falling, etc).
+	 * This is automatically replicated through the Character owner and for client-server movement functions.
 	 * @see SetMovementMode()
 	 */
 	UPROPERTY(Category=MovementMode, BlueprintReadOnly)
 	TEnumAsByte<enum EMovementMode> MovementMode;
 
 	/**
-	 * Current custom sub-mode if MovementMode is set to Custom. This is automatically replicated through the Character owner and for client-server movement functions.
+	 * Current custom sub-mode if MovementMode is set to Custom.
+	 * This is automatically replicated through the Character owner and for client-server movement functions.
 	 * @see SetMovementMode()
 	 */
 	UPROPERTY(Category=MovementMode, BlueprintReadOnly)
@@ -187,7 +189,8 @@ public:
 	float GravityScale;
 
 	/**
-	 * Coefficient of friction. This property allows you to control how much friction is applied when moving across the ground, applying an opposing force that scales with current velocity.
+	 * Coefficient of friction.
+	 * This property allows you to control how much friction is applied when moving across the ground, applying an opposing force that scales with current velocity.
 	 * This can be used to simulate slippery surfaces such as ice or oil by changing the value (possibly based on the material pawn is standing on).
 	 * @see BrakingDecelerationWalking
 	 */
diff --git a/Engine/Source/Runtime/Engine/Classes/GameFramework/Pawn.h b/Engine/Source/Runtime/Engine/Classes/GameFramework/Pawn.h
--- a/Engine/Source/Runtime/Engine/Classes/GameFramework/Pawn.h
+++ b/Engine/Source/Runtime/Engine/Classes/GameFramework/Pawn.h
@@ -22,8 +22,8 @@ class UPrimitiveComponent;
 
 /** 
  * Pawn is the base class of all actors that can be possessed by players or AI.
- * Pawns are the physical representations of players and creatures in a level.
- * They are responsible for all physical interaction between the player or AI and the world.
+ * They are the physical representations of players and creatures in a level.
+ *
  * @see https://docs.unrealengine.com/latest/INT/Gameplay/Framework/Pawn/
  */
 UCLASS(abstract, config=Game, BlueprintType, Blueprintable, hidecategories=(Navigation, "AI|Navigation"))
diff --git a/Engine/Source/Runtime/Engine/Classes/GameFramework/PlayerController.h b/Engine/Source/Runtime/Engine/Classes/GameFramework/PlayerController.h
--- a/Engine/Source/Runtime/Engine/Classes/GameFramework/PlayerController.h
+++ b/Engine/Source/Runtime/Engine/Classes/GameFramework/PlayerController.h
@@ -63,7 +63,7 @@ struct FDynamicForceFeedbackDetails
  * orientation of the controlled Pawn.
  *
  * In networked games, PlayerControllers exist on the server for every player-controlled pawn,
- * and also on the controlling client's machine. They do *not* exist on a client's
+ * and also on the controlling client's machine. They do NOT exist on a client's
  * machine for pawns controlled by remote players elsewhere on the network.
  *
  * @see https://docs.unrealengine.com/latest/INT/Gameplay/Framework/Controller/PlayerController/
