Change retargetting logic to deal with AnimSequences that are referenced by other anim sequences

#TTP 344172 - ANIM: Fix crash of Retargeting additive animation

diff --git a/Engine/Source/Editor/UnrealEd/Private/EditorAnimUtils.cpp b/Engine/Source/Editor/UnrealEd/Private/EditorAnimUtils.cpp
--- a/Engine/Source/Editor/UnrealEd/Private/EditorAnimUtils.cpp
+++ b/Engine/Source/Editor/UnrealEd/Private/EditorAnimUtils.cpp
@@ -74,6 +74,13 @@ namespace EditorAnimUtils
 			{
 				GetAllAnimationSequencesReferredInBlueprint( (*Iter), ComplexAnimsToRetarget, AnimSequencesToRetarget);
 			}
+
+			int SequenceIndex = 0;
+			while (SequenceIndex < AnimSequencesToRetarget.Num())
+			{
+				UAnimSequence* Seq = AnimSequencesToRetarget[SequenceIndex++];
+				Seq->GetAllAnimationSequencesReferred(AnimSequencesToRetarget);
+			}
 		}
 	}
 
@@ -175,6 +182,7 @@ namespace EditorAnimUtils
 				NewNameMapping->AddName(Curve.LastObservedName, Curve.CurveUid);
 			}
 
+			AssetToRetarget->ReplaceReferredAnimations(DuplicatedSequences);
 			AssetToRetarget->ReplaceSkeleton(NewSkeleton, bConvertAnimationDataInComponentSpaces);
 		}
 
diff --git a/Engine/Source/Runtime/Engine/Classes/Animation/AnimSequence.h b/Engine/Source/Runtime/Engine/Classes/Animation/AnimSequence.h
--- a/Engine/Source/Runtime/Engine/Classes/Animation/AnimSequence.h
+++ b/Engine/Source/Runtime/Engine/Classes/Animation/AnimSequence.h
@@ -504,6 +504,7 @@ class UAnimSequence : public UAnimSequenceBase
 	ENGINE_API virtual bool IsValidAdditive() const override;
 #if WITH_EDITOR
 	ENGINE_API virtual bool GetAllAnimationSequencesReferred(TArray<UAnimSequence*>& AnimationSequences) override;
+	ENGINE_API virtual void ReplaceReferredAnimations(const TMap<UAnimSequence*, UAnimSequence*>& ReplacementMap) override;
 	ENGINE_API virtual int32 GetNumberOfFrames() override { return NumFrames; }
 #endif
 	// End of UAnimationAsset interface
diff --git a/Engine/Source/Runtime/Engine/Private/Animation/AnimSequence.cpp b/Engine/Source/Runtime/Engine/Private/Animation/AnimSequence.cpp
--- a/Engine/Source/Runtime/Engine/Private/Animation/AnimSequence.cpp
+++ b/Engine/Source/Runtime/Engine/Private/Animation/AnimSequence.cpp
@@ -2258,9 +2258,23 @@ void UAnimSequence::RemoveTrack(int32 TrackIndex)
 
 bool UAnimSequence::GetAllAnimationSequencesReferred(TArray<UAnimSequence*>& AnimationSequences)
 {
-	AnimationSequences.Add(this);
+	if (RefPoseSeq)
+	{
+		AnimationSequences.AddUnique(RefPoseSeq);
+	}
+	return AnimationSequences.Num() > 0;
+}
 
-	return true;
+void UAnimSequence::ReplaceReferredAnimations(const TMap<UAnimSequence*, UAnimSequence*>& ReplacementMap)
+{
+	if (RefPoseSeq)
+	{
+		UAnimSequence* const* ReplacementAsset = (UAnimSequence*const*)ReplacementMap.Find(RefPoseSeq);
+		if (ReplacementAsset)
+		{
+			RefPoseSeq = *ReplacementAsset;
+		}
+	}
 }
 
 bool UAnimSequence::AddLoopingInterpolation()
diff --git a/Engine/Source/Runtime/Engine/Private/AnimationAsset.cpp b/Engine/Source/Runtime/Engine/Private/AnimationAsset.cpp
--- a/Engine/Source/Runtime/Engine/Private/AnimationAsset.cpp
+++ b/Engine/Source/Runtime/Engine/Private/AnimationAsset.cpp
@@ -61,6 +61,10 @@ bool UAnimationAsset::ReplaceSkeleton(USkeleton* NewSkeleton, bool bConvertSpace
 	{
 		// get all sequences that need to change
 		TArray<UAnimSequence*> AnimSeqsToReplace;
+		if (UAnimSequence* AnimSequence = Cast<UAnimSequence>(this))
+		{
+			AnimSeqsToReplace.AddUnique(AnimSequence);
+		}
 		if (GetAllAnimationSequencesReferred(AnimSeqsToReplace))
 		{
 			for (auto Iter = AnimSeqsToReplace.CreateIterator(); Iter; ++Iter)
