Add call to Actor EndPlay for quit (in standalone game)
Move logic for when to dispatch EndPlay in to a non-virtual RouteEndPlay function
Only call EndPlay if the Actor has been initialized AND the world has begun

diff --git a/Engine/Source/Editor/UnrealEd/Private/PlayLevel.cpp b/Engine/Source/Editor/UnrealEd/Private/PlayLevel.cpp
--- a/Engine/Source/Editor/UnrealEd/Private/PlayLevel.cpp
+++ b/Engine/Source/Editor/UnrealEd/Private/PlayLevel.cpp
@@ -445,7 +445,7 @@ void UEditorEngine::TeardownPlaySession(FWorldContext &PieWorldContext)
 	// Go through and let all the PlayWorld Actor's know they are being destroyed
 	for (FActorIterator ActorIt(PlayWorld); ActorIt; ++ActorIt)
 	{
-		ActorIt->EndPlay(EEndPlayReason::EndPlayInEditor);
+		ActorIt->RouteEndPlay(EEndPlayReason::EndPlayInEditor);
 	}
 
 	// Move blueprint debugging pointers back to the objects in the editor world
diff --git a/Engine/Source/Runtime/Engine/Classes/GameFramework/Actor.h b/Engine/Source/Runtime/Engine/Classes/GameFramework/Actor.h
--- a/Engine/Source/Runtime/Engine/Classes/GameFramework/Actor.h
+++ b/Engine/Source/Runtime/Engine/Classes/GameFramework/Actor.h
@@ -1737,7 +1737,10 @@ public:
 	/**	Do anything needed to clear out cross level references; Called from ULevel::PreSave	 */
 	virtual void ClearCrossLevelReferences();
 	
-	/** Called whenever this actor is being removed from a level */
+	/** Non-virtual function to evaluate which portions of the EndPlay process should be dispatched for each actor */
+	void RouteEndPlay(const EEndPlayReason::Type EndPlayReason);
+
+	/** Overridable function called whenever this actor is being removed from a level */
 	virtual void EndPlay(const EEndPlayReason::Type EndPlayReason);
 
 	/** iterates up the Base chain to see whether or not this Actor is based on the given Actor
diff --git a/Engine/Source/Runtime/Engine/Private/Actor.cpp b/Engine/Source/Runtime/Engine/Private/Actor.cpp
--- a/Engine/Source/Runtime/Engine/Private/Actor.cpp
+++ b/Engine/Source/Runtime/Engine/Private/Actor.cpp
@@ -1477,11 +1477,22 @@ void AActor::PrestreamTextures( float Seconds, bool bEnableStreaming, int32 Cine
 
 void AActor::OnRep_Instigator() {}
 
+void AActor::RouteEndPlay(const EEndPlayReason::Type EndPlayReason)
+{
+	if (bActorInitialized)
+	{
+		UninitializeComponents();
+
+		UWorld* World = GetWorld();
+		if (World && World->HasBegunPlay())
+		{
+			EndPlay(EndPlayReason);
+		}
+	}
+}
 
 void AActor::EndPlay(const EEndPlayReason::Type EndPlayReason)
 {
-	UninitializeComponents();
-
 	// Dispatch the blueprint events
 	ReceiveEndPlay(EndPlayReason);
 	OnEndPlay.Broadcast(EndPlayReason);
@@ -1547,10 +1558,8 @@ FTransform AActor::GetTransform() const
 
 void AActor::Destroyed()
 {
-	if (bActorInitialized)
-	{
-		EndPlay(EEndPlayReason::ActorDestroyed);
-	}
+	RouteEndPlay(EEndPlayReason::ActorDestroyed);
+
 	ReceiveDestroyed();
 	OnDestroyed.Broadcast();
 	GetWorld()->RemoveNetworkActor(this);
diff --git a/Engine/Source/Runtime/Engine/Private/UnrealEngine.cpp b/Engine/Source/Runtime/Engine/Private/UnrealEngine.cpp
--- a/Engine/Source/Runtime/Engine/Private/UnrealEngine.cpp
+++ b/Engine/Source/Runtime/Engine/Private/UnrealEngine.cpp
@@ -8382,10 +8382,7 @@ bool UEngine::LoadMap( FWorldContext& WorldContext, FURL URL, class UPendingNetG
 
 		for (FActorIterator ActorIt(WorldContext.World()); ActorIt; ++ActorIt)
 		{
-			if (ActorIt->bActorInitialized)
-			{
-				ActorIt->EndPlay(EEndPlayReason::LevelTransition);
-			}
+			ActorIt->RouteEndPlay(EEndPlayReason::LevelTransition);
 		}
 
 		// Do this after destroying pawns/playercontrollers, in case that spawns new things (e.g. dropped weapons)
