Rename (though backwards compatible via deprecation) AGameMode::SetGameMode to AGameMode::GetGameModeClass
Make GetGameModeClass and GetDefaultGameClassPath const and virtual

diff --git a/Engine/Source/Runtime/Engine/Classes/GameFramework/GameMode.h b/Engine/Source/Runtime/Engine/Classes/GameFramework/GameMode.h
--- a/Engine/Source/Runtime/Engine/Classes/GameFramework/GameMode.h
+++ b/Engine/Source/Runtime/Engine/Classes/GameFramework/GameMode.h
@@ -48,7 +48,7 @@ struct FGameClassShortName
 // C++ UGameEngine::LoadMap().  The class of this GameMode actor is determined by
 // (in order) either the URL ?game=xxx, or the
 // DefaultGameMode entry in the game's .ini file (in the /Script/Engine.Engine section).
-// The GameMode used can be overridden in the GameMode function SetGameMode(), called
+// The GameMode used can be overridden in the GameMode function GetGameModeClass(), called
 // on the game class picked by the above process.
 //
 //=============================================================================
@@ -354,16 +354,20 @@ public:
 	/** 
 	 * @return the full path to the optimal GameMode class to use for the specified map and options
 	 * this is used for preloading cooked packages, etc. and therefore doesn't need to include any fallbacks
-	 * as SetGameMode() will be called later to actually find/load the desired class
+	 * as GetGameModeClass() will be called later to actually find/load the desired class
 	 */
-	FString GetDefaultGameClassPath(const FString& MapName, const FString& Options, const FString& Portal);
+	virtual FString GetDefaultGameClassPath(const FString& MapName, const FString& Options, const FString& Portal) const;
 
 	/** 
 	 * @return the class of GameMode to spawn for the game on the specified map and the specified options
 	 * this function should include any fallbacks in case the desired class can't be found
 	 */
+	virtual TSubclassOf<AGameMode> GetGameModeClass(const FString& MapName, const FString& Options, const FString& Portal) const;
+
+	DEPRECATED(4.5, "AGameMode::SetGameMode renamed AGameMode::GetGameModeClass")
 	TSubclassOf<AGameMode> SetGameMode(const FString& MapName, const FString& Options, const FString& Portal);
 
+
 	/** @return GameSession class to use for this game  */
 	virtual TSubclassOf<class AGameSession> GetGameSessionClass() const;
 
diff --git a/Engine/Source/Runtime/Engine/Private/GameMode.cpp b/Engine/Source/Runtime/Engine/Private/GameMode.cpp
--- a/Engine/Source/Runtime/Engine/Private/GameMode.cpp
+++ b/Engine/Source/Runtime/Engine/Private/GameMode.cpp
@@ -1212,15 +1212,19 @@ int32 AGameMode::GetIntOption( const FString& Options, const FString& ParseStrin
 	return CurrentValue;
 }
 
-FString AGameMode::GetDefaultGameClassPath(const FString& MapName, const FString& Options, const FString& Portal)
+FString AGameMode::GetDefaultGameClassPath(const FString& MapName, const FString& Options, const FString& Portal) const
 {
 	return GetClass()->GetPathName();
 }
 
+TSubclassOf<AGameMode> AGameMode::GetGameModeClass(const FString& MapName, const FString& Options, const FString& Portal) const
+{
+	return GetClass();
+}
 
 TSubclassOf<AGameMode> AGameMode::SetGameMode(const FString& MapName, const FString& Options, const FString& Portal)
 {
-	return GetClass();
+	return GetGameModeClass(MapName, Options, Portal);
 }
 
 TSubclassOf<AGameSession> AGameMode::GetGameSessionClass() const
