Fixing stupid mistake.

diff --git a/Engine/Source/Runtime/Engine/Private/Particles/ParticleSystemRender.cpp b/Engine/Source/Runtime/Engine/Private/Particles/ParticleSystemRender.cpp
--- a/Engine/Source/Runtime/Engine/Private/Particles/ParticleSystemRender.cpp
+++ b/Engine/Source/Runtime/Engine/Private/Particles/ParticleSystemRender.cpp
@@ -6566,7 +6566,7 @@ int32 FDynamicRibbonEmitterData::FillVertexData(struct FAsyncBufferFillData& Dat
 						FinalPos += Data.LocalToWorld.GetOrigin();
 					}
 					Vertex = (FParticleBeamTrailVertex*)(TempVertexData);
-					Vertex->Position = FinalPos + 
+					Vertex->Position = FinalPos;
 					Vertex->OldPosition = FinalPos;
 					Vertex->ParticleId	= 0;
 					Vertex->Size.X = InterpSize;
