#docs - Fix comment on UPrimitiveComponent.bGenerateOverlapEvents to indicate that it must be set on both source and target components for the overlap event to occur.

diff --git a/Engine/Source/Runtime/Engine/Classes/Components/PrimitiveComponent.h b/Engine/Source/Runtime/Engine/Classes/Components/PrimitiveComponent.h
--- a/Engine/Source/Runtime/Engine/Classes/Components/PrimitiveComponent.h
+++ b/Engine/Source/Runtime/Engine/Classes/Components/PrimitiveComponent.h
@@ -149,7 +149,9 @@ public:
 	uint32 bAlwaysCreatePhysicsState:1;
 
 	/**
-	 * If true, this component will generate overlap events when it is overlapping other components (e.g. Begin Overlap)
+	 * If true, this component will generate overlap events when it is overlapping other components (eg Begin Overlap).
+	 * Both components (this and the other) must have this enabled for overlap events to occur.
+	 *
 	 * @see [Overlap Events] https://docs.unrealengine.com/latest/INT/Engine/Physics/Collision/index.html#overlapandgenerateoverlapevents
 	 * @see UpdateOverlaps(), BeginComponentOverlap(), EndComponentOverlap()
 	 */
