Fixed another memory profiler compile error
- MemoryProfiler2 now works again!

diff --git a/Engine/Source/Runtime/Engine/Private/MallocProfilerEx.cpp b/Engine/Source/Runtime/Engine/Private/MallocProfilerEx.cpp
--- a/Engine/Source/Runtime/Engine/Private/MallocProfilerEx.cpp
+++ b/Engine/Source/Runtime/Engine/Private/MallocProfilerEx.cpp
@@ -48,13 +48,13 @@ void FMallocProfilerEx::WriteLoadedLevels( UWorld* InWorld )
 		{
 			ULevelStreaming* LevelStreaming = InWorld->StreamingLevels[LevelIndex];
 			if ((LevelStreaming != NULL)
-				&& (LevelStreaming->PackageName != NAME_None)
-				&& (LevelStreaming->PackageName != InWorld->GetOutermost()->GetFName())
+				&& (LevelStreaming->GetWorldAssetPackageFName() != NAME_None)
+				&& (LevelStreaming->GetWorldAssetPackageFName() != InWorld->GetOutermost()->GetFName())
 				&& (LevelStreaming->GetLoadedLevel() != NULL))
 			{
 				NumLoadedLevels++;
 
-				int32 LevelPackageIndex = GetNameTableIndex(LevelStreaming->PackageName);
+				int32 LevelPackageIndex = GetNameTableIndex(LevelStreaming->GetWorldAssetPackageFName());
 
 				BufferedFileWriter << LevelPackageIndex;
 			}
