Fix welding crash when non-root components are attached to other actors

diff --git a/Engine/Source/Runtime/Engine/Private/PrimitiveComponentPhysics.cpp b/Engine/Source/Runtime/Engine/Private/PrimitiveComponentPhysics.cpp
--- a/Engine/Source/Runtime/Engine/Private/PrimitiveComponentPhysics.cpp
+++ b/Engine/Source/Runtime/Engine/Private/PrimitiveComponentPhysics.cpp
@@ -615,7 +615,7 @@ void UPrimitiveComponent::WeldTo(USceneComponent* InParent, FName InSocketName /
 void UPrimitiveComponent::UnWeldFromParent()
 {
 	FBodyInstance * NewRootBI = GetBodyInstance(NAME_None, false);
-	if (NewRootBI == NULL || NewRootBI->bWelded == false)
+	if (NewRootBI == NULL || NewRootBI->bWelded == false || IsPendingKill())
 	{
 		return;
 	}
