add missing eMTD in BodyInstance sweep

diff --git a/Engine/Source/Runtime/Engine/Private/PhysicsEngine/BodyInstance.cpp b/Engine/Source/Runtime/Engine/Private/PhysicsEngine/BodyInstance.cpp
--- a/Engine/Source/Runtime/Engine/Private/PhysicsEngine/BodyInstance.cpp
+++ b/Engine/Source/Runtime/Engine/Private/PhysicsEngine/BodyInstance.cpp
@@ -3374,7 +3374,7 @@ bool FBodyInstance::InternalSweepPhysX(struct FHitResult& OutHit, const FVector&
 	const float DeltaMag = Delta.Size();
 	if(DeltaMag > KINDA_SMALL_NUMBER)
 	{
-		PxSceneQueryFlags POutputFlags = PxSceneQueryFlag::eIMPACT | PxSceneQueryFlag::eNORMAL | PxSceneQueryFlag::eDISTANCE;
+		PxSceneQueryFlags POutputFlags = PxSceneQueryFlag::eIMPACT | PxSceneQueryFlag::eNORMAL | PxSceneQueryFlag::eDISTANCE | PxSceneQueryFlag::eMTD;
 		
 		PxQuat PGeomRot = PxQuat::createIdentity();
 
