Missed files from CL# 2302887

diff --git a/Engine/Source/Runtime/Engine/Classes/Engine/EngineTypes.h b/Engine/Source/Runtime/Engine/Classes/Engine/EngineTypes.h
--- a/Engine/Source/Runtime/Engine/Classes/Engine/EngineTypes.h
+++ b/Engine/Source/Runtime/Engine/Classes/Engine/EngineTypes.h
@@ -2162,6 +2162,7 @@ namespace EEndPlayReason
 		LevelTransition,	// When the world is being unloaded for a level transition
 		EndPlayInEditor,	// When the world is being unloaded because PIE is ending
 		RemovedFromWorld,	// When the level it is a member of is streamed out
+		Quit,			// When the application is being exited
 	};
 
 }
diff --git a/Engine/Source/Runtime/Engine/Private/GameEngine.cpp b/Engine/Source/Runtime/Engine/Private/GameEngine.cpp
--- a/Engine/Source/Runtime/Engine/Private/GameEngine.cpp
+++ b/Engine/Source/Runtime/Engine/Private/GameEngine.cpp
@@ -650,6 +650,11 @@ bool UGameEngine::HandleExitCommand( const TCHAR* Cmd, FOutputDevice& Ar )
 
 		// Shut down any existing game connections
 		ShutdownWorldNetDriver(World);
+
+		for (FActorIterator ActorIt(World); ActorIt; ++ActorIt)
+		{
+			ActorIt->RouteEndPlay(EEndPlayReason::Quit);
+		}
 	}
 
 	Ar.Log( TEXT("Closing by request") );
diff --git a/Engine/Source/Runtime/Engine/Private/World.cpp b/Engine/Source/Runtime/Engine/Private/World.cpp
--- a/Engine/Source/Runtime/Engine/Private/World.cpp
+++ b/Engine/Source/Runtime/Engine/Private/World.cpp
@@ -1833,7 +1833,8 @@ void UWorld::RemoveFromWorld( ULevel* Level )
 			AActor* Actor = Level->Actors[ActorIdx];
 			if (Actor != NULL)
 			{
-				Actor->EndPlay(EEndPlayReason::RemovedFromWorld);
+				Actor->RouteEndPlay(EEndPlayReason::RemovedFromWorld);
+
 				if (NetDriver)
 				{
 					NetDriver->NotifyActorLevelUnloaded(Actor);
@@ -4405,10 +4406,7 @@ UWorld* FSeamlessTravelHandler::Tick()
 						UE_LOG(LogWorld, Warning, TEXT("Actor '%s' was indicated to be kept but exists in level '%s', not the persistent level.  Actor will not travel."), *TheActor->GetName(), *TheActor->GetLevel()->GetOutermost()->GetName());
 					}
 
-					if (TheActor->bActorInitialized)
-					{
-						TheActor->EndPlay(EEndPlayReason::LevelTransition);
-					}
+					TheActor->RouteEndPlay(EEndPlayReason::LevelTransition);
 
 					// otherwise, set to be deleted
 					KeepAnnotation.Clear(TheActor);
