Viewport changes to suport new input mode API on playercontroller.

diff --git a/Engine/Source/Runtime/Engine/Classes/Engine/GameViewportClient.h b/Engine/Source/Runtime/Engine/Classes/Engine/GameViewportClient.h
--- a/Engine/Source/Runtime/Engine/Classes/Engine/GameViewportClient.h
+++ b/Engine/Source/Runtime/Engine/Classes/Engine/GameViewportClient.h
@@ -248,6 +248,8 @@ public:
 
 	class FSceneViewport* GetGameViewport();
 
+	TSharedPtr<class SViewport> GetGameViewportWidget();
+
 	UGameInstance* GetGameInstance() const;
 
 	virtual void Init(struct FWorldContext& WorldContext, UGameInstance* OwningGameInstance);
@@ -635,6 +637,38 @@ public:
 		SoundShowFlags = InSoundShowFlags;
 	}
 
+	/**
+	 * Set whether to ignore input.
+	 */
+	void SetIgnoreInput(bool Ignore)
+	{
+		bIgnoreInput = Ignore;
+	}
+
+	/**
+	 * Check whether we should ignore input.
+	 */
+	virtual bool IgnoreInput() override
+	{
+		return bIgnoreInput;
+	}
+
+	/**
+	 * Set the mouse capture behavior when the viewport is clicked
+	 */
+	void SetCaptureMouseOnClick(EMouseCaptureMode::Type Mode)
+	{
+		MouseCaptureMode = Mode;
+	}
+
+	/**
+	 * Gets the mouse capture behavior when the viewport is clicked
+	 */
+	virtual EMouseCaptureMode::Type CaptureMouseOnClick() override
+	{
+		return MouseCaptureMode;
+	}
+
 private:
 	/**
 	 * Set a specific stat to either enabled or disabled (returns the number of remaining enabled stats)
@@ -714,6 +748,12 @@ private:
 
 	/** Disables splitscreen, useful when game code is in menus, and doesn't want splitscreen on */
 	bool bDisableSplitScreenOverride;
+
+	/** Whether or not to ignore input */
+	bool bIgnoreInput;
+
+	/** Mouse capture behavior when the viewport is clicked */
+	EMouseCaptureMode::Type MouseCaptureMode;
 };
 
 
diff --git a/Engine/Source/Runtime/Engine/Classes/GameFramework/PlayerController.h b/Engine/Source/Runtime/Engine/Classes/GameFramework/PlayerController.h
--- a/Engine/Source/Runtime/Engine/Classes/GameFramework/PlayerController.h
+++ b/Engine/Source/Runtime/Engine/Classes/GameFramework/PlayerController.h
@@ -12,8 +12,11 @@
 #include "GameFramework/Controller.h"
 #include "Engine/LatentActionManager.h"
 #include "GenericPlatform/IForceFeedbackSystem.h"
+#include "SlateCore.h"
+
 #include "PlayerController.generated.h"
 
+
 class FPrimitiveComponentId;
 
 /** Default delegate that provides an implementation for those that don't have special needs other than a toggle */
@@ -24,6 +27,17 @@ DECLARE_DELEGATE_ThreeParams(FGetAudioListenerPos, FVector& /*Location*/, FVecto
 
 DECLARE_LOG_CATEGORY_EXTERN(LogPlayerController, Log, All);
 
+UENUM()
+namespace EInputMode
+{
+	enum Type
+	{
+		UIOnly,
+		UIAndGame,
+		GameOnly,
+	};
+}
+
 UENUM()
 namespace EDynamicForceFeedbackAction
 {
@@ -956,6 +970,11 @@ public:
 	UFUNCTION(BlueprintCallable, Category = "Game|Player")
 	virtual void SetVirtualJoystickVisibility(bool bVisible);
 
+public:
+	/** Set the Input Mode. */
+	UFUNCTION(BlueprintCallable, Category = "Input")
+	virtual void SetInputMode(EInputMode::Type Mode);
+
 public:
 	/**
 	 * Change Camera mode
@@ -1399,4 +1418,10 @@ public:
 	/** The value of SeamlessTravelCount, upon the last call to GameMode::HandleSeamlessTravelPlayer; used to detect seamless travel */
 	UPROPERTY()
 	uint16		LastCompletedSeamlessTravelCount;
+
+	/** Stores the last input mode set */
+	//EInputMode  InputMode; right now we don't save this as it could be a half truth.
+
+	/** FReply used to defer some slate operations. */
+	FReply		SlateOperations;
 };
diff --git a/Engine/Source/Runtime/Engine/Private/GameViewportClient.cpp b/Engine/Source/Runtime/Engine/Private/GameViewportClient.cpp
--- a/Engine/Source/Runtime/Engine/Private/GameViewportClient.cpp
+++ b/Engine/Source/Runtime/Engine/Private/GameViewportClient.cpp
@@ -82,6 +82,8 @@ UGameViewportClient::UGameViewportClient(const class FPostConstructInitializePro
 	, EngineShowFlags(ESFIM_Game)
 	, CurrentBufferVisualizationMode(NAME_None)
 	, HighResScreenshotDialog(NULL)
+	, bIgnoreInput(false)
+	, MouseCaptureMode(EMouseCaptureMode::CapturePermanently)
 {
 
 	TitleSafeZone.MaxPercentX = 0.9f;
@@ -175,6 +177,17 @@ FSceneViewport* UGameViewportClient::GetGameViewport()
 	return static_cast<FSceneViewport*>(Viewport);
 }
 
+TSharedPtr<class SViewport> UGameViewportClient::GetGameViewportWidget()
+{
+	FSceneViewport* SceneViewport = GetGameViewport();
+	if (SceneViewport != nullptr)
+	{
+		TWeakPtr<SViewport> WeakViewportWidget = SceneViewport->GetViewportWidget();
+		TSharedPtr<SViewport> ViewportWidget = WeakViewportWidget.Pin();
+		return ViewportWidget;
+	}
+	return nullptr;
+}
 
 void UGameViewportClient::Tick( float DeltaTime )
 {
@@ -209,6 +222,11 @@ UGameInstance* UGameViewportClient::GetGameInstance() const
 
 bool UGameViewportClient::InputKey(FViewport* InViewport, int32 ControllerId, FKey Key, EInputEvent EventType, float AmountDepressed, bool bGamepad)
 {
+	if (IgnoreInput())
+	{
+		return false;
+	}
+
 	if (InViewport->IsPlayInEditorViewport() && Key.IsGamepadKey())
 	{
 		GEngine->RemapGamepadControllerIdForPIE(this, ControllerId);
@@ -248,6 +266,11 @@ bool UGameViewportClient::InputKey(FViewport* InViewport, int32 ControllerId, FK
 
 bool UGameViewportClient::InputAxis(FViewport* InViewport, int32 ControllerId, FKey Key, float Delta, float DeltaTime, int32 NumSamples, bool bGamepad)
 {
+	if (IgnoreInput())
+	{
+		return false;
+	}
+
 	bool bResult = false;
 
 	if (InViewport->IsPlayInEditorViewport() && Key.IsGamepadKey())
@@ -298,6 +321,11 @@ bool UGameViewportClient::InputAxis(FViewport* InViewport, int32 ControllerId, F
 
 bool UGameViewportClient::InputChar(FViewport* InViewport, int32 ControllerId, TCHAR Character)
 {
+	if (IgnoreInput())
+	{
+		return false;
+	}
+
 	// should probably just add a ctor to FString that takes a TCHAR
 	FString CharacterString;
 	CharacterString += Character;
@@ -316,6 +344,11 @@ bool UGameViewportClient::InputChar(FViewport* InViewport, int32 ControllerId, T
 
 bool UGameViewportClient::InputTouch(FViewport* InViewport, int32 ControllerId, uint32 Handle, ETouchType::Type Type, const FVector2D& TouchLocation, FDateTime DeviceTimestamp, uint32 TouchpadIndex)
 {
+	if (IgnoreInput())
+	{
+		return false;
+	}
+
 	// route to subsystems that care
 	bool bResult = (ViewportConsole ? ViewportConsole->InputTouch(ControllerId, Handle, Type, TouchLocation, DeviceTimestamp, TouchpadIndex) : false);
 	if (!bResult)
@@ -332,6 +365,11 @@ bool UGameViewportClient::InputTouch(FViewport* InViewport, int32 ControllerId,
 
 bool UGameViewportClient::InputMotion(FViewport* InViewport, int32 ControllerId, const FVector& Tilt, const FVector& RotationRate, const FVector& Gravity, const FVector& Acceleration)
 {
+	if (IgnoreInput())
+	{
+		return false;
+	}
+
 	// route to subsystems that care
 	bool bResult = false;
 
diff --git a/Engine/Source/Runtime/Engine/Private/PlayerController.cpp b/Engine/Source/Runtime/Engine/Private/PlayerController.cpp
--- a/Engine/Source/Runtime/Engine/Private/PlayerController.cpp
+++ b/Engine/Source/Runtime/Engine/Private/PlayerController.cpp
@@ -39,6 +39,7 @@ const float RetryServerCheckSpectatorThrottleTime = 0.25f;
 
 APlayerController::APlayerController(const class FPostConstructInitializeProperties& PCIP)
 	: Super(PCIP)
+	, SlateOperations(FReply::Unhandled())
 {
 	NetPriority = 3.0f;
 	CheatClass = UCheatManager::StaticClass();
@@ -4245,6 +4246,44 @@ void APlayerController::SetVirtualJoystickVisibility(bool bVisible)
 	}
 }
 
+void APlayerController::SetInputMode(EInputMode::Type Mode)
+{
+	UGameViewportClient* GameViewportClient = GetWorld()->GetGameViewport();
+	if (GameViewportClient != nullptr)
+	{
+		TSharedPtr<SViewport> ViewportWidget = GameViewportClient->GetGameViewportWidget();
+		if (ViewportWidget.IsValid())
+		{
+			switch (Mode)
+			{
+			case EInputMode::UIOnly:
+				SlateOperations.ReleaseMouseCapture();
+				SlateOperations.ReleaseJoystickCapture();
+				SlateOperations.ReleaseMouseLock();
+				GameViewportClient->SetIgnoreInput(true);
+				GameViewportClient->SetCaptureMouseOnClick(EMouseCaptureMode::NoCapture);
+				break;
+			case EInputMode::UIAndGame:
+				SlateOperations.ReleaseMouseCapture();
+				SlateOperations.ReleaseJoystickCapture();
+				SlateOperations.ReleaseMouseLock();
+				GameViewportClient->SetIgnoreInput(false);
+				GameViewportClient->SetCaptureMouseOnClick(EMouseCaptureMode::CaptureDuringMouseDown);
+				break;
+			case EInputMode::GameOnly:
+				TSharedRef<SViewport> pViewportWidgetRef = ViewportWidget.ToSharedRef();
+				SlateOperations.CaptureMouse(pViewportWidgetRef);
+				SlateOperations.CaptureJoystick(pViewportWidgetRef);
+				SlateOperations.LockMouseToWidget(pViewportWidgetRef);
+				SlateOperations.SetKeyboardFocus(pViewportWidgetRef, EKeyboardFocusCause::SetDirectly);
+				GameViewportClient->SetIgnoreInput(false);
+				GameViewportClient->SetCaptureMouseOnClick(EMouseCaptureMode::CapturePermanently);
+				break;
+			}
+		}
+	}
+}
+
 void APlayerController::UpdateCameraManager(float DeltaSeconds)
 {
 	if (PlayerCameraManager != NULL)
diff --git a/Engine/Source/Runtime/Engine/Private/Slate/SceneViewport.cpp b/Engine/Source/Runtime/Engine/Private/Slate/SceneViewport.cpp
--- a/Engine/Source/Runtime/Engine/Private/Slate/SceneViewport.cpp
+++ b/Engine/Source/Runtime/Engine/Private/Slate/SceneViewport.cpp
@@ -351,15 +351,34 @@ FReply FSceneViewport::OnMouseButtonDown( const FGeometry& InGeometry, const FPo
 		{
 			CurrentReplyState = FReply::Unhandled(); 
 		}
-	}
 
-	// Mouse down should focus viewport for keyboard input
-	CurrentReplyState.SetKeyboardFocus( ViewportWidget.Pin().ToSharedRef(), EKeyboardFocusCause::Mouse );
-	CurrentReplyState.UseHighPrecisionMouseMovement( ViewportWidget.Pin().ToSharedRef() );
+		if (ViewportClient->CaptureMouseOnClick() != EMouseCaptureMode::NoCapture && !ViewportClient->IgnoreInput())
+		{
+			TSharedRef<SViewport> ViewportWidgetRef = ViewportWidget.Pin().ToSharedRef();
+
+			// Mouse down should focus viewport for keyboard input
+			CurrentReplyState.SetKeyboardFocus(ViewportWidgetRef, EKeyboardFocusCause::Mouse);
+			
+			UWorld* World = ViewportClient->GetWorld();
+			if (World && World->IsGameWorld() && World->GetFirstPlayerController())
+			{
+				CurrentReplyState.CaptureMouse(ViewportWidgetRef);
+				CurrentReplyState.LockMouseToWidget(ViewportWidgetRef);
+				if (!World->GetFirstPlayerController()->ShouldShowMouseCursor())
+				{
+					CurrentReplyState.UseHighPrecisionMouseMovement(ViewportWidgetRef);
+				}
+			}
+			else
+			{
+				CurrentReplyState.UseHighPrecisionMouseMovement(ViewportWidgetRef);
+			}
+		}
+	}
 	
 	// Re-set prevent throttling here as it can get reset when inside of InputKey()
 	CurrentReplyState.PreventThrottling();
-	
+
 	return CurrentReplyState;
 }
 
@@ -375,22 +394,28 @@ FReply FSceneViewport::OnMouseButtonUp( const FGeometry& InGeometry, const FPoin
 
 	// Switch to the viewport clients world before processing input
 	FScopedConditionalWorldSwitcher WorldSwitcher( ViewportClient );
-	bool bIsCursorForcedVisible = true;
+	bool bIsCursorVisible = true;
+	bool bReleaseMouse = true;
 	if( ViewportClient && GetSizeXY() != FIntPoint::ZeroValue  )
 	{
 		if( !ViewportClient->InputKey(this,0,InMouseEvent.GetEffectingButton(), IE_Released ) )
 		{
 			CurrentReplyState = FReply::Unhandled(); 
 		}
-		bIsCursorForcedVisible = ViewportClient->GetCursor( this, GetMouseX(), GetMouseY() ) != EMouseCursor::None;
+		bIsCursorVisible = ViewportClient->GetCursor(this, GetMouseX(), GetMouseY()) != EMouseCursor::None;
+		bReleaseMouse = bIsCursorVisible || ViewportClient->CaptureMouseOnClick() == EMouseCaptureMode::CaptureDuringMouseDown;
 	}
-	if( !( ( FApp::IsGame() && !GIsEditor ) || bIsPlayInEditorViewport ) || bIsCursorForcedVisible )
+	if (!((FApp::IsGame() && !GIsEditor) || bIsPlayInEditorViewport) || bReleaseMouse)
 	{
 		// On mouse up outside of the game (editor viewport) or if the cursor is visible in game, we should make sure the mouse is no longer captured
 		// as long as the left or right mouse buttons are not still down
 		if( !InMouseEvent.IsMouseButtonDown( EKeys::RightMouseButton ) && !InMouseEvent.IsMouseButtonDown( EKeys::LeftMouseButton ))
 		{
 			CurrentReplyState.ReleaseMouseCapture();
+			if (bIsCursorVisible)
+			{
+				CurrentReplyState.ReleaseMouseLock();
+			}
 		}
 	}
 	return CurrentReplyState;
diff --git a/Engine/Source/Runtime/Engine/Public/UnrealClient.h b/Engine/Source/Runtime/Engine/Public/UnrealClient.h
--- a/Engine/Source/Runtime/Engine/Public/UnrealClient.h
+++ b/Engine/Source/Runtime/Engine/Public/UnrealClient.h
@@ -623,6 +623,19 @@ extern ENGINE_API bool IsAltDown(FViewport* Viewport);
 extern ENGINE_API bool GetViewportScreenShot(FViewport* Viewport, TArray<FColor>& Bitmap, const FIntRect& ViewRect = FIntRect());
 extern ENGINE_API bool GetHighResScreenShotInput(const TCHAR* Cmd, FOutputDevice& Ar, uint32& OutXRes, uint32& OutYRes, float& OutResMult, FIntRect& OutCaptureRegion, bool& OutShouldEnableMask);
 
+namespace EMouseCaptureMode
+{
+	enum Type
+	{
+		/** Do not capture the mouse at all */
+		NoCapture,
+		/** Capture the mouse permanently when the viewport is clicked */
+		CapturePermanently,
+		/** Capture the mouse during a mouse down, releases on mouse up */
+		CaptureDuringMouseDown
+	};
+}
+
 /**
  * An abstract interface to a viewport's client.
  * The viewport's client processes input received by the viewport, and draws the viewport.
@@ -854,6 +867,16 @@ public:
 	 * Set the sound stat flags enabled for this viewport
 	 */
 	virtual void SetSoundShowFlags(const ESoundShowFlags::Type InSoundShowFlags) {}
+
+	/**
+	 * Check whether we should ignore input.
+	 */
+	virtual bool IgnoreInput() { return false; }
+
+	/**
+	 * Gets the mouse capture behavior when the viewport is clicked
+	 */
+	virtual EMouseCaptureMode::Type CaptureMouseOnClick() { return EMouseCaptureMode::CapturePermanently; }
 };
 
 /** Tracks the viewport client that should process the stat command, can be NULL */
diff --git a/Engine/Source/Runtime/Launch/Private/LaunchEngineLoop.cpp b/Engine/Source/Runtime/Launch/Private/LaunchEngineLoop.cpp
--- a/Engine/Source/Runtime/Launch/Private/LaunchEngineLoop.cpp
+++ b/Engine/Source/Runtime/Launch/Private/LaunchEngineLoop.cpp
@@ -13,7 +13,6 @@
 #include "TaskGraphInterfaces.h"
 #include "Runtime/Core/Public/Modules/ModuleVersion.h"
 
-
 #include "Projects.h"
 #include "UProjectInfo.h"
 #include "EngineVersion.h"
@@ -1976,6 +1975,41 @@ void FEngineLoop::Exit()
 	FIOSystem::Shutdown();
 }
 
+void FEngineLoop::ProcessPlayerControllersSlateOperations() const
+{
+	FSlateApplication& SlateApp = FSlateApplication::Get();
+
+	// For all the game worlds drill down to the player controller for each game viewport and process it's slate operation
+	for (const FWorldContext& Context : GEngine->GetWorldContexts())
+	{
+		UWorld* CurWorld = Context.World();
+		if (CurWorld != nullptr && CurWorld->IsGameWorld())
+		{
+			ULocalPlayer* LocalPlayer = CurWorld->GetFirstLocalPlayerFromController();
+			UGameViewportClient* GameViewportClient = CurWorld->GetGameViewport();
+			if (LocalPlayer != nullptr && GameViewportClient != nullptr)
+			{
+				TSharedPtr<SViewport> ViewportWidget = GameViewportClient->GetGameViewportWidget();
+				if (ViewportWidget.IsValid())
+				{
+					APlayerController* PlayerController = LocalPlayer->PlayerController;
+					if (PlayerController != nullptr)
+					{
+						FReply& TheReply = PlayerController->SlateOperations;
+
+						FWidgetPath PathToWidget;
+						SlateApp.GeneratePathToWidgetUnchecked(ViewportWidget.ToSharedRef(), PathToWidget);
+
+						SlateApp.ProcessReply(PathToWidget, TheReply, NULL, NULL);
+
+						TheReply = FReply::Unhandled();
+					}
+				}
+			}
+		}
+	}
+}
+
 bool FEngineLoop::ShouldUseIdleMode() const
 {
 	static const auto CVarIdleWhenNotForeground = IConsoleManager::Get().FindTConsoleVariableDataInt(TEXT("t.IdleWhenNotForeground"));
@@ -2104,6 +2138,9 @@ void FEngineLoop::Tick()
 		if (FSlateApplication::IsInitialized() && !bIdleMode)
 		{
 			check(!IsRunningDedicatedServer());
+			
+			// Process slate operations accumulated in the world ticks.
+			ProcessPlayerControllersSlateOperations();
 
 			// Tick Slate application
 			FSlateApplication::Get().Tick();
diff --git a/Engine/Source/Runtime/Launch/Public/LaunchEngineLoop.h b/Engine/Source/Runtime/Launch/Public/LaunchEngineLoop.h
--- a/Engine/Source/Runtime/Launch/Public/LaunchEngineLoop.h
+++ b/Engine/Source/Runtime/Launch/Public/LaunchEngineLoop.h
@@ -130,6 +130,12 @@ public:
 	 */
 	static void AppExit( );
 
+private:
+	/**
+	 * Utility function that processes Slate operations
+	 */
+	void ProcessPlayerControllersSlateOperations() const;
+
 protected:
 
 	/**
