Missed file for CL# 2302986

diff --git a/Engine/Source/Runtime/Engine/Private/World.cpp b/Engine/Source/Runtime/Engine/Private/World.cpp
--- a/Engine/Source/Runtime/Engine/Private/World.cpp
+++ b/Engine/Source/Runtime/Engine/Private/World.cpp
@@ -2697,7 +2697,7 @@ bool UWorld::SetGameMode(const FURL& InURL)
 		else
 		{
 			// Remove any directory path from the map for the purpose of setting the game type
-			GameClass = GameClass->GetDefaultObject<AGameMode>()->SetGameMode(FPaths::GetBaseFilename(InURL.Map), Options, *InURL.Portal);
+			GameClass = GameClass->GetDefaultObject<AGameMode>()->GetGameModeClass(FPaths::GetBaseFilename(InURL.Map), Options, *InURL.Portal);
 		}
 
 		// no matter how the game was specified, we can use it to load the PostLoadClass packages
