Properly strip the PIE prefix when looking for GameMode based on map prefix

diff --git a/Engine/Source/Runtime/Engine/Private/World.cpp b/Engine/Source/Runtime/Engine/Private/World.cpp
--- a/Engine/Source/Runtime/Engine/Private/World.cpp
+++ b/Engine/Source/Runtime/Engine/Private/World.cpp
@@ -2655,7 +2655,9 @@ bool UWorld::SetGameMode(const FURL& InURL)
 			FString MapNameNoPath = FPaths::GetBaseFilename(MapName);
 			if (MapNameNoPath.StartsWith(PLAYWORLD_PACKAGE_PREFIX))
 			{
-				const int32 PrefixLen = FString(PLAYWORLD_PACKAGE_PREFIX).Len();
+				FWorldContext WorldContext = GEngine->GetWorldContextFromWorldChecked(this);
+
+				const int32 PrefixLen = BuildPIEPackagePrefix(WorldContext.PIEInstance).Len();
 				MapNameNoPath = MapNameNoPath.Mid(PrefixLen);
 			}
 			
