#docs - Fix some instances of "i.e." in comments terminating the short description in doxygen comments.

diff --git a/Engine/Source/Runtime/Core/Public/Math/OrientedBox.h b/Engine/Source/Runtime/Core/Public/Math/OrientedBox.h
--- a/Engine/Source/Runtime/Core/Public/Math/OrientedBox.h
+++ b/Engine/Source/Runtime/Core/Public/Math/OrientedBox.h
@@ -8,7 +8,7 @@
 
 
 /**
- * Structure for arbitrarily oriented boxes (i.e. not necessarily axis-aligned).
+ * Structure for arbitrarily oriented boxes (not necessarily axis-aligned).
  */
 struct FOrientedBox
 {
diff --git a/Engine/Source/Runtime/Engine/Classes/GameFramework/Character.h b/Engine/Source/Runtime/Engine/Classes/GameFramework/Character.h
--- a/Engine/Source/Runtime/Engine/Classes/GameFramework/Character.h
+++ b/Engine/Source/Runtime/Engine/Classes/GameFramework/Character.h
@@ -411,7 +411,7 @@ protected:
 public:
 
 	/**
-	 * True if jump is actively providing a force; i.e. jump key is held and the time it has been held is less than JumpMaxHoldTime.
+	 * True if jump is actively providing a force, such as when the jump key is held and the time it has been held is less than JumpMaxHoldTime.
 	 * @see CharacterMovement->IsFalling
 	 */
 	UFUNCTION(BlueprintCallable, Category="Pawn|Character")
