#ue4 - WalkableSlopeOverride now has an option to make a surface completely unwalkable to Characters (they will not step up on it and will "fall"/slide when on it).

- Previously you could get close by setting it to decrease the walkable slope to zero, but then flat surfaces were still walkable.
- Also improved comments on the enum values.

diff --git a/Engine/Source/Runtime/Engine/Classes/Engine/EngineTypes.h b/Engine/Source/Runtime/Engine/Classes/Engine/EngineTypes.h
--- a/Engine/Source/Runtime/Engine/Classes/Engine/EngineTypes.h
+++ b/Engine/Source/Runtime/Engine/Classes/Engine/EngineTypes.h
@@ -2330,18 +2330,35 @@ struct FRepAttachment
 };
 
 
-/** Enum controlling behavior of FWalkableSlopeOverride */
+/**
+ * Controls behavior of WalkableSlopeOverride, determining how to affect walkability of surfaces for Characters.
+ * @see FWalkableSlopeOverride
+ */
 UENUM(BlueprintType)
 enum EWalkableSlopeBehavior
 {
-	// Don't affect the walkable slope
-	WalkableSlope_Default	UMETA(DisplayName="Unchanged"),
+	/** Don't affect the walkable slope. */
+	WalkableSlope_Default		UMETA(DisplayName="Unchanged"),
+
+	/**
+	 * Increase walkable slope.
+	 * Makes it easier to walk up a surface, by allowing traversal over higher-than-usual angles.
+	 * @see FWalkableSlopeOverride::WalkableSlopeAngle
+	 */
+	WalkableSlope_Increase		UMETA(DisplayName="Increase Walkable Slope"),
 
-	// Increase walkable slope.
-	WalkableSlope_Increase	UMETA(DisplayName="Increase Walkable Slope"),
+	/**
+	 * Decrease walkable slope.
+	 * Makes it harder to walk up a surface, by restricting traversal to lower-than-usual angles.
+	 * @see FWalkableSlopeOverride::WalkableSlopeAngle
+	 */
+	WalkableSlope_Decrease		UMETA(DisplayName="Decrease Walkable Slope"),
 
-	// Decrease walkable slope.
-	WalkableSlope_Decrease	UMETA(DisplayName="Decrease Walkable Slope"),
+	/**
+	 * Make surface unwalkable.
+	 * Note: WalkableSlopeAngle will be ignored.
+	 */
+	WalkableSlope_Unwalkable	UMETA(DisplayName="Unwalkable"),
 	
 	WalkableSlope_Max		UMETA(Hidden),
 };
@@ -2391,6 +2408,12 @@ struct FWalkableSlopeOverride
 				return FMath::Max(InWalkableFloorZ, FMath::Cos(Angle));
 			}
 
+			case WalkableSlope_Unwalkable:
+			{
+				// Z component of a normal will always be less than this, so this will be unwalkable.
+				return 2.0f;
+			}
+
 			default:
 			{
 				return InWalkableFloorZ;
