added validating simple move requests and reporting them in PIE's message log
ttp# 332178

diff --git a/Engine/Source/Runtime/AIModule/Classes/Navigation/PathFollowingComponent.h b/Engine/Source/Runtime/AIModule/Classes/Navigation/PathFollowingComponent.h
--- a/Engine/Source/Runtime/AIModule/Classes/Navigation/PathFollowingComponent.h
+++ b/Engine/Source/Runtime/AIModule/Classes/Navigation/PathFollowingComponent.h
@@ -203,6 +203,7 @@ class AIMODULE_API UPathFollowingComponent : public UActorComponent, public IAIR
 	FORCEINLINE FVector GetCurrentDirection() const { return MoveSegmentDirection; }
 	FORCEINLINE FBasedPosition GetCurrentTargetLocationBased() const { return CurrentDestination; }
 	FORCEINLINE UObject* GetCurrentCustomLinkOb() const { return CurrentCustomLinkOb.Get(); }
+	FORCEINLINE class UNavMovementComponent* GetMovementComponent() const { return MovementComp; }
 
 	/** will be deprecated soon, please use AIController.GetMoveStatus instead! */
 	UFUNCTION(BlueprintCallable, Category="AI|Components|PathFollowing")
diff --git a/Engine/Source/Runtime/Engine/Private/AI/Navigation/NavigationSystem.cpp b/Engine/Source/Runtime/Engine/Private/AI/Navigation/NavigationSystem.cpp
--- a/Engine/Source/Runtime/Engine/Private/AI/Navigation/NavigationSystem.cpp
+++ b/Engine/Source/Runtime/Engine/Private/AI/Navigation/NavigationSystem.cpp
@@ -25,6 +25,7 @@
 #include "AI/Navigation/NavigationSystem.h"
 #include "AI/Navigation/NavRelevantComponent.h"
 #include "AI/Navigation/NavigationPath.h"
+#include "MessageLog.h"
 
 static const uint32 INITIAL_ASYNC_QUERIES_SIZE = 32;
 static const uint32 REGISTRATION_QUEUE_SIZE = 16;	// and we'll not reallocate
@@ -853,6 +854,32 @@ bool UNavigationSystem::ProjectPointToNavigation(const FVector& Point, FNavLocat
 		, QueryFilter);
 }
 
+void ValidateSimpleMoveRequest(UPathFollowingComponent* PFollowComp)
+{
+	// simple move requests require a character in valid movement mode
+	UCharacterMovementComponent* CharMovement = Cast<UCharacterMovementComponent>(PFollowComp->GetMovementComponent());
+	ACharacter* CharOwner = CharMovement ? CharMovement->GetCharacterOwner() : NULL;
+	if (CharMovement == NULL || CharMovement->UpdatedComponent == NULL ||
+		CharMovement->UpdatedComponent->IsSimulatingPhysics() ||
+		CharOwner == NULL || CharOwner->IsMatineeControlled())
+	{
+		FString ReasonDesc = (CharMovement == NULL) ? TEXT("not using character movement") :
+			(CharMovement->UpdatedComponent == NULL) ? TEXT("uninitialized movement component") :
+			CharMovement->UpdatedComponent->IsSimulatingPhysics() ? TEXT("rigid body") :
+			(CharOwner == NULL) ? TEXT("not using character class") :
+			CharOwner->IsMatineeControlled() ? TEXT("matinee control") :
+			TEXT("unknown reason");
+
+		AController* OwnerController = Cast<AController>(PFollowComp->GetOuter());
+
+		FText WarnMsg = FText::Format(NSLOCTEXT("Navigation", "SimpleMoveWarning", "Simple move request called on pawn {0} using invalid movement mode: {1}"),
+			FText::FromName(OwnerController && OwnerController->GetPawn() ? OwnerController->GetPawn()->GetFName() : NAME_None), FText::FromString(ReasonDesc));
+
+		UE_LOG(LogNavigation, Warning, TEXT("%s"), *WarnMsg.ToString());
+		FMessageLog("PIE").Warning(WarnMsg);
+	}
+}
+
 void UNavigationSystem::SimpleMoveToActor(AController* Controller, const AActor* Goal)
 {
 	if (Goal == NULL || Controller == NULL || Controller->GetPawn() == NULL)
@@ -869,6 +896,8 @@ void UNavigationSystem::SimpleMoveToActor(AController* Controller, const AActor*
 
 	if (PFindComp && PFollowComp && !PFollowComp->HasReached(Goal))
 	{
+		ValidateSimpleMoveRequest(PFollowComp);
+
 		const bool bPathExists = PFindComp->FindPathToActor(Goal);
 		if (bPathExists)
 		{
@@ -893,6 +922,8 @@ void UNavigationSystem::SimpleMoveToLocation(AController* Controller, const FVec
 
 	if (PFindComp && PFollowComp && !PFollowComp->HasReached(Goal))
 	{
+		ValidateSimpleMoveRequest(PFollowComp);
+
 		const bool bPathExists = PFindComp->FindPathToLocation(Goal);
 		if (bPathExists)
 		{
