Remove Single from friendly names of scene queries

diff --git a/Engine/Source/Runtime/Engine/Classes/Kismet/KismetSystemLibrary.h b/Engine/Source/Runtime/Engine/Classes/Kismet/KismetSystemLibrary.h
--- a/Engine/Source/Runtime/Engine/Classes/Kismet/KismetSystemLibrary.h
+++ b/Engine/Source/Runtime/Engine/Classes/Kismet/KismetSystemLibrary.h
@@ -556,7 +556,7 @@ class UKismetSystemLibrary : public UBlueprintFunctionLibrary
 	 * @param OutHit		Properties of the trace hit.
 	 * @return				True if there was a hit, false otherwise.
 	 */
-	UFUNCTION(BlueprintCallable, Category="Collision", meta=(bIgnoreSelf="true", WorldContext="WorldContextObject", AutoCreateRefTerm="ActorsToIgnore", FriendlyName="SingleLineTraceByChannel", Keywords="raycast"))
+	UFUNCTION(BlueprintCallable, Category="Collision", meta=(bIgnoreSelf="true", WorldContext="WorldContextObject", AutoCreateRefTerm="ActorsToIgnore", FriendlyName="LineTraceByChannel", Keywords="raycast"))
 	static ENGINE_API bool LineTraceSingle_NEW(UObject* WorldContextObject, const FVector Start, const FVector End, ETraceTypeQuery TraceChannel, bool bTraceComplex, const TArray<AActor*>& ActorsToIgnore, EDrawDebugTrace::Type DrawDebugType, FHitResult& OutHit, bool bIgnoreSelf);
 	
 	/**
@@ -586,7 +586,7 @@ class UKismetSystemLibrary : public UBlueprintFunctionLibrary
 	 * @param OutHit		Properties of the trace hit.
 	 * @return				True if there was a hit, false otherwise.
 	 */
-	UFUNCTION(BlueprintCallable, Category="Collision", meta=(bIgnoreSelf="true", WorldContext="WorldContextObject", AutoCreateRefTerm="ActorsToIgnore", FriendlyName = "SingleSphereTraceByChannel", Keywords="sweep"))
+	UFUNCTION(BlueprintCallable, Category="Collision", meta=(bIgnoreSelf="true", WorldContext="WorldContextObject", AutoCreateRefTerm="ActorsToIgnore", FriendlyName = "SphereTraceByChannel", Keywords="sweep"))
 	static ENGINE_API bool SphereTraceSingle_NEW(UObject* WorldContextObject, const FVector Start, const FVector End, float Radius, ETraceTypeQuery TraceChannel, bool bTraceComplex, const TArray<AActor*>& ActorsToIgnore, EDrawDebugTrace::Type DrawDebugType, FHitResult& OutHit, bool bIgnoreSelf);
 
 	/**
@@ -618,7 +618,7 @@ class UKismetSystemLibrary : public UBlueprintFunctionLibrary
 	* @param OutHit			Properties of the trace hit.
 	* @return				True if there was a hit, false otherwise.
 	*/
-	UFUNCTION(BlueprintCallable, Category = "Collision", meta = (bIgnoreSelf = "true", HidePin = "WorldContextObject", DefaultToSelf = "WorldContextObject", AutoCreateRefTerm = "ActorsToIgnore", FriendlyName = "SingleBoxTraceByChannel", Keywords="sweep"))
+	UFUNCTION(BlueprintCallable, Category = "Collision", meta = (bIgnoreSelf = "true", HidePin = "WorldContextObject", DefaultToSelf = "WorldContextObject", AutoCreateRefTerm = "ActorsToIgnore", FriendlyName = "BoxTraceByChannel", Keywords="sweep"))
 	static ENGINE_API bool BoxTraceSingle(UObject* WorldContextObject, const FVector Start, const FVector End, const FVector HalfSize, const FRotator Orientation, ETraceTypeQuery TraceChannel, bool bTraceComplex, const TArray<AActor*>& ActorsToIgnore, EDrawDebugTrace::Type DrawDebugType, FHitResult& OutHit, bool bIgnoreSelf);
 
 	/**
@@ -652,7 +652,7 @@ class UKismetSystemLibrary : public UBlueprintFunctionLibrary
 	 * @param OutHit		Properties of the trace hit.
 	 * @return				True if there was a hit, false otherwise.
 	 */
-	UFUNCTION(BlueprintCallable, Category="Collision", meta=(bIgnoreSelf="true", WorldContext="WorldContextObject", AutoCreateRefTerm="ActorsToIgnore", FriendlyName = "SingleCapsuleTraceByChannel", Keywords="sweep"))
+	UFUNCTION(BlueprintCallable, Category="Collision", meta=(bIgnoreSelf="true", WorldContext="WorldContextObject", AutoCreateRefTerm="ActorsToIgnore", FriendlyName = "CapsuleTraceByChannel", Keywords="sweep"))
 	static ENGINE_API bool CapsuleTraceSingle_NEW(UObject* WorldContextObject, const FVector Start, const FVector End, float Radius, float HalfHeight, ETraceTypeQuery TraceChannel, bool bTraceComplex, const TArray<AActor*>& ActorsToIgnore, EDrawDebugTrace::Type DrawDebugType, FHitResult& OutHit, bool bIgnoreSelf);
 
 	/**
@@ -684,7 +684,7 @@ class UKismetSystemLibrary : public UBlueprintFunctionLibrary
 	 * @param OutHit		Properties of the trace hit.
 	 * @return				True if there was a hit, false otherwise.
 	 */
-	UFUNCTION(BlueprintCallable, Category="Collision", meta=(bIgnoreSelf="true", WorldContext="WorldContextObject", AutoCreateRefTerm="ActorsToIgnore", FriendlyName = "SingleLineTraceForObjects", Keywords="raycast"))
+	UFUNCTION(BlueprintCallable, Category="Collision", meta=(bIgnoreSelf="true", WorldContext="WorldContextObject", AutoCreateRefTerm="ActorsToIgnore", FriendlyName = "LineTraceForObjects", Keywords="raycast"))
 	static ENGINE_API bool LineTraceSingleForObjects(UObject* WorldContextObject, const FVector Start, const FVector End, const TArray<TEnumAsByte<EObjectTypeQuery> > & ObjectTypes, bool bTraceComplex, const TArray<AActor*>& ActorsToIgnore, EDrawDebugTrace::Type DrawDebugType, FHitResult& OutHit, bool bIgnoreSelf);
 	
 	/**
@@ -714,7 +714,7 @@ class UKismetSystemLibrary : public UBlueprintFunctionLibrary
 	 * @param OutHit		Properties of the trace hit.
 	 * @return				True if there was a hit, false otherwise.
 	 */
-	UFUNCTION(BlueprintCallable, Category="Collision", meta=(bIgnoreSelf="true", WorldContext="WorldContextObject", AutoCreateRefTerm="ActorsToIgnore", FriendlyName = "SingleSphereTraceForObjects", Keywords="sweep"))
+	UFUNCTION(BlueprintCallable, Category="Collision", meta=(bIgnoreSelf="true", WorldContext="WorldContextObject", AutoCreateRefTerm="ActorsToIgnore", FriendlyName = "SphereTraceForObjects", Keywords="sweep"))
 	static ENGINE_API bool SphereTraceSingleForObjects(UObject* WorldContextObject, const FVector Start, const FVector End, float Radius, const TArray<TEnumAsByte<EObjectTypeQuery> > & ObjectTypes, bool bTraceComplex, const TArray<AActor*>& ActorsToIgnore, EDrawDebugTrace::Type DrawDebugType, FHitResult& OutHit, bool bIgnoreSelf);
 
 	/**
@@ -747,7 +747,7 @@ class UKismetSystemLibrary : public UBlueprintFunctionLibrary
 	* @param OutHit			Properties of the trace hit.
 	* @return				True if there was a hit, false otherwise.
 	*/
-	UFUNCTION(BlueprintCallable, Category = "Collision", meta = (bIgnoreSelf = "true", HidePin = "WorldContextObject", DefaultToSelf = "WorldContextObject", AutoCreateRefTerm = "ActorsToIgnore", FriendlyName = "SingleBoxTraceForObjects", Keywords="sweep"))
+	UFUNCTION(BlueprintCallable, Category = "Collision", meta = (bIgnoreSelf = "true", HidePin = "WorldContextObject", DefaultToSelf = "WorldContextObject", AutoCreateRefTerm = "ActorsToIgnore", FriendlyName = "BoxTraceForObjects", Keywords="sweep"))
 	static ENGINE_API bool BoxTraceSingleForObjects(UObject* WorldContextObject, const FVector Start, const FVector End, const FVector HalfSize, const FRotator Orientation, const TArray<TEnumAsByte<EObjectTypeQuery> > & ObjectTypes, bool bTraceComplex, const TArray<AActor*>& ActorsToIgnore, EDrawDebugTrace::Type DrawDebugType, FHitResult& OutHit, bool bIgnoreSelf);
 
 
@@ -781,7 +781,7 @@ class UKismetSystemLibrary : public UBlueprintFunctionLibrary
 	 * @param OutHit		Properties of the trace hit.
 	 * @return				True if there was a hit, false otherwise.
 	 */
-	UFUNCTION(BlueprintCallable, Category="Collision", meta=(bIgnoreSelf="true", WorldContext="WorldContextObject", AutoCreateRefTerm="ActorsToIgnore", FriendlyName = "SingleCapsuleTraceForObjects", Keywords="sweep"))
+	UFUNCTION(BlueprintCallable, Category="Collision", meta=(bIgnoreSelf="true", WorldContext="WorldContextObject", AutoCreateRefTerm="ActorsToIgnore", FriendlyName = "CapsuleTraceForObjects", Keywords="sweep"))
 	static ENGINE_API bool CapsuleTraceSingleForObjects(UObject* WorldContextObject, const FVector Start, const FVector End, float Radius, float HalfHeight, const TArray<TEnumAsByte<EObjectTypeQuery> > & ObjectTypes, bool bTraceComplex, const TArray<AActor*>& ActorsToIgnore, EDrawDebugTrace::Type DrawDebugType, FHitResult& OutHit, bool bIgnoreSelf);
 
 	/**
