UE4: Fixing a rare crash involving releasing the AdjacencyIndexBuffer on static meshes.

diff --git a/Engine/Source/Runtime/Engine/Private/StaticMesh.cpp b/Engine/Source/Runtime/Engine/Private/StaticMesh.cpp
--- a/Engine/Source/Runtime/Engine/Private/StaticMesh.cpp
+++ b/Engine/Source/Runtime/Engine/Private/StaticMesh.cpp
@@ -553,10 +553,11 @@ void FStaticMeshLODResources::ReleaseResources()
 	DEC_DWORD_STAT_BY( STAT_StaticMeshIndexMemory, StaticMeshIndexMemory );
 
 	// Release the vertex and index buffers.
-	if (AdjacencyIndexBuffer.IsInitialized())
-	{
-		BeginReleaseResource(&AdjacencyIndexBuffer);
-	}
+	
+	// AdjacencyIndexBuffer may not be initialized at this time, but it is safe to release it anyway.
+	// The bInitialized flag will be safely checked in the render thread.
+	// This avoids a race condition regarding releasing this resource.
+	BeginReleaseResource(&AdjacencyIndexBuffer);
 
 	BeginReleaseResource(&IndexBuffer);
 	BeginReleaseResource(&WireframeIndexBuffer);
