#docs - Further comment improvements for WalkableSlopeOverride

diff --git a/Engine/Source/Runtime/Engine/Classes/Engine/EngineTypes.h b/Engine/Source/Runtime/Engine/Classes/Engine/EngineTypes.h
--- a/Engine/Source/Runtime/Engine/Classes/Engine/EngineTypes.h
+++ b/Engine/Source/Runtime/Engine/Classes/Engine/EngineTypes.h
@@ -2333,6 +2333,7 @@ struct FRepAttachment
 /**
  * Controls behavior of WalkableSlopeOverride, determining how to affect walkability of surfaces for Characters.
  * @see FWalkableSlopeOverride
+ * @see UCharacterMovementComponent::GetWalkableFloorAngle(), UCharacterMovementComponent::SetWalkableFloorAngle()
  */
 UENUM(BlueprintType)
 enum EWalkableSlopeBehavior
@@ -2370,12 +2371,12 @@ struct FWalkableSlopeOverride
 {
 	GENERATED_USTRUCT_BODY()
 
-	// Behavior of this surface (whether we affect the walkable normal)
+	/** Behavior of this surface (whether we affect the walkable slope). */
 	UPROPERTY(EditAnywhere, BlueprintReadWrite, Category=WalkableSlopeOverride)
 	TEnumAsByte<EWalkableSlopeBehavior> WalkableSlopeBehavior;
 
 	/**
-	 * Override slope, applying the rules of the Walkable Slope Behavior.
+	 * Override walkable slope, applying the rules of the Walkable Slope Behavior.
 	 */
 	UPROPERTY(EditAnywhere, BlueprintReadWrite, Category=WalkableSlopeOverride, meta=(ClampMin="0", ClampMax="90", UIMin="0", UIMax="90"))
 	float WalkableSlopeAngle;
