Fixed material break attributes reorder from clear coat change.

diff --git a/Engine/Shaders/ShaderVersion.usf b/Engine/Shaders/ShaderVersion.usf
--- a/Engine/Shaders/ShaderVersion.usf
+++ b/Engine/Shaders/ShaderVersion.usf
@@ -2,4 +2,4 @@
 // Each time the console command is executed it generates a new GUID. As this file is included
 // in common.usf (which should be included in any shader) it allows to invalidate the shader DDC.
 // 
-// GUID = 1B144F2441A3CDA26663BB820BFD049B
\ No newline at end of file
+// GUID = 1B1A4F2441A3CDA2433BB820BFD049B
\ No newline at end of file
diff --git a/Engine/Source/Runtime/Engine/Private/Materials/Material.cpp b/Engine/Source/Runtime/Engine/Private/Materials/Material.cpp
--- a/Engine/Source/Runtime/Engine/Private/Materials/Material.cpp
+++ b/Engine/Source/Runtime/Engine/Private/Materials/Material.cpp
@@ -1716,25 +1716,22 @@ void UMaterial::Serialize(FArchive& Ar)
 	}
 #endif // #if WITH_EDITOR
 
-	if( Ar.UE4Ver() < VER_UE4_MATERIAL_ATTRIBUTES_REORDERING )
-	{
-		DoMaterialAttributeReorder(&DiffuseColor);
-		DoMaterialAttributeReorder(&SpecularColor);
-		DoMaterialAttributeReorder(&BaseColor);
-		DoMaterialAttributeReorder(&Metallic);
-		DoMaterialAttributeReorder(&Specular);
-		DoMaterialAttributeReorder(&Roughness);
-		DoMaterialAttributeReorder(&Normal);
-		DoMaterialAttributeReorder(&EmissiveColor);
-		DoMaterialAttributeReorder(&Opacity);
-		DoMaterialAttributeReorder(&OpacityMask);
-		DoMaterialAttributeReorder(&WorldPositionOffset);
-		DoMaterialAttributeReorder(&WorldDisplacement);
-		DoMaterialAttributeReorder(&TessellationMultiplier);
-		DoMaterialAttributeReorder(&SubsurfaceColor);
-		DoMaterialAttributeReorder(&AmbientOcclusion);
-		DoMaterialAttributeReorder(&Refraction);
-	}
+	DoMaterialAttributeReorder(&DiffuseColor,			Ar.UE4Ver());
+	DoMaterialAttributeReorder(&SpecularColor,			Ar.UE4Ver());
+	DoMaterialAttributeReorder(&BaseColor,				Ar.UE4Ver());
+	DoMaterialAttributeReorder(&Metallic,				Ar.UE4Ver());
+	DoMaterialAttributeReorder(&Specular,				Ar.UE4Ver());
+	DoMaterialAttributeReorder(&Roughness,				Ar.UE4Ver());
+	DoMaterialAttributeReorder(&Normal,					Ar.UE4Ver());
+	DoMaterialAttributeReorder(&EmissiveColor,			Ar.UE4Ver());
+	DoMaterialAttributeReorder(&Opacity,				Ar.UE4Ver());
+	DoMaterialAttributeReorder(&OpacityMask,			Ar.UE4Ver());
+	DoMaterialAttributeReorder(&WorldPositionOffset,	Ar.UE4Ver());
+	DoMaterialAttributeReorder(&WorldDisplacement,		Ar.UE4Ver());
+	DoMaterialAttributeReorder(&TessellationMultiplier,	Ar.UE4Ver());
+	DoMaterialAttributeReorder(&SubsurfaceColor,		Ar.UE4Ver());
+	DoMaterialAttributeReorder(&AmbientOcclusion,		Ar.UE4Ver());
+	DoMaterialAttributeReorder(&Refraction,				Ar.UE4Ver());
 }
 
 void UMaterial::PostDuplicate(bool bDuplicateForPIE)
diff --git a/Engine/Source/Runtime/Engine/Private/Materials/MaterialExpressions.cpp b/Engine/Source/Runtime/Engine/Private/Materials/MaterialExpressions.cpp
--- a/Engine/Source/Runtime/Engine/Private/Materials/MaterialExpressions.cpp
+++ b/Engine/Source/Runtime/Engine/Private/Materials/MaterialExpressions.cpp
@@ -436,28 +436,12 @@ void UMaterialExpression::Serialize( FArchive& Ar )
 {
 	Super::Serialize(Ar);
 
-	if (Ar.UE4Ver() < VER_UE4_CHANGED_MATERIAL_REFACTION_TYPE)
+	const TArray<FExpressionInput*> Inputs = GetInputs();
+	for (int32 InputIndex = 0; InputIndex < Inputs.Num(); ++InputIndex)
 	{
-		const TArray<FExpressionInput*> Inputs = GetInputs();
-		for (int32 InputIndex = 0; InputIndex < Inputs.Num(); ++InputIndex)
-		{
-			FExpressionInput* Input = Inputs[InputIndex];
-			if (Ar.UE4Ver() < VER_UE4_MATERIAL_ATTRIBUTES_REORDERING)
-			{
-				DoMaterialAttributeReorder(Input);
-			}
-			if (Input->Expression
-				&& Input->OutputIndex == 13
-				&& Input->Expression->IsA(UMaterialExpressionBreakMaterialAttributes::StaticClass())
-				&& Input->MaskB)
-			{
-				Input->Mask = 1;
-				Input->MaskR = 1;
-				Input->MaskG = 1;
-				Input->MaskB = 0;
-				Input->MaskA = 0;
-			}
-		}
+		FExpressionInput* Input = Inputs[InputIndex];
+
+		DoMaterialAttributeReorder(Input, Ar.UE4Ver());
 	}
 }
 
@@ -3273,9 +3257,9 @@ UMaterialExpressionBreakMaterialAttributes::UMaterialExpressionBreakMaterialAttr
 	Outputs.Add(FExpressionOutput(TEXT("ClearCoat"), 1, 1, 1, 1, 0));
 	Outputs.Add(FExpressionOutput(TEXT("ClearCoatRoughness"), 1, 1, 1, 1, 0));
 	Outputs.Add(FExpressionOutput(TEXT("AmbientOcclusion"), 1, 1, 1, 1, 0));
-	Outputs.Add(FExpressionOutput(TEXT("Refraction"), 1, 1, 1, 0, 0));
+	Outputs.Add(FExpressionOutput(TEXT("Refraction"), 1, 1, 1, 1, 0));
 
-	for (int32 UVIndex = 0; UVIndex < MP_CustomizedUVs7 - MP_CustomizedUVs0; UVIndex++)
+	for (int32 UVIndex = 0; UVIndex <= MP_CustomizedUVs7 - MP_CustomizedUVs0; UVIndex++)
 	{
 		Outputs.Add(FExpressionOutput(*FString::Printf(TEXT("CustomizedUV%u"), UVIndex), 1, 1, 1, 0, 0));
 	}
diff --git a/Engine/Source/Runtime/Engine/Private/Materials/MaterialShared.cpp b/Engine/Source/Runtime/Engine/Private/Materials/MaterialShared.cpp
--- a/Engine/Source/Runtime/Engine/Private/Materials/MaterialShared.cpp
+++ b/Engine/Source/Runtime/Engine/Private/Materials/MaterialShared.cpp
@@ -111,7 +111,7 @@ int32 FColorMaterialInput::CompileWithDefault(class FMaterialCompiler* Compiler,
 		}
 	}
 
-	return GetDefaultExpressionForMaterialProperty(Compiler, Property);
+	return Compiler->ForceCast(GetDefaultExpressionForMaterialProperty(Compiler, Property), MCT_Float3);
 }
 
 int32 FScalarMaterialInput::CompileWithDefault(class FMaterialCompiler* Compiler, EMaterialProperty Property)
@@ -129,7 +129,7 @@ int32 FScalarMaterialInput::CompileWithDefault(class FMaterialCompiler* Compiler
 		}
 	}
 
-	return GetDefaultExpressionForMaterialProperty(Compiler, Property);
+	return Compiler->ForceCast(GetDefaultExpressionForMaterialProperty(Compiler, Property), MCT_Float1);
 }
 
 int32 FVectorMaterialInput::CompileWithDefault(class FMaterialCompiler* Compiler, EMaterialProperty Property)
@@ -146,7 +146,7 @@ int32 FVectorMaterialInput::CompileWithDefault(class FMaterialCompiler* Compiler
 			return ResultIndex;
 		}
 	}
-	return GetDefaultExpressionForMaterialProperty(Compiler, Property);
+	return Compiler->ForceCast(GetDefaultExpressionForMaterialProperty(Compiler, Property), MCT_Float3);
 }
 
 int32 FVector2MaterialInput::CompileWithDefault(class FMaterialCompiler* Compiler, EMaterialProperty Property)
@@ -164,7 +164,7 @@ int32 FVector2MaterialInput::CompileWithDefault(class FMaterialCompiler* Compile
 		}
 	}
 
-	return GetDefaultExpressionForMaterialProperty(Compiler, Property);
+	return Compiler->ForceCast(GetDefaultExpressionForMaterialProperty(Compiler, Property), MCT_Float2);
 }
 
 int32 FMaterialAttributesInput::CompileWithDefault(class FMaterialCompiler* Compiler, EMaterialProperty Property)
@@ -2277,20 +2277,38 @@ int32 UMaterialInterface::CompileProperty(FMaterialCompiler* Compiler, EMaterial
 	}
 }
 
-/** TODO - This can be removed whenever VER_UE4_MATERIAL_ATTRIBUTES_REORDERING is no longer relevant. */
-void DoMaterialAttributeReorder(FExpressionInput* Input)
+void DoMaterialAttributeReorder(FExpressionInput* Input, int32 UE4Ver)
 {
 	if( Input && Input->Expression && Input->Expression->IsA(UMaterialExpressionBreakMaterialAttributes::StaticClass()) )
 	{
-		switch(Input->OutputIndex)
+		if( UE4Ver < VER_UE4_MATERIAL_ATTRIBUTES_REORDERING )
 		{
-		case 4: Input->OutputIndex = 7; break;
-		case 5: Input->OutputIndex = 4; break;
-		case 6: Input->OutputIndex = 5; break;
-		case 7: Input->OutputIndex = 6; break;
+			switch(Input->OutputIndex)
+			{
+			case 4: Input->OutputIndex = 7; break;
+			case 5: Input->OutputIndex = 4; break;
+			case 6: Input->OutputIndex = 5; break;
+			case 7: Input->OutputIndex = 6; break;
+			}
+		}
+		
+		if( UE4Ver < VER_UE4_CHANGED_MATERIAL_REFACTION_TYPE && Input->OutputIndex == 13 )
+		{
+			Input->Mask = 1;
+			Input->MaskR = 1;
+			Input->MaskG = 1;
+			Input->MaskB = 1;
+			Input->MaskA = 0;
+		}
+
+		// closest version to the clear coat change
+		if( UE4Ver < VER_UE4_ADD_ROOTCOMPONENT_TO_FOLIAGEACTOR && Input->OutputIndex >= 12 )
+		{
+			Input->OutputIndex += 2;
 		}
 	}
 }
+
 FMaterialInstanceBasePropertyOverrides::FMaterialInstanceBasePropertyOverrides()
 	:bOverride_OpacityMaskClipValue(false)
 	,bOverride_BlendMode(false)
diff --git a/Engine/Source/Runtime/Engine/Public/MaterialShared.h b/Engine/Source/Runtime/Engine/Public/MaterialShared.h
--- a/Engine/Source/Runtime/Engine/Public/MaterialShared.h
+++ b/Engine/Source/Runtime/Engine/Public/MaterialShared.h
@@ -1596,4 +1596,4 @@ ENGINE_API int32 GetDefaultExpressionForMaterialProperty(FMaterialCompiler* Comp
 ENGINE_API FString GetNameOfMaterialProperty(EMaterialProperty Property);
 
 /** TODO - This can be removed whenever VER_UE4_MATERIAL_ATTRIBUTES_REORDERING is no longer relevant. */
-ENGINE_API void DoMaterialAttributeReorder(FExpressionInput* Input);
+ENGINE_API void DoMaterialAttributeReorder(FExpressionInput* Input, int32 UE4Ver);
