Editor: Improved display names and tooltips of some rendering options

diff --git a/Engine/Source/Runtime/Engine/Classes/Engine/RendererSettings.h b/Engine/Source/Runtime/Engine/Classes/Engine/RendererSettings.h
--- a/Engine/Source/Runtime/Engine/Classes/Engine/RendererSettings.h
+++ b/Engine/Source/Runtime/Engine/Classes/Engine/RendererSettings.h
@@ -187,38 +187,38 @@ class ENGINE_API URendererSettings
 	uint32 bSeparateTranslucency:1;
 
 	UPROPERTY(config, EditAnywhere, Category=Postprocessing, meta=(
-		ConsoleVariable="r.CustomDepth",DisplayName="Custom Depth",
+		ConsoleVariable="r.CustomDepth",DisplayName="Custom Depth Pass",
 		ToolTip="Whether the custom depth pass for tagging primitives for postprocessing passes is enabled. Enabling it on demand can save memory but may cause a hitch the first time the feature is used."))
 	TEnumAsByte<ECustomDepth::Type> CustomDepth;
 
-	UPROPERTY(config, EditAnywhere, Category = Postprocessing, meta = (
-		ConsoleVariable = "r.DefaultFeature.Bloom", DisplayName = "Bloom is enabled by default",
-		ToolTip = "Whether the default for Bloom is enabled or not (postprocess volume/camera/game setting still can enable)"))
+	UPROPERTY(config, EditAnywhere, Category = DefaultPostprocessingSettings, meta = (
+		ConsoleVariable = "r.DefaultFeature.Bloom", DisplayName = "Bloom",
+		ToolTip = "Whether the default for Bloom is enabled or not (postprocess volume/camera/game setting can still override and enable or disable it independently)"))
 	uint32 bDefaultFeatureBloom : 1;
 
-	UPROPERTY(config, EditAnywhere, Category = Postprocessing, meta = (
-		ConsoleVariable = "r.DefaultFeature.AmbientOcclusion", DisplayName = "Ambient Occlusion is enabled by default",
-		ToolTip = "Whether the default for AmbientOcclusion is enabled or not (postprocess volume/camera/game setting still can enable)"))
+	UPROPERTY(config, EditAnywhere, Category = DefaultPostprocessingSettings, meta = (
+		ConsoleVariable = "r.DefaultFeature.AmbientOcclusion", DisplayName = "Ambient Occlusion",
+		ToolTip = "Whether the default for AmbientOcclusion is enabled or not (postprocess volume/camera/game setting can still override and enable or disable it independently)"))
 	uint32 bDefaultFeatureAmbientOcclusion : 1;
 
-	UPROPERTY(config, EditAnywhere, Category = Postprocessing, meta = (
-		ConsoleVariable = "r.DefaultFeature.AutoExposure", DisplayName = "Auto Exposure is enabled by default",
-		ToolTip = "Whether the default for AutoExposure is enabled or not (postprocess volume/camera/game setting still can enable)"))
+	UPROPERTY(config, EditAnywhere, Category = DefaultPostprocessingSettings, meta = (
+		ConsoleVariable = "r.DefaultFeature.AutoExposure", DisplayName = "Auto Exposure",
+		ToolTip = "Whether the default for AutoExposure is enabled or not (postprocess volume/camera/game setting can still override and enable or disable it independently)"))
 	uint32 bDefaultFeatureAutoExposure : 1;
 
-	UPROPERTY(config, EditAnywhere, Category = Postprocessing, meta = (
-		ConsoleVariable = "r.DefaultFeature.MotionBlur", DisplayName = "Motion Blur is enabled by default",
-		ToolTip = "Whether the default for MotionBlur is enabled or not (postprocess volume/camera/game setting still can enable)"))
+	UPROPERTY(config, EditAnywhere, Category = DefaultPostprocessingSettings, meta = (
+		ConsoleVariable = "r.DefaultFeature.MotionBlur", DisplayName = "Motion Blur",
+		ToolTip = "Whether the default for MotionBlur is enabled or not (postprocess volume/camera/game setting can still override and enable or disable it independently)"))
 	uint32 bDefaultFeatureMotionBlur : 1;
 
-	UPROPERTY(config, EditAnywhere, Category = Postprocessing, meta = (
-		ConsoleVariable = "r.DefaultFeature.LensFlare", DisplayName = "Lens Flares (Image based) are enabled by default",
-		ToolTip = "Whether the default for LensFlare is enabled or not (postprocess volume/camera/game setting still can enable)"))
+	UPROPERTY(config, EditAnywhere, Category = DefaultPostprocessingSettings, meta = (
+		ConsoleVariable = "r.DefaultFeature.LensFlare", DisplayName = "Lens Flares (Image based)",
+		ToolTip = "Whether the default for LensFlare is enabled or not (postprocess volume/camera/game setting can still override and enable or disable it independently)"))
 	uint32 bDefaultFeatureLensFlare : 1;
 
-	UPROPERTY(config, EditAnywhere, Category = Postprocessing, meta = (
-		ConsoleVariable = "r.DefaultFeature.AntiAliasing", DisplayName = "Anti Aliasing method that is used by default",
-		ToolTip = "What anti-aliasing mode is used by default (postprocess volume/camera/game setting still can enable)"))
+	UPROPERTY(config, EditAnywhere, Category = DefaultPostprocessingSettings, meta = (
+		ConsoleVariable = "r.DefaultFeature.AntiAliasing", DisplayName = "Anti-Aliasing Method",
+		ToolTip = "What anti-aliasing mode is used by default (postprocess volume/camera/game setting can still override and enable or disable it independently)"))
 	TEnumAsByte<EAntiAliasingMethodUI::Type> DefaultFeatureAntiAliasing;
 
 	UPROPERTY(config, EditAnywhere, Category = Optimizations, meta = (
