Fixed so that different node matched skeleton works each other.

diff --git a/Engine/Source/Runtime/Engine/Classes/Animation/AnimSequence.h b/Engine/Source/Runtime/Engine/Classes/Animation/AnimSequence.h
--- a/Engine/Source/Runtime/Engine/Classes/Animation/AnimSequence.h
+++ b/Engine/Source/Runtime/Engine/Classes/Animation/AnimSequence.h
@@ -113,6 +113,45 @@ struct FRawAnimSequenceTrack
 	}
 };
 
+/**
+ * These two always should go together, but it is not right now. 
+ * I wonder in the future, we change all compressed to be inside as well, so they all stay together
+ * When remove tracks, it should be handled together
+ */
+USTRUCT()
+struct FAnimSequenceTrackContainer
+{
+	GENERATED_USTRUCT_BODY()
+
+	UPROPERTY()
+	TArray<struct FRawAnimSequenceTrack> AnimationTracks;
+
+	UPROPERTY()
+	TArray<FName>						TrackNames;
+
+	// @todo expand this struct to work better and assign data better
+	void Initialize(int32 NumNode)
+	{
+		AnimationTracks.Empty(NumNode);
+		AnimationTracks.AddZeroed(NumNode);
+		TrackNames.Empty(NumNode);
+		TrackNames.AddZeroed(NumNode);
+	}
+
+	void Initialize(TArray<FName> InTrackNames)
+	{
+		TrackNames = InTrackNames;
+		int32 NumNode = TrackNames.Num();
+		AnimationTracks.Empty(NumNode);
+		AnimationTracks.AddZeroed(NumNode);
+	}
+
+	int32 GetNum() const
+	{
+		check (TrackNames.Num() == AnimationTracks.Num());
+		return (AnimationTracks.Num());
+	}
+};
 /**
  *	@note we have plan to support skeleton hierarchy. 
  *	when that happens, we'd like to keep skeleton indexing
@@ -727,9 +766,9 @@ private:
 	void RemoveNaNTracks();
 
 	/** Retargeting functions */
-	bool ConvertAnimationDataToRiggingData(TArray<struct FRawAnimSequenceTrack> & RawRiggingAnimationData);
-	bool ConvertRiggingDataToAnimationData(TArray<struct FRawAnimSequenceTrack> & RawRiggingAnimationData);
-	int32 GetSpaceBasedAnimationData(TArray< TArray<FTransform> > & AnimationDataInComponentSpace, TArray<struct FRawAnimSequenceTrack> * RawRiggingAnimationData) const;
+	bool ConvertAnimationDataToRiggingData(FAnimSequenceTrackContainer & RiggingAnimationData);
+	bool ConvertRiggingDataToAnimationData(FAnimSequenceTrackContainer & RiggingAnimationData);
+	int32 GetSpaceBasedAnimationData(TArray< TArray<FTransform> > & AnimationDataInComponentSpace, FAnimSequenceTrackContainer * RiggingAnimationData) const;
 #endif
 
 	/**
diff --git a/Engine/Source/Runtime/Engine/Classes/Animation/Rig.h b/Engine/Source/Runtime/Engine/Classes/Animation/Rig.h
--- a/Engine/Source/Runtime/Engine/Classes/Animation/Rig.h
+++ b/Engine/Source/Runtime/Engine/Classes/Animation/Rig.h
@@ -171,7 +171,7 @@ public:
 	// rig control related
 	int32 GetTransformBaseNum() const;
 	const TArray<FTransformBase> & GetTransformBases() const { return TransformBases; }
-	const FTransformBase * GetTransformBase(int32 ControlIndex) const;
+	const FTransformBase * GetTransformBase(int32 TransformBaseIndex) const;
 	const FTransformBase * GetTransformBaseByNodeName(FName NodeName) const;
 	int32 FindTransformBaseByNodeName(FName NodeName) const;
 	int32 FindTransformParentNode(int32 NodeIndex, bool bTranslate, int32 Index=0) const;
diff --git a/Engine/Source/Runtime/Engine/Classes/Animation/Skeleton.h b/Engine/Source/Runtime/Engine/Classes/Animation/Skeleton.h
--- a/Engine/Source/Runtime/Engine/Classes/Animation/Skeleton.h
+++ b/Engine/Source/Runtime/Engine/Classes/Animation/Skeleton.h
@@ -632,6 +632,7 @@ public:
 	ENGINE_API void SetRigConfig(URig * Rig);
 	ENGINE_API FName GetRigBoneMapping(const FName & NodeName) const;
 	ENGINE_API bool SetRigBoneMapping(const FName & NodeName, FName BoneName);
+	ENGINE_API FName GetRigNodeNameFromBoneName(const FName & BoneName) const;
 	// this make sure it stays within the valid range
 	ENGINE_API int32 GetMappedValidNodes(TArray<FName> &OutValidNodeNames);
 	// verify if it has all latest data
diff --git a/Engine/Source/Runtime/Engine/Private/Animation/AnimSequence.cpp b/Engine/Source/Runtime/Engine/Private/Animation/AnimSequence.cpp
--- a/Engine/Source/Runtime/Engine/Private/Animation/AnimSequence.cpp
+++ b/Engine/Source/Runtime/Engine/Private/Animation/AnimSequence.cpp
@@ -1727,14 +1727,14 @@ void UAnimSequence::RemapTracksToNewSkeleton( USkeleton * NewSkeleton, bool bCon
 				}
 			}
 
-			TArray<struct FRawAnimSequenceTrack> RawRiggingAnimationData;
+			FAnimSequenceTrackContainer RiggingAnimationData;
 
 			// now convert animation data to rig data
-			ConvertAnimationDataToRiggingData(RawRiggingAnimationData);
+			ConvertAnimationDataToRiggingData(RiggingAnimationData);
 
 			// here we have to watch out the index
-			// The RawRiggingAnimationData will contain only the nodes that are mapped to source skeleton
-			// and here we convert everything that is in RawRiggingAnimationData which means based on source data
+			// The RiggingAnimationData will contain only the nodes that are mapped to source skeleton
+			// and here we convert everything that is in RiggingAnimationData which means based on source data
 			// when mapped back to new skeleton, it will discard results that are not mapped to target skeleton
 			
 			TArray<FName> SrcValidNodeNames;
@@ -1761,7 +1761,7 @@ void UAnimSequence::RemapTracksToNewSkeleton( USkeleton * NewSkeleton, bool bCon
 			{
 				int32 NodeIndex = Rig->FindNode(SrcValidNodeNames[SrcTrackIndex]);
 				check (NodeIndex != INDEX_NONE);
-				auto & RawAnimation = RawRiggingAnimationData[SrcTrackIndex];
+				auto & RawAnimation = RiggingAnimationData.AnimationTracks[SrcTrackIndex];
 
 				// find rotation parent node
 				int32 RotParentTrackIndex = FindValidTransformParentTrack(Rig, NodeIndex, false, SrcValidNodeNames);
@@ -1770,7 +1770,7 @@ void UAnimSequence::RemapTracksToNewSkeleton( USkeleton * NewSkeleton, bool bCon
 				for(int32 Key=0; Key<NumKeys; ++Key)
 				{
 					FTransform AnimatedLocalKey;
-					ExtractBoneTransform(RawRiggingAnimationData, AnimatedLocalKey, SrcTrackIndex, Interval*Key);
+					ExtractBoneTransform(RiggingAnimationData.AnimationTracks, AnimatedLocalKey, SrcTrackIndex, Interval*Key);
 
 					AnimatedLocalKey.ScaleTranslation(OldToNewTranslationRatio[NodeIndex]);
 
@@ -1799,7 +1799,7 @@ void UAnimSequence::RemapTracksToNewSkeleton( USkeleton * NewSkeleton, bool bCon
 			}
 
 			// now animation is converted to component space
-			TArray<struct FRawAnimSequenceTrack> NewRawAnimationData = RawRiggingAnimationData;
+			TArray<struct FRawAnimSequenceTrack> NewRawAnimationData = RiggingAnimationData.AnimationTracks;
 			for (int32 SrcTrackIndex=0; SrcTrackIndex<SrcNumTracks; ++SrcTrackIndex)
 			{
 				int32 NodeIndex = Rig->FindNode(SrcValidNodeNames[SrcTrackIndex]);
@@ -1851,51 +1851,20 @@ void UAnimSequence::RemapTracksToNewSkeleton( USkeleton * NewSkeleton, bool bCon
 					RawAnimation.PosKeys[Key] = ConvertedLocalSpaceAnimations[SrcTrackIndex][Key].GetLocation();
 					RawAnimation.RotKeys[Key] = ConvertedLocalSpaceAnimations[SrcTrackIndex][Key].GetRotation();
 					RawAnimation.ScaleKeys[Key] = ConvertedLocalSpaceAnimations[SrcTrackIndex][Key].GetScale3D();
-				}
-			}
 
-			// now add missing tracsk, before converting back to rigging data, we need whole nodes to be mapped
-			// @Todo fix this to do this earlier, just like to avoid confusion of so many index
-			if (SrcNumTracks  != Rig->GetNodeNum())
-			{
-				// if not all entries are filled up
-				// insert nodes
-				int32 NumNodes = Rig->GetNodeNum();
-				for (int32 NodeIndex=0; NodeIndex < NumNodes; ++NodeIndex)
-				{
-					FName NodeName = Rig->GetNodeName(NodeIndex);
-					// if not found
-					if (SrcValidNodeNames.Contains(NodeName) == false)
-					{
-						NewRawAnimationData.InsertZeroed(NodeIndex);
-
-						auto & RawAnimation = NewRawAnimationData[NodeIndex];
-						RawAnimation.PosKeys.Empty(NumKeys);
-						RawAnimation.PosKeys.AddUninitialized(NumKeys);
-						RawAnimation.RotKeys.Empty(NumKeys);
-						RawAnimation.RotKeys.AddUninitialized(NumKeys);
-						RawAnimation.ScaleKeys.Empty(NumKeys);
-						RawAnimation.ScaleKeys.AddUninitialized(NumKeys);
-
-						for(int32 Key=0; Key<NumKeys; ++Key)
-						{
-							RawAnimation.PosKeys[Key] = FVector::ZeroVector;
-							RawAnimation.RotKeys[Key] = FQuat::Identity;
-							RawAnimation.ScaleKeys[Key] = FVector(1.f);
-						}
-					}
+					// normalize rotation
+					RawAnimation.RotKeys[Key].Normalize();
 				}
 			}
 
-			RawRiggingAnimationData = NewRawAnimationData;
-			// here RawRiggingAnimationData always should contain # of nodes of entries
-			check (RawRiggingAnimationData.Num() == Rig->GetNodeNum());
+			RiggingAnimationData.AnimationTracks = MoveTemp(NewRawAnimationData);
+			RiggingAnimationData.TrackNames = MoveTemp(SrcValidNodeNames);
 
 			// set new skeleton
 			SetSkeleton(NewSkeleton);
 
 			// convert back to animated data with new skeleton
-			ConvertRiggingDataToAnimationData(RawRiggingAnimationData);
+			ConvertRiggingDataToAnimationData(RiggingAnimationData);
 		}
 		// @todo end rig testing
 		// @IMPORTANT: now otherwise this will try to do bone to bone mapping
@@ -2347,7 +2316,7 @@ int32 FindParentNodeIndex(URig * Rig, USkeleton * Skeleton, FName ParentNodeName
 	return Skeleton->GetReferenceSkeleton().FindBoneIndex(ParentBoneName);
 }
 
-int32 UAnimSequence::GetSpaceBasedAnimationData(TArray< TArray<FTransform> > & AnimationDataInComponentSpace, TArray<struct FRawAnimSequenceTrack> * RawRiggingAnimationData) const
+int32 UAnimSequence::GetSpaceBasedAnimationData(TArray< TArray<FTransform> > & AnimationDataInComponentSpace, FAnimSequenceTrackContainer * RiggingAnimationData) const
 {
 	USkeleton * Skeleton = GetSkeleton();
 
@@ -2368,31 +2337,28 @@ int32 UAnimSequence::GetSpaceBasedAnimationData(TArray< TArray<FTransform> > & A
 		AnimationDataInComponentSpace[BoneIndex].AddUninitialized(NumKeys);
 	}
 
-	if (RawRiggingAnimationData)
+	if (RiggingAnimationData)
 	{
 		const URig * Rig = Skeleton->GetRig();
 
 		check(Rig);
 
-		TArray<bool> bCompleteComponentSpaces;
-		bCompleteComponentSpaces.AddZeroed(NumBones);
-
 		// @Todo sort rig control by parent -> child, so that we can calculate parent first before child
 		// but that depends on each skeleton
-
-		const TArray<FTransformBase> & TransformBases = Rig->GetTransformBases();
 		// first put all the component space
-		for (const auto & TransformBase : TransformBases)
-		{
-			const FName & NodeName = TransformBase.Node;
-			const int32 & NodeIndex = Rig->FindNode(NodeName);
-			const FName & BoneName = Skeleton->GetRigBoneMapping(NodeName);
-			const int32 & BoneIndex = RefSkeleton.FindBoneIndex(BoneName);
 
-			if (BoneIndex != INDEX_NONE)
+		for(int32 BoneIndex = 0; BoneIndex < NumBones; ++BoneIndex)
+		{
+			const FName & BoneName = RefSkeleton.GetBoneName(BoneIndex);
+			const FName & NodeName = Skeleton->GetRigNodeNameFromBoneName(BoneName);
+			const FTransformBase * TransformBase = Rig->GetTransformBaseByNodeName(NodeName);
+			const int32 NodeIndex = RiggingAnimationData->TrackNames.Find(NodeName);
+			if (NodeIndex != INDEX_NONE)
 			{
+				check (TransformBase);
+
 				// now calculate the component space
-				const TArray<FRigTransformConstraint>	& RotTransformConstraints = TransformBase.Constraints[EControlConstraint::Type::Orientation].TransformConstraints;
+				const TArray<FRigTransformConstraint>	& RotTransformConstraints = TransformBase->Constraints[EControlConstraint::Type::Orientation].TransformConstraints;
 
 				FQuat ComponentRotation;
 				FTransform ComponentTranslation;
@@ -2401,84 +2367,69 @@ int32 UAnimSequence::GetSpaceBasedAnimationData(TArray< TArray<FTransform> > & A
 				//if (RotTransformConstraints.Num() > 0)
 				{
 					const FName & ParentNodeName = RotTransformConstraints[0].ParentSpace;
-					const int32 & ParentNodeIndex = Rig->FindNode(ParentNodeName);
 					const FName & ParentBoneName = Skeleton->GetRigBoneMapping(ParentNodeName);
 					const int32 & ParentBoneIndex = RefSkeleton.FindBoneIndex(ParentBoneName);
 
-					if (ParentBoneIndex != INDEX_NONE)
+					if(ParentBoneIndex != INDEX_NONE)
 					{
-						// make sure parent is calculated first
-						check(bCompleteComponentSpaces[ParentBoneIndex]);
-
-						for (int32 Key = 0; Key < NumKeys; ++Key)
+						for(int32 Key = 0; Key < NumKeys; ++Key)
 						{
-							ComponentRotation = AnimationDataInComponentSpace[ParentBoneIndex][Key].GetRotation() * (*RawRiggingAnimationData)[NodeIndex].RotKeys[Key];
+							ComponentRotation = AnimationDataInComponentSpace[ParentBoneIndex][Key].GetRotation() * RiggingAnimationData->AnimationTracks[NodeIndex].RotKeys[Key];
 							AnimationDataInComponentSpace[BoneIndex][Key].SetRotation(ComponentRotation);
 						}
 					}
 					else
 					{
-						for (int32 Key = 0; Key < NumKeys; ++Key)
+						for(int32 Key = 0; Key < NumKeys; ++Key)
 						{
-							ComponentRotation = (*RawRiggingAnimationData)[NodeIndex].RotKeys[Key];
+							ComponentRotation = RiggingAnimationData->AnimationTracks[NodeIndex].RotKeys[Key];
 							AnimationDataInComponentSpace[BoneIndex][Key].SetRotation(ComponentRotation);
 						}
 					}
-
-					bCompleteComponentSpaces[BoneIndex] = true;
 				}
 
-				const TArray<FRigTransformConstraint>	& PosTransformConstraints = TransformBase.Constraints[EControlConstraint::Type::Translation].TransformConstraints;
+				const TArray<FRigTransformConstraint>	& PosTransformConstraints = TransformBase->Constraints[EControlConstraint::Type::Translation].TransformConstraints;
 
 				//if (PosTransformConstraints.Num() > 0)
 				{
 					const FName & ParentNodeName = PosTransformConstraints[0].ParentSpace;
-					const int32 & ParentNodeIndex = Rig->FindNode(ParentNodeName);
 					const FName & ParentBoneName = Skeleton->GetRigBoneMapping(ParentNodeName);
 					const int32 & ParentBoneIndex = RefSkeleton.FindBoneIndex(ParentBoneName);
 
-					if (ParentBoneIndex != INDEX_NONE)
+					if(ParentBoneIndex != INDEX_NONE)
 					{
-						// make sure parent is calculated first
-						check(bCompleteComponentSpaces[ParentBoneIndex]);
-
-						for (int32 Key = 0; Key < NumKeys; ++Key)
+						for(int32 Key = 0; Key < NumKeys; ++Key)
 						{
-							ComponentTranslation = FTransform((*RawRiggingAnimationData)[NodeIndex].PosKeys[Key]) * AnimationDataInComponentSpace[ParentBoneIndex][Key];
+							ComponentTranslation = FTransform(RiggingAnimationData->AnimationTracks[NodeIndex].PosKeys[Key]) * AnimationDataInComponentSpace[ParentBoneIndex][Key];
 							AnimationDataInComponentSpace[BoneIndex][Key].SetTranslation(ComponentTranslation.GetTranslation());
 							AnimationDataInComponentSpace[BoneIndex][Key].SetScale3D(ComponentScale);
 						}
 					}
 					else
 					{
-						for (int32 Key = 0; Key < NumKeys; ++Key)
+						for(int32 Key = 0; Key < NumKeys; ++Key)
 						{
-							ComponentTranslation = FTransform((*RawRiggingAnimationData)[NodeIndex].PosKeys[Key]);
+							ComponentTranslation = FTransform(RiggingAnimationData->AnimationTracks[NodeIndex].PosKeys[Key]);
 							AnimationDataInComponentSpace[BoneIndex][Key].SetTranslation(ComponentTranslation.GetTranslation());
 							AnimationDataInComponentSpace[BoneIndex][Key].SetScale3D(ComponentScale);
 						}
 					}
 				}
 			}
-		}
-
-		for (int32 BoneIndex = 0; BoneIndex < NumBones; ++BoneIndex)
-		{
-			// fill up all not calculated ones
-			if (bCompleteComponentSpaces[BoneIndex] == false)
+			else
 			{
 				int32 ParentIndex = RefSkeleton.GetParentIndex(BoneIndex);
 				const FTransform & LocalSpace = RefSkeleton.GetRefBonePose()[BoneIndex];
-				if (ParentIndex != INDEX_NONE)
+				if(ParentIndex != INDEX_NONE)
 				{
-					for (int32 Key = 0; Key < NumKeys; ++Key)
+					for(int32 Key = 0; Key < NumKeys; ++Key)
 					{
 						AnimationDataInComponentSpace[BoneIndex][Key] = LocalSpace * AnimationDataInComponentSpace[ParentIndex][Key];
 					}
 				}
 				else
 				{
-					for (int32 Key = 0; Key < NumKeys; ++Key)
+					for(int32 Key = 0; Key < NumKeys; ++Key)
 					{
 						AnimationDataInComponentSpace[BoneIndex][Key] = LocalSpace;
 					}
@@ -2539,7 +2490,7 @@ int32 UAnimSequence::GetSpaceBasedAnimationData(TArray< TArray<FTransform> > & A
 	return AnimationDataInComponentSpace.Num();
 }
 
-bool UAnimSequence::ConvertAnimationDataToRiggingData(TArray<struct FRawAnimSequenceTrack> & RawRiggingAnimationData)
+bool UAnimSequence::ConvertAnimationDataToRiggingData(FAnimSequenceTrackContainer & RiggingAnimationData)
 {
 	USkeleton * Skeleton = GetSkeleton();
 	if (Skeleton && Skeleton->GetRig())
@@ -2552,13 +2503,12 @@ bool UAnimSequence::ConvertAnimationDataToRiggingData(TArray<struct FRawAnimSequ
 
 		if (NumBones > 0)
 		{
-			RawRiggingAnimationData.Empty(NumNodes);
-			RawRiggingAnimationData.AddZeroed(NumNodes);
+			RiggingAnimationData.Initialize(ValidNodeNames);
 
 			// first we copy all space bases back to it
 			for (int32 NodeIndex = 0; NodeIndex < NumNodes; ++NodeIndex)
 			{
-				struct FRawAnimSequenceTrack & Track = RawRiggingAnimationData[NodeIndex];
+				struct FRawAnimSequenceTrack & Track = RiggingAnimationData.AnimationTracks[NodeIndex];
 				const FName & NodeName = ValidNodeNames[NodeIndex];
 				const FName & BoneName = Skeleton->GetRigBoneMapping(NodeName);
 				const int32 & BoneIndex = Skeleton->GetReferenceSkeleton().FindBoneIndex(BoneName);
@@ -2671,18 +2621,30 @@ bool UAnimSequence::ConvertAnimationDataToRiggingData(TArray<struct FRawAnimSequ
 	return false;
 }
 
-bool UAnimSequence::ConvertRiggingDataToAnimationData(TArray<struct FRawAnimSequenceTrack> & RawRiggingAnimationData)
+bool UAnimSequence::ConvertRiggingDataToAnimationData(FAnimSequenceTrackContainer & RiggingAnimationData)
 {
-	if (RawRiggingAnimationData.Num() > 0)
+	if (RiggingAnimationData.GetNum() > 0)
 	{
 		TArray< TArray<FTransform> > AnimationDataInComponentSpace;
-		int32 NumBones = GetSpaceBasedAnimationData(AnimationDataInComponentSpace, &RawRiggingAnimationData);
+		int32 NumBones = GetSpaceBasedAnimationData(AnimationDataInComponentSpace, &RiggingAnimationData);
 
 		USkeleton * Skeleton = GetSkeleton();
 		TArray<FRawAnimSequenceTrack> OldAnimationData = RawAnimationData;
 		TArray<FName> OldAnimationTrackNames = AnimationTrackNames;
 		TArray<FName> ValidNodeNames;
-		int32 ValidNumNodes = Skeleton->GetMappedValidNodes(ValidNodeNames);
+		Skeleton->GetMappedValidNodes(ValidNodeNames);
+		// remove from ValidNodeNames if it doesn't belong to AnimationTrackNames
+		for (int32 NameIndex=0; NameIndex<ValidNodeNames.Num(); ++NameIndex)
+		{
+			if (RiggingAnimationData.TrackNames.Contains(ValidNodeNames[NameIndex]) == false)
+			{
+				ValidNodeNames.RemoveAt(NameIndex);
+				--NameIndex;
+			}
+		}
+
+		int32 ValidNumNodes = ValidNodeNames.Num();
+
 		// get local spaces
 		// add all tracks?
 		AnimationTrackNames.Empty(ValidNumNodes);
diff --git a/Engine/Source/Runtime/Engine/Private/Animation/Rig.cpp b/Engine/Source/Runtime/Engine/Private/Animation/Rig.cpp
--- a/Engine/Source/Runtime/Engine/Private/Animation/Rig.cpp
+++ b/Engine/Source/Runtime/Engine/Private/Animation/Rig.cpp
@@ -174,11 +174,11 @@ int32 URig::GetTransformBaseNum() const
 	return TransformBases.Num();
 }
 
-const FTransformBase * URig::GetTransformBase(int32 ControlIndex) const
+const FTransformBase * URig::GetTransformBase(int32 TransformBaseIndex) const
 {
-	if (TransformBases.IsValidIndex(ControlIndex))
+	if (TransformBases.IsValidIndex(TransformBaseIndex))
 	{
-		return &TransformBases[ControlIndex];
+		return &TransformBases[TransformBaseIndex];
 	}
 
 	return NULL;
@@ -186,11 +186,11 @@ const FTransformBase * URig::GetTransformBase(int32 ControlIndex) const
 
 const FTransformBase * URig::GetTransformBaseByNodeName(FName NodeName) const
 {
-	int32 ControlIndex = FindTransformBaseByNodeName(NodeName);
+	int32 TransformBaseIndex = FindTransformBaseByNodeName(NodeName);
 
-	if (ControlIndex != INDEX_NONE)
+	if (TransformBaseIndex != INDEX_NONE)
 	{
-		return &TransformBases[ControlIndex];
+		return &TransformBases[TransformBaseIndex];
 	}
 
 	return NULL;
diff --git a/Engine/Source/Runtime/Engine/Private/Animation/Skeleton.cpp b/Engine/Source/Runtime/Engine/Private/Animation/Skeleton.cpp
--- a/Engine/Source/Runtime/Engine/Private/Animation/Skeleton.cpp
+++ b/Engine/Source/Runtime/Engine/Private/Animation/Skeleton.cpp
@@ -1042,6 +1042,19 @@ FName USkeleton::GetRigBoneMapping(const FName & NodeName) const
 	return NAME_None;
 }
 
+FName USkeleton::GetRigNodeNameFromBoneName(const FName & BoneName) const
+{
+	for(const auto & NameMap : RigConfig.BoneMappingTable)
+	{
+		if(NameMap.BoneName == BoneName)
+		{
+			return NameMap.NodeName;
+		}
+	}
+
+	return NAME_None;
+}
+
 int32 USkeleton::GetMappedValidNodes(TArray<FName> &OutValidNodeNames)
 {
 	OutValidNodeNames.Empty();
