making multi object query return block to be more consistent with single.

diff --git a/Engine/Source/Runtime/Engine/Private/Collision/PhysXCollision.cpp b/Engine/Source/Runtime/Engine/Private/Collision/PhysXCollision.cpp
--- a/Engine/Source/Runtime/Engine/Private/Collision/PhysXCollision.cpp
+++ b/Engine/Source/Runtime/Engine/Private/Collision/PhysXCollision.cpp
@@ -494,7 +494,7 @@ namespace ECollisionQuery
 #define TRACE_MULTI		1
 #define TRACE_SINGLE	0
 
-PxSceneQueryHitType::Enum FPxQueryFilterCallback::CalcQueryHitType(const PxFilterData &PQueryFilter, const PxFilterData &PShapeFilter)
+PxSceneQueryHitType::Enum FPxQueryFilterCallback::CalcQueryHitType(const PxFilterData &PQueryFilter, const PxFilterData &PShapeFilter, bool bPreFilter)
 {
 	ECollisionQuery::Type QueryType = (ECollisionQuery::Type)PQueryFilter.word0;
 	ECollisionChannel const ShapeChannel = (ECollisionChannel)(PShapeFilter.word3 >> 24);
@@ -505,13 +505,14 @@ PxSceneQueryHitType::Enum FPxQueryFilterCallback::CalcQueryHitType(const PxFilte
 		// do I belong to one of objects of interest?
 		if ( ShapeBit & PQueryFilter.word1 )
 		{
-			if (MultiTrace)
+			if (bPreFilter)	//In the case of an object query we actually want to return all object types (or first in single case). So in PreFilter we have to trick physx by not blocking in the multi case, and blocking in the single case.
 			{
-				return PxSceneQueryHitType::eTOUCH;
+
+				return MultiTrace ? PxSceneQueryHitType::eTOUCH : PxSceneQueryHitType::eBLOCK;
 			}
 			else
 			{
-				return PxSceneQueryHitType::eBLOCK;
+				return PxSceneQueryHitType::eBLOCK;	//In the case where an object query is being resolved for the user we just return a block because object query doesn't have the concept of overlap at all and block seems more natural
 			}
 		}
 	}
@@ -610,7 +611,7 @@ PxSceneQueryHitType::Enum FPxQueryFilterCallback::preFilter(const PxFilterData&
 		return (PrefilterReturnValue = PxSceneQueryHitType::eNONE);
 	}
 
-	PxSceneQueryHitType::Enum Result = FPxQueryFilterCallback::CalcQueryHitType(filterData, ShapeFilter);
+	PxSceneQueryHitType::Enum Result = FPxQueryFilterCallback::CalcQueryHitType(filterData, ShapeFilter, true);
 
 	if(bSingleQuery && Result == PxSceneQueryHitType::eTOUCH)
 	{
diff --git a/Engine/Source/Runtime/Engine/Private/Collision/PhysXCollision.h b/Engine/Source/Runtime/Engine/Private/Collision/PhysXCollision.h
--- a/Engine/Source/Runtime/Engine/Private/Collision/PhysXCollision.h
+++ b/Engine/Source/Runtime/Engine/Private/Collision/PhysXCollision.h
@@ -107,7 +107,7 @@ public:
 	 *
 	 * @return PxSceneQueryHitType from both FilterData
 	 */
-	static PxSceneQueryHitType::Enum CalcQueryHitType(const PxFilterData &PQueryFilter, const PxFilterData &PShapeFilter);
+	static PxSceneQueryHitType::Enum CalcQueryHitType(const PxFilterData &PQueryFilter, const PxFilterData &PShapeFilter, bool bPreFilter = false);
 	
 	virtual PxSceneQueryHitType::Enum preFilter(const PxFilterData& filterData, const PxShape* shape, const PxRigidActor* actor, PxSceneQueryFlags& queryFlags) override;
 
