Fix incorrectly normalizing raycast dir for destructible component

diff --git a/Engine/Source/Runtime/Engine/Private/DestructibleComponent.cpp b/Engine/Source/Runtime/Engine/Private/DestructibleComponent.cpp
--- a/Engine/Source/Runtime/Engine/Private/DestructibleComponent.cpp
+++ b/Engine/Source/Runtime/Engine/Private/DestructibleComponent.cpp
@@ -1023,7 +1023,7 @@ bool UDestructibleComponent::LineTraceComponent( FHitResult& OutHit, const FVect
 		PxF32 HitTime = 0.0f;
 		PxVec3 HitNormal;
 		
-		int32 ChunkIdx = ApexDestructibleActor->rayCast(HitTime, HitNormal, U2PVector(Start), U2PVector((End-Start).SafeNormal()), NxDestructibleActorRaycastFlags::AllChunks);
+		int32 ChunkIdx = ApexDestructibleActor->rayCast(HitTime, HitNormal, U2PVector(Start), U2PVector(End-Start), NxDestructibleActorRaycastFlags::AllChunks);
 
 		if (ChunkIdx != NxModuleDestructibleConst::INVALID_CHUNK_INDEX && HitTime <= 1.0f)
 		{
