Expose ragdoll simulation on dedicated server as project setting

diff --git a/Engine/Source/Runtime/Engine/Classes/Components/SkeletalMeshComponent.h b/Engine/Source/Runtime/Engine/Classes/Components/SkeletalMeshComponent.h
--- a/Engine/Source/Runtime/Engine/Classes/Components/SkeletalMeshComponent.h
+++ b/Engine/Source/Runtime/Engine/Classes/Components/SkeletalMeshComponent.h
@@ -335,9 +335,11 @@ public:
 	UPROPERTY(transient)
 	uint32 bBlendPhysics:1;
 
-	/** temporary variable to enable/disable skeletalmeshcomponent physics SIM/blending in dedicated server 
-	 * @todo this should move to more global option **/
-	UPROPERTY()
+	/**
+	 *  If true, simulate physics for this component on a dedicated server.
+	 *  This should be set if simulating physics and replicating with a dedicated server.
+	 */
+	UPROPERTY(EditAnywhere, AdvancedDisplay, BlueprintReadWrite, Category = SkeletalMesh)
 	uint32 bEnablePhysicsOnDedicatedServer:1;
 
 	/**
diff --git a/Engine/Source/Runtime/Engine/Classes/PhysicsEngine/PhysicsSettings.h b/Engine/Source/Runtime/Engine/Classes/PhysicsEngine/PhysicsSettings.h
--- a/Engine/Source/Runtime/Engine/Classes/PhysicsEngine/PhysicsSettings.h
+++ b/Engine/Source/Runtime/Engine/Classes/PhysicsEngine/PhysicsSettings.h
@@ -107,6 +107,13 @@ class ENGINE_API UPhysicsSettings
 	UPROPERTY(config, EditAnywhere, Category = Simulation)
 	float MaxDepenetrationVelocity;
 
+	/**
+	*  If true, simulate physics for this component on a dedicated server.
+	*  This should be set if simulating physics and replicating with a dedicated server.
+	*/
+	UPROPERTY(config, EditAnywhere, Category = Simulation)
+	bool bSimulateSkeletalMeshOnDedicatedServer;
+
 
 	/** Max Physics Delta Time to be clamped. */
 	UPROPERTY(config, EditAnywhere, meta=(ClampMin="0.0013", UIMin = "0.0013", ClampMax="1.0", UIMax="1.0"), Category=Framerate)
diff --git a/Engine/Source/Runtime/Engine/Private/PhysicsEngine/PhysicsSettings.cpp b/Engine/Source/Runtime/Engine/Private/PhysicsEngine/PhysicsSettings.cpp
--- a/Engine/Source/Runtime/Engine/Private/PhysicsEngine/PhysicsSettings.cpp
+++ b/Engine/Source/Runtime/Engine/Private/PhysicsEngine/PhysicsSettings.cpp
@@ -9,6 +9,7 @@ UPhysicsSettings::UPhysicsSettings(const class FPostConstructInitializePropertie
 	, DefaultFluidFriction(0.3f)
 	, bEnableAsyncScene(false)
 	, bEnable2DPhysics(false)
+	, bSimulateSkeletalMeshOnDedicatedServer(true)
 	, MaxPhysicsDeltaTime(1.f / 30.f)
 	, bSubstepping(false)
 	, MaxSubstepDeltaTime(1.f / 60.f)
diff --git a/Engine/Source/Runtime/Engine/Private/SkeletalMeshComponent.cpp b/Engine/Source/Runtime/Engine/Private/SkeletalMeshComponent.cpp
--- a/Engine/Source/Runtime/Engine/Private/SkeletalMeshComponent.cpp
+++ b/Engine/Source/Runtime/Engine/Private/SkeletalMeshComponent.cpp
@@ -105,6 +105,7 @@ USkeletalMeshComponent::USkeletalMeshComponent(const class FPostConstructInitial
 	PreClothTickFunction.bCanEverTick = true;
 	PreClothTickFunction.bStartWithTickEnabled = true;
 
+
 #if WITH_APEX_CLOTHING
 	ClothMaxDistanceScale = 1.0f;
 	ClothTeleportMode = FClothingActor::Continuous;
@@ -126,7 +127,7 @@ USkeletalMeshComponent::USkeletalMeshComponent(const class FPostConstructInitial
 
 	DefaultPlayRate_DEPRECATED = 1.0f;
 	bDefaultPlaying_DEPRECATED = true;
-	bEnablePhysicsOnDedicatedServer = false;
+	bEnablePhysicsOnDedicatedServer = UPhysicsSettings::Get()->bSimulateSkeletalMeshOnDedicatedServer;
 	bEnableUpdateRateOptimizations = false;
 
 	bTickInEditor = true;
