Fix material expression mask issue again.

diff --git a/Engine/Source/Runtime/Core/Public/UObject/ObjectVersion.h b/Engine/Source/Runtime/Core/Public/UObject/ObjectVersion.h
--- a/Engine/Source/Runtime/Core/Public/UObject/ObjectVersion.h
+++ b/Engine/Source/Runtime/Core/Public/UObject/ObjectVersion.h
@@ -618,6 +618,8 @@ enum EUnrealEngineObjectUE4Version
 	// Rename USpringArmComponent::bUsePawnViewRotation to bUsePawnControlRotation
 	// Rename UCameraComponent::bUsePawnViewRotation to bUsePawnControlRotation
 	VER_UE4_RENAME_CAMERA_COMPONENT_CONTROL_ROTATION,
+	// Fix bad refraction material attribute masks
+	VER_UE4_FIX_REFRACTION_INPUT_MASKING,
 
 	// -----<new versions can be added before this line>-------------------------------------------------
 	// - this needs to be the last line (see note below)
diff --git a/Engine/Source/Runtime/Engine/Private/Materials/MaterialShared.cpp b/Engine/Source/Runtime/Engine/Private/Materials/MaterialShared.cpp
--- a/Engine/Source/Runtime/Engine/Private/Materials/MaterialShared.cpp
+++ b/Engine/Source/Runtime/Engine/Private/Materials/MaterialShared.cpp
@@ -55,25 +55,19 @@ int32 FExpressionInput::Compile(class FMaterialCompiler* Compiler, int32 Multipl
 	if(Expression)
 	{
 		Expression->ValidateState();
+		
+		int32 ExpressionResult = Compiler->CallExpression(FMaterialExpressionKey(Expression, OutputIndex, MultiplexIndex),Compiler);
 
-		if(Mask)
+		if(Mask && ExpressionResult != INDEX_NONE)
 		{
-			int32 ExpressionResult = Compiler->CallExpression(FMaterialExpressionKey(Expression, OutputIndex, MultiplexIndex),Compiler);
-			if(ExpressionResult != INDEX_NONE)
-			{
-				return Compiler->ComponentMask(
-					ExpressionResult,
-					!!MaskR,!!MaskG,!!MaskB,!!MaskA
-					);
-			}
-			else
-			{
-				return INDEX_NONE;
-			}
+			return Compiler->ComponentMask(
+				ExpressionResult,
+				!!MaskR,!!MaskG,!!MaskB,!!MaskA
+				);
 		}
 		else
 		{
-			return Compiler->CallExpression(FMaterialExpressionKey(Expression, OutputIndex, MultiplexIndex),Compiler);
+			return ExpressionResult;
 		}
 	}
 	else
@@ -2292,7 +2286,7 @@ void DoMaterialAttributeReorder(FExpressionInput* Input, int32 UE4Ver)
 			}
 		}
 		
-		if( UE4Ver < VER_UE4_CHANGED_MATERIAL_REFACTION_TYPE && Input->OutputIndex == 13 )
+		if( UE4Ver < VER_UE4_FIX_REFRACTION_INPUT_MASKING && Input->OutputIndex == 13 )
 		{
 			Input->Mask = 1;
 			Input->MaskR = 1;
