AActor::SetLifeSpan checks if actor has PendingKill flag.
https://udn.unrealengine.com/questions/220862/destroyed-actor-access-leads-to-crash.html

diff --git a/Engine/Source/Runtime/Engine/Private/Actor.cpp b/Engine/Source/Runtime/Engine/Private/Actor.cpp
--- a/Engine/Source/Runtime/Engine/Private/Actor.cpp
+++ b/Engine/Source/Runtime/Engine/Private/Actor.cpp
@@ -3264,7 +3264,7 @@ void AActor::SetLifeSpan( float InLifespan )
 	// Store the new value
 	InitialLifeSpan = InLifespan;
 	// Initialize a timer for the actors lifespan if there is one. Otherwise clear any existing timer
-	if( Role==ROLE_Authority || bTearOff )
+	if ((Role == ROLE_Authority || bTearOff) && !IsPendingKill())
 	{
 		if( InLifespan > 0.0f)
 		{
