Core - Added a command "Obj MemSub -class=[classname] -cull=[no of kb]" that prints memory stats for all objects of the specified class, this is inclusive stats, includes all subobjects
Core - Added an option " -class=[classname]" to "Obj Mem" command, which tries to narrow results to the specified class

diff --git a/Engine/Source/Runtime/Engine/Private/UnrealEngine.cpp b/Engine/Source/Runtime/Engine/Private/UnrealEngine.cpp
--- a/Engine/Source/Runtime/Engine/Private/UnrealEngine.cpp
+++ b/Engine/Source/Runtime/Engine/Private/UnrealEngine.cpp
@@ -1989,73 +1989,57 @@ static void ShowSubobjectGraph( FOutputDevice& Ar, UObject* CurrentObject, const
 			Ar.Logf(TEXT("%s+ %s"), *IndentString, IndentString.Len() == 0 ? *CurrentObject->GetPathName() : *CurrentObject->GetName());
 			for ( int32 ObjIndex = 0; ObjIndex < ReferencedObjs.Num(); ObjIndex++ )
 			{
-				ShowSubobjectGraph(Ar, ReferencedObjs[ObjIndex], IndentString + TEXT("|\t"));
+				ShowSubobjectGraph( Ar, ReferencedObjs[ObjIndex], IndentString + TEXT( "|\t" ) );
 			}
 		}
 	}
 }
-
-/** Holds information about memory usage. */
-struct FMemItem
+struct FItem
 {
-	int32		Count;
-	SIZE_T	Num, Max, Res;
-	UObject* Object;
+	UClass*	Class;
+	int32 Count;
+	SIZE_T Num;
+	SIZE_T Max;
+	/** Resource size of the object and all of its references, the 'old-style'. */
+	SIZE_T ResourceSize;
+	/** Only exclusive resource size, the truer resource size. */
+	SIZE_T TrueResourceSize;
 
-	FMemItem()
-		: Count(0)
-		, Num(0)
-		, Max(0)
-		, Res(0)
-		, Object( NULL )
+	FItem( UClass* InClass=NULL )
+	: Class(InClass), Count(0), Num(0), Max(0), ResourceSize(0), TrueResourceSize(0)
 	{}
 
-	FMemItem( UObject* InObject, SIZE_T	InRes )
-		: Count(0)
-		, Num(0)
-		, Max(0)
-		, Res(InRes)
-		, Object( InObject )
+	FItem( UClass* InClass, int32 InCount, SIZE_T InNum, SIZE_T InMax, SIZE_T InResourceSize, SIZE_T InTrueResourceSize ) :
+		Class( InClass ),
+		Count( InCount ),
+		Num( InNum ), 
+		Max( InMax ), 
+		ResourceSize( InResourceSize ), 
+		TrueResourceSize( InTrueResourceSize )
 	{}
 
-	void Add( FArchiveCountMem& Ar, SIZE_T InRes )
-	{
-		Count ++; 
-		Num += Ar.GetNum(); 
-		Max += Ar.GetMax(); 
-		Res += InRes;
-	}
-
-	void AddRes( SIZE_T InRes )
-	{
-		Count ++; 
-		Res += InRes;
-	}
-};
-
-struct FItem
-{
-	UClass*	Class;
-	int32		Count;
-	SIZE_T	Num, Max, Res, TrueRes;
-	FItem( UClass* InClass=NULL )
-	: Class(InClass), Count(0), Num(0), Max(0), Res(0), TrueRes(0)
-	{}
-	void Add( FArchiveCountMem& Ar, SIZE_T InRes, SIZE_T InTrueRes )
+	void Add( FArchiveCountMem& Ar, SIZE_T InResourceSize, SIZE_T InTrueResourceSize )
 	{
 		Count++;
 		Num += Ar.GetNum();
 		Max += Ar.GetMax();
-		Res += InRes;
-		TrueRes += InTrueRes;
+		ResourceSize += InResourceSize;
+		TrueResourceSize += InTrueResourceSize;
 	}
 };
+
 struct FSubItem
 {
 	UObject* Object;
-	SIZE_T Num, Max, Res, TrueRes;
-	FSubItem( UObject* InObject, SIZE_T InNum, SIZE_T InMax, SIZE_T InRes, SIZE_T InTrueRes )
-	: Object( InObject ), Num( InNum ), Max( InMax ), Res( InRes ), TrueRes( InTrueRes )
+	SIZE_T Num;
+	SIZE_T Max;
+	/** Resource size of the object and all of its references, the 'old-style'. */
+	SIZE_T ResourceSize;
+	/** Only exclusive resource size, the truer resource size. */
+	SIZE_T TrueResourceSize;
+
+	FSubItem( UObject* InObject, SIZE_T InNum, SIZE_T InMax, SIZE_T InResourceSize, SIZE_T InTrueResourceSize )
+	: Object( InObject ), Num( InNum ), Max( InMax ), ResourceSize( InResourceSize ), TrueResourceSize( InTrueResourceSize )
 	{}
 };
 
@@ -4298,15 +4282,15 @@ struct FHierarchyNode
 	TSet<UObject*> Items;
 	int64 Inc;
 	int64 Exc;
-	int32 IncCnt;
-	int32 ExcCnt;
+	int32 IncCount;
+	int32 ExcCount;
 	FHierarchyNode()
 		: This(NULL)
 		, Parent(NULL)
 		, Inc(-1)
 		, Exc(-1)
-		, IncCnt(-1)
-		, ExcCnt(-1)
+		, IncCount(-1)
+		, ExcCount(-1)
 	{
 	}
 	bool operator< (FHierarchyNode const& Other) const
@@ -4379,45 +4363,45 @@ struct FHierarchy
 			AddClass((UClass*)This);
 		}
 	}
-	FHierarchyNode& Compute(UObject* This, TMap<UObject*, FSubItem> const& Objects, bool bCntItems)
+	FHierarchyNode& Compute(UObject* This, TMap<UObject*, FSubItem> const& Objects, bool bCountItems)
 	{
 		FHierarchyNode& Node = Nodes.FindChecked(This);
 		if (Node.Inc < 0)
 		{
 			Node.Exc = 0;
-			Node.ExcCnt = 1;
+			Node.ExcCount = 1;
 			if (This)
 			{
 				FSubItem const& Item = Objects.FindChecked(This);
 				Node.Exc += Item.Max;
-				Node.Exc += Item.TrueRes;
-				if (bCntItems)
+				Node.Exc += Item.TrueResourceSize;
+				if (bCountItems)
 				{
-					Node.ExcCnt += Node.Items.Num();
+					Node.ExcCount += Node.Items.Num();
 				}
 				else
 				{
-					Node.ExcCnt += Node.Children.Num();
+					Node.ExcCount += Node.Children.Num();
 				}
 			}
 			Node.Inc = Node.Exc;
-			Node.IncCnt = Node.ExcCnt;
+			Node.IncCount = Node.ExcCount;
 			for (TSet<UObject*>::TConstIterator It(Node.Children); It; ++It)
 			{
-				FHierarchyNode& Child = Compute(*It, Objects, bCntItems);
+				FHierarchyNode& Child = Compute(*It, Objects, bCountItems);
 				Node.Inc += Child.Inc;
-				if (!bCntItems)
+				if (!bCountItems)
 				{
-					Node.IncCnt += Child.IncCnt;
+					Node.IncCount += Child.IncCount;
 				}
 			}
 			for (TSet<UObject*>::TConstIterator It(Node.Items); It; ++It)
 			{
-				FHierarchyNode& Child = Compute(*It, Objects, bCntItems);
+				FHierarchyNode& Child = Compute(*It, Objects, bCountItems);
 				Node.Inc += Child.Inc;
-				if (bCntItems)
+				if (bCountItems)
 				{
-					Node.IncCnt += Child.IncCnt;
+					Node.IncCount += Child.IncCount;
 				}
 			}
 		}
@@ -4432,7 +4416,7 @@ struct FHierarchy
 		}
 		Out.Sort();
 	}
-	FString Size(uint64 Mem)
+	static FString Size(uint64 Mem)
 	{
 		if (Mem / 1024 < 10000)
 		{
@@ -4444,14 +4428,22 @@ struct FHierarchy
 		}
 		return FString::Printf(TEXT("%4lldG"), Mem / (1024*1024*1024));
 	}
-	void LogSet(TSet<UObject*> const& In, bool bCntItems, int Indent)
+	void LogSet(TSet<UObject*> const& In, UClass* ClassToCheck, bool bCntItems, int Indent)
 	{
 		TArray<FHierarchyNode> Children;
 		SortSet(In, Children);
 		int32 Index = 0;
 		for (; Index < Children.Num(); Index++)
 		{
-			if (!Log(Children[Index].This, bCntItems, Indent + 1, Index + 1 < Children.Num()))
+			UObject* Child = Children[Index].This;
+			// Only makes sense for flat hierarchy.
+			const bool bIsClassToCheck = Child->IsA( ClassToCheck );
+			if( !bIsClassToCheck )
+			{
+				continue;
+			}
+
+			if( !Log( Child, ClassToCheck, bCntItems, Indent + 1, Index + 1 < Children.Num() ) )
 			{
 				break;
 			}
@@ -4462,21 +4454,21 @@ struct FHierarchy
 			FHierarchyNode Extra;
 			Extra.Exc = 0;
 			Extra.Inc = 0;
-			Extra.ExcCnt = 0;
-			Extra.IncCnt = 0;
+			Extra.ExcCount = 0;
+			Extra.IncCount = 0;
 			for (; Index < Children.Num(); Index++)
 			{
 				Extra.Exc += Children[Index].Exc;
 				Extra.Inc += Children[Index].Inc;
-				Extra.ExcCnt += Children[Index].ExcCnt;
-				Extra.IncCnt += Children[Index].IncCnt;
+				Extra.ExcCount += Children[Index].ExcCount;
+				Extra.IncCount += Children[Index].IncCount;
 				NumExtra++;
 			}
-			FString Line = FString::Printf(TEXT("%s        %5d %s (%d)"), *Size(Extra.Inc), Extra.IncCnt, TEXT("More"), NumExtra);
+			FString Line = FString::Printf(TEXT("%s        %5d %s (%d)"), *Size(Extra.Inc), Extra.IncCount, TEXT("More"), NumExtra);
 			UE_LOG(LogEngine, Log, TEXT("%s%s") , FCString::Spc(2 * (Indent + 1)), *Line);
 		}
 	}
-	bool Log(UObject* This, bool bCntItems, int Indent = 0, bool bAllowCull = true)
+	bool Log(UObject* This, UClass* ClassToCheck, bool bCountItems, int Indent = 0, bool bAllowCull = true)
 	{
 		FHierarchyNode& Node = Nodes.FindChecked(This);
 		if (bAllowCull && Node.Inc < Limit && Node.Exc < Limit)
@@ -4486,30 +4478,31 @@ struct FHierarchy
 		FString Line;
 		if (Node.IsLeaf())
 		{
-			Line = FString::Printf(TEXT("%s        %5d %s"), *Size(Node.Inc), Node.IncCnt, Node.This ? *Node.This->GetFullName() : TEXT("Root"));
+			Line = FString::Printf(TEXT("%s        %5d %s"), *Size(Node.Inc), Node.IncCount, Node.This ? *Node.This->GetFullName() : TEXT("Root"));
 			UE_LOG(LogEngine, Log, TEXT("%s%s") , FCString::Spc(2 * Indent), *Line);
 		}
 		else
 		{
-			Line = FString::Printf(TEXT("%s %sx %5d %s"), *Size(Node.Inc), *Size(Node.Exc), Node.IncCnt, Node.This ? *Node.This->GetFullName() : TEXT("Root"));
+			Line = FString::Printf(TEXT("%s %sx %5d %s"), *Size(Node.Inc), *Size(Node.Exc), Node.IncCount, Node.This ? *Node.This->GetFullName() : TEXT("Root"));
 			UE_LOG(LogEngine, Log, TEXT("%s%s") , FCString::Spc(2 * Indent), *Line);
-			if (bCntItems && Node.Children.Num())
+			if (bCountItems && Node.Children.Num())
 			{
 				UE_LOG(LogEngine, Log, TEXT("%s%s") , FCString::Spc(2 * (Indent + 1)), TEXT("Child Classes"));
 			}
-			LogSet(Node.Children, bCntItems, Indent + 2);
+			LogSet(Node.Children, ClassToCheck, bCountItems, Indent + 2);
 
-			if (bCntItems && Node.Items.Num())
+			if (bCountItems && Node.Items.Num())
 			{
 				UE_LOG(LogEngine, Log, TEXT("%s%s") , FCString::Spc(2 * (Indent + 1)), TEXT("Instances"));
 			}
-			LogSet(Node.Items, bCntItems, Indent);
+			LogSet(Node.Items, ClassToCheck, bCountItems, Indent);
 		}
 
 		return true;
 	}
 };
 
+// @TODO yrx 2014-09-15 Move to ObjectCommads.cpp or ObjectExec.cpp
 bool UEngine::HandleObjCommand( const TCHAR* Cmd, FOutputDevice& Ar )
 {
 	if (FParse::Command(&Cmd,TEXT("LIST2")))
@@ -4526,13 +4519,163 @@ bool UEngine::HandleObjCommand( const TCHAR* Cmd, FOutputDevice& Ar )
 		MemAnalyze.PrintResults(Ar, FObjectMemoryAnalyzer::EPrintFlags::PrintReferences);
 		return true;
 	}
-	if (FParse::Command(&Cmd,TEXT("Mem")))
-	{			
+	else if (FParse::Command(&Cmd,TEXT("MemSub")))
+	{
+		struct FCompareByInclusiveSize
+		{
+			FORCEINLINE bool operator()( const FItem& A, const FItem& B ) const 
+			{ 
+				return A.Max > B.Max; 
+			}
+		};
+
+		struct FLocal
+		{
+			static void GetReferencedObjs( UObject* CurrentObject, TArray<UObject*>& out_ReferencedObjs )
+			{
+				TArray<UObject*> ReferencedObjs;
+				FReferenceFinder RefCollector( ReferencedObjs, CurrentObject, true, false, false, false );
+				RefCollector.FindReferences( CurrentObject );
+
+				out_ReferencedObjs.Append( ReferencedObjs );
+				for( UObject*& RefObj : ReferencedObjs )
+				{
+					GetReferencedObjs( RefObj, out_ReferencedObjs );
+				}
+			}
+		};
+
+		int32 Limit = 16;
+		FParse::Value(Cmd, TEXT("CULL="), Limit);
+		Limit *= 1024;
+
+		UClass* ClassToCheck = NULL;
+		ParseObject<UClass>(Cmd, TEXT("CLASS="  ), ClassToCheck, ANY_PACKAGE );
+		if (ClassToCheck == NULL)
+		{
+			ClassToCheck = UObject::StaticClass();
+		}
+
+		TMap<UClass*,FItem> ObjectsByClass;
+
+		UE_LOG( LogEngine, Log, TEXT("**********************************************") );
+		UE_LOG( LogEngine, Log, TEXT("Obj MemSub for class '%s'"), *ClassToCheck->GetName() );
+		UE_LOG( LogEngine, Log, TEXT("") );
+
+		for( FObjectIterator It(ClassToCheck); It; ++It )
+		{
+			UObject* Obj = *It;
+			if( Obj->IsTemplate( RF_ClassDefaultObject ) )
+			{
+				continue;
+			}
+
+			// Get references.
+			TArray<UObject*> ReferencedObjects;
+			FLocal::GetReferencedObjs( Obj, ReferencedObjects );
+
+			// Calculate memory usage.
+			FItem ThisObject( Obj->GetClass() );
+			for( UObject*& RefObj : ReferencedObjects )
+			{
+				FArchiveCountMem Count( RefObj );
+				// Get the 'old-style' resource size and the truer resource size
+				const SIZE_T ResourceSize = It->GetResourceSize( EResourceSizeMode::Inclusive );
+				const SIZE_T TrueResourceSize = It->GetResourceSize( EResourceSizeMode::Exclusive );
+				ThisObject.Add( Count, ResourceSize, TrueResourceSize );
+			}
+
+			FItem& ClassObjects = ObjectsByClass.FindOrAdd( ThisObject.Class );
+			ClassObjects.Count++;
+			ClassObjects.Num += ThisObject.Num;
+			ClassObjects.Max += ThisObject.Max;
+			ClassObjects.ResourceSize += ThisObject.ResourceSize;
+			ClassObjects.TrueResourceSize += ThisObject.TrueResourceSize;
+		}
+
+		ObjectsByClass.ValueSort( FCompareByInclusiveSize() );
+
+		UE_LOG( LogEngine, Log, TEXT("%32s, %6s %6s %6s %6s, %6s"), 
+			TEXT("Class"),
+			TEXT("IncMax"),
+			TEXT("IncNum"),
+			TEXT("ResInc"),
+			TEXT("ResExc"),
+			TEXT("Count") );
+
+		FItem Total;
+		FItem Culled;
+		for( const auto& It : ObjectsByClass )
+		{
+			UClass* Class = It.Key;
+			const FItem& ClassObjects = It.Value;
+
+			if( ClassObjects.Max < Limit )
+			{
+				Culled.Count += ClassObjects.Count;
+				Culled.Num += ClassObjects.Num;
+				Culled.Max += ClassObjects.Max;
+				Culled.ResourceSize += ClassObjects.ResourceSize;
+				Culled.TrueResourceSize += ClassObjects.TrueResourceSize;
+			}
+			else
+			{
+				UE_LOG( LogEngine, Log, TEXT("%32s, %6s %6s %6s %6s, %6i"), 
+					*Class->GetName(), 
+					*FHierarchy::Size(ClassObjects.Max), *FHierarchy::Size(ClassObjects.Num), 
+					*FHierarchy::Size(ClassObjects.ResourceSize), *FHierarchy::Size(ClassObjects.TrueResourceSize),
+					ClassObjects.Count 
+					);
+
+			}
+
+			Total.Count += ClassObjects.Count;
+			Total.Num += ClassObjects.Num;
+			Total.Max += ClassObjects.Max;
+			Total.ResourceSize += ClassObjects.ResourceSize;
+			Total.TrueResourceSize += ClassObjects.TrueResourceSize;
+		}
+
+		if( Culled.Count > 0 )
+		{
+			UE_LOG( LogEngine, Log, TEXT("") );
+			UE_LOG( LogEngine, Log, TEXT("%32s, %6s %6s %6s %6s, %6i"), 
+				TEXT("(Culled)"), 
+				*FHierarchy::Size(Culled.Max), *FHierarchy::Size(Culled.Num), 
+				*FHierarchy::Size(Culled.ResourceSize), *FHierarchy::Size(Culled.TrueResourceSize),
+				Culled.Count 
+				);
+		}
 
+		UE_LOG( LogEngine, Log, TEXT("") );
+		UE_LOG( LogEngine, Log, TEXT("%32s, %6s %6s %6s %6s, %6i"), 
+			TEXT("Total"), 
+			*FHierarchy::Size(Total.Max), *FHierarchy::Size(Total.Num), 
+			*FHierarchy::Size(Total.ResourceSize), *FHierarchy::Size(Total.TrueResourceSize),
+			Total.Count 
+			);
+		UE_LOG( LogEngine, Log, TEXT("**********************************************") );
+		return true;
+	}
+	else if (FParse::Command(&Cmd,TEXT("Mem")))
+	{
 		int32 Limit = 50;
 		FParse::Value(Cmd, TEXT("CULL="), Limit);
 		Limit *= 1024;
 
+		UClass* ClassToCheck = NULL;
+		ParseObject<UClass>(Cmd, TEXT("CLASS="  ), ClassToCheck, ANY_PACKAGE );
+
+		if (ClassToCheck == NULL)
+		{
+			ClassToCheck = UObject::StaticClass();
+		}
+		else
+		{
+			// Class is set, so lower a bit the limit.
+			Limit /= 10;
+		}
+
 		FHierarchy Classes(Limit);
 		FHierarchy Outers(Limit);
 		FHierarchy Flat(Limit);
@@ -4549,16 +4692,19 @@ bool UEngine::HandleObjCommand( const TCHAR* Cmd, FOutputDevice& Ar )
 			Outers.AddOuter(*It);
 			Flat.AddFlat(*It);
 		}
+
 		UE_LOG(LogEngine, Log, TEXT("********************************************** By Outer Hierarchy") );
 		Outers.Compute(NULL, Objects, false);
-		Outers.Log(NULL, false);
+		Outers.Log( ClassToCheck, UObject::StaticClass(), false );
+		
 
 		UE_LOG(LogEngine, Log, TEXT("********************************************** By Class Hierarchy") );
 		Classes.Compute(NULL, Objects, true);
-		Classes.Log(NULL, true);
+		Classes.Log( ClassToCheck, UObject::StaticClass(), true );
+
 		UE_LOG(LogEngine, Log, TEXT("********************************************** Flat") );
 		Flat.Compute(NULL, Objects, false);
-		Flat.Log(NULL, false);
+		Flat.Log( NULL, ClassToCheck, false );
 		UE_LOG(LogEngine, Log, TEXT("**********************************************") );
 
 		return true;
@@ -4763,7 +4909,7 @@ bool UEngine::HandleObjCommand( const TCHAR* Cmd, FOutputDevice& Ar )
 				}					
 				//MSEND
 
-				Ar.Logf( TEXT("%140s % 10iK % 10iK % 10iK % 10iK"), *ObjItem.Object->GetFullName(), (int32)ObjItem.Num / 1024, (int32)ObjItem.Max / 1024, (int32)ObjItem.Res / 1024, (int32)ObjItem.TrueRes / 1024 );
+				Ar.Logf( TEXT("%140s % 10iK % 10iK % 10iK % 10iK"), *ObjItem.Object->GetFullName(), (int32)ObjItem.Num / 1024, (int32)ObjItem.Max / 1024, (int32)ObjItem.ResourceSize / 1024, (int32)ObjItem.TrueResourceSize / 1024 );
 			}
 			Ar.Log( TEXT("") );
 		}
@@ -4786,11 +4932,11 @@ bool UEngine::HandleObjCommand( const TCHAR* Cmd, FOutputDevice& Ar )
 
 			for( int32 i=0; i<List.Num(); i++ )
 			{
-				Ar.Logf(TEXT(" %100s % 6i % 10iK % 10iK % 10iK % 10iK"), *List[i].Class->GetName(), (int32)List[i].Count, (int32)(List[i].Num/1024), (int32)(List[i].Max/1024), (int32)(List[i].Res/1024), (int32)(List[i].TrueRes/1024) );
+				Ar.Logf(TEXT(" %100s % 6i % 10iK % 10iK % 10iK % 10iK"), *List[i].Class->GetName(), (int32)List[i].Count, (int32)(List[i].Num/1024), (int32)(List[i].Max/1024), (int32)(List[i].ResourceSize/1024), (int32)(List[i].TrueResourceSize/1024) );
 			}
 			Ar.Log( TEXT("") );
 		}
-		Ar.Logf( TEXT("%i Objects (%.3fM / %.3fM / %.3fM / %.3fM)"), Total.Count, (float)Total.Num/1024.0/1024.0, (float)Total.Max/1024.0/1024.0, (float)Total.Res/1024.0/1024.0, (float)Total.TrueRes/1024.0/1024.0 );
+		Ar.Logf( TEXT("%i Objects (%.3fM / %.3fM / %.3fM / %.3fM)"), Total.Count, (float)Total.Num/1024.0/1024.0, (float)Total.Max/1024.0/1024.0, (float)Total.ResourceSize/1024.0/1024.0, (float)Total.TrueResourceSize/1024.0/1024.0 );
 		return true;
 
 	}
