CIS Fix

diff --git a/Engine/Source/Runtime/Engine/Private/UnrealEngine.cpp b/Engine/Source/Runtime/Engine/Private/UnrealEngine.cpp
--- a/Engine/Source/Runtime/Engine/Private/UnrealEngine.cpp
+++ b/Engine/Source/Runtime/Engine/Private/UnrealEngine.cpp
@@ -4610,7 +4610,7 @@ bool UEngine::HandleObjCommand( const TCHAR* Cmd, FOutputDevice& Ar )
 			UClass* Class = It.Key;
 			const FItem& ClassObjects = It.Value;
 
-			if( ClassObjects.Max < Limit )
+			if( ClassObjects.Max < (SIZE_T)Limit )
 			{
 				Culled.Count += ClassObjects.Count;
 				Culled.Num += ClassObjects.Num;
