Fix BSP raycasts not working when initially overlapping. Solution is to force BSP to use double sided collision

diff --git a/Engine/Source/Runtime/Engine/Private/ModelComponent.cpp b/Engine/Source/Runtime/Engine/Private/ModelComponent.cpp
--- a/Engine/Source/Runtime/Engine/Private/ModelComponent.cpp
+++ b/Engine/Source/Runtime/Engine/Private/ModelComponent.cpp
@@ -243,6 +243,8 @@ void UModelComponent::PostLoad()
 	{
 		ModelBodySetup->bGenerateMirroredCollision = false;
 	}
+
+	ModelBodySetup->bDoubleSidedGeometry = true;	//saved content wants this to be true
 }
 
 #if WITH_EDITOR
@@ -507,6 +509,7 @@ void UModelComponent::CreateModelBodySetup()
 
 	ModelBodySetup->CollisionTraceFlag = CTF_UseComplexAsSimple;
 	ModelBodySetup->bGenerateMirroredCollision = false;
+	ModelBodySetup->bDoubleSidedGeometry = true;
 }
 
 #if WITH_EDITOR
