Actually removing a test variable from the matinee perf capture function.

diff --git a/Engine/Source/Runtime/Engine/Private/Tests/AutomationTestCommon.cpp b/Engine/Source/Runtime/Engine/Private/Tests/AutomationTestCommon.cpp
--- a/Engine/Source/Runtime/Engine/Private/Tests/AutomationTestCommon.cpp
+++ b/Engine/Source/Runtime/Engine/Private/Tests/AutomationTestCommon.cpp
@@ -88,7 +88,7 @@ bool FMatineePerformanceCaptureCommand::Update()
 	for (TObjectIterator<AMatineeActor> It; It; ++It)
 	{
 		AMatineeActor* MatineeActor = *It;
-		FString MatineeFOOName = MatineeActor->GetName();
+
 		if (MatineeActor->GetName().Equals(MatineeName,ESearchCase::IgnoreCase))
 		{	
 			//add latent action to execute this matinee
