Removed "Title" property from WorldSettings.
#ttp 345480 - NUX: Level Editing: Why does World Settings have a "Title" property?
#branch UE4
#proj Runtime.Engine

diff --git a/Engine/Source/Runtime/Engine/Classes/GameFramework/WorldSettings.h b/Engine/Source/Runtime/Engine/Classes/GameFramework/WorldSettings.h
--- a/Engine/Source/Runtime/Engine/Classes/GameFramework/WorldSettings.h
+++ b/Engine/Source/Runtime/Engine/Classes/GameFramework/WorldSettings.h
@@ -218,10 +218,6 @@ class ENGINE_API AWorldSettings : public AInfo, public IInterface_AssetUserData
 	UPROPERTY(EditAnywhere, BlueprintReadOnly, meta=(DisplayName = "Override World Gravity"), Category = Physics)
 	uint32 bGlobalGravitySet:1;
 
-	/** title of the map displayed in the UI */
-	UPROPERTY(EditAnywhere, BlueprintReadOnly, localized, Category=World)
-	FString Title;
-
 	// any actor falling below this level gets destroyed
 	UPROPERTY(EditAnywhere, BlueprintReadOnly, Category=World, meta=(editcondition = "bEnableWorldBoundsChecks"))
 	float KillZ;
