Renamed text references from TextureGroup to LODGroup in Texture Editor.
#ttp 345448 - NUX:Texture Editor: Some combo boxes say "FromTextureGroup" but the actual setting is called "LOD Group"
#branch UE4
#proj Runtime.Engine

diff --git a/Engine/Source/Runtime/Engine/Classes/Engine/Texture.h b/Engine/Source/Runtime/Engine/Classes/Engine/Texture.h
--- a/Engine/Source/Runtime/Engine/Classes/Engine/Texture.h
+++ b/Engine/Source/Runtime/Engine/Classes/Engine/Texture.h
@@ -31,8 +31,8 @@ enum TextureFilter
 	TF_Nearest UMETA(DisplayName="Nearest"),
 	TF_Bilinear UMETA(DisplayName="Bi-linear"),
 	TF_Trilinear UMETA(DisplayName="Tri-linear"),
-	/** use setting from the LOD TextureGroup */
-	TF_Default UMETA(DisplayName="Default (from TextureGroup)"),
+	/** use setting from the LOD Group */
+	TF_Default UMETA(DisplayName="Default (from LOD Group)"),
 	TF_MAX,
 };
 
@@ -504,7 +504,7 @@ public:
 	UPROPERTY(EditAnywhere, BlueprintReadWrite, Category=Texture, AssetRegistrySearchable, AdvancedDisplay)
 	TEnumAsByte<enum TextureFilter> Filter;
 
-	/** Texture group this texture belongs to for LOD bias */
+	/** LOD group this texture belongs to for LOD bias */
 	UPROPERTY(EditAnywhere, BlueprintReadWrite, Category=LevelOfDetail, meta=(DisplayName="LOD Group"), AssetRegistrySearchable)
 	TEnumAsByte<enum TextureGroup> LODGroup;
 
diff --git a/Engine/Source/Runtime/Engine/Classes/Engine/TextureDefines.h b/Engine/Source/Runtime/Engine/Classes/Engine/TextureDefines.h
--- a/Engine/Source/Runtime/Engine/Classes/Engine/TextureDefines.h
+++ b/Engine/Source/Runtime/Engine/Classes/Engine/TextureDefines.h
@@ -56,7 +56,7 @@ UENUM()
 enum TextureMipGenSettings
 {
 	// default for the "texture"
-	TMGS_FromTextureGroup UMETA(DisplayName="FromTextureGroup"),
+	TMGS_FromTextureGroup UMETA(DisplayName="FromLODGroup"),
 	// 2x2 average, default for the "texture group"
 	TMGS_SimpleAverage UMETA(DisplayName="SimpleAverage"),
 	// 8x8 with sharpening: 0=no sharpening but better quality which is softer, 1..little, 5=medium, 10=extreme
