Fix case where encroaching geometry blocks on an overlap. OverlapTest was using object query instead of channel query which always assumes blocking

diff --git a/Engine/Source/Runtime/Engine/Private/LevelActor.cpp b/Engine/Source/Runtime/Engine/Private/LevelActor.cpp
--- a/Engine/Source/Runtime/Engine/Private/LevelActor.cpp
+++ b/Engine/Source/Runtime/Engine/Private/LevelActor.cpp
@@ -708,7 +708,7 @@ bool UWorld::EncroachingBlockingGeometry(AActor* TestActor, FVector TestLocation
 			{
 				// must be registered
 				FComponentQueryParams Params(NAME_EncroachingBlockingGeometry, TestActor);
-				bFoundBlockingHit = ComponentOverlapMulti(Overlaps, PrimComp, TestLocation, TestActor->GetActorRotation(), Params);
+				bFoundBlockingHit = ComponentOverlapMulti(Overlaps, PrimComp, TestLocation, TestActor->GetActorRotation(), BlockingChannel, Params);
 			}
 			else
 			{
