Fix crash when selecting root bone and Parent Space on transform bone node

#TTP 344009 - PERSONA: Crash: AnimGraph: Transform Bone with no parent in parent space

diff --git a/Engine/Source/Runtime/Engine/Private/Animation/AnimationRuntime.cpp b/Engine/Source/Runtime/Engine/Private/Animation/AnimationRuntime.cpp
--- a/Engine/Source/Runtime/Engine/Private/Animation/AnimationRuntime.cpp
+++ b/Engine/Source/Runtime/Engine/Private/Animation/AnimationRuntime.cpp
@@ -1067,8 +1067,11 @@ void FAnimationRuntime::ConvertCSTransformToBoneSpace
 		case BCS_ParentBoneSpace :
 			{
 				const int32 ParentIndex = MeshBases.GetParentBoneIndex(BoneIndex);
-				const FTransform ParentTM = MeshBases.GetComponentSpaceTransform(ParentIndex);
-				CSBoneTM.SetToRelativeTransform(ParentTM);
+				if (ParentIndex != INDEX_NONE)
+				{
+					const FTransform ParentTM = MeshBases.GetComponentSpaceTransform(ParentIndex);
+					CSBoneTM.SetToRelativeTransform(ParentTM);
+				}
 			}
 			break;
 
