-Add RepNotify conditions.
-Remove client side equality check for calling repnotifies on custom delta serilaized properties

diff --git a/Engine/Source/Runtime/CoreUObject/Public/UObject/CoreNet.h b/Engine/Source/Runtime/CoreUObject/Public/UObject/CoreNet.h
--- a/Engine/Source/Runtime/CoreUObject/Public/UObject/CoreNet.h
+++ b/Engine/Source/Runtime/CoreUObject/Public/UObject/CoreNet.h
@@ -182,6 +182,12 @@ enum ELifetimeCondition
 	COND_Max				= 9,
 };
 
+enum ELifetimeRepNotifyCondition
+{
+	REPNOTIFY_OnChanged		= 0,		// Only call the property's RepNotify function if it changes from the local value
+	REPNOTIFY_Always		= 1,		// Always Call the property's RepNotify function when it is received from the server
+};
+
 /** FLifetimeProperty
  *	This class is used to track a property that is marked to be replicated for the lifetime of the actor channel.
  *  This doesn't mean the property will necessarily always be replicated, it just means:
@@ -193,16 +199,18 @@ class FLifetimeProperty
 public:
 	uint16				RepIndex;
 	ELifetimeCondition	Condition;
+	ELifetimeRepNotifyCondition RepNotifyCondition;
 
-	FLifetimeProperty() : RepIndex( 0 ), Condition( COND_None ) {}
-	FLifetimeProperty( int32 InRepIndex ) : RepIndex( InRepIndex ), Condition( COND_None ) { check( InRepIndex <= 65535 ); }
-	FLifetimeProperty( int32 InRepIndex, ELifetimeCondition InCondition ) : RepIndex( InRepIndex ), Condition( InCondition ) { check( InRepIndex <= 65535 ); }
+	FLifetimeProperty() : RepIndex( 0 ), Condition( COND_None ), RepNotifyCondition(REPNOTIFY_OnChanged) {}
+	FLifetimeProperty( int32 InRepIndex ) : RepIndex( InRepIndex ), Condition( COND_None ), RepNotifyCondition(REPNOTIFY_OnChanged) { check( InRepIndex <= 65535 ); }
+	FLifetimeProperty(int32 InRepIndex, ELifetimeCondition InCondition, ELifetimeRepNotifyCondition InRepNotifyCondition=REPNOTIFY_OnChanged) : RepIndex(InRepIndex), Condition(InCondition), RepNotifyCondition(InRepNotifyCondition) { check(InRepIndex <= 65535); }
 
 	inline bool operator==( const FLifetimeProperty & Other ) const
 	{
 		if ( RepIndex == Other.RepIndex )
 		{
 			check( Condition == Other.Condition );		// Can't have different conditions if the RepIndex matches, doesn't make sense
+			check( RepNotifyCondition == Other.RepNotifyCondition);
 			return true;
 		}
 
diff --git a/Engine/Source/Runtime/Engine/Private/DataReplication.cpp b/Engine/Source/Runtime/Engine/Private/DataReplication.cpp
--- a/Engine/Source/Runtime/Engine/Private/DataReplication.cpp
+++ b/Engine/Source/Runtime/Engine/Private/DataReplication.cpp
@@ -494,14 +494,6 @@ bool FObjectReplicator::ReceivedBunch( FInBunch &Bunch, const FReplicationFlags
 				TArray<uint8>	MetaData;
 				PTRINT Offset = 0;
 
-				// Copy current value over to Recent for comparison
-				if ( DestRecent )
-				{
-					Offset = ReplicatedProp->ContainerPtrToValuePtr<uint8>(DestRecent, Element) - DestRecent;
-					check( Offset >= 0 && Offset < RepState->StaticBuffer.Num() ); //@todo if we move properties outside of the memory block, then this will not work anyway
-					ReplicatedProp->CopySingleValue( DestRecent + Offset, Data );
-				}
-
 				// Receive custom delta property.
 				UStructProperty * StructProperty = Cast< UStructProperty >( ReplicatedProp );
 
@@ -545,26 +537,11 @@ bool FObjectReplicator::ReceivedBunch( FInBunch &Bunch, const FReplicationFlags
 					return false;
 				}
 
-				// See if it changed from our local value
-				bool PropertyChanged = true;
-
-				if ( DestRecent )
-				{
-					// POD types can do a memcmp with a call to Identical
-					if ( ReplicatedProp->Identical( DestRecent + Offset, Data ) )
-					{
-						PropertyChanged = false;
-					}
-				}
-
 				// Successfully received it.
-				UE_LOG( LogNetTraffic, Log, TEXT( " %s - %s - Change: %d" ), *Object->GetName(), *ReplicatedProp->GetName(), PropertyChanged );
+				UE_LOG( LogNetTraffic, Log, TEXT( " %s - %s" ), *Object->GetName(), *ReplicatedProp->GetName());
 
 				// Notify the Object if this var is RepNotify
-				if ( PropertyChanged )
-				{
-					QueuePropertyRepNotify( Object, ReplicatedProp, Element, MetaData );
-				}
+				QueuePropertyRepNotify( Object, ReplicatedProp, Element, MetaData );
 			}	
 			
 			// Next.
diff --git a/Engine/Source/Runtime/Engine/Private/RepLayout.cpp b/Engine/Source/Runtime/Engine/Private/RepLayout.cpp
--- a/Engine/Source/Runtime/Engine/Private/RepLayout.cpp
+++ b/Engine/Source/Runtime/Engine/Private/RepLayout.cpp
@@ -1051,7 +1051,7 @@ bool FRepLayout::ReadProperty(
 		Cmd.Property->NetSerializeItem( ReaderState.Bunch, ReaderState.Bunch.PackageMap, Data + SwappedCmd.Offset );
 
 		// Check to see if this property changed
-		if ( !PropertiesAreIdentical( Cmd, StoredData + Cmd.Offset, Data + SwappedCmd.Offset ) )
+		if (Parent.RepNotifyCondition == REPNOTIFY_Always || !PropertiesAreIdentical(Cmd, StoredData + Cmd.Offset, Data + SwappedCmd.Offset))
 		{
 			ReaderState.RepState->RepNotifies.AddUnique( Parent.Property );
 		}
@@ -1157,7 +1157,7 @@ bool FRepLayout::ReceiveProperties_DynamicArray_r(
 
 	const FRepParentCmd & Parent = Parents[Cmd.ParentIndex];
 
-	if ( Array->Num() != ArrayNum && Parent.Property->HasAnyPropertyFlags( CPF_RepNotify ) )
+	if ( Array->Num() != ArrayNum && (Parent.RepNotifyCondition == REPNOTIFY_Always || Parent.Property->HasAnyPropertyFlags( CPF_RepNotify ) ) )
 	{
 		ReaderState.RepState->RepNotifies.AddUnique( Parent.Property );
 	}
@@ -1323,7 +1323,7 @@ void FRepLayout::UpdateUnmappedObjects_r(
 
 			UObject * OldObject = ObjProperty->GetObjectPropertyValue( Data + AbsOffset );
 
-			if ( OldObject != Object )
+			if (Parent.RepNotifyCondition == REPNOTIFY_Always || OldObject != Object)
 			{
 				if ( !bOutSomeObjectsWereMapped )
 				{
@@ -1756,7 +1756,7 @@ void FRepLayout::DiffProperties_r(
 		const FRepLayoutCmd & SwappedCmd = Cmd;//Parent.RoleSwapIndex != -1 ? Cmds[Parents[Parent.RoleSwapIndex].CmdStart] : Cmd;
 
 		// Make the shadow state match the actual state at the time of send
-		if ( !PropertiesAreIdentical( Cmd, (const void*)( Data + SwappedCmd.Offset ), (const void*)( StoredData + Cmd.Offset ) ) )
+		if ( Parent.RepNotifyCondition == REPNOTIFY_Always || !PropertiesAreIdentical( Cmd, (const void*)( Data + SwappedCmd.Offset ), (const void*)( StoredData + Cmd.Offset ) ) )
 		{
 			bOutDifferent = true;
 
@@ -2107,6 +2107,7 @@ void FRepLayout::InitFromObjectClass( UClass * InObjectClass )
 		{
 			// Store the condition on the parent in case we need it
 			Parents[LifetimeProps[i].RepIndex].Condition = LifetimeProps[i].Condition;
+			Parents[LifetimeProps[i].RepIndex].RepNotifyCondition = LifetimeProps[i].RepNotifyCondition;
 
 			if ( Parents[LifetimeProps[i].RepIndex].Flags & PARENT_IsCustomDelta )
 			{
diff --git a/Engine/Source/Runtime/Engine/Public/Net/RepLayout.h b/Engine/Source/Runtime/Engine/Public/Net/RepLayout.h
--- a/Engine/Source/Runtime/Engine/Public/Net/RepLayout.h
+++ b/Engine/Source/Runtime/Engine/Public/Net/RepLayout.h
@@ -178,6 +178,7 @@ public:
 		CmdEnd( 0 ), 
 		RoleSwapIndex( -1 ), 
 		Condition( COND_None ),
+		RepNotifyCondition(REPNOTIFY_OnChanged),
 		Flags( 0 )
 	{}
 
@@ -187,6 +188,7 @@ public:
 	uint16				CmdEnd;
 	int32				RoleSwapIndex;
 	ELifetimeCondition	Condition;
+	ELifetimeRepNotifyCondition	RepNotifyCondition;
 
 	uint32				Flags;
 };
diff --git a/Engine/Source/Runtime/Engine/Public/Net/UnrealNetwork.h b/Engine/Source/Runtime/Engine/Public/Net/UnrealNetwork.h
--- a/Engine/Source/Runtime/Engine/Public/Net/UnrealNetwork.h
+++ b/Engine/Source/Runtime/Engine/Public/Net/UnrealNetwork.h
@@ -94,6 +94,16 @@ static UProperty* GetReplicatedProperty(UClass* CallingClass, UClass* PropClass,
 	}																					\
 }
 
+/** Allows gamecode to specify RepNotify condition: REPNOTIFY_OnChanged (default) or REPNOTIFY_Always for when repnotify function is called  */
+#define DOREPLIFETIME_CONDITION_NOTIFY(c,v,cond, rncond) \
+{ \
+	static UProperty* sp##v = GetReplicatedProperty(StaticClass(), c::StaticClass(),GET_MEMBER_NAME_CHECKED(c,v)); \
+	for ( int32 i = 0; i < sp##v->ArrayDim; i++ )										\
+	{																					\
+		OutLifetimeProps.AddUnique( FLifetimeProperty( sp##v->RepIndex + i, cond, rncond) );	\
+	}																					\
+}
+
 #define DOREPLIFETIME_ACTIVE_OVERRIDE(c,v,active)	\
 {													\
 	static UProperty* sp##v = GetReplicatedProperty(StaticClass(), c::StaticClass(),GET_MEMBER_NAME_CHECKED(c,v)); \
