Logic fix

diff --git a/Engine/Source/Runtime/Engine/Private/RepLayout.cpp b/Engine/Source/Runtime/Engine/Private/RepLayout.cpp
--- a/Engine/Source/Runtime/Engine/Private/RepLayout.cpp
+++ b/Engine/Source/Runtime/Engine/Private/RepLayout.cpp
@@ -1157,7 +1157,7 @@ bool FRepLayout::ReceiveProperties_DynamicArray_r(
 
 	const FRepParentCmd & Parent = Parents[Cmd.ParentIndex];
 
-	if ( Array->Num() != ArrayNum && (Parent.RepNotifyCondition == REPNOTIFY_Always || Parent.Property->HasAnyPropertyFlags( CPF_RepNotify ) ) )
+	if ( (Array->Num() != ArrayNum || Parent.RepNotifyCondition == REPNOTIFY_Always) && Parent.Property->HasAnyPropertyFlags( CPF_RepNotify ) )
 	{
 		ReaderState.RepState->RepNotifies.AddUnique( Parent.Property );
 	}
