Remove redundent (and flawed) overlap check. The calls to OverlapMulti and ComponentOverlapMulti should suffice.

TTP#345428

diff --git a/Engine/Source/Runtime/Engine/Private/LevelActor.cpp b/Engine/Source/Runtime/Engine/Private/LevelActor.cpp
--- a/Engine/Source/Runtime/Engine/Private/LevelActor.cpp
+++ b/Engine/Source/Runtime/Engine/Private/LevelActor.cpp
@@ -716,14 +716,6 @@ bool UWorld::EncroachingBlockingGeometry(AActor* TestActor, FVector TestLocation
 			}
 		}
 
-		for( int32 HitIdx = 0; HitIdx < Overlaps.Num(); HitIdx++ )
-		{
-			if ( Overlaps[HitIdx].Component.IsValid() && (Overlaps[HitIdx].Component.Get()->GetCollisionResponseToChannel(BlockingChannel) == ECR_Block) )
-			{
-				bFoundBlockingHit = true;
-				break;
-			}
-		}
 		if ( !bFoundBlockingHit )
 		{
 			return false;
