Fix UpdateShapeFilter ignoring the actor's collision enable flag when using a skeletal mesh component

TTP#345425

diff --git a/Engine/Source/Runtime/Engine/Private/PhysicsEngine/BodyInstance.cpp b/Engine/Source/Runtime/Engine/Private/PhysicsEngine/BodyInstance.cpp
--- a/Engine/Source/Runtime/Engine/Private/PhysicsEngine/BodyInstance.cpp
+++ b/Engine/Source/Runtime/Engine/Private/PhysicsEngine/BodyInstance.cpp
@@ -793,9 +793,16 @@ void FBodyInstance::UpdatePhysicsFilterData(bool bForceSimpleAsComplex)
 
 		// In skeletal case, collision enable/disable/movement should be overriden by mesh component
 		// being in the physics asset, and not having collision is a waste and it can cause a bug where disconnected bodies
-		UseCollisionEnabled = SkelMeshComp->BodyInstance.CollisionEnabled; 
-		ObjectType = SkelMeshComp->GetCollisionObjectType();
+		if (Owner && Owner->GetActorEnableCollision())	//we only override if actor has collision enabled
+		{
+			UseCollisionEnabled = SkelMeshComp->BodyInstance.CollisionEnabled;
+		}
+		else
+		{	//actor has collision disabled so disable regardless of component setting
+			UseCollisionEnabled = ECollisionEnabled::NoCollision;
+		}
 
+		ObjectType = SkelMeshComp->GetCollisionObjectType();
 		bUseCollisionEnabledOverride = true;
 
 		if (BodySetup->CollisionReponse == EBodyCollisionResponse::BodyCollision_Enabled)
