To the FMatineePerformanceFCaptureCommand I added a log that will output the name of the matinee that is not being used.
AutomationTestCommon.cpp

diff --git a/Engine/Source/Runtime/Engine/Private/Tests/AutomationTestCommon.cpp b/Engine/Source/Runtime/Engine/Private/Tests/AutomationTestCommon.cpp
--- a/Engine/Source/Runtime/Engine/Private/Tests/AutomationTestCommon.cpp
+++ b/Engine/Source/Runtime/Engine/Private/Tests/AutomationTestCommon.cpp
@@ -88,9 +88,11 @@ bool FMatineePerformanceCaptureCommand::Update()
 	for (TObjectIterator<AMatineeActor> It; It; ++It)
 	{
 		AMatineeActor* MatineeActor = *It;
-
+		FString MatineeFOOName = MatineeActor->GetName();
 		if (MatineeActor->GetName().Equals(MatineeName,ESearchCase::IgnoreCase))
 		{	
+
+
 			//add latent action to execute this matinee
 			ADD_LATENT_AUTOMATION_COMMAND(FPlayMatineeLatentCommand(MatineeActor));
 
@@ -103,6 +105,10 @@ bool FMatineePerformanceCaptureCommand::Update()
 			//Stop the Stat FPS Chart command
 			ADD_LATENT_AUTOMATION_COMMAND(FExecWorldStringLatentCommand(TEXT("StopFPSChart")));
 		}
+		else
+		{
+			UE_LOG(LogEngineAutomationLatentCommand, Log, TEXT("'%s' is not the matinee name that is being searched for."), *MatineeActor->GetName())
+		}
 	}
 
 	return true;
