remove unused pointer

diff --git a/Engine/Source/Runtime/Engine/Private/DataReplication.cpp b/Engine/Source/Runtime/Engine/Private/DataReplication.cpp
--- a/Engine/Source/Runtime/Engine/Private/DataReplication.cpp
+++ b/Engine/Source/Runtime/Engine/Private/DataReplication.cpp
@@ -488,8 +488,6 @@ bool FObjectReplicator::ReceivedBunch( FInBunch &Bunch, const FReplicationFlags
 				}
 
 				// Pointer to destination.
-				uint8 * DestRecent = RepState->StaticBuffer.GetTypedData();
-
 				uint8 * Data = ReplicatedProp->ContainerPtrToValuePtr<uint8>((uint8*)Object, Element);
 				TArray<uint8>	MetaData;
 				PTRINT Offset = 0;
