Initial demo record/playback support

diff --git a/Engine/Config/BaseEngine.ini b/Engine/Config/BaseEngine.ini
--- a/Engine/Config/BaseEngine.ini
+++ b/Engine/Config/BaseEngine.ini
@@ -153,6 +153,7 @@ MinDesiredFrameRate=35.000000
 InitialButtonRepeatDelay=0.2
 ButtonRepeatDelay=0.1
 NetDriverDefinitions=(DefName="GameNetDriver",DriverClassName="/Script/OnlineSubsystemUtils.IpNetDriver",DriverClassNameFallback="/Script/OnlineSubsystemUtils.IpNetDriver")
++NetDriverDefinitions=(DefName="DemoNetDriver",DriverClassName="/Script/Engine.DemoNetDriver",DriverClassNameFallback="/Script/Engine.DemoNetDriver")
 
 bSmoothFrameRate=true
 SmoothedFrameRateRange=(LowerBound=(Type="ERangeBoundTypes::Inclusive",Value=22),UpperBound=(Type="ERangeBoundTypes::Exclusive",Value=62))
@@ -1442,6 +1443,10 @@ LanServerMaxTickRate=35
 NetConnectionClassName="/Script/OnlineSubsystemUtils.IpConnection"
 MaxPortCountToTry=512
 
+[/Script/Engine.DemoNetDriver]
+NetConnectionClassName="/Script/Engine.DemoNetConnection"
+DemoSpectatorClass=Engine.PlayerController
+
 [TextureStreaming]
 NeverStreamOutTextures=False
 MinTextureResidentMipCount=7
diff --git a/Engine/Source/Runtime/Engine/Classes/Engine/DemoNetConnection.h b/Engine/Source/Runtime/Engine/Classes/Engine/DemoNetConnection.h
new file mode 100644
--- /dev/null
+++ b/Engine/Source/Runtime/Engine/Classes/Engine/DemoNetConnection.h
@@ -0,0 +1,30 @@
+// Copyright 1998-2014 Epic Games, Inc. All Rights Reserved.
+
+//
+// Simulated network connection for recording and playing back game sessions.
+//
+
+#pragma once
+#include "DemoNetConnection.generated.h"
+
+UCLASS(transient, config=Engine)
+class UDemoNetConnection : public UNetConnection
+{
+	GENERATED_UCLASS_BODY()
+
+	// Begin UNetConnection interface.
+	virtual void		InitConnection( class UNetDriver* InDriver, EConnectionState InState, const FURL& InURL, int32 InConnectionSpeed = 0 ) override;
+	virtual FString		LowLevelGetRemoteAddress( bool bAppendPort=false ) override;
+	virtual FString		LowLevelDescribe() override;
+	virtual void		LowLevelSend( void* Data, int32 Count ) override;
+	virtual int32		IsNetReady( bool Saturate ) override;
+	virtual void		FlushNet( bool bIgnoreSimulation = false ) override;
+	virtual void		HandleClientPlayer( APlayerController* PC, class UNetConnection* NetConnection ) override;
+	virtual bool		ClientHasInitializedLevelFor( const UObject* TestObject ) override;
+	// End UNetConnection interface.
+
+	/** @return The DemoRecording driver object */
+	FORCEINLINE class UDemoNetDriver* GetDriver() { return (UDemoNetDriver*)Driver; }
+};
+
+
diff --git a/Engine/Source/Runtime/Engine/Classes/Engine/DemoNetDriver.h b/Engine/Source/Runtime/Engine/Classes/Engine/DemoNetDriver.h
new file mode 100644
--- /dev/null
+++ b/Engine/Source/Runtime/Engine/Classes/Engine/DemoNetDriver.h
@@ -0,0 +1,67 @@
+// Copyright 1998-2014 Epic Games, Inc. All Rights Reserved.
+
+//
+// Simulated network driver for recording and playing back game sessions.
+#pragma once
+#include "DemoNetDriver.generated.h"
+
+UCLASS(transient, config=Engine)
+class UDemoNetDriver : public UNetDriver
+{
+	GENERATED_UCLASS_BODY()
+
+	// Variables.
+
+	/** Name of the file to read/write from */
+	FString			DemoFilename;
+
+	/** Handle to the archive that will read/write network packets */
+	FArchive*		FileAr;
+
+	/** @todo document */
+	int32			DemoFrameNum;
+
+	/** Last time (in real seconds) that we recorded a frame */
+	double			LastRecordTime;
+
+	/** Time (in game seconds) that have elapsed between recorded frames */
+	float			DemoDeltaTime;
+
+	/** during playback, set to total number of frames recorded in the demo */
+	int32			PlaybackTotalFrames;
+
+	UPROPERTY( config )
+	FString			DemoSpectatorClass;
+
+	// Begin UNetDriver interface.
+	virtual bool InitBase( bool bInitAsClient, FNetworkNotify* InNotify, const FURL& URL, bool bReuseAddressAndPort, FString& Error ) override;
+	virtual void FinishDestroy() override;
+	virtual FString LowLevelGetNetworkNumber() override;
+	virtual bool InitConnect( FNetworkNotify* InNotify, const FURL& ConnectURL, FString& Error ) override;
+	virtual bool InitListen( FNetworkNotify* InNotify, FURL& ListenURL, bool bReuseAddressAndPort, FString& Error ) override;
+	virtual void TickDispatch( float DeltaSeconds ) override;
+	virtual void TickFlush( float DeltaSeconds ) override;
+	virtual void ProcessRemoteFunction( class AActor* Actor, class UFunction* Function, void* Parameters, struct FOutParmRec* OutParms, struct FFrame* Stack, class UObject* SubObject = NULL );
+	virtual bool IsAvailable() const override { return true; }
+	// End UNetDriver interface.
+
+	// Begin FExec interface.
+	virtual bool Exec( UWorld* InWorld, const TCHAR* Cmd, FOutputDevice& Ar ) override;
+	// End FExec interface.
+
+	/** @todo document */
+	bool UpdateDemoTime( float* DeltaTime, float TimeDilation );
+
+	/** Called when demo playback finishes, either because we reached the end of the file or because the demo spectator was destroyed */
+	void DemoPlaybackEnded();
+
+	/** @return true if the net resource is valid or false if it should not be used */
+	virtual bool IsNetResourceValid(void) { return true; }
+
+	void TickDemoRecord( float DeltaSeconds );
+	bool ReadDemoFrame();
+	void TickDemoPlayback( float DeltaSeconds );
+	void SpawnDemoRecSpectator( UNetConnection* Connection );
+
+	void StopDemo();
+};
diff --git a/Engine/Source/Runtime/Engine/Classes/Engine/DemoPendingNetGame.h b/Engine/Source/Runtime/Engine/Classes/Engine/DemoPendingNetGame.h
new file mode 100644
--- /dev/null
+++ b/Engine/Source/Runtime/Engine/Classes/Engine/DemoPendingNetGame.h
@@ -0,0 +1,10 @@
+// Copyright 1998-2014 Epic Games, Inc. All Rights Reserved.
+
+#pragma once
+#include "DemoPendingNetGame.generated.h"
+
+UCLASS(transient, config=Engine)
+class UDemoPendingNetGame : public UPendingNetGame
+{
+	GENERATED_UCLASS_BODY()
+};
diff --git a/Engine/Source/Runtime/Engine/Classes/Engine/World.h b/Engine/Source/Runtime/Engine/Classes/Engine/World.h
--- a/Engine/Source/Runtime/Engine/Classes/Engine/World.h
+++ b/Engine/Source/Runtime/Engine/Classes/Engine/World.h
@@ -539,6 +539,10 @@ class ENGINE_API UWorld : public UObject, public FNetworkNotify
 	UPROPERTY(Transient)
 	class ULevel*								CurrentLevelPendingVisibility;
 	
+	/** Fake NetDriver for capturing network traffic to record demos															*/
+	UPROPERTY()
+	class UDemoNetDriver*						DemoNetDriver;
+
 	/** Particle event manager **/
 	UPROPERTY()
 	class AParticleEventManager*				MyParticleEventManager;
@@ -2008,7 +2012,13 @@ public:
 	bool HandleTraceTagCommand( const TCHAR* Cmd, FOutputDevice& Ar );
 	bool HandleFlushPersistentDebugLinesCommand( const TCHAR* Cmd, FOutputDevice& Ar );
 	bool HandleLogActorCountsCommand( const TCHAR* Cmd, FOutputDevice& Ar, UWorld* InWorld );
-	
+	bool HandleDemoRecordCommand( const TCHAR* Cmd, FOutputDevice& Ar, UWorld* InWorld );
+	bool HandleDemoPlayCommand( const TCHAR* Cmd, FOutputDevice& Ar, UWorld* InWorld );
+	bool HandleDemoStopCommand( const TCHAR* Cmd, FOutputDevice& Ar, UWorld* InWorld );
+
+	// Destroys the current demo net driver
+	void DestroyDemoNetDriver();
+
 	// Start listening for connections.
 	bool Listen( FURL& InURL );
 
diff --git a/Engine/Source/Runtime/Engine/Private/DemoNetDriver.cpp b/Engine/Source/Runtime/Engine/Private/DemoNetDriver.cpp
new file mode 100644
--- /dev/null
+++ b/Engine/Source/Runtime/Engine/Private/DemoNetDriver.cpp
@@ -0,0 +1,611 @@
+// Copyright 1998-2014 Epic Games, Inc. All Rights Reserved.
+
+/*=============================================================================
+	UDemoNetDriver.cpp: Simulated network driver for recording and playing back game sessions.
+=============================================================================*/
+
+#include "EnginePrivate.h"
+#include "Engine/DemoNetDriver.h"
+#include "Engine/DemoNetConnection.h"
+#include "Engine/DemoPendingNetGame.h"
+
+DEFINE_LOG_CATEGORY_STATIC( LogDemo, Log, All );
+
+/*-----------------------------------------------------------------------------
+	UDemoNetDriver.
+-----------------------------------------------------------------------------*/
+
+UDemoNetDriver::UDemoNetDriver(const class FPostConstructInitializeProperties& PCIP)
+	: Super(PCIP)
+{
+}
+
+bool UDemoNetDriver::InitBase( bool bInitAsClient, FNetworkNotify* InNotify, const FURL& URL, bool bReuseAddressAndPort, FString& Error )
+{
+	if ( Super::InitBase( bInitAsClient, InNotify, URL, bReuseAddressAndPort, Error ) )
+	{
+		DemoFilename	= URL.Map;
+		Time			= 0;
+		DemoFrameNum	= 0;
+
+		return true;
+	}
+
+	return false;
+}
+
+void UDemoNetDriver::FinishDestroy()
+{
+	if ( !HasAnyFlags( RF_ClassDefaultObject ) )
+	{
+		// Make sure we stop any recording that might be going on
+		if ( ClientConnections.Num() > 0 )
+		{
+			StopDemo();
+		}
+	}
+
+	Super::FinishDestroy();
+}
+
+FString UDemoNetDriver::LowLevelGetNetworkNumber()
+{
+	return FString( TEXT( "" ) );
+}
+
+bool UDemoNetDriver::InitConnect( FNetworkNotify* InNotify, const FURL& ConnectURL, FString& Error )
+{
+	// handle default initialization
+	if ( !InitBase( true, InNotify, ConnectURL, false, Error ) )
+	{
+		return false;
+	}
+
+	// open the pre-recorded demo file
+	FileAr = IFileManager::Get().CreateFileReader( *DemoFilename );
+
+	if ( !FileAr )
+	{
+		Error = FString::Printf( TEXT( "Couldn't open demo file %s for reading" ), *DemoFilename );
+		UE_LOG( LogDemo, Error, TEXT( "UDemoNetDriver::InitConnect: %s" ), *Error );
+		return false;
+	}
+
+	// Playback, local machine is a client, and the demo stream acts "as if" it's the server.
+	ServerConnection = ConstructObject<UNetConnection>( UDemoNetConnection::StaticClass() );
+	ServerConnection->InitConnection( this, USOCK_Pending, ConnectURL, 1000000 );
+
+#if 1
+	// Create fake control channel
+	ServerConnection->CreateChannel( CHTYPE_Control, 1 );
+#endif
+
+	// use the same byte format regardless of platform so that the demos are cross platform
+	// DEMO_FIXME: This is messing up for some reason, investigate
+	//FileAr->SetByteSwapping( true );
+
+	int32 EngineVersion = 0;
+	(*FileAr) << EngineVersion;
+	(*FileAr) << PlaybackTotalFrames;
+
+	UE_LOG( LogDemo, Log, TEXT( "Starting demo playback with demo. Filename: %s, Frames: %i, Version %i" ), *DemoFilename, PlaybackTotalFrames, EngineVersion );
+
+#if 1
+	// Bypass UDemoPendingNetLevel
+	FString LevelName;
+	(*FileAr) << LevelName;
+
+	FString LoadMapError;
+
+	FURL DemoURL;
+	DemoURL.Map = LevelName;
+
+	FWorldContext * WorldContext = GEngine->GetWorldContextFromWorld( World );
+
+	if ( WorldContext == NULL )
+	{
+		Error = FString::Printf( TEXT( "No world context" ), *DemoFilename );
+		UE_LOG( LogDemo, Error, TEXT( "UDemoNetDriver::InitConnect: %s" ), *Error );
+		return false;
+	}
+
+	World->DemoNetDriver = NULL;
+	SetWorld( NULL );
+
+	if ( !GEngine->LoadMap( *WorldContext, DemoURL, NULL, LoadMapError ) )
+	{
+		Error = LoadMapError;
+		UE_LOG( LogDemo, Error, TEXT( "UDemoNetDriver::InitConnect: LoadMap failed: failed: %s" ), *Error );
+		return false;
+	}
+
+	SetWorld( WorldContext->World() );
+	WorldContext->World()->DemoNetDriver = this;
+#endif
+
+	DemoDeltaTime = 0;
+
+	return true;
+}
+
+bool UDemoNetDriver::InitListen( FNetworkNotify* InNotify, FURL& ListenURL, bool bReuseAddressAndPort, FString& Error )
+{
+	if ( !InitBase( false, InNotify, ListenURL, bReuseAddressAndPort, Error ) )
+	{
+		return false;
+	}
+
+	check( World != NULL );
+
+	class AWorldSettings * WorldSettings = World->GetWorldSettings();
+
+	if ( !WorldSettings )
+	{
+		Error = TEXT( "No WorldSettings!!" );
+		return false;
+	}
+
+	// Recording, local machine is server, demo stream acts "as if" it's a client.
+	UDemoNetConnection* Connection = ConstructObject<UDemoNetConnection>( UDemoNetConnection::StaticClass() );
+	Connection->InitConnection( this, USOCK_Open, ListenURL, 1000000 );
+	Connection->InitSendBuffer();
+
+	FileAr = IFileManager::Get().CreateFileWriter( *DemoFilename );
+	ClientConnections.Add( Connection );
+
+	if( !FileAr )
+	{
+		Error = FString::Printf( TEXT("Couldn't open demo file %s for writing"), *DemoFilename );//@todo demorec: localize
+		return false;
+	}
+
+	// use the same byte format regardless of platform so that the demos are cross platform
+	//@note: swap on non console platforms as the console archives have byte swapping compiled out by default
+	//FileAr->SetByteSwapping(true);
+
+	// write engine version info
+	int32 EngineVersion = GEngineNetVersion;
+	(*FileAr) << EngineVersion;
+
+	// write placeholder for total frames - will be updated when the demo is stopped
+	PlaybackTotalFrames = 0;
+	(*FileAr) << PlaybackTotalFrames;
+
+#if 0
+	// Create the control channel.
+	Connection->CreateChannel( CHTYPE_Control, 1, 0 );
+
+	// Send initial message.
+	uint8 IsLittleEndian = uint8( PLATFORM_LITTLE_ENDIAN );
+	check( IsLittleEndian == !!IsLittleEndian ); // should only be one or zero
+	FNetControlMessage<NMT_Hello>::Send( Connection, IsLittleEndian, GEngineMinNetVersion, GEngineNetVersion, Cast<UGeneralProjectSettings>(UGeneralProjectSettings::StaticClass()->GetDefaultObject())->ProjectID );
+	Connection->FlushNet();
+
+	// WelcomePlayer will send the needed map name
+	World->WelcomePlayer( Connection );
+#else
+	// Bypass UDemoPendingNetLevel
+	FString LevelName = World->GetCurrentLevel()->GetOutermost()->GetName();
+	(*FileAr) << LevelName;
+#endif
+
+	// Spawn the demo recording spectator.
+	SpawnDemoRecSpectator( Connection );
+
+	DemoDeltaTime = 0;
+	LastRecordTime = 0;
+
+	return true;
+}
+
+void UDemoNetDriver::TickDispatch( float DeltaSeconds )
+{
+	Super::TickDispatch( DeltaSeconds );
+}
+
+void UDemoNetDriver::TickFlush( float DeltaSeconds )
+{
+	Super::TickFlush( DeltaSeconds );
+
+	if ( FileAr )
+	{
+		if ( ClientConnections.Num() > 0 )
+		{
+			TickDemoRecord( DeltaSeconds );
+		}
+		else if ( ServerConnection != NULL )
+		{
+			TickDemoPlayback( DeltaSeconds );
+		}
+	}
+}
+
+void UDemoNetDriver::ProcessRemoteFunction( class AActor* Actor, class UFunction* Function, void* Parameters, struct FOutParmRec* OutParms, struct FFrame* Stack, class UObject* SubObject )
+{
+
+}
+
+bool UDemoNetDriver::Exec( UWorld* InWorld, const TCHAR* Cmd, FOutputDevice& Ar )
+{
+	return Super::Exec( InWorld, Cmd, Ar);
+}
+
+void UDemoNetDriver::StopDemo()
+{
+	if ( !ServerConnection && ClientConnections.Num() == 0 )
+	{
+		UE_LOG( LogDemo, Log, TEXT( "StopDemo: No demo is playing" ) );
+		return;
+	}
+
+	UE_LOG( LogDemo, Log, TEXT( "StopDemo: Demo %s stopped at frame %d" ), *DemoFilename, DemoFrameNum );
+
+	if ( !ServerConnection )
+	{
+		// write the total number of frames in the placeholder at the beginning of the file
+		if ( FileAr != NULL && World != NULL )
+		{
+			PlaybackTotalFrames = DemoFrameNum;
+			int32 OldPos = FileAr->Tell();
+			FileAr->Seek( sizeof( int32 ) );
+			(*FileAr) << PlaybackTotalFrames;
+			FileAr->Seek( OldPos );
+		}
+
+		// let GC cleanup the object
+		if ( ClientConnections.Num() > 0 && ClientConnections[0] != NULL )
+		{
+			ClientConnections[0]->Close();
+			ClientConnections[0]->CleanUp(); // make sure DemoRecSpectator gets destroyed immediately
+		}
+	}
+	else
+	{
+		// flush out any pending network traffic
+		ServerConnection->FlushNet();
+
+#if 0		
+		for (int32 i = ServerConnection->OpenChannels.Num() - 1; i >= 0; i--)
+		{
+			UChannel* OpenChannel = ServerConnection->OpenChannels[i];
+			if (OpenChannel != NULL)
+			{
+				UActorChannel* ActorChannel = Cast< UActorChannel >( OpenChannel );
+				if ( ActorChannel != NULL )
+				{
+					if ( !ActorChannel->GetActor()->IsNetStartupActor() )
+					{
+						OpenChannel->ConditionalCleanUp();
+					}
+				}
+			}
+		}
+#endif
+
+		ServerConnection->State = USOCK_Closed;
+		ServerConnection->Close();
+		ServerConnection->CleanUp(); // make sure DemoRecSpectator gets destroyed immediately
+		ServerConnection = NULL;
+		
+		//GEngine->SetClientTravel( World, TEXT( "?closed" ), TRAVEL_Absolute );
+		//SetWorld( NULL );
+	}
+
+	delete FileAr;
+	FileAr = NULL;
+
+	check( ClientConnections.Num() == 0 );
+	check( ServerConnection == NULL );
+}
+
+/*-----------------------------------------------------------------------------
+Demo Recording tick.
+-----------------------------------------------------------------------------*/
+
+//static int GNumReplicatedActors = 0;
+
+static void DemoReplicateActor(AActor* Actor, UNetConnection* Connection, bool IsNetClient)
+{
+	// All actors marked for replication are assumed to be relevant for demo recording.
+	/*
+	if
+		(	Actor
+		&&	((IsNetClient && Actor->bTearOff) || Actor->RemoteRole != ROLE_None || (IsNetClient && Actor->Role != ROLE_None && Actor->Role != ROLE_Authority) || Actor->bForceDemoRelevant)
+		&&  (!Actor->bNetTemporary || !Connection->SentTemporaries.Contains(Actor))
+		&& (Actor == Connection->PlayerController || Cast<APlayerController>(Actor) == NULL)
+		)
+	*/
+	if ( Actor != NULL && Actor->GetRemoteRole() != ROLE_None && ( Actor == Connection->PlayerController || Cast< APlayerController >( Actor ) == NULL ) )
+	{
+		// Create a new channel for this actor.
+		const bool StartupActor = Actor->IsNetStartupActor();
+		UActorChannel* Channel = Connection->ActorChannels.FindRef( Actor );
+
+		if ( !Channel && ( !StartupActor || Connection->ClientHasInitializedLevelFor( Actor ) ) )
+		{
+			// create a channel if possible
+			Channel = (UActorChannel*)Connection->CreateChannel( CHTYPE_Actor, 1 );
+			if ( Channel != NULL )
+			{
+				Channel->SetChannelActor( Actor );
+			}
+		}
+
+		if ( Channel )
+		{
+			// Send it out!
+			check( !Channel->Closing );
+
+			if ( Channel->IsNetReady( 0 ) )
+			{
+				Channel->ReplicateActor();
+				//GNumReplicatedActors++;
+			}
+		}
+	}
+}
+
+void UDemoNetDriver::TickDemoRecord( float DeltaSeconds )
+{
+	if ( ClientConnections.Num() == 0 )
+	{
+		return;
+	}
+
+	if ( FileAr == NULL )
+	{
+		return;
+	}
+
+	DemoDeltaTime += DeltaSeconds;
+
+	const double CurrentSeconds = FPlatformTime::Seconds();
+
+	const double RECORD_DELAY = 1.0 / 10.0;
+
+	if ( CurrentSeconds - LastRecordTime < RECORD_DELAY )
+	{
+		return;		// Not enough real-time has passed to record another frame
+	}
+
+	LastRecordTime = CurrentSeconds;
+
+	//GNumReplicatedActors = 0;
+
+	// Save out a frame
+	DemoFrameNum++;
+	ReplicationFrame++;
+
+	// Save elapsed game time
+	*FileAr << DemoDeltaTime;
+
+	DemoDeltaTime = 0;
+
+	// Make sure we don't have anything in the buffer for this new frame
+	check( ClientConnections[0]->SendBuffer.GetNumBits() == 0 );
+
+	const bool IsNetClient = ( GetWorld()->GetNetDriver() != NULL && GetWorld()->GetNetDriver()->GetNetMode() == NM_Client );
+
+	DemoReplicateActor( World->GetWorldSettings(), ClientConnections[0], IsNetClient );
+
+	for ( int32 i = 0; i < World->NetworkActors.Num(); i++ )
+	{
+		AActor * Actor = World->NetworkActors[i];
+
+		Actor->PreReplication( *FindOrCreateRepChangedPropertyTracker( Actor ).Get() );
+		DemoReplicateActor( Actor, ClientConnections[0], IsNetClient );
+	}
+
+	// Make sure nothing is left over
+	ClientConnections[0]->FlushNet();
+
+	check( ClientConnections[0]->SendBuffer.GetNumBits() == 0 );
+
+	// Write a count of 0 to signal the end of the frame
+	int32 EndCount = 0;
+
+	*FileAr << EndCount;
+
+	//UE_LOG( LogDemo, Log, TEXT( "Num replicated actors: %i, %i"), GNumReplicatedActors, ReplicationFrame );
+}
+
+bool UDemoNetDriver::ReadDemoFrame()
+{
+	if ( FileAr->AtEnd() || FileAr->IsError() )
+	{
+		StopDemo();
+		return false;
+	}
+
+	const int32 OldFilePos = FileAr->Tell();
+
+	float ServerDeltaTime;
+
+	// Peek at the next demo delta time, and see if we should process this frame
+	*FileAr << ServerDeltaTime;
+
+	if ( FileAr->IsError() )
+	{
+		UE_LOG( LogDemo, Error, TEXT( "Failed to read demo ServerDeltaTime" ) );
+		StopDemo();
+		return false;
+	}
+
+	if ( DemoDeltaTime - ServerDeltaTime < 0 )//&& ServerConnection->State != USOCK_Pending )
+	{
+		// Not enough time has passed to read another frame
+		FileAr->Seek( OldFilePos );
+		return false;
+	}
+
+	DemoDeltaTime -= ServerDeltaTime;
+
+	while ( true )
+	{
+		uint8 Data[ 512 + 8 ];
+
+		int32 PacketBytes;
+
+		*FileAr << PacketBytes;
+
+		if ( FileAr->IsError() )
+		{
+			UE_LOG( LogDemo, Error, TEXT( "Failed to read demo PacketBytes" ) );
+			StopDemo();
+			return false;
+		}
+
+		if ( PacketBytes == 0 )
+		{
+			break;
+		}
+
+		// Read data from file.
+		FileAr->Serialize( Data, PacketBytes );
+
+		if ( FileAr->IsError() )
+		{
+			UE_LOG( LogDemo, Error, TEXT( "Failed to read demo file packet" ) );
+			StopDemo();
+			return false;
+		}
+
+		// Process incoming packet.
+		ServerConnection->ReceivedRawPacket( Data, PacketBytes );
+
+		if ( ServerConnection == NULL || ServerConnection->State == USOCK_Closed )
+		{
+			// something we received resulted in the demo being stopped
+			StopDemo();
+			return false;
+		}
+	}
+
+	return true;
+}
+
+void UDemoNetDriver::TickDemoPlayback( float DeltaSeconds )
+{
+	if ( ServerConnection == NULL || ServerConnection->State == USOCK_Closed )
+	{
+		StopDemo();
+		return;
+	}
+
+	DemoDeltaTime += DeltaSeconds;
+
+	while ( true )
+	{
+		// Read demo frames until we are cought up
+		if ( !ReadDemoFrame() )
+		{
+			break;
+		}
+	}
+}
+
+void UDemoNetDriver::SpawnDemoRecSpectator( UNetConnection* Connection )
+{
+	check( Connection != NULL );
+	UClass* C = StaticLoadClass( AActor::StaticClass(), NULL, *DemoSpectatorClass, NULL, LOAD_None, NULL );
+	APlayerController* Controller = World->SpawnActor<APlayerController>( C );
+
+	for ( FActorIterator It( World ); It; ++It)
+	{
+		if ( It->IsA( APlayerStart::StaticClass() ) )
+		{
+			Controller->SetInitialLocationAndRotation( It->GetActorLocation() + FVector( 0, 0, 1000 ), It->GetActorRotation() );
+			break;
+		}
+	}
+
+	//Controller->SetReplicates( true );
+
+	Controller->SetPlayer( Connection );
+}
+
+/*-----------------------------------------------------------------------------
+	UDemoNetConnection.
+-----------------------------------------------------------------------------*/
+
+UDemoNetConnection::UDemoNetConnection( const class FPostConstructInitializeProperties& PCIP ) : Super( PCIP )
+{
+	MaxPacket = 512;
+	InternalAck = true;
+}
+
+void UDemoNetConnection::InitConnection( UNetDriver* InDriver, EConnectionState InState, const FURL& InURL, int32 InConnectionSpeed )
+{
+	// default implementation
+	Super::InitConnection( InDriver, InState, InURL, InConnectionSpeed );
+
+	InitSendBuffer();
+
+	// the driver must be a DemoRecording driver (GetDriver makes assumptions to avoid Cast'ing each time)
+	check( InDriver->IsA( UDemoNetDriver::StaticClass() ) );
+}
+
+FString UDemoNetConnection::LowLevelGetRemoteAddress( bool bAppendPort )
+{
+	return TEXT( "UDemoNetConnection" );
+}
+
+void UDemoNetConnection::LowLevelSend( void* Data, int32 Count )
+{
+	if ( Count == 0 )
+	{
+		UE_LOG( LogDemo, Warning, TEXT( "UDemoNetConnection::LowLevelSend: Ignoring empty packet." ) );
+		return;
+	}
+
+	if ( !GetDriver()->ServerConnection && GetDriver()->FileAr )
+	{
+		*GetDriver()->FileAr << Count;
+		GetDriver()->FileAr->Serialize( Data, Count );
+	}
+}
+
+FString UDemoNetConnection::LowLevelDescribe()
+{
+	return TEXT( "Demo recording/playback driver connection" );
+}
+
+int32 UDemoNetConnection::IsNetReady( bool Saturate )
+{
+	return 1;
+}
+
+void UDemoNetConnection::FlushNet( bool bIgnoreSimulation )
+{
+	// in playback, there is no data to send except
+	// channel closing if an error occurs.
+	if ( GetDriver()->ServerConnection != NULL )
+	{
+		InitSendBuffer();
+	}
+	else
+	{
+		Super::FlushNet( bIgnoreSimulation );
+	}
+}
+
+void UDemoNetConnection::HandleClientPlayer( APlayerController* PC, UNetConnection* NetConnection )
+{
+	Super::HandleClientPlayer( PC, NetConnection );
+	//PC->bDemoOwner = true;	// DEMO_FIXME
+}
+
+bool UDemoNetConnection::ClientHasInitializedLevelFor(const UObject* TestObject)
+{
+	// there is no way to know when demo playback will load levels,
+	// so we assume it always is anytime after initial replication and rely on the client to load it in time to make the reference work
+	return ( GetDriver()->DemoFrameNum > 2 || Super::ClientHasInitializedLevelFor( TestObject ) );
+}
+
+/*-----------------------------------------------------------------------------
+	UDemoPendingNetGame.
+-----------------------------------------------------------------------------*/
+
+UDemoPendingNetGame::UDemoPendingNetGame( const class FPostConstructInitializeProperties& PCIP ) : Super( PCIP )
+{
+}
diff --git a/Engine/Source/Runtime/Engine/Private/NetworkDriver.cpp b/Engine/Source/Runtime/Engine/Private/NetworkDriver.cpp
--- a/Engine/Source/Runtime/Engine/Private/NetworkDriver.cpp
+++ b/Engine/Source/Runtime/Engine/Private/NetworkDriver.cpp
@@ -149,7 +149,7 @@ void UNetDriver::AssertValid()
 
 void UNetDriver::TickFlush(float DeltaSeconds)
 {
-	if ( IsServer() )
+	if ( IsServer() && ClientConnections.Num() > 0 && ClientConnections[0]->InternalAck == false )
 	{
 		// Update all clients.
 #if WITH_SERVER_CODE
diff --git a/Engine/Source/Runtime/Engine/Private/PlayerController.cpp b/Engine/Source/Runtime/Engine/Private/PlayerController.cpp
--- a/Engine/Source/Runtime/Engine/Private/PlayerController.cpp
+++ b/Engine/Source/Runtime/Engine/Private/PlayerController.cpp
@@ -3732,7 +3732,7 @@ void APlayerController::TickActor( float DeltaSeconds, ELevelTick TickType, FAct
 	else if (Role > ROLE_SimulatedProxy)
 	{
 		// Process PlayerTick with input.
-		if (!PlayerInput)
+		if (!PlayerInput && (Player == NULL || Cast<ULocalPlayer>( Player ) != NULL))
 		{
 			InitInputSystem();
 		}
diff --git a/Engine/Source/Runtime/Engine/Private/UnrealEngine.cpp b/Engine/Source/Runtime/Engine/Private/UnrealEngine.cpp
--- a/Engine/Source/Runtime/Engine/Private/UnrealEngine.cpp
+++ b/Engine/Source/Runtime/Engine/Private/UnrealEngine.cpp
@@ -7410,6 +7410,9 @@ void UEngine::ShutdownWorldNetDriver( UWorld * World )
 			DestroyNamedNetDriver(World, NetDriver->NetDriverName);
 		}
 
+		// Take care of the demo net driver specifically (so the world can clear the DemoNetDriver property)
+		World->DestroyDemoNetDriver();
+
 		// Also disconnect any net drivers that have this set as their world, to avoid GC issues
 		FWorldContext &Context = GEngine->GetWorldContextFromWorldChecked(World);
 
@@ -7423,7 +7426,6 @@ void UEngine::ShutdownWorldNetDriver( UWorld * World )
 				Index--;
 			}
 		}
-
 	}
 }
 
diff --git a/Engine/Source/Runtime/Engine/Private/World.cpp b/Engine/Source/Runtime/Engine/Private/World.cpp
--- a/Engine/Source/Runtime/Engine/Private/World.cpp
+++ b/Engine/Source/Runtime/Engine/Private/World.cpp
@@ -34,6 +34,7 @@
 #include "AI/AISystemBase.h"
 #include "SceneInterface.h"
 #include "Camera/CameraActor.h"
+#include "Engine/DemoNetDriver.h"
 
 #if WITH_EDITOR
 	#include "Editor/UnrealEd/Public/Kismet2/KismetEditorUtilities.h"
@@ -189,6 +190,7 @@ void UWorld::AddReferencedObjects(UObject* InThis, FReferenceCollector& Collecto
 
 		Collector.AddReferencedObject( This->CurrentLevel, This );
 		Collector.AddReferencedObject( This->NetDriver, This );
+		Collector.AddReferencedObject( This->DemoNetDriver, This );
 		Collector.AddReferencedObject( This->LineBatcher, This );
 		Collector.AddReferencedObject( This->PersistentLineBatcher, This );
 		Collector.AddReferencedObject( This->ForegroundLineBatcher, This );
@@ -2574,6 +2576,18 @@ bool UWorld::Exec( UWorld* InWorld, const TCHAR* Cmd, FOutputDevice& Ar )
 	{		
 		return HandleLogActorCountsCommand( Cmd, Ar, InWorld );
 	}
+	else if (FParse::Command(&Cmd, TEXT("DEMOREC")))
+	{		
+		return HandleDemoRecordCommand( Cmd, Ar, InWorld );
+	}
+	else if( FParse::Command( &Cmd, TEXT("DEMOPLAY") ) )
+	{		
+		return HandleDemoPlayCommand( Cmd, Ar, InWorld );
+	}
+	else if( FParse::Command( &Cmd, TEXT("DEMOSTOP") ) )
+	{		
+		return HandleDemoStopCommand( Cmd, Ar, InWorld );
+	}
 	else if( ExecPhysCommands( Cmd, &Ar, InWorld ) )
 	{
 		return HandleLogActorCountsCommand( Cmd, Ar, InWorld );
@@ -2604,6 +2618,152 @@ bool UWorld::HandleLogActorCountsCommand( const TCHAR* Cmd, FOutputDevice& Ar, U
 	return true;
 }
 
+bool UWorld::HandleDemoRecordCommand( const TCHAR* Cmd, FOutputDevice& Ar, UWorld* InWorld )
+{
+	// Attempt to make the dir if it doesn't exist.
+	const FString DemoDir = FPaths::GameSavedDir() + TEXT( "Demos" );
+
+	IFileManager::Get().MakeDirectory( *DemoDir, true );
+
+	FURL DemoURL;
+	FString DemoName;
+
+	if ( !FParse::Token( Cmd, DemoName, 0 ) )
+	{
+		Ar.Logf( TEXT( "Missing demo name." ) );
+		return true;
+	}
+
+	int32 Year, Month, DayOfWeek, Day, Hour, Min, Sec, MSec;
+	FPlatformTime::SystemTime( Year, Month, DayOfWeek, Day, Hour, Min, Sec, MSec );
+
+	DemoName.ReplaceInline( TEXT( "%td" ), *FDateTime::Now().ToString() );
+	DemoName.ReplaceInline( TEXT( "%d" ), *FString::Printf( TEXT( "%i-%i-%i" ), Month, Day, Year ) );
+	DemoName.ReplaceInline( TEXT( "%t" ), *FString::Printf( TEXT( "%i" ), ( ( Hour * 3600 ) + ( Min * 60 ) + Sec ) * 1000 + MSec ) );
+	DemoName.ReplaceInline( TEXT( "%v" ), *FString::Printf( TEXT( "%i" ), GEngineVersion.GetChangelist() ) );
+
+	// replace bad characters with underscores
+	DemoName.ReplaceInline( TEXT( "\\" ),	TEXT( "_" ) );
+	DemoName.ReplaceInline( TEXT( "/" ),	TEXT( "_" ) );
+	DemoName.ReplaceInline( TEXT( "." ),	TEXT( "_" ) );
+	DemoName.ReplaceInline( TEXT( " " ),	TEXT( "_" ) );
+	DemoName.ReplaceInline( TEXT( "%" ),	TEXT( "_" ) );
+
+	// replace the current URL's map with a demo extension
+	DemoURL.Map = DemoDir + TEXT( "/" ) + DemoName + TEXT( ".demo" );
+
+	DestroyDemoNetDriver();
+
+	const FName NAME_DemoNetDriver( TEXT( "DemoNetDriver" ) );
+
+	if ( !GEngine->CreateNamedNetDriver( this, NAME_DemoNetDriver, NAME_DemoNetDriver ) )
+	{
+		Ar.Logf( TEXT( "Failed to create demo net driver!" ) );
+		return true;
+	}
+
+	DemoNetDriver = Cast< UDemoNetDriver >( GEngine->FindNamedNetDriver( this, NAME_DemoNetDriver ) );
+
+	check( DemoNetDriver != NULL );
+
+	DemoNetDriver->SetWorld( this );
+
+	FString Error;
+
+	if ( !DemoNetDriver->InitListen( this, DemoURL, false, Error ) )
+	{
+		Ar.Logf( TEXT( "Demo recording failed: %s" ), *Error );
+		DemoNetDriver = NULL;
+	}
+	else
+	{
+		Ar.Logf( TEXT( "Num Network Actors: %i" ), InWorld->NetworkActors.Num() );
+	}
+
+	return true;
+}
+
+bool UWorld::HandleDemoPlayCommand( const TCHAR* Cmd, FOutputDevice& Ar, UWorld* InWorld )
+{
+	FString Temp;
+
+	if ( FParse::Token( Cmd, Temp, 0 ) )
+	{
+		DestroyDemoNetDriver();
+
+		FURL DemoURL;
+		UE_LOG( LogWorld, Log, TEXT( "Attempting to play demo %s" ), *Temp );
+
+		const FString DemoDir = FPaths::GameSavedDir() + TEXT( "Demos" );
+
+		DemoURL.Map = DemoDir + TEXT( "/" ) + Temp + TEXT( ".demo" );
+
+#if 0
+		UPendingNetGame * NewPendingNetGame = new UDemoPendingNetGame( FPostConstructInitializeProperties(), DemoURL );
+
+		if ( !NewPendingNetGame->DemoNetDriver )
+		{
+			Ar.Logf( TEXT( "Demo playback failed: %s" ), *NewPendingNetGame->ConnectionError );
+			NewPendingNetGame = NULL;
+		}
+
+		GEngine->CancelPending( InWorld, NewPendingNetGame );
+#else
+		const FName NAME_DemoNetDriver( TEXT( "DemoNetDriver" ) );
+
+		if ( !GEngine->CreateNamedNetDriver( this, NAME_DemoNetDriver, NAME_DemoNetDriver ) )
+		{
+			Ar.Logf( TEXT( "Failed to create demo net driver!" ) );
+			return true;
+		}
+
+		DemoNetDriver = Cast< UDemoNetDriver >( GEngine->FindNamedNetDriver( this, NAME_DemoNetDriver ) );
+
+		check( DemoNetDriver != NULL );
+
+		DemoNetDriver->SetWorld( this );
+
+		FString Error;
+
+		if ( !DemoNetDriver->InitConnect( this, DemoURL, Error ) )
+		{
+			Ar.Logf( TEXT( "Demo playback failed: %s" ), *Error );
+			DestroyDemoNetDriver();
+		}
+#endif
+	}
+	else
+	{
+		Ar.Log( TEXT( "You must specify a filename" ) );
+	}
+
+	return true;
+}
+
+bool UWorld::HandleDemoStopCommand( const TCHAR* Cmd, FOutputDevice& Ar, UWorld* InWorld )
+{
+	DestroyDemoNetDriver();
+	return true;
+}
+
+void UWorld::DestroyDemoNetDriver()
+{
+	if ( DemoNetDriver != NULL )
+	{
+		const FName DemoNetDriverName = DemoNetDriver->NetDriverName;
+
+		check( GEngine->FindNamedNetDriver( this, DemoNetDriverName ) == DemoNetDriver );
+
+		DemoNetDriver->StopDemo();
+		DemoNetDriver->SetWorld( NULL );
+
+		GEngine->DestroyNamedNetDriver( this, DemoNetDriverName );
+
+		check( GEngine->FindNamedNetDriver( this, DemoNetDriverName ) == NULL );
+
+		DemoNetDriver = NULL;
+	}
+}
 
 bool UWorld::SetGameMode(const FURL& InURL)
 {
