Made the Custom material node work with TextureObject and removed old CustomTexture node as it's no longer needed.

diff --git a/Engine/Source/Editor/UnrealEd/Private/MaterialGraphNode.cpp b/Engine/Source/Editor/UnrealEd/Private/MaterialGraphNode.cpp
--- a/Engine/Source/Editor/UnrealEd/Private/MaterialGraphNode.cpp
+++ b/Engine/Source/Editor/UnrealEd/Private/MaterialGraphNode.cpp
@@ -11,7 +11,6 @@
 #include "Materials/MaterialExpressionConstant2Vector.h"
 #include "Materials/MaterialExpressionConstant3Vector.h"
 #include "Materials/MaterialExpressionConstant4Vector.h"
-#include "Materials/MaterialExpressionCustomTexture.h"
 #include "Materials/MaterialExpressionFontSample.h"
 #include "Materials/MaterialExpressionFontSampleParameter.h"
 #include "Materials/MaterialExpressionFunctionInput.h"
@@ -714,10 +713,6 @@ bool UMaterialGraphNode::UsesObjectColour(UMaterialExpression* Expression)
 	{
 		return true;
 	}
-	else if (Expression->IsA<UMaterialExpressionCustomTexture>())
-	{
-		return true;
-	}
 	else if (Expression->IsA<UMaterialExpressionFontSample>())
 	{
 		return true;
diff --git a/Engine/Source/Runtime/Engine/Classes/Materials/MaterialExpressionCustomTexture.h b/Engine/Source/Runtime/Engine/Classes/Materials/MaterialExpressionCustomTexture.h
deleted file mode 100644
diff --git a/Engine/Source/Runtime/Engine/Private/Materials/HLSLMaterialTranslator.h b/Engine/Source/Runtime/Engine/Private/Materials/HLSLMaterialTranslator.h
--- a/Engine/Source/Runtime/Engine/Private/Materials/HLSLMaterialTranslator.h
+++ b/Engine/Source/Runtime/Engine/Private/Materials/HLSLMaterialTranslator.h
@@ -3408,24 +3408,38 @@ protected:
 			case MCT_Float:
 			case MCT_Float1:
 				InputParamDecl += TEXT("MaterialFloat ");
+				InputParamDecl += Custom->Inputs[i].InputName;
 				break;
 			case MCT_Float2:
 				InputParamDecl += TEXT("MaterialFloat2 ");
+				InputParamDecl += Custom->Inputs[i].InputName;
 				break;
 			case MCT_Float3:
 				InputParamDecl += TEXT("MaterialFloat3 ");
+				InputParamDecl += Custom->Inputs[i].InputName;
 				break;
 			case MCT_Float4:
 				InputParamDecl += TEXT("MaterialFloat4 ");
+				InputParamDecl += Custom->Inputs[i].InputName;
 				break;
 			case MCT_Texture2D:
-				InputParamDecl += TEXT("sampler2D ");
+				InputParamDecl += TEXT("Texture2D ");
+				InputParamDecl += Custom->Inputs[i].InputName;
+				InputParamDecl += TEXT(", sampler ");
+				InputParamDecl += Custom->Inputs[i].InputName;
+				InputParamDecl += TEXT("Sampler ");
+				break;
+			case MCT_TextureCube:
+				InputParamDecl += TEXT("TextureCube ");
+				InputParamDecl += Custom->Inputs[i].InputName;
+				InputParamDecl += TEXT(", sampler ");
+				InputParamDecl += Custom->Inputs[i].InputName;
+				InputParamDecl += TEXT("Sampler ");
 				break;
 			default:
 				return Errorf(TEXT("Bad type %s for %s input %s"),DescribeType(GetParameterType(CompiledInputs[i])), *Custom->Description, *Custom->Inputs[i].InputName);
 				break;
 			}
-			InputParamDecl += Custom->Inputs[i].InputName;
 		}
 		int32 CustomExpressionIndex = CustomExpressionImplementations.Num();
 		FString Code = Custom->Code;
@@ -3446,8 +3460,18 @@ protected:
 			{
 				continue;
 			}
+
+			FString ParamCode = GetParameterCode(CompiledInputs[i]);
+			EMaterialValueType ParamType = GetParameterType(CompiledInputs[i]);
+
 			CodeChunk += TEXT(",");
-			CodeChunk += *GetParameterCode(CompiledInputs[i]);
+			CodeChunk += *ParamCode;
+			if (ParamType == MCT_Texture2D || ParamType == MCT_TextureCube)
+			{
+				CodeChunk += TEXT(",");
+				CodeChunk += *ParamCode;
+				CodeChunk += TEXT("Sampler");
+			}
 		}
 		CodeChunk += TEXT(")");
 
diff --git a/Engine/Source/Runtime/Engine/Private/Materials/MaterialExpressions.cpp b/Engine/Source/Runtime/Engine/Private/Materials/MaterialExpressions.cpp
--- a/Engine/Source/Runtime/Engine/Private/Materials/MaterialExpressions.cpp
+++ b/Engine/Source/Runtime/Engine/Private/Materials/MaterialExpressions.cpp
@@ -29,7 +29,6 @@
 #include "Materials/MaterialExpressionCosine.h"
 #include "Materials/MaterialExpressionCrossProduct.h"
 #include "Materials/MaterialExpressionCustom.h"
-#include "Materials/MaterialExpressionCustomTexture.h"
 #include "Materials/MaterialExpressionDDX.h"
 #include "Materials/MaterialExpressionDDY.h"
 #include "Materials/MaterialExpressionDepthFade.h"
@@ -6035,121 +6034,6 @@ uint32 UMaterialExpressionCustom::GetOutputType(int32 OutputIndex)
 }
 #endif // WITH_EDITOR
 
-///////////////////////////////////////////////////////////////////////////////
-// UMaterialExpressionCustomTexture
-///////////////////////////////////////////////////////////////////////////////
-UMaterialExpressionCustomTexture::UMaterialExpressionCustomTexture(const class FPostConstructInitializeProperties& PCIP)
-	: Super(PCIP)
-{
-	// Structure to hold one-time initialization
-	struct FConstructorStatics
-	{
-		FString NAME_Custom;
-		FConstructorStatics()
-			: NAME_Custom(LOCTEXT( "Custom", "Custom" ).ToString())
-		{
-		}
-	};
-	static FConstructorStatics ConstructorStatics;
-
-	MenuCategories.Add(ConstructorStatics.NAME_Custom);
-	bCollapsed = false;
-}
-
-int32 UMaterialExpressionCustomTexture::Compile(class FMaterialCompiler* Compiler, int32 OutputIndex, int32 MultiplexIndex)
-{
-#if PLATFORM_EXCEPTIONS_DISABLED
-	// if we can't throw the error below, attempt to thwart the error by using the default texture
-	// @todo: handle missing cubemaps and 3d textures?
-	if (!Texture)
-	{
-		UE_LOG(LogMaterial, Log, TEXT("Using default texture instead of real texture!"));
-		Texture = GEngine->DefaultTexture;
-	}
-#endif
-
-	if (Texture)
-	{
-		return  Compiler->Texture(Texture);
-	}
-	else
-	{
-		if (Desc.Len() > 0)
-		{
-			return Compiler->Errorf(TEXT("%s> Missing input texture"), *Desc);
-		}
-		else
-		{
-			return Compiler->Errorf(TEXT("CustomTexture> Missing input texture"));
-		}
-	}
-}
-
-int32 UMaterialExpressionCustomTexture::CompilePreview(class FMaterialCompiler* Compiler, int32 OutputIndex, int32 MultiplexIndex)
-{
-#if PLATFORM_EXCEPTIONS_DISABLED
-	// if we can't throw the error below, attempt to thwart the error by using the default texture
-	// @todo: handle missing cubemaps and 3d textures?
-	if (!Texture)
-	{
-		UE_LOG(LogMaterial, Log, TEXT("Using default texture instead of real texture!"));
-		Texture = GEngine->DefaultTexture;
-	}
-#endif
-
-	if (Texture)
-	{
-		// Preview just returns the sampled texture rather than the texture sampler.
-		int32 TextureCodeIndex = Compiler->Texture(Texture);
-		return Compiler->TextureSample(TextureCodeIndex, Compiler->TextureCoordinate(0, false, false), SAMPLERTYPE_Color);
-	}
-	else
-	{
-		if (Desc.Len() > 0)
-		{
-			return Compiler->Errorf(TEXT("%s> Missing input texture"), *Desc);
-		}
-		else
-		{
-			return Compiler->Errorf(TEXT("CustomTexture> Missing input texture"));
-		}
-	}
-}
-
-int32 UMaterialExpressionCustomTexture::GetWidth() const
-{
-	return ME_STD_THUMBNAIL_SZ+(ME_STD_BORDER*2);
-}
-
-void UMaterialExpressionCustomTexture::GetCaption(TArray<FString>& OutCaptions) const
-{
-	OutCaptions.Add(TEXT("Custom Texture"));
-}
-
-bool UMaterialExpressionCustomTexture::MatchesSearchQuery( const TCHAR* SearchQuery )
-{
-	if( Texture!=NULL && Texture->GetName().Contains(SearchQuery) )
-	{
-		return true;
-	}
-
-	return Super::MatchesSearchQuery(SearchQuery);
-}
-
-#if WITH_EDITOR
-uint32 UMaterialExpressionCustomTexture::GetOutputType(int32 OutputIndex)
-{
-	if (Cast<UTextureCube>(Texture) != NULL)
-	{
-		return MCT_TextureCube;
-	}
-	else
-	{
-		return MCT_Texture2D;
-	}
-}
-#endif // WITH_EDITOR
-
 ///////////////////////////////////////////////////////////////////////////////
 // UMaterialFunction
 ///////////////////////////////////////////////////////////////////////////////
