Fixed case where landscape texture could be created with wrong outer
TTP #345615

diff --git a/Engine/Source/Runtime/Engine/Private/Landscape/LandscapeEdit.cpp b/Engine/Source/Runtime/Engine/Private/Landscape/LandscapeEdit.cpp
--- a/Engine/Source/Runtime/Engine/Private/Landscape/LandscapeEdit.cpp
+++ b/Engine/Source/Runtime/Engine/Private/Landscape/LandscapeEdit.cpp
@@ -4765,7 +4765,7 @@ UMaterialInstance* ULandscapeComponent::GeneratePlatformPixelData(TArray<UTextur
 	}
 
 	int32 WeightmapSize = (SubsectionSizeQuads + 1) * NumSubsections;
-	UTexture2D* WeightmapTexture = GetLandscapeActor()->CreateLandscapeTexture(WeightmapSize, WeightmapSize, TEXTUREGROUP_Terrain_Weightmap, TSF_BGRA8);
+	UTexture2D* WeightmapTexture = GetLandscapeProxy()->CreateLandscapeTexture(WeightmapSize, WeightmapSize, TEXTUREGROUP_Terrain_Weightmap, TSF_BGRA8);
 	CreateEmptyTextureMips(WeightmapTexture);
 
 	{
