Adding an updatable texture to easily allow writing to textures every frame.

Added an updatable texture interface, which will be implemented by the main texture types of each slate renderer.
Ensured that it is safe to update/resize the texture from the game or render thread in the case of the RHI renderer.
D3D could be improved if we could write the whole texture at once but the only working method I could find exposes memory buffer with 16 byte alignment for each row.

Fixed non-unity build failure

diff --git a/Engine/Source/Developer/StandaloneRenderer/Private/OpenGL/SlateOpenGLRenderer.cpp b/Engine/Source/Developer/StandaloneRenderer/Private/OpenGL/SlateOpenGLRenderer.cpp
--- a/Engine/Source/Developer/StandaloneRenderer/Private/OpenGL/SlateOpenGLRenderer.cpp
+++ b/Engine/Source/Developer/StandaloneRenderer/Private/OpenGL/SlateOpenGLRenderer.cpp
@@ -216,3 +216,21 @@ void FSlateOpenGLRenderer::LoadStyleResources( const ISlateStyle& Style )
 		TextureManager->LoadStyleResources( Style );
 	}
 }
+
+FSlateUpdatableTexture* FSlateOpenGLRenderer::CreateUpdatableTexture(uint32 Width, uint32 Height)
+{
+	TArray<uint8> RawData;
+	RawData.AddZeroed(4);
+	FSlateOpenGLTexture* NewTexture = new FSlateOpenGLTexture(Width, Height);
+#if !PLATFORM_USES_ES2
+	NewTexture->Init(GL_SRGB8_ALPHA8, RawData);
+#else
+	NewTexture->Init(GL_SRGB8_ALPHA8_EXT, RawData);
+#endif
+	return NewTexture;
+}
+
+void FSlateOpenGLRenderer::ReleaseUpdatableTexture(FSlateUpdatableTexture* Texture)
+{
+	delete Texture;
+}
diff --git a/Engine/Source/Developer/StandaloneRenderer/Private/OpenGL/SlateOpenGLRenderer.h b/Engine/Source/Developer/StandaloneRenderer/Private/OpenGL/SlateOpenGLRenderer.h
--- a/Engine/Source/Developer/StandaloneRenderer/Private/OpenGL/SlateOpenGLRenderer.h
+++ b/Engine/Source/Developer/StandaloneRenderer/Private/OpenGL/SlateOpenGLRenderer.h
@@ -145,6 +145,8 @@ public:
 	virtual void ReleaseDynamicResource( const FSlateBrush& Brush ) override;
 	virtual bool GenerateDynamicImageResource(FName ResourceName, uint32 Width, uint32 Height, const TArray< uint8 >& Bytes) override;
 	virtual void LoadStyleResources(const ISlateStyle& Style) override;
+	virtual FSlateUpdatableTexture* CreateUpdatableTexture(uint32 Width, uint32 Height) override;
+	virtual void ReleaseUpdatableTexture(FSlateUpdatableTexture* Texture) override;
 
 private:
 	/** 
diff --git a/Engine/Source/Developer/StandaloneRenderer/Private/OpenGL/SlateOpenGLTextures.cpp b/Engine/Source/Developer/StandaloneRenderer/Private/OpenGL/SlateOpenGLTextures.cpp
--- a/Engine/Source/Developer/StandaloneRenderer/Private/OpenGL/SlateOpenGLTextures.cpp
+++ b/Engine/Source/Developer/StandaloneRenderer/Private/OpenGL/SlateOpenGLTextures.cpp
@@ -36,6 +36,34 @@ void FSlateOpenGLTexture::Init( GLuint TextureID )
 	ShaderResource = TextureID;
 }
 
+void FSlateOpenGLTexture::ResizeTexture(uint32 Width, uint32 Height)
+{
+	SizeX = Width;
+	SizeY = Height;
+}
+
+void FSlateOpenGLTexture::UpdateTexture(const TArray<uint8>& Bytes)
+{
+	// Ensure texturing is enabled before setting texture properties
+#if !PLATFORM_USES_ES2 && !PLATFORM_LINUX
+	glEnable(GL_TEXTURE_2D);
+#endif
+	glBindTexture(GL_TEXTURE_2D, ShaderResource);
+	glTexParameteri(GL_TEXTURE_2D, GL_TEXTURE_MIN_FILTER, GL_LINEAR);
+	glTexParameteri(GL_TEXTURE_2D, GL_TEXTURE_MAG_FILTER, GL_LINEAR);
+#if !PLATFORM_LINUX
+	glTexEnvi(GL_TEXTURE_ENV, GL_TEXTURE_ENV_MODE, GL_REPLACE);
+#endif
+
+	// Upload the texture data
+#if !PLATFORM_USES_ES2
+	glTexImage2D(GL_TEXTURE_2D, 0, GL_SRGB8_ALPHA8, SizeX, SizeY, 0, GL_RGBA, GL_UNSIGNED_INT_8_8_8_8_REV, Bytes.GetData());
+#else
+	glTexImage2D(GL_TEXTURE_2D, 0, GL_SRGB8_ALPHA8_EXT, SizeX, SizeY, 0, GL_RGBA, GL_UNSIGNED_INT_8_8_8_8_REV, Bytes.GetData());
+#endif
+	CHECK_GL_ERRORS;
+}
+
 FSlateFontTextureOpenGL::FSlateFontTextureOpenGL( uint32 Width, uint32 Height )
 	: FSlateFontAtlas( Width, Height ) 
 	, Texture(NULL)
diff --git a/Engine/Source/Developer/StandaloneRenderer/Private/OpenGL/SlateOpenGLTextures.h b/Engine/Source/Developer/StandaloneRenderer/Private/OpenGL/SlateOpenGLTextures.h
--- a/Engine/Source/Developer/StandaloneRenderer/Private/OpenGL/SlateOpenGLTextures.h
+++ b/Engine/Source/Developer/StandaloneRenderer/Private/OpenGL/SlateOpenGLTextures.h
@@ -5,7 +5,7 @@
 
 #pragma once
 
-class FSlateOpenGLTexture : public TSlateTexture< GLuint >
+class FSlateOpenGLTexture : public TSlateTexture< GLuint >, public FSlateUpdatableTexture
 {
 public:
 	FSlateOpenGLTexture( uint32 InSizeX, uint32 InSizeY )
@@ -28,6 +28,11 @@ public:
 	uint32 GetWidth() const { return SizeX; }
 	uint32 GetHeight() const { return SizeY; }
 
+	// FSlateUpdatableTexture interface
+	virtual FSlateShaderResource* GetSlateResource() override {return this;}
+	virtual void ResizeTexture( uint32 Width, uint32 Height ) override;
+	virtual void UpdateTexture(const TArray<uint8>& Bytes) override;
+
 private:
 	static GLuint NullTexture;
 
diff --git a/Engine/Source/Developer/StandaloneRenderer/Private/Windows/D3D/SlateD3DRenderer.cpp b/Engine/Source/Developer/StandaloneRenderer/Private/Windows/D3D/SlateD3DRenderer.cpp
--- a/Engine/Source/Developer/StandaloneRenderer/Private/Windows/D3D/SlateD3DRenderer.cpp
+++ b/Engine/Source/Developer/StandaloneRenderer/Private/Windows/D3D/SlateD3DRenderer.cpp
@@ -120,6 +120,18 @@ void FSlateD3DRenderer::LoadStyleResources( const ISlateStyle& Style )
 	}
 }
 
+FSlateUpdatableTexture* FSlateD3DRenderer::CreateUpdatableTexture(uint32 Width, uint32 Height)
+{
+	FSlateD3DTexture* NewTexture = new FSlateD3DTexture(Width, Height);
+	NewTexture->Init(DXGI_FORMAT_B8G8R8A8_UNORM_SRGB, NULL, D3D11_USAGE_DYNAMIC, D3D11_CPU_ACCESS_WRITE);
+	return NewTexture;
+}
+
+void FSlateD3DRenderer::ReleaseUpdatableTexture(FSlateUpdatableTexture* Texture)
+{
+	delete Texture;
+}
+
 void FSlateD3DRenderer::Private_CreateViewport( TSharedRef<SWindow> InWindow, const FVector2D &WindowSize )
 {
 	TSharedRef< FGenericWindow > NativeWindow = InWindow->GetNativeWindow().ToSharedRef();
diff --git a/Engine/Source/Developer/StandaloneRenderer/Private/Windows/D3D/SlateD3DRenderer.h b/Engine/Source/Developer/StandaloneRenderer/Private/Windows/D3D/SlateD3DRenderer.h
--- a/Engine/Source/Developer/StandaloneRenderer/Private/Windows/D3D/SlateD3DRenderer.h
+++ b/Engine/Source/Developer/StandaloneRenderer/Private/Windows/D3D/SlateD3DRenderer.h
@@ -56,6 +56,8 @@ public:
 	virtual void ReleaseDynamicResource( const FSlateBrush& Brush ) override;
 	virtual bool GenerateDynamicImageResource(FName ResourceName, uint32 Width, uint32 Height, const TArray< uint8 >& Bytes) override;
 	virtual void LoadStyleResources( const ISlateStyle& Style ) override;
+	virtual FSlateUpdatableTexture* CreateUpdatableTexture(uint32 Width, uint32 Height) override;
+	virtual void ReleaseUpdatableTexture(FSlateUpdatableTexture* Texture) override;
 	
 	void CreateDevice();
 	void CreateDepthStencilBuffer( FSlateD3DViewport& Viewport );
diff --git a/Engine/Source/Developer/StandaloneRenderer/Private/Windows/D3D/SlateD3DTextures.cpp b/Engine/Source/Developer/StandaloneRenderer/Private/Windows/D3D/SlateD3DTextures.cpp
--- a/Engine/Source/Developer/StandaloneRenderer/Private/Windows/D3D/SlateD3DTextures.cpp
+++ b/Engine/Source/Developer/StandaloneRenderer/Private/Windows/D3D/SlateD3DTextures.cpp
@@ -35,6 +35,26 @@ void FSlateD3DTexture::Init( DXGI_FORMAT InFormat, D3D11_SUBRESOURCE_DATA* Inita
 	check( SUCCEEDED(Hr) );
 }
 
+void FSlateD3DTexture::ResizeTexture(uint32 Width, uint32 Height)
+{
+	// Seems only way to resize d3d texture is recreate it
+	SizeX = Width;
+	SizeY = Height;
+	Init(DXGI_FORMAT_B8G8R8A8_UNORM_SRGB, NULL, D3D11_USAGE_DYNAMIC, D3D11_CPU_ACCESS_WRITE);
+}
+
+void FSlateD3DTexture::UpdateTexture(const TArray<uint8>& Bytes)
+{
+	// TODO: Improve the memory copying here, have tried using UpdateSubresource but it doesn't seem to work
+	D3D11_MAPPED_SUBRESOURCE Resource;
+	GD3DDeviceContext->Map(D3DTexture, 0, D3D11_MAP_WRITE_DISCARD, 0, &Resource);
+	for (uint32 Row = 0; Row < SizeY; ++Row)
+	{
+		FMemory::Memcpy((uint8*)Resource.pData + Row * Resource.RowPitch, Bytes.GetTypedData() + Row * SizeX*4, SizeX*4);
+	}
+	GD3DDeviceContext->Unmap(D3DTexture, 0);
+}
+
 FSlateTextureAtlasD3D::FSlateTextureAtlasD3D( uint32 Width, uint32 Height, uint32 StrideBytes, ESlateTextureAtlasPaddingStyle PaddingStyle )
 	: FSlateTextureAtlas( Width, Height, StrideBytes, PaddingStyle )
 	, AtlasTexture( new FSlateD3DTexture( Width, Height ) )
diff --git a/Engine/Source/Developer/StandaloneRenderer/Private/Windows/D3D/SlateD3DTextures.h b/Engine/Source/Developer/StandaloneRenderer/Private/Windows/D3D/SlateD3DTextures.h
--- a/Engine/Source/Developer/StandaloneRenderer/Private/Windows/D3D/SlateD3DTextures.h
+++ b/Engine/Source/Developer/StandaloneRenderer/Private/Windows/D3D/SlateD3DTextures.h
@@ -6,7 +6,7 @@
 /**
  * Encapsulates a D3D11 texture that can be accessed by a shader                   
  */
-class FSlateD3DTexture : public TSlateTexture< TRefCountPtr<ID3D11ShaderResourceView> >
+class FSlateD3DTexture : public TSlateTexture< TRefCountPtr<ID3D11ShaderResourceView> >, public FSlateUpdatableTexture
 {
 public:
 	FSlateD3DTexture( uint32 InSizeX, uint32 InSizeY )
@@ -28,6 +28,11 @@ public:
 	uint32 GetHeight() const { return SizeY; }
 
 	TRefCountPtr<ID3D11Texture2D> GetTextureResource() const { return D3DTexture; }
+
+	// FSlateUpdatableTexture interface
+	virtual FSlateShaderResource* GetSlateResource() override {return this;}
+	virtual void ResizeTexture( uint32 Width, uint32 Height ) override;
+	virtual void UpdateTexture(const TArray<uint8>& Bytes) override;
 private:
 	/** Actual texture.  In D3D the SRV is used by the shader so our parent class holds that */
 	TRefCountPtr<ID3D11Texture2D> D3DTexture;
diff --git a/Engine/Source/Runtime/Engine/Private/DemoNetDriver.cpp b/Engine/Source/Runtime/Engine/Private/DemoNetDriver.cpp
--- a/Engine/Source/Runtime/Engine/Private/DemoNetDriver.cpp
+++ b/Engine/Source/Runtime/Engine/Private/DemoNetDriver.cpp
@@ -8,6 +8,8 @@
 #include "Engine/DemoNetDriver.h"
 #include "Engine/DemoNetConnection.h"
 #include "Engine/DemoPendingNetGame.h"
+#include "Engine/ActorChannel.h"
+#include "RepLayout.h"
 
 DEFINE_LOG_CATEGORY_STATIC( LogDemo, Log, All );
 
diff --git a/Engine/Source/Runtime/Engine/Private/Slate/SlateTextures.cpp b/Engine/Source/Runtime/Engine/Private/Slate/SlateTextures.cpp
--- a/Engine/Source/Runtime/Engine/Private/Slate/SlateTextures.cpp
+++ b/Engine/Source/Runtime/Engine/Private/Slate/SlateTextures.cpp
@@ -119,7 +119,49 @@ void FSlateTexture2DRHIRef::ClearTextureData()
 	TextureData.Reset();
 }
 
+void FSlateTexture2DRHIRef::ResizeTexture(uint32 Width, uint32 Height)
+{
+	if (GetWidth() != Width || GetHeight() != Height)
+	{
+		if (IsInRenderingThread())
+		{
+			Resize(Width, Height);
+		}
+		else
+		{
+			FIntPoint Dimensions(Width, Height);
+			ENQUEUE_UNIQUE_RENDER_COMMAND_TWOPARAMETER(ResizeSlateTexture,
+			FSlateTexture2DRHIRef*, TextureRHIRef, this,
+			FIntPoint, InDimensions, Dimensions,
+			{
+				TextureRHIRef->Resize(InDimensions.X, InDimensions.Y);
+			});
+		}
+	}
+}
 
+void FSlateTexture2DRHIRef::UpdateTexture(const TArray<uint8>& Bytes)
+{
+	if (IsInRenderingThread())
+	{
+		uint32 Stride = 0;
+		void* TextureBuffer = RHILockTexture2D(GetTypedResource(), 0, RLM_WriteOnly, Stride, false);
+		FMemory::Memcpy(TextureBuffer, Bytes.GetData(), Bytes.Num());
+		RHIUnlockTexture2D(GetTypedResource(), 0, false);
+	}
+	else
+	{
+		ENQUEUE_UNIQUE_RENDER_COMMAND_TWOPARAMETER(UpdateSlateTexture,
+		FSlateTexture2DRHIRef*, TextureRHIRef, this,
+		const TArray<uint8>&, TextureData, Bytes,
+		{
+			uint32 Stride = 0;
+			void* TextureBuffer = RHILockTexture2D(TextureRHIRef->GetTypedResource(), 0, RLM_WriteOnly, Stride, false);
+			FMemory::Memcpy(TextureBuffer, TextureData.GetData(), TextureData.Num());
+			RHIUnlockTexture2D(TextureRHIRef->GetTypedResource(), 0, false);
+		});
+	}
+}
 
 void FSlateRenderTargetRHI::SetRHIRef( FTexture2DRHIRef InRenderTargetTexture, uint32 InWidth, uint32 InHeight )
 {
diff --git a/Engine/Source/Runtime/Engine/Public/Slate/SlateTextures.h b/Engine/Source/Runtime/Engine/Public/Slate/SlateTextures.h
--- a/Engine/Source/Runtime/Engine/Public/Slate/SlateTextures.h
+++ b/Engine/Source/Runtime/Engine/Public/Slate/SlateTextures.h
@@ -7,7 +7,7 @@ struct FSlateTextureData;
 /**
  * Encapsulates a Texture2DRHIRef for use by a Slate rendering implementation                   
  */
-class ENGINE_API FSlateTexture2DRHIRef : public TSlateTexture<FTexture2DRHIRef>, public FRenderResource
+class ENGINE_API FSlateTexture2DRHIRef : public TSlateTexture<FTexture2DRHIRef>, public FSlateUpdatableTexture, public FRenderResource
 {
 public:
 	FSlateTexture2DRHIRef( FTexture2DRHIRef InRef, uint32 InWidth, uint32 InHeight );
@@ -62,6 +62,13 @@ public:
 	 * Clears texture data being used.  Can only be accessed on the render thread                   
 	 */
 	void ClearTextureData();
+
+	// FSlateUpdatableTexture interface
+	virtual FSlateShaderResource* GetSlateResource() override {return this;}
+	virtual FRenderResource* GetRenderResource() override {return this;}
+	virtual void ResizeTexture( uint32 Width, uint32 Height ) override;
+	virtual void UpdateTexture(const TArray<uint8>& Bytes) override;
+
 protected:
 	/** Width of this texture */
 	uint32 Width;
diff --git a/Engine/Source/Runtime/SlateCore/Public/Rendering/SlateRenderer.h b/Engine/Source/Runtime/SlateCore/Public/Rendering/SlateRenderer.h
--- a/Engine/Source/Runtime/SlateCore/Public/Rendering/SlateRenderer.h
+++ b/Engine/Source/Runtime/SlateCore/Public/Rendering/SlateRenderer.h
@@ -5,6 +5,11 @@
 
 class SWindow;
 class FRHITexture2D;
+class ISlateStyle;
+class FSlateDrawBuffer;
+class FSlateUpdatableTexture;
+
+struct FSlateBrush;
 
 typedef FRHITexture2D* FTexture2DRHIParamRef;
 
@@ -201,6 +206,23 @@ public:
 	 */
 	virtual void SetWindowRenderTarget(const SWindow& Window, FTexture2DRHIParamRef RT) {}
 
+	/**
+	 * Create an updatable texture that can receive new data dynamically
+	 *
+	 * @param	Width	Initial width of the texture
+	 * @param	Height	Initial height of the texture
+	 *
+	 * @return	Newly created updatable texture
+	 */
+	virtual FSlateUpdatableTexture* CreateUpdatableTexture(uint32 Width, uint32 Height) = 0;
+
+	/**
+	 * Return an updatable texture to the renderer for release
+	 *
+	 * @param	Texture	The texture we are releasing (should not use this pointer after calling)
+	 */
+	virtual void ReleaseUpdatableTexture(FSlateUpdatableTexture* Texture) = 0;
+
 private:
 
 	// Non-copyable
diff --git a/Engine/Source/Runtime/SlateCore/Public/SlateCore.h b/Engine/Source/Runtime/SlateCore/Public/SlateCore.h
--- a/Engine/Source/Runtime/SlateCore/Public/SlateCore.h
+++ b/Engine/Source/Runtime/SlateCore/Public/SlateCore.h
@@ -102,6 +102,7 @@ DECLARE_MEMORY_STAT_EXTERN(TEXT("Index Buffer Memory"), STAT_SlateIndexBufferMem
 // Textures
 #include "SlateShaderResource.h"
 #include "SlateTextureData.h"
+#include "SlateUpdatableTexture.h"
 #include "TextureAtlas.h"
 #include "TextureManager.h"
 
diff --git a/Engine/Source/Runtime/SlateCore/Public/Textures/SlateUpdatableTexture.h b/Engine/Source/Runtime/SlateCore/Public/Textures/SlateUpdatableTexture.h
new file mode 100644
--- /dev/null
+++ b/Engine/Source/Runtime/SlateCore/Public/Textures/SlateUpdatableTexture.h
@@ -0,0 +1,45 @@
+// Copyright 1998-2014 Epic Games, Inc. All Rights Reserved.
+
+#pragma once
+
+class FSlateShaderResource;
+class FRenderResource;
+
+/**
+ * An interface to deal with a slate texture that can be updated dynamically
+ */
+class FSlateUpdatableTexture
+{
+public:
+	/** Virtual destructor */
+	virtual ~FSlateUpdatableTexture() {}
+
+	/**
+	 * Gets the interface to the underlying platform independent texture
+	 *
+	 * @return	FSlateShaderResource*	pointer to the shader resource
+	 */
+	virtual FSlateShaderResource* GetSlateResource() = 0;
+
+	/**
+	 * Gets the interface to the underlying render resource (may not always be used)
+	 *
+	 * @return	FRenderResource*	pointer to the render resource
+	 */
+	virtual FRenderResource* GetRenderResource() {return nullptr;}
+
+	/**
+	 * Resize the texture.
+	 *
+	 * @param Width New texture width
+	 * @param Height New texture height
+	 */
+	virtual void ResizeTexture( uint32 Width, uint32 Height ) = 0;
+
+	/**
+	 * Updates the texture contents via a byte array
+	 *
+	 * @param Bytes Array of texture data
+	 */
+	virtual void UpdateTexture(const TArray<uint8>& Bytes) = 0;
+};
diff --git a/Engine/Source/Runtime/SlateRHIRenderer/Private/SlateRHIRenderer.cpp b/Engine/Source/Runtime/SlateRHIRenderer/Private/SlateRHIRenderer.cpp
--- a/Engine/Source/Runtime/SlateRHIRenderer/Private/SlateRHIRenderer.cpp
+++ b/Engine/Source/Runtime/SlateRHIRenderer/Private/SlateRHIRenderer.cpp
@@ -1143,6 +1143,35 @@ void FSlateRHIRenderer::SetColorVisionDeficiencyType( uint32 Type )
 	GSlateShaderColorVisionDeficiencyType = Type;
 }
 
+FSlateUpdatableTexture* FSlateRHIRenderer::CreateUpdatableTexture(uint32 Width, uint32 Height)
+{
+	const bool bCreateEmptyTexture = true;
+	FSlateTexture2DRHIRef* NewTexture = new FSlateTexture2DRHIRef(Width, Height, PF_B8G8R8A8, nullptr, TexCreate_Dynamic, bCreateEmptyTexture);
+	if (IsInRenderingThread())
+	{
+		NewTexture->InitResource();
+	}
+	else
+	{
+		BeginInitResource(NewTexture);
+	}
+	return NewTexture;
+}
+
+void FSlateRHIRenderer::ReleaseUpdatableTexture(FSlateUpdatableTexture* Texture)
+{
+	if (IsInRenderingThread())
+	{
+		Texture->GetRenderResource()->InitResource();
+	}
+	else
+	{
+		BeginReleaseResource(Texture->GetRenderResource());
+		FlushRenderingCommands();
+	}
+	delete Texture;
+}
+
 bool FSlateRHIRenderer::AreShadersInitialized() const
 {
 #if WITH_EDITORONLY_DATA
diff --git a/Engine/Source/Runtime/SlateRHIRenderer/Private/SlateRHIRenderer.h b/Engine/Source/Runtime/SlateRHIRenderer/Private/SlateRHIRenderer.h
--- a/Engine/Source/Runtime/SlateRHIRenderer/Private/SlateRHIRenderer.h
+++ b/Engine/Source/Runtime/SlateRHIRenderer/Private/SlateRHIRenderer.h
@@ -166,6 +166,8 @@ public:
 	virtual bool GenerateDynamicImageResource( FName ResourceName, uint32 Width, uint32 Height, const TArray< uint8 >& Bytes ) override;
 	virtual void* GetViewportResource( const SWindow& Window ) override;
 	virtual void SetColorVisionDeficiencyType( uint32 Type ) override;
+	virtual FSlateUpdatableTexture* CreateUpdatableTexture(uint32 Width, uint32 Height) override;
+	virtual void ReleaseUpdatableTexture(FSlateUpdatableTexture* Texture) override;
 
 	/** Draws windows from a FSlateDrawBuffer on the render thread */
 	void DrawWindow_RenderThread(FRHICommandListImmediate& RHICmdList, const FSlateRHIRenderer::FViewportInfo& ViewportInfo, const FSlateWindowElementList& WindowElementList, bool bLockToVsync);
