Make PhysicalMaterial blueprint readonly. Really don't want users trying to modify assets at runtime using blueprint

TTP#343946

diff --git a/Engine/Source/Runtime/Engine/Classes/PhysicalMaterials/PhysicalMaterial.h b/Engine/Source/Runtime/Engine/Classes/PhysicalMaterials/PhysicalMaterial.h
--- a/Engine/Source/Runtime/Engine/Classes/PhysicalMaterials/PhysicalMaterial.h
+++ b/Engine/Source/Runtime/Engine/Classes/PhysicalMaterials/PhysicalMaterial.h
@@ -43,11 +43,11 @@ class ENGINE_API UPhysicalMaterial : public UObject
 	//
 	
 	/** Friction value of surface, controls how easily things can slide on this surface */
-	UPROPERTY(EditAnywhere, BlueprintReadWrite, Category=PhysicalMaterial)
+	UPROPERTY(EditAnywhere, BlueprintReadOnly, Category=PhysicalMaterial)
 	float Friction;
 
 	/** Friction combine mode, controls how friction is computed for multiple materials. */
-	UPROPERTY(EditAnywhere, BlueprintReadWrite, Category = PhysicalMaterial, meta = (editcondition = "bOverrideFrictionCombineMode"))
+	UPROPERTY(EditAnywhere, BlueprintReadOnly, Category = PhysicalMaterial, meta = (editcondition = "bOverrideFrictionCombineMode"))
 	TEnumAsByte<EFrictionCombineMode::Type> FrictionCombineMode;
 
 	/** If set we will use the FrictionCombineMode of this material, instead of the FrictionCombineMode found in the project settings. */
@@ -55,7 +55,7 @@ class ENGINE_API UPhysicalMaterial : public UObject
 	bool bOverrideFrictionCombineMode;
 
 	/** Resitution or 'bouncyness' of this surface */
-	UPROPERTY(EditAnywhere, BlueprintReadWrite, Category=PhysicalMaterial)
+	UPROPERTY(EditAnywhere, BlueprintReadOnly, Category = PhysicalMaterial)
 	float Restitution;
 
 	//
@@ -63,7 +63,7 @@ class ENGINE_API UPhysicalMaterial : public UObject
 	//
 	
 	/** Used with the shape of the object to calculate its mass properties. The higher the number, the heavier the object. g per cubic cm. */
-	UPROPERTY(EditAnywhere, BlueprintReadWrite, Category=PhysicalMaterial)
+	UPROPERTY(EditAnywhere, BlueprintReadOnly, Category = PhysicalMaterial)
 	float Density;
 
 	/** 
@@ -71,11 +71,11 @@ class ENGINE_API UPhysicalMaterial : public UObject
 	 *	In actuality, larger objects do not tend to be solid, and become more like 'shells' (e.g. a car is not a solid piece of metal).
 	 *	Values are clamped to 1 or less.
 	 */
-	UPROPERTY(EditAnywhere, BlueprintReadWrite, Category=Advanced)
+	UPROPERTY(EditAnywhere, BlueprintReadOnly, Category = Advanced)
 	float RaiseMassToPower;
 
 	/** How much to scale the damage threshold by on any destructible we are applied to */
-	UPROPERTY(EditAnywhere, BlueprintReadWrite, Category=Destruction)
+	UPROPERTY(EditAnywhere, BlueprintReadOnly, Category = Destruction)
 	float DestructibleDamageThresholdScale;
 
 	UPROPERTY(/*deprecated*/)
@@ -84,7 +84,7 @@ class ENGINE_API UPhysicalMaterial : public UObject
 	/**
 	 * To edit surface type for your project, use ProjectSettings/Physics/PhysicalSurface section
 	*/
-	UPROPERTY(EditAnywhere, BlueprintReadWrite, Category=PhysicalProperties)
+	UPROPERTY(EditAnywhere, BlueprintReadOnly, Category = PhysicalProperties)
 	TEnumAsByte<EPhysicalSurface> SurfaceType;
 
 	/** Overall tire friction scalar for every type of tire. This value is multiplied against our parents' values. */
