exposing postprocess setting to blueprints

diff --git a/Engine/Source/Runtime/Engine/Classes/Engine/Scene.h b/Engine/Source/Runtime/Engine/Classes/Engine/Scene.h
--- a/Engine/Source/Runtime/Engine/Classes/Engine/Scene.h
+++ b/Engine/Source/Runtime/Engine/Classes/Engine/Scene.h
@@ -340,56 +340,56 @@ struct FPostProcessSettings
 	// -----------------------------------------------------------------------
 
 
-	UPROPERTY(interp, Category=Film, meta=(editcondition = "bOverride_FilmWhitePoint", DisplayName = "Tint"))
+	UPROPERTY(interp, BlueprintReadWrite, Category=Film, meta=(editcondition = "bOverride_FilmWhitePoint", DisplayName = "Tint"))
 	FLinearColor FilmWhitePoint;
-	UPROPERTY(interp, Category=Film, AdvancedDisplay, meta=(editcondition = "bOverride_FilmShadowTint", DisplayName = "Tint Shadow"))
+	UPROPERTY(interp, BlueprintReadWrite, Category=Film, AdvancedDisplay, meta=(editcondition = "bOverride_FilmShadowTint", DisplayName = "Tint Shadow"))
 	FLinearColor FilmShadowTint;
-	UPROPERTY(interp, Category=Film, AdvancedDisplay, meta=(UIMin = "0.0", UIMax = "1.0", editcondition = "bOverride_FilmShadowTintBlend", DisplayName = "Tint Shadow Blend"))
+	UPROPERTY(interp, BlueprintReadWrite, Category=Film, AdvancedDisplay, meta=(UIMin = "0.0", UIMax = "1.0", editcondition = "bOverride_FilmShadowTintBlend", DisplayName = "Tint Shadow Blend"))
 	float FilmShadowTintBlend;
-	UPROPERTY(interp, Category=Film, AdvancedDisplay, meta=(UIMin = "0.0", UIMax = "1.0", editcondition = "bOverride_FilmShadowTintAmount", DisplayName = "Tint Shadow Amount"))
+	UPROPERTY(interp, BlueprintReadWrite, Category=Film, AdvancedDisplay, meta=(UIMin = "0.0", UIMax = "1.0", editcondition = "bOverride_FilmShadowTintAmount", DisplayName = "Tint Shadow Amount"))
 	float FilmShadowTintAmount;
 
-	UPROPERTY(interp, Category=Film, meta=(UIMin = "0.0", UIMax = "2.0", editcondition = "bOverride_FilmSaturation", DisplayName = "Saturation"))
+	UPROPERTY(interp, BlueprintReadWrite, Category=Film, meta=(UIMin = "0.0", UIMax = "2.0", editcondition = "bOverride_FilmSaturation", DisplayName = "Saturation"))
 	float FilmSaturation;
-	UPROPERTY(interp, Category=Film, AdvancedDisplay, meta=(editcondition = "bOverride_FilmChannelMixerRed", DisplayName = "Channel Mixer Red"))
+	UPROPERTY(interp, BlueprintReadWrite, Category=Film, AdvancedDisplay, meta=(editcondition = "bOverride_FilmChannelMixerRed", DisplayName = "Channel Mixer Red"))
 	FLinearColor FilmChannelMixerRed;
-	UPROPERTY(interp, Category=Film, AdvancedDisplay, meta=(editcondition = "bOverride_FilmChannelMixerGreen", DisplayName = "Channel Mixer Green"))
+	UPROPERTY(interp, BlueprintReadWrite, Category=Film, AdvancedDisplay, meta=(editcondition = "bOverride_FilmChannelMixerGreen", DisplayName = "Channel Mixer Green"))
 	FLinearColor FilmChannelMixerGreen;
-	UPROPERTY(interp, Category=Film, AdvancedDisplay, meta=(editcondition = "bOverride_FilmChannelMixerBlue", DisplayName = " Channel Mixer Blue"))
+	UPROPERTY(interp, BlueprintReadWrite, Category=Film, AdvancedDisplay, meta=(editcondition = "bOverride_FilmChannelMixerBlue", DisplayName = " Channel Mixer Blue"))
 	FLinearColor FilmChannelMixerBlue;
 
-	UPROPERTY(interp, Category=Film, meta=(UIMin = "0.0", UIMax = "1.0", editcondition = "bOverride_FilmContrast", DisplayName = "Contrast"))
+	UPROPERTY(interp, BlueprintReadWrite, Category=Film, meta=(UIMin = "0.0", UIMax = "1.0", editcondition = "bOverride_FilmContrast", DisplayName = "Contrast"))
 	float FilmContrast;
-	UPROPERTY(interp, Category=Film, AdvancedDisplay, meta=(UIMin = "0.0", UIMax = "1.0", editcondition = "bOverride_FilmToeAmount", DisplayName = "Crush Shadows"))
+	UPROPERTY(interp, BlueprintReadWrite, Category=Film, AdvancedDisplay, meta=(UIMin = "0.0", UIMax = "1.0", editcondition = "bOverride_FilmToeAmount", DisplayName = "Crush Shadows"))
 	float FilmToeAmount;
-	UPROPERTY(interp, Category=Film, AdvancedDisplay, meta=(UIMin = "0.0", UIMax = "1.0", editcondition = "bOverride_FilmHealAmount", DisplayName = "Crush Highlights"))
+	UPROPERTY(interp, BlueprintReadWrite, Category=Film, AdvancedDisplay, meta=(UIMin = "0.0", UIMax = "1.0", editcondition = "bOverride_FilmHealAmount", DisplayName = "Crush Highlights"))
 	float FilmHealAmount;
-	UPROPERTY(interp, Category=Film, AdvancedDisplay, meta=(UIMin = "1.0", UIMax = "4.0", editcondition = "bOverride_FilmDynamicRange", DisplayName = "Dynamic Range"))
+	UPROPERTY(interp, BlueprintReadWrite, Category=Film, AdvancedDisplay, meta=(UIMin = "1.0", UIMax = "4.0", editcondition = "bOverride_FilmDynamicRange", DisplayName = "Dynamic Range"))
 	float FilmDynamicRange;
 
 
 
 	/** Scene tint color */
-	UPROPERTY(interp, Category=SceneColor, AdvancedDisplay, meta=(editcondition = "bOverride_SceneColorTint", HideAlphaChannel))
+	UPROPERTY(interp, BlueprintReadWrite, Category=SceneColor, AdvancedDisplay, meta=(editcondition = "bOverride_SceneColorTint", HideAlphaChannel))
 	FLinearColor SceneColorTint;
 	
 	/** in percent, Scene chromatic aberration / color fringe (camera imperfection) to simulate an artifact that happens in real-world lens, mostly visible in the image corners. */
-	UPROPERTY(interp, Category=SceneColor, meta=(UIMin = "0.0", UIMax = "5.0", editcondition = "bOverride_SceneFringeIntensity", DisplayName = "Fringe Intensity"))
+	UPROPERTY(interp, BlueprintReadWrite, Category=SceneColor, meta=(UIMin = "0.0", UIMax = "5.0", editcondition = "bOverride_SceneFringeIntensity", DisplayName = "Fringe Intensity"))
 	float SceneFringeIntensity;
 
 	/** 0..1, Scene chromatic aberration / color fringe (camera imperfection) to simulate an artifact that happens in real-world lens, mostly visible in the image corners. */
-	UPROPERTY(interp, Category=SceneColor, AdvancedDisplay, meta=(ClampMin = "0.0", UIMax = "1.0", editcondition = "bOverride_SceneFringeSaturation", DisplayName = "Fringe Saturation"))
+	UPROPERTY(interp, BlueprintReadWrite, Category=SceneColor, AdvancedDisplay, meta=(ClampMin = "0.0", UIMax = "1.0", editcondition = "bOverride_SceneFringeSaturation", DisplayName = "Fringe Saturation"))
 	float SceneFringeSaturation;
 
 	/** Multiplier for all bloom contributions >=0: off, 1(default), >1 brighter */
-	UPROPERTY(interp, Category=Bloom, meta=(ClampMin = "0.0", UIMax = "8.0", editcondition = "bOverride_BloomIntensity", DisplayName = "Intensity"))
+	UPROPERTY(interp, BlueprintReadWrite, Category=Bloom, meta=(ClampMin = "0.0", UIMax = "8.0", editcondition = "bOverride_BloomIntensity", DisplayName = "Intensity"))
 	float BloomIntensity;
 
 	/**
 	 * minimum brightness the bloom starts having effect
 	 * -1:all pixels affect bloom equally (dream effect), 0:all pixels affect bloom brights more, 1(default), >1 brighter
 	 */
-	UPROPERTY(interp, Category=Bloom, AdvancedDisplay, meta=(ClampMin = "-1.0", UIMin = "0.0", UIMax = "8.0", editcondition = "bOverride_BloomThreshold", DisplayName = "Threshold"))
+	UPROPERTY(interp, BlueprintReadWrite, Category=Bloom, AdvancedDisplay, meta=(ClampMin = "-1.0", UIMin = "0.0", UIMax = "8.0", editcondition = "bOverride_BloomThreshold", DisplayName = "Threshold"))
 	float BloomThreshold;
 
 	/**
@@ -397,59 +397,59 @@ struct FPostProcessSettings
 	 * (is done in 1/2 resolution, larger values cost more performance, good for high frequency details)
 	 * >=0: can be clamped because of shader limitations
 	 */
-	UPROPERTY(interp, Category=Bloom, AdvancedDisplay, meta=(ClampMin = "0.0", UIMax = "4.0", editcondition = "bOverride_Bloom1Size", DisplayName = "#1 Size"))
+	UPROPERTY(interp, BlueprintReadWrite, Category=Bloom, AdvancedDisplay, meta=(ClampMin = "0.0", UIMax = "4.0", editcondition = "bOverride_Bloom1Size", DisplayName = "#1 Size"))
 	float Bloom1Size;
 	/**
 	 * Diameter size for Bloom2 in percent of the screen width
 	 * (is done in 1/4 resolution, larger values cost more performance)
 	 * >=0: can be clamped because of shader limitations
 	 */
-	UPROPERTY(interp, Category=Bloom, AdvancedDisplay, meta=(ClampMin = "0.0", UIMax = "8.0", editcondition = "bOverride_Bloom2Size", DisplayName = "#2 Size"))
+	UPROPERTY(interp, BlueprintReadWrite, Category=Bloom, AdvancedDisplay, meta=(ClampMin = "0.0", UIMax = "8.0", editcondition = "bOverride_Bloom2Size", DisplayName = "#2 Size"))
 	float Bloom2Size;
 	/**
 	 * Diameter size for Bloom3 in percent of the screen width
 	 * (is done in 1/8 resolution, larger values cost more performance)
 	 * >=0: can be clamped because of shader limitations
 	 */
-	UPROPERTY(interp, Category=Bloom, AdvancedDisplay, meta=(ClampMin = "0.0", UIMax = "16.0", editcondition = "bOverride_Bloom3Size", DisplayName = "#3 Size"))
+	UPROPERTY(interp, BlueprintReadWrite, Category=Bloom, AdvancedDisplay, meta=(ClampMin = "0.0", UIMax = "16.0", editcondition = "bOverride_Bloom3Size", DisplayName = "#3 Size"))
 	float Bloom3Size;
 	/**
 	 * Diameter size for Bloom4 in percent of the screen width
 	 * (is done in 1/16 resolution, larger values cost more performance, best for wide contributions)
 	 * >=0: can be clamped because of shader limitations
 	 */
-	UPROPERTY(interp, Category=Bloom, AdvancedDisplay, meta=(ClampMin = "0.0", UIMax = "32.0", editcondition = "bOverride_Bloom4Size", DisplayName = "#4 Size"))
+	UPROPERTY(interp, BlueprintReadWrite, Category=Bloom, AdvancedDisplay, meta=(ClampMin = "0.0", UIMax = "32.0", editcondition = "bOverride_Bloom4Size", DisplayName = "#4 Size"))
 	float Bloom4Size;
 	/**
 	 * Diameter size for Bloom5 in percent of the screen width
 	 * (is done in 1/32 resolution, larger values cost more performance, best for wide contributions)
 	 * >=0: can be clamped because of shader limitations
 	 */
-	UPROPERTY(interp, Category=Bloom, AdvancedDisplay, meta=(ClampMin = "0.0", UIMax = "64.0", editcondition = "bOverride_Bloom5Size", DisplayName = "#5 Size"))
+	UPROPERTY(interp, BlueprintReadWrite, Category=Bloom, AdvancedDisplay, meta=(ClampMin = "0.0", UIMax = "64.0", editcondition = "bOverride_Bloom5Size", DisplayName = "#5 Size"))
 	float Bloom5Size;
 
 	/** Bloom1 tint color */
-	UPROPERTY(interp, Category=Bloom, AdvancedDisplay, meta=(editcondition = "bOverride_Bloom1Tint", DisplayName = "#1 Tint", HideAlphaChannel))
+	UPROPERTY(interp, BlueprintReadWrite, Category=Bloom, AdvancedDisplay, meta=(editcondition = "bOverride_Bloom1Tint", DisplayName = "#1 Tint", HideAlphaChannel))
 	FLinearColor Bloom1Tint;
 	/** Bloom2 tint color */
-	UPROPERTY(interp, Category=Bloom, AdvancedDisplay, meta=(editcondition = "bOverride_Bloom2Tint", DisplayName = "#2 Tint", HideAlphaChannel))
+	UPROPERTY(interp, BlueprintReadWrite, Category=Bloom, AdvancedDisplay, meta=(editcondition = "bOverride_Bloom2Tint", DisplayName = "#2 Tint", HideAlphaChannel))
 	FLinearColor Bloom2Tint;
 	/** Bloom3 tint color */
-	UPROPERTY(interp, Category=Bloom, AdvancedDisplay, meta=(editcondition = "bOverride_Bloom3Tint", DisplayName = "#3 Tint", HideAlphaChannel))
+	UPROPERTY(interp, BlueprintReadWrite, Category=Bloom, AdvancedDisplay, meta=(editcondition = "bOverride_Bloom3Tint", DisplayName = "#3 Tint", HideAlphaChannel))
 	FLinearColor Bloom3Tint;
 	/** Bloom4 tint color */
-	UPROPERTY(interp, Category=Bloom, AdvancedDisplay, meta=(editcondition = "bOverride_Bloom4Tint", DisplayName = "#4 Tint", HideAlphaChannel))
+	UPROPERTY(interp, BlueprintReadWrite, Category=Bloom, AdvancedDisplay, meta=(editcondition = "bOverride_Bloom4Tint", DisplayName = "#4 Tint", HideAlphaChannel))
 	FLinearColor Bloom4Tint;
 	/** Bloom5 tint color */
-	UPROPERTY(interp, Category=Bloom, AdvancedDisplay, meta=(editcondition = "bOverride_Bloom5Tint", DisplayName = "#5 Tint", HideAlphaChannel))
+	UPROPERTY(interp, BlueprintReadWrite, Category=Bloom, AdvancedDisplay, meta=(editcondition = "bOverride_Bloom5Tint", DisplayName = "#5 Tint", HideAlphaChannel))
 	FLinearColor Bloom5Tint;
 
 	/** BloomDirtMask intensity */
-	UPROPERTY(interp, Category=Bloom, meta=(ClampMin = "0.0", UIMax = "8.0", editcondition = "bOverride_BloomDirtMaskIntensity", DisplayName = "Dirt Mask Intensity"))
+	UPROPERTY(interp, BlueprintReadWrite, Category=Bloom, meta=(ClampMin = "0.0", UIMax = "8.0", editcondition = "bOverride_BloomDirtMaskIntensity", DisplayName = "Dirt Mask Intensity"))
 	float BloomDirtMaskIntensity;
 
 	/** BloomDirtMask tint color */
-	UPROPERTY(interp, Category=Bloom, AdvancedDisplay, meta=(editcondition = "bOverride_BloomDirtMaskTint", DisplayName = "Dirt Mask Tint"))
+	UPROPERTY(interp, BlueprintReadWrite, Category=Bloom, AdvancedDisplay, meta=(editcondition = "bOverride_BloomDirtMaskTint", DisplayName = "Dirt Mask Tint"))
 	FLinearColor BloomDirtMaskTint;
 
 	/**
@@ -459,15 +459,15 @@ struct FPostProcessSettings
 	class UTexture* BloomDirtMask;	// The plan is to replace this texture with many small textures quads for better performance, more control and to animate the effect.
 
 	/** How strong the dynamic GI from the LPV should be. 0.0 is off, 1.0 is the "normal" value, but higher values can be used to boost the effect*/
-	UPROPERTY(interp, Category=LightPropagationVolume, meta=(editcondition = "bOverride_LPVIntensity", UIMin = "0", UIMax = "20", DisplayName = "Intensity") )
+	UPROPERTY(interp, BlueprintReadWrite, Category=LightPropagationVolume, meta=(editcondition = "bOverride_LPVIntensity", UIMin = "0", UIMax = "20", DisplayName = "Intensity") )
 	float LPVIntensity;
 
 	/** CURRENTLY DISABLED - The strength of the warp offset for reducing light bleeding. 0.0 is off, 1.0 is the "normal" value, but higher values can be used to boost the effect*/
-	UPROPERTY(interp, Category=LightPropagationVolume, AdvancedDisplay, meta=(editcondition = "bOverride_LPVWarpIntensity", UIMin = "0", UIMax = "1", DisplayName = "(DISABLED) Grid Warp Intensity"))
+	UPROPERTY(interp, BlueprintReadWrite, Category=LightPropagationVolume, AdvancedDisplay, meta=(editcondition = "bOverride_LPVWarpIntensity", UIMin = "0", UIMax = "1", DisplayName = "(DISABLED) Grid Warp Intensity"))
 	float LPVWarpIntensity;
 
 	/** Bias applied to light injected into the LPV in cell units. Increase to reduce bleeding through thin walls*/
-	UPROPERTY(interp, Category=LightPropagationVolume, meta=(editcondition = "bOverride_LPVVplInjectionBias", UIMin = "0", UIMax = "2", DisplayName = "Light Injection Bias") )
+	UPROPERTY(interp, BlueprintReadWrite, Category=LightPropagationVolume, meta=(editcondition = "bOverride_LPVVplInjectionBias", UIMin = "0", UIMax = "2", DisplayName = "Light Injection Bias") )
 	float LPVVplInjectionBias;
 
 	/** The size of the LPV volume, in Unreal units*/
@@ -475,32 +475,32 @@ struct FPostProcessSettings
 	float LPVSize;
 
 	/** Secondary occlusion strength (bounce light shadows). Set to 0 to disable*/
-	UPROPERTY(interp, Category=LightPropagationVolume, meta=(editcondition = "bOverride_LPVSecondaryOcclusionIntensity", UIMin = "0", UIMax = "1", DisplayName = "Secondary Occlusion Intensity") )
+	UPROPERTY(interp, BlueprintReadWrite, Category=LightPropagationVolume, meta=(editcondition = "bOverride_LPVSecondaryOcclusionIntensity", UIMin = "0", UIMax = "1", DisplayName = "Secondary Occlusion Intensity") )
 	float LPVSecondaryOcclusionIntensity;
 
 	/** Secondary bounce light strength (bounce light shadows). Set to 0 to disable*/
-	UPROPERTY(interp, Category=LightPropagationVolume, meta=(editcondition = "bOverride_LPVSecondaryBounceIntensity", UIMin = "0", UIMax = "1", DisplayName = "Secondary Bounce Intensity") )
+	UPROPERTY(interp, BlueprintReadWrite, Category=LightPropagationVolume, meta=(editcondition = "bOverride_LPVSecondaryBounceIntensity", UIMin = "0", UIMax = "1", DisplayName = "Secondary Bounce Intensity") )
 	float LPVSecondaryBounceIntensity;
 
 	/** Bias applied to the geometry volume in cell units. Increase to reduce darkening due to secondary occlusion */
-	UPROPERTY(interp, Category=LightPropagationVolume, AdvancedDisplay, meta=(editcondition = "bOverride_LPVGeometryVolumeBias", UIMin = "0", UIMax = "2", DisplayName = "Geometry Volume Bias"))
+	UPROPERTY(interp, BlueprintReadWrite, Category=LightPropagationVolume, AdvancedDisplay, meta=(editcondition = "bOverride_LPVGeometryVolumeBias", UIMin = "0", UIMax = "2", DisplayName = "Geometry Volume Bias"))
 	float LPVGeometryVolumeBias;
 
-	UPROPERTY(interp, Category=LightPropagationVolume, meta=(editcondition = "bOverride_LPVEmissiveInjectionIntensity", UIMin = "0", UIMax = "20", DisplayName = "Emissive Injection Intensity") )
+	UPROPERTY(interp, BlueprintReadWrite, Category=LightPropagationVolume, meta=(editcondition = "bOverride_LPVEmissiveInjectionIntensity", UIMin = "0", UIMax = "20", DisplayName = "Emissive Injection Intensity") )
 	float LPVEmissiveInjectionIntensity;
 
 	/** How strong light transmission from the LPV should be. 0.0 is off, 1.0 is the "normal" value, but higher values can be used to boost the effect*/
-	UPROPERTY(interp, Category=LightPropagationVolume, meta=(editcondition = "bOverride_LPVTransmissionIntensity", UIMin = "0", UIMax = "2", DisplayName = "Transmission Intensity") )
+	UPROPERTY(interp, BlueprintReadWrite, Category=LightPropagationVolume, meta=(editcondition = "bOverride_LPVTransmissionIntensity", UIMin = "0", UIMax = "2", DisplayName = "Transmission Intensity") )
 	float LPVTransmissionIntensity;
 
 	/** AmbientCubemap tint color */
-	UPROPERTY(interp, Category=AmbientCubemap, AdvancedDisplay, meta=(editcondition = "bOverride_AmbientCubemapTint", DisplayName = "Tint"))
+	UPROPERTY(interp, BlueprintReadWrite, Category=AmbientCubemap, AdvancedDisplay, meta=(editcondition = "bOverride_AmbientCubemapTint", DisplayName = "Tint"))
 	FLinearColor AmbientCubemapTint;
 	/**
 	 * To scale the Ambient cubemap brightness
 	 * >=0: off, 1(default), >1 brighter
 	 */
-	UPROPERTY(interp, Category=AmbientCubemap, meta=(ClampMin = "0.0", UIMax = "4.0", editcondition = "bOverride_AmbientCubemapIntensity", DisplayName = "Intensity"))
+	UPROPERTY(interp, BlueprintReadWrite, Category=AmbientCubemap, meta=(ClampMin = "0.0", UIMax = "4.0", editcondition = "bOverride_AmbientCubemapIntensity", DisplayName = "Intensity"))
 	float AmbientCubemapIntensity;
 	/** The Ambient cubemap (Affects diffuse and specular shading), blends additively which if different from all other settings here */
 	UPROPERTY(EditAnywhere, BlueprintReadWrite, Category=AmbientCubemap, meta=(DisplayName = "Cubemap Texture"))
@@ -513,7 +513,7 @@ struct FPostProcessSettings
 	 * bright spots.
 	 * >0, <100, good values are in the range 70 .. 80
 	 */
-	UPROPERTY(interp, Category=AutoExposure, AdvancedDisplay, meta=(ClampMin = "0.0", ClampMax = "100.0", editcondition = "bOverride_AutoExposureLowPercent", DisplayName = "Low Percent"))
+	UPROPERTY(interp, BlueprintReadWrite, Category=AutoExposure, AdvancedDisplay, meta=(ClampMin = "0.0", ClampMax = "100.0", editcondition = "bOverride_AutoExposureLowPercent", DisplayName = "Low Percent"))
 	float AutoExposureLowPercent;
 
 	/**
@@ -523,7 +523,7 @@ struct FPostProcessSettings
 	 * bright spots.
 	 * >0, <100, good values are in the range 80 .. 95
 	 */
-	UPROPERTY(interp, Category=AutoExposure, AdvancedDisplay, meta=(ClampMin = "0.0", ClampMax = "100.0", editcondition = "bOverride_AutoExposureHighPercent", DisplayName = "High Percent"))
+	UPROPERTY(interp, BlueprintReadWrite, Category=AutoExposure, AdvancedDisplay, meta=(ClampMin = "0.0", ClampMax = "100.0", editcondition = "bOverride_AutoExposureHighPercent", DisplayName = "High Percent"))
 	float AutoExposureHighPercent;
 
 	/**
@@ -533,7 +533,7 @@ struct FPostProcessSettings
 	 * effect and defined the HDR range - you don't want to change that late in the project development.
 	 * Eye Adaptation is disabled if MinBrightness = MaxBrightness
 	 */
-	UPROPERTY(interp, Category=AutoExposure, meta=(ClampMin = "0.0", UIMax = "10.0", editcondition = "bOverride_AutoExposureMinBrightness", DisplayName = "Min Brightness"))
+	UPROPERTY(interp, BlueprintReadWrite, Category=AutoExposure, meta=(ClampMin = "0.0", UIMax = "10.0", editcondition = "bOverride_AutoExposureMinBrightness", DisplayName = "Min Brightness"))
 	float AutoExposureMinBrightness;
 
 	/**
@@ -543,46 +543,46 @@ struct FPostProcessSettings
 	 * effect and defined the HDR range - you don't want to change that late in the project development.
 	 * Eye Adaptation is disabled if MinBrightness = MaxBrightness
 	 */
-	UPROPERTY(interp, Category=AutoExposure, meta=(ClampMin = "0.0", UIMax = "10.0", editcondition = "bOverride_AutoExposureMaxBrightness", DisplayName = "Max Brightness"))
+	UPROPERTY(interp, BlueprintReadWrite, Category=AutoExposure, meta=(ClampMin = "0.0", UIMax = "10.0", editcondition = "bOverride_AutoExposureMaxBrightness", DisplayName = "Max Brightness"))
 	float AutoExposureMaxBrightness;
 
 	/** >0 */
-	UPROPERTY(interp, Category=AutoExposure, AdvancedDisplay, meta=(ClampMin = "0.02", UIMax = "20.0", editcondition = "bOverride_AutoExposureSpeedUp", DisplayName = "Speed Up"))
+	UPROPERTY(interp, BlueprintReadWrite, Category=AutoExposure, AdvancedDisplay, meta=(ClampMin = "0.02", UIMax = "20.0", editcondition = "bOverride_AutoExposureSpeedUp", DisplayName = "Speed Up"))
 	float AutoExposureSpeedUp;
 
 	/** >0 */
-	UPROPERTY(interp, Category=AutoExposure, AdvancedDisplay, meta=(ClampMin = "0.02", UIMax = "20.0", editcondition = "bOverride_AutoExposureSpeedDown", DisplayName = "Speed Down"))
+	UPROPERTY(interp, BlueprintReadWrite, Category=AutoExposure, AdvancedDisplay, meta=(ClampMin = "0.02", UIMax = "20.0", editcondition = "bOverride_AutoExposureSpeedDown", DisplayName = "Speed Down"))
 	float AutoExposureSpeedDown;
 
 	/**
 	 * Logarithmic adjustment for the exposure. Only used if a tonemapper is specified.
 	 * 0: no adjustment, -1:2x darker, -2:4x darker, 1:2x brighter, 2:4x brighter, ...
 	 */
-	UPROPERTY(interp, Category = AutoExposure, meta = (UIMin = "-8.0", UIMax = "8.0", editcondition = "bOverride_AutoExposureBias", DisplayName = "Exposure Bias"))
+	UPROPERTY(interp, BlueprintReadWrite, Category = AutoExposure, meta = (UIMin = "-8.0", UIMax = "8.0", editcondition = "bOverride_AutoExposureBias", DisplayName = "Exposure Bias"))
 	float AutoExposureBias;
 
 	/** temporary exposed until we found good values, -8: 1/256, -10: 1/1024 */
-	UPROPERTY(interp, Category=AutoExposure, AdvancedDisplay, meta=(UIMin = "-16", UIMax = "0.0", editcondition = "bOverride_HistogramLogMin"))
+	UPROPERTY(interp, BlueprintReadWrite, Category=AutoExposure, AdvancedDisplay, meta=(UIMin = "-16", UIMax = "0.0", editcondition = "bOverride_HistogramLogMin"))
 	float HistogramLogMin;
 
 	/** temporary exposed until we found good values 4: 16, 8: 256 */
-	UPROPERTY(interp, Category=AutoExposure, AdvancedDisplay, meta=(UIMin = "0.0", UIMax = "16.0", editcondition = "bOverride_HistogramLogMax"))
+	UPROPERTY(interp, BlueprintReadWrite, Category=AutoExposure, AdvancedDisplay, meta=(UIMin = "0.0", UIMax = "16.0", editcondition = "bOverride_HistogramLogMax"))
 	float HistogramLogMax;
 
 	/** Brightness scale of the image cased lens flares (linear) */
-	UPROPERTY(interp, Category=LensFlares, meta=(UIMin = "0.0", UIMax = "16.0", editcondition = "bOverride_LensFlareIntensity", DisplayName = "Intensity"))
+	UPROPERTY(interp, BlueprintReadWrite, Category=LensFlares, meta=(UIMin = "0.0", UIMax = "16.0", editcondition = "bOverride_LensFlareIntensity", DisplayName = "Intensity"))
 	float LensFlareIntensity;
 
 	/** Tint color for the image based lens flares. */
-	UPROPERTY(interp, Category=LensFlares, AdvancedDisplay, meta=(editcondition = "bOverride_LensFlareTint", DisplayName = "Tint"))
+	UPROPERTY(interp, BlueprintReadWrite, Category=LensFlares, AdvancedDisplay, meta=(editcondition = "bOverride_LensFlareTint", DisplayName = "Tint"))
 	FLinearColor LensFlareTint;
 
 	/** Size of the Lens Blur (in percent of the view width) that is done with the Bokeh texture (note: performance cost is radius*radius) */
-	UPROPERTY(interp, Category=LensFlares, meta=(UIMin = "0.0", UIMax = "32.0", editcondition = "bOverride_LensFlareBokehSize", DisplayName = "BokehSize"))
+	UPROPERTY(interp, BlueprintReadWrite, Category=LensFlares, meta=(UIMin = "0.0", UIMax = "32.0", editcondition = "bOverride_LensFlareBokehSize", DisplayName = "BokehSize"))
 	float LensFlareBokehSize;
 
 	/** Minimum brightness the lens flare starts having effect (this should be as high as possible to avoid the performance cost of blurring content that is too dark too see) */
-	UPROPERTY(interp, Category=LensFlares, AdvancedDisplay, meta=(UIMin = "0.1", UIMax = "32.0", editcondition = "bOverride_LensFlareThreshold", DisplayName = "Threshold"))
+	UPROPERTY(interp, BlueprintReadWrite, Category=LensFlares, AdvancedDisplay, meta=(UIMin = "0.1", UIMax = "32.0", editcondition = "bOverride_LensFlareThreshold", DisplayName = "Threshold"))
 	float LensFlareThreshold;
 
 	/** Defines the shape of the Bokeh when the image base lens flares are blurred, cannot be blended */
@@ -590,35 +590,35 @@ struct FPostProcessSettings
 	class UTexture* LensFlareBokehShape;
 
 	/** RGB defines the lens flare color, A it's position. This is a temporary solution. */
-	UPROPERTY(EditAnywhere, Category=LensFlares, meta=(editcondition = "bOverride_LensFlareTints", DisplayName = "Tints"))
+	UPROPERTY(EditAnywhere, BlueprintReadWrite, Category=LensFlares, meta=(editcondition = "bOverride_LensFlareTints", DisplayName = "Tints"))
 	FLinearColor LensFlareTints[8];
 
 	/** 0..1 0=off/no vignette .. 1=strong vignette */
-	UPROPERTY(interp, Category=SceneColor, meta=(UIMin = "0.0", UIMax = "1.0", editcondition = "bOverride_VignetteIntensity"))
+	UPROPERTY(interp, BlueprintReadWrite, Category=SceneColor, meta=(UIMin = "0.0", UIMax = "1.0", editcondition = "bOverride_VignetteIntensity"))
 	float VignetteIntensity;
 
 	/** Vignette color. */
-	UPROPERTY(interp, Category=SceneColor, AdvancedDisplay, meta=(editcondition = "bOverride_VignetteColor", DisplayName = "Vignette Color"))
+	UPROPERTY(interp, BlueprintReadWrite, Category=SceneColor, AdvancedDisplay, meta=(editcondition = "bOverride_VignetteColor", DisplayName = "Vignette Color"))
 	FLinearColor VignetteColor;
 
 	/** 0..1 grain jitter */
-	UPROPERTY(interp, Category = SceneColor, AdvancedDisplay, meta=(UIMin = "0.0", UIMax = "1.0", editcondition = "bOverride_GrainJitter"))
+	UPROPERTY(interp, BlueprintReadWrite, Category = SceneColor, AdvancedDisplay, meta=(UIMin = "0.0", UIMax = "1.0", editcondition = "bOverride_GrainJitter"))
 	float GrainJitter;
 
 	/** 0..1 grain intensity */
-	UPROPERTY(interp, Category=SceneColor, meta=(UIMin = "0.0", UIMax = "1.0", editcondition = "bOverride_GrainIntensity"))
+	UPROPERTY(interp, BlueprintReadWrite, Category=SceneColor, meta=(UIMin = "0.0", UIMax = "1.0", editcondition = "bOverride_GrainIntensity"))
 	float GrainIntensity;
 
 	/** 0..1 0=off/no ambient occlusion .. 1=strong ambient occlusion, defines how much it affects the non direct lighting after base pass */
-	UPROPERTY(interp, Category=AmbientOcclusion, meta=(ClampMin = "0.0", ClampMax = "1.0", editcondition = "bOverride_AmbientOcclusionIntensity", DisplayName = "Intensity"))
+	UPROPERTY(interp, BlueprintReadWrite, Category=AmbientOcclusion, meta=(ClampMin = "0.0", ClampMax = "1.0", editcondition = "bOverride_AmbientOcclusionIntensity", DisplayName = "Intensity"))
 	float AmbientOcclusionIntensity;
 
 	/** 0..1 0=no effect on static lighting .. 1=AO affects the stat lighting, 0 is free meaning no extra rendering pass */
-	UPROPERTY(interp, Category=AmbientOcclusion, AdvancedDisplay, meta=(ClampMin = "0.0", ClampMax = "1.0", editcondition = "bOverride_AmbientOcclusionStaticFraction", DisplayName = "Static Fraction"))
+	UPROPERTY(interp, BlueprintReadWrite, Category=AmbientOcclusion, AdvancedDisplay, meta=(ClampMin = "0.0", ClampMax = "1.0", editcondition = "bOverride_AmbientOcclusionStaticFraction", DisplayName = "Static Fraction"))
 	float AmbientOcclusionStaticFraction;
 
 	/** >0, in unreal units, bigger values means even distant surfaces affect the ambient occlusion */
-	UPROPERTY(interp, Category=AmbientOcclusion, meta=(ClampMin = "0.1", UIMax = "200.0", editcondition = "bOverride_AmbientOcclusionRadius", DisplayName = "Radius"))
+	UPROPERTY(interp, BlueprintReadWrite, Category=AmbientOcclusion, meta=(ClampMin = "0.1", UIMax = "200.0", editcondition = "bOverride_AmbientOcclusionRadius", DisplayName = "Radius"))
 	float AmbientOcclusionRadius;
 
 	/** true: AO radius is in world space units, false: AO radius is locked the view space in 400 units */
@@ -626,51 +626,51 @@ struct FPostProcessSettings
 	uint32 AmbientOcclusionRadiusInWS:1;
 
 	/** >0, in unreal units, at what distance the AO effect disppears in the distance (avoding artifacts and AO effects on huge object) */
-	UPROPERTY(interp, Category=AmbientOcclusion, AdvancedDisplay, meta=(ClampMin = "0.0", UIMax = "20000.0", editcondition = "bOverride_AmbientOcclusionFadeDistance", DisplayName = "Fade Out Distance"))
+	UPROPERTY(interp, BlueprintReadWrite, Category=AmbientOcclusion, AdvancedDisplay, meta=(ClampMin = "0.0", UIMax = "20000.0", editcondition = "bOverride_AmbientOcclusionFadeDistance", DisplayName = "Fade Out Distance"))
 	float AmbientOcclusionFadeDistance;
 	
 	/** >0, in unreal units, how many units before AmbientOcclusionFadeOutDistance it starts fading out */
-	UPROPERTY(interp, Category=AmbientOcclusion, AdvancedDisplay, meta=(ClampMin = "0.0", UIMax = "20000.0", editcondition = "bOverride_AmbientOcclusionFadeRadius", DisplayName = "Fade Out Radius"))
+	UPROPERTY(interp, BlueprintReadWrite, Category=AmbientOcclusion, AdvancedDisplay, meta=(ClampMin = "0.0", UIMax = "20000.0", editcondition = "bOverride_AmbientOcclusionFadeRadius", DisplayName = "Fade Out Radius"))
 	float AmbientOcclusionFadeRadius;
 
 	/** >0, in unreal units, how wide the ambient occlusion effect should affect the geometry (in depth) */
-	UPROPERTY(interp, Category=AmbientOcclusion, AdvancedDisplay, meta=(ClampMin = "0.0", UIMax = "2000.0", editcondition = "bOverride_AmbientOcclusionDistance", DisplayName = "Occlusion Distance"))
+	UPROPERTY(interp, BlueprintReadWrite, Category=AmbientOcclusion, AdvancedDisplay, meta=(ClampMin = "0.0", UIMax = "2000.0", editcondition = "bOverride_AmbientOcclusionDistance", DisplayName = "Occlusion Distance"))
 	float AmbientOcclusionDistance;
 
 	/** >0, in unreal units, bigger values means even distant surfaces affect the ambient occlusion */
-	UPROPERTY(interp, Category=AmbientOcclusion, AdvancedDisplay, meta=(ClampMin = "0.1", UIMax = "4.0", editcondition = "bOverride_AmbientOcclusionPower", DisplayName = "Power"))
+	UPROPERTY(interp, BlueprintReadWrite, Category=AmbientOcclusion, AdvancedDisplay, meta=(ClampMin = "0.1", UIMax = "4.0", editcondition = "bOverride_AmbientOcclusionPower", DisplayName = "Power"))
 	float AmbientOcclusionPower;
 
 	/** >0, in unreal units, default (3.0) works well for flat surfaces but can reduce details */
-	UPROPERTY(interp, Category=AmbientOcclusion, AdvancedDisplay, meta=(ClampMin = "0.0", UIMax = "10.0", editcondition = "bOverride_AmbientOcclusionBias", DisplayName = "Bias"))
+	UPROPERTY(interp, BlueprintReadWrite, Category=AmbientOcclusion, AdvancedDisplay, meta=(ClampMin = "0.0", UIMax = "10.0", editcondition = "bOverride_AmbientOcclusionBias", DisplayName = "Bias"))
 	float AmbientOcclusionBias;
 
 	/** 0=lowest quality..100=maximum quality, only a few quality levels are implemented, no soft transition */
-	UPROPERTY(interp, Category=AmbientOcclusion, AdvancedDisplay, meta=(ClampMin = "0.0", UIMax = "100.0", editcondition = "bOverride_AmbientOcclusionQuality", DisplayName = "Quality"))
+	UPROPERTY(interp, BlueprintReadWrite, Category=AmbientOcclusion, AdvancedDisplay, meta=(ClampMin = "0.0", UIMax = "100.0", editcondition = "bOverride_AmbientOcclusionQuality", DisplayName = "Quality"))
 	float AmbientOcclusionQuality;
 
 	/** Affects the blend over the multiple mips (lower resolution versions) , 0:fully use full resolution, 1::fully use low resolution, around 0.6 seems to be a good value */
-	UPROPERTY(interp, Category=AmbientOcclusion, AdvancedDisplay, meta=(ClampMin = "0.1", UIMax = "1.0", editcondition = "bOverride_AmbientOcclusionMipBlend", DisplayName = "Mip Blend"))
+	UPROPERTY(interp, BlueprintReadWrite, Category=AmbientOcclusion, AdvancedDisplay, meta=(ClampMin = "0.1", UIMax = "1.0", editcondition = "bOverride_AmbientOcclusionMipBlend", DisplayName = "Mip Blend"))
 	float AmbientOcclusionMipBlend;
 
 	/** Affects the radius AO radius scale over the multiple mips (lower resolution versions) */
-	UPROPERTY(interp, Category=AmbientOcclusion, AdvancedDisplay, meta=(ClampMin = "0.5", UIMax = "4.0", editcondition = "bOverride_AmbientOcclusionMipScale", DisplayName = "Mip Scale"))
+	UPROPERTY(interp, BlueprintReadWrite, Category=AmbientOcclusion, AdvancedDisplay, meta=(ClampMin = "0.5", UIMax = "4.0", editcondition = "bOverride_AmbientOcclusionMipScale", DisplayName = "Mip Scale"))
 	float AmbientOcclusionMipScale;
 
 	/** to tweak the bilateral upsampling when using multiple mips (lower resolution versions) */
-	UPROPERTY(interp, Category=AmbientOcclusion, AdvancedDisplay, meta=(ClampMin = "0.0", UIMax = "0.1", editcondition = "bOverride_AmbientOcclusionMipThreshold", DisplayName = "Mip Threshold"))
+	UPROPERTY(interp, BlueprintReadWrite, Category=AmbientOcclusion, AdvancedDisplay, meta=(ClampMin = "0.0", UIMax = "0.1", editcondition = "bOverride_AmbientOcclusionMipThreshold", DisplayName = "Mip Threshold"))
 	float AmbientOcclusionMipThreshold;
 
 	/** Adjusts indirect lighting color. (1,1,1) is default. (0,0,0) to disable GI. The show flag 'Global Illumination' must be enabled to use this property. */
-	UPROPERTY(interp, Category=GlobalIllumination, meta=(editcondition = "bOverride_IndirectLightingColor", DisplayName = "Indirect Lighting Color"))
+	UPROPERTY(interp, BlueprintReadWrite, Category=GlobalIllumination, meta=(editcondition = "bOverride_IndirectLightingColor", DisplayName = "Indirect Lighting Color"))
 	FLinearColor IndirectLightingColor;
 
 	/** Scales the indirect lighting contribution. A value of 0 disables GI. Default is 1. The show flag 'Global Illumination' must be enabled to use this property. */
-	UPROPERTY(interp, Category=GlobalIllumination, meta=(ClampMin = "0", UIMax = "4.0", editcondition = "bOverride_IndirectLightingIntensity", DisplayName = "Indirect Lighting Intensity"))
+	UPROPERTY(interp, BlueprintReadWrite, Category=GlobalIllumination, meta=(ClampMin = "0", UIMax = "4.0", editcondition = "bOverride_IndirectLightingIntensity", DisplayName = "Indirect Lighting Intensity"))
 	float IndirectLightingIntensity;
 
 	/** 0..1=full intensity */
-	UPROPERTY(interp, Category=SceneColor, meta=(ClampMin = "0", ClampMax = "1.0", editcondition = "bOverride_ColorGradingIntensity", DisplayName = "Color Grading Intensity"))
+	UPROPERTY(interp, BlueprintReadWrite, Category=SceneColor, meta=(ClampMin = "0", ClampMax = "1.0", editcondition = "bOverride_ColorGradingIntensity", DisplayName = "Color Grading Intensity"))
 	float ColorGradingIntensity;
 
 	/** Name of the LUT texture e.g. MyPackage01.LUTNeutral, empty if not used */
@@ -682,35 +682,35 @@ struct FPostProcessSettings
 	TEnumAsByte<enum EDepthOfFieldMethod> DepthOfFieldMethod;
 
 	/** Distance in which the Depth of Field effect should be sharp, in unreal units (cm) */
-	UPROPERTY(interp, Category=DepthOfField, meta=(UIMin = "0.0", UIMax = "10000.0", editcondition = "bOverride_DepthOfFieldFocalDistance", DisplayName = "Focal Distance"))
+	UPROPERTY(interp, BlueprintReadWrite, Category=DepthOfField, meta=(UIMin = "0.0", UIMax = "10000.0", editcondition = "bOverride_DepthOfFieldFocalDistance", DisplayName = "Focal Distance"))
 	float DepthOfFieldFocalDistance;
 
 	/** Artificial region where all content is in focus, starting after DepthOfFieldFocalDistance, in unreal units  (cm) */
-	UPROPERTY(interp, Category=DepthOfField, meta=(UIMin = "0.0", UIMax = "10000.0", editcondition = "bOverride_DepthOfFieldFocalRegion", DisplayName = "Focal Region"))
+	UPROPERTY(interp, BlueprintReadWrite, Category=DepthOfField, meta=(UIMin = "0.0", UIMax = "10000.0", editcondition = "bOverride_DepthOfFieldFocalRegion", DisplayName = "Focal Region"))
 	float DepthOfFieldFocalRegion;
 
 	/** To define the width of the transition region next to the focal region on the near side (cm) */
-	UPROPERTY(interp, Category=DepthOfField, meta=(UIMin = "0.0", UIMax = "10000.0", editcondition = "bOverride_DepthOfFieldNearTransitionRegion", DisplayName = "Near Transition Region"))
+	UPROPERTY(interp, BlueprintReadWrite, Category=DepthOfField, meta=(UIMin = "0.0", UIMax = "10000.0", editcondition = "bOverride_DepthOfFieldNearTransitionRegion", DisplayName = "Near Transition Region"))
 	float DepthOfFieldNearTransitionRegion;
 
 	/** To define the width of the transition region next to the focal region on the near side (cm) */
-	UPROPERTY(interp, Category=DepthOfField, meta=(UIMin = "0.0", UIMax = "10000.0", editcondition = "bOverride_DepthOfFieldFarTransitionRegion", DisplayName = "Far Transition Region"))
+	UPROPERTY(interp, BlueprintReadWrite, Category=DepthOfField, meta=(UIMin = "0.0", UIMax = "10000.0", editcondition = "bOverride_DepthOfFieldFarTransitionRegion", DisplayName = "Far Transition Region"))
 	float DepthOfFieldFarTransitionRegion;
 
 	/** BokehDOF only: To amplify the depth of field effect (like aperture)  0=off */
-	UPROPERTY(interp, Category=DepthOfField, meta=(ClampMin = "0.0", ClampMax = "2.0", editcondition = "bOverride_DepthOfFieldScale", DisplayName = "Scale"))
+	UPROPERTY(interp, BlueprintReadWrite, Category=DepthOfField, meta=(ClampMin = "0.0", ClampMax = "2.0", editcondition = "bOverride_DepthOfFieldScale", DisplayName = "Scale"))
 	float DepthOfFieldScale;
 
 	/** BokehDOF only: Maximum size of the Depth of Field blur (in percent of the view width) (note: performance cost scales with size*size) */
-	UPROPERTY(interp, Category=DepthOfField, meta=(UIMin = "0.0", UIMax = "32.0", editcondition = "bOverride_DepthOfFieldMaxBokehSize", DisplayName = "Max Bokeh Size"))
+	UPROPERTY(interp, BlueprintReadWrite, Category=DepthOfField, meta=(UIMin = "0.0", UIMax = "32.0", editcondition = "bOverride_DepthOfFieldMaxBokehSize", DisplayName = "Max Bokeh Size"))
 	float DepthOfFieldMaxBokehSize;
 
 	/** Gaussian only: Maximum size of the Depth of Field blur (in percent of the view width) (note: performance cost scales with size) */
-	UPROPERTY(interp, Category=DepthOfField, meta=(UIMin = "0.0", UIMax = "32.0", editcondition = "bOverride_DepthOfFieldNearBlurSize", DisplayName = "Near Blur Size"))
+	UPROPERTY(interp, BlueprintReadWrite, Category=DepthOfField, meta=(UIMin = "0.0", UIMax = "32.0", editcondition = "bOverride_DepthOfFieldNearBlurSize", DisplayName = "Near Blur Size"))
 	float DepthOfFieldNearBlurSize;
 
 	/** Gaussian only: Maximum size of the Depth of Field blur (in percent of the view width) (note: performance cost scales with size) */
-	UPROPERTY(interp, Category=DepthOfField, meta=(UIMin = "0.0", UIMax = "32.0", editcondition = "bOverride_DepthOfFieldFarBlurSize", DisplayName = "Far Blur Size"))
+	UPROPERTY(interp, BlueprintReadWrite, Category=DepthOfField, meta=(UIMin = "0.0", UIMax = "32.0", editcondition = "bOverride_DepthOfFieldFarBlurSize", DisplayName = "Far Blur Size"))
 	float DepthOfFieldFarBlurSize;
 
 	/** Defines the shape of the Bokeh when object get out of focus, cannot be blended */
@@ -718,33 +718,33 @@ struct FPostProcessSettings
 	class UTexture* DepthOfFieldBokehShape;
 
 	/** Occlusion tweak factor 1 (0.18 to get natural occlusion, 0.4 to solve layer color leaking issues) */
-	UPROPERTY(interp, Category=DepthOfField, AdvancedDisplay, meta=(ClampMin = "0.0", ClampMax = "1.0", editcondition = "bOverride_DepthOfFieldOcclusion", DisplayName = "Occlusion"))
+	UPROPERTY(interp, BlueprintReadWrite, Category=DepthOfField, AdvancedDisplay, meta=(ClampMin = "0.0", ClampMax = "1.0", editcondition = "bOverride_DepthOfFieldOcclusion", DisplayName = "Occlusion"))
 	float DepthOfFieldOcclusion;
 	
 	/** Color threshold to do full quality DOF */
-	UPROPERTY(interp, Category=DepthOfField, AdvancedDisplay, meta=(ClampMin = "0.0", ClampMax = "10.0", editcondition = "bOverride_DepthOfFieldColorThreshold", DisplayName = "Color Threshold"))
+	UPROPERTY(interp, BlueprintReadWrite, Category=DepthOfField, AdvancedDisplay, meta=(ClampMin = "0.0", ClampMax = "10.0", editcondition = "bOverride_DepthOfFieldColorThreshold", DisplayName = "Color Threshold"))
 	float DepthOfFieldColorThreshold;
 
 	/** Size threshold to do full quality DOF */
-	UPROPERTY(interp, Category=DepthOfField, AdvancedDisplay, meta=(ClampMin = "0.0", ClampMax = "1.0", editcondition = "bOverride_DepthOfFieldSizeThreshold", DisplayName = "Size Threshold"))
+	UPROPERTY(interp, BlueprintReadWrite, Category=DepthOfField, AdvancedDisplay, meta=(ClampMin = "0.0", ClampMax = "1.0", editcondition = "bOverride_DepthOfFieldSizeThreshold", DisplayName = "Size Threshold"))
 	float DepthOfFieldSizeThreshold;
 	
 	/** Artificial distance to allow the skybox to be in focus (e.g. 200000), <=0 to switch the feature off, only for GaussianDOF, can cost performance */
-	UPROPERTY(interp, Category=DepthOfField, AdvancedDisplay, meta=(ClampMin = "0.0", ClampMax = "200000.0", editcondition = "bOverride_DepthOfFieldSkyFocusDistance", DisplayName = "Sky Distance"))
+	UPROPERTY(interp, BlueprintReadWrite, Category=DepthOfField, AdvancedDisplay, meta=(ClampMin = "0.0", ClampMax = "200000.0", editcondition = "bOverride_DepthOfFieldSkyFocusDistance", DisplayName = "Sky Distance"))
 	float DepthOfFieldSkyFocusDistance;
 	
 	/** Strength of motion blur, 0:off, should be renamed to intensity */
-	UPROPERTY(interp, Category=MotionBlur, meta=(ClampMin = "0.0", ClampMax = "1.0", editcondition = "bOverride_MotionBlurAmount", DisplayName = "Amount"))
+	UPROPERTY(interp, BlueprintReadWrite, Category=MotionBlur, meta=(ClampMin = "0.0", ClampMax = "1.0", editcondition = "bOverride_MotionBlurAmount", DisplayName = "Amount"))
 	float MotionBlurAmount;
 	/** max distortion caused by motion blur, in percent of the screen width, 0:off */
-	UPROPERTY(interp, Category=MotionBlur, meta=(ClampMin = "0.0", ClampMax = "100.0", editcondition = "bOverride_MotionBlurMax", DisplayName = "Max"))
+	UPROPERTY(interp, BlueprintReadWrite, Category=MotionBlur, meta=(ClampMin = "0.0", ClampMax = "100.0", editcondition = "bOverride_MotionBlurMax", DisplayName = "Max"))
 	float MotionBlurMax;
 	/** The minimum projected screen radius for a primitive to be drawn in the velocity pass, percentage of screen width. smaller numbers cause more draw calls, default: 4% */
-	UPROPERTY(interp, Category=MotionBlur, AdvancedDisplay, meta=(ClampMin = "0.0", UIMax = "100.0", editcondition = "bOverride_MotionBlurPerObjectSize", DisplayName = "Per Object Size"))
+	UPROPERTY(interp, BlueprintReadWrite, Category=MotionBlur, AdvancedDisplay, meta=(ClampMin = "0.0", UIMax = "100.0", editcondition = "bOverride_MotionBlurPerObjectSize", DisplayName = "Per Object Size"))
 	float MotionBlurPerObjectSize;
 
 	/** to render with lower resolution and upscale, controlled by console variable, 100:off, needs to be <99 to see effect, only applied in game  */
-	UPROPERTY(interp, Category=Misc, AdvancedDisplay, meta=(ClampMin = "0.0", ClampMax = "400.0", editcondition = "bOverride_ScreenPercentage"))
+	UPROPERTY(interp, BlueprintReadWrite, Category=Misc, AdvancedDisplay, meta=(ClampMin = "0.0", ClampMax = "400.0", editcondition = "bOverride_ScreenPercentage"))
 	float ScreenPercentage;
 
 	/** TemporalAA, FXAA, ... */
@@ -752,15 +752,15 @@ struct FPostProcessSettings
 	TEnumAsByte<enum EAntiAliasingMethod> AntiAliasingMethod;
 
 	/** Enable/Fade/disable the Screen Space Reflection feature, in percent, avoid numbers between 0 and 1 fo consistency */
-	UPROPERTY(interp, Category=ScreenSpaceReflections, meta=(ClampMin = "0.0", ClampMax = "100.0", editcondition = "bOverride_ScreenSpaceReflectionIntensity", DisplayName = "Intensity"))
+	UPROPERTY(interp, BlueprintReadWrite, Category=ScreenSpaceReflections, meta=(ClampMin = "0.0", ClampMax = "100.0", editcondition = "bOverride_ScreenSpaceReflectionIntensity", DisplayName = "Intensity"))
 	float ScreenSpaceReflectionIntensity;
 
 	/** 0=lowest quality..100=maximum quality, only a few quality levels are implemented, no soft transition, 50 is the default for better performance. */
-	UPROPERTY(interp, Category=ScreenSpaceReflections, meta=(ClampMin = "0.0", UIMax = "100.0", editcondition = "bOverride_ScreenSpaceReflectionQuality", DisplayName = "Quality"))
+	UPROPERTY(interp, BlueprintReadWrite, Category=ScreenSpaceReflections, meta=(ClampMin = "0.0", UIMax = "100.0", editcondition = "bOverride_ScreenSpaceReflectionQuality", DisplayName = "Quality"))
 	float ScreenSpaceReflectionQuality;
 
 	/** Until what roughness we fade the screen space reflections, 0.8 works well, smaller can run faster */
-	UPROPERTY(interp, Category=ScreenSpaceReflections, meta=(ClampMin = "0.01", ClampMax = "1.0", editcondition = "bOverride_ScreenSpaceReflectionMaxRoughness", DisplayName = "Max Roughness"))
+	UPROPERTY(interp, BlueprintReadWrite, Category=ScreenSpaceReflections, meta=(ClampMin = "0.01", ClampMax = "1.0", editcondition = "bOverride_ScreenSpaceReflectionMaxRoughness", DisplayName = "Max Roughness"))
 	float ScreenSpaceReflectionMaxRoughness;
 
 	// Note: Adding properties before this line require also changes to the OverridePostProcessSettings() function and 
@@ -768,7 +768,7 @@ struct FPostProcessSettings
 	// -----------------------------------------------------------------------
 	
 	/** Allows custom post process materials to be defined,		using a MaterialInstance with the same Material as its parent to allow blending. Make sure you use the "PostProcess" domain type */
-	UPROPERTY(EditAnywhere, Category=Misc)
+	UPROPERTY(EditAnywhere, BlueprintReadWrite, Category=Misc)
 	TArray<UObject*> Blendables;
 
 	// good start values for a new volume, by default no value is overriding
