exposing dbuffer as project setting

diff --git a/Engine/Source/Runtime/Engine/Classes/Engine/RendererSettings.h b/Engine/Source/Runtime/Engine/Classes/Engine/RendererSettings.h
--- a/Engine/Source/Runtime/Engine/Classes/Engine/RendererSettings.h
+++ b/Engine/Source/Runtime/Engine/Classes/Engine/RendererSettings.h
@@ -223,14 +223,19 @@ class ENGINE_API URendererSettings
 
 	UPROPERTY(config, EditAnywhere, Category = Optimizations, meta = (
 		ConsoleVariable="r.EarlyZPass",DisplayName="Early Z-pass",
-		ToolTip="Whether to use a depth only pass to initialize Z culling for the base pass."))
+		ToolTip="Whether to use a depth only pass to initialize Z culling for the base pass. Need to reload the level!"))
 	TEnumAsByte<EEarlyZPass::Type> EarlyZPass;
 
 	UPROPERTY(config, EditAnywhere, Category=Optimizations, meta=(
 		ConsoleVariable="r.EarlyZPassMovable",DisplayName="Movables in early Z-pass",
-		ToolTip="Whether to render movable objects in the early Z pass."))
+		ToolTip="Whether to render movable objects in the early Z pass. Need to reload the level!"))
 	uint32 bEarlyZPassMovable:1;
 
+	UPROPERTY(config, EditAnywhere, Category=Lighting, meta=(
+		ConsoleVariable="r.DBuffer",DisplayName="DBuffer Decals",
+		ToolTip="Experimental decal feature (see r.DBuffer, ideally combined with 'Movables in early Z-pass' and 'Early Z-pass')"))
+	uint32 bDBuffer:1;
+
 	UPROPERTY(config, EditAnywhere, Category=Optimizations, meta=(
 		ConsoleVariable="r.ClearSceneMethod",DisplayName="Clear Scene",
 		ToolTip="Control how the scene is cleared before rendering"))
