Extended visual logs to log corectly from Objects, not only from Actors. That Object has to work corectly with GetWorldFromContextObject function. #ue4

diff --git a/Engine/Source/Runtime/AIModule/Classes/EnvironmentQuery/EnvQueryDebugHelpers.h b/Engine/Source/Runtime/AIModule/Classes/EnvironmentQuery/EnvQueryDebugHelpers.h
--- a/Engine/Source/Runtime/AIModule/Classes/EnvironmentQuery/EnvQueryDebugHelpers.h
+++ b/Engine/Source/Runtime/AIModule/Classes/EnvironmentQuery/EnvQueryDebugHelpers.h
@@ -146,12 +146,8 @@ FArchive& operator<<(FArchive& Ar, EQSDebug::FQueryData& Data)
 	if (FVisualLog::Get().IsRecording() && (!FVisualLog::Get().IsAllBlocked() || FVisualLog::Get().InWhitelist(CategoryName.GetCategoryName()))) \
 	{ \
 		const AActor* OwnerActor = FVisualLog::Get().GetVisualLogRedirection(Query->Owner.Get()); \
-		ensure(OwnerActor != NULL); \
-		if (OwnerActor) \
-		{  \
-			UEnvQueryDebugHelpers::LogQuery(OwnerActor, Query, CategoryName.GetCategoryName(), ELogVerbosity::Verbosity, \
-				FString::Printf(TEXT("Executed EQS: \n - Name: '%s' (id=%d, option=%d),\n - All Items: %d,\n - ValidItems: %d"), *Query->QueryName, Query->QueryID, Query->OptionIndex, Query->ItemDetails.Num(), Query->NumValidItems)); \
-		} \
+			UEnvQueryDebugHelpers::LogQuery(OwnerActor ? OwnerActor : Query->Owner.Get(), Query, CategoryName.GetCategoryName(), ELogVerbosity::Verbosity, \
+			FString::Printf(TEXT("Executed EQS: \n - Name: '%s' (id=%d, option=%d),\n - All Items: %d,\n - ValidItems: %d"), *Query->QueryName, Query->QueryID, Query->OptionIndex, Query->ItemDetails.Num(), Query->NumValidItems)); \
 	} \
 }
 
@@ -170,6 +166,6 @@ class AIMODULE_API UEnvQueryDebugHelpers : public UObject
 	static void QueryToDebugData(struct FEnvQueryInstance* Query, EQSDebug::FQueryData& EQSLocalData);
 	static void QueryToBlobArray(struct FEnvQueryInstance* Query, TArray<uint8>& BlobArray, bool bUseCompression = false);
 	static void BlobArrayToDebugData(const TArray<uint8>& BlobArray, EQSDebug::FQueryData& EQSLocalData, bool bUseCompression = false);
-	static void LogQuery(const class AActor* LogOwnerActor, struct FEnvQueryInstance* Query, const FName& CategoryName, ELogVerbosity::Type, const FString& AdditionalLogInfo);
+	static void LogQuery(const class UObject* QueryOwner, struct FEnvQueryInstance* Query, const FName& CategoryName, ELogVerbosity::Type, const FString& AdditionalLogInfo);
 #endif
 };
diff --git a/Engine/Source/Runtime/AIModule/Private/EnvironmentQuery/EnvQueryDebugHelpers.cpp b/Engine/Source/Runtime/AIModule/Private/EnvironmentQuery/EnvQueryDebugHelpers.cpp
--- a/Engine/Source/Runtime/AIModule/Private/EnvironmentQuery/EnvQueryDebugHelpers.cpp
+++ b/Engine/Source/Runtime/AIModule/Private/EnvironmentQuery/EnvQueryDebugHelpers.cpp
@@ -129,14 +129,14 @@ void  UEnvQueryDebugHelpers::BlobArrayToDebugData(const TArray<uint8>& BlobArray
 	}
 }
 
-void UEnvQueryDebugHelpers::LogQuery(const class AActor* LogOwnerActor, struct FEnvQueryInstance* Query, const FName& CategoryName, ELogVerbosity::Type Type, const FString& AdditionalLogInfo)
+void UEnvQueryDebugHelpers::LogQuery(const class UObject* QueryOwner, struct FEnvQueryInstance* Query, const FName& CategoryName, ELogVerbosity::Type Type, const FString& AdditionalLogInfo)
 {
 	TArray<uint8> BlobArray;
 	UEnvQueryDebugHelpers::QueryToBlobArray(Query, BlobArray);
-	FVisLogEntry* EntryToWrite = FVisualLog::Get().GetEntryToWrite(LogOwnerActor);
+	FVisLogEntry* EntryToWrite = FVisualLog::Get().GetEntryToWrite(QueryOwner);
 	if (EntryToWrite)
 	{
-		FVisualLog::Get().LogLine(LogOwnerActor, CategoryName, Type, AdditionalLogInfo, Query->QueryID, *EVisLogTags::TAG_EQS);
+		FVisualLog::Get().LogLine(QueryOwner, CategoryName, Type, AdditionalLogInfo, Query->QueryID, *EVisLogTags::TAG_EQS);
 		EntryToWrite->AddDataBlock(EVisLogTags::TAG_EQS, BlobArray, CategoryName);
 	}
 }
diff --git a/Engine/Source/Runtime/Engine/Private/VisualLog.cpp b/Engine/Source/Runtime/Engine/Private/VisualLog.cpp
--- a/Engine/Source/Runtime/Engine/Private/VisualLog.cpp
+++ b/Engine/Source/Runtime/Engine/Private/VisualLog.cpp
@@ -26,7 +26,7 @@ FVisLogEntry::FVisLogEntry(const class AActor* InActor, TArray<TWeakObjectPtr<UO
 			TWeakObjectPtr<UObject>* WeakActorPtr = Children->GetTypedData();
 			for (int32 Index = 0; Index < Children->Num(); ++Index, ++WeakActorPtr)
 			{
-				if (WeakActorPtr->IsValid() )
+				if (WeakActorPtr->IsValid())
 				{
 					const AActor* ChildActor = Cast<AActor>(WeakActorPtr->Get());
 					if (ChildActor)
@@ -39,6 +39,27 @@ FVisLogEntry::FVisLogEntry(const class AActor* InActor, TArray<TWeakObjectPtr<UO
 	}
 }
 
+FVisLogEntry::FVisLogEntry(float InTimeStamp, FVector InLocation, const UObject* Object, TArray<TWeakObjectPtr<UObject> >* Children)
+{
+	TimeStamp = InTimeStamp;
+	Location = InLocation;
+	if (Children != NULL)
+	{
+		TWeakObjectPtr<UObject>* WeakActorPtr = Children->GetTypedData();
+		for (int32 Index = 0; Index < Children->Num(); ++Index, ++WeakActorPtr)
+		{
+			if (WeakActorPtr->IsValid())
+			{
+				const AActor* ChildActor = Cast<AActor>(WeakActorPtr->Get());
+				if (ChildActor)
+				{
+					ChildActor->GrabDebugSnapshot(this);
+				}
+			}
+		}
+	}
+}
+
 FVisLogEntry::FVisLogEntry(TSharedPtr<FJsonValue> FromJson)
 {
 	TSharedPtr<FJsonObject> JsonEntryObject = FromJson->AsObject();
@@ -374,12 +395,24 @@ void FVisLogEntry::AddDataBlock(const FString& TagName, const TArray<uint8>& Blo
 //----------------------------------------------------------------------//
 // FActorsVisLog
 //----------------------------------------------------------------------//
-FActorsVisLog::FActorsVisLog(const class AActor* Actor, TArray<TWeakObjectPtr<UObject> >* Children)
-	: Name(Actor->GetFName())
-	, FullName(Actor->GetFullName())
+FActorsVisLog::FActorsVisLog(const class UObject* Object, TArray<TWeakObjectPtr<UObject> >* Children)
+	: Name(Object->GetFName())
+	, FullName(Object->GetFullName())
 {
+	const class AActor* AsActor = FVisualLog::Get().GetVisualLogRedirection(Object);
+
 	Entries.Reserve(VisLogInitialSize);
-	Entries.Add(MakeShareable(new FVisLogEntry(Actor, Children)));
+	if (!AsActor)
+	{
+		UWorld* World = GEngine->GetWorldFromContextObject(Object);
+		check(World);
+
+		Entries.Add(MakeShareable(new FVisLogEntry(World->GetTimeSeconds(), FVector(), Object, Children)));
+	}
+	else
+	{
+		Entries.Add(MakeShareable(new FVisLogEntry(AsActor, Children)));
+	}
 }
 
 FActorsVisLog::FActorsVisLog(TSharedPtr<FJsonValue> FromJson)
@@ -553,24 +586,29 @@ void FVisualLog::SetIsRecording(bool NewRecording, bool bRecordToFile)
 	}
 }
 
-FVisLogEntry*  FVisualLog::GetEntryToWrite(const class AActor* Actor)
+FVisLogEntry*  FVisualLog::GetEntryToWrite(const class UObject* Object)
 {
-	const class AActor* RedirectionActor = GetVisualLogRedirection(Actor);
-	const class AActor* LogOwner = RedirectionActor ? RedirectionActor : Actor;
-	ensure(Actor && Actor->GetWorld() && LogOwner);
-	if (!LogOwner)
-	{
-		return NULL;
-	}
-	const float TimeStamp = Actor->GetWorld()->TimeSeconds;
-	TSharedPtr<FActorsVisLog> Log = GetLog(LogOwner);
+	const class AActor* RedirectionActor = GetVisualLogRedirection(Object);
+
+	UWorld* World = RedirectionActor ? RedirectionActor->GetWorld() : GEngine->GetWorldFromContextObject(Object);
+	ensure(World);
+
+	const float TimeStamp = World->TimeSeconds;
+	TSharedPtr<FActorsVisLog> Log = GetLog(RedirectionActor ? RedirectionActor : Object);
 	const int32 LastIndex = Log->Entries.Num() - 1;
 	FVisLogEntry* Entry = Log->Entries.Num() > 0 ? Log->Entries[LastIndex].Get() : NULL;
 
 	if (Entry == NULL || Entry->TimeStamp < TimeStamp)
 	{
 		// create new entry
-		Entry = Log->Entries[Log->Entries.Add( MakeShareable(new FVisLogEntry(LogOwner, RedirectsMap.Find(LogOwner))) )].Get();
+		if (RedirectionActor)
+		{
+			Entry = Log->Entries[Log->Entries.Add(MakeShareable(new FVisLogEntry(RedirectionActor, RedirectsMap.Find(RedirectionActor))))].Get();
+		}
+		else
+		{
+			Entry = Log->Entries[Log->Entries.Add(MakeShareable(new FVisLogEntry(TimeStamp, FVector(), Object, RedirectsMap.Find(RedirectionActor))))].Get();
+		}
 	}
 
 	check(Entry);
@@ -611,7 +649,7 @@ const class AActor* FVisualLog::GetVisualLogRedirection(const class UObject* Sou
 
 void FVisualLog::RedirectToVisualLog(const class UObject* Src, const class AActor* Dest)
 {
-	Redirect( const_cast<class UObject*>(Src), Dest);
+	Redirect(const_cast<class UObject*>(Src), Dest);
 }
 
 void FVisualLog::Redirect(UObject* Source, const AActor* NewRedirection)
@@ -641,8 +679,6 @@ void FVisualLog::Redirect(UObject* Source, const AActor* NewRedirection)
 		return;
 	}
 
-	UE_VLOG(Source, LogVisual, Display, TEXT("Binding %s to log %s"), *Source->GetName(), *NewRedirection->GetName());
-
 	TArray<TWeakObjectPtr<UObject> >& NewTargetChildren = RedirectsMap.FindOrAdd(NewRedirection);
 
 	NewTargetChildren.AddUnique(Source);
@@ -670,14 +706,14 @@ void FVisualLog::Redirect(UObject* Source, const AActor* NewRedirection)
 	UE_CVLOG(NewRedirection != NULL, NewRedirection, LogVisual, Display, TEXT("Binding %s to log %s"), *Source->GetName(), *NewRedirection->GetName());
 }
 
-void FVisualLog::LogLine(const AActor* Actor, const FName& CategoryName, ELogVerbosity::Type Verbosity, const FString& Line, int64 UserData, FName TagName)
+void FVisualLog::LogLine(const UObject* Object, const FName& CategoryName, ELogVerbosity::Type Verbosity, const FString& Line, int64 UserData, FName TagName)
 {
-	if (IsRecording() == false || Actor == NULL || Actor->IsPendingKill() || (IsAllBlocked() && !InWhitelist(CategoryName)))
+	if (IsRecording() == false || Object == NULL || Object->IsPendingKill() || (IsAllBlocked() && !InWhitelist(CategoryName)))
 	{
 		return;
 	}
 	
-	FVisLogEntry* Entry = GetEntryToWrite(Actor);
+	FVisLogEntry* Entry = GetEntryToWrite(Object);
 
 	if (Entry != NULL)
 	{
diff --git a/Engine/Source/Runtime/Engine/Public/VisualLog.h b/Engine/Source/Runtime/Engine/Public/VisualLog.h
--- a/Engine/Source/Runtime/Engine/Public/VisualLog.h
+++ b/Engine/Source/Runtime/Engine/Public/VisualLog.h
@@ -167,6 +167,7 @@ struct ENGINE_API FVisLogEntry
 
 
 	FVisLogEntry(const class AActor* Actor, TArray<TWeakObjectPtr<UObject> >* Children);
+	FVisLogEntry(float TimeStamp, FVector Location, const UObject* Object, TArray<TWeakObjectPtr<UObject> >* Children);
 	FVisLogEntry(TSharedPtr<FJsonValue> FromJson);
 	
 	TSharedPtr<FJsonValue> ToJson() const;
@@ -213,18 +214,16 @@ struct ENGINE_API FVisLogEntry
 #define REDIRECT_TO_VLOG(Dest) FVisualLog::Get().RedirectToVisualLog(this, Dest)
 #define REDIRECT_OBJECT_TO_VLOG(Src, Dest) FVisualLog::Get().RedirectToVisualLog(Src, Dest)
 
+#define CONNECT_WITH_VLOG(Dest)
+#define CONNECT_OBJECT_WITH_VLOG(Src, Dest)
+
 #define UE_VLOG(Object, CategoryName, Verbosity, Format, ...) \
 { \
 	SCOPE_CYCLE_COUNTER(STAT_VisualLog); \
 	static_assert((ELogVerbosity::Verbosity & ELogVerbosity::VerbosityMask) < ELogVerbosity::NumVerbosity && ELogVerbosity::Verbosity > 0, "Verbosity must be constant and in range."); \
 	if (FVisualLog::Get().IsRecording() && (!FVisualLog::Get().IsAllBlocked() || FVisualLog::Get().InWhitelist(CategoryName.GetCategoryName()))) \
 	{ \
-		const AActor* MyOwnerActor = FVisualLog::Get().GetVisualLogRedirection(Object); \
-		ensure(MyOwnerActor != NULL); \
-		if (MyOwnerActor) \
-		{  \
-			FVisualLog::Get().LogLine(MyOwnerActor, CategoryName.GetCategoryName(), ELogVerbosity::Verbosity, FString::Printf(Format, ##__VA_ARGS__)); \
-		} \
+		FVisualLog::Get().LogLine(Object, CategoryName.GetCategoryName(), ELogVerbosity::Verbosity, FString::Printf(Format, ##__VA_ARGS__)); \
 	} \
 	if (UE_LOG_CHECK_COMPILEDIN_VERBOSITY(CategoryName, Verbosity)) \
 	{ \
@@ -246,12 +245,7 @@ struct ENGINE_API FVisLogEntry
 	static_assert((ELogVerbosity::Verbosity & ELogVerbosity::VerbosityMask) < ELogVerbosity::NumVerbosity && ELogVerbosity::Verbosity > 0, "Verbosity must be constant and in range."); \
 	if (FVisualLog::Get().IsRecording() && (!FVisualLog::Get().IsAllBlocked() || FVisualLog::Get().InWhitelist(CategoryName.GetCategoryName()))) \
 	{ \
-		const AActor* OwnerActor = FVisualLog::Get().GetVisualLogRedirection(Object); \
-		ensure(OwnerActor != NULL); \
-		if (OwnerActor) \
-		{  \
-			FVisualLog::Get().GetEntryToWrite(OwnerActor)->AddElement(SegmentStart, SegmentEnd, CategoryName.GetCategoryName(), Color, FString::Printf(DescriptionFormat, ##__VA_ARGS__), Thickness); \
-		} \
+		FVisualLog::Get().GetEntryToWrite(Object)->AddElement(SegmentStart, SegmentEnd, CategoryName.GetCategoryName(), Color, FString::Printf(DescriptionFormat, ##__VA_ARGS__), Thickness); \
 	} \
 }
 
@@ -264,12 +258,7 @@ struct ENGINE_API FVisLogEntry
 	static_assert((ELogVerbosity::Verbosity & ELogVerbosity::VerbosityMask) < ELogVerbosity::NumVerbosity && ELogVerbosity::Verbosity > 0, "Verbosity must be constant and in range."); \
 	if (FVisualLog::Get().IsRecording() && (!FVisualLog::Get().IsAllBlocked() || FVisualLog::Get().InWhitelist(CategoryName.GetCategoryName()))) \
 	{ \
-		const AActor* OwnerActor = FVisualLog::Get().GetVisualLogRedirection(Object); \
-		ensure(OwnerActor != NULL); \
-		if (OwnerActor) \
-		{  \
-			FVisualLog::Get().GetEntryToWrite(OwnerActor)->AddElement(Location, CategoryName.GetCategoryName(), Color, FString::Printf(DescriptionFormat, ##__VA_ARGS__), Radius); \
-		} \
+		FVisualLog::Get().GetEntryToWrite(Object)->AddElement(Location, CategoryName.GetCategoryName(), Color, FString::Printf(DescriptionFormat, ##__VA_ARGS__), Radius); \
 	} \
 }
 
@@ -279,12 +268,7 @@ struct ENGINE_API FVisLogEntry
 	static_assert((ELogVerbosity::Verbosity & ELogVerbosity::VerbosityMask) < ELogVerbosity::NumVerbosity && ELogVerbosity::Verbosity > 0, "Verbosity must be constant and in range."); \
 	if (FVisualLog::Get().IsRecording() && (!FVisualLog::Get().IsAllBlocked() || FVisualLog::Get().InWhitelist(CategoryName.GetCategoryName()))) \
 	{ \
-		const AActor* OwnerActor = FVisualLog::Get().GetVisualLogRedirection(Object); \
-		ensure(OwnerActor != NULL); \
-		if (OwnerActor) \
-		{  \
-			FVisualLog::Get().GetEntryToWrite(OwnerActor)->AddElement(Box, CategoryName.GetCategoryName(), Color, FString::Printf(DescriptionFormat, ##__VA_ARGS__)); \
-		} \
+		FVisualLog::Get().GetEntryToWrite(Object)->AddElement(Box, CategoryName.GetCategoryName(), Color, FString::Printf(DescriptionFormat, ##__VA_ARGS__)); \
 	} \
 }
 
@@ -294,12 +278,7 @@ struct ENGINE_API FVisLogEntry
 	static_assert((ELogVerbosity::Verbosity & ELogVerbosity::VerbosityMask) < ELogVerbosity::NumVerbosity && ELogVerbosity::Verbosity > 0, "Verbosity must be constant and in range."); \
 	if (FVisualLog::Get().IsRecording() && (!FVisualLog::Get().IsAllBlocked() || FVisualLog::Get().InWhitelist(CategoryName.GetCategoryName()))) \
 	{ \
-		const AActor* OwnerActor = FVisualLog::Get().GetVisualLogRedirection(Object); \
-		ensure(OwnerActor != NULL); \
-		if (OwnerActor) \
-		{  \
-			FVisualLog::Get().GetEntryToWrite(OwnerActor)->AddHistogramData(Data, CategoryName.GetCategoryName(), GraphName, DataName); \
-		} \
+		FVisualLog::Get().GetEntryToWrite(Object)->AddHistogramData(Data, CategoryName.GetCategoryName(), GraphName, DataName); \
 	} \
 }
 
@@ -342,7 +321,7 @@ struct ENGINE_API FActorsVisLog
 	FString FullName;
 	TArray<TSharedPtr<FVisLogEntry> > Entries;
 
-	FActorsVisLog(const class AActor* Actor, TArray<TWeakObjectPtr<UObject> >* Children);
+	FActorsVisLog(const class UObject* Object, TArray<TWeakObjectPtr<UObject> >* Children);
 	FActorsVisLog(TSharedPtr<FJsonValue> FromJson);
 
 	TSharedPtr<FJsonValue> ToJson() const;
@@ -363,7 +342,7 @@ struct FVisualLogExtensionInterface
 class ENGINE_API FVisualLog : public FOutputDevice
 {
 public:
-	typedef TMap<const class AActor*, TSharedPtr<FActorsVisLog> > FLogsMap;
+	typedef TMap<const class UObject*, TSharedPtr<FActorsVisLog> > FLogsMap;
 	typedef TMap<const class AActor*, TArray<TWeakObjectPtr<UObject> > > FLogRedirectsMap;
 
 	DECLARE_DELEGATE_TwoParams(FOnNewLogCreatedDelegate, const AActor*, TSharedPtr<FActorsVisLog>);
@@ -380,12 +359,10 @@ public:
 	void Cleanup(bool bReleaseMemory = false);
 
 	void Redirect(class UObject* Source, const class AActor* NewRedirection);
-	
-	const class AActor* GetVisualLogRedirection(const class UObject* Source);
-
 	void RedirectToVisualLog(const class UObject* Src, const class AActor* Dest);
+	const class AActor* GetVisualLogRedirection(const class UObject* Source);
 
-	void LogLine(const class AActor* Actor, const FName& CategoryName, ELogVerbosity::Type Verbosity, const FString& Line, int64 UserData = 0, FName TagName = NAME_Name);
+	void LogLine(const class UObject* Object, const FName& CategoryName, ELogVerbosity::Type Verbosity, const FString& Line, int64 UserData = 0, FName TagName = NAME_Name);
 
 	const FLogsMap* GetLogs() const { return &LogsMap; }
 
@@ -406,7 +383,7 @@ public:
 
 	FArchive* FileAr;
 
-	FVisLogEntry* GetEntryToWrite(const class AActor* Actor);
+	FVisLogEntry* GetEntryToWrite(const class UObject* Object);
 
 	/** highly encouradged to use FVisualLogFilenameGetterDelegate::CreateUObject with this */
 	void SetLogFileNameGetter(const FVisualLogFilenameGetterDelegate& InLogFileNameGetter) { LogFileNameGetter = InLogFileNameGetter; }
@@ -424,12 +401,13 @@ protected:
 
 	void LogElementImpl(FVisLogSelfDrawingElement* Element);
 
-	FORCEINLINE_DEBUGGABLE TSharedPtr<FActorsVisLog> GetLog(const class AActor* Actor)
+	FORCEINLINE_DEBUGGABLE TSharedPtr<FActorsVisLog> GetLog(const class UObject* Object)
 	{
-		TSharedPtr<FActorsVisLog>* Log = LogsMap.Find(Actor);
+		TSharedPtr<FActorsVisLog>* Log = LogsMap.Find(Object);
 		if (Log == NULL)
 		{
-			Log = &(LogsMap.Add(Actor, MakeShareable(new FActorsVisLog(Actor, RedirectsMap.Find(Actor)))));
+			const class AActor* Actor = GetVisualLogRedirection(Object);
+			Log = &(LogsMap.Add(Object, MakeShareable(new FActorsVisLog(Object, RedirectsMap.Find(Actor)))));
 			OnNewLogCreated.ExecuteIfBound(Actor, *Log);
 		}
 		return *Log;
@@ -465,4 +443,6 @@ private:
 	#define UE_VLOG_HISTOGRAM(Actor, CategoryName, Verbosity, GraphName, DataName, Data)
 	#define REDIRECT_TO_VLOG(Dest)
 	#define REDIRECT_OBJECT_TO_VLOG(Src, Destination) 
+	#define CONNECT_WITH_VLOG(Dest)
+	#define CONNECT_OBJECT_WITH_VLOG(Src, Dest)
 #endif //ENABLE_VISUAL_LOG
