Refactor GetMaxTickRate so that it's const.
Updating RunningAverage is now done only during initial update of time.

TTP#345354

diff --git a/Engine/Source/Editor/UnrealEd/Classes/Editor/EditorEngine.h b/Engine/Source/Editor/UnrealEd/Classes/Editor/EditorEngine.h
--- a/Engine/Source/Editor/UnrealEd/Classes/Editor/EditorEngine.h
+++ b/Engine/Source/Editor/UnrealEd/Classes/Editor/EditorEngine.h
@@ -651,7 +651,7 @@ public:
 	
 
 	/** Get tick rate limitor. */
-	virtual float GetMaxTickRate( float DeltaTime, bool bAllowFrameRateSmoothing = true ) override;
+	virtual float GetMaxTickRate( float DeltaTime, bool bAllowFrameRateSmoothing = true ) const override;
 
 	/**
 	 * Initializes the Editor.
diff --git a/Engine/Source/Editor/UnrealEd/Private/Editor.cpp b/Engine/Source/Editor/UnrealEd/Private/Editor.cpp
--- a/Engine/Source/Editor/UnrealEd/Private/Editor.cpp
+++ b/Engine/Source/Editor/UnrealEd/Private/Editor.cpp
@@ -1473,7 +1473,7 @@ void UEditorEngine::Tick( float DeltaSeconds, bool bIdleMode )
 	}
 }
 
-float UEditorEngine::GetMaxTickRate( float DeltaTime, bool bAllowFrameRateSmoothing )
+float UEditorEngine::GetMaxTickRate( float DeltaTime, bool bAllowFrameRateSmoothing ) const
 {
 	float MaxTickRate = 0.0f;
 	if( !ShouldThrottleCPUUsage() )
diff --git a/Engine/Source/Runtime/Engine/Classes/Engine/Engine.h b/Engine/Source/Runtime/Engine/Classes/Engine/Engine.h
--- a/Engine/Source/Runtime/Engine/Classes/Engine/Engine.h
+++ b/Engine/Source/Runtime/Engine/Classes/Engine/Engine.h
@@ -1711,7 +1711,13 @@ public:
 	void TickDeferredCommands();
 
 	/** Get tick rate limiter. */
-	virtual float GetMaxTickRate( float /*DeltaTime*/, bool bAllowFrameRateSmoothing = true );
+	virtual float GetMaxTickRate(float DeltaTime, bool bAllowFrameRateSmoothing = true) const;
+
+	/** Updates the running average delta time */
+	virtual void UpdateRunningAverageDeltaTime(float DeltaTime, bool bAllowFrameRateSmoothing = true);
+
+	/** Whether we're allowed to do frame rate smoothing */
+	bool IsAllowedFramerateSmoothing(bool bAllowFrameRateSmoothing) const;
 
 	/**
 	 * Pauses / un-pauses the game-play when focus of the game's window gets lost / gained.
@@ -2155,6 +2161,13 @@ protected:
 	 */
 	virtual void InitializeObjectReferences();
 
+	/** 
+	 * Initializes the running average delta to some good initial framerate 
+	 */
+	virtual void InitializeRunningAverageDeltaTime();
+
+	float RunningAverageDeltaTime;
+
 	/** Broadcasts when a world is added. */
 	FWorldAddedEvent			WorldAddedEvent;
 
diff --git a/Engine/Source/Runtime/Engine/Classes/Engine/GameEngine.h b/Engine/Source/Runtime/Engine/Classes/Engine/GameEngine.h
--- a/Engine/Source/Runtime/Engine/Classes/Engine/GameEngine.h
+++ b/Engine/Source/Runtime/Engine/Classes/Engine/GameEngine.h
@@ -90,7 +90,7 @@ public:
 	virtual void Init(class IEngineLoop* InEngineLoop) override;
 	virtual void PreExit() override;
 	virtual void Tick( float DeltaSeconds, bool bIdleMode ) override;
-	virtual float GetMaxTickRate( float DeltaTime, bool bAllowFrameRateSmoothing = true ) override;
+	virtual float GetMaxTickRate( float DeltaTime, bool bAllowFrameRateSmoothing = true ) const override;
 	virtual void ProcessToggleFreezeCommand( UWorld* InWorld ) override;
 	virtual void ProcessToggleFreezeStreamingCommand( UWorld* InWorld ) override;
 	// End UEngine interface.
diff --git a/Engine/Source/Runtime/Engine/Private/GameEngine.cpp b/Engine/Source/Runtime/Engine/Private/GameEngine.cpp
--- a/Engine/Source/Runtime/Engine/Private/GameEngine.cpp
+++ b/Engine/Source/Runtime/Engine/Private/GameEngine.cpp
@@ -696,7 +696,7 @@ void UGameEngine::PostLoadMap()
 {
 }
 
-float UGameEngine::GetMaxTickRate( float DeltaTime, bool bAllowFrameRateSmoothing )
+float UGameEngine::GetMaxTickRate(float DeltaTime, bool bAllowFrameRateSmoothing) const
 {
 	float MaxTickRate = 0.f;
 
diff --git a/Engine/Source/Runtime/Engine/Private/UnrealEngine.cpp b/Engine/Source/Runtime/Engine/Private/UnrealEngine.cpp
--- a/Engine/Source/Runtime/Engine/Private/UnrealEngine.cpp
+++ b/Engine/Source/Runtime/Engine/Private/UnrealEngine.cpp
@@ -662,6 +662,8 @@ void UEngine::Init(IEngineLoop* InEngineLoop)
 	}
 #endif // !UE_BUILD_SHIPPING
 
+	InitializeRunningAverageDeltaTime();
+
 	// Add to root.
 	AddToRoot();
 
@@ -991,6 +993,9 @@ void UEngine::UpdateTimeAndHandleMaxTickRate()
 			DeltaTime = 0.01;
 		}
 
+		// Give engine chance to update frame rate smoothing
+		UpdateRunningAverageDeltaTime(DeltaTime);
+
 		// Get max tick rate based on network settings and current delta time.
 		const float MaxTickRate	= GetMaxTickRate( DeltaTime );
 		float WaitTime		= 0;
@@ -5877,51 +5882,64 @@ static TAutoConsoleVariable<int32> CVarUnsteadyFPS(
 	TEXT("t.UnsteadyFPS"),0,
 	TEXT("Causes FPS to bounce around randomly in the 8-32 range."));
 
-/** Get tick rate limitor. */
-float UEngine::GetMaxTickRate( float DeltaTime, bool bAllowFrameRateSmoothing )
+void UEngine::InitializeRunningAverageDeltaTime()
 {
-	float MaxTickRate = 0;
+	// Running average delta time, initial value at 100 FPS so fast machines don't have to creep up
+	// to a good frame rate due to code limiting upward "mobility".
+	RunningAverageDeltaTime = 1 / 100.f;
+}
 
-	if (FPlatformProperties::AllowsFramerateSmoothing())
+bool UEngine::IsAllowedFramerateSmoothing(bool bAllowFrameRateSmoothing) const
+{
+	return FPlatformProperties::AllowsFramerateSmoothing() && bSmoothFrameRate && bAllowFrameRateSmoothing && !IsRunningDedicatedServer();
+}
+
+/** Compute tick rate limitor. */
+void UEngine::UpdateRunningAverageDeltaTime(float DeltaTime, bool bAllowFrameRateSmoothing)
+{
+	if (IsAllowedFramerateSmoothing(bAllowFrameRateSmoothing))
 	{
 		// Smooth the framerate if wanted. The code uses a simplistic running average. Other approaches, like reserving
 		// a percentage of time, ended up creating negative feedback loops in conjunction with GPU load and were abandonend.
-		if( bSmoothFrameRate && bAllowFrameRateSmoothing && !IsRunningDedicatedServer() )
+		if (DeltaTime < 0.0f)
 		{
-			if( DeltaTime < 0.0f )
-			{
 #if PLATFORM_ANDROID
 				UE_LOG(LogEngine, Warning, TEXT("Detected negative delta time - ignoring"));
 				DeltaTime = 0.01;
 #elif (UE_BUILD_SHIPPING && WITH_EDITOR)
-				// End users don't have access to the secure parts of UDN. The localized string points to the release notes,
-				// which should include a link to the AMD CPU drivers download site.
-				UE_LOG(LogEngine, Fatal, TEXT("%s"), TEXT("CPU time drift detected! Please consult release notes on how to address this."));
+			// End users don't have access to the secure parts of UDN. The localized string points to the release notes,
+			// which should include a link to the AMD CPU drivers download site.
+			UE_LOG(LogEngine, Fatal, TEXT("%s"), TEXT("CPU time drift detected! Please consult release notes on how to address this."));
 #else
-				// Send developers to the support list thread.
-				UE_LOG(LogEngine, Fatal, TEXT("Negative delta time! Please see https://udn.epicgames.com/lists/showpost.php?list=ue3bugs&id=4364"));
+			// Send developers to the support list thread.
+			UE_LOG(LogEngine, Fatal, TEXT("Negative delta time! Please see https://udn.epicgames.com/lists/showpost.php?list=ue3bugs&id=4364"));
 #endif
-			}
+		}
+
+		// Keep track of running average over 300 frames, clamping at min of 5 FPS for individual delta times.
+		RunningAverageDeltaTime = FMath::Lerp<float>(RunningAverageDeltaTime, FMath::Min<float>(DeltaTime, 0.2f), 1 / 300.f);
+	}
+}
 
-			// Running average delta time, initial value at 100 FPS so fast machines don't have to creep up
-			// to a good frame rate due to code limiting upward "mobility".
-			static float RunningAverageDeltaTime = 1 / 100.f;
 
-			// Keep track of running average over 300 frames, clamping at min of 5 FPS for individual delta times.
-			RunningAverageDeltaTime = FMath::Lerp<float>( RunningAverageDeltaTime, FMath::Min<float>( DeltaTime, 0.2f ), 1 / 300.f );
+/** Get tick rate limitor. */
+float UEngine::GetMaxTickRate(float DeltaTime, bool bAllowFrameRateSmoothing) const
+{
+	float MaxTickRate = 0;
 
-			// Work in FPS domain as that is what the function will return.
-			MaxTickRate = 1.f / RunningAverageDeltaTime;
+	if (IsAllowedFramerateSmoothing(bAllowFrameRateSmoothing))
+	{
+		// Work in FPS domain as that is what the function will return.
+		MaxTickRate = 1.f / RunningAverageDeltaTime;
 
-			// Clamp FPS into ini defined min/ max range.
-			if (SmoothedFrameRateRange.HasLowerBound())
-			{
-				MaxTickRate = FMath::Max( MaxTickRate, SmoothedFrameRateRange.GetLowerBoundValue() );
-			}
-			if (SmoothedFrameRateRange.HasUpperBound())
-			{
-				MaxTickRate = FMath::Min( MaxTickRate, SmoothedFrameRateRange.GetUpperBoundValue() );
-			}
+		// Clamp FPS into ini defined min/ max range.
+		if (SmoothedFrameRateRange.HasLowerBound())
+		{
+			MaxTickRate = FMath::Max( MaxTickRate, SmoothedFrameRateRange.GetLowerBoundValue() );
+		}
+		if (SmoothedFrameRateRange.HasUpperBound())
+		{
+			MaxTickRate = FMath::Min( MaxTickRate, SmoothedFrameRateRange.GetUpperBoundValue() );
 		}
 	}
 
