reverted recent change on static array as blueprints cannot do that yet
(LogCompile:Error: Static array cannot be exposed to blueprint)

diff --git a/Engine/Source/Runtime/Engine/Classes/Engine/Scene.h b/Engine/Source/Runtime/Engine/Classes/Engine/Scene.h
--- a/Engine/Source/Runtime/Engine/Classes/Engine/Scene.h
+++ b/Engine/Source/Runtime/Engine/Classes/Engine/Scene.h
@@ -590,7 +590,7 @@ struct FPostProcessSettings
 	class UTexture* LensFlareBokehShape;
 
 	/** RGB defines the lens flare color, A it's position. This is a temporary solution. */
-	UPROPERTY(EditAnywhere, BlueprintReadWrite, Category=LensFlares, meta=(editcondition = "bOverride_LensFlareTints", DisplayName = "Tints"))
+	UPROPERTY(EditAnywhere, Category=LensFlares, meta=(editcondition = "bOverride_LensFlareTints", DisplayName = "Tints"))
 	FLinearColor LensFlareTints[8];
 
 	/** 0..1 0=off/no vignette .. 1=strong vignette */
