Parents that are set to movable will automatically set children to movable #codereview(James.Golding,Ori.Cohen)

diff --git a/Engine/Source/Runtime/Engine/Private/SceneComponent.cpp b/Engine/Source/Runtime/Engine/Private/SceneComponent.cpp
--- a/Engine/Source/Runtime/Engine/Private/SceneComponent.cpp
+++ b/Engine/Source/Runtime/Engine/Private/SceneComponent.cpp
@@ -1135,6 +1135,17 @@ void USceneComponent::SetMobility(EComponentMobility::Type NewMobility)
 	{
 		FComponentReregisterContext ReregisterContext(this);
 		Mobility = NewMobility;
+
+		if (Mobility == EComponentMobility::Movable)	//if we're now movable all children should be updated as having static children is invalid
+		{
+			for (USceneComponent* ChildComponent : AttachChildren)
+			{
+				if (ChildComponent)
+				{
+					ChildComponent->SetMobility(NewMobility);
+				}
+			}
+		}
 	}
 }
 
