Engine: Fix typo in MapCheck_Message_SimulatePhyNoSimpleCollision message

diff --git a/Engine/Source/Runtime/Engine/Private/StaticMeshComponent.cpp b/Engine/Source/Runtime/Engine/Private/StaticMeshComponent.cpp
--- a/Engine/Source/Runtime/Engine/Private/StaticMeshComponent.cpp
+++ b/Engine/Source/Runtime/Engine/Private/StaticMeshComponent.cpp
@@ -380,7 +380,7 @@ void UStaticMeshComponent::CheckForErrors()
 	{
 		FMessageLog("MapCheck").Warning()
 			->AddToken(FUObjectToken::Create(this))
-			->AddToken(FTextToken::Create(LOCTEXT( "MapCheck_Message_SimulatePhyNoSimpleCollision", "Using bSimulatePhysics but StaticMesh has not simple collision.")));
+			->AddToken(FTextToken::Create(LOCTEXT( "MapCheck_Message_SimulatePhyNoSimpleCollision", "Using bSimulatePhysics but StaticMesh has no simple collision.")));
 	}
 
 	if( Mobility == EComponentMobility::Movable &&
