Fixed crash when toggling visibility of a renamed level
TTP# 344719

diff --git a/Engine/Source/Runtime/Engine/Private/LevelStreaming.cpp b/Engine/Source/Runtime/Engine/Private/LevelStreaming.cpp
--- a/Engine/Source/Runtime/Engine/Private/LevelStreaming.cpp
+++ b/Engine/Source/Runtime/Engine/Private/LevelStreaming.cpp
@@ -375,10 +375,12 @@ bool ULevelStreaming::RequestLevel(UWorld* PersistentWorld, bool bAllowLevelLoad
 				return false;
 			}
 #endif
-			SetLoadedLevel(World->PersistentLevel);
-			
-			// Broadcast level loaded event to blueprints
-			OnLevelLoaded.Broadcast();
+			if (World->PersistentLevel != LoadedLevel)
+			{
+				SetLoadedLevel(World->PersistentLevel);
+				// Broadcast level loaded event to blueprints
+				OnLevelLoaded.Broadcast();
+			}
 			
 			return true;
 		}
