CPF_EditInline flag removed (replaced by metadata).
CPF_PersistentInstance  flag added.

diff --git a/Engine/Source/Editor/KismetCompiler/Private/KismetCompilerCppBackend.cpp b/Engine/Source/Editor/KismetCompiler/Private/KismetCompilerCppBackend.cpp
--- a/Engine/Source/Editor/KismetCompiler/Private/KismetCompilerCppBackend.cpp
+++ b/Engine/Source/Editor/KismetCompiler/Private/KismetCompilerCppBackend.cpp
@@ -312,7 +312,7 @@ struct FEmitProperyHelper
 			HANDLE_CPF_TAG(TEXT("RepRetry"), CPF_RepRetry)
 			HANDLE_CPF_TAG(TEXT("Interp"), CPF_Edit | CPF_BlueprintVisible | CPF_Interp)
 			HANDLE_CPF_TAG(TEXT("NonTransactional"), CPF_NonTransactional)
-			HANDLE_CPF_TAG(TEXT("Instanced"), CPF_EditInline | CPF_ExportObject | CPF_InstancedReference)
+			HANDLE_CPF_TAG(TEXT("Instanced"), CPF_PersistentInstance | CPF_ExportObject | CPF_InstancedReference)
 			HANDLE_CPF_TAG(TEXT("BlueprintAssignable"), CPF_BlueprintAssignable)
 			HANDLE_CPF_TAG(TEXT("BlueprintCallable"), CPF_BlueprintCallable)
 			HANDLE_CPF_TAG(TEXT("BlueprintAuthorityOnly"), CPF_BlueprintAuthorityOnly)
diff --git a/Engine/Source/Editor/PropertyEditor/Private/PropertyNode.cpp b/Engine/Source/Editor/PropertyEditor/Private/PropertyNode.cpp
--- a/Engine/Source/Editor/PropertyEditor/Private/PropertyNode.cpp
+++ b/Engine/Source/Editor/PropertyEditor/Private/PropertyNode.cpp
@@ -129,7 +129,7 @@ void FPropertyNode::InitNode( const FPropertyNodeInitParams& InitParams )
 
 		// true if the property can be expanded into the property window; that is, instead of seeing
 		// a pointer to the object, you see the object's properties.
-		const bool bEditInline = ( (MyProperty->PropertyFlags&CPF_EditInline) && bIsObjectOrInterface && GotReadAddresses );
+		const bool bEditInline = bIsObjectOrInterface && GotReadAddresses && MyProperty->HasMetaData(TEXT("EditInline"));
 		SetNodeFlags(EPropertyNodeFlags::EditInline, bEditInline);
 
 		//Get the property max child depth
@@ -555,20 +555,6 @@ bool FPropertyNode::IsEditConst() const
 						bIsEditConst = true;
 						break;
 					}
-					/*else
-					{
-						// If 
-						if ( const FPropertyNode* ObjectParentNode = ObjectPropertyNode->GetParentNode() )
-						{
-							if ( const UProperty* ParentProperty = ObjectParentNode->GetProperty() )
-							{
-								if ( ParentProperty->PropertyFlags & CPF_EditInline )
-								{
-									return ObjectParentNode->IsEditConst();
-								}
-							}
-						}
-					}*/
 				}
 			}
 		}
@@ -1026,7 +1012,7 @@ struct FPropertyItemComponentCollector
 	{
 		if ( Property != NULL )
 		{
-			bContainsEditInlineNew |= (Property->PropertyFlags&CPF_EditInline) != 0 && (Property->PropertyFlags & CPF_EditConst) == 0;
+			bContainsEditInlineNew |= Property->HasMetaData(TEXT("EditInline")) && ((Property->PropertyFlags & CPF_EditConst) == 0);
 
 			if ( ProcessObjectProperty(Cast<UObjectPropertyBase>(Property), PropertyValueAddress) )
 			{
diff --git a/Engine/Source/Programs/UnrealHeaderTool/Private/HeaderParser.cpp b/Engine/Source/Programs/UnrealHeaderTool/Private/HeaderParser.cpp
--- a/Engine/Source/Programs/UnrealHeaderTool/Private/HeaderParser.cpp
+++ b/Engine/Source/Programs/UnrealHeaderTool/Private/HeaderParser.cpp
@@ -343,6 +343,11 @@ namespace
 			FError::Throwf(TEXT("SealedEvent cannot be used on Blueprint events"));
 		}
 	}
+
+	void AddEditInlineMetaData(TMap<FName, FString>& MetaData)
+	{
+		MetaData.Add(TEXT("EditInline"), TEXT("true"));
+	}
 }
 	
 /////////////////////////////////////////////////////
@@ -2555,7 +2560,8 @@ bool FHeaderParser::GetVarType
 			}
 			else if (Specifier == TEXT("Instanced"))
 			{
-				Flags |= CPF_EditInline | CPF_ExportObject | CPF_InstancedReference;
+				Flags |= CPF_PersistentInstance | CPF_ExportObject | CPF_InstancedReference;
+				AddEditInlineMetaData(MetaDataFromNewStyle);
 			}
 			else if (Specifier == TEXT("BlueprintAssignable"))
 			{
@@ -2662,13 +2668,6 @@ bool FHeaderParser::GetVarType
 		FError::Throwf(TEXT("Unknown access level"));
 	}
 
-
-	if ((Flags & CPF_EditInline) && (Flags & CPF_ExportObject) && !(Flags & CPF_InstancedReference))
-	{
-		Flags |= CPF_InstancedReference;
-		FError::Throwf(TEXT("Properties with 'editinline export' should instead use 'instanced' for clarity ") );
-	}
-
 	// Get variable type.
 	bool bUnconsumedStructKeyword = false;
 	bool bUnconsumedClassKeyword  = false;
@@ -3273,7 +3272,7 @@ bool FHeaderParser::GetVarType
 	}
 
 	// Perform some more specific validation on the property flags
-	if ((VarProperty.PropertyFlags & CPF_EditInline) && VarProperty.Type != CPT_ObjectReference)
+	if ((VarProperty.PropertyFlags & CPF_PersistentInstance) && VarProperty.Type != CPT_ObjectReference)
 	{
 		FError::Throwf(TEXT("'Instanced' is only allowed on object property (or array of objects)"));
 	}
@@ -3638,11 +3637,7 @@ UProperty* FHeaderParser::GetVarNameAndDim
 				if ( DoesAnythingInHierarchyHaveDefaultToInstanced( VarProperty.PropertyClass ) )
 				{
 					VarProperty.PropertyFlags |= CPF_InstancedReference;
-
-					if (0 != (VarProperty.PropertyFlags & CPF_Edit))
-					{
-						VarProperty.PropertyFlags |= CPF_EditInline;
-					}
+					AddEditInlineMetaData(VarProperty.MetaData);
 				}
 				NewProperty = new(NewScope,PropertyName,ObjectFlags)UObjectProperty(FPostConstructInitializeProperties());
 			}
@@ -3730,7 +3725,14 @@ UProperty* FHeaderParser::GetVarNameAndDim
 			if (NewProperty->PropertyFlags & (CPF_ContainsInstancedReference | CPF_InstancedReference))
 			{
 				VarProperty.PropertyFlags |= CPF_ContainsInstancedReference;
-				VarProperty.PropertyFlags &= ~CPF_InstancedReference; //this was propagated to the inner
+				VarProperty.PropertyFlags &= ~(CPF_InstancedReference | CPF_PersistentInstance); //this was propagated to the inner
+
+				if (0 != (NewProperty->PropertyFlags & CPF_PersistentInstance))
+				{
+					TMap<FName, FString> MetaData;
+					AddEditInlineMetaData(MetaData);
+					AddMetaDataToClassData(NewProperty, VarProperty.MetaData);
+				}
 			}
 			NewProperty = Array;
 		}
@@ -6635,7 +6637,6 @@ FHeaderParser::FHeaderParser(FFeedbackContext* InWarn)
 	LegalVariableSpecifiers.Add(TEXT("DuplicateTransient"));
 	LegalVariableSpecifiers.Add(TEXT("Ref"));
 	LegalVariableSpecifiers.Add(TEXT("Export"));
-	LegalVariableSpecifiers.Add(TEXT("EditInline"));
 	LegalVariableSpecifiers.Add(TEXT("NoClear"));
 	LegalVariableSpecifiers.Add(TEXT("EditFixedSize"));
 	LegalVariableSpecifiers.Add(TEXT("Replicated"));
diff --git a/Engine/Source/Runtime/CoreUObject/Private/UObject/Class.cpp b/Engine/Source/Runtime/CoreUObject/Private/UObject/Class.cpp
--- a/Engine/Source/Runtime/CoreUObject/Private/UObject/Class.cpp
+++ b/Engine/Source/Runtime/CoreUObject/Private/UObject/Class.cpp
@@ -3757,7 +3757,6 @@ void UFunction::Link(FArchive& Ar, bool bRelinkExistingProperties)
 
 bool UFunction::IsSignatureCompatibleWith(const UFunction* OtherFunction) const
 {
-	// The script compiler pointlessly sets CPF_EditInline on components passed as arguments
 	const uint64 IgnoreFlags = UFunction::GetDefaultIgnoredSignatureCompatibilityFlags();
 
 	return IsSignatureCompatibleWith(OtherFunction, IgnoreFlags);
diff --git a/Engine/Source/Runtime/CoreUObject/Private/UObject/Obj.cpp b/Engine/Source/Runtime/CoreUObject/Private/UObject/Obj.cpp
--- a/Engine/Source/Runtime/CoreUObject/Private/UObject/Obj.cpp
+++ b/Engine/Source/Runtime/CoreUObject/Private/UObject/Obj.cpp
@@ -2349,7 +2349,7 @@ TArray<const TCHAR*> ParsePropertyFlags(uint64 Flags)
 		TEXT("0x0000000000800000"),
 		TEXT("CPF_SaveGame"),	
 		TEXT("CPF_NoClear"),
-		TEXT("CPF_EditInline"),
+		TEXT("0x0000000004000000"),
 		TEXT("CPF_ReferenceParm"),
 		TEXT("CPF_BlueprintAssignable"),
 		TEXT("CPF_Deprecated"),
@@ -2372,6 +2372,7 @@ TArray<const TCHAR*> ParsePropertyFlags(uint64 Flags)
 		TEXT("CPF_TextExportTransient"),
 		TEXT("CPF_NonPIEDuplicateTransient"),
 		TEXT("CPF_ExposeOnSpawn"),
+		TEXT("CPF_PersistentInstance"),
 	};
 
 	for (const TCHAR* FlagName : PropertyFlags)
diff --git a/Engine/Source/Runtime/CoreUObject/Public/UObject/Class.h b/Engine/Source/Runtime/CoreUObject/Public/UObject/Class.h
--- a/Engine/Source/Runtime/CoreUObject/Public/UObject/Class.h
+++ b/Engine/Source/Runtime/CoreUObject/Public/UObject/Class.h
@@ -1248,7 +1248,7 @@ public:
 	FORCEINLINE static uint64 GetDefaultIgnoredSignatureCompatibilityFlags()
 	{
 		//@TODO: UCREMOVAL: CPF_ConstParm added as a hack to get blueprints compiling with a const DamageType parameter.
-		const uint64 IgnoreFlags = CPF_EditInline | CPF_ExportObject | CPF_InstancedReference | CPF_ContainsInstancedReference | CPF_ComputedFlags | CPF_ConstParm;
+		const uint64 IgnoreFlags = CPF_PersistentInstance | CPF_ExportObject | CPF_InstancedReference | CPF_ContainsInstancedReference | CPF_ComputedFlags | CPF_ConstParm;
 		return IgnoreFlags;
 	}
 
diff --git a/Engine/Source/Runtime/CoreUObject/Public/UObject/ObjectBase.h b/Engine/Source/Runtime/CoreUObject/Public/UObject/ObjectBase.h
--- a/Engine/Source/Runtime/CoreUObject/Public/UObject/ObjectBase.h
+++ b/Engine/Source/Runtime/CoreUObject/Public/UObject/ObjectBase.h
@@ -305,7 +305,7 @@ typedef uint64 EClassCastFlags;
 //#define CPF_    							DECLARE_UINT64(0x0000000000800000)		// 
 #define CPF_SaveGame						DECLARE_UINT64(0x0000000001000000)		// Property should be serialized for save games
 #define CPF_NoClear							DECLARE_UINT64(0x0000000002000000)		// Hide clear (and browse) button.
-#define CPF_EditInline						DECLARE_UINT64(0x0000000004000000)		// Edit this object reference inline. Property must have also CPF_InstancedReference. The sub-object instance will be persistent.
+//#define CPF_  							DECLARE_UINT64(0x0000000004000000)		//
 #define CPF_ReferenceParm					DECLARE_UINT64(0x0000000008000000)		// Value is passed by reference; CPF_OutParam and CPF_Param should also be set.
 #define CPF_BlueprintAssignable				DECLARE_UINT64(0x0000000010000000)		// MC Delegates only.  Property should be exposed for assigning in blueprint code
 #define CPF_Deprecated  					DECLARE_UINT64(0x0000000020000000)		// Property is deprecated.  Read it from an archive, but don't save it.
@@ -328,6 +328,7 @@ typedef uint64 EClassCastFlags;
 #define CPF_TextExportTransient				DECLARE_UINT64(0x0000400000000000)		// Property shouldn't be exported to text format (e.g. copy/paste)
 #define CPF_NonPIEDuplicateTransient		DECLARE_UINT64(0x0000800000000000)		// Property should only be copied in PIE
 #define CPF_ExposeOnSpawn					DECLARE_UINT64(0x0001000000000000)		// Property is exposed on spawn
+#define CPF_PersistentInstance				DECLARE_UINT64(0x0002000000000000)		//
 
 #define CPF_NonPIETransient \
 	EMIT_DEPRECATED_WARNING_MESSAGE("CPF_NonPIETransient is deprecated. Please use CPF_NonPIEDuplicateTransient instead.") \
@@ -336,7 +337,7 @@ typedef uint64 EClassCastFlags;
 /** @name Combinations flags */
 //@{
 #define CPF_ParmFlags				(CPF_Parm | CPF_OutParm | CPF_ReturnParm | CPF_ReferenceParm | CPF_ConstParm)
-#define CPF_PropagateToArrayInner	(CPF_ExportObject | CPF_EditInline | CPF_InstancedReference | CPF_ContainsInstancedReference | CPF_Localized | CPF_Config | CPF_EditConst | CPF_Deprecated | CPF_EditorOnly | CPF_AutoWeak )
+#define CPF_PropagateToArrayInner	(CPF_ExportObject | CPF_PersistentInstance | CPF_InstancedReference | CPF_ContainsInstancedReference | CPF_Localized | CPF_Config | CPF_EditConst | CPF_Deprecated | CPF_EditorOnly | CPF_AutoWeak )
 
 /** the flags that should never be set on interface properties */
 #define CPF_InterfaceClearMask		(CPF_ExportObject|CPF_InstancedReference|CPF_ContainsInstancedReference)
diff --git a/Engine/Source/Runtime/Engine/Private/UnrealEngine.cpp b/Engine/Source/Runtime/Engine/Private/UnrealEngine.cpp
--- a/Engine/Source/Runtime/Engine/Private/UnrealEngine.cpp
+++ b/Engine/Source/Runtime/Engine/Private/UnrealEngine.cpp
@@ -5901,23 +5901,23 @@ void UEngine::UpdateRunningAverageDeltaTime(float DeltaTime, bool bAllowFrameRat
 	{
 		// Smooth the framerate if wanted. The code uses a simplistic running average. Other approaches, like reserving
 		// a percentage of time, ended up creating negative feedback loops in conjunction with GPU load and were abandonend.
-		if (DeltaTime < 0.0f)
-		{
+			if( DeltaTime < 0.0f )
+			{
 #if PLATFORM_ANDROID
 				UE_LOG(LogEngine, Warning, TEXT("Detected negative delta time - ignoring"));
 				DeltaTime = 0.01;
 #elif (UE_BUILD_SHIPPING && WITH_EDITOR)
-			// End users don't have access to the secure parts of UDN. The localized string points to the release notes,
-			// which should include a link to the AMD CPU drivers download site.
-			UE_LOG(LogEngine, Fatal, TEXT("%s"), TEXT("CPU time drift detected! Please consult release notes on how to address this."));
+				// End users don't have access to the secure parts of UDN. The localized string points to the release notes,
+				// which should include a link to the AMD CPU drivers download site.
+				UE_LOG(LogEngine, Fatal, TEXT("%s"), TEXT("CPU time drift detected! Please consult release notes on how to address this."));
 #else
-			// Send developers to the support list thread.
-			UE_LOG(LogEngine, Fatal, TEXT("Negative delta time! Please see https://udn.epicgames.com/lists/showpost.php?list=ue3bugs&id=4364"));
+				// Send developers to the support list thread.
+				UE_LOG(LogEngine, Fatal, TEXT("Negative delta time! Please see https://udn.epicgames.com/lists/showpost.php?list=ue3bugs&id=4364"));
 #endif
-		}
+			}
 
-		// Keep track of running average over 300 frames, clamping at min of 5 FPS for individual delta times.
-		RunningAverageDeltaTime = FMath::Lerp<float>(RunningAverageDeltaTime, FMath::Min<float>(DeltaTime, 0.2f), 1 / 300.f);
+			// Keep track of running average over 300 frames, clamping at min of 5 FPS for individual delta times.
+			RunningAverageDeltaTime = FMath::Lerp<float>( RunningAverageDeltaTime, FMath::Min<float>( DeltaTime, 0.2f ), 1 / 300.f );
 	}
 }
 
@@ -5929,19 +5929,19 @@ float UEngine::GetMaxTickRate(float DeltaTime, bool bAllowFrameRateSmoothing) co
 
 	if (IsAllowedFramerateSmoothing(bAllowFrameRateSmoothing))
 	{
-		// Work in FPS domain as that is what the function will return.
-		MaxTickRate = 1.f / RunningAverageDeltaTime;
+			// Work in FPS domain as that is what the function will return.
+			MaxTickRate = 1.f / RunningAverageDeltaTime;
 
-		// Clamp FPS into ini defined min/ max range.
-		if (SmoothedFrameRateRange.HasLowerBound())
-		{
-			MaxTickRate = FMath::Max( MaxTickRate, SmoothedFrameRateRange.GetLowerBoundValue() );
-		}
-		if (SmoothedFrameRateRange.HasUpperBound())
-		{
-			MaxTickRate = FMath::Min( MaxTickRate, SmoothedFrameRateRange.GetUpperBoundValue() );
+			// Clamp FPS into ini defined min/ max range.
+			if (SmoothedFrameRateRange.HasLowerBound())
+			{
+				MaxTickRate = FMath::Max( MaxTickRate, SmoothedFrameRateRange.GetLowerBoundValue() );
+			}
+			if (SmoothedFrameRateRange.HasUpperBound())
+			{
+				MaxTickRate = FMath::Min( MaxTickRate, SmoothedFrameRateRange.GetUpperBoundValue() );
+			}
 		}
-	}
 
 	if (CVarCauseHitches.GetValueOnGameThread())
 	{
@@ -9857,8 +9857,9 @@ struct FFindInstancedReferenceSubobjectHelper
 		check(ContainerAddress);
 		for (UProperty* Prop = (Struct ? Struct->RefLink : NULL); Prop; Prop = Prop->NextRef)
 		{
-			if (Prop->HasAllPropertyFlags(CPF_InstancedReference | CPF_EditInline))
+			if (Prop->HasAllPropertyFlags(CPF_PersistentInstance))
 			{
+				ensure(Prop->HasAllPropertyFlags(CPF_InstancedReference));
 				auto ObjectProperty = Cast<const UObjectProperty>(Prop);
 				if (ObjectProperty)
 				{
