Don't rely on MouseEnter/Leave events to track when the mouse is over the viewport, but  instead rely on the cached geometry position

diff --git a/Engine/Source/Runtime/Engine/Classes/Engine/GameViewportClient.h b/Engine/Source/Runtime/Engine/Classes/Engine/GameViewportClient.h
--- a/Engine/Source/Runtime/Engine/Classes/Engine/GameViewportClient.h
+++ b/Engine/Source/Runtime/Engine/Classes/Engine/GameViewportClient.h
@@ -205,9 +205,6 @@ public:
 	/** set to disable world rendering */
 	uint32 bDisableWorldRendering:1;
 
-	/** Whether the mouse is currently over this game viewport */
-	uint32 bMouseOverViewport:1;
-
 protected:
 	/**
 	 * The splitscreen type that is actually being used; takes into account the number of players and other factors (such as cinematic mode)
diff --git a/Engine/Source/Runtime/Engine/Private/GameViewportClient.cpp b/Engine/Source/Runtime/Engine/Private/GameViewportClient.cpp
--- a/Engine/Source/Runtime/Engine/Private/GameViewportClient.cpp
+++ b/Engine/Source/Runtime/Engine/Private/GameViewportClient.cpp
@@ -409,8 +409,6 @@ void UGameViewportClient::MouseEnter(FViewport* InViewport, int32 x, int32 y)
 {
 	Super::MouseEnter(InViewport, x, y);
 
-	bMouseOverViewport = true;
-
 	if (GetDefault<UInputSettings>()->bUseMouseForTouch && !GetGameViewport()->GetPlayInEditorIsSimulate())
 	{
 		FSlateApplication::Get().SetGameIsFakingTouchEvents(true);
@@ -421,8 +419,6 @@ void UGameViewportClient::MouseLeave(FViewport* InViewport)
 {
 	Super::MouseLeave(InViewport);
 
-	bMouseOverViewport = false;
-
 	if (GetDefault<UInputSettings>()->bUseMouseForTouch)
 	{
 		UE_LOG(LogTemp, Log, TEXT("MouseLeave"));
@@ -439,12 +435,15 @@ bool UGameViewportClient::GetMousePosition(FVector2D& MousePosition) const
 {
 	bool bGotMousePosition = false;
 
-	if (Viewport && bMouseOverViewport && FSlateApplication::Get().IsMouseAttached())
+	if (Viewport && FSlateApplication::Get().IsMouseAttached())
 	{
 		FIntPoint MousePos;
 		Viewport->GetMousePos(MousePos);
-		MousePosition = FVector2D(MousePos);
-		bGotMousePosition = true;
+		if (MousePos.X >= 0 && MousePos.Y >= 0)
+		{
+			MousePosition = FVector2D(MousePos);
+			bGotMousePosition = true;
+		}
 	}
 
 	return bGotMousePosition;
