Prevented vislog to log for default objects - World is needed for visual logs to track timestamp #ue4

diff --git a/Engine/Source/Runtime/Engine/Public/VisualLog.h b/Engine/Source/Runtime/Engine/Public/VisualLog.h
--- a/Engine/Source/Runtime/Engine/Public/VisualLog.h
+++ b/Engine/Source/Runtime/Engine/Public/VisualLog.h
@@ -221,7 +221,7 @@ struct ENGINE_API FVisLogEntry
 { \
 	SCOPE_CYCLE_COUNTER(STAT_VisualLog); \
 	static_assert((ELogVerbosity::Verbosity & ELogVerbosity::VerbosityMask) < ELogVerbosity::NumVerbosity && ELogVerbosity::Verbosity > 0, "Verbosity must be constant and in range."); \
-	if (FVisualLog::Get().IsRecording() && (!FVisualLog::Get().IsAllBlocked() || FVisualLog::Get().InWhitelist(CategoryName.GetCategoryName()))) \
+	if (FVisualLog::Get().IsRecording() && (!FVisualLog::Get().IsAllBlocked() || FVisualLog::Get().InWhitelist(CategoryName.GetCategoryName())) && Object && !Object->HasAnyFlags(RF_ClassDefaultObject)) \
 	{ \
 		FVisualLog::Get().LogLine(Object, CategoryName.GetCategoryName(), ELogVerbosity::Verbosity, FString::Printf(Format, ##__VA_ARGS__)); \
 	} \
@@ -243,7 +243,7 @@ struct ENGINE_API FVisLogEntry
 { \
 	SCOPE_CYCLE_COUNTER(STAT_VisualLog); \
 	static_assert((ELogVerbosity::Verbosity & ELogVerbosity::VerbosityMask) < ELogVerbosity::NumVerbosity && ELogVerbosity::Verbosity > 0, "Verbosity must be constant and in range."); \
-	if (FVisualLog::Get().IsRecording() && (!FVisualLog::Get().IsAllBlocked() || FVisualLog::Get().InWhitelist(CategoryName.GetCategoryName()))) \
+	if (FVisualLog::Get().IsRecording() && (!FVisualLog::Get().IsAllBlocked() || FVisualLog::Get().InWhitelist(CategoryName.GetCategoryName())) && Object && !Object->HasAnyFlags(RF_ClassDefaultObject)) \
 	{ \
 		FVisualLog::Get().GetEntryToWrite(Object)->AddElement(SegmentStart, SegmentEnd, CategoryName.GetCategoryName(), Color, FString::Printf(DescriptionFormat, ##__VA_ARGS__), Thickness); \
 	} \
@@ -256,7 +256,7 @@ struct ENGINE_API FVisLogEntry
 { \
 	SCOPE_CYCLE_COUNTER(STAT_VisualLog); \
 	static_assert((ELogVerbosity::Verbosity & ELogVerbosity::VerbosityMask) < ELogVerbosity::NumVerbosity && ELogVerbosity::Verbosity > 0, "Verbosity must be constant and in range."); \
-	if (FVisualLog::Get().IsRecording() && (!FVisualLog::Get().IsAllBlocked() || FVisualLog::Get().InWhitelist(CategoryName.GetCategoryName()))) \
+	if (FVisualLog::Get().IsRecording() && (!FVisualLog::Get().IsAllBlocked() || FVisualLog::Get().InWhitelist(CategoryName.GetCategoryName())) && Object && !Object->HasAnyFlags(RF_ClassDefaultObject)) \
 	{ \
 		FVisualLog::Get().GetEntryToWrite(Object)->AddElement(Location, CategoryName.GetCategoryName(), Color, FString::Printf(DescriptionFormat, ##__VA_ARGS__), Radius); \
 	} \
@@ -266,7 +266,7 @@ struct ENGINE_API FVisLogEntry
 { \
 	SCOPE_CYCLE_COUNTER(STAT_VisualLog); \
 	static_assert((ELogVerbosity::Verbosity & ELogVerbosity::VerbosityMask) < ELogVerbosity::NumVerbosity && ELogVerbosity::Verbosity > 0, "Verbosity must be constant and in range."); \
-	if (FVisualLog::Get().IsRecording() && (!FVisualLog::Get().IsAllBlocked() || FVisualLog::Get().InWhitelist(CategoryName.GetCategoryName()))) \
+	if (FVisualLog::Get().IsRecording() && (!FVisualLog::Get().IsAllBlocked() || FVisualLog::Get().InWhitelist(CategoryName.GetCategoryName())) && Object && !Object->HasAnyFlags(RF_ClassDefaultObject)) \
 	{ \
 		FVisualLog::Get().GetEntryToWrite(Object)->AddElement(Box, CategoryName.GetCategoryName(), Color, FString::Printf(DescriptionFormat, ##__VA_ARGS__)); \
 	} \
@@ -276,7 +276,7 @@ struct ENGINE_API FVisLogEntry
 { \
 	SCOPE_CYCLE_COUNTER(STAT_VisualLog); \
 	static_assert((ELogVerbosity::Verbosity & ELogVerbosity::VerbosityMask) < ELogVerbosity::NumVerbosity && ELogVerbosity::Verbosity > 0, "Verbosity must be constant and in range."); \
-	if (FVisualLog::Get().IsRecording() && (!FVisualLog::Get().IsAllBlocked() || FVisualLog::Get().InWhitelist(CategoryName.GetCategoryName()))) \
+	if (FVisualLog::Get().IsRecording() && (!FVisualLog::Get().IsAllBlocked() || FVisualLog::Get().InWhitelist(CategoryName.GetCategoryName())) && Object && !Object->HasAnyFlags(RF_ClassDefaultObject)) \
 	{ \
 		FVisualLog::Get().GetEntryToWrite(Object)->AddHistogramData(Data, CategoryName.GetCategoryName(), GraphName, DataName); \
 	} \
