#UE4

- Adding Native NetSerialize support to FastArrayReplication path
- Changed a few function to take in a UScriptStruct to avoid a cast inside SerializePropertiesForStruct

diff --git a/Engine/Source/Runtime/CoreUObject/Public/UObject/CoreNet.h b/Engine/Source/Runtime/CoreUObject/Public/UObject/CoreNet.h
--- a/Engine/Source/Runtime/CoreUObject/Public/UObject/CoreNet.h
+++ b/Engine/Source/Runtime/CoreUObject/Public/UObject/CoreNet.h
@@ -278,7 +278,7 @@ class INetSerializeCB
 public:
 	INetSerializeCB() { }
 
-	virtual void NetSerializeStruct( UStruct * Struct, FArchive & Ar, UPackageMap *	Map, void * Data, bool & bHasUnmapped ) = 0;
+	virtual void NetSerializeStruct( UScriptStruct* Struct, FArchive& Ar, UPackageMap* Map, void* Data, bool& bHasUnmapped ) = 0;
 };
 
 class IRepChangedPropertyTracker
diff --git a/Engine/Source/Runtime/Engine/Classes/Engine/NetSerialization.h b/Engine/Source/Runtime/Engine/Classes/Engine/NetSerialization.h
--- a/Engine/Source/Runtime/Engine/Classes/Engine/NetSerialization.h
+++ b/Engine/Source/Runtime/Engine/Classes/Engine/NetSerialization.h
@@ -370,7 +370,7 @@ struct FFastArraySerializer
 template< typename Type, typename SerializerType >
 bool FFastArraySerializer::FastArrayDeltaSerialize( TArray<Type> &Items, FNetDeltaSerializeInfo& Parms, SerializerType& ArraySerializer )
 {
-	UStruct* InnerStruct = Type::StaticStruct();
+	UScriptStruct* InnerStruct = Type::StaticStruct();
 
 	if (Parms.OutBunch)
 	{
diff --git a/Engine/Source/Runtime/Engine/Private/DataReplication.cpp b/Engine/Source/Runtime/Engine/Private/DataReplication.cpp
--- a/Engine/Source/Runtime/Engine/Private/DataReplication.cpp
+++ b/Engine/Source/Runtime/Engine/Private/DataReplication.cpp
@@ -18,11 +18,23 @@ public:
 
 	FNetSerializeCB( UNetDriver * InNetDriver ) : Driver( InNetDriver ) { }
 
-	virtual void NetSerializeStruct( UStruct * Struct, FArchive & Ar, UPackageMap *	Map, void * Data, bool & bHasUnmapped )
+	virtual void NetSerializeStruct( UScriptStruct* Struct, FArchive& Ar, UPackageMap* Map, void* Data, bool& bHasUnmapped )
 	{
-		TSharedPtr<FRepLayout> RepLayout = Driver->GetStructRepLayout( Struct );
-
-		RepLayout->SerializePropertiesForStruct( Struct, Ar, Map, Data, bHasUnmapped );
+		if (Struct->StructFlags & STRUCT_NetSerializeNative)
+		{
+			UScriptStruct::ICppStructOps* CppStructOps = Struct->GetCppStructOps();
+			check(CppStructOps); // else should not have STRUCT_NetSerializeNative
+			check(!Struct->InheritedCppStructOps()); // else should not have STRUCT_NetSerializeNative
+			if (!CppStructOps->NetSerialize(Ar, Map, bHasUnmapped, Data))
+			{
+				UE_LOG(LogNet, Warning, TEXT("NetSerializeStruct: Native NetSerialize %s failed."), *Struct->GetFullName());
+			}
+		}
+		else
+		{
+			TSharedPtr<FRepLayout> RepLayout = Driver->GetStructRepLayout(Struct);
+			RepLayout->SerializePropertiesForStruct(Struct, Ar, Map, Data, bHasUnmapped);
+		}
 	}
 
 	UNetDriver * Driver;
