PIE - Move BeginPlay to after the netdriver initialization step (to be more consistent with LoadMap)

diff --git a/Engine/Source/Runtime/Engine/Private/GameInstance.cpp b/Engine/Source/Runtime/Engine/Private/GameInstance.cpp
--- a/Engine/Source/Runtime/Engine/Private/GameInstance.cpp
+++ b/Engine/Source/Runtime/Engine/Private/GameInstance.cpp
@@ -215,10 +215,6 @@ bool UGameInstance::StartPIEGameInstance(ULocalPlayer* LocalPlayer, bool bInSimu
 			}
 		}
 
-		static bool bDoBeginPlayHere = true;
-		if (bDoBeginPlayHere)
-			PlayWorld->BeginPlay();
-
 		UGameViewportClient* const GameViewport = GetGameViewportClient();
 		if (GameViewport != NULL && GameViewport->Viewport != NULL)
 		{
@@ -231,6 +227,8 @@ bool UGameInstance::StartPIEGameInstance(ULocalPlayer* LocalPlayer, bool bInSimu
 			// start listen server with the built URL
 			PlayWorld->Listen(URL);
 		}
+
+		PlayWorld->BeginPlay();
 	}
 
 	return true;
