#ue4 - Fix Character movement restricting air control because of incorrect sweep tests.

- CharacterMovementComponent::FindAirControlImpact() should always use the capsule shape, not the boxes used for floor checks when box checks are enabled. If boxes are used (from FloorSweepTest), those boxes could start in penetration since they extend out of the lower portion of the capsule. This would disable air control even if trying to move away from an object.

TTP 347625

diff --git a/Engine/Source/Runtime/Engine/Private/CharacterMovementComponent.cpp b/Engine/Source/Runtime/Engine/Private/CharacterMovementComponent.cpp
--- a/Engine/Source/Runtime/Engine/Private/CharacterMovementComponent.cpp
+++ b/Engine/Source/Runtime/Engine/Private/CharacterMovementComponent.cpp
@@ -3036,7 +3036,8 @@ bool UCharacterMovementComponent::FindAirControlImpact(float DeltaTime, float Ti
 		InitCollisionParams(CapsuleQuery, ResponseParam);
 		const FVector CapsuleLocation = UpdatedComponent->GetComponentLocation();
 		const FCollisionShape CapsuleShape = GetPawnCapsuleCollisionShape(SHRINK_None);
-		if (FloorSweepTest(OutHitResult, CapsuleLocation, CapsuleLocation + TestWalk, UpdatedComponent->GetCollisionObjectType(), CapsuleShape, CapsuleQuery, ResponseParam))
+		
+		if (GetWorld()->SweepSingle(OutHitResult, CapsuleLocation, CapsuleLocation + TestWalk, FQuat::Identity, UpdatedComponent->GetCollisionObjectType(), CapsuleShape, CapsuleQuery, ResponseParam))
 		{
 			return true;
 		}
