Added tiny font as alternative to small font for reporter graphs - it can be useful when showing many 2d graphs at once on the screen #ue4
Added cursor for reporter graph, to be able to set it and show values at specific location. It works only with style set to lines #ue4
Enabled cursor and tiny font for 2d graph visualization in LogVisualizer #ue4

diff --git a/Engine/Source/Developer/LogVisualizer/Private/SLogVisualizer.cpp b/Engine/Source/Developer/LogVisualizer/Private/SLogVisualizer.cpp
--- a/Engine/Source/Developer/LogVisualizer/Private/SLogVisualizer.cpp
+++ b/Engine/Source/Developer/LogVisualizer/Private/SLogVisualizer.cpp
@@ -2181,6 +2181,9 @@ void SLogVisualizer::DrawOnCanvas(UCanvas* Canvas, APlayerController*)
 					CurrentX += GraphWidth + XGraphSpacing;
 					HistogramGraph->SetAxesMinMax(FVector2D(TimeStampWindow.X, It->Value.Min.Y), FVector2D(TimeStampWindow.Y, It->Value.Max.Y));
 					
+					HistogramGraph->DrawCursorOnGraph(true);
+					HistogramGraph->UseTinyFont(CollectedGraphs.Num() >= 5);
+					HistogramGraph->SetCursorLocation(Entry->TimeStamp);
 					HistogramGraph->SetNumThresholds(0);
 					HistogramGraph->SetStyles(EGraphAxisStyle::Grid, EGraphDataStyle::Lines);
 					HistogramGraph->SetBackgroundColor( FColor(0,0,0, 200) );
diff --git a/Engine/Source/Runtime/Engine/Classes/Debug/ReporterGraph.h b/Engine/Source/Runtime/Engine/Classes/Debug/ReporterGraph.h
--- a/Engine/Source/Runtime/Engine/Classes/Debug/ReporterGraph.h
+++ b/Engine/Source/Runtime/Engine/Classes/Debug/ReporterGraph.h
@@ -185,6 +185,14 @@ public:
 	 */
 	bool IsOffsetForDataSetsEnabled() { return !!bOffsetDataSets; }
 
+	/** Sets cursor location on line graph to show value at specyfic place */
+	void SetCursorLocation(float InValue) { CursorLocation = InValue; }
+
+	/** setter to force tiny font instead of small font */
+	void UseTinyFont(bool InUseTinyFont) { bUseTinyFont = InUseTinyFont;  }
+
+	/** setter to enable or disable cursor for line graphs */
+	void DrawCursorOnGraph(bool InDrawCursorOnGraph) { bDrawCursorOnGraph = InDrawCursorOnGraph;  }
 protected:
 
 	/** Draw background under graph
@@ -234,6 +242,9 @@ protected:
 	*/
 	FVector2D DataToNormalized(const FVector2D& InVector);
 
+	/** Returns default font used to print texts */
+	UFont* GetDefaultFont();
+
 	/** The screen size of the graph */
 	FRect GraphScreenSize;
 
@@ -269,6 +280,16 @@ protected:
 	/** Background color to draw under graph */
 	FColor BackgroundColor;
 
+	/** Current location for cursor on line graphs */
+	float CursorLocation;
+
+	/** if set it enables small offset for graphs to better visualize overlapping data sets  */
 	int32 bOffsetDataSets : 1;
+
+	/** if set it forces tinyt font for texts */
+	int32 bUseTinyFont : 1;
+
+	/** if set it enables cursor for line graphs */
+	int32 bDrawCursorOnGraph : 1;
 };
 
diff --git a/Engine/Source/Runtime/Engine/Private/Debug/ReporterGraph.cpp b/Engine/Source/Runtime/Engine/Private/Debug/ReporterGraph.cpp
--- a/Engine/Source/Runtime/Engine/Private/Debug/ReporterGraph.cpp
+++ b/Engine/Source/Runtime/Engine/Private/Debug/ReporterGraph.cpp
@@ -15,6 +15,8 @@ UReporterGraph::UReporterGraph(const class FPostConstructInitializeProperties& P
 	NumXNotches = 10;
 	NumYNotches = 10;
 	bOffsetDataSets = false;
+	UseTinyFont(false);
+	DrawCursorOnGraph(false);
 
 	BackgroundColor = FColor(0, 0, 0, 0);
 }
@@ -150,7 +152,7 @@ void UReporterGraph::DrawLegend(UCanvas* Canvas)
 {
 	FVector2D CurrentTextPos = GraphScreenSize.Max;
 	CurrentTextPos.X = GraphScreenSize.Min.X;
-	UFont* Font = GEngine->GetSmallFont();
+	UFont* Font = GetDefaultFont();
 
 	int32 DummyY, CurrentX;
 	StringSize(Font, CurrentX, DummyY, TEXT("99.99"));
@@ -191,7 +193,7 @@ void UReporterGraph::DrawAxes(UCanvas* Canvas)
 	FVector2D YMax = GraphScreenSize.Min;
 	YMax.Y = GraphScreenSize.Max.Y;
 
-	UFont* Font = GEngine->GetSmallFont();
+	UFont* Font = GetDefaultFont();
 	int32 StringSizeX, StringSizeY;
 	// Draw the X axis
 	StringSize(Font, StringSizeX, StringSizeY, *FString::Printf(TEXT("%.2f"), GraphMinMaxData.Max.X) );
@@ -268,7 +270,7 @@ void UReporterGraph::DrawAxis(UCanvas* Canvas, FVector2D Start, FVector2D End, f
 		} break;
 	}
 	
-	UFont* Font = GEngine->GetSmallFont();
+	UFont* Font = GetDefaultFont();
 	const FVector2D Width = FVector2D((GraphScreenSize.Max.X - GraphScreenSize.Min.X), 0);
 	const FVector2D Height = FVector2D(0, (GraphScreenSize.Max.Y - GraphScreenSize.Min.Y));
 
@@ -305,7 +307,7 @@ void UReporterGraph::DrawAxis(UCanvas* Canvas, FVector2D Start, FVector2D End, f
 
 void UReporterGraph::DrawThresholds(UCanvas* Canvas)
 {
-	UFont* Font = GEngine->GetSmallFont();
+	UFont* Font = GetDefaultFont();
 	for(int32 i = 0; i < Thresholds.Num(); i++)
 	{
 		if(Thresholds[i].Threshold < GraphMinMaxData.Max.Y)
@@ -335,9 +337,19 @@ void UReporterGraph::DrawData(UCanvas* Canvas)
 	const FVector2D Max = FVector2D(GraphScreenSize.Max.X * Canvas->SizeX, Canvas->SizeY - GraphScreenSize.Max.Y * Canvas->SizeY);
 	const float Height = GraphScreenSize.Max.Y - GraphScreenSize.Min.Y;
 	const float dx = Height / FMath::Abs(Max.Y - Min.Y);
+	UFont* Font = GetDefaultFont();
+	int32 StringSizeX, StringSizeY;
+	// Draw the X axis
+	StringSize(Font, StringSizeX, StringSizeY, TEXT("0"));
 
 	float UpOffset = 0;
 	float DownOffset = 0;
+
+	if (bDrawCursorOnGraph && DataStyle == EGraphDataStyle::Lines)
+	{
+		DrawLine(Canvas, DataToNormalized(FVector2D(CursorLocation, GraphMinMaxData.Min.Y)), DataToNormalized(FVector2D(CursorLocation, GraphMinMaxData.Max.Y)), FLinearColor::White, EReporterLineStyle::Line);
+	}
+
 	for (int32 i = 0; i < CurrentData.Num(); i++)
 	{
 		if (IsOffsetForDataSetsEnabled())
@@ -352,29 +364,32 @@ void UReporterGraph::DrawData(UCanvas* Canvas)
 			}
 		}
 
-		for(int32 j = 1; j < CurrentData[i].Data.Num(); j++)
+		for (int32 j = 1; j < CurrentData[i].Data.Num(); j++)
 		{
-			Start = CurrentData[i].Data[j - 1];
-			End = CurrentData[i].Data[j];
+			FVector2D DataStart = Start = CurrentData[i].Data[j - 1];
+			FVector2D DataEnd = End = CurrentData[i].Data[j];
 
 			Start = DataToNormalized(Start);
 			End = DataToNormalized(End);
 
-			if(DataStyle == EGraphDataStyle::Lines)
+			if (DataStyle == EGraphDataStyle::Lines)
 			{
-				if (i % 2)
-				{
-					DrawLine(Canvas, Start + FVector2D(UpOffset, UpOffset), End + FVector2D(UpOffset, UpOffset), CurrentData[i].Color);
-				}
-				else
+				const FVector2D DrawOffset = FVector2D(UpOffset, UpOffset) * (i % 2 ? 1.0 : -1.0f);
+				DrawLine(Canvas, Start + DrawOffset, End + DrawOffset, CurrentData[i].Color);
+
+				if (bDrawCursorOnGraph && CursorLocation >= DataStart.X && CursorLocation < DataEnd.X)
 				{
-					DrawLine(Canvas, Start - FVector2D(DownOffset, DownOffset), End - FVector2D(DownOffset, DownOffset), CurrentData[i].Color);
+					const float t = (CursorLocation - DataStart.X) / (DataEnd.X - DataStart.X);
+					FVector2D Location = FMath::Lerp<FVector2D, float>(DataStart, DataEnd, t);
+
+					FVector2D TextPos = ToScreenSpace(DataToNormalized(Location), Canvas) + DrawOffset;
+					Canvas->Canvas->DrawShadowedString(StringSizeX + TextPos.X, TextPos.Y, *FString::Printf(TEXT("%1.2f"), Location.Y), Font, CurrentData[i].Color);
 				}
 			}
 			else
 			{
 				FVector2D Position0, Position1, Position2;
-				
+
 				// draw the top triangle of the quad
 				Position0.X = Start.X;
 				Position0.Y = (GraphMinMaxData.Min.Y * (GraphScreenSize.Max.Y - GraphScreenSize.Min.Y)) + GraphScreenSize.Min.Y;
@@ -382,7 +397,7 @@ void UReporterGraph::DrawData(UCanvas* Canvas)
 				Position2 = Start;
 				DrawTriangle(Canvas, Position0, Position1, Position2, CurrentData[i].Color);
 
-				
+
 				// draw the second triangle of the quad
 				Position0.X = Start.X;
 				Position0.Y = (GraphMinMaxData.Min.Y * (GraphScreenSize.Max.Y - GraphScreenSize.Min.Y)) + GraphScreenSize.Min.Y;
@@ -416,3 +431,13 @@ FVector2D UReporterGraph::DataToNormalized(const FVector2D& InVector)
 	
 	return OutVector;
 }
+					   
+UFont* UReporterGraph::GetDefaultFont()
+{
+	if (bUseTinyFont)
+	{
+		return GEngine->GetTinyFont();
+	}
+
+	return GEngine->GetSmallFont();
+}
