#UE4

- Mac Fix?

diff --git a/Engine/Source/Runtime/Engine/Classes/Engine/NetSerialization.h b/Engine/Source/Runtime/Engine/Classes/Engine/NetSerialization.h
--- a/Engine/Source/Runtime/Engine/Classes/Engine/NetSerialization.h
+++ b/Engine/Source/Runtime/Engine/Classes/Engine/NetSerialization.h
@@ -370,7 +370,7 @@ struct FFastArraySerializer
 template< typename Type, typename SerializerType >
 bool FFastArraySerializer::FastArrayDeltaSerialize( TArray<Type> &Items, FNetDeltaSerializeInfo& Parms, SerializerType& ArraySerializer )
 {
-	UScriptStruct* InnerStruct = Type::StaticStruct();
+	class UScriptStruct* InnerStruct = Type::StaticStruct();
 
 	if (Parms.OutBunch)
 	{
@@ -510,7 +510,8 @@ bool FFastArraySerializer::FastArrayDeltaSerialize( TArray<Type> &Items, FNetDel
 				// Set to 0 to mean 'unmapped'. This will force reserialization on next update.
 				*NewMap.Find(ID) = 0;
 				NewState->ArrayReplicationKey = INDEX_NONE;
-				UE_LOG( LogNetSerialization, Log, TEXT("   Property: %s is unmapped. Will reserialize."), *InnerStruct->GetName() );
+				UStruct* StructPtr = InnerStruct;
+				UE_LOG(LogNetSerialization, Log, TEXT("   Property: %s is unmapped. Will reserialize."), *StructPtr->GetName());
 			}
 		}
 
