disabled 1 pixel black border by default
can show up in screenshots
if some postprocess has issues we should fix those.

diff --git a/Engine/Source/Runtime/Engine/Private/GameViewportClient.cpp b/Engine/Source/Runtime/Engine/Private/GameViewportClient.cpp
--- a/Engine/Source/Runtime/Engine/Private/GameViewportClient.cpp
+++ b/Engine/Source/Runtime/Engine/Private/GameViewportClient.cpp
@@ -47,7 +47,7 @@ DEFINE_STAT(STAT_UIDrawingTime);
 
 static TAutoConsoleVariable<int32> CVarSetBlackBordersEnabled(
 	TEXT("r.BlackBorders"),
-	1,
+	0,
 	TEXT("To draw black borders around the rendered image\n")
 	TEXT("(prevents artifacts from post processing passes that read outside of the image e.g. PostProcessAA)\n")
 	TEXT("in pixels, 0:off"),
