#docs - Fix comment for UMovementComponent::AddInputVector(). [REDACTED_HANDLE] was before [REDACTED_HANDLE] and [REDACTED_HANDLE] data so the latter were all stripped from editor tooltips.

diff --git a/Engine/Source/Runtime/Engine/Classes/GameFramework/PawnMovementComponent.h b/Engine/Source/Runtime/Engine/Classes/GameFramework/PawnMovementComponent.h
--- a/Engine/Source/Runtime/Engine/Classes/GameFramework/PawnMovementComponent.h
+++ b/Engine/Source/Runtime/Engine/Classes/GameFramework/PawnMovementComponent.h
@@ -26,11 +26,11 @@ public:
 	/**
 	 * Adds the given vector to the accumulated input in world space. Input vectors are usually between 0 and 1 in magnitude. 
 	 * They are accumulated during a frame then applied as acceleration during the movement update.
-	 * @see APawn::AddMovementInput()
 	 *
-	 * @param WorldDirection:	Direction in world space to apply input
-	 * @param ScaleValue:		Scale to apply to input. This can be used for analog input, ie a value of 0.5 applies half the normal value.
-	 * @param bForce:			If true always add the input, ignoring the result of IsMoveInputIgnored().
+	 * @param WorldDirection	Direction in world space to apply input
+	 * @param ScaleValue		Scale to apply to input. This can be used for analog input, ie a value of 0.5 applies half the normal value.
+	 * @param bForce			If true always add the input, ignoring the result of IsMoveInputIgnored().
+	 * @see APawn::AddMovementInput()
 	 */
 	UFUNCTION(BlueprintCallable, Category="Pawn|Components|PawnMovement")
 	virtual void AddInputVector(FVector WorldVector, bool bForce = false);
