Fix the Mac compile errors by including what is used - don't depend on UHT generating a particular include order, it is different on each platform.

diff --git a/Engine/Source/Developer/AITestSuite/Private/Tests/AITestsCommon.cpp b/Engine/Source/Developer/AITestSuite/Private/Tests/AITestsCommon.cpp
--- a/Engine/Source/Developer/AITestSuite/Private/Tests/AITestsCommon.cpp
+++ b/Engine/Source/Developer/AITestSuite/Private/Tests/AITestsCommon.cpp
@@ -1,6 +1,8 @@
 // Copyright 1998-2014 Epic Games, Inc. All Rights Reserved.
 
 #include "AITestSuitePrivatePCH.h"
+#include "EngineGlobals.h"
+#include "Engine.h"
 
 bool FAITestCommand_WaitSeconds::Update()
 {
diff --git a/Engine/Source/Runtime/AIModule/Classes/Actions/PawnActionsComponent.h b/Engine/Source/Runtime/AIModule/Classes/Actions/PawnActionsComponent.h
--- a/Engine/Source/Runtime/AIModule/Classes/Actions/PawnActionsComponent.h
+++ b/Engine/Source/Runtime/AIModule/Classes/Actions/PawnActionsComponent.h
@@ -2,9 +2,9 @@
 
 #pragma once
 #include "Components/ActorComponent.h"
+#include "GameFramework/Pawn.h"
 #include "PawnActionsComponent.generated.h"
 
-class APawn;
 class UPawnAction;
 
 USTRUCT()
diff --git a/Engine/Source/Runtime/AIModule/Classes/Perception/AIPerceptionSystem.h b/Engine/Source/Runtime/AIModule/Classes/Perception/AIPerceptionSystem.h
--- a/Engine/Source/Runtime/AIModule/Classes/Perception/AIPerceptionSystem.h
+++ b/Engine/Source/Runtime/AIModule/Classes/Perception/AIPerceptionSystem.h
@@ -3,6 +3,7 @@
 #pragma once
 
 #include "AIPerceptionTypes.h"
+#include "Kismet/BlueprintFunctionLibrary.h"
 #include "AIPerceptionSystem.generated.h"
 
 DECLARE_LOG_CATEGORY_EXTERN(LogAIPerception, Warning, All);
diff --git a/Engine/Source/Runtime/Engine/Classes/Engine/NetSerialization.h b/Engine/Source/Runtime/Engine/Classes/Engine/NetSerialization.h
--- a/Engine/Source/Runtime/Engine/Classes/Engine/NetSerialization.h
+++ b/Engine/Source/Runtime/Engine/Classes/Engine/NetSerialization.h
@@ -6,8 +6,10 @@
 =============================================================================*/
 
 #pragma once
-#include "NetSerialization.generated.h"
 
+#include "CoreNet.h"
+#include "CoreUObject.h"
+#include "NetSerialization.generated.h"
 
 /**
  *	===================== NetSerialize and NetDeltaSerialize customization. =====================
