Multiprocess PIE(viewport) client can now connect to a server

#ttp 335628 - +CRITICAL: Regression: SHOOTERGAME: PIE as client only opens one player window

#branch UE4

#change Modified the clients PIEInstance (in this instance) to be -1, so that it uses normal assets (those determined by the server) and not the PIE ones otherwise the client will get kicked from the server. Made a slight change to prevent a call to BuildPIEPackagePrefix from crashing too.

#change Modified AsyncLoading to be able to specify the PackageFlags applied to the package, rather than assuming them based off the PIEInstanceID. LevelStreaming (which is the only place which utilised this) can now check the WorldContext to determine if the PKG_PlayInEditor PackageFlag should be set, rather than PIEInstanceID - which brings the code inline with the rest of the solution. Streaming levels only get unloaded if the package has this flag set (unloading doesn't care about PIEInstanceID, so it was odd that the flag was set based on this in the first place). Without this change, PIE clients would keep streaming levels around in memory which would break things later on.

diff --git a/Engine/Source/Editor/UnrealEd/Private/PlayLevel.cpp b/Engine/Source/Editor/UnrealEd/Private/PlayLevel.cpp
--- a/Engine/Source/Editor/UnrealEd/Private/PlayLevel.cpp
+++ b/Engine/Source/Editor/UnrealEd/Private/PlayLevel.cpp
@@ -1935,8 +1935,9 @@ void UEditorEngine::PlayInEditor( UWorld* InWorld, bool bInSimulateInEditor )
 
 	if (bInSimulateInEditor || (PlayInSettings->PlayNetMode == EPlayNetMode::PIE_Standalone && !bSupportsOnlinePIE) || !PlayInSettings->RunUnderOneProcess)
 	{
-		// Only spawning 1 PIE instance under this process
-		UGameInstance* const GameInstance = CreatePIEGameInstance(0, bInSimulateInEditor, bAnyBlueprintErrors, bStartInSpectatorMode, false, PIEStartTime);
+		// Only spawning 1 PIE instance under this process, only set the PIEInstance value if we're not connecting to another local instance of the game, otherwise it will run the wrong streaming levels
+		const int32 PIEInstance = ( !PlayInSettings->RunUnderOneProcess && PlayInSettings->PlayNetMode == EPlayNetMode::PIE_Client ) ? INDEX_NONE : 0;
+		UGameInstance* const GameInstance = CreatePIEGameInstance(PIEInstance, bInSimulateInEditor, bAnyBlueprintErrors, bStartInSpectatorMode, false, PIEStartTime);
 
 		if (bInSimulateInEditor)
 		{
diff --git a/Engine/Source/Runtime/CoreUObject/Private/Serialization/AsyncLoading.cpp b/Engine/Source/Runtime/CoreUObject/Private/Serialization/AsyncLoading.cpp
--- a/Engine/Source/Runtime/CoreUObject/Private/Serialization/AsyncLoading.cpp
+++ b/Engine/Source/Runtime/CoreUObject/Private/Serialization/AsyncLoading.cpp
@@ -296,13 +296,10 @@ EAsyncPackageState::Type FAsyncPackage::CreateLinker()
 		UPackage* Package = CreatePackage( NULL, *PackageName );
 		FScopeCycleCounterUObject ConstructorScope(Package, GET_STATID(STAT_FAsyncPackage_CreateLinker));
 
-		// Mark PIE package
+		// Set package specific data 
+		Package->PackageFlags |= PackageFlags;
 #if WITH_EDITOR
-		if (PIEInstanceID != INDEX_NONE)
-		{
-			Package->PackageFlags |= PKG_PlayInEditor;
-			Package->PIEInstanceID = PIEInstanceID;
-		}
+		Package->PIEInstanceID = PIEInstanceID;
 #endif
 
 		// Always store package filename we loading from
diff --git a/Engine/Source/Runtime/CoreUObject/Public/Serialization/AsyncLoading.h b/Engine/Source/Runtime/CoreUObject/Public/Serialization/AsyncLoading.h
--- a/Engine/Source/Runtime/CoreUObject/Public/Serialization/AsyncLoading.h
+++ b/Engine/Source/Runtime/CoreUObject/Public/Serialization/AsyncLoading.h
@@ -36,6 +36,7 @@ struct FAsyncPackage : public FGCObject
 	,	LastTypeOfWorkPerformed		( NULL					)
 	,	LoadStartTime				( 0.0					)
 	,	LoadPercentage				( 0						)
+	,	PackageFlags				( 0						)
 #if WITH_EDITOR
 	,	PIEInstanceID				( INDEX_NONE			)
 #endif	
@@ -123,11 +124,11 @@ struct FAsyncPackage : public FGCObject
 	}
 
 	/**
-	 *	Sets PIE instance ID this package belongs to, INDEX_NONE otherwise
-	 *  Will also mark loading package as PKG_PlayInEditor
+	 *	Sets the flags and PIE ID for the package
 	 */
-	void SetPIEInstanceID(int32 InPIEInstanceID)
+	void SetPackageData(uint32 InPackageFlags, int32 InPIEInstanceID)
 	{
+		PackageFlags = InPackageFlags;
 #if WITH_EDITOR
 		PIEInstanceID = InPIEInstanceID;
 #endif
@@ -196,7 +197,8 @@ private:
 	double						LoadStartTime;
 	/** Estimated load percentage.																		*/
 	float						LoadPercentage;
-	
+	/** The flags that should be applied to the package */
+	uint32						PackageFlags;
 #if WITH_EDITOR
 	/** Editor only: PIE instance ID this package belongs to, INDEX_NONE otherwise */
 	int32						PIEInstanceID;
diff --git a/Engine/Source/Runtime/Engine/Private/GameInstance.cpp b/Engine/Source/Runtime/Engine/Private/GameInstance.cpp
--- a/Engine/Source/Runtime/Engine/Private/GameInstance.cpp
+++ b/Engine/Source/Runtime/Engine/Private/GameInstance.cpp
@@ -59,7 +59,7 @@ bool UGameInstance::InitPIE(bool bAnyBlueprintErrors, int32 PIEInstance)
 
 	// Establish World Context for PIE World
 	WorldContext->LastURL.Map = WorldPackageName;
-	WorldContext->PIEPrefix = UWorld::BuildPIEPackagePrefix(PIEInstance);
+	WorldContext->PIEPrefix = WorldContext->PIEInstance != INDEX_NONE ? UWorld::BuildPIEPackagePrefix(WorldContext->PIEInstance) : FString();
 
 	const ULevelEditorPlaySettings* PlayInSettings = GetDefault<ULevelEditorPlaySettings>();
 
diff --git a/Engine/Source/Runtime/Engine/Private/LevelStreaming.cpp b/Engine/Source/Runtime/Engine/Private/LevelStreaming.cpp
--- a/Engine/Source/Runtime/Engine/Private/LevelStreaming.cpp
+++ b/Engine/Source/Runtime/Engine/Private/LevelStreaming.cpp
@@ -386,6 +386,7 @@ bool ULevelStreaming::RequestLevel(UWorld* PersistentWorld, bool bAllowLevelLoad
 		}
 	}
 
+	uint32 PackageFlags = 0;
 	int32 PIEInstanceID = INDEX_NONE;
 
 	// copy streaming level on demand if we are in PIE
@@ -393,6 +394,10 @@ bool ULevelStreaming::RequestLevel(UWorld* PersistentWorld, bool bAllowLevelLoad
 	if ( PersistentWorld->IsPlayInEditor() )
 	{
 #if WITH_EDITOR
+		if ((PersistentWorld->GetOutermost()->PackageFlags & PKG_PlayInEditor) != 0)
+		{
+			PackageFlags |= PKG_PlayInEditor;
+		}
 		PIEInstanceID = PersistentWorld->GetOutermost()->PIEInstanceID;
 #endif
 		const FString PrefixedLevelName = DesiredPackageName.ToString();
@@ -450,8 +455,9 @@ bool ULevelStreaming::RequestLevel(UWorld* PersistentWorld, bool bAllowLevelLoad
 			{
 				// Load localized part of level first in case it exists. We don't need to worry about GC or completion 
 				// callback as we always kick off another async IO for the level below.
-				LoadPackageAsync(*(GetWorldAssetPackageName() + LOCALIZED_SEEKFREE_SUFFIX), NULL, NAME_None, *LocalizedPackageName)
-					.SetPIEInstanceID(PIEInstanceID);
+				LoadPackageAsync(*(GetWorldAssetPackageName() + LOCALIZED_SEEKFREE_SUFFIX), 
+					NULL, NAME_None, *LocalizedPackageName
+					).SetPackageData(PackageFlags, PIEInstanceID);
 			}
 		}
 
@@ -466,7 +472,7 @@ bool ULevelStreaming::RequestLevel(UWorld* PersistentWorld, bool bAllowLevelLoad
 			LoadPackageAsync(*DesiredPackageName.ToString(), 
 				FLoadPackageAsyncDelegate::CreateUObject(this, &ULevelStreaming::AsyncLevelLoadComplete), 
 				NULL, NAME_None, *PackageNameToLoadFrom
-				).SetPIEInstanceID(PIEInstanceID);
+				).SetPackageData(PackageFlags, PIEInstanceID);
 
 			// streamingServer: server loads everything?
 			// Editor immediately blocks on load and we also block if background level streaming is disabled.
