Add comments to OwnedComponents related functions

diff --git a/Engine/Source/Runtime/Engine/Classes/GameFramework/Actor.h b/Engine/Source/Runtime/Engine/Classes/GameFramework/Actor.h
--- a/Engine/Source/Runtime/Engine/Classes/GameFramework/Actor.h
+++ b/Engine/Source/Runtime/Engine/Classes/GameFramework/Actor.h
@@ -2086,11 +2086,25 @@ public:
 		}
 	}
 
+	/** Puts a component in to the OwnedComponents array of the Actor.
+	 *  The Component must be owned by the Actor or else it will assert
+	 *  In general this should not need to be called directly by anything other than UActorComponent functions
+	 */
 	void AddOwnedComponent(UActorComponent* Component);
+
+	/** Removes a component from the OwnedComponents array of the Actor.
+	 *  In general this should not need to be called directly by anything other than UActorComponent functions
+	 */
 	void RemoveOwnedComponent(UActorComponent* Component);
+
 #if DO_CHECK
+	// Utility function for validating that a component is correctly in its Owner's OwnedComponents array
 	bool OwnsComponent(UActorComponent* Component) const;
 #endif
+
+	/** Force the Actor to clear and rebuild its OwnedComponents array by evaluating all children (recursively) and locating components
+	 *  In general this should not need to be called directly, but can sometimes be necessary as part of undo/redo code paths.
+	 */
 	void ResetOwnedComponents();
 
 private:
