Assert that only Components can only be added to the OwnedComponents array of their Owner

diff --git a/Engine/Source/Runtime/Engine/Private/Actor.cpp b/Engine/Source/Runtime/Engine/Private/Actor.cpp
--- a/Engine/Source/Runtime/Engine/Private/Actor.cpp
+++ b/Engine/Source/Runtime/Engine/Private/Actor.cpp
@@ -1960,6 +1960,7 @@ enum ECollisionResponse AActor::GetComponentsCollisionResponseToChannel(enum ECo
 
 void AActor::AddOwnedComponent(UActorComponent* Component)
 {
+	check(Component->GetOwner() == this);
 	OwnedComponents.AddUnique(Component);
 }
 
