Added more geometry shapes to visual logger #ue4

diff --git a/Engine/Source/Developer/LogVisualizer/Private/SLogVisualizer.cpp b/Engine/Source/Developer/LogVisualizer/Private/SLogVisualizer.cpp
--- a/Engine/Source/Developer/LogVisualizer/Private/SLogVisualizer.cpp
+++ b/Engine/Source/Developer/LogVisualizer/Private/SLogVisualizer.cpp
@@ -2332,6 +2332,60 @@ void SLogVisualizer::DrawOnCanvas(UCanvas* Canvas, APlayerController*)
 						}
 					}
 					break;
+				case FVisLogEntry::FElementToDraw::Cone:
+					{
+						const float Thickness = float(ElementToDraw->Thicknes);
+						const bool bDrawLabel = ElementToDraw->Description.IsEmpty() == false;
+						for (int32 Index = 0; Index + 2 < ElementToDraw->Points.Num(); Index += 3)
+						{
+							const FVector Orgin = ElementToDraw->Points[Index];
+							const FVector Direction = ElementToDraw->Points[Index + 1];
+							const FVector Angles = ElementToDraw->Points[Index + 2];
+							DrawDebugCone(World, Orgin, Direction, Angles.X, Angles.Y, Angles.Z, 16, Color);
+							if (bDrawLabel)
+							{
+								const FVector ScreenLoc = Canvas->Project(Orgin);
+								Canvas->DrawText(Font, FString::Printf(TEXT("%s_%d"), *ElementToDraw->Description, Index), ScreenLoc.X, ScreenLoc.Y);
+							}
+						}
+					}
+					break;
+				case FVisLogEntry::FElementToDraw::Cylinder:
+					{
+						const float Thickness = float(ElementToDraw->Thicknes);
+						const bool bDrawLabel = ElementToDraw->Description.IsEmpty() == false;
+						for (int32 Index = 0; Index + 2 < ElementToDraw->Points.Num(); Index += 3)
+						{
+							const FVector Start = ElementToDraw->Points[Index];
+							const FVector End = ElementToDraw->Points[Index + 1];
+							const FVector OtherData = ElementToDraw->Points[Index + 2];
+							DrawDebugCylinder(World, Start, End, OtherData.X, 16, Color);
+							if (bDrawLabel)
+							{
+								const FVector ScreenLoc = Canvas->Project(Start);
+								Canvas->DrawText(Font, FString::Printf(TEXT("%s_%d"), *ElementToDraw->Description, Index), ScreenLoc.X, ScreenLoc.Y);
+							}
+						}
+					}
+					break;
+				case FVisLogEntry::FElementToDraw::Capsule:
+					{
+						const float Thickness = float(ElementToDraw->Thicknes);
+						const bool bDrawLabel = ElementToDraw->Description.IsEmpty() == false;
+						for (int32 Index = 0; Index + 2 < ElementToDraw->Points.Num(); Index += 3)
+						{
+							const FVector Center = ElementToDraw->Points[Index];
+							const FVector FirstData = ElementToDraw->Points[Index + 1];
+							const FVector SecondData = ElementToDraw->Points[Index + 2];
+							DrawDebugCapsule(World, Center, FirstData.X, FirstData.Y, FQuat(FirstData.Z, SecondData.X, SecondData.Y, SecondData.Z), Color);
+							if (bDrawLabel)
+							{
+								const FVector ScreenLoc = Canvas->Project(Center);
+								Canvas->DrawText(Font, FString::Printf(TEXT("%s_%d"), *ElementToDraw->Description, Index), ScreenLoc.X, ScreenLoc.Y);
+							}
+						}
+					}
+					break;
 				}
 			}
 		}
diff --git a/Engine/Source/Runtime/Engine/Private/VisualLog.cpp b/Engine/Source/Runtime/Engine/Private/VisualLog.cpp
--- a/Engine/Source/Runtime/Engine/Private/VisualLog.cpp
+++ b/Engine/Source/Runtime/Engine/Private/VisualLog.cpp
@@ -371,6 +371,28 @@ void FVisLogEntry::AddElement(const FBox& Box, const FName& CategoryName, const
 	ElementsToDraw.Add(Element);
 }
 
+void FVisLogEntry::AddElement(const FVector& Orgin, const FVector& Direction, float Length, float AngleWidth, float AngleHeight, const FName& CategoryName, const FColor& Color, const FString& Description, uint16 Thickness)
+{
+	FElementToDraw Element(Description, Color, Thickness, CategoryName);
+	Element.Points.Reserve(3);
+	Element.Points.Add(Orgin);
+	Element.Points.Add(Direction);
+	Element.Points.Add(FVector(Length, AngleWidth, AngleHeight));
+	Element.Type = FElementToDraw::Cone;
+	ElementsToDraw.Add(Element);
+}
+
+void FVisLogEntry::AddElement(const FVector& Start, const FVector& End, float Radius, const FName& CategoryName, const FColor& Color, const FString& Description, uint16 Thickness)
+{
+	FElementToDraw Element(Description, Color, Thickness, CategoryName);
+	Element.Points.Reserve(3);
+	Element.Points.Add(Start);
+	Element.Points.Add(End);
+	Element.Points.Add(FVector(Radius, Thickness, 0));
+	Element.Type = FElementToDraw::Cylinder;
+	ElementsToDraw.Add(Element);
+}
+
 void FVisLogEntry::AddHistogramData(const FVector2D& DataSample, const FName& CategoryName, const FName& GraphName, const FName& DataName)
 {
 	FHistogramSample Sample;
@@ -392,6 +414,17 @@ void FVisLogEntry::AddDataBlock(const FString& TagName, const TArray<uint8>& Blo
 	DataBlocks.Add(DataBlock);
 }
 
+void FVisLogEntry::AddCapsule(FVector const& Center, float HalfHeight, float Radius, const FQuat & Rotation, const FName& CategoryName, const FColor& Color, const FString& Description)
+{
+	FElementToDraw Element(Description, Color, 0, CategoryName);
+	Element.Points.Reserve(3);
+	Element.Points.Add(Center);
+	Element.Points.Add(FVector(HalfHeight, Radius, Rotation.X));
+	Element.Points.Add(FVector(Rotation.Y, Rotation.Z, Rotation.W));
+	Element.Type = FElementToDraw::Capsule;
+	ElementsToDraw.Add(Element);
+}
+
 //----------------------------------------------------------------------//
 // FActorsVisLog
 //----------------------------------------------------------------------//
diff --git a/Engine/Source/Runtime/Engine/Public/VisualLog.h b/Engine/Source/Runtime/Engine/Public/VisualLog.h
--- a/Engine/Source/Runtime/Engine/Public/VisualLog.h
+++ b/Engine/Source/Runtime/Engine/Public/VisualLog.h
@@ -81,6 +81,9 @@ struct ENGINE_API FVisLogEntry
 			Segment, // pairs of points 
 			Path,	// sequence of point
 			Box,
+			Cone,
+			Cylinder,
+			Capsule,
 			// note that in order to remain backward compatibility in terms of log
 			// serialization new enum values need to be added at the end
 		};
@@ -181,12 +184,21 @@ struct ENGINE_API FVisLogEntry
 	// box
 	void AddElement(const FBox& Box, const FName& CategoryName, const FColor& Color = FColor::White, const FString& Description = TEXT(""), uint16 Thickness = 0);
 
+	// Cone
+	void AddElement(const FVector& Orgin, const FVector& Direction, float Length, float AngleWidth, float AngleHeight, const FName& CategoryName, const FColor& Color = FColor::White, const FString& Description = TEXT(""), uint16 Thickness = 0);
+
+	// Cylinder
+	void AddElement(const FVector& Start, const FVector& End, float Radius, const FName& CategoryName, const FColor& Color = FColor::White, const FString& Description = TEXT(""), uint16 Thickness = 0);
+
 	// histogram sample
 	void AddHistogramData(const FVector2D& DataSample, const FName& CategoryName, const FName& GraphName, const FName& DataName);
 
 	// Custom data block
 	void AddDataBlock(const FString& TagName, const TArray<uint8>& BlobDataArray, const FName& CategoryName);
 
+	// capsule
+	void AddCapsule(FVector const& Center, float HalfHeight, float Radius, const FQuat & Rotation, const FName& CategoryName, const FColor& Color = FColor::White, const FString& Description = TEXT(""));
+
 	// find index of status category
 	FORCEINLINE int32 FindStatusIndex(const FString& CategoryName)
 	{
@@ -272,6 +284,36 @@ struct ENGINE_API FVisLogEntry
 	} \
 }
 
+#define UE_VLOG_CONE(Object, CategoryName, Verbosity, Orgin, Direction, Length, Angle, Color, DescriptionFormat, ...) \
+{ \
+	SCOPE_CYCLE_COUNTER(STAT_VisualLog); \
+	static_assert((ELogVerbosity::Verbosity & ELogVerbosity::VerbosityMask) < ELogVerbosity::NumVerbosity && ELogVerbosity::Verbosity > 0, "Verbosity must be constant and in range."); \
+	if (FVisualLog::Get().IsRecording() && (!FVisualLog::Get().IsAllBlocked() || FVisualLog::Get().InWhitelist(CategoryName.GetCategoryName())) && Object && !Object->HasAnyFlags(RF_ClassDefaultObject)) \
+			{ \
+		FVisualLog::Get().GetEntryToWrite(Object)->AddElement(Orgin, Direction, Length, Angle, Angle, CategoryName.GetCategoryName(), Color, FString::Printf(DescriptionFormat, ##__VA_ARGS__)); \
+			} \
+}
+
+#define UE_VLOG_CYLINDER(Object, CategoryName, Verbosity, Start, End, Radius, Color, DescriptionFormat, ...) \
+{ \
+	SCOPE_CYCLE_COUNTER(STAT_VisualLog); \
+	static_assert((ELogVerbosity::Verbosity & ELogVerbosity::VerbosityMask) < ELogVerbosity::NumVerbosity && ELogVerbosity::Verbosity > 0, "Verbosity must be constant and in range."); \
+	if (FVisualLog::Get().IsRecording() && (!FVisualLog::Get().IsAllBlocked() || FVisualLog::Get().InWhitelist(CategoryName.GetCategoryName())) && Object && !Object->HasAnyFlags(RF_ClassDefaultObject)) \
+	{ \
+		FVisualLog::Get().GetEntryToWrite(Object)->AddElement(Start, End, Radius, CategoryName.GetCategoryName(), Color, FString::Printf(DescriptionFormat, ##__VA_ARGS__)); \
+	} \
+}
+
+#define UE_VLOG_CAPSULE(Object, CategoryName, Verbosity, Center, HalfHeight, Radius, Rotation, Color, DescriptionFormat, ...) \
+{ \
+	SCOPE_CYCLE_COUNTER(STAT_VisualLog); \
+	static_assert((ELogVerbosity::Verbosity & ELogVerbosity::VerbosityMask) < ELogVerbosity::NumVerbosity && ELogVerbosity::Verbosity > 0, "Verbosity must be constant and in range."); \
+	if (FVisualLog::Get().IsRecording() && (!FVisualLog::Get().IsAllBlocked() || FVisualLog::Get().InWhitelist(CategoryName.GetCategoryName())) && Object && !Object->HasAnyFlags(RF_ClassDefaultObject)) \
+					{ \
+		FVisualLog::Get().GetEntryToWrite(Object)->AddCapsule(Center, HalfHeight, Radius, Rotation, CategoryName.GetCategoryName(), Color, FString::Printf(DescriptionFormat, ##__VA_ARGS__)); \
+					} \
+}
+
 #define UE_VLOG_HISTOGRAM(Object, CategoryName, Verbosity, GraphName, DataName, Data) \
 { \
 	SCOPE_CYCLE_COUNTER(STAT_VisualLog); \
@@ -301,6 +343,7 @@ namespace VisualLogJson
 	static const FString TAG_ELEMENTSTODRAW = TEXT("ElementsToDraw");
 	static const FString TAG_TAGNAME = TEXT("DataBlockTagName");
 	static const FString TAG_USERDATA = TEXT("UserData");
+	static const FString TAG_CAPSULESTODRAW = TEXT("CapsulesToDraw");
 
 	static const FString TAG_HISTOGRAMSAMPLES = TEXT("HistogramSamples");
 	static const FString TAG_HISTOGRAMSAMPLE = TEXT("Sample");
