Do not need to manually AddOwnedComponents, is handled as part of component construction

diff --git a/Engine/Source/Runtime/Engine/Private/Landscape/LandscapeSplines.cpp b/Engine/Source/Runtime/Engine/Private/Landscape/LandscapeSplines.cpp
--- a/Engine/Source/Runtime/Engine/Private/Landscape/LandscapeSplines.cpp
+++ b/Engine/Source/Runtime/Engine/Private/Landscape/LandscapeSplines.cpp
@@ -910,7 +910,6 @@ void ULandscapeSplineControlPoint::UpdateSplinePoints(bool bUpdateCollision, boo
 				Modify();
 				OuterSplines->GetOuter()->Modify();
 				MeshComponent = ConstructObject<UControlPointMeshComponent>(UControlPointMeshComponent::StaticClass(), OuterSplines->GetOuter(), NAME_None, RF_Transactional | RF_TextExportTransient);
-				OuterSplines->GetOwner()->AddOwnedComponent(MeshComponent);
 				MeshComponent->bSelected = bSelected;
 				MeshComponent->AttachTo(OuterSplines);
 			}
@@ -1513,7 +1512,6 @@ void ULandscapeSplineSegment::UpdateSplinePoints(bool bUpdateCollision)
 			{
 				OuterSplines->GetOuter()->Modify();
 				MeshComponent = ConstructObject<USplineMeshComponent>(USplineMeshComponent::StaticClass(), OuterSplines->GetOuter(), NAME_None, RF_Transactional|RF_TextExportTransient);
-				OuterSplines->GetOwner()->AddOwnedComponent(MeshComponent);
 				MeshComponent->bSelected = bSelected;
 				MeshComponent->AttachTo(OuterSplines);
 				MeshComponents.Add(MeshComponent);
