Fixed compilation issue  when compiling with disabled visual logger

diff --git a/Engine/Source/Runtime/Engine/Public/VisualLog.h b/Engine/Source/Runtime/Engine/Public/VisualLog.h
--- a/Engine/Source/Runtime/Engine/Public/VisualLog.h
+++ b/Engine/Source/Runtime/Engine/Public/VisualLog.h
@@ -484,6 +484,10 @@ private:
 	#define UE_VLOG_LOCATION(Actor, CategoryName, Verbosity, Location, Radius, Color, DescriptionFormat, ...)
 	#define UE_VLOG_BOX(Actor, CategoryName, Verbosity, Box, Color, DescriptionFormat, ...) 
 	#define UE_VLOG_HISTOGRAM(Actor, CategoryName, Verbosity, GraphName, DataName, Data)
+	#define UE_VLOG_CONE(Object, CategoryName, Verbosity, Orgin, Direction, Length, Angle, Color, DescriptionFormat, ...)
+	#define UE_VLOG_CYLINDER(Object, CategoryName, Verbosity, Start, End, Radius, Color, DescriptionFormat, ...)
+	#define UE_VLOG_CAPSULE(Object, CategoryName, Verbosity, Center, HalfHeight, Radius, Rotation, Color, DescriptionFormat, ...)
+
 	#define REDIRECT_TO_VLOG(Dest)
 	#define REDIRECT_OBJECT_TO_VLOG(Src, Destination) 
 	#define CONNECT_WITH_VLOG(Dest)
