Giving documentation/comment nodes a icon for the action menu system.

diff --git a/Engine/Source/Runtime/Engine/Classes/EdGraph/EdGraphNode_Comment.h b/Engine/Source/Runtime/Engine/Classes/EdGraph/EdGraphNode_Comment.h
--- a/Engine/Source/Runtime/Engine/Classes/EdGraph/EdGraphNode_Comment.h
+++ b/Engine/Source/Runtime/Engine/Classes/EdGraph/EdGraphNode_Comment.h
@@ -58,6 +58,7 @@ public:
 	virtual bool ShouldDrawNodeAsComment() const override { return true; }
 	ENGINE_API virtual FString GetDocumentationLink() const override;
 	ENGINE_API virtual FString GetDocumentationExcerptName() const override;
+	ENGINE_API virtual FName GetPaletteIcon(FLinearColor& OutColor) const override;
 	// End UEdGraphNode interface
 
 	/** Add a node that will be dragged when this comment is dragged */
diff --git a/Engine/Source/Runtime/Engine/Classes/EdGraph/EdGraphNode_Documentation.h b/Engine/Source/Runtime/Engine/Classes/EdGraph/EdGraphNode_Documentation.h
--- a/Engine/Source/Runtime/Engine/Classes/EdGraph/EdGraphNode_Documentation.h
+++ b/Engine/Source/Runtime/Engine/Classes/EdGraph/EdGraphNode_Documentation.h
@@ -32,6 +32,7 @@ public:
 	ENGINE_API virtual TSharedPtr<class INameValidatorInterface> MakeNameValidator() const override;
 	ENGINE_API virtual FString GetDocumentationLink() const override { return Link; }
 	ENGINE_API virtual FString GetDocumentationExcerptName() const override { return Excerpt; }
+	ENGINE_API virtual FName GetPaletteIcon(FLinearColor& OutColor) const override { return TEXT("GraphEditor.Documentation_16x");  }
 	// End UEdGraphNode interface
 
 	/** Set the Bounds for the comment node */
diff --git a/Engine/Source/Runtime/Engine/Private/EdGraph/EdGraphNode_Comment.cpp b/Engine/Source/Runtime/Engine/Private/EdGraph/EdGraphNode_Comment.cpp
--- a/Engine/Source/Runtime/Engine/Private/EdGraph/EdGraphNode_Comment.cpp
+++ b/Engine/Source/Runtime/Engine/Private/EdGraph/EdGraphNode_Comment.cpp
@@ -74,6 +74,12 @@ FString UEdGraphNode_Comment::GetDocumentationExcerptName() const
 	return TEXT("UEdGraphNode_Comment");
 }
 
+FName UEdGraphNode_Comment::GetPaletteIcon(FLinearColor& OutColor) const
+{
+	OutColor = FLinearColor::White;
+	return TEXT("GraphEditor.Comment_16x");
+}
+
 FText UEdGraphNode_Comment::GetNodeTitle(ENodeTitleType::Type TitleType) const
 {
 	if(TitleType == ENodeTitleType::MenuTitle)
