Remove unused code

diff --git a/Engine/Source/Runtime/Engine/Private/UnrealEngine.cpp b/Engine/Source/Runtime/Engine/Private/UnrealEngine.cpp
--- a/Engine/Source/Runtime/Engine/Private/UnrealEngine.cpp
+++ b/Engine/Source/Runtime/Engine/Private/UnrealEngine.cpp
@@ -8509,13 +8509,6 @@ bool UEngine::LoadMap( FWorldContext& WorldContext, FURL URL, class UPendingNetG
 	UE_LOG(LogLoad, Log,  TEXT("LoadMap: %s"), *URL.ToString() );
 	GInitRunaway();
 
-	// Get network package map.
-	UPackageMap* PackageMap = NULL;
-	if( Pending )
-	{
-		PackageMap = Pending->GetNetDriver()->ServerConnection->PackageMap;
-	}
-
 	// Unload the current world
 	if( WorldContext.World() )
 	{
