Replicate demo spectator so that client can use for spectating, and spawn in a good location

diff --git a/Engine/Source/Runtime/Engine/Private/DemoNetDriver.cpp b/Engine/Source/Runtime/Engine/Private/DemoNetDriver.cpp
--- a/Engine/Source/Runtime/Engine/Private/DemoNetDriver.cpp
+++ b/Engine/Source/Runtime/Engine/Private/DemoNetDriver.cpp
@@ -509,6 +509,7 @@ void UDemoNetDriver::TickDemoPlayback( float DeltaSeconds )
 void UDemoNetDriver::SpawnDemoRecSpectator( UNetConnection* Connection )
 {
 	check( Connection != NULL );
+
 	UClass* C = StaticLoadClass( AActor::StaticClass(), NULL, *DemoSpectatorClass, NULL, LOAD_None, NULL );
 	APlayerController* Controller = World->SpawnActor<APlayerController>( C );
 
@@ -516,12 +517,13 @@ void UDemoNetDriver::SpawnDemoRecSpectator( UNetConnection* Connection )
 	{
 		if ( It->IsA( APlayerStart::StaticClass() ) )
 		{
-			Controller->SetInitialLocationAndRotation( It->GetActorLocation() + FVector( 0, 0, 1000 ), It->GetActorRotation() );
+			Controller->SetInitialLocationAndRotation( It->GetActorLocation(), It->GetActorRotation() );
 			break;
 		}
 	}
 
-	//Controller->SetReplicates( true );
+	Controller->SetReplicates( true );
+	Controller->SetAutonomousProxy( true );
 
 	Controller->SetPlayer( Connection );
 }
@@ -594,7 +596,20 @@ void UDemoNetConnection::FlushNet( bool bIgnoreSimulation )
 void UDemoNetConnection::HandleClientPlayer( APlayerController* PC, UNetConnection* NetConnection )
 {
 	Super::HandleClientPlayer( PC, NetConnection );
-	//PC->bDemoOwner = true;	// DEMO_FIXME
+
+	for ( FActorIterator It( Driver->World ); It; ++It)
+	{
+		if ( It->IsA( APlayerStart::StaticClass() ) )
+		{
+			PC->SetInitialLocationAndRotation( It->GetActorLocation(), It->GetActorRotation() );
+
+			if ( PC->GetPawn() )
+			{
+				PC->GetPawn()->TeleportTo( It->GetActorLocation(), It->GetActorRotation(), false, true );
+			}
+			break;
+		}
+	}
 }
 
 bool UDemoNetConnection::ClientHasInitializedLevelFor(const UObject* TestObject)
