Engine: UPrimitiveComponent shadow property clarifications
- Add EditCondition=CastShadow to various advanced shadowing flags
- Add an edit condition on CastShadow dependent on having at least one lit material assigned to the component
- Clarified that CastShadow is ignored on components that have no non-Unlit materials

diff --git a/Engine/Source/Runtime/Engine/Classes/Components/PrimitiveComponent.h b/Engine/Source/Runtime/Engine/Classes/Components/PrimitiveComponent.h
--- a/Engine/Source/Runtime/Engine/Classes/Components/PrimitiveComponent.h
+++ b/Engine/Source/Runtime/Engine/Classes/Components/PrimitiveComponent.h
@@ -255,20 +255,24 @@ public:
 
 	// Lighting flags
 	
-	/** Controls whether the primitive component should cast a shadow or not. **/
+	/**
+	 * Controls whether the primitive component should cast a shadow or not.
+	 *
+	 * This flag is ignored (no shadows will be generated) if all materials on this component have an Unlit shading model.
+	 */
 	UPROPERTY(EditAnywhere, BlueprintReadOnly, Category=Lighting)
 	uint32 CastShadow:1;
 
 	/** Controls whether the primitive should inject light into the Light Propagation Volume.  This flag is only used if CastShadow is true. **/
-	UPROPERTY(EditAnywhere, BlueprintReadOnly, Category=Lighting, AdvancedDisplay)
+	UPROPERTY(EditAnywhere, BlueprintReadOnly, Category=Lighting, AdvancedDisplay, meta=(EditCondition="CastShadow"))
 	uint32 bAffectDynamicIndirectLighting:1;
 
 	/** Controls whether the primitive should cast shadows in the case of non precomputed shadowing.  This flag is only used if CastShadow is true. **/
-	UPROPERTY(EditAnywhere, BlueprintReadOnly, Category=Lighting, AdvancedDisplay)
+	UPROPERTY(EditAnywhere, BlueprintReadOnly, Category=Lighting, AdvancedDisplay, meta=(EditCondition="CastShadow"))
 	uint32 bCastDynamicShadow:1;
 
 	/** Whether the object should cast a static shadow from shadow casting lights.  This flag is only used if CastShadow is true. */
-	UPROPERTY(EditAnywhere, BlueprintReadOnly, Category=Lighting, AdvancedDisplay)
+	UPROPERTY(EditAnywhere, BlueprintReadOnly, Category=Lighting, AdvancedDisplay, meta=(EditCondition="CastShadow"))
 	uint32 bCastStaticShadow:1;
 
 	/** 
@@ -276,14 +280,14 @@ public:
 	 * Volumetric translucent shadows are useful for primitives with smoothly changing opacity like particles representing a volume, 
 	 * But have artifacts when used on highly opaque surfaces.
 	 */
-	UPROPERTY(EditAnywhere, AdvancedDisplay, BlueprintReadOnly, Category=Lighting)
+	UPROPERTY(EditAnywhere, AdvancedDisplay, BlueprintReadOnly, Category=Lighting, meta=(EditCondition="CastShadow"))
 	uint32 bCastVolumetricTranslucentShadow:1;
 
 	/** 
 	 * Whether this component should create a per-object shadow that gives higher effective shadow resolution. 
 	 * Useful for cinematic character shadowing.
 	 */
-	UPROPERTY(EditAnywhere, AdvancedDisplay, BlueprintReadOnly, Category=Lighting)
+	UPROPERTY(EditAnywhere, AdvancedDisplay, BlueprintReadOnly, Category=Lighting, meta=(EditCondition="CastShadow"))
 	uint32 bCastInsetShadow:1;
 
 	/** 
@@ -291,11 +295,11 @@ public:
 	 *	Controls whether the primitive should cast shadows when hidden.
 	 *	This flag is only used if CastShadow is true.
 	 */
-	UPROPERTY(EditAnywhere, AdvancedDisplay, BlueprintReadOnly, Category=Lighting)
+	UPROPERTY(EditAnywhere, AdvancedDisplay, BlueprintReadOnly, Category=Lighting, meta=(EditCondition="CastShadow"))
 	uint32 bCastHiddenShadow:1;
 
 	/** Whether this primitive should cast dynamic shadows as if it were a two sided material. */
-	UPROPERTY(EditAnywhere, AdvancedDisplay, BlueprintReadOnly, Category=Lighting)
+	UPROPERTY(EditAnywhere, AdvancedDisplay, BlueprintReadOnly, Category=Lighting, meta=(EditCondition="CastShadow"))
 	uint32 bCastShadowAsTwoSided:1;
 
 	/** 
diff --git a/Engine/Source/Runtime/Engine/Private/PrimitiveComponent.cpp b/Engine/Source/Runtime/Engine/Private/PrimitiveComponent.cpp
--- a/Engine/Source/Runtime/Engine/Private/PrimitiveComponent.cpp
+++ b/Engine/Source/Runtime/Engine/Private/PrimitiveComponent.cpp
@@ -524,9 +524,9 @@ bool UPrimitiveComponent::CanEditChange(const UProperty* InProperty) const
 	bool bIsEditable = Super::CanEditChange( InProperty );
 	if (bIsEditable && InProperty)
 	{
-		const FString PropertyName = InProperty->GetName();
+		const FName PropertyName = InProperty->GetFName();
 
-		if (PropertyName == TEXT("bLightAsIfStatic"))
+		if (PropertyName == GET_MEMBER_NAME_CHECKED(UPrimitiveComponent, bLightAsIfStatic))
 		{
 			// Disable editing bLightAsIfStatic on static components, since it has no effect
 			return Mobility != EComponentMobility::Static;
@@ -537,10 +537,31 @@ bool UPrimitiveComponent::CanEditChange(const UProperty* InProperty) const
 			return Mobility != EComponentMobility::Movable || bLightAsIfStatic;
 		}
 
-		if (PropertyName == TEXT("IndirectLightingCacheQuality"))
+		if (PropertyName == GET_MEMBER_NAME_CHECKED(UPrimitiveComponent, IndirectLightingCacheQuality))
 		{
 			return Mobility == EComponentMobility::Movable;
 		}
+
+		if (PropertyName == GET_MEMBER_NAME_CHECKED(UPrimitiveComponent, CastShadow))
+		{
+			// Look for any lit materials
+			bool bHasAnyLitMaterials = false;
+			const int32 NumMaterials = GetNumMaterials();
+			for (int32 MaterialIndex = 0; (MaterialIndex < NumMaterials) && !bHasAnyLitMaterials; ++MaterialIndex)
+			{
+				if (UMaterialInterface* Material = GetMaterial(MaterialIndex))
+				{
+					if (Material->GetShadingModel() != MSM_Unlit)
+					{
+						bHasAnyLitMaterials = true;
+					}
+				}
+			}
+
+			// If there's at least one lit section it could cast shadows, so let the property be edited.
+			// The 0 materials catch is in case any components aren't properly implementing the GetMaterial API, they might or might not work with shadows.
+			return (NumMaterials == 0) || bHasAnyLitMaterials;
+		}
 	}
 
 	return bIsEditable;
