Network demo fixes

* Using DemoPendingNetGame so that IsServer and IsClient checks work properly
* IsServer and IsClient checks now consider DemoNetDriver

diff --git a/Engine/Source/Runtime/Engine/Classes/Engine/PendingNetGame.h b/Engine/Source/Runtime/Engine/Classes/Engine/PendingNetGame.h
--- a/Engine/Source/Runtime/Engine/Classes/Engine/PendingNetGame.h
+++ b/Engine/Source/Runtime/Engine/Classes/Engine/PendingNetGame.h
@@ -75,6 +75,13 @@ class UPendingNetGame :
 	UPROPERTY()
 	class UNetDriver*		NetDriver;
 
+	/** 
+	 * Demo Net driver created for loading demos, but we need to go through pending net game
+	 * Transferred to world on successful connection
+	 */
+	UPROPERTY()
+	class UDemoNetDriver*	DemoNetDriver;
+
 public:
 	/** URL associated with this level. */
 	FURL					URL;
diff --git a/Engine/Source/Runtime/Engine/Private/DemoNetDriver.cpp b/Engine/Source/Runtime/Engine/Private/DemoNetDriver.cpp
--- a/Engine/Source/Runtime/Engine/Private/DemoNetDriver.cpp
+++ b/Engine/Source/Runtime/Engine/Private/DemoNetDriver.cpp
@@ -114,7 +114,13 @@ bool UDemoNetDriver::InitConnect( FNetworkNotify* InNotify, const FURL& ConnectU
 	World->DemoNetDriver = NULL;
 	SetWorld( NULL );
 
-	if ( !GEngine->LoadMap( *WorldContext, DemoURL, NULL, LoadMapError ) )
+	UDemoPendingNetGame * NewPendingNetGame = new UDemoPendingNetGame( FPostConstructInitializeProperties() );
+
+	NewPendingNetGame->DemoNetDriver = this;
+
+	WorldContext->PendingNetGame = NewPendingNetGame;
+
+	if ( !GEngine->LoadMap( *WorldContext, DemoURL, NewPendingNetGame, LoadMapError ) )
 	{
 		Error = LoadMapError;
 		UE_LOG( LogDemo, Error, TEXT( "UDemoNetDriver::InitConnect: LoadMap failed: failed: %s" ), *Error );
@@ -123,6 +129,7 @@ bool UDemoNetDriver::InitConnect( FNetworkNotify* InNotify, const FURL& ConnectU
 
 	SetWorld( WorldContext->World() );
 	WorldContext->World()->DemoNetDriver = this;
+	WorldContext->PendingNetGame = NULL;
 #endif
 
 	DemoDeltaTime = 0;
diff --git a/Engine/Source/Runtime/Engine/Private/UnrealEngine.cpp b/Engine/Source/Runtime/Engine/Private/UnrealEngine.cpp
--- a/Engine/Source/Runtime/Engine/Private/UnrealEngine.cpp
+++ b/Engine/Source/Runtime/Engine/Private/UnrealEngine.cpp
@@ -84,7 +84,7 @@
 #include "Engine/GameInstance.h"
 #include "HotReloadInterface.h"
 #include "STestSuite.h"
-
+#include "Engine/DemoNetDriver.h"
 
 DEFINE_LOG_CATEGORY_STATIC(LogEngine, Log, All);
 
@@ -9038,6 +9038,13 @@ void UEngine::MovePendingLevel(FWorldContext &Context)
 		NetDriver->SetWorld(Context.World());
 	}
 
+	// Attach the DemoNetDriver to the world if there is one
+	if ( Context.PendingNetGame->DemoNetDriver != NULL )
+	{
+		Context.PendingNetGame->DemoNetDriver->SetWorld( Context.World() );
+		Context.World()->DemoNetDriver = Context.PendingNetGame->DemoNetDriver;
+	}
+
 	// Reset the Navigation System
 	Context.World()->SetNavigationSystem(NULL);
 }
diff --git a/Engine/Source/Runtime/Engine/Private/World.cpp b/Engine/Source/Runtime/Engine/Private/World.cpp
--- a/Engine/Source/Runtime/Engine/Private/World.cpp
+++ b/Engine/Source/Runtime/Engine/Private/World.cpp
@@ -4022,8 +4022,17 @@ bool UWorld::IsClient()
 
 bool UWorld::IsServer()
 {
-	// @todo ONLINE move this into net driver since there will no longer be one netdriver
-	return (!NetDriver || NetDriver->IsServer());
+	if ( NetDriver != NULL )
+	{
+		return NetDriver->IsServer();
+	}
+
+	if ( DemoNetDriver != NULL )
+	{
+		return DemoNetDriver->IsServer();
+	}
+
+	return true;
 }
 
 void UWorld::PrepareMapChange(const TArray<FName>& LevelNames)
@@ -5203,7 +5212,17 @@ ULevel* UWorld::GetCurrentLevel() const
 
 ENetMode UWorld::GetNetMode() const
 {
-	return NetDriver ? NetDriver->GetNetMode() : NM_Standalone;
+	if ( NetDriver != NULL )
+	{
+		return NetDriver->GetNetMode();
+	}
+
+	if ( DemoNetDriver )
+	{
+		return DemoNetDriver->GetNetMode();
+	}
+
+	return NM_Standalone;
 }
 
 void UWorld::CopyGameState(AGameMode* FromGameMode, AGameState* FromGameState)
