Fix bug where interior volume would pop to target value instead of ramping in over time.

diff --git a/Engine/Source/Runtime/Engine/Private/ActiveSound.cpp b/Engine/Source/Runtime/Engine/Private/ActiveSound.cpp
--- a/Engine/Source/Runtime/Engine/Private/ActiveSound.cpp
+++ b/Engine/Source/Runtime/Engine/Private/ActiveSound.cpp
@@ -330,10 +330,10 @@ void FActiveSound::HandleInteriorVolumes( const FListener& Listener, FSoundParse
 	if( Listener.Volume == ReverbVolume || !bAllowSpatialization )
 	{
 		// Ambient and listener in same ambient zone
-		CurrentInteriorVolume = ( SourceInteriorVolume * ( 1.0f - Listener.InteriorVolumeInterp ) ) + Listener.InteriorVolumeInterp;
+		CurrentInteriorVolume = ( SourceInteriorVolume * ( 1.0f - Listener.InteriorVolumeInterp ) ) + (Listener.InteriorSettings.InteriorVolume * Listener.InteriorVolumeInterp);
 		ParseParams.VolumeMultiplier *= CurrentInteriorVolume;
 
-		CurrentInteriorLPF = ( SourceInteriorLPF * ( 1.0f - Listener.InteriorLPFInterp ) ) + Listener.InteriorLPFInterp;
+		CurrentInteriorLPF = ( SourceInteriorLPF * ( 1.0f - Listener.InteriorLPFInterp ) ) + (Listener.InteriorSettings.InteriorLPF * Listener.InteriorLPFInterp);
 		ParseParams.HighFrequencyGain *= CurrentInteriorLPF;
 
 		UE_LOG(LogAudio, Verbose, TEXT( "Ambient in same volume. Volume *= %g LPF *= %g (%s)" ),
@@ -357,16 +357,16 @@ void FActiveSound::HandleInteriorVolumes( const FListener& Listener, FSoundParse
 		else
 		{
 			// The ambient sound is 'inside' - use the ambient sound's interior volume multiplied with the listeners exterior volume
-			CurrentInteriorVolume = ( SourceInteriorVolume * ( 1.0f - Listener.InteriorVolumeInterp ) ) + ( Ambient.InteriorVolume * Listener.InteriorVolumeInterp );
-			float CurrentExteriorVolume = ( SourceInteriorVolume * ( 1.0f - Listener.ExteriorVolumeInterp ) ) + ( Listener.InteriorSettings.ExteriorVolume * Listener.ExteriorVolumeInterp );
-			ParseParams.VolumeMultiplier *= CurrentInteriorVolume*CurrentExteriorVolume;
+			CurrentInteriorVolume = ( SourceInteriorVolume * ( 1.0f - Listener.InteriorVolumeInterp ) ) + ( Ambient.InteriorVolume * Listener.InteriorVolumeInterp )
+										* ( SourceInteriorVolume * ( 1.0f - Listener.ExteriorVolumeInterp ) ) + ( Listener.InteriorSettings.ExteriorVolume * Listener.ExteriorVolumeInterp );
+			ParseParams.VolumeMultiplier *= CurrentInteriorVolume;
 
 			CurrentInteriorLPF = ( SourceInteriorLPF * ( 1.0f - Listener.InteriorLPFInterp ) ) + ( Ambient.InteriorLPF * Listener.InteriorLPFInterp );
 			float CurrentExteriorLPF = ( SourceInteriorLPF * ( 1.0f - Listener.ExteriorLPFInterp ) ) + ( Listener.InteriorSettings.ExteriorLPF * Listener.ExteriorLPFInterp );
 			ParseParams.HighFrequencyGain *= CurrentInteriorLPF*CurrentExteriorLPF;
 
 			UE_LOG(LogAudio, Verbose, TEXT( "Ambient in diff volume, ambient inside. Volume *= %g LPF *= %g (%s)" ),
-				CurrentInteriorVolume*CurrentExteriorVolume, CurrentInteriorLPF*CurrentExteriorLPF, ( WaveInstances.Num() > 0 ) ? GetAWaveName(WaveInstances) : TEXT( "NULL" ) );
+				CurrentInteriorVolume, CurrentInteriorLPF*CurrentExteriorLPF, ( WaveInstances.Num() > 0 ) ? GetAWaveName(WaveInstances) : TEXT( "NULL" ) );
 		}
 	}
 }
