Created LOD Group for 2D Pixels, which disables mipmaps and filtering.
#ttp 345443 - NUX:Texture Editor: "Nearest" texture filtering option is buried away, but this is super useful for 2D games
#branch UE4
#proj Runtime.Engine

diff --git a/Engine/Source/Runtime/Engine/Classes/Engine/TextureDefines.h b/Engine/Source/Runtime/Engine/Classes/Engine/TextureDefines.h
--- a/Engine/Source/Runtime/Engine/Classes/Engine/TextureDefines.h
+++ b/Engine/Source/Runtime/Engine/Classes/Engine/TextureDefines.h
@@ -49,6 +49,8 @@ enum TextureGroup
 	TEXTUREGROUP_Bokeh UMETA(DisplayName="Bokeh"),
 	// no compression, created on import of a .IES file
 	TEXTUREGROUP_IESLightProfile UMETA(DisplayName="IESLightProfile"),
+	// non-filtered, useful for 2D rendering
+	TEXTUREGROUP_Pixels2D UMETA(DisplayName="2D Pixels"),
 	TEXTUREGROUP_MAX,
 };
 
