#ue4 - Fix path following in CharacterMovementComponent, so that it turns more quickly and considers current velocity when trying to reach destination speed.

- Restored "bRequestedMoveUseAcceleration" which can be set to false to disable acceleration and use requested pathfinding velocity directly (instant turning).
- API Note: ApplyRequestedMove() has had params reordered and changed; it now includes Friction and Braking, and has all the "out" params at the end.

TTP 345589

diff --git a/Engine/Source/Runtime/Engine/Classes/GameFramework/CharacterMovementComponent.h b/Engine/Source/Runtime/Engine/Classes/GameFramework/CharacterMovementComponent.h
--- a/Engine/Source/Runtime/Engine/Classes/GameFramework/CharacterMovementComponent.h
+++ b/Engine/Source/Runtime/Engine/Classes/GameFramework/CharacterMovementComponent.h
@@ -324,7 +324,7 @@ public:
 	UPROPERTY(Category="Character Movement", EditAnywhere, AdvancedDisplay)
 	uint32 bEnableScopedMovementUpdates:1;
 
-	/** Ignores Acceleration component, and forces max acceleration to drive APawn at full velocity. */
+	/** Ignores size of acceleration component, and forces max acceleration to drive character at full velocity. */
 	UPROPERTY()
 	uint32 bForceMaxAccel:1;    
 
@@ -621,7 +621,11 @@ public:
 	UPROPERTY(Category="Avoidance", EditAnywhere, BlueprintReadOnly)
 	uint32 bUseRVOAvoidance:1;
 
-	/** Should use acceleration for path following? */
+	/**
+	 * Should use acceleration for path following?
+	 * If true, acceleration is applied when path following to reach the target velocity.
+	 * If false, path following velocity is set directly, disregarding acceleration.
+	 */
 	UPROPERTY(Category="Character Movement", EditAnywhere, BlueprintReadWrite, AdvancedDisplay)
 	uint32 bRequestedMoveUseAcceleration:1;
 
@@ -869,8 +873,21 @@ public:
 	  */
 	virtual FRotator ComputeOrientToMovementRotation(const FRotator& CurrentRotation, float DeltaTime, FRotator& DeltaRotation);
 
-	/** use velocity requested by path following */
-	virtual bool ApplyRequestedMove(FVector& NewAcceleration, float DeltaTime, float MaxAccel, float MaxSpeed, float& OutRequestedSpeed);
+	/**
+	 * Use velocity requested by path following to compute a requested acceleration and speed.
+	 * This does not affect the Acceleration member variable, as that is used to indicate input acceleration.
+	 * This may directly affect current Velocity.
+	 *
+	 * @param DeltaTime				Time slice for this operation
+	 * @param MaxAccel				Max acceleration allowed in OutAcceleration result.
+	 * @param MaxSpeed				Max speed allowed when computing OutRequestedSpeed.
+	 * @param Friction				Current friction.
+	 * @param BrakingDeceleration	Current braking deceleration.
+	 * @param OutAcceleration		Acceleration computed based on requested velocity.
+	 * @param OutRequestedSpeed		Speed of resulting velocity request, which can affect the max speed allowed by movement.
+	 * @return Whether there is a requested velocity and acceleration, resulting in valid OutAcceleration and OutRequestedSpeed values.
+	 */
+	virtual bool ApplyRequestedMove(float DeltaTime, float MaxAccel, float MaxSpeed, float Friction, float BrakingDeceleration, FVector& OutAcceleration, float& OutRequestedSpeed);
 
 	/** Called if bNotifyApex is true and character has just passed the apex of its jump. */
 	virtual void NotifyJumpApex();
diff --git a/Engine/Source/Runtime/Engine/Private/CharacterMovementComponent.cpp b/Engine/Source/Runtime/Engine/Private/CharacterMovementComponent.cpp
--- a/Engine/Source/Runtime/Engine/Private/CharacterMovementComponent.cpp
+++ b/Engine/Source/Runtime/Engine/Private/CharacterMovementComponent.cpp
@@ -372,6 +372,11 @@ void UCharacterMovementComponent::SetUpdatedComponent(UPrimitiveComponent* NewUp
 	Super::SetUpdatedComponent(NewUpdatedComponent);
 	CharacterOwner = Cast<ACharacter>(PawnOwner);
 
+	if (UpdatedComponent == NULL)
+	{
+		StopActiveMovement();
+	}
+
 	if (bEnablePhysicsInteraction)
 	{
 		UpdatedComponent->OnComponentBeginOverlap.AddDynamic(this, &UCharacterMovementComponent::CapsuleTouched);
@@ -1454,9 +1459,6 @@ void UCharacterMovementComponent::PerformMovement(float DeltaSeconds)
 		// change position
 		StartNewPhysics(DeltaSeconds, 0);
 
-		// consume path following requested velocity
-		bHasRequestedVelocity = false;
-
 		if (!HasValidData())
 		{
 			return;
@@ -1510,6 +1512,9 @@ void UCharacterMovementComponent::PerformMovement(float DeltaSeconds)
 			RootMotionParams.Clear();
 		}
 
+		// consume path following requested velocity
+		bHasRequestedVelocity = false;
+
 		OnMovementUpdated(DeltaSeconds, OldLocation, OldVelocity);
 	} // End scoped movement update
 
@@ -2139,12 +2144,13 @@ void UCharacterMovementComponent::CalcVelocity(float DeltaTime, float Friction,
 	float MaxSpeed = GetMaxSpeed();
 	
 	// Check if path following requested movement
+	bool bZeroRequestedAcceleration = true;
 	FVector RequestedAcceleration = FVector::ZeroVector;
 	float RequestedSpeed = 0.0f;
-	if(ApplyRequestedMove(RequestedAcceleration, DeltaTime, MaxAccel, MaxSpeed, RequestedSpeed))
+	if (ApplyRequestedMove(DeltaTime, MaxAccel, MaxSpeed, Friction, BrakingDeceleration, RequestedAcceleration, RequestedSpeed))
 	{
-		Acceleration += RequestedAcceleration;
-		Acceleration = Acceleration.ClampMaxSize(MaxAccel);
+		RequestedAcceleration.ClampMaxSize(MaxAccel);
+		bZeroRequestedAcceleration = false;
 	}
 
 	if (bForceMaxAccel)
@@ -2164,12 +2170,15 @@ void UCharacterMovementComponent::CalcVelocity(float DeltaTime, float Friction,
 	}
 
 	// Path following above didn't care about the analog modifier, but we do for everything else below, so get the fully modified value.
-	// Use max of requested speed and max speed if we modified the speed in ApplyRequestedMove above
+	// Use max of requested speed and max speed if we modified the speed in ApplyRequestedMove above.
 	MaxSpeed = FMath::Max(RequestedSpeed, MaxSpeed * AnalogInputModifier);
 
 	// Apply braking or deceleration
+	const bool bZeroAcceleration = Acceleration.IsZero();
 	const bool bVelocityOverMax = IsExceedingMaxSpeed(MaxSpeed);
-	if (Acceleration.IsZero() || bVelocityOverMax)
+	
+	// Only apply braking if there is no acceleration, or we are over our max speed and need to slow down to it.
+	if ((bZeroAcceleration && bZeroRequestedAcceleration) || bVelocityOverMax)
 	{
 		const FVector OldVelocity = Velocity;
 		ApplyVelocityBraking(DeltaTime, Friction, BrakingDeceleration);
@@ -2180,8 +2189,9 @@ void UCharacterMovementComponent::CalcVelocity(float DeltaTime, float Friction,
 			Velocity = SafeNormalPrecise(OldVelocity) * MaxSpeed;
 		}
 	}
-	else
+	else if (!bZeroAcceleration)
 	{
+		// Friction affects our ability to change direction. This is only done for input acceleration, not path following.
 		const FVector AccelDir = SafeNormalPrecise(Acceleration);
 		const float VelSize = Velocity.Size();
 		Velocity = Velocity - (Velocity - AccelDir * VelSize) * FMath::Min(DeltaTime * Friction, 1.f);
@@ -2193,9 +2203,10 @@ void UCharacterMovementComponent::CalcVelocity(float DeltaTime, float Friction,
 		Velocity = Velocity * (1.f - FMath::Min(Friction * DeltaTime, 1.f));
 	}
 
-	// Apply input acceleration
+	// Apply acceleration
 	const float NewMaxSpeed = (IsExceedingMaxSpeed(MaxSpeed)) ? Velocity.Size() : MaxSpeed;
 	Velocity += Acceleration * DeltaTime;
+	Velocity += RequestedAcceleration * DeltaTime;
 	Velocity = ClampMaxSizePrecise(Velocity, NewMaxSpeed);
 
 	if (bUseRVOAvoidance)
@@ -2204,27 +2215,48 @@ void UCharacterMovementComponent::CalcVelocity(float DeltaTime, float Friction,
 	}
 }
 
-bool UCharacterMovementComponent::ApplyRequestedMove(FVector& NewAcceleration, float DeltaTime, float MaxAccel, float MaxSpeed, float& OutRequestedSpeed)
+
+bool UCharacterMovementComponent::ApplyRequestedMove(float DeltaTime, float MaxAccel, float MaxSpeed, float Friction, float BrakingDeceleration, FVector& OutAcceleration, float& OutRequestedSpeed)
 {
 	if (bHasRequestedVelocity)
 	{
-		OutRequestedSpeed = RequestedVelocity.Size();
-		if (OutRequestedSpeed < KINDA_SMALL_NUMBER)
+		const float RequestedSpeedSquared = RequestedVelocity.SizeSquared();
+		if (RequestedSpeedSquared < KINDA_SMALL_NUMBER)
 		{
-			return true;
+			return false;
 		}
 
-		const FVector RequestedMoveDir = RequestedVelocity / OutRequestedSpeed;
-
-		OutRequestedSpeed = bRequestedMoveWithMaxSpeed ? MaxSpeed : FMath::Min(MaxSpeed, OutRequestedSpeed);
-		const FVector MoveVelocity = RequestedMoveDir * OutRequestedSpeed;
+		// Compute requested speed from path following
+		float RequestedSpeed = FMath::Sqrt(RequestedSpeedSquared);
+		const FVector RequestedMoveDir = RequestedVelocity / RequestedSpeed;
+		RequestedSpeed = (bRequestedMoveWithMaxSpeed ? MaxSpeed : FMath::Min(MaxSpeed, RequestedSpeed));
+		
+		// Compute actual requested velocity
+		const FVector MoveVelocity = RequestedMoveDir * RequestedSpeed;
+		
+		// Compute acceleration. Use MaxAccel to limit speed increase, 1% buffer.
+		FVector NewAcceleration = FVector::ZeroVector;
+		const float CurrentSpeedSq = Velocity.SizeSquared();
+		if (bRequestedMoveUseAcceleration && CurrentSpeedSq < FMath::Square(RequestedSpeed * 1.01f))
+		{
+			// Turn in the same manner as with input acceleration.
+			const float VelSize = FMath::Sqrt(CurrentSpeedSq);
+			Velocity = Velocity - (Velocity - RequestedMoveDir * VelSize) * FMath::Min(DeltaTime * Friction, 1.f);
 
-		NewAcceleration = MoveVelocity / DeltaTime;
-		if (NewAcceleration.SizeSquared() > FMath::Square(MaxAccel) || bForceMaxAccel)
+			// How much do we need to accelerate to get to the new velocity?
+			NewAcceleration = ((MoveVelocity - Velocity) / DeltaTime);
+			NewAcceleration = NewAcceleration.ClampMaxSize(MaxAccel);
+		}
+		else
 		{
-			NewAcceleration = RequestedMoveDir * MaxAccel;
+			// Just set velocity directly.
+			// If decelerating we do so instantly, so we don't slide through the destination if we can't brake fast enough.
+			Velocity = MoveVelocity;
 		}
 
+		// Copy to out params
+		OutRequestedSpeed = RequestedSpeed;
+		OutAcceleration = NewAcceleration;
 		return true;
 	}
 
@@ -3722,9 +3754,15 @@ FRotator UCharacterMovementComponent::GetDeltaRotation(float DeltaTime)
 
 FRotator UCharacterMovementComponent::ComputeOrientToMovementRotation(const FRotator& CurrentRotation, float DeltaTime, FRotator& DeltaRotation)
 {
-	// Do nothing if not accelerating.
 	if (Acceleration.SizeSquared() < KINDA_SMALL_NUMBER)
 	{
+		// AI path following request can orient us in that direction (it's effectively an acceleration)
+		if (bHasRequestedVelocity && RequestedVelocity.SizeSquared() > KINDA_SMALL_NUMBER)
+		{
+			return RequestedVelocity.SafeNormal().Rotation();
+		}
+
+		// Don't change rotation if there is no acceleration.
 		return CurrentRotation;
 	}
 
