Engine: Split analytics enables to separate hardware survey from editor

diff --git a/Engine/Source/Runtime/Engine/Classes/Engine/Engine.h b/Engine/Source/Runtime/Engine/Classes/Engine/Engine.h
--- a/Engine/Source/Runtime/Engine/Classes/Engine/Engine.h
+++ b/Engine/Source/Runtime/Engine/Classes/Engine/Engine.h
@@ -1070,12 +1070,12 @@ public:
 	UPROPERTY(EditAnywhere, config, Category=Settings)
 	float NearClipPlane;
 
-	/** Should the engine analytics provider be initialized on startup?  Note: The engine/editor must be restarted for changes to take effect. */
+	/** Can the editor report usage analytics (types of assets being spawned, etc...) back to Epic in order for us to improve the editor user experience?  Note: The editor must be restarted for changes to take effect. */
 	UPROPERTY(EditAnywhere, config, Category=Settings, AdvancedDisplay)
-	uint32 bEngineAnalyticsProviderEnabled:1;
+	uint32 bEditorAnalyticsEnabled:1;
 
-	/** Can the game/application report anonymous hardware survey statistics back to Epic?  (such as display resolution and GPU model) */
-	UPROPERTY(EditAnywhere, config, Category=Settings, AdvancedDisplay, meta=(EditCondition="bEngineAnalyticsProviderEnabled"))
+	/** Can a runtime game/application report anonymous hardware survey statistics (such as display resolution and GPU model) back to Epic? */
+	UPROPERTY(EditAnywhere, config, Category=Settings, AdvancedDisplay)
 	uint32 bHardwareSurveyEnabled:1;
 
 	/** Flag for completely disabling subtitles for localized sounds. */
diff --git a/Engine/Source/Runtime/Engine/Private/EngineAnalytics.cpp b/Engine/Source/Runtime/Engine/Private/EngineAnalytics.cpp
--- a/Engine/Source/Runtime/Engine/Private/EngineAnalytics.cpp
+++ b/Engine/Source/Runtime/Engine/Private/EngineAnalytics.cpp
@@ -119,10 +119,22 @@ void FEngineAnalytics::Initialize()
 	checkf(!bIsInitialized, TEXT("FEngineAnalytics::Initialize called more than once."));
 
 	// Never use analytics when running a commandlet tool
+	bool bShouldInitAnalytics = !IsRunningCommandlet();
+
 	check(GEngine);
-	const bool bShouldInitAnalytics = !IsRunningCommandlet() && GEngine->bEngineAnalyticsProviderEnabled;
+#if WITH_EDITORONLY_DATA
+	if (GIsEditor)
+	{
+		bShouldInitAnalytics = bShouldInitAnalytics && GEngine->bEditorAnalyticsEnabled;
+	}
+	else
+#endif
+	{
+		// Outside of the editor, the only engine analytics usage is the hardware survey
+		bShouldInitAnalytics = bShouldInitAnalytics && GEngine->bHardwareSurveyEnabled;
+	}
 
-	if( bShouldInitAnalytics )
+	if (bShouldInitAnalytics)
 	{
 #if DEBUG_ANALYTICS
 		// Dummy analytics provider just logs events locally
diff --git a/Engine/Source/Runtime/Engine/Private/GameEngine.cpp b/Engine/Source/Runtime/Engine/Private/GameEngine.cpp
--- a/Engine/Source/Runtime/Engine/Private/GameEngine.cpp
+++ b/Engine/Source/Runtime/Engine/Private/GameEngine.cpp
@@ -393,7 +393,7 @@ UEngine::UEngine(const class FPostConstructInitializeProperties& PCIP)
 	SelectionHighlightIntensityBillboards = 0.25f;
 
 	bUseSound = true;
-	bEngineAnalyticsProviderEnabled = true;
+	bEditorAnalyticsEnabled = true;
 	bHardwareSurveyEnabled = true;
 	bPendingHardwareSurveyResults = false;
 	bIsInitialized = false;
diff --git a/Engine/Source/Runtime/Engine/Private/UnrealEngine.cpp b/Engine/Source/Runtime/Engine/Private/UnrealEngine.cpp
--- a/Engine/Source/Runtime/Engine/Private/UnrealEngine.cpp
+++ b/Engine/Source/Runtime/Engine/Private/UnrealEngine.cpp
@@ -5543,7 +5543,10 @@ void UEngine::OnLostFocusPause(bool EnablePause)
 
 void UEngine::InitHardwareSurvey()
 {
-	bool bEnabled = bHardwareSurveyEnabled;
+	bool bEnabled = true;
+#if !WITH_EDITORONLY_DATA
+	bEnabled = bHardwareSurveyEnabled;
+#endif
 
 	// The hardware survey costs time and we don't want to slow down debug builds.
 	// This is mostly because of the CPU benchmark running in the survey and the results in debug are not being valid.
@@ -5588,7 +5591,7 @@ void UEngine::TickHardwareSurvey()
 bool UEngine::IsHardwareSurveyRequired()
 {
 	// Analytics must have been initialized FIRST.
-	if (!FEngineAnalytics::IsAvailable() || IsRunningDedicatedServer() )
+	if (!FEngineAnalytics::IsAvailable() || IsRunningDedicatedServer())
 	{
 		return false;
 	}
