Fixed case where foliage could be painted with a wrong scale
TTP# 347365

diff --git a/Engine/Source/Runtime/Engine/Private/InstancedFoliage.cpp b/Engine/Source/Runtime/Engine/Private/InstancedFoliage.cpp
--- a/Engine/Source/Runtime/Engine/Private/InstancedFoliage.cpp
+++ b/Engine/Source/Runtime/Engine/Private/InstancedFoliage.cpp
@@ -300,7 +300,9 @@ void FFoliageMeshInfo::AddInstance(AInstancedFoliageActor* InIFA, UFoliageType*
 
 		BestCluster->ClusterComponent->AttachTo(InIFA->GetRootComponent());
 		BestCluster->ClusterComponent->RegisterComponent();
-		BestCluster->ClusterComponent->SetWorldTransform(InstanceToWorld);
+		
+		// Use only instance translation as a component transform
+		BestCluster->ClusterComponent->SetWorldTransform(FTransform(InstanceToWorld.GetTranslation()));
 	}
 	else
 	{
