Removed commented lines

diff --git a/Engine/Source/Runtime/Engine/Private/Animation/AnimNode_Trail.cpp b/Engine/Source/Runtime/Engine/Private/Animation/AnimNode_Trail.cpp
--- a/Engine/Source/Runtime/Engine/Private/Animation/AnimNode_Trail.cpp
+++ b/Engine/Source/Runtime/Engine/Private/Animation/AnimNode_Trail.cpp
@@ -71,7 +71,6 @@ void FAnimNode_Trail::EvaluateBoneTransforms(USkeletalMeshComponent* SkelComp, c
 
 		//Insert indices at the start of array, so that parents are before children in the array.
 		int32 TransformIndex = ChainLength - (i + 1);
-		//ChainBoneIndicesTemp.Add(WalkBoneIndex);
 		ChainBoneIndices[TransformIndex] = WalkBoneIndex;
 	}
 
