Fix LPF the same way Volume was in CL# 2310237

diff --git a/Engine/Source/Runtime/Engine/Private/ActiveSound.cpp b/Engine/Source/Runtime/Engine/Private/ActiveSound.cpp
--- a/Engine/Source/Runtime/Engine/Private/ActiveSound.cpp
+++ b/Engine/Source/Runtime/Engine/Private/ActiveSound.cpp
@@ -361,12 +361,12 @@ void FActiveSound::HandleInteriorVolumes( const FListener& Listener, FSoundParse
 										* ( SourceInteriorVolume * ( 1.0f - Listener.ExteriorVolumeInterp ) ) + ( Listener.InteriorSettings.ExteriorVolume * Listener.ExteriorVolumeInterp );
 			ParseParams.VolumeMultiplier *= CurrentInteriorVolume;
 
-			CurrentInteriorLPF = ( SourceInteriorLPF * ( 1.0f - Listener.InteriorLPFInterp ) ) + ( Ambient.InteriorLPF * Listener.InteriorLPFInterp );
-			float CurrentExteriorLPF = ( SourceInteriorLPF * ( 1.0f - Listener.ExteriorLPFInterp ) ) + ( Listener.InteriorSettings.ExteriorLPF * Listener.ExteriorLPFInterp );
-			ParseParams.HighFrequencyGain *= CurrentInteriorLPF*CurrentExteriorLPF;
+			CurrentInteriorLPF = ( SourceInteriorLPF * ( 1.0f - Listener.InteriorLPFInterp ) ) + ( Ambient.InteriorLPF * Listener.InteriorLPFInterp )
+										* ( SourceInteriorLPF * ( 1.0f - Listener.ExteriorLPFInterp ) ) + ( Listener.InteriorSettings.ExteriorLPF * Listener.ExteriorLPFInterp );
+			ParseParams.HighFrequencyGain *= CurrentInteriorLPF;
 
 			UE_LOG(LogAudio, Verbose, TEXT( "Ambient in diff volume, ambient inside. Volume *= %g LPF *= %g (%s)" ),
-				CurrentInteriorVolume, CurrentInteriorLPF*CurrentExteriorLPF, ( WaveInstances.Num() > 0 ) ? GetAWaveName(WaveInstances) : TEXT( "NULL" ) );
+				CurrentInteriorVolume, CurrentInteriorLPF, ( WaveInstances.Num() > 0 ) ? GetAWaveName(WaveInstances) : TEXT( "NULL" ) );
 		}
 	}
 }
