Rename AReverbVolume to AAudioVolume

diff --git a/Engine/Config/BaseEngine.ini b/Engine/Config/BaseEngine.ini
--- a/Engine/Config/BaseEngine.ini
+++ b/Engine/Config/BaseEngine.ini
@@ -196,7 +196,8 @@ bUseStreamingPause=false
 +ActiveClassRedirects=(OldClassName="AmbientSoundSimpleToggleable",NewClassName="AmbientSound")
 +ActiveClassRedirects=(OldClassName="AmbientSoundNonLoop",NewClassName="AmbientSound")
 +ActiveClassRedirects=(OldClassName="AmbientSoundNonLoopingToggleable",NewClassName="AmbientSound")
-+ActiveClassRedirects=(OldClassName="ReverbVolumeToggleable",NewClassName="ReverbVolume")
++ActiveClassRedirects=(OldClassName="ReverbVolume",NewClassName="AudioVolume")
++ActiveClassRedirects=(OldClassName="ReverbVolumeToggleable",NewClassName="AudioVolume")
 +ActiveClassRedirects=(OldClassName="BlueprintActorBase",NewClassName="Actor")
 +ActiveClassRedirects=(OldClassName="WorldInfo",NewClassName="WorldSettings")
 +ActiveClassRedirects=(OldClassName="RB_Handle",NewClassName="PhysicsHandleComponent")
diff --git a/Engine/Source/Editor/EditorStyle/Private/SlateEditorStyle.cpp b/Engine/Source/Editor/EditorStyle/Private/SlateEditorStyle.cpp
--- a/Engine/Source/Editor/EditorStyle/Private/SlateEditorStyle.cpp
+++ b/Engine/Source/Editor/EditorStyle/Private/SlateEditorStyle.cpp
@@ -4668,6 +4668,7 @@ void FSlateEditorStyle::FStyle::SetupClassIconsAndThumbnails()
 			TEXT("AnimMontage"),
 			TEXT("AnimSequence"),
 			TEXT("AtmosphericFog"),
+			TEXT("AudioVolume"),
 			TEXT("BehaviorTree"),
 			TEXT("BlackboardData"),
 			TEXT("BlendSpace"),
@@ -4736,8 +4737,6 @@ void FSlateEditorStyle::FStyle::SetupClassIconsAndThumbnails()
 			TEXT("RadialForceActor"),
 			TEXT("ReflectionCapture"),
 			TEXT("ReverbEffect"),
-			TEXT("ReverbVolume"),
-			TEXT("ReverbVolumeToggleable"),
 			TEXT("SceneCapture2D"),
 			TEXT("SceneCaptureCube"),
 			TEXT("SkyLight"),
diff --git a/Engine/Source/Editor/UnrealEd/Private/LevelEditorViewport.cpp b/Engine/Source/Editor/UnrealEd/Private/LevelEditorViewport.cpp
--- a/Engine/Source/Editor/UnrealEd/Private/LevelEditorViewport.cpp
+++ b/Engine/Source/Editor/UnrealEd/Private/LevelEditorViewport.cpp
@@ -3966,7 +3966,7 @@ void FLevelEditorViewportClient::UpdateAudioListener( const FSceneView& View )
 			FInteriorSettings InteriorSettings;
 			const FVector& ViewLocation = GetViewLocation();
 
-			class AReverbVolume* ReverbVolume = GetWorld()->GetAudioSettings( ViewLocation, &ReverbSettings, &InteriorSettings );
+			class AAudioVolume* AudioVolume = GetWorld()->GetAudioSettings( ViewLocation, &ReverbSettings, &InteriorSettings );
 
 			FMatrix CameraToWorld = View.ViewMatrices.ViewMatrix.InverseFast();
 			FVector ProjUp = CameraToWorld.TransformVector(FVector(0, 1000, 0));
@@ -3976,8 +3976,8 @@ void FLevelEditorViewportClient::UpdateAudioListener( const FSceneView& View )
 			ListenerTransform.SetTranslation(ViewLocation);
 			ListenerTransform.NormalizeRotation();
 
-			AudioDevice->SetListener( 0, ListenerTransform, 0.f, ReverbVolume, InteriorSettings );
-			AudioDevice->SetReverbSettings( ReverbVolume, ReverbSettings );
+			AudioDevice->SetListener( 0, ListenerTransform, 0.f, AudioVolume, InteriorSettings );
+			AudioDevice->SetReverbSettings( AudioVolume, ReverbSettings );
 		}
 	}
 }
diff --git a/Engine/Source/Runtime/Engine/Classes/Engine/World.h b/Engine/Source/Runtime/Engine/Classes/Engine/World.h
--- a/Engine/Source/Runtime/Engine/Classes/Engine/World.h
+++ b/Engine/Source/Runtime/Engine/Classes/Engine/World.h
@@ -846,8 +846,8 @@ public:
 	/** An array of post processing volumes, sorted in ascending order of priority.					*/
 	TArray< IInterface_PostProcessVolume * > PostProcessVolumes;
 
-	/** Linked list of reverb volumes, sorted in descending order of priority.							*/
-	TAutoWeakObjectPtr<class AReverbVolume> HighestPriorityReverbVolume;
+	/** Pointer to the higest priority audio volumes, each volume has a reference to the next lower priority volume creating a linked list of prioritized audio volumes in descending order */
+	TAutoWeakObjectPtr<class AAudioVolume> HighestPriorityAudioVolume;
 
 	/** Time in FPlatformTime::Seconds unbuilt time was last encountered. 0 means not yet.							*/
 	double LastTimeUnbuiltLightingWasEncountered;
@@ -2268,15 +2268,15 @@ public:
 	/*****************************************************************************************************/
 
 	/**
-	 * Finds the reverb settings to use for a given view location, taking into account the world's default
-	 * settings and the reverb volumes in the world.
+	 * Finds the audio settings to use for a given view location, taking into account the world's default
+	 * settings and the audio volumes in the world.
 	 *
 	 * @param	ViewLocation			Current view location.
 	 * @param	OutReverbSettings		[out] Upon return, the reverb settings for a camera at ViewLocation.
 	 * @param	OutInteriorSettings		[out] Upon return, the interior settings for a camera at ViewLocation.
-	 * @return							If the settings came from a reverb volume, the reverb volume's object index is returned.
+	 * @return							If the settings came from an audio volume, the audio volume object is returned.
 	 */
-	class AReverbVolume* GetAudioSettings( const FVector& ViewLocation, struct FReverbSettings* OutReverbSettings, struct FInteriorSettings* OutInteriorSettings );
+	class AAudioVolume* GetAudioSettings( const FVector& ViewLocation, struct FReverbSettings* OutReverbSettings, struct FInteriorSettings* OutInteriorSettings );
 
 	/** Return the URL of this level on the local machine. */
 	virtual FString GetLocalURL() const;
diff --git a/Engine/Source/Runtime/Engine/Classes/GameFramework/WorldSettings.h b/Engine/Source/Runtime/Engine/Classes/GameFramework/WorldSettings.h
--- a/Engine/Source/Runtime/Engine/Classes/GameFramework/WorldSettings.h
+++ b/Engine/Source/Runtime/Engine/Classes/GameFramework/WorldSettings.h
@@ -2,7 +2,7 @@
 
 #pragma once
 #include "GameFramework/Info.h"
-#include "Sound/ReverbVolume.h"
+#include "Sound/AudioVolume.h"
 #include "WorldSettings.generated.h"
 
 UENUM()
@@ -324,11 +324,11 @@ class ENGINE_API AWorldSettings : public AInfo, public IInterface_AssetUserData
 
 	/************************************/
 	/** AUDIO SETTINGS **/
-	/** Default reverb settings used by reverb volumes.													*/
+	/** Default reverb settings used by audio volumes.													*/
 	UPROPERTY(EditAnywhere, config, Category=Audio)
 	struct FReverbSettings DefaultReverbSettings;
 
-	/** Default interior settings used by reverb volumes.												*/
+	/** Default interior settings used by audio volumes.												*/
 	UPROPERTY(EditAnywhere, config, Category=Audio)
 	struct FInteriorSettings DefaultAmbientZoneSettings;
 
diff --git a/Engine/Source/Runtime/Engine/Classes/Sound/AudioVolume.h b/Engine/Source/Runtime/Engine/Classes/Sound/AudioVolume.h
new file mode 100644
--- /dev/null
+++ b/Engine/Source/Runtime/Engine/Classes/Sound/AudioVolume.h
@@ -0,0 +1,199 @@
+// Copyright 1998-2014 Epic Games, Inc. All Rights Reserved.
+
+/**
+ *
+ * Used to affect audio settings in the game and editor.
+ */
+
+#pragma once
+#include "GameFramework/Volume.h"
+#include "AudioVolume.generated.h"
+
+/**
+ * DEPRECATED: Exists for backwards compatibility
+ * Indicates a reverb preset to use.
+ *
+ */
+UENUM()
+enum ReverbPreset
+{
+	REVERB_Default,
+	REVERB_Bathroom,
+	REVERB_StoneRoom,
+	REVERB_Auditorium,
+	REVERB_ConcertHall,
+	REVERB_Cave,
+	REVERB_Hallway,
+	REVERB_StoneCorridor,
+	REVERB_Alley,
+	REVERB_Forest,
+	REVERB_City,
+	REVERB_Mountains,
+	REVERB_Quarry,
+	REVERB_Plain,
+	REVERB_ParkingLot,
+	REVERB_SewerPipe,
+	REVERB_Underwater,
+	REVERB_SmallRoom,
+	REVERB_MediumRoom,
+	REVERB_LargeRoom,
+	REVERB_MediumHall,
+	REVERB_LargeHall,
+	REVERB_Plate,
+	REVERB_MAX,
+};
+
+/** Struct encapsulating settings for reverb effects. */
+USTRUCT()
+struct FReverbSettings
+{
+	GENERATED_USTRUCT_BODY()
+
+	/* Whether to apply the reverb settings below. */
+	UPROPERTY(EditAnywhere, BlueprintReadWrite, Category=ReverbSettings )
+	uint32 bApplyReverb:1;
+
+	/** The reverb preset to employ. */
+	UPROPERTY()
+	TEnumAsByte<enum ReverbPreset> ReverbType_DEPRECATED;
+
+	/** The reverb asset to employ. */
+	UPROPERTY(EditAnywhere, BlueprintReadWrite, Category=ReverbSettings)
+	class UReverbEffect* ReverbEffect;
+
+	/** Volume level of the reverb affect. */
+	UPROPERTY(EditAnywhere, BlueprintReadWrite, Category=ReverbSettings)
+	float Volume;
+
+	/** Time to fade from the current reverb settings into this setting, in seconds. */
+	UPROPERTY(EditAnywhere, BlueprintReadWrite, Category=ReverbSettings)
+	float FadeTime;
+
+
+
+		FReverbSettings()
+		: bApplyReverb(true)
+		, ReverbType_DEPRECATED(REVERB_Default)
+		, ReverbEffect(NULL)
+		, Volume(0.5f)
+		, FadeTime(2.0f)
+		{
+		}
+
+		void PostSerialize(const FArchive& Ar);
+	
+};
+
+template<>
+struct TStructOpsTypeTraits<FReverbSettings> : public TStructOpsTypeTraitsBase
+{
+	enum 
+	{
+		WithPostSerialize = true,
+	};
+};
+
+/** Struct encapsulating settings for interior areas. */
+//@warning: manually mirrored in Components.h
+USTRUCT()
+struct FInteriorSettings
+{
+	GENERATED_USTRUCT_BODY()
+
+	// Whether these interior settings are the default values for the world
+	UPROPERTY()
+	uint32 bIsWorldSettings:1;
+
+	// The desired volume of sounds outside the volume when the player is inside the volume
+	UPROPERTY(EditAnywhere, BlueprintReadWrite, Category=InteriorSettings)
+	float ExteriorVolume;
+
+	// The time over which to interpolate from the current volume to the desired volume of sounds outside the volume when the player enters the volume
+	UPROPERTY(EditAnywhere, BlueprintReadWrite, Category=InteriorSettings)
+	float ExteriorTime;
+
+	// The desired LPF of sounds outside the volume when the player is inside the volume
+	UPROPERTY(EditAnywhere, BlueprintReadWrite, Category=InteriorSettings)
+	float ExteriorLPF;
+
+	// The time over which to interpolate from the current LPF to the desired LPF of sounds outside the volume when the player enters the volume
+	UPROPERTY(EditAnywhere, BlueprintReadWrite, Category=InteriorSettings)
+	float ExteriorLPFTime;
+
+	// The desired volume of sounds inside the volume when the player is inside the volume
+	UPROPERTY(EditAnywhere, BlueprintReadWrite, Category=InteriorSettings)
+	float InteriorVolume;
+
+	// The time over which to interpolate from the current volume to the desired volume of sounds inside the volume when the player enters the volume
+	UPROPERTY(EditAnywhere, BlueprintReadWrite, Category=InteriorSettings)
+	float InteriorTime;
+
+	// The desired LPF of sounds inside  the volume when the player is inside the volume
+	UPROPERTY(EditAnywhere, BlueprintReadWrite, Category=InteriorSettings)
+	float InteriorLPF;
+
+	// The time over which to interpolate from the current LPF to the desired LPF of sounds inside the volume when the player enters the volume
+	UPROPERTY(EditAnywhere, BlueprintReadWrite, Category=InteriorSettings)
+	float InteriorLPFTime;
+
+	FInteriorSettings()
+		: bIsWorldSettings(false)
+		, ExteriorVolume(1.0f)
+		, ExteriorTime(0.5f)
+		, ExteriorLPF(1.0f)
+		, ExteriorLPFTime(0.5f)
+		, InteriorVolume(1.0f)
+		, InteriorTime(0.5f)
+		, InteriorLPF(1.0f)
+		, InteriorLPFTime(0.5f)
+		{
+		}
+	
+};
+
+UCLASS(hidecategories=(Advanced, Attachment, Collision, Volume))
+class AAudioVolume : public AVolume
+{
+	GENERATED_UCLASS_BODY()
+
+	/**
+	 * Priority of this volume. In the case of overlapping volumes the one with the highest priority
+	 * is chosen. The order is undefined if two or more overlapping volumes have the same priority.
+	 */
+	UPROPERTY(EditAnywhere, BlueprintReadWrite, Category=AudioVolume)
+	float Priority;
+
+	/** whether this volume is currently enabled and able to affect sounds */
+	UPROPERTY(EditAnywhere, BlueprintReadWrite, replicated, Category=AudioVolume)
+	uint32 bEnabled:1;
+
+	/** Reverb settings to use for this volume. */
+	UPROPERTY(EditAnywhere, BlueprintReadWrite, Category=Reverb)
+	struct FReverbSettings Settings;
+
+	/** Interior settings used for this volume */
+	UPROPERTY(EditAnywhere, BlueprintReadWrite, Category=AmbientZone)
+	struct FInteriorSettings AmbientZoneSettings;
+
+	/** Next volume in linked listed, sorted by priority in descending order. */
+	UPROPERTY(transient)
+	class AAudioVolume* NextLowerPriorityVolume;
+
+
+
+	// Begin UObject Interface
+#if WITH_EDITOR
+	virtual void PostEditChangeProperty(FPropertyChangedEvent& PropertyChangedEvent) override;
+#endif // WITH_EDITOR
+	// End UObject Interface
+
+	// Begin AActor Interface
+	virtual void PostUnregisterAllComponents() override;
+protected:
+	virtual void PostRegisterAllComponents() override;
+public:
+	// End AActor Interface
+};
+
+
+
diff --git a/Engine/Source/Runtime/Engine/Classes/Sound/ReverbVolume.h b/Engine/Source/Runtime/Engine/Classes/Sound/ReverbVolume.h
deleted file mode 100644
diff --git a/Engine/Source/Runtime/Engine/Private/ActiveSound.cpp b/Engine/Source/Runtime/Engine/Private/ActiveSound.cpp
--- a/Engine/Source/Runtime/Engine/Private/ActiveSound.cpp
+++ b/Engine/Source/Runtime/Engine/Private/ActiveSound.cpp
@@ -47,7 +47,7 @@ FActiveSound::FActiveSound()
 	, OcclusionCheckInterval(0.f)
 	, LastOcclusionCheckTime(0.f)
 	, LastLocation(FVector::ZeroVector)
-	, LastReverbVolume(NULL)
+	, LastAudioVolume(NULL)
 	, LastUpdateTime(0.f)
 	, SourceInteriorVolume(1.f)
 	, SourceInteriorLPF(1.f)
@@ -66,7 +66,7 @@ FArchive& operator<<( FArchive& Ar, FActiveSound* ActiveSound )
 	if( !Ar.IsLoading() && !Ar.IsSaving() )
 	{
 		Ar << ActiveSound->Sound;
-		Ar << ActiveSound->LastReverbVolume;
+		Ar << ActiveSound->LastAudioVolume;
 		Ar << ActiveSound->WaveInstances;
 		Ar << ActiveSound->SoundNodeOffsetMap;
 	}
@@ -86,7 +86,7 @@ void FActiveSound::AddReferencedObjects( FReferenceCollector& Collector)
 	}
 
 	Collector.AddReferencedObject(Sound);
-	Collector.AddReferencedObject(LastReverbVolume);
+	Collector.AddReferencedObject(LastAudioVolume);
 	Collector.AddReferencedObject(SoundClassOverride);
 }
 
@@ -304,22 +304,22 @@ void FActiveSound::HandleInteriorVolumes( const FListener& Listener, FSoundParse
 {
 	// Get the settings of the ambient sound
 	FInteriorSettings Ambient;
-	class AReverbVolume* ReverbVolume;
+	class AAudioVolume* AudioVolume;
 	if (!bGotInteriorSettings || (GIsEditor && !World->IsGameWorld()) || (ParseParams.Transform.GetTranslation() - LastLocation).SizeSquared() > KINDA_SMALL_NUMBER)
 	{
-		ReverbVolume = World->GetAudioSettings(ParseParams.Transform.GetTranslation(), NULL, &Ambient);
+		AudioVolume = World->GetAudioSettings(ParseParams.Transform.GetTranslation(), NULL, &Ambient);
 		LastInteriorSettings = Ambient;
-		LastReverbVolume = ReverbVolume;
+		LastAudioVolume = AudioVolume;
 		bGotInteriorSettings = true;
 	}
 	else
 	{
 		// use previous settings as we haven't moved
 		Ambient = LastInteriorSettings;
-		ReverbVolume = LastReverbVolume;
+		AudioVolume = LastAudioVolume;
 	}
 
-	// Check to see if we've moved to a new reverb volume
+	// Check to see if we've moved to a new audio volume
 	if( LastUpdateTime < Listener.InteriorStartTime )
 	{
 		SourceInteriorVolume = CurrentInteriorVolume;
@@ -327,7 +327,7 @@ void FActiveSound::HandleInteriorVolumes( const FListener& Listener, FSoundParse
 		LastUpdateTime = FApp::GetCurrentTime();
 	}
 
-	if( Listener.Volume == ReverbVolume || !bAllowSpatialization )
+	if( Listener.Volume == AudioVolume || !bAllowSpatialization )
 	{
 		// Ambient and listener in same ambient zone
 		CurrentInteriorVolume = ( SourceInteriorVolume * ( 1.0f - Listener.InteriorVolumeInterp ) ) + (Listener.InteriorSettings.InteriorVolume * Listener.InteriorVolumeInterp);
@@ -362,7 +362,14 @@ void FActiveSound::HandleInteriorVolumes( const FListener& Listener, FSoundParse
 			ParseParams.VolumeMultiplier *= CurrentInteriorVolume;
 
 			CurrentInteriorLPF = ( SourceInteriorLPF * ( 1.0f - Listener.InteriorLPFInterp ) ) + ( Ambient.InteriorLPF * Listener.InteriorLPFInterp )
+>>>> ORIGINAL //depot/UE4/Engine/Source/Runtime/Engine/Private/ActiveSound.cpp#47
+			float CurrentExteriorLPF = ( SourceInteriorLPF * ( 1.0f - Listener.ExteriorLPFInterp ) ) + ( Listener.InteriorSettings.ExteriorLPF * Listener.ExteriorLPFInterp );
+			ParseParams.HighFrequencyGain *= CurrentInteriorLPF*CurrentExteriorLPF;
+==== THEIRS //depot/UE4/Engine/Source/Runtime/Engine/Private/ActiveSound.cpp#48
 										* ( SourceInteriorLPF * ( 1.0f - Listener.ExteriorLPFInterp ) ) + ( Listener.InteriorSettings.ExteriorLPF * Listener.ExteriorLPFInterp );
+==== YOURS //Marc.Audy_Z2487_Reference/UE4/Engine/Source/Runtime/Engine/Private/ActiveSound.cpp
+										*( SourceInteriorLPF * ( 1.0f - Listener.ExteriorLPFInterp ) ) + ( Listener.InteriorSettings.ExteriorLPF * Listener.ExteriorLPFInterp );
+<<<<
 			ParseParams.HighFrequencyGain *= CurrentInteriorLPF;
 
 			UE_LOG(LogAudio, Verbose, TEXT( "Ambient in diff volume, ambient inside. Volume *= %g LPF *= %g (%s)" ),
diff --git a/Engine/Source/Runtime/Engine/Private/AudioDevice.cpp b/Engine/Source/Runtime/Engine/Private/AudioDevice.cpp
--- a/Engine/Source/Runtime/Engine/Private/AudioDevice.cpp
+++ b/Engine/Source/Runtime/Engine/Private/AudioDevice.cpp
@@ -30,7 +30,7 @@ FAudioDevice::FAudioDevice()
 	, BaseSoundMix(NULL)
 	, DefaultBaseSoundMix(NULL)
 	, Effects(NULL)
-	, CurrentReverbVolume(NULL)
+	, CurrentAudioVolume(NULL)
 	, HighestPriorityReverb(NULL)
 {
 }
@@ -1339,7 +1339,7 @@ void FAudioDevice::InvalidateCachedInteriorVolumes() const
 	}
 }
 
-void FListener::ApplyInteriorSettings( class AReverbVolume* InVolume, const FInteriorSettings& Settings )
+void FListener::ApplyInteriorSettings( class AAudioVolume* InVolume, const FInteriorSettings& Settings )
 {
 	if( InVolume != Volume )
 	{
@@ -1356,7 +1356,7 @@ void FListener::ApplyInteriorSettings( class AReverbVolume* InVolume, const FInt
 }
 
 
-void FAudioDevice::SetListener( const int32 InViewportIndex, const FTransform& InListenerTransform, const float InDeltaSeconds, class AReverbVolume* Volume, const FInteriorSettings& InteriorSettings )
+void FAudioDevice::SetListener( const int32 InViewportIndex, const FTransform& InListenerTransform, const float InDeltaSeconds, class AAudioVolume* Volume, const FInteriorSettings& InteriorSettings )
 {
 	FTransform ListenerTransform = InListenerTransform;
 	
@@ -1556,7 +1556,7 @@ void FAudioDevice::DeactivateReverbEffect(FName TagName)
 	}
 }
 
-void FAudioDevice::SetReverbSettings( class AReverbVolume* Volume, const FReverbSettings& ReverbSettings )
+void FAudioDevice::SetReverbSettings( class AAudioVolume* Volume, const FReverbSettings& ReverbSettings )
 {
 	const FReverbSettings* ActivatedReverb = &ReverbSettings;
 	if (HighestPriorityReverb && (!Volume || HighestPriorityReverb->Priority > Volume->Priority))
diff --git a/Engine/Source/Runtime/Engine/Private/AudioVolume.cpp b/Engine/Source/Runtime/Engine/Private/AudioVolume.cpp
new file mode 100644
--- /dev/null
+++ b/Engine/Source/Runtime/Engine/Private/AudioVolume.cpp
@@ -0,0 +1,259 @@
+// Copyright 1998-2014 Epic Games, Inc. All Rights Reserved.
+
+/*=============================================================================
+	AudioVolume.cpp: Used to affect audio settings in the game and editor.
+=============================================================================*/
+
+#include "EnginePrivate.h"
+#include "AudioDevice.h"
+#include "Net/UnrealNetwork.h"
+#include "Sound/ReverbEffect.h"
+
+void FReverbSettings::PostSerialize(const FArchive& Ar)
+{
+	if( Ar.UE4Ver() < VER_UE4_REVERB_EFFECT_ASSET_TYPE )
+	{
+		FString ReverbAssetName;
+		switch(ReverbType_DEPRECATED)
+		{
+			case REVERB_Default:
+				// No replacement asset for default reverb type
+				return;
+
+			case REVERB_Bathroom:
+				ReverbAssetName = TEXT("/Engine/EngineSounds/ReverbSettings/Bathroom.Bathroom");
+				break;
+
+			case REVERB_StoneRoom:
+				ReverbAssetName = TEXT("/Engine/EngineSounds/ReverbSettings/StoneRoom.StoneRoom");
+				break;
+
+			case REVERB_Auditorium:
+				ReverbAssetName = TEXT("/Engine/EngineSounds/ReverbSettings/Auditorium.Auditorium");
+				break;
+
+			case REVERB_ConcertHall:
+				ReverbAssetName = TEXT("/Engine/EngineSounds/ReverbSettings/ConcertHall.ConcertHall");
+				break;
+
+			case REVERB_Cave:
+				ReverbAssetName = TEXT("/Engine/EngineSounds/ReverbSettings/Cave.Cave");
+				break;
+
+			case REVERB_Hallway:
+				ReverbAssetName = TEXT("/Engine/EngineSounds/ReverbSettings/Hallway.Hallway");
+				break;
+
+			case REVERB_StoneCorridor:
+				ReverbAssetName = TEXT("/Engine/EngineSounds/ReverbSettings/StoneCorridor.StoneCorridor");
+				break;
+
+			case REVERB_Alley:
+				ReverbAssetName = TEXT("/Engine/EngineSounds/ReverbSettings/Alley.Alley");
+				break;
+
+			case REVERB_Forest:
+				ReverbAssetName = TEXT("/Engine/EngineSounds/ReverbSettings/Forest.Forest");
+				break;
+
+			case REVERB_City:
+				ReverbAssetName = TEXT("/Engine/EngineSounds/ReverbSettings/City.City");
+				break;
+
+			case REVERB_Mountains:
+				ReverbAssetName = TEXT("/Engine/EngineSounds/ReverbSettings/Mountains.Mountains");
+				break;
+
+			case REVERB_Quarry:
+				ReverbAssetName = TEXT("/Engine/EngineSounds/ReverbSettings/Quarry.Quarry");
+				break;
+
+			case REVERB_Plain:
+				ReverbAssetName = TEXT("/Engine/EngineSounds/ReverbSettings/Plain.Plain");
+				break;
+
+			case REVERB_ParkingLot:
+				ReverbAssetName = TEXT("/Engine/EngineSounds/ReverbSettings/ParkingLot.ParkingLot");
+				break;
+
+			case REVERB_SewerPipe:
+				ReverbAssetName = TEXT("/Engine/EngineSounds/ReverbSettings/SewerPipe.SewerPipe");
+				break;
+
+			case REVERB_Underwater:
+				ReverbAssetName = TEXT("/Engine/EngineSounds/ReverbSettings/Underwater.Underwater");
+				break;
+
+			case REVERB_SmallRoom:
+				ReverbAssetName = TEXT("/Engine/EngineSounds/ReverbSettings/SmallRoom.SmallRoom");
+				break;
+
+			case REVERB_MediumRoom:
+				ReverbAssetName = TEXT("/Engine/EngineSounds/ReverbSettings/MediumRoom.MediumRoom");
+				break;
+
+			case REVERB_LargeRoom:
+				ReverbAssetName = TEXT("/Engine/EngineSounds/ReverbSettings/LargeRoom.LargeRoom");
+				break;
+
+			case REVERB_MediumHall:
+				ReverbAssetName = TEXT("/Engine/EngineSounds/ReverbSettings/MediumHall.MediumHall");
+				break;
+
+			case REVERB_LargeHall:
+				ReverbAssetName = TEXT("/Engine/EngineSounds/ReverbSettings/LargeHall.LargeHall");
+				break;
+
+			case REVERB_Plate:
+				ReverbAssetName = TEXT("/Engine/EngineSounds/ReverbSettings/Plate.Plate");
+				break;
+
+			default:
+				// This should cover every type of reverb preset
+				checkNoEntry();
+				break;
+		}
+
+		ReverbEffect = LoadObject<UReverbEffect>(NULL, *ReverbAssetName);
+		check( ReverbEffect );
+	}
+}
+
+AAudioVolume::AAudioVolume(const class FPostConstructInitializeProperties& PCIP)
+	: Super(PCIP)
+{
+	BrushComponent->BodyInstance.bEnableCollision_DEPRECATED = false;
+
+	BrushComponent->SetCollisionProfileName(UCollisionProfile::NoCollision_ProfileName);
+	BrushComponent->bAlwaysCreatePhysicsState = true;
+
+	bColored = true;
+	BrushColor = FColor(255, 255, 0, 255);
+
+	bEnabled = true;
+}
+
+void AAudioVolume::GetLifetimeReplicatedProps( TArray< FLifetimeProperty > & OutLifetimeProps ) const
+{
+	Super::GetLifetimeReplicatedProps( OutLifetimeProps );
+	DOREPLIFETIME( AAudioVolume, bEnabled );
+}
+
+void AAudioVolume::PostUnregisterAllComponents( void )
+{
+	// Route clear to super first.
+	Super::PostUnregisterAllComponents();
+
+	// World will be NULL during exit purge.
+	if( GetWorld() )
+	{
+		AAudioVolume* CurrentVolume = GetWorld()->HighestPriorityAudioVolume;
+		AAudioVolume* PreviousVolume = NULL;
+
+		// Iterate over linked list, removing this volume if found.
+		while( CurrentVolume )
+		{
+			// Found.
+			if( CurrentVolume == this )
+			{
+				// Remove from linked list.
+				if( PreviousVolume )
+				{
+					PreviousVolume->NextLowerPriorityVolume = NextLowerPriorityVolume;
+				}
+				else
+				{
+					// Special case removal from first entry.
+					GetWorld()->HighestPriorityAudioVolume = NextLowerPriorityVolume;
+				}
+
+				break;
+			}
+
+			// List traversal.
+			PreviousVolume = CurrentVolume;
+			CurrentVolume = CurrentVolume->NextLowerPriorityVolume;
+		}
+
+		// Reset next pointer to avoid dangling end bits and also for GC.
+		NextLowerPriorityVolume = NULL;
+
+		if (FAudioDevice* AudioDevice = GEngine->GetAudioDevice())
+		{
+			AudioDevice->InvalidateCachedInteriorVolumes();
+		}
+	}
+}
+
+void AAudioVolume::PostRegisterAllComponents()
+{
+	// Route update to super first.
+	Super::PostRegisterAllComponents();
+
+	AAudioVolume* CurrentVolume = GetWorld()->HighestPriorityAudioVolume;
+	AAudioVolume* PreviousVolume = NULL;
+
+	// Find where to insert in sorted linked list.
+	if( CurrentVolume )
+	{
+		// Avoid double insertion!
+		while( CurrentVolume && CurrentVolume != this )
+		{
+			// We use > instead of >= to be sure that we are not inserting twice in the case of multiple volumes having
+			// the same priority and the current one already having being inserted after one with the same priority.
+			if( Priority > CurrentVolume->Priority )
+			{
+				if ( PreviousVolume )
+				{
+					// Insert before current node by fixing up previous to point to current.
+					PreviousVolume->NextLowerPriorityVolume = this;
+				}
+				else
+				{
+					// Special case for insertion at the beginning.
+					GetWorld()->HighestPriorityAudioVolume = this;
+				}
+
+				// Point to current volume, finalizing insertion.
+				NextLowerPriorityVolume = CurrentVolume;
+				return;
+			}
+
+			// List traversal.
+			PreviousVolume = CurrentVolume;
+			CurrentVolume = CurrentVolume->NextLowerPriorityVolume;
+		}
+
+		// We're the lowest priority volume, insert at the end.
+		if( !CurrentVolume )
+		{
+			checkSlow( PreviousVolume );
+			PreviousVolume->NextLowerPriorityVolume = this;
+			NextLowerPriorityVolume = NULL;
+		}
+	}
+	else
+	{
+		// First volume in the world info.
+		GetWorld()->HighestPriorityAudioVolume = this;
+		NextLowerPriorityVolume	= NULL;
+	}
+
+	if (FAudioDevice* AudioDevice = GEngine->GetAudioDevice())
+	{
+		AudioDevice->InvalidateCachedInteriorVolumes();
+	}
+}
+
+#if WITH_EDITOR
+void AAudioVolume::PostEditChangeProperty(FPropertyChangedEvent& PropertyChangedEvent)
+{
+	Super::PostEditChangeProperty(PropertyChangedEvent);
+
+	Settings.Volume = FMath::Clamp<float>( Settings.Volume, 0.0f, 1.0f );
+	AmbientZoneSettings.InteriorTime = FMath::Max<float>( 0.01f, AmbientZoneSettings.InteriorTime );
+	AmbientZoneSettings.InteriorLPFTime = FMath::Max<float>( 0.01f, AmbientZoneSettings.InteriorLPFTime );
+	AmbientZoneSettings.ExteriorTime = FMath::Max<float>( 0.01f, AmbientZoneSettings.ExteriorTime );
+	AmbientZoneSettings.ExteriorLPFTime = FMath::Max<float>( 0.01f, AmbientZoneSettings.ExteriorLPFTime );
+}
+#endif // WITH_EDITOR
diff --git a/Engine/Source/Runtime/Engine/Private/GameViewportClient.cpp b/Engine/Source/Runtime/Engine/Private/GameViewportClient.cpp
--- a/Engine/Source/Runtime/Engine/Private/GameViewportClient.cpp
+++ b/Engine/Source/Runtime/Engine/Private/GameViewportClient.cpp
@@ -758,7 +758,7 @@ void UGameViewportClient::Draw(FViewport* InViewport, FCanvas* SceneCanvas)
 	FAudioDevice* AudioDevice = GEngine->GetAudioDevice();
 	bool bReverbSettingsFound = false;
 	FReverbSettings ReverbSettings;
-	class AReverbVolume* ReverbVolume = NULL;
+	class AAudioVolume* AudioVolume = nullptr;
 
 	for( FConstPlayerControllerIterator Iterator = GetWorld()->GetPlayerControllerIterator(); Iterator; ++Iterator )
 	{
@@ -841,16 +841,16 @@ void UGameViewportClient::Draw(FViewport* InViewport, FCanvas* SceneCanvas)
 
 								FReverbSettings PlayerReverbSettings;
 								FInteriorSettings PlayerInteriorSettings;
-								class AReverbVolume* PlayerReverbVolume = GetWorld()->GetAudioSettings( ViewLocation, &PlayerReverbSettings, &PlayerInteriorSettings );
+								class AAudioVolume* PlayerAudioVolume = GetWorld()->GetAudioSettings( ViewLocation, &PlayerReverbSettings, &PlayerInteriorSettings );
 
-								if (ReverbVolume == NULL || (PlayerReverbVolume != NULL && PlayerReverbVolume->Priority > ReverbVolume->Priority))
+								if (AudioVolume == nullptr || (PlayerAudioVolume != nullptr && PlayerAudioVolume->Priority > AudioVolume->Priority))
 								{
-									ReverbVolume = PlayerReverbVolume;
+									AudioVolume = PlayerAudioVolume;
 									ReverbSettings = PlayerReverbSettings;
 								}
 
 								uint32 ViewportIndex = PlayerViewMap.Num()-1;
-								AudioDevice->SetListener(ViewportIndex, ListenerTransform, (View->bCameraCut ? 0.f : GetWorld()->GetDeltaSeconds()), PlayerReverbVolume, PlayerInteriorSettings);
+								AudioDevice->SetListener(ViewportIndex, ListenerTransform, (View->bCameraCut ? 0.f : GetWorld()->GetDeltaSeconds()), PlayerAudioVolume, PlayerInteriorSettings);
 							}
 					
 						}
@@ -866,7 +866,7 @@ void UGameViewportClient::Draw(FViewport* InViewport, FCanvas* SceneCanvas)
 
 	if (bReverbSettingsFound)
 	{
-		AudioDevice->SetReverbSettings( ReverbVolume, ReverbSettings );
+		AudioDevice->SetReverbSettings( AudioVolume, ReverbSettings );
 	}
 
 	// Update level streaming.
diff --git a/Engine/Source/Runtime/Engine/Private/LevelActor.cpp b/Engine/Source/Runtime/Engine/Private/LevelActor.cpp
--- a/Engine/Source/Runtime/Engine/Private/LevelActor.cpp
+++ b/Engine/Source/Runtime/Engine/Private/LevelActor.cpp
@@ -896,21 +896,12 @@ void UWorld::RefreshStreamingLevels()
 #endif // WITH_EDITOR
 
 
-/**
- * Finds the reverb settings to use for a given view location, taking into account the world's default
- * settings and the reverb volumes in the world.
- *
- * @param	ViewLocation			Current view location.
- * @param	OutReverbSettings		[out] Upon return, the reverb settings for a camera at ViewLocation.
- * @param	OutInteriorSettings		[out] Upon return, the interior settings for a camera at ViewLocation.
- * @return							If the settings came from a reverb volume, the reverb volume's object index is returned.
- */
-AReverbVolume* UWorld::GetAudioSettings( const FVector& ViewLocation, FReverbSettings* OutReverbSettings, FInteriorSettings* OutInteriorSettings )
+AAudioVolume* UWorld::GetAudioSettings( const FVector& ViewLocation, FReverbSettings* OutReverbSettings, FInteriorSettings* OutInteriorSettings )
 {
 	// Find the highest priority volume encompassing the current view location. This is made easier by the linked
 	// list being sorted by priority. @todo: it remains to be seen whether we should trade off sorting for constant
 	// time insertion/ removal time via e.g. TLinkedList.
-	AReverbVolume* Volume = HighestPriorityReverbVolume;
+	AAudioVolume* Volume = HighestPriorityAudioVolume;
 	while( Volume )
 	{
 		// Volume encompasses, break out of loop.
diff --git a/Engine/Source/Runtime/Engine/Private/ReverbVolume.cpp b/Engine/Source/Runtime/Engine/Private/ReverbVolume.cpp
deleted file mode 100644
diff --git a/Engine/Source/Runtime/Engine/Private/UnrealEngine.cpp b/Engine/Source/Runtime/Engine/Private/UnrealEngine.cpp
--- a/Engine/Source/Runtime/Engine/Private/UnrealEngine.cpp
+++ b/Engine/Source/Runtime/Engine/Private/UnrealEngine.cpp
@@ -10995,14 +10995,14 @@ int32 UEngine::RenderStatReverb(UWorld* World, FViewport* Viewport, FCanvas* Can
 			ULocalPlayer* LocalPlayer = GetFirstGamePlayer(World);
 			if (LocalPlayer)
 			{
-				const AReverbVolume* ReverbVolume = AudioDevice->CurrentReverbVolume;
-				if (ReverbVolume && ReverbVolume->Settings.ReverbEffect)
+				const AAudioVolume* AudioVolume = AudioDevice->CurrentAudioVolume;
+				if (AudioVolume && AudioVolume->Settings.ReverbEffect)
 				{
-					TheString = FString::Printf(TEXT("  Reverb Volume Effect: %s (Priority: %g Volume Name: %s)"), *ReverbVolume->Settings.ReverbEffect->GetName(), ReverbVolume->Priority, *ReverbVolume->GetName());
+					TheString = FString::Printf(TEXT("  Audio Volume Reverb Effect: %s (Priority: %g Volume Name: %s)"), *AudioVolume->Settings.ReverbEffect->GetName(), AudioVolume->Priority, *AudioVolume->GetName());
 				}
 				else
 				{
-					TheString = TEXT("  Reverb Volume: None");
+					TheString = TEXT("  Audio Volume Reverb Effect: None");
 				}
 				Canvas->DrawShadowedString(X, Y, *TheString, GetSmallFont(), FLinearColor::White);
 				Y += 12;
diff --git a/Engine/Source/Runtime/Engine/Public/ActiveSound.h b/Engine/Source/Runtime/Engine/Public/ActiveSound.h
--- a/Engine/Source/Runtime/Engine/Public/ActiveSound.h
+++ b/Engine/Source/Runtime/Engine/Public/ActiveSound.h
@@ -2,7 +2,7 @@
 
 #pragma once
 
-#include "Sound/ReverbVolume.h"
+#include "Sound/AudioVolume.h"
 #include "Sound/SoundAttenuation.h"
 
 struct ENGINE_API FActiveSound
@@ -109,7 +109,7 @@ public:
 	/** cache what volume settings we had last time so we don't have to search again if we didn't move */
 	FInteriorSettings LastInteriorSettings;
 
-	class AReverbVolume* LastReverbVolume;
+	class AAudioVolume* LastAudioVolume;
 
 	// To remember where the volumes are interpolating to and from
 	double LastUpdateTime; 
diff --git a/Engine/Source/Runtime/Engine/Public/AudioDevice.h b/Engine/Source/Runtime/Engine/Public/AudioDevice.h
--- a/Engine/Source/Runtime/Engine/Public/AudioDevice.h
+++ b/Engine/Source/Runtime/Engine/Public/AudioDevice.h
@@ -67,7 +67,7 @@ struct FListener
 	struct FInteriorSettings InteriorSettings;
 
 	/** The volume the listener resides in */
-	class AReverbVolume* Volume;
+	class AAudioVolume* Volume;
 
 	/** The times of interior volumes fading in and out */
 	double InteriorStartTime;
@@ -97,7 +97,7 @@ struct FListener
 	/** 
 	 * Apply the interior settings to ambient sounds
 	 */
-	void ApplyInteriorSettings( class AReverbVolume* Volume, const FInteriorSettings& Settings );
+	void ApplyInteriorSettings( class AAudioVolume* Volume, const FInteriorSettings& Settings );
 
 	FListener()
 		: Transform(FTransform::Identity)
@@ -327,17 +327,17 @@ public:
 	 * @param   ListenerIndex		The index of the listener
 	 * @param   ListenerTransform   The listener's world transform
 	 * @param   DeltaSeconds		The amount of time over which velocity should be calculated.  If 0, then velocity will not be calculated.
-	 * @param   Volume				The reverb volume this listener is in
+	 * @param   Volume				The audio volume this listener is in
 	 * @param	ReverbSettings		The reverb settings for this user to use.
 	 * @param	InteriorSettings	The interior settings for this user to use.
 	 */
-	void SetListener( const int32 InListenerIndex, const FTransform& ListenerTransform, const float InDeltaSeconds, class AReverbVolume* Volume, const FInteriorSettings& InteriorSettings );
+	void SetListener( const int32 InListenerIndex, const FTransform& ListenerTransform, const float InDeltaSeconds, class AAudioVolume* Volume, const FInteriorSettings& InteriorSettings );
 
 	/**
 	 * Starts a transition to new reverb and interior settings
 	 *
 	 */
-	void SetReverbSettings( class AReverbVolume* Volume, const FReverbSettings& ReverbSettings );
+	void SetReverbSettings( class AAudioVolume* Volume, const FReverbSettings& ReverbSettings );
 
 	/**
 	 * Creates an audio component to handle playing a sound cue
@@ -500,7 +500,7 @@ public:
 
 	const TArray<FActiveSound*>& GetActiveSounds() const { return ActiveSounds; }
 
-	/* When the set of Reverb volumes have changed invalidate the cached values of active sounds */
+	/* When the set of Audio volumes have changed invalidate the cached values of active sounds */
 	void InvalidateCachedInteriorVolumes() const;
 
     /** Suspend any context related objects */
@@ -775,7 +775,7 @@ public:
 	class FAudioEffectsManager*						Effects;
 
 	/** The volume the listener resides in */
-	const class AReverbVolume*						CurrentReverbVolume;
+	const class AAudioVolume*						CurrentAudioVolume;
 
 	/** Reverb Effects activated without volumes */
 	TMap<FName, FActivatedReverb>					ActivatedReverbs;
diff --git a/Engine/Source/Runtime/Engine/Public/Engine.h b/Engine/Source/Runtime/Engine/Public/Engine.h
--- a/Engine/Source/Runtime/Engine/Public/Engine.h
+++ b/Engine/Source/Runtime/Engine/Public/Engine.h
@@ -40,7 +40,7 @@
 #include "GameFramework/PainCausingVolume.h"
 #include "Interfaces/Interface_PostProcessVolume.h"
 #include "Engine/PostProcessVolume.h"
-#include "Sound/ReverbVolume.h"
+#include "Sound/AudioVolume.h"
 #include "Engine/TriggerVolume.h"
 #include "Camera/CameraActor.h"
 #include "AI/Navigation/NavAgentInterface.h"
diff --git a/Engine/Source/Runtime/Engine/Public/ShowFlags.h b/Engine/Source/Runtime/Engine/Public/ShowFlags.h
--- a/Engine/Source/Runtime/Engine/Public/ShowFlags.h
+++ b/Engine/Source/Runtime/Engine/Public/ShowFlags.h
@@ -94,7 +94,7 @@ struct FEngineShowFlags
 			LocNames.Add(TEXT("PostProcessVolume"), LOCTEXT("PostProcessVolumeSF", "Post Process"));
 			LocNames.Add(TEXT("PrecomputedVisibilityOverrideVolume"), LOCTEXT("PrecomputedVisibilityOverrideVolumeSF", "Precomputed Visibility Override"));
 			LocNames.Add(TEXT("PrecomputedVisibilityVolume"), LOCTEXT("PrecomputedVisibilityVolumeSF", "Precomputed Visibility"));
-			LocNames.Add(TEXT("ReverbVolume"), LOCTEXT("ReverbVolumeSF", "Reverb"));
+			LocNames.Add(TEXT("AudioVolume"), LOCTEXT("AudioVolumeSF", "Audio"));
 			LocNames.Add(TEXT("TriggerVolume"), LOCTEXT("TriggerVolumeSF", "Trigger"));
 			LocNames.Add(TEXT("VisualizeAO"), LOCTEXT("VisualizeAOSF", "Ambient Occlusion"));
 		}
diff --git a/Engine/Source/Runtime/Windows/XAudio2/Private/XAudio2Source.cpp b/Engine/Source/Runtime/Windows/XAudio2/Private/XAudio2Source.cpp
--- a/Engine/Source/Runtime/Windows/XAudio2/Private/XAudio2Source.cpp
+++ b/Engine/Source/Runtime/Windows/XAudio2/Private/XAudio2Source.cpp
@@ -1380,15 +1380,6 @@ void FSpatializationHelper::DumpSpatializationState() const
 	FLocal::DumpChannelArray(TEXT("  "), TEXT("MatrixCoefficients"), ARRAY_COUNT(MatrixCoefficients), MatrixCoefficients);
 }
 
-/**
- * Calculates the spatialized volumes for each channel.
- *
- * @param	OrientFront				The listener's facing direction.
- * @param	ListenerPosition		The position of the listener.
- * @param	EmitterPosition			The position of the emitter.
- * @param	OutVolumes				An array of floats with one volume for each output channel.
- * @param	OutReverbLevel			The reverb volume
- */
 void FSpatializationHelper::CalculateDolbySurroundRate( const FVector& OrientFront, const FVector& ListenerPosition, const FVector& EmitterPosition, float OmniRadius, float* OutVolumes  )
 {
 	uint32 CalculateFlags = X3DAUDIO_CALCULATE_MATRIX | X3DAUDIO_CALCULATE_DOPPLER | X3DAUDIO_CALCULATE_REVERB;
diff --git a/Engine/Source/Runtime/Windows/XAudio2/Private/XAudio2Support.h b/Engine/Source/Runtime/Windows/XAudio2/Private/XAudio2Support.h
--- a/Engine/Source/Runtime/Windows/XAudio2/Private/XAudio2Support.h
+++ b/Engine/Source/Runtime/Windows/XAudio2/Private/XAudio2Support.h
@@ -478,8 +478,8 @@ public:
 	 * @param	OrientFront				The listener's facing direction.
 	 * @param	ListenerPosition		The position of the listener.
 	 * @param	EmitterPosition			The position of the emitter.
+	 * @param	OmniRadius				At what distance we start treating the sound source as spatialized
 	 * @param	OutVolumes				An array of floats with one volume for each output channel.
-	 * @param	OutReverbLevel			The reverb volume
 	 */
 	void CalculateDolbySurroundRate( const FVector& OrientFront, const FVector& ListenerPosition, const FVector& EmitterPosition, float OmniRadius, float* OutVolumes );
 };
