Fix compile error from CL# 2311102

diff --git a/Engine/Source/Runtime/Engine/Private/ActiveSound.cpp b/Engine/Source/Runtime/Engine/Private/ActiveSound.cpp
--- a/Engine/Source/Runtime/Engine/Private/ActiveSound.cpp
+++ b/Engine/Source/Runtime/Engine/Private/ActiveSound.cpp
@@ -362,14 +362,7 @@ void FActiveSound::HandleInteriorVolumes( const FListener& Listener, FSoundParse
 			ParseParams.VolumeMultiplier *= CurrentInteriorVolume;
 
 			CurrentInteriorLPF = ( SourceInteriorLPF * ( 1.0f - Listener.InteriorLPFInterp ) ) + ( Ambient.InteriorLPF * Listener.InteriorLPFInterp )
->>>> ORIGINAL //depot/UE4/Engine/Source/Runtime/Engine/Private/ActiveSound.cpp#47
-			float CurrentExteriorLPF = ( SourceInteriorLPF * ( 1.0f - Listener.ExteriorLPFInterp ) ) + ( Listener.InteriorSettings.ExteriorLPF * Listener.ExteriorLPFInterp );
-			ParseParams.HighFrequencyGain *= CurrentInteriorLPF*CurrentExteriorLPF;
-==== THEIRS //depot/UE4/Engine/Source/Runtime/Engine/Private/ActiveSound.cpp#48
 										* ( SourceInteriorLPF * ( 1.0f - Listener.ExteriorLPFInterp ) ) + ( Listener.InteriorSettings.ExteriorLPF * Listener.ExteriorLPFInterp );
-==== YOURS //Marc.Audy_Z2487_Reference/UE4/Engine/Source/Runtime/Engine/Private/ActiveSound.cpp
-										*( SourceInteriorLPF * ( 1.0f - Listener.ExteriorLPFInterp ) ) + ( Listener.InteriorSettings.ExteriorLPF * Listener.ExteriorLPFInterp );
-<<<<
 			ParseParams.HighFrequencyGain *= CurrentInteriorLPF;
 
 			UE_LOG(LogAudio, Verbose, TEXT( "Ambient in diff volume, ambient inside. Volume *= %g LPF *= %g (%s)" ),
