Make UGameInstance::Init execute for both PIE and Standalone
Redirect current UGameInstance::Init to ::InitializeStandalone and ::InitPie to InitializePIE, both of which are now non-virtual.
Add UGameInstance::Shutdown to allow execution of behaviour on shut down of the game instance (analogous to EndPlay for EndPlayInEditor and Quit on Actors)
Expose blueprint implementable events for Init and Shutdown

diff --git a/Engine/Source/Editor/UnrealEd/Private/PlayLevel.cpp b/Engine/Source/Editor/UnrealEd/Private/PlayLevel.cpp
--- a/Engine/Source/Editor/UnrealEd/Private/PlayLevel.cpp
+++ b/Engine/Source/Editor/UnrealEd/Private/PlayLevel.cpp
@@ -448,6 +448,8 @@ void UEditorEngine::TeardownPlaySession(FWorldContext &PieWorldContext)
 		ActorIt->RouteEndPlay(EEndPlayReason::EndPlayInEditor);
 	}
 
+	PieWorldContext.OwningGameInstance->Shutdown();
+
 	// Move blueprint debugging pointers back to the objects in the editor world
 	PlayWorld->TransferBlueprintDebugReferences(EditorWorld);
 
@@ -2251,7 +2253,7 @@ UGameInstance* UEditorEngine::CreatePIEGameInstance(int32 PIEInstance, bool bInS
 	// We need to temporarily add the GameInstance to the root because the InitPIE call can do garbage collection wiping out the GameInstance
 	GameInstance->AddToRoot();
 
-	bool bSuccess = GameInstance->InitPIE(bAnyBlueprintErrors, PIEInstance);
+	bool bSuccess = GameInstance->InitializePIE(bAnyBlueprintErrors, PIEInstance);
 	if (!bSuccess)
 	{
 		FEditorDelegates::EndPIE.Broadcast(bInSimulateInEditor);
diff --git a/Engine/Source/Runtime/Engine/Classes/Engine/GameInstance.h b/Engine/Source/Runtime/Engine/Classes/Engine/GameInstance.h
--- a/Engine/Source/Runtime/Engine/Classes/Engine/GameInstance.h
+++ b/Engine/Source/Runtime/Engine/Classes/Engine/GameInstance.h
@@ -55,11 +55,25 @@ public:
 	virtual class UWorld* GetWorld() const override;
 	// End UObject Interface
 
-	/** Gives GameInstance an opportunity to set up what it needs */
+	/** virtual function to allow custom GameInstances an opportunity to set up what it needs */
 	virtual void Init();
 
+	/** Opportunity for blueprints to handle the game instance being initialized. */
+	UFUNCTION(BlueprintImplementableEvent, meta=(FriendlyName = "Init"))
+	virtual void ReceiveInit();
+
+	/** virtual function to allow custom GameInstances an opportunity to do cleanup when shutting down */
+	virtual void Shutdown();
+
+	/** Opportunity for blueprints to handle the game instance being shutdown. */
+	UFUNCTION(BlueprintImplementableEvent, meta=(FriendlyName = "Shutdown"))
+	virtual void ReceiveShutdown();
+
+	/* Called to initialize the game instance for standalone instances of the game */
+	void InitializeStandalone();
 #if WITH_EDITOR
-	virtual bool InitPIE(bool bAnyBlueprintErrors, int32 PIEInstance);
+	/* Called to initialize the game instance for PIE instances of the game */
+	bool InitializePIE(bool bAnyBlueprintErrors, int32 PIEInstance);
 
 	bool StartPIEGameInstance(ULocalPlayer* LocalPlayer, bool bInSimulateInEditor, bool bAnyBlueprintErrors, bool bStartInSpectatorMode);
 #endif
diff --git a/Engine/Source/Runtime/Engine/Private/GameEngine.cpp b/Engine/Source/Runtime/Engine/Private/GameEngine.cpp
--- a/Engine/Source/Runtime/Engine/Private/GameEngine.cpp
+++ b/Engine/Source/Runtime/Engine/Private/GameEngine.cpp
@@ -431,7 +431,7 @@ void UGameEngine::Init(IEngineLoop* InEngineLoop)
 		
 		GameInstance = ConstructObject<UGameInstance>(GameInstanceClass, this);
 
-		GameInstance->Init();
+		GameInstance->InitializeStandalone();
 	}
  
 //  	// Creates the initial world context. For GameEngine, this should be the only WorldContext that ever gets created.
@@ -655,6 +655,8 @@ bool UGameEngine::HandleExitCommand( const TCHAR* Cmd, FOutputDevice& Ar )
 		{
 			ActorIt->RouteEndPlay(EEndPlayReason::Quit);
 		}
+
+		World->GetGameInstance()->Shutdown();
 	}
 
 	Ar.Log( TEXT("Closing by request") );
diff --git a/Engine/Source/Runtime/Engine/Private/GameInstance.cpp b/Engine/Source/Runtime/Engine/Private/GameInstance.cpp
--- a/Engine/Source/Runtime/Engine/Private/GameInstance.cpp
+++ b/Engine/Source/Runtime/Engine/Private/GameInstance.cpp
@@ -31,15 +31,27 @@ UEngine* UGameInstance::GetEngine() const
 
 void UGameInstance::Init()
 {
-	UEngine* const Engine = GetEngine();
+	ReceiveInit();
+}
+
+void UGameInstance::Shutdown()
+{
+	ReceiveShutdown();
+}
+
+void UGameInstance::InitializeStandalone()
+{
+	UGameEngine* const Engine = CastChecked<UGameEngine>(GetEngine());
 
 	// Creates the world context. This should be the only WorldContext that ever gets created for this GameInstance.
 	WorldContext = &Engine->CreateNewWorldContext(EWorldType::Game);
 	WorldContext->OwningGameInstance = this;
+
+	Init();
 }
 
 #if WITH_EDITOR
-bool UGameInstance::InitPIE(bool bAnyBlueprintErrors, int32 PIEInstance)
+bool UGameInstance::InitializePIE(bool bAnyBlueprintErrors, int32 PIEInstance)
 {
 	UEditorEngine* const EditorEngine = CastChecked<UEditorEngine>(GetEngine());
 
@@ -98,6 +110,8 @@ bool UGameInstance::InitPIE(bool bAnyBlueprintErrors, int32 PIEInstance)
 	NewWorld->ClearFlags(RF_Standalone);
 	NewWorld->bKismetScriptError = bAnyBlueprintErrors;
 
+	Init();
+
 	return true;
 }
 
