Fix minimum mass being computed for wrong units #codereview(James.Golding)

diff --git a/Engine/Source/Runtime/Engine/Private/PhysicsEngine/BodyInstance.cpp b/Engine/Source/Runtime/Engine/Private/PhysicsEngine/BodyInstance.cpp
--- a/Engine/Source/Runtime/Engine/Private/PhysicsEngine/BodyInstance.cpp
+++ b/Engine/Source/Runtime/Engine/Private/PhysicsEngine/BodyInstance.cpp
@@ -2731,6 +2731,22 @@ int32 GetNumSimShapes(PxRigidDynamic* PRigidDynamic)
 }
 #endif // WITH_PHYSX
 
+float KgPerM3ToKgPerCm3(float KgPerM3)
+{
+	//1m = 100cm => 1m^3 = (100cm)^3 = 1000000cm^3
+	//kg/m^3 = kg/1000000cm^3
+	const float M3ToCm3Inv = 1.f / (100.f * 100.f * 100.f);
+	return KgPerM3 * M3ToCm3Inv;
+}
+
+float gPerCm3ToKgPerCm3(float gPerCm3)
+{
+	//1000g = 1kg
+	//kg/cm^3 = 1000g/cm^3 => g/cm^3 = kg/1000 cm^3
+	const float gToKG = 1.f / 1000.f;
+	return gPerCm3 * gToKG;
+}
+
 void FBodyInstance::UpdateMassProperties()
 {
 	UPhysicalMaterial* PhysMat = GetSimplePhysicalMaterial();
@@ -2743,7 +2759,7 @@ void FBodyInstance::UpdateMassProperties()
 		// First, reset mass to default
 
 		// physical material - nothing can weigh less than hydrogen (0.09 kg/m^3)
-		float DensityKGPerCubicUU = FMath::Max(0.00009f, PhysMat->Density * 0.001f);
+		float DensityKGPerCubicUU = FMath::Max(KgPerM3ToKgPerCm3(0.09f), gPerCm3ToKgPerCm3(PhysMat->Density));
 		PxRigidBodyExt::updateMassAndInertia(*PRigidDynamic, DensityKGPerCubicUU);
 
 		// Then scale mass to avoid big differences between big and small objects.
