From John Barrett: Moved FDummyViewport from UnrealClient.cpp to UnrealClient.h, so it can be used by a plugin.

diff --git a/Engine/Source/Runtime/Engine/Private/UnrealClient.cpp b/Engine/Source/Runtime/Engine/Private/UnrealClient.cpp
--- a/Engine/Source/Runtime/Engine/Private/UnrealClient.cpp
+++ b/Engine/Source/Runtime/Engine/Private/UnrealClient.cpp
@@ -684,67 +684,6 @@ FViewport::FViewport(FViewportClient* InViewportClient):
 	bIsPlayInEditorViewport = false;
 }
 
-// todo: this should be refactored
-class FDummyViewport : public FViewport
-{
-public:
-	FDummyViewport(FViewportClient* InViewportClient)
-		: FViewport(InViewportClient)
-		, DebugCanvas(this, NULL, InViewportClient->GetWorld(), InViewportClient->GetWorld()->FeatureLevel)
-	{
-		DebugCanvas.SetAllowedModes(0);
-	}
-
-	// Begin FViewport interface
-	virtual void BeginRenderFrame(FRHICommandListImmediate& RHICmdList) override
-	{
-		check( IsInRenderingThread() );
-		SetRenderTarget(RHICmdList,  RenderTargetTextureRHI,  FTexture2DRHIRef() );
-	};
-
-	virtual void EndRenderFrame(FRHICommandListImmediate& RHICmdList, bool bPresent, bool bLockToVsync) override
-	{
-		check( IsInRenderingThread() );
-	}
-
-	virtual void*	GetWindow() { return 0; }
-	virtual void	MoveWindow(int32 NewPosX, int32 NewPosY, int32 NewSizeX, int32 NewSizeY) {}
-	virtual void	Destroy() {}
-	virtual bool	CaptureJoystickInput(bool Capture) { return false; }
-	virtual bool	KeyState(FKey Key) const { return false; }
-	virtual int32	GetMouseX() const { return 0; }
-	virtual int32	GetMouseY() const { return 0; }
-	virtual void	GetMousePos( FIntPoint& MousePosition, const bool bLocalPosition = true) { MousePosition = FIntPoint(0, 0); }
-	virtual void	SetMouse(int32 x, int32 y) { }
-	virtual void	ProcessInput( float DeltaTime ) { }
-	virtual FVector2D VirtualDesktopPixelToViewport(FIntPoint VirtualDesktopPointPx) const override { return FVector2D::ZeroVector; }
-	virtual FIntPoint ViewportToVirtualDesktopPixel(FVector2D ViewportCoordinate) const override { return FIntPoint::ZeroValue; }
-	virtual void InvalidateDisplay() { }
-	virtual void DeferInvalidateHitProxy() { }
-	virtual FViewportFrame* GetViewportFrame() { return 0; }
-	virtual FCanvas* GetDebugCanvas() { return &DebugCanvas; }
-	// End FViewport interface
-
-	// Begin FRenderResource interface
-	virtual void InitDynamicRHI()
-	{
-		FTexture2DRHIRef ShaderResourceTextureRHI;
-
-		FRHIResourceCreateInfo CreateInfo;
-		RHICreateTargetableShaderResource2D( SizeX, SizeY, PF_B8G8R8A8, 1, TexCreate_None, TexCreate_RenderTargetable, false, CreateInfo, RenderTargetTextureRHI, ShaderResourceTextureRHI );
-	}
-
-	// @todo UE4 DLL: Without these functions we get unresolved linker errors with FRenderResource
-	virtual void InitRHI() override{}
-	virtual void ReleaseRHI() override{}
-	virtual void InitResource() override{ FViewport::InitResource(); }
-	virtual void ReleaseResource() override { FViewport::ReleaseResource(); }
-	virtual FString GetFriendlyName() const { return FString(TEXT("FDummyViewport"));}
-	// End FRenderResource interface
-private:
-	FCanvas DebugCanvas;
-};
-
 bool FViewport::TakeHighResScreenShot()
 {
 	if( GScreenshotResolutionX == 0 && GScreenshotResolutionY == 0 )
@@ -1680,4 +1619,28 @@ void FCommonViewportClient::DrawHighResScreenshotCaptureRegion(FCanvas& Canvas)
 	LineItem.Draw( &Canvas, FVector2D(Config.UnscaledCaptureRegion.Max.X, Config.UnscaledCaptureRegion.Min.Y), FVector2D(Config.UnscaledCaptureRegion.Max.X, Config.UnscaledCaptureRegion.Max.Y));
 	LineItem.Draw( &Canvas, FVector2D(Config.UnscaledCaptureRegion.Max.X, Config.UnscaledCaptureRegion.Max.Y), FVector2D(Config.UnscaledCaptureRegion.Min.X, Config.UnscaledCaptureRegion.Max.Y));
 	LineItem.Draw( &Canvas, FVector2D(Config.UnscaledCaptureRegion.Min.X, Config.UnscaledCaptureRegion.Max.Y), FVector2D(Config.UnscaledCaptureRegion.Min.X, Config.UnscaledCaptureRegion.Min.Y));
-}
\ No newline at end of file
+}
+
+
+/**
+ * FDummyViewport
+ */
+
+FDummyViewport::FDummyViewport(FViewportClient* InViewportClient)
+	: FViewport(InViewportClient)
+	, DebugCanvas(NULL)
+{
+	UWorld* CurWorld = (InViewportClient != NULL ? InViewportClient->GetWorld() : NULL);
+	DebugCanvas = new FCanvas(this, NULL, CurWorld, (CurWorld != NULL ? CurWorld->FeatureLevel : GRHIFeatureLevel));
+		
+	DebugCanvas->SetAllowedModes(0);
+}
+
+FDummyViewport::~FDummyViewport()
+{
+	if (DebugCanvas != NULL)
+	{
+		delete DebugCanvas;
+		DebugCanvas = NULL;
+	}
+}
diff --git a/Engine/Source/Runtime/Engine/Public/UnrealClient.h b/Engine/Source/Runtime/Engine/Public/UnrealClient.h
--- a/Engine/Source/Runtime/Engine/Public/UnrealClient.h
+++ b/Engine/Source/Runtime/Engine/Public/UnrealClient.h
@@ -899,4 +899,66 @@ public:
 	}
 
 	ENGINE_API void DrawHighResScreenshotCaptureRegion(FCanvas& Canvas);
-};
\ No newline at end of file
+};
+
+
+/**
+ * Minimal viewport for assisting with taking screenshots (also used within a plugin)
+ * @todo: This should be refactored
+ */
+class FDummyViewport : public FViewport
+{
+public:
+	ENGINE_API FDummyViewport(FViewportClient* InViewportClient);
+
+	virtual ~FDummyViewport();
+
+	// Begin FViewport interface
+	virtual void BeginRenderFrame(FRHICommandListImmediate& RHICmdList) override
+	{
+		check( IsInRenderingThread() );
+		SetRenderTarget(RHICmdList,  RenderTargetTextureRHI,  FTexture2DRHIRef() );
+	};
+
+	virtual void EndRenderFrame(FRHICommandListImmediate& RHICmdList, bool bPresent, bool bLockToVsync) override
+	{
+		check( IsInRenderingThread() );
+	}
+
+	virtual void*	GetWindow() { return 0; }
+	virtual void	MoveWindow(int32 NewPosX, int32 NewPosY, int32 NewSizeX, int32 NewSizeY) {}
+	virtual void	Destroy() {}
+	virtual bool	CaptureJoystickInput(bool Capture) { return false; }
+	virtual bool	KeyState(FKey Key) const { return false; }
+	virtual int32	GetMouseX() const { return 0; }
+	virtual int32	GetMouseY() const { return 0; }
+	virtual void	GetMousePos( FIntPoint& MousePosition, const bool bLocalPosition = true) { MousePosition = FIntPoint(0, 0); }
+	virtual void	SetMouse(int32 x, int32 y) { }
+	virtual void	ProcessInput( float DeltaTime ) { }
+	virtual FVector2D VirtualDesktopPixelToViewport(FIntPoint VirtualDesktopPointPx) const override { return FVector2D::ZeroVector; }
+	virtual FIntPoint ViewportToVirtualDesktopPixel(FVector2D ViewportCoordinate) const override { return FIntPoint::ZeroValue; }
+	virtual void InvalidateDisplay() { }
+	virtual void DeferInvalidateHitProxy() { }
+	virtual FViewportFrame* GetViewportFrame() { return 0; }
+	virtual FCanvas* GetDebugCanvas() { return DebugCanvas; }
+	// End FViewport interface
+
+	// Begin FRenderResource interface
+	virtual void InitDynamicRHI()
+	{
+		FTexture2DRHIRef ShaderResourceTextureRHI;
+
+		FRHIResourceCreateInfo CreateInfo;
+		RHICreateTargetableShaderResource2D( SizeX, SizeY, PF_B8G8R8A8, 1, TexCreate_None, TexCreate_RenderTargetable, false, CreateInfo, RenderTargetTextureRHI, ShaderResourceTextureRHI );
+	}
+
+	// @todo UE4 DLL: Without these functions we get unresolved linker errors with FRenderResource
+	virtual void InitRHI() override{}
+	virtual void ReleaseRHI() override{}
+	virtual void InitResource() override{ FViewport::InitResource(); }
+	virtual void ReleaseResource() override { FViewport::ReleaseResource(); }
+	virtual FString GetFriendlyName() const { return FString(TEXT("FDummyViewport"));}
+	// End FRenderResource interface
+private:
+	FCanvas* DebugCanvas;
+};
