BP: Added some initial support for using GUIDs in blueprints (GitHub PR# 467 by james0x0A)

https://github.com/EpicGames/UnrealEngine/pull/467

- did not expose FGuid fields in BPs
- removed ZeroedGuid() because NewGuid() is already zeroed
- made the string conversion function match ToString() behavior
- cosmetic changes

diff --git a/Engine/Source/Runtime/CoreUObject/Classes/Object.h b/Engine/Source/Runtime/CoreUObject/Classes/Object.h
--- a/Engine/Source/Runtime/CoreUObject/Classes/Object.h
+++ b/Engine/Source/Runtime/CoreUObject/Classes/Object.h
@@ -184,7 +184,7 @@ namespace EMouseCursor
 }
 
 // A globally unique identifier.
-USTRUCT(immutable, noexport)
+USTRUCT(immutable, noexport, BlueprintType)
 struct FGuid
 {
 	UPROPERTY(EditAnywhere, SaveGame, Category=Guid)
diff --git a/Engine/Source/Runtime/Engine/Classes/Kismet/KismetGuidLibrary.h b/Engine/Source/Runtime/Engine/Classes/Kismet/KismetGuidLibrary.h
new file mode 100644
--- /dev/null
+++ b/Engine/Source/Runtime/Engine/Classes/Kismet/KismetGuidLibrary.h
@@ -0,0 +1,41 @@
+// Copyright 1998-2014 Epic Games, Inc. All Rights Reserved.
+
+#pragma once
+
+#include "KismetGuidLibrary.generated.h"
+
+
+UCLASS(MinimalAPI)
+class UKismetGuidLibrary
+	: public UBlueprintFunctionLibrary
+{
+	GENERATED_UCLASS_BODY()
+
+	/** Returns true if the values are equal (A == B) */
+	UFUNCTION(BlueprintPure, meta=(FriendlyName="Equal (Guid)", CompactNodeTitle="=="), Category="Guid")
+	static bool EqualEqual_GuidGuid( const FGuid& A, const FGuid& B );
+	
+	/** Returns true if the values are not equal (A != B) */
+	UFUNCTION(BlueprintPure, meta=(FriendlyName="NotEqual (Guid)", CompactNodeTitle="!="), Category="Guid")
+	static bool NotEqual_GuidGuid( const FGuid& A, const FGuid& B );
+
+	/** Checks whether the given GUID is valid */
+	UFUNCTION(BlueprintPure, meta=(FriendlyName="IsValid (Guid)", CompactNodeTitle="isValid?"), Category="Guid")
+	static bool IsValid_Guid( const FGuid& InGuid );
+
+	/** Invalidates the given GUID */
+	UFUNCTION(BlueprintCallable, meta=(FriendlyName="Invalidate_Guid (Guid)"), Category="Guid")
+	static void Invalidate_Guid( UPARAM(ref) FGuid& InGuid );
+
+	/** Returns a new unique GUID */
+	UFUNCTION(BlueprintPure, Category="Guid")
+	static FGuid NewGuid();
+	
+	/** Converts a GUID value to a string, in the form 'A-B-C-D' */
+	UFUNCTION(BlueprintPure, meta=(FriendlyName="ToString (Guid)", CompactNodeTitle="->"), Category="Guid")
+	static FString Conv_GuidToString( const FGuid& InGuid );
+
+	/** Converts a String of format EGuidFormats::Type to a Guid. Returns Guid OutGuid, Returns bool Success */
+	UFUNCTION(BlueprintPure, meta=(FriendlyName="Parse String to Guid"), Category="Guid")
+	static void Parse_StringToGuid( const FString& GuidString, FGuid& OutGuid, bool& Success );
+};
diff --git a/Engine/Source/Runtime/Engine/Private/KismetGuidLibrary.cpp b/Engine/Source/Runtime/Engine/Private/KismetGuidLibrary.cpp
new file mode 100644
--- /dev/null
+++ b/Engine/Source/Runtime/Engine/Private/KismetGuidLibrary.cpp
@@ -0,0 +1,54 @@
+// Copyright 1998-2014 Epic Games, Inc. All Rights Reserved.
+
+#include "EnginePrivate.h"
+#include "Kismet/KismetGuidLibrary.h"
+
+
+/* Guid functions
+ *****************************************************************************/
+
+UKismetGuidLibrary::UKismetGuidLibrary( const class FPostConstructInitializeProperties& PCIP )
+	: Super(PCIP)
+{ }
+
+
+FString UKismetGuidLibrary::Conv_GuidToString( const FGuid& InGuid )
+{
+	return InGuid.ToString(EGuidFormats::Digits);
+}
+
+
+bool UKismetGuidLibrary::EqualEqual_GuidGuid( const FGuid& A, const FGuid& B )
+{
+	return A == B;
+}
+
+
+bool UKismetGuidLibrary::NotEqual_GuidGuid( const FGuid& A, const FGuid& B )
+{
+	return A != B;
+}
+
+
+bool UKismetGuidLibrary::IsValid_Guid( const FGuid& InGuid )
+{
+	return InGuid.IsValid();
+}
+
+
+void UKismetGuidLibrary::Invalidate_Guid( UPARAM(ref) FGuid& InGuid )
+{
+	InGuid.Invalidate();
+}
+
+
+FGuid UKismetGuidLibrary::NewGuid()
+{
+	return FGuid::NewGuid();
+}
+
+
+void UKismetGuidLibrary::Parse_StringToGuid( const FString& GuidString, FGuid& OutGuid, bool& Success )
+{
+	Success = FGuid::Parse(GuidString, OutGuid);
+}
