#ue4 - AController::IsInState() is now const.

diff --git a/Engine/Source/Runtime/Engine/Classes/GameFramework/Controller.h b/Engine/Source/Runtime/Engine/Classes/GameFramework/Controller.h
--- a/Engine/Source/Runtime/Engine/Classes/GameFramework/Controller.h
+++ b/Engine/Source/Runtime/Engine/Classes/GameFramework/Controller.h
@@ -126,7 +126,7 @@ public:
 	 * @param StateName the name of the state to test against
 	 * @return true if current state is StateName
 	 */
-	bool IsInState(FName InStateName);
+	bool IsInState(FName InStateName) const;
 	
 	/** @return the name of the current state */
 	FName GetStateName() const;
diff --git a/Engine/Source/Runtime/Engine/Private/Controller.cpp b/Engine/Source/Runtime/Engine/Private/Controller.cpp
--- a/Engine/Source/Runtime/Engine/Private/Controller.cpp
+++ b/Engine/Source/Runtime/Engine/Private/Controller.cpp
@@ -508,7 +508,7 @@ FName AController::GetStateName() const
 	return StateName;
 }
 
-bool AController::IsInState(FName InStateName)
+bool AController::IsInState(FName InStateName) const
 {
 	return (StateName == InStateName);
 }
