Integrated CL #2311379 from Dev to Main - Fix for NavMesh rendering performance in editor while navmesh rebuilding.

diff --git a/Engine/Source/Editor/UnrealEd/Private/NavigationBuildingNotification.cpp b/Engine/Source/Editor/UnrealEd/Private/NavigationBuildingNotification.cpp
--- a/Engine/Source/Editor/UnrealEd/Private/NavigationBuildingNotification.cpp
+++ b/Engine/Source/Editor/UnrealEd/Private/NavigationBuildingNotification.cpp
@@ -74,6 +74,11 @@ void FNavigationBuildingNotificationImpl::BuildFinished()
 		}
 	}
 
+	if (EEngine && EEngine->IsAnyViewportRealtime() == false)
+	{
+		EEngine->RedrawLevelEditingViewports();
+	}
+
 	FEditorBuildUtils::PathBuildingFinished();
 }
 
diff --git a/Engine/Source/Runtime/Engine/Private/AI/Navigation/RecastNavMesh.cpp b/Engine/Source/Runtime/Engine/Private/AI/Navigation/RecastNavMesh.cpp
--- a/Engine/Source/Runtime/Engine/Private/AI/Navigation/RecastNavMesh.cpp
+++ b/Engine/Source/Runtime/Engine/Private/AI/Navigation/RecastNavMesh.cpp
@@ -1431,13 +1431,6 @@ void ARecastNavMesh::UpdateDrawing()
 	ConstructRenderingComponentImpl();
 #endif
 	UpdateNavMeshDrawing();
-
-#if WITH_EDITOR
-	if (GEditor != NULL)
-	{
-		GEditor->RedrawLevelEditingViewports();
-	}
-#endif // WITH_EDITOR
 }
 
 void ARecastNavMesh::DrawDebugPathCorridor(NavNodeRef const* PathPolys, int32 NumPathPolys, bool bPersistent) const
