#ue4 - APlayerController::PostProcessInput() uses IsLookInputIgnored() virtual rather than reading IgnoreLookInput variable directly.

diff --git a/Engine/Source/Runtime/Engine/Private/PlayerController.cpp b/Engine/Source/Runtime/Engine/Private/PlayerController.cpp
--- a/Engine/Source/Runtime/Engine/Private/PlayerController.cpp
+++ b/Engine/Source/Runtime/Engine/Private/PlayerController.cpp
@@ -2357,7 +2357,7 @@ void APlayerController::PreProcessInput(const float DeltaTime, const bool bGameP
 
 void APlayerController::PostProcessInput(const float DeltaTime, const bool bGamePaused)
 {
-	if( IgnoreLookInput )
+	if( IsLookInputIgnored() )
 	{
 		// zero look inputs
 		RotationInput = FRotator::ZeroRotator;
