Fix for crash in cook in the editor if the level is already loaded and tries to load it's sublevels again.

diff --git a/Engine/Source/Runtime/Engine/Private/LevelActor.cpp b/Engine/Source/Runtime/Engine/Private/LevelActor.cpp
--- a/Engine/Source/Runtime/Engine/Private/LevelActor.cpp
+++ b/Engine/Source/Runtime/Engine/Private/LevelActor.cpp
@@ -766,10 +766,10 @@ void UWorld::LoadSecondaryLevels(bool bForce, TSet<FString>* CookedPackages)
 	{
 		for( int32 LevelIndex=0; LevelIndex<StreamingLevels.Num(); LevelIndex++ )
 		{
-			bool bLoadedLevelPackage = false;
 			ULevelStreaming* const StreamingLevel = StreamingLevels[LevelIndex];
 			if( StreamingLevel )
 			{
+				bool bAlreadyCooked = false;
 				// If we are cooking don't cook sub levels multiple times if they've already been cooked
 				FString PackageFilename;
 				const FString StreamingLevelWorldAssetPackageName = StreamingLevel->GetWorldAssetPackageName();
@@ -778,10 +778,22 @@ void UWorld::LoadSecondaryLevels(bool bForce, TSet<FString>* CookedPackages)
 					if (FPackageName::DoesPackageExist(StreamingLevelWorldAssetPackageName, NULL, &PackageFilename))
 					{
 						PackageFilename = FPaths::ConvertRelativePathToFull(PackageFilename);
+						bAlreadyCooked |= CookedPackages->Contains( PackageFilename );
 					}
 				}
-				if (CookedPackages == NULL || !CookedPackages->Contains(PackageFilename))
+
+
+				bool bAlreadyLoaded = false;
+				UPackage* const LevelPackage = FindObject<UPackage>(NULL, *StreamingLevelWorldAssetPackageName,true);
+				// don't need to do any extra work if the level is already loaded
+				if ( LevelPackage && LevelPackage->IsFullyLoaded() ) 
+				{
+					bAlreadyLoaded = true;
+				}
+
+				if ( !bAlreadyCooked && !bAlreadyLoaded )
 				{
+					bool bLoadedLevelPackage = false;
 					const FName StreamingLevelWorldAssetPackageFName = StreamingLevel->GetWorldAssetPackageFName();
 					// Load the package and find the world object.
 					if( FPackageName::IsShortPackageName(StreamingLevelWorldAssetPackageFName) == false )
@@ -789,7 +801,7 @@ void UWorld::LoadSecondaryLevels(bool bForce, TSet<FString>* CookedPackages)
 						ULevel::StreamedLevelsOwningWorld.Add(StreamingLevelWorldAssetPackageFName, this);
 						UPackage* const LevelPackage = LoadPackage( NULL, *StreamingLevelWorldAssetPackageName, LOAD_None );
 						ULevel::StreamedLevelsOwningWorld.Remove(StreamingLevelWorldAssetPackageFName);
-					
+
 						if( LevelPackage )
 						{
 							bLoadedLevelPackage = true;
@@ -809,14 +821,14 @@ void UWorld::LoadSecondaryLevels(bool bForce, TSet<FString>* CookedPackages)
 								check( LoadedWorld->GetLinker() );
 								LoadedWorld->GetLinker()->Preload( LoadedWorld );
 							}
-							
+
 
 							// Keep reference to prevent garbage collection.
 							check( LoadedWorld->PersistentLevel );
-						
+
 							ULevel* NewLoadedLevel = LoadedWorld->PersistentLevel;
 							NewLoadedLevel->OwningWorld = this;
-					
+
 							StreamingLevel->SetLoadedLevel(NewLoadedLevel);
 						}
 					}
