Fixed expanding NavBounds by a size of a voxel as part of navmesh generation #UE4

diff --git a/Engine/Source/Runtime/Engine/Private/AI/Navigation/RecastNavMeshGenerator.cpp b/Engine/Source/Runtime/Engine/Private/AI/Navigation/RecastNavMeshGenerator.cpp
--- a/Engine/Source/Runtime/Engine/Private/AI/Navigation/RecastNavMeshGenerator.cpp
+++ b/Engine/Source/Runtime/Engine/Private/AI/Navigation/RecastNavMeshGenerator.cpp
@@ -3408,7 +3408,7 @@ void FRecastNavMeshGenerator::SetUpGeneration(float CellSize, float CellHeight,
 	}
 	
 	// expand bounds a bit to support later inclusion tests
-	NavBounds.ExpandBy(CellSize);
+	NavBounds = NavBounds.ExpandBy(CellSize);
 
 	bool bAdjust = false;
 	bool bClampBounds = false;
