#ue4 - Improved documentation for ACharacter::Landed() / OnLanded() indicating that movement mode is still "falling" during this time, and pointing out alternate OnMovementModeChanged() event.

diff --git a/Engine/Source/Runtime/Engine/Classes/GameFramework/Character.h b/Engine/Source/Runtime/Engine/Classes/GameFramework/Character.h
--- a/Engine/Source/Runtime/Engine/Classes/GameFramework/Character.h
+++ b/Engine/Source/Runtime/Engine/Classes/GameFramework/Character.h
@@ -467,10 +467,24 @@ public:
 	/** Called when character's jump reaches Apex. Needs CharacterMovement->bNotifyApex = true */
 	virtual void NotifyJumpApex() {}
 
-	/** Called on landing after falling has completed, to perform actions based on the Hit result. Triggers the OnLanded event. */
+	/**
+	 * Called upon landing when falling, to perform actions based on the Hit result. Triggers the OnLanded event.
+	 * Note that movement mode is still "Falling" during this event. Current Velocity value is the velocity at the time of landing.
+	 * Consider OnMovementModeChanged() as well, as that can be used once the movement mode changes to the new mode (most likely Walking).
+	 *
+	 * @param Hit Result describing the landing that resulted in a valid landing spot.
+	 * @see OnMovementModeChanged()
+	 */
 	virtual void Landed(const FHitResult& Hit);
 
-	/** Called on landing after falling has completed, to perform actions based on the Hit result. */
+	/**
+	* Called upon landing when falling, to perform actions based on the Hit result.
+	* Note that movement mode is still "Falling" during this event. Current Velocity value is the velocity at the time of landing.
+	* Consider OnMovementModeChanged() as well, as that can be used once the movement mode changes to the new mode (most likely Walking).
+	*
+	* @param Hit Result describing the landing that resulted in a valid landing spot.
+	* @see OnMovementModeChanged()
+	*/
 	UFUNCTION(BlueprintImplementableEvent)
 	virtual void OnLanded(const FHitResult& Hit);
 
