Engine: Add missing world->local conversion inside USplineComponent::GetWorldLocationAtSplinePoint

diff --git a/Engine/Source/Runtime/Engine/Private/SplineComponent.cpp b/Engine/Source/Runtime/Engine/Private/SplineComponent.cpp
--- a/Engine/Source/Runtime/Engine/Private/SplineComponent.cpp
+++ b/Engine/Source/Runtime/Engine/Private/SplineComponent.cpp
@@ -233,9 +233,9 @@ void USplineComponent::SetSplineLocalPoints(const TArray<FVector>& Points)
 
 void USplineComponent::SetWorldLocationAtSplinePoint(int32 PointIndex, const FVector& InLocation)
 {
-	if((PointIndex >= 0) && (PointIndex < SplineInfo.Points.Num()))
+	if ((PointIndex >= 0) && (PointIndex < SplineInfo.Points.Num()))
 	{
-		SplineInfo.Points[PointIndex].OutVal = InLocation;
+		SplineInfo.Points[PointIndex].OutVal = ComponentToWorld.InverseTransformPosition(InLocation);
 		UpdateSpline();
 	}
 }
