CouldntSetMaterialUsage warning only in editor

diff --git a/Engine/Source/Runtime/Engine/Private/Materials/Material.cpp b/Engine/Source/Runtime/Engine/Private/Materials/Material.cpp
--- a/Engine/Source/Runtime/Engine/Private/Materials/Material.cpp
+++ b/Engine/Source/Runtime/Engine/Private/Materials/Material.cpp
@@ -960,7 +960,7 @@ bool UMaterial::SetMaterialUsage(bool &bNeedsRecompile, EMaterialUsage Usage, co
 				{
 					UE_LOG(LogMaterial, Warning, TEXT("     The material will recompile every editor launch until resaved."));
 				}
-				else
+				else if (GIsEditor && !FApp::IsGame())
 				{
 #if WITH_EDITOR
 					FFormatNamedArguments Args;
