Added Lightmass world settings option to not compress lightmaps

diff --git a/Engine/Source/Editor/UnrealEd/Private/StaticLightingSystem/StaticLightingSystem.cpp b/Engine/Source/Editor/UnrealEd/Private/StaticLightingSystem/StaticLightingSystem.cpp
--- a/Engine/Source/Editor/UnrealEd/Private/StaticLightingSystem/StaticLightingSystem.cpp
+++ b/Engine/Source/Editor/UnrealEd/Private/StaticLightingSystem/StaticLightingSystem.cpp
@@ -423,6 +423,8 @@ bool FStaticLightingSystem::BeginLightmassProcess()
 		verify(GConfig->GetBool(TEXT("DevOptions.StaticLighting"), TEXT("bRebuildDirtyGeometryForLighting"), bRebuildDirtyGeometryForLighting, GLightmassIni));
 		verify(GConfig->GetBool(TEXT("DevOptions.StaticLighting"), TEXT("bCompressLightmaps"), GCompressLightmaps, GLightmassIni));
 
+		GCompressLightmaps = GCompressLightmaps && World->GetWorldSettings()->LightmassSettings.bCompressLightmaps;
+
 		GAllowLightmapPadding = true;
 		FMemory::Memzero(&LightingMeshBounds, sizeof(FBox));
 		FMemory::Memzero(&AutomaticImportanceVolumeBounds, sizeof(FBox));
diff --git a/Engine/Source/Runtime/Engine/Classes/GameFramework/WorldSettings.h b/Engine/Source/Runtime/Engine/Classes/GameFramework/WorldSettings.h
--- a/Engine/Source/Runtime/Engine/Classes/GameFramework/WorldSettings.h
+++ b/Engine/Source/Runtime/Engine/Classes/GameFramework/WorldSettings.h
@@ -128,6 +128,13 @@ struct FLightmassWorldInfoSettings
 	UPROPERTY(EditAnywhere, Category=LightmassGeneral, AdvancedDisplay, meta=(UIMin = "0.1", UIMax = "100.0"))
 	float VolumeLightSamplePlacementScale;
 
+	/** 
+	 * Whether to compress lightmap textures.  Disabling lightmap texture compression will reduce artifacts but increase memory and disk size by 4x.
+	 * Use caution when disabling this.
+	 */
+	UPROPERTY(EditAnywhere, Category=LightmassGeneral, AdvancedDisplay)
+	uint32 bCompressLightmaps:1;
+
 	FLightmassWorldInfoSettings()
 		: StaticLightingLevelScale(1)
 		, NumIndirectLightingBounces(3)
@@ -146,6 +153,7 @@ struct FLightmassWorldInfoSettings
 		, bVisualizeMaterialDiffuse(false)
 		, bVisualizeAmbientOcclusion(false)
 		, VolumeLightSamplePlacementScale(1)
+		, bCompressLightmaps(true)
 	{
 	}
 };
