Fixed landscape heightmap streaming

diff --git a/Engine/Source/Runtime/Engine/Private/Landscape/LandscapeRender.cpp b/Engine/Source/Runtime/Engine/Private/Landscape/LandscapeRender.cpp
--- a/Engine/Source/Runtime/Engine/Private/Landscape/LandscapeRender.cpp
+++ b/Engine/Source/Runtime/Engine/Private/Landscape/LandscapeRender.cpp
@@ -2795,10 +2795,13 @@ void ULandscapeComponent::GetStreamingTextureInfo(TArray<FStreamingTexturePrimit
 	}
 
 	// Heightmap
-	FStreamingTexturePrimitiveInfo& StreamingHeightmap = *new(OutStreamingTextures) FStreamingTexturePrimitiveInfo;
-	StreamingHeightmap.Bounds = BoundingSphere;
-	StreamingHeightmap.TexelFactor = ForcedLOD >= 0 ? -13+ForcedLOD : TexelFactor; // Minus Value indicate ForcedLOD, 13 for 8k texture
-	StreamingHeightmap.Texture = HeightmapTexture;
+	{
+		float HeightmapTexelFactor = TexelFactor*HeightmapTexture->GetSizeY()/ComponentSizeQuads;
+		FStreamingTexturePrimitiveInfo& StreamingHeightmap = *new(OutStreamingTextures) FStreamingTexturePrimitiveInfo;
+		StreamingHeightmap.Bounds = BoundingSphere;
+		StreamingHeightmap.TexelFactor = ForcedLOD >= 0 ? -13+ForcedLOD : HeightmapTexelFactor; // Minus Value indicate ForcedLOD, 13 for 8k texture
+		StreamingHeightmap.Texture = HeightmapTexture;
+	}
 
 	// XYOffset
 	if (XYOffsetmapTexture)
