Corrected heightmap texel factor calculation to use number of vertices instead of number of quads

diff --git a/Engine/Source/Runtime/Engine/Private/Landscape/LandscapeRender.cpp b/Engine/Source/Runtime/Engine/Private/Landscape/LandscapeRender.cpp
--- a/Engine/Source/Runtime/Engine/Private/Landscape/LandscapeRender.cpp
+++ b/Engine/Source/Runtime/Engine/Private/Landscape/LandscapeRender.cpp
@@ -2796,7 +2796,7 @@ void ULandscapeComponent::GetStreamingTextureInfo(TArray<FStreamingTexturePrimit
 
 	// Heightmap
 	{
-		float HeightmapTexelFactor = TexelFactor*HeightmapTexture->GetSizeY()/ComponentSizeQuads;
+		float HeightmapTexelFactor = TexelFactor*HeightmapTexture->GetSizeY()/(ComponentSizeQuads+1);
 		FStreamingTexturePrimitiveInfo& StreamingHeightmap = *new(OutStreamingTextures) FStreamingTexturePrimitiveInfo;
 		StreamingHeightmap.Bounds = BoundingSphere;
 		StreamingHeightmap.TexelFactor = ForcedLOD >= 0 ? -13+ForcedLOD : HeightmapTexelFactor; // Minus Value indicate ForcedLOD, 13 for 8k texture
