Fix for crash cooking maps with blueprints that change atmosphere parameters in the construction script

diff --git a/Engine/Source/Runtime/Engine/Private/Atmosphere/Atmosphere.cpp b/Engine/Source/Runtime/Engine/Private/Atmosphere/Atmosphere.cpp
--- a/Engine/Source/Runtime/Engine/Private/Atmosphere/Atmosphere.cpp
+++ b/Engine/Source/Runtime/Engine/Private/Atmosphere/Atmosphere.cpp
@@ -547,19 +547,24 @@ void UAtmosphericFogComponent::StartPrecompute()
 		}
 		PrecomputeDataHandler = new FAtmospherePrecomputeDataHandler(this);
 
-		ENQUEUE_UNIQUE_RENDER_COMMAND_TWOPARAMETER(
-			FStartPrecompute,
-			FScene*,Scene,(FScene*)GetScene(),
-			UAtmosphericFogComponent*,Component,this,
+		if (GetScene())
 		{
-			if (Scene && Scene->AtmosphericFog)
+			FAtmosphericFogSceneInfo* AtmosphericFogSceneInfo = GetScene()->GetAtmosphericFogSceneInfo();
+
+			if (AtmosphericFogSceneInfo)
 			{
-				if (Scene->AtmosphericFog->Component == Component)
+				ENQUEUE_UNIQUE_RENDER_COMMAND_TWOPARAMETER(
+					FStartPrecompute,
+					FAtmosphericFogSceneInfo*, AtmosphericFogSceneInfo, AtmosphericFogSceneInfo, 
+					UAtmosphericFogComponent*, Component, this,
 				{
-					Scene->AtmosphericFog->bNeedRecompute = true;
-				}
+					if (AtmosphericFogSceneInfo->Component == Component)
+					{
+						AtmosphericFogSceneInfo->bNeedRecompute = true;
+					}
+				});
 			}
-		});
+		}
 	}
 #endif
 }
diff --git a/Engine/Source/Runtime/Engine/Public/SceneInterface.h b/Engine/Source/Runtime/Engine/Public/SceneInterface.h
--- a/Engine/Source/Runtime/Engine/Public/SceneInterface.h
+++ b/Engine/Source/Runtime/Engine/Public/SceneInterface.h
@@ -7,6 +7,7 @@ class FMaterial;
 class FMaterialShaderMap;
 class ULightComponent;
 class USkyLightComponent;
+class FAtmosphericFogSceneInfo;
 
 /**
  * An interface to the private scene manager implementation of a scene.  Use GetRendererModule().AllocateScene to create.
@@ -168,6 +169,11 @@ public:
 	 */	
 	virtual void RemoveAtmosphericFog(class UAtmosphericFogComponent* FogComponent) = 0;
 
+	/**
+	 * Returns the scene's FAtmosphericFogSceneInfo if it exists
+	 */
+	virtual FAtmosphericFogSceneInfo* GetAtmosphericFogSceneInfo() = 0;
+
 	/**
 	 * Adds a wind source component to the scene.
 	 * @param WindComponent - The component to add.
diff --git a/Engine/Source/Runtime/Renderer/Private/RendererScene.cpp b/Engine/Source/Runtime/Renderer/Private/RendererScene.cpp
--- a/Engine/Source/Runtime/Renderer/Private/RendererScene.cpp
+++ b/Engine/Source/Runtime/Renderer/Private/RendererScene.cpp
@@ -2070,7 +2070,7 @@ public:
 	virtual void RemoveExponentialHeightFog(class UExponentialHeightFogComponent* FogComponent){}
 	virtual void AddAtmosphericFog(class UAtmosphericFogComponent* FogComponent) {}
 	virtual void RemoveAtmosphericFog(class UAtmosphericFogComponent* FogComponent) {}
-
+	virtual FAtmosphericFogSceneInfo* GetAtmosphericFogSceneInfo() override { return NULL; }
 	virtual void AddWindSource(class UWindDirectionalSourceComponent* WindComponent) {}
 	virtual void RemoveWindSource(class UWindDirectionalSourceComponent* WindComponent) {}
 	virtual const TArray<class FWindSourceSceneProxy*>& GetWindSources_RenderThread() const
diff --git a/Engine/Source/Runtime/Renderer/Private/ScenePrivate.h b/Engine/Source/Runtime/Renderer/Private/ScenePrivate.h
--- a/Engine/Source/Runtime/Renderer/Private/ScenePrivate.h
+++ b/Engine/Source/Runtime/Renderer/Private/ScenePrivate.h
@@ -1349,6 +1349,7 @@ public:
 	virtual void RemoveExponentialHeightFog(UExponentialHeightFogComponent* FogComponent);
 	virtual void AddAtmosphericFog(UAtmosphericFogComponent* FogComponent);
 	virtual void RemoveAtmosphericFog(UAtmosphericFogComponent* FogComponent);
+	virtual FAtmosphericFogSceneInfo* GetAtmosphericFogSceneInfo() override { return AtmosphericFog; }
 	virtual void AddWindSource(UWindDirectionalSourceComponent* WindComponent);
 	virtual void RemoveWindSource(UWindDirectionalSourceComponent* WindComponent);
 	virtual const TArray<FWindSourceSceneProxy*>& GetWindSources_RenderThread() const;
