TArray documentation and a coupld of methods deprecation.

diff --git a/Engine/Source/Runtime/Core/Public/Containers/UnrealString.h b/Engine/Source/Runtime/Core/Public/Containers/UnrealString.h
--- a/Engine/Source/Runtime/Core/Public/Containers/UnrealString.h
+++ b/Engine/Source/Runtime/Core/Public/Containers/UnrealString.h
@@ -126,7 +126,7 @@ public:
 			int32 DestLen = FPlatformString::ConvertedLength<TCHAR>(Src, SrcLen);
 			Data.AddUninitialized(DestLen);
 
-			FPlatformString::Convert(Data.GetTypedData(), DestLen, Src, SrcLen);
+			FPlatformString::Convert(Data.GetData(), DestLen, Src, SrcLen);
 		}
 	}
 
@@ -142,7 +142,7 @@ public:
 
 		if( Data.Num() > 0 )
 		{
-			FCString::Strncpy( Data.GetTypedData(), InSrc, InCount+1 );
+			FCString::Strncpy(Data.GetData(), InSrc, InCount + 1);
 		}
 	}
 
@@ -176,7 +176,7 @@ public:
 	 */
 	FORCEINLINE FString& operator=( const TCHAR* Other )
 	{
-		if( Data.GetTypedData() != Other )
+		if (Data.GetData() != Other)
 		{
 			int32 Len = (Other && *Other) ? FCString::Strlen(Other)+1 : 0;
 			Data.Empty(Len);
@@ -253,7 +253,7 @@ public:
 	{
 		int32 Num = Data.Num();
 		checkSlow(Num >= 0);
-		checkSlow(!Num || !Data.GetTypedData()[Num - 1]);
+		checkSlow(!Num || !Data.GetData()[Num - 1]);
 		checkSlow(Data.GetSlack() >= 0);
 	}
 
@@ -293,7 +293,7 @@ public:
 	 */
 	FORCEINLINE const TCHAR* operator*() const
 	{
-		return Data.Num() ? Data.GetTypedData() : TEXT("");
+		return Data.Num() ? Data.GetData() : TEXT("");
 	}
 
 
@@ -346,7 +346,7 @@ public:
 		// Reserve enough space - including an extra gap for a null terminator if we don't already have a string allocated
 		Data.AddUninitialized(Count + (Index ? 0 : 1));
 
-		TCHAR* EndPtr = Data.GetTypedData() + Index - (Index ? 1 : 0);
+		TCHAR* EndPtr = Data.GetData() + Index - (Index ? 1 : 0);
 
 		// Copy characters to end of string, overwriting null terminator if we already have one
 		CopyAssignItems(EndPtr, Array, Count);
@@ -511,7 +511,7 @@ public:
 		CheckInvariants();
 		Str.CheckInvariants();
 
-		AppendChars(Str.Data.GetTypedData(), Str.Len());
+		AppendChars(Str.Data.GetData(), Str.Len());
 
 		return *this;
 	}
@@ -565,7 +565,7 @@ private:
 		int32 RhsLen = Rhs.Len();
 
 		FString Result(MoveTemp(Lhs), RhsLen);
-		Result.AppendChars(Rhs.Data.GetTypedData(), RhsLen);
+		Result.AppendChars(Rhs.Data.GetData(), RhsLen);
 		
 		return Result;
 	}
@@ -588,9 +588,9 @@ private:
 		FString Result;
 		Result.Data.AddUninitialized(LhsLen + RhsLen + 1);
 
-		TCHAR* ResultData = Result.Data.GetTypedData();
+		TCHAR* ResultData = Result.Data.GetData();
 		CopyAssignItems(ResultData, Lhs, LhsLen);
-		CopyAssignItems(ResultData + LhsLen, Rhs.Data.GetTypedData(), RhsLen);
+		CopyAssignItems(ResultData + LhsLen, Rhs.Data.GetData(), RhsLen);
 		*(ResultData + LhsLen + RhsLen) = 0;
 		
 		return Result;
diff --git a/Engine/Source/Runtime/Core/Public/Misc/DataCodec.h b/Engine/Source/Runtime/Core/Public/Misc/DataCodec.h
--- a/Engine/Source/Runtime/Core/Public/Misc/DataCodec.h
+++ b/Engine/Source/Runtime/Core/Public/Misc/DataCodec.h
@@ -53,7 +53,7 @@ public:
 		TArray<int32> CompressPosition;
 		CompressBufferArray.AddUninitialized(MAX_BUFFER_SIZE);
 		CompressPosition   .AddUninitialized(MAX_BUFFER_SIZE+1);
-		CompressBuffer = CompressBufferArray.GetTypedData();
+		CompressBuffer = CompressBufferArray.GetData();
 		int32 i, First=0, Last=0;
 		while( !In.AtEnd() )
 		{
@@ -89,7 +89,7 @@ public:
 			In << DecompressLength << First << Last;
 			check(DecompressLength<=MAX_BUFFER_SIZE+1);
 			check(DecompressLength<=In.TotalSize()-In.Tell());
-			In.Serialize( DecompressBuffer.GetTypedData(), ++DecompressLength );
+			In.Serialize(DecompressBuffer.GetData(), ++DecompressLength);
 			for( i=0; i<257; i++ )
 				DecompressCount[ i ]=0;
 			for( i=0; i<DecompressLength; i++ )
@@ -295,8 +295,8 @@ public:
 		In << Total;
 		TArray<uint8> InArray;
 		InArray.AddUninitialized( In.TotalSize()-In.Tell() );
-		In.Serialize( InArray.GetTypedData(), InArray.Num() );
-		FBitReader Reader( InArray.GetTypedData(), InArray.Num()*8 );
+		In.Serialize(InArray.GetData(), InArray.Num());
+		FBitReader Reader(InArray.GetData(), InArray.Num() * 8);
 		FHuffman Root(-1);
 		Root.ReadTable( Reader );
 		while( Total-- > 0 )
diff --git a/Engine/Source/Runtime/Core/Public/Serialization/Bitstreams.h b/Engine/Source/Runtime/Core/Public/Serialization/Bitstreams.h
--- a/Engine/Source/Runtime/Core/Public/Serialization/Bitstreams.h
+++ b/Engine/Source/Runtime/Core/Public/Serialization/Bitstreams.h
@@ -41,11 +41,11 @@ public:
 	 */
 	FORCEINLINE uint8* GetData(void)
 	{
-		return Buffer.GetTypedData();
+		return Buffer.GetData();
 	}
 	FORCEINLINE const uint8* GetData(void) const
 	{
-		return Buffer.GetTypedData();
+		return Buffer.GetData();
 	}
 	FORCEINLINE const TArray<uint8>* GetBuffer(void) const
 	{
diff --git a/Engine/Source/Runtime/CoreUObject/Private/Serialization/AsyncLoading.cpp b/Engine/Source/Runtime/CoreUObject/Private/Serialization/AsyncLoading.cpp
--- a/Engine/Source/Runtime/CoreUObject/Private/Serialization/AsyncLoading.cpp
+++ b/Engine/Source/Runtime/CoreUObject/Private/Serialization/AsyncLoading.cpp
@@ -403,7 +403,7 @@ void FAsyncPackage::AddImportDependency(int32 CurrentPackageIndex, const FString
 	if (ExistingAsyncPackageIndex == INDEX_NONE)
 	{
 		PackageToStream = new FAsyncPackage(PendingImport, NULL, NAME_None, PendingImport);
-		GObjAsyncPackages.InsertRawItem(PackageToStream, CurrentPackageIndex);
+		GObjAsyncPackages.Insert(PackageToStream, CurrentPackageIndex);
 	}
 	else
 	{
diff --git a/Engine/Source/Runtime/CoreUObject/Private/UObject/Class.cpp b/Engine/Source/Runtime/CoreUObject/Private/UObject/Class.cpp
--- a/Engine/Source/Runtime/CoreUObject/Private/UObject/Class.cpp
+++ b/Engine/Source/Runtime/CoreUObject/Private/UObject/Class.cpp
@@ -2830,7 +2830,7 @@ void UClass::Link(FArchive& Ar, bool bRelinkExistingProperties)
 			}
 		};
 
-		Sort( NetProperties.GetTypedData(), NetProperties.Num(), FCompareUFieldOffsets() );
+		Sort(NetProperties.GetData(), NetProperties.Num(), FCompareUFieldOffsets());
 
 		for ( int32 i = 0; i < NetProperties.Num(); i++ )
 		{
@@ -2850,7 +2850,7 @@ void UClass::Link(FArchive& Ar, bool bRelinkExistingProperties)
 				return A.GetName() < B.GetName();
 			}
 		};
-		Sort( NetFields.GetTypedData(), NetFields.Num(), FCompareUFieldNames() );
+		Sort(NetFields.GetData(), NetFields.Num(), FCompareUFieldNames());
 	}
 }
 
diff --git a/Engine/Source/Runtime/CoreUObject/Private/UObject/GarbageCollection.cpp b/Engine/Source/Runtime/CoreUObject/Private/UObject/GarbageCollection.cpp
--- a/Engine/Source/Runtime/CoreUObject/Private/UObject/GarbageCollection.cpp
+++ b/Engine/Source/Runtime/CoreUObject/Private/UObject/GarbageCollection.cpp
@@ -547,7 +547,7 @@ public:
 
 				if (CurrentIndex != ObjectsToSerialize.Num())
 				{
-					UObject* NextObject = ObjectsToSerialize.GetTypedData()[CurrentIndex]; // special syntax avoiding out of bounds checking
+					UObject* NextObject = ObjectsToSerialize.GetData()[CurrentIndex]; // special syntax avoiding out of bounds checking
 					FPlatformMisc::PrefetchBlock(NextObject, NextObject->GetClass()->GetPropertiesSize());
 				}
 
@@ -563,7 +563,7 @@ public:
 				uint32 ReferenceTokenStreamIndex	= 0;
 
 				// Create stack entry and initialize sane values.
-				FStackEntry* RESTRICT StackEntry = Stack.GetTypedData();
+				FStackEntry* RESTRICT StackEntry = Stack.GetData();
 				uint8* StackEntryData		= (uint8*) CurrentObject;
 				StackEntry->Data			= StackEntryData;
 				StackEntry->Stride			= 0;
@@ -702,7 +702,7 @@ public:
 						UE_LOG(LogGarbage, Fatal,TEXT("Unknown token"));
 					}
 				}
-				check( StackEntry == Stack.GetTypedData() );
+				check(StackEntry == Stack.GetData());
 
 #if PERF_DETAILED_PER_CLASS_GC_STATS
 				// Detailed per class stats should not be performed when parallel GC is running
diff --git a/Engine/Source/Runtime/CoreUObject/Private/UObject/Linker.cpp b/Engine/Source/Runtime/CoreUObject/Private/UObject/Linker.cpp
--- a/Engine/Source/Runtime/CoreUObject/Private/UObject/Linker.cpp
+++ b/Engine/Source/Runtime/CoreUObject/Private/UObject/Linker.cpp
@@ -342,7 +342,7 @@ void ULinker::UpdateScriptSHAKey(const TArray<uint8>& ScriptCode)
 	// if we are doing SHA, update it
 	if (ScriptSHA && ScriptCode.Num())
 	{
-		ScriptSHA->Update((uint8*)ScriptCode.GetTypedData(), ScriptCode.Num());
+		ScriptSHA->Update((uint8*)ScriptCode.GetData(), ScriptCode.Num());
 	}
 }
 
diff --git a/Engine/Source/Runtime/CoreUObject/Private/UObject/ReferenceChainSearch.cpp b/Engine/Source/Runtime/CoreUObject/Private/UObject/ReferenceChainSearch.cpp
--- a/Engine/Source/Runtime/CoreUObject/Private/UObject/ReferenceChainSearch.cpp
+++ b/Engine/Source/Runtime/CoreUObject/Private/UObject/ReferenceChainSearch.cpp
@@ -449,7 +449,7 @@ bool FReferenceChainSearch::ProcessObject( UObject* CurrentObject )
 	Stack.AddUninitialized( 128 );
 
 	// Create strack entry and initialize sane values.
-	FStackEntry* RESTRICT StackEntry = Stack.GetTypedData();
+	FStackEntry* RESTRICT StackEntry = Stack.GetData();
 	uint8* StackEntryData		= (uint8*) CurrentObject;
 	StackEntry->Data			= StackEntryData;
 	StackEntry->Stride			= 0;
@@ -633,7 +633,7 @@ bool FReferenceChainSearch::ProcessObject( UObject* CurrentObject )
 			UE_LOG(LogReferenceChain, Fatal,TEXT("Unknown token"));
 		}
 	}
-	check( StackEntry == Stack.GetTypedData() );
+	check(StackEntry == Stack.GetData());
 	return bRetVal;
 }
 
diff --git a/Engine/Source/Runtime/CoreUObject/Private/UObject/SavePackage.cpp b/Engine/Source/Runtime/CoreUObject/Private/UObject/SavePackage.cpp
--- a/Engine/Source/Runtime/CoreUObject/Private/UObject/SavePackage.cpp
+++ b/Engine/Source/Runtime/CoreUObject/Private/UObject/SavePackage.cpp
@@ -3393,7 +3393,7 @@ bool UPackage::SavePackage( UPackage* InOuter, UObject* Base, EObjectFlags TopLe
 					ScriptSHABytes->AddUninitialized(20);
 
 					// retrieve it
-					Linker->GetScriptSHAKey(ScriptSHABytes->GetTypedData());
+					Linker->GetScriptSHAKey(ScriptSHABytes->GetData());
 				}
 
 				
diff --git a/Engine/Source/Runtime/CoreUObject/Private/UObject/ScriptCore.cpp b/Engine/Source/Runtime/CoreUObject/Private/UObject/ScriptCore.cpp
--- a/Engine/Source/Runtime/CoreUObject/Private/UObject/ScriptCore.cpp
+++ b/Engine/Source/Runtime/CoreUObject/Private/UObject/ScriptCore.cpp
@@ -201,7 +201,7 @@ void FFrame::Serialize( const TCHAR* V, ELogVerbosity::Type Verbosity, const cla
 			V,
 			*Object->GetFullName(),
 			*Node->GetFullName(),
-			Code - Node->Script.GetTypedData(),
+			Code - Node->Script.GetData(),
 			*GetStackTrace()
 		);
 	}
@@ -214,7 +214,7 @@ void FFrame::Serialize( const TCHAR* V, ELogVerbosity::Type Verbosity, const cla
 			V,
 			*Object->GetFullName(),
 			*Node->GetFullName(),
-			Code - Node->Script.GetTypedData(),
+			Code - Node->Script.GetData(),
 			GScriptStackForScriptWarning ? *(FString(TEXT("\r\n")) + GetStackTrace()) : TEXT("")
 		);
 #endif
diff --git a/Engine/Source/Runtime/CoreUObject/Public/UObject/Stack.h b/Engine/Source/Runtime/CoreUObject/Public/UObject/Stack.h
--- a/Engine/Source/Runtime/CoreUObject/Public/UObject/Stack.h
+++ b/Engine/Source/Runtime/CoreUObject/Public/UObject/Stack.h
@@ -171,7 +171,7 @@ public:
 inline FFrame::FFrame( UObject* InObject, UFunction* InNode, void* InLocals, FFrame* InPreviousFrame, UField* InPropertyChainForCompiledIn )
 	: Node(InNode)
 	, Object(InObject)
-	, Code(InNode->Script.GetTypedData())
+	, Code(InNode->Script.GetData())
 	, Locals((uint8*)InLocals)
 	, MostRecentProperty(NULL)
 	, MostRecentPropertyAddress(NULL)
diff --git a/Engine/Source/Runtime/CoreUObject/Public/UObject/UObjectAnnotation.h b/Engine/Source/Runtime/CoreUObject/Public/UObject/UObjectAnnotation.h
--- a/Engine/Source/Runtime/CoreUObject/Public/UObject/UObjectAnnotation.h
+++ b/Engine/Source/Runtime/CoreUObject/Public/UObject/UObjectAnnotation.h
@@ -514,7 +514,7 @@ public:
 				int32 Start = AnnotationArray.AddUninitialized(AddNum);
 				while (AddNum--) 
 				{
-					new (AnnotationArray.GetTypedData() + Start++) TAnnotation();
+					new (AnnotationArray.GetData() + Start++) TAnnotation();
 				}
 			}
 			AnnotationArray[Index] = Annotation;
diff --git a/Engine/Source/Runtime/Engine/Classes/AI/Navigation/RecastNavMesh.h b/Engine/Source/Runtime/Engine/Classes/AI/Navigation/RecastNavMesh.h
--- a/Engine/Source/Runtime/Engine/Classes/AI/Navigation/RecastNavMesh.h
+++ b/Engine/Source/Runtime/Engine/Classes/AI/Navigation/RecastNavMesh.h
@@ -116,7 +116,7 @@ struct ENGINE_API FNavMeshPath : public FNavigationPath
 		}
 
 		float TotalCost = 0.f;
-		const float* Cost = PathCorridorCost.GetTypedData();
+		const float* Cost = PathCorridorCost.GetData();
 		for (int32 PolyIndex = PathPointIndex; PolyIndex < PathCorridorCost.Num(); ++PolyIndex, ++Cost)
 		{
 			TotalCost += *Cost;
@@ -891,8 +891,8 @@ public:
 	/** get information about given tile index */
 	FTileSetItem* GetTileSetItemAt(int32 TileX, int32 TileY);
 	
-	const FTileSetItem* GetTileSet() const { return TileSet.GetTypedData(); }
-	FTileSetItem* GetTileSet() { return TileSet.GetTypedData(); }
+	const FTileSetItem* GetTileSet() const { return TileSet.GetData(); }
+	FTileSetItem* GetTileSet() { return TileSet.GetData(); }
 	
 	virtual bool NeedsRebuild() const override;
 
diff --git a/Engine/Source/Runtime/Engine/Classes/Components/NiagaraComponent.h b/Engine/Source/Runtime/Engine/Classes/Components/NiagaraComponent.h
--- a/Engine/Source/Runtime/Engine/Classes/Components/NiagaraComponent.h
+++ b/Engine/Source/Runtime/Engine/Classes/Components/NiagaraComponent.h
@@ -42,13 +42,13 @@ public:
 	const FVector4 *GetAttributeData(FName Name) const
 	{
 		int32 Offset = AttrMap[Name] * ParticleAllocation;
-		return ParticleBuffers[CurrentBuffer].GetTypedData() + Offset;
+		return ParticleBuffers[CurrentBuffer].GetData() + Offset;
 	}
 
 	FVector4 *GetAttributeDataWrite(FName Name) 
 	{
 		int32 Offset = AttrMap[Name] * ParticleAllocation;
-		return ParticleBuffers[CurrentBuffer].GetTypedData() + Offset;
+		return ParticleBuffers[CurrentBuffer].GetData() + Offset;
 	}
 
 	
@@ -67,8 +67,8 @@ public:
 	void SetNumParticles(uint32 Num)	{ NumParticles = Num;  }
 	void SwapBuffers()					{ CurrentBuffer ^= 0x1; }
 
-	FVector4 *GetCurrentBuffer()		{ return ParticleBuffers[CurrentBuffer].GetTypedData(); }
-	FVector4 *GetPreviousBuffer()		{ return ParticleBuffers[CurrentBuffer^0x1].GetTypedData(); }
+	FVector4 *GetCurrentBuffer()		{ return ParticleBuffers[CurrentBuffer].GetData(); }
+	FVector4 *GetPreviousBuffer()		{ return ParticleBuffers[CurrentBuffer^0x1].GetData(); }
 
 private:
 	uint32 CurrentBuffer, NumParticles, ParticleAllocation;
diff --git a/Engine/Source/Runtime/Engine/Private/AI/Navigation/NavCollision.cpp b/Engine/Source/Runtime/Engine/Private/AI/Navigation/NavCollision.cpp
--- a/Engine/Source/Runtime/Engine/Private/AI/Navigation/NavCollision.cpp
+++ b/Engine/Source/Runtime/Engine/Private/AI/Navigation/NavCollision.cpp
@@ -435,7 +435,7 @@ FByteBulkData* UNavCollision::GetCookedData(FName Format)
 			if (OutData.Num())
 			{
 				Result->Lock(LOCK_READ_WRITE);
-				FMemory::Memcpy(Result->Realloc(OutData.Num()), OutData.GetTypedData(), OutData.Num());
+				FMemory::Memcpy(Result->Realloc(OutData.Num()), OutData.GetData(), OutData.Num());
 				Result->Unlock();
 			}
 		}
diff --git a/Engine/Source/Runtime/Engine/Private/AI/Navigation/NavFilters/NavigationQueryFilter.cpp b/Engine/Source/Runtime/Engine/Private/AI/Navigation/NavFilters/NavigationQueryFilter.cpp
--- a/Engine/Source/Runtime/Engine/Private/AI/Navigation/NavFilters/NavigationQueryFilter.cpp
+++ b/Engine/Source/Runtime/Engine/Private/AI/Navigation/NavFilters/NavigationQueryFilter.cpp
@@ -74,7 +74,7 @@ void FNavigationQueryFilter::SetExcludedArea(uint8 AreaType)
 
 void FNavigationQueryFilter::SetAllAreaCosts(const TArray<float>& CostArray)
 {
-	SetAllAreaCosts(CostArray.GetTypedData(), CostArray.Num());
+	SetAllAreaCosts(CostArray.GetData(), CostArray.Num());
 }
 
 void FNavigationQueryFilter::SetAllAreaCosts(const float* CostArray, const int32 Count)
diff --git a/Engine/Source/Runtime/Engine/Private/AI/Navigation/NavLinkRenderingComponent.cpp b/Engine/Source/Runtime/Engine/Private/AI/Navigation/NavLinkRenderingComponent.cpp
--- a/Engine/Source/Runtime/Engine/Private/AI/Navigation/NavLinkRenderingComponent.cpp
+++ b/Engine/Source/Runtime/Engine/Private/AI/Navigation/NavLinkRenderingComponent.cpp
@@ -107,7 +107,7 @@ FNavLinkRenderingProxy::FNavLinkRenderingProxy(const UPrimitiveComponent* InComp
 
 void FNavLinkRenderingProxy::StorePointLinks(const FTransform& LocalToWorld, const TArray<FNavigationLink>& LinksArray)
 {
-	const FNavigationLink* Link = LinksArray.GetTypedData();
+	const FNavigationLink* Link = LinksArray.GetData();
 	for (int32 LinkIndex = 0; LinkIndex < LinksArray.Num(); ++LinkIndex, ++Link)
 	{	
 		FNavLinkDrawing LinkDrawing;
@@ -122,7 +122,7 @@ void FNavLinkRenderingProxy::StorePointLinks(const FTransform& LocalToWorld, con
 
 void FNavLinkRenderingProxy::StoreSegmentLinks(const FTransform& LocalToWorld, const TArray<FNavigationSegmentLink>& LinksArray)
 {
-	const FNavigationSegmentLink* Link = LinksArray.GetTypedData();
+	const FNavigationSegmentLink* Link = LinksArray.GetData();
 	for (int32 LinkIndex = 0; LinkIndex < LinksArray.Num(); ++LinkIndex, ++Link)
 	{	
 		FNavLinkSegmentDrawing LinkDrawing;
diff --git a/Engine/Source/Runtime/Engine/Private/AI/Navigation/NavigationData.cpp b/Engine/Source/Runtime/Engine/Private/AI/Navigation/NavigationData.cpp
--- a/Engine/Source/Runtime/Engine/Private/AI/Navigation/NavigationData.cpp
+++ b/Engine/Source/Runtime/Engine/Private/AI/Navigation/NavigationData.cpp
@@ -316,7 +316,7 @@ void ANavigationData::PurgeUnusedPaths()
 	check(IsInGameThread());
 
 	const int32 Count = ActivePaths.Num();
-	FNavPathWeakPtr* WeakPathPtr = (ActivePaths.GetTypedData() + Count - 1);
+	FNavPathWeakPtr* WeakPathPtr = (ActivePaths.GetData() + Count - 1);
 	for (int32 i = Count - 1; i >= 0; --i, --WeakPathPtr)
 	{
 		if (WeakPathPtr->IsValid() == false)
diff --git a/Engine/Source/Runtime/Engine/Private/AI/Navigation/NavigationModifier.cpp b/Engine/Source/Runtime/Engine/Private/AI/Navigation/NavigationModifier.cpp
--- a/Engine/Source/Runtime/Engine/Private/AI/Navigation/NavigationModifier.cpp
+++ b/Engine/Source/Runtime/Engine/Private/AI/Navigation/NavigationModifier.cpp
@@ -104,7 +104,7 @@ bool UNavLinkDefinition::HasAdjustableLinks() const
 	}
 
 	bHasDeterminedAdjustableLinks = true;
-	const FNavigationLink* Link = Links.GetTypedData();
+	const FNavigationLink* Link = Links.GetData();
 	for (int32 LinkIndex = 0; LinkIndex < Links.Num(); ++LinkIndex, ++Link)
 	{
 		if (Link->MaxFallDownLength > 0)
@@ -113,7 +113,7 @@ bool UNavLinkDefinition::HasAdjustableLinks() const
 			return true;
 		}
 	}
-	const FNavigationSegmentLink* SegmentLink = SegmentLinks.GetTypedData();
+	const FNavigationSegmentLink* SegmentLink = SegmentLinks.GetData();
 	for (int32 LinkIndex = 0; LinkIndex < SegmentLinks.Num(); ++LinkIndex, ++SegmentLink)
 	{
 		if (SegmentLink->MaxFallDownLength > 0)
@@ -290,7 +290,7 @@ void FAreaNavModifier::SetConvex(const TArray<FVector>& InPoints, const int32 Fi
 
 		// check if there's a similar point already in HullVertices array
 		bool bUnique = true;
-		const FVector* RESTRICT Start = HullVertices.GetTypedData();
+		const FVector* RESTRICT Start = HullVertices.GetData();
 		for (const FVector* RESTRICT Data = Start, * RESTRICT DataEnd = Data + HullVertices.Num(); Data != DataEnd; ++Data)
 		{
 			if (FVector::DistSquared(*Data, TransformedPoint) < CONVEX_HULL_POINTS_MIN_DISTANCE_SQ)
@@ -480,7 +480,7 @@ FCompositeNavModifier FCompositeNavModifier::GetInstantiatedMetaModifier(const F
 	Result = *this;
 
 	{
-		FAreaNavModifier* Area = Result.Areas.GetTypedData();
+		FAreaNavModifier* Area = Result.Areas.GetData();
 		for (int32 Index = 0; Index < Result.Areas.Num(); ++Index, ++Area)
 		{
 			if (Area->HasMetaAreas())
@@ -491,7 +491,7 @@ FCompositeNavModifier FCompositeNavModifier::GetInstantiatedMetaModifier(const F
 	}
 
 	{
-		FSimpleLinkNavModifier* SimpleLink = Result.SimpleLinks.GetTypedData();
+		FSimpleLinkNavModifier* SimpleLink = Result.SimpleLinks.GetData();
 		for (int32 Index = 0; Index < Result.SimpleLinks.Num(); ++Index, ++SimpleLink)
 		{
 			if (SimpleLink->HasMetaAreas())
@@ -563,7 +563,7 @@ uint32 FCompositeNavModifier::GetAllocatedSize() const
 {
 	uint32 MemUsed = Areas.GetAllocatedSize() + SimpleLinks.GetAllocatedSize() + CustomLinks.GetAllocatedSize();
 
-	const FSimpleLinkNavModifier* SimpleLink = SimpleLinks.GetTypedData();
+	const FSimpleLinkNavModifier* SimpleLink = SimpleLinks.GetData();
 	for (int32 Index = 0; Index < SimpleLinks.Num(); ++Index, ++SimpleLink)
 	{
 		MemUsed += SimpleLink->Links.GetAllocatedSize();
diff --git a/Engine/Source/Runtime/Engine/Private/AI/Navigation/NavigationPath.cpp b/Engine/Source/Runtime/Engine/Private/AI/Navigation/NavigationPath.cpp
--- a/Engine/Source/Runtime/Engine/Private/AI/Navigation/NavigationPath.cpp
+++ b/Engine/Source/Runtime/Engine/Private/AI/Navigation/NavigationPath.cpp
@@ -321,7 +321,7 @@ float FNavMeshPath::GetStringPulledLength(const int32 StartingPoint) const
 	}
 
 	float TotalLength = 0.f;
-	const FNavPathPoint* PrevPoint = PathPoints.GetTypedData() + StartingPoint;
+	const FNavPathPoint* PrevPoint = PathPoints.GetData() + StartingPoint;
 	const FNavPathPoint* PathPoint = PrevPoint + 1;
 
 	for (int32 PathPointIndex = StartingPoint + 1; PathPointIndex < PathPoints.Num(); ++PathPointIndex, ++PathPoint, ++PrevPoint)
@@ -344,7 +344,7 @@ float FNavMeshPath::GetPathCorridorLength(const int32 StartingEdge) const
 	}
 
 	
-	const FNavigationPortalEdge* PrevEdge = PathCorridorEdges.GetTypedData() + StartingEdge;
+	const FNavigationPortalEdge* PrevEdge = PathCorridorEdges.GetData() + StartingEdge;
 	const FNavigationPortalEdge* CorridorEdge = PrevEdge + 1;
 	FVector PrevEdgeMiddle = PrevEdge->GetMiddlePoint();
 
@@ -580,7 +580,7 @@ void FNavMeshPath::OffsetFromCorners(float Distance)
 	bool bNeedToCopyResults = false;
 	int32 SingleNodePassCount = 0;
 
-	FNavPathPoint* PathPoint = PathPoints.GetTypedData();
+	FNavPathPoint* PathPoint = PathPoints.GetData();
 	// it's possible we'll be inserting points into the path, so we need to buffer the result
 	TArray<FPathPointInfo> FirstPassPoints;
 	FirstPassPoints.Reserve(PathPoints.Num() + 2);
@@ -796,7 +796,7 @@ void FNavMeshPath::DebugDraw(const ANavigationData* NavData, FColor PathColor, U
 	const ARecastNavMesh* RecastNavMesh = Cast<const ARecastNavMesh>(NavData);
 
 	// tmp hack
-	const FNavigationPortalEdge* Edge = (const_cast<FNavMeshPath*>(this))->GetPathCorridorEdges()->GetTypedData();
+	const FNavigationPortalEdge* Edge = (const_cast<FNavMeshPath*>(this))->GetPathCorridorEdges()->GetData();
 	//const FNavigationPortalEdge* Edge = PathCorridorEdges.GetTypedData();
 	const int32 CorridorEdgesCount = PathCorridorEdges.Num();
 
diff --git a/Engine/Source/Runtime/Engine/Private/AI/Navigation/NavigationSystem.cpp b/Engine/Source/Runtime/Engine/Private/AI/Navigation/NavigationSystem.cpp
--- a/Engine/Source/Runtime/Engine/Private/AI/Navigation/NavigationSystem.cpp
+++ b/Engine/Source/Runtime/Engine/Private/AI/Navigation/NavigationSystem.cpp
@@ -693,7 +693,7 @@ uint32 UNavigationSystem::FindPathAsync(const FNavAgentProperties& AgentProperti
 void UNavigationSystem::AbortAsyncFindPathRequest(uint32 AsynPathQueryID)
 {
 	check(IsInGameThread());
-	FAsyncPathFindingQuery* Query = AsyncPathFindingQueries.GetTypedData();
+	FAsyncPathFindingQuery* Query = AsyncPathFindingQueries.GetData();
 	for (int32 Index = 0; Index < AsyncPathFindingQueries.Num(); ++Index, ++Query)
 	{
 		if (Query->QueryID == AsynPathQueryID)
@@ -730,7 +730,7 @@ void UNavigationSystem::PerformAsyncQueries(TArray<FAsyncPathFindingQuery> PathF
 	}
 	
 	const int32 QueriesCount = PathFindingQueries.Num();
-	FAsyncPathFindingQuery* Query = PathFindingQueries.GetTypedData();
+	FAsyncPathFindingQuery* Query = PathFindingQueries.GetData();
 
 	for (int32 QueryIndex = 0; QueryIndex < QueriesCount; ++QueryIndex, ++Query)
 	{
@@ -1244,7 +1244,7 @@ FBox UNavigationSystem::GetLevelBounds(ULevel* InLevel) const
 
 	if (InLevel)
 	{
-		AActor** Actor = InLevel->Actors.GetTypedData();
+		AActor** Actor = InLevel->Actors.GetData();
 		const int32 ActorCount = InLevel->Actors.Num();
 		for (int32 ActorIndex = 0; ActorIndex < ActorCount; ++ActorIndex, ++Actor)
 		{
@@ -1304,7 +1304,7 @@ void UNavigationSystem::ProcessRegistrationCandidates()
 		return;
 	}
 
-	ANavigationData** NavDataPtr = NavDataRegistrationQueue.GetTypedData();
+	ANavigationData** NavDataPtr = NavDataRegistrationQueue.GetData();
 	const int CandidatesCount = NavDataRegistrationQueue.Num();
 
 	for (int32 CandidateIndex = 0; CandidateIndex < CandidatesCount; ++CandidateIndex, ++NavDataPtr)
@@ -1388,7 +1388,7 @@ UNavigationSystem::ERegistrationResult UNavigationSystem::RegisterNavData(ANavig
 		{
 			// ok, so this navigation agent doesn't have its navmesh registered yet, but do we want to support it?
 			bool bAgentSupported = false;
-			FNavDataConfig* Agent = SupportedAgents.GetTypedData();
+			FNavDataConfig* Agent = SupportedAgents.GetData();
 			for (int32 AgentIndex = 0; AgentIndex < SupportedAgents.Num(); ++AgentIndex, ++Agent)
 			{
 				if (Agent->IsEquivalent(NavConfig) == true)
@@ -2465,7 +2465,7 @@ void UNavigationSystem::SpawnMissingNavigationData()
 		if (Nav != NULL && Nav->GetTypedOuter<UWorld>() == NavWorld && Nav->IsPendingKill() == false)
 		{
 			// find out which one it is
-			const FNavDataConfig* AgentProps = SupportedAgents.GetTypedData();
+			const FNavDataConfig* AgentProps = SupportedAgents.GetData();
 			for (int32 AgentIndex = 0; AgentIndex < SupportedAgentsCount; ++AgentIndex, ++AgentProps)
 			{
 				if (AlreadyInstantiated[AgentIndex] == true)
diff --git a/Engine/Source/Runtime/Engine/Private/AI/Navigation/PImplRecastNavMesh.cpp b/Engine/Source/Runtime/Engine/Private/AI/Navigation/PImplRecastNavMesh.cpp
--- a/Engine/Source/Runtime/Engine/Private/AI/Navigation/PImplRecastNavMesh.cpp
+++ b/Engine/Source/Runtime/Engine/Private/AI/Navigation/PImplRecastNavMesh.cpp
@@ -917,7 +917,7 @@ void FPImplRecastNavMesh::PostProcessPath(dtStatus FindPathStatus, FNavMeshPath&
 		
 		// copy over corridor poly data
 		Path.PathCorridor.AddUninitialized(PathResult.size());
-		NavNodeRef* DestCorridorPoly = Path.PathCorridor.GetTypedData();
+		NavNodeRef* DestCorridorPoly = Path.PathCorridor.GetData();
 		for (int i = 0; i < PathResult.size(); ++i, ++DestCorridorPoly)
 		{
 			*DestCorridorPoly = PathResult.getRef(i);
@@ -992,7 +992,7 @@ bool FPImplRecastNavMesh::FindStraightPath(const FVector& StartLoc, const FVecto
 
 	dtQueryResult StringPullResult;
 	const dtStatus StringPullStatus = NavQuery.findStraightPath(&RecastStartPos.X, &RecastEndPos.X,
-		PathCorridor.GetTypedData(), PathCorridor.Num(), StringPullResult, DT_STRAIGHTPATH_AREA_CROSSINGS);
+		PathCorridor.GetData(), PathCorridor.Num(), StringPullResult, DT_STRAIGHTPATH_AREA_CROSSINGS);
 
 	PathPoints.Reset();
 	if (dtStatusSucceed(StringPullStatus))
@@ -1000,7 +1000,7 @@ bool FPImplRecastNavMesh::FindStraightPath(const FVector& StartLoc, const FVecto
 		PathPoints.AddZeroed(StringPullResult.size());
 
 		// convert to desired format
-		FNavPathPoint* CurVert = PathPoints.GetTypedData();
+		FNavPathPoint* CurVert = PathPoints.GetData();
 
 		for (int32 VertIdx = 0; VertIdx < StringPullResult.size(); ++VertIdx)
 		{
@@ -1418,7 +1418,7 @@ bool FPImplRecastNavMesh::GetPolysWithinPathingDistance(FVector const& StartLoc,
 	int32 NumPolys;
 
 	dtStatus Status = NavQuery.findPolysInPathDistance(StartPolyID, RecastStartPos
-		, PathingDistance, QueryFilter, FoundPolys.GetTypedData(), &NumPolys, Filter.GetMaxSearchNodes());
+		, PathingDistance, QueryFilter, FoundPolys.GetData(), &NumPolys, Filter.GetMaxSearchNodes());
 
 	FoundPolys.RemoveAt(NumPolys, FoundPolys.Num() - NumPolys);
 
@@ -2248,7 +2248,7 @@ void FPImplRecastNavMesh::GetNavMeshTilesAt(int32 TileX, int32 TileY, TArray<int
 		TArray<const dtMeshTile*> Tiles;
 		Tiles.AddZeroed(MaxTiles);
 
-		const int32 NumTiles = ConstRecastNavMesh->getTilesAt(TileX, TileY, Tiles.GetTypedData(), MaxTiles);
+		const int32 NumTiles = ConstRecastNavMesh->getTilesAt(TileX, TileY, Tiles.GetData(), MaxTiles);
 		for (int32 i = 0; i < NumTiles; i++)
 		{
 			dtTileRef TileRef = ConstRecastNavMesh->getTileRef(Tiles[i]);
diff --git a/Engine/Source/Runtime/Engine/Private/AI/Navigation/RecastNavMesh.cpp b/Engine/Source/Runtime/Engine/Private/AI/Navigation/RecastNavMesh.cpp
--- a/Engine/Source/Runtime/Engine/Private/AI/Navigation/RecastNavMesh.cpp
+++ b/Engine/Source/Runtime/Engine/Private/AI/Navigation/RecastNavMesh.cpp
@@ -318,7 +318,7 @@ void ARecastNavMesh::TickMe(float DeltaTime)
 		TNavStatArray<FNavMeshGenerationResult> AsyncResults;
 		((FRecastNavMeshGenerator*)NavDataGenerator.Get())->GetAsyncResultsCopy(AsyncResults, /*bClearSource=*/true);
 
-		FNavMeshGenerationResult* Result = AsyncResults.GetTypedData();
+		FNavMeshGenerationResult* Result = AsyncResults.GetData();
 		for (int32 ResultIndex = 0; ResultIndex < AsyncResults.Num(); ++ResultIndex, ++Result)
 		{
 			UpdateNavTileData(Result->OldRawNavData, Result->NewNavData);
@@ -333,7 +333,7 @@ void ARecastNavMesh::TickMe(float DeltaTime)
 		FreshTiles.Reset();
 		const int32 PathsCount = ActivePaths.Num();
 
-		FNavPathWeakPtr* WeakPathPtr = (ActivePaths.GetTypedData() + PathsCount - 1);
+		FNavPathWeakPtr* WeakPathPtr = (ActivePaths.GetData() + PathsCount - 1);
 		for (int32 PathIndex = PathsCount - 1; PathIndex >= 0; --PathIndex, --WeakPathPtr)
 		{
 			FNavPathSharedPtr SharedPath = WeakPathPtr->Pin();
@@ -352,7 +352,7 @@ void ARecastNavMesh::TickMe(float DeltaTime)
 				}
 
 				const int32 PathLenght = Path->PathCorridor.Num();
-				const NavNodeRef* PathPoly = Path->PathCorridor.GetTypedData();
+				const NavNodeRef* PathPoly = Path->PathCorridor.GetData();
 				for (int32 NodeIndex = 0; NodeIndex < PathLenght; ++NodeIndex, ++PathPoly)
 				{
 					const uint32 NodeTileIdx = RecastNavMeshImpl->GetTileIndexFromPolyRef(*PathPoly);
@@ -509,7 +509,7 @@ ANavigationData* ARecastNavMesh::CreateNavigationInstances(UNavigationSystem* Na
 			if (Nav != NULL && Nav->GetTypedOuter<UWorld>() == NavSys->GetWorld() && Nav->IsPendingKill() == false)
 			{
 				// find out which one it is
-				const FNavDataConfig* AgentProps = SupportedAgents->GetTypedData();
+				const FNavDataConfig* AgentProps = SupportedAgents->GetData();
 				for (int i = 0; i < SupportedAgentsCount; ++i, ++AgentProps)
 				{
 					if (AlreadyInstantiated[i] == true)
@@ -1819,12 +1819,12 @@ void ARecastNavMesh::UpdateNavTileData(uint8* OldTileData, FNavMeshTileData NewT
 		int32 OldDataSize = 0;
 
 		// find shared pointer holding OldTileData
-		const FNavMeshTileData* const RESTRICT DataEnd = TileNavDataContainer.GetTypedData() + TileNavDataContainer.Num();
-		for(const FNavMeshTileData* RESTRICT Data = TileNavDataContainer.GetTypedData(); Data < DataEnd; ++Data)
+		const FNavMeshTileData* const RESTRICT DataEnd = TileNavDataContainer.GetData() + TileNavDataContainer.Num();
+		for(const FNavMeshTileData* RESTRICT Data = TileNavDataContainer.GetData(); Data < DataEnd; ++Data)
 		{
 			if( *Data == OldTileData )
 			{
-				OldDataIndex = (int32)(Data - TileNavDataContainer.GetTypedData());
+				OldDataIndex = (int32)(Data - TileNavDataContainer.GetData());
 				OldDataSize = Data->DataSize;
 				break;
 			}
diff --git a/Engine/Source/Runtime/Engine/Private/AI/Navigation/RecastNavMeshGenerator.cpp b/Engine/Source/Runtime/Engine/Private/AI/Navigation/RecastNavMeshGenerator.cpp
--- a/Engine/Source/Runtime/Engine/Private/AI/Navigation/RecastNavMeshGenerator.cpp
+++ b/Engine/Source/Runtime/Engine/Private/AI/Navigation/RecastNavMeshGenerator.cpp
@@ -213,7 +213,7 @@ static void ExportGeomToOBJFile(const FString& InFileName, const TNavStatArray<f
 			CompressedBuffer.Init(0, HeaderSize + FMath::Trunc(1.1f * UncompressedSize));
 
 			int32 CompressedSize = CompressedBuffer.Num() - HeaderSize;
-			uint8* DestBuffer = CompressedBuffer.GetTypedData();
+			uint8* DestBuffer = CompressedBuffer.GetData();
 			FMemory::Memcpy(DestBuffer, &UncompressedSize, HeaderSize);
 			DestBuffer += HeaderSize;
 
@@ -407,14 +407,14 @@ static void StoreCollisionCache(FRecastGeometryExport* GeomExport)
 	GeomExport->Data->CollisionData.AddUninitialized(CacheSize);
 
 	// store collisions
-	uint8* RawMemory = GeomExport->Data->CollisionData.GetTypedData();
+	uint8* RawMemory = GeomExport->Data->CollisionData.GetData();
 	FRecastGeometryCache* CacheMemory = (FRecastGeometryCache*)RawMemory;
 	CacheMemory->Header = HeaderInfo;
 	CacheMemory->Verts = 0;
 	CacheMemory->Indices = 0;
 
-	FMemory::Memcpy(RawMemory + HeaderSize, GeomExport->VertexBuffer.GetTypedData(), CoordsSize);
-	FMemory::Memcpy(RawMemory + HeaderSize + CoordsSize, GeomExport->IndexBuffer.GetTypedData(), IndicesSize);
+	FMemory::Memcpy(RawMemory + HeaderSize, GeomExport->VertexBuffer.GetData(), CoordsSize);
+	FMemory::Memcpy(RawMemory + HeaderSize + CoordsSize, GeomExport->IndexBuffer.GetData(), IndicesSize);
 }
 
 #if WITH_PHYSX
@@ -711,7 +711,7 @@ FORCEINLINE_DEBUGGABLE void ExportRigidBodyConvexElements(UBodySetup& BodySetup,
 {
 #if WITH_PHYSX
 	const int32 ConvexCount = BodySetup.AggGeom.ConvexElems.Num();
-	FKConvexElem const * ConvexElem = BodySetup.AggGeom.ConvexElems.GetTypedData();
+	FKConvexElem const * ConvexElem = BodySetup.AggGeom.ConvexElems.GetData();
 
 	for(int32 i=0; i< ConvexCount; ++i, ++ConvexElem)
 	{
@@ -995,7 +995,7 @@ FORCEINLINE void TransformVertexSoupToRecast(const TArray<FVector>& VertexSoup,
 
 	const int32 StaticFacesCount = VertexSoup.Num() / 3;
 	int32 VertsCount = Verts.Num();
-	const FVector* Vertex = VertexSoup.GetTypedData();
+	const FVector* Vertex = VertexSoup.GetData();
 
 	for (int32 k = 0; k < StaticFacesCount; ++k, Vertex += 3)
 	{
@@ -1012,7 +1012,7 @@ FORCEINLINE void TransformVertexSoupToRecast(const TArray<FVector>& VertexSoup,
 
 FORCEINLINE void CovertCoordDataToRecast(TNavStatArray<float>& Coords)
 {
-	float* CoordPtr = Coords.GetTypedData();
+	float* CoordPtr = Coords.GetData();
 	const int32 MaxIt = Coords.Num() / 3;
 	for (int32 i = 0; i < MaxIt; i++)
 	{
@@ -1663,7 +1663,7 @@ void FRecastTileGenerator::Init(class FRecastNavMeshGenerator* ParentGenerator,
 
 		InclusionBounds.Reserve(BoundingBoxes.Num());
 
-		const FBox* Bounds = BoundingBoxes.GetTypedData();	
+		const FBox* Bounds = BoundingBoxes.GetData();	
 		for (int32 i = 0; i < BoundingBoxes.Num() && bFullyEncapsulatedByInclusionBounds == false; ++i, ++Bounds)
 		{	
 			if (Bounds->Intersect( TileBB ))
@@ -1753,7 +1753,7 @@ void FRecastTileGenerator::ApplyVoxelFilter(struct rcHeightfield* HF, float Walk
 		// it doesn't create "fake cliffs"
 		const float ExpandBBBy = WalkableRadius*CellSize;
 
-		const FBox* BBox = InclusionBounds.GetTypedData();
+		const FBox* BBox = InclusionBounds.GetData();
 		// optimized common case of single box
 		if (InclusionBounds.Num() == 1)
 		{
@@ -1821,7 +1821,7 @@ void FRecastTileGenerator::ApplyVoxelFilter(struct rcHeightfield* HF, float Walk
 							const FVector SpanMaxV(SpanX, SpanY, SpanMax);
 
 							bool bIsInsideAnyBB = false;
-							const FBox* BB = Bounds.GetTypedData();
+							const FBox* BB = Bounds.GetData();
 							for (int32 BoundIndex = 0; BoundIndex < BoundsCount; ++BoundIndex, ++BB)
 							{
 								if (BB->IsInside(SpanMinV) || BB->IsInside(SpanMaxV))
@@ -1846,7 +1846,7 @@ void FRecastTileGenerator::ApplyVoxelFilter(struct rcHeightfield* HF, float Walk
 
 void FRecastTileGenerator::PrepareVoxelCache(const TNavStatArray<uint8>& RawCollisionCache, TNavStatArray<rcSpanCache>& SpanData)
 {
-	FRecastGeometryCache CachedCollisions(RawCollisionCache.GetTypedData());
+	FRecastGeometryCache CachedCollisions(RawCollisionCache.GetData());
 	VoxelCacheContext.SetupForTile(BMin, BMax, RasterizationPadding);
 
 	float SlopeCosPerActor = WalkableSlopeCos;
@@ -1859,23 +1859,23 @@ void FRecastTileGenerator::PrepareVoxelCache(const TNavStatArray<uint8>& RawColl
 	rcMarkWalkableTrianglesCos(0, SlopeCosPerActor,
 		CachedCollisions.Verts, CachedCollisions.Header.NumVerts,
 		CachedCollisions.Indices, CachedCollisions.Header.NumFaces,
-		TriAreas.GetTypedData());
+		TriAreas.GetData());
 
 	rcRasterizeTriangles(0, CachedCollisions.Verts, CachedCollisions.Header.NumVerts,
-		CachedCollisions.Indices, TriAreas.GetTypedData(), CachedCollisions.Header.NumFaces,
+		CachedCollisions.Indices, TriAreas.GetData(), CachedCollisions.Header.NumFaces,
 		*VoxelCacheContext.RasterizeHF, WalkableClimbVX);
 
 	const int32 NumSpans = rcCountSpans(0, *VoxelCacheContext.RasterizeHF);
 	if (NumSpans > 0)
 	{
 		SpanData.AddZeroed(NumSpans);
-		rcCacheSpans(0, *VoxelCacheContext.RasterizeHF, SpanData.GetTypedData());
+		rcCacheSpans(0, *VoxelCacheContext.RasterizeHF, SpanData.GetData());
 	}
 }
 
 bool FRecastTileGenerator::HasVoxelCache(const TNavStatArray<uint8>& RawVoxelCache, rcSpanCache*& CachedVoxels, int32& NumCachedVoxels) const
 {
-	FRecastVoxelCache VoxelCache(RawVoxelCache.GetTypedData());
+	FRecastVoxelCache VoxelCache(RawVoxelCache.GetData());
 	for (FRecastVoxelCache::FTileInfo* iTile = VoxelCache.Tiles; iTile; iTile = iTile->NextTile)
 	{
 		if (iTile->TileX == TileX && iTile->TileY == TileY)
@@ -1896,7 +1896,7 @@ void FRecastTileGenerator::AddVoxelCache(TNavStatArray<uint8>& RawVoxelCache, co
 		RawVoxelCache.AddZeroed(sizeof(int32));
 	}
 
-	int32* NumTiles = (int32*)RawVoxelCache.GetTypedData();
+	int32* NumTiles = (int32*)RawVoxelCache.GetData();
 	*NumTiles = *NumTiles + 1;
 
 	const int32 NewCacheIdx = RawVoxelCache.Num();
@@ -1905,12 +1905,12 @@ void FRecastTileGenerator::AddVoxelCache(TNavStatArray<uint8>& RawVoxelCache, co
 	const int32 EntrySize = HeaderSize + VoxelsSize;
 	RawVoxelCache.AddZeroed(EntrySize);
 
-	FRecastVoxelCache::FTileInfo* TileInfo = (FRecastVoxelCache::FTileInfo*)(RawVoxelCache.GetTypedData() + NewCacheIdx);
+	FRecastVoxelCache::FTileInfo* TileInfo = (FRecastVoxelCache::FTileInfo*)(RawVoxelCache.GetData() + NewCacheIdx);
 	TileInfo->TileX = TileX;
 	TileInfo->TileY = TileY;
 	TileInfo->NumSpans = NumCachedVoxels;
 
-	FMemory::Memcpy(RawVoxelCache.GetTypedData() + NewCacheIdx + HeaderSize, CachedVoxels, VoxelsSize);
+	FMemory::Memcpy(RawVoxelCache.GetData() + NewCacheIdx + HeaderSize, CachedVoxels, VoxelsSize);
 }
 
 void FRecastTileGenerator::ClearGeometry()
@@ -1978,7 +1978,7 @@ void FRecastTileGenerator::AppendModifier(const FCompositeNavModifier& Modifier,
 #endif
 
 	// evaluate custom links
-	const FCustomLinkNavModifier* LinkModifier = Modifier.GetCustomLinks().GetTypedData();
+	const FCustomLinkNavModifier* LinkModifier = Modifier.GetCustomLinks().GetData();
 	for (int32 i = 0; i < Modifier.GetCustomLinks().Num(); i++, LinkModifier++)
 	{
 		FSimpleLinkNavModifier SimpleLinkCollection(UNavLinkDefinition::GetLinksDefinition(LinkModifier->GetNavLinkClass()), LinkModifier->LocalToWorld);
@@ -1997,7 +1997,7 @@ void FRecastTileGenerator::AppendGeometry(const TNavStatArray<uint8>& RawCollisi
 		return;
 	}
 
-	FRecastGeometryCache CollisionCache(RawCollisionCache.GetTypedData());
+	FRecastGeometryCache CollisionCache(RawCollisionCache.GetData());
 
 #if CACHE_NAV_GENERATOR_DATA
 	const int32 FirstNewCoord = GeomCoords.Num();
@@ -2007,8 +2007,8 @@ void FRecastTileGenerator::AppendGeometry(const TNavStatArray<uint8>& RawCollisi
 	GeomCoords.AddZeroed(CollisionCache.Header.NumVerts * 3);
 	GeomIndices.AddZeroed(CollisionCache.Header.NumFaces * 3);
 
-	FMemory::Memcpy(GeomCoords.GetTypedData() + FirstNewCoord, CollisionCache.Verts, sizeof(float) * CollisionCache.Header.NumVerts * 3);
-	int32* DestIndex = GeomIndices.GetTypedData() + FirstNewIndex;
+	FMemory::Memcpy(GeomCoords.GetData() + FirstNewCoord, CollisionCache.Verts, sizeof(float) * CollisionCache.Header.NumVerts * 3);
+	int32* DestIndex = GeomIndices.GetData() + FirstNewIndex;
 #else
 	const int32 FirstNewCoord = StaticGeomCoords.Num();
 	const int32 FirstNewIndex = StaticGeomIndices.Num();
@@ -2017,8 +2017,8 @@ void FRecastTileGenerator::AppendGeometry(const TNavStatArray<uint8>& RawCollisi
 	StaticGeomCoords.AddZeroed(CollisionCache.Header.NumVerts * 3);
 	StaticGeomIndices.AddZeroed(CollisionCache.Header.NumFaces * 3);
 
-	FMemory::Memcpy(StaticGeomCoords.GetTypedData() + FirstNewCoord, CollisionCache.Verts, sizeof(float) * CollisionCache.Header.NumVerts * 3);
-	int32* DestIndex = StaticGeomIndices.GetTypedData() + FirstNewIndex;
+	FMemory::Memcpy(StaticGeomCoords.GetData() + FirstNewCoord, CollisionCache.Verts, sizeof(float) * CollisionCache.Header.NumVerts * 3);
+	int32* DestIndex = StaticGeomIndices.GetData() + FirstNewIndex;
 #endif
 
 	for (int32 i = 0; i < CollisionCache.Header.NumFaces * 3; i++, DestIndex++)
@@ -2042,8 +2042,8 @@ void FRecastTileGenerator::AppendGeometry(const TNavStatArray<FVector>& Verts, c
 	GeomCoords.AddZeroed(Verts.Num() * 3);
 	GeomIndices.AddZeroed(Faces.Num());
 
-	float* DestCoord = GeomCoords.GetTypedData() + FirstNewCoord;
-	int32* DestIndex = GeomIndices.GetTypedData() + FirstNewIndex;
+	float* DestCoord = GeomCoords.GetData() + FirstNewCoord;
+	int32* DestIndex = GeomIndices.GetData() + FirstNewIndex;
 #else
 	const int32 FirstNewCoord = StaticGeomCoords.Num();
 	const int32 FirstNewIndex = StaticGeomIndices.Num();
@@ -2052,8 +2052,8 @@ void FRecastTileGenerator::AppendGeometry(const TNavStatArray<FVector>& Verts, c
 	StaticGeomCoords.AddZeroed(Verts.Num() * 3);
 	StaticGeomIndices.AddZeroed(Faces.Num());
 
-	float* DestCoord = StaticGeomCoords.GetTypedData() + FirstNewCoord;
-	int32* DestIndex = StaticGeomIndices.GetTypedData() + FirstNewIndex;
+	float* DestCoord = StaticGeomCoords.GetData() + FirstNewCoord;
+	int32* DestIndex = StaticGeomIndices.GetData() + FirstNewIndex;
 #endif
 
 	for (int32 i = 0; i < Verts.Num(); i++)
@@ -2077,7 +2077,7 @@ void FRecastTileGenerator::AppendVoxels(rcSpanCache* SpanData, int32 NumSpans)
 		const int32 FirstSpanIdx = StaticGeomSpans.Num();
 		StaticGeomSpans.AddZeroed(NumSpans);
 
-		FMemory::Memcpy(StaticGeomSpans.GetTypedData() + FirstSpanIdx, SpanData, NumSpans * sizeof(rcSpanCache));
+		FMemory::Memcpy(StaticGeomSpans.GetData() + FirstSpanIdx, SpanData, NumSpans * sizeof(rcSpanCache));
 #endif
 	}
 }
@@ -2286,19 +2286,19 @@ bool FRecastTileGenerator::GenerateCompressedLayers(class FNavMeshBuildContext*
 		TriAreas.AddZeroed(NumFaces);
 
 		rcMarkWalkableTriangles(BuildContext, TileConfig.walkableSlopeAngle,
-			PtrGeomCoords->GetTypedData(), NumVerts, PtrGeomIndices->GetTypedData(), NumFaces,
-			TriAreas.GetTypedData());
+			PtrGeomCoords->GetData(), NumVerts, PtrGeomIndices->GetData(), NumFaces,
+			TriAreas.GetData());
 
 		rcRasterizeTriangles(BuildContext,
-			PtrGeomCoords->GetTypedData(), NumVerts, 
-			PtrGeomIndices->GetTypedData(), TriAreas.GetTypedData(), NumFaces,
+			PtrGeomCoords->GetData(), NumVerts, 
+			PtrGeomIndices->GetData(), TriAreas.GetData(), NumFaces,
 			*RasterContext.SolidHF, TileConfig.walkableClimb);
 	}
 	else if (PtrGeomSpans->Num())
 	{
 		DECLARE_SCOPE_CYCLE_COUNTER(TEXT("Rasterization: apply voxel cache"), Stat_RecastRasterCacheApply, STATGROUP_Navigation);
 
-		rcAddSpans(BuildContext, *RasterContext.SolidHF, TileConfig.walkableClimb, PtrGeomSpans->GetTypedData(), PtrGeomSpans->Num());
+		rcAddSpans(BuildContext, *RasterContext.SolidHF, TileConfig.walkableClimb, PtrGeomSpans->GetData(), PtrGeomSpans->Num());
 	}
 
 	if (!RasterContext.SolidHF || RasterContext.SolidHF->pools == 0)
@@ -2500,7 +2500,7 @@ void FRecastTileGenerator::MarkStaticAreas(class FNavMeshBuildContext* BuildCont
 	RECAST_STAT(STAT_Navigation_Async_MarkAreas);
 	const float ExpandBy = TileConfig.AgentRadius * 1.5;
 
-	const FAreaNavModifier* Modifier = PtrStaticAreas->GetTypedData();
+	const FAreaNavModifier* Modifier = PtrStaticAreas->GetData();
 	for (int32 ModifierIndex = 0; ModifierIndex < NumAreas; ++ModifierIndex, ++Modifier)
 	{
 		const int32* AreaID = AdditionalCachedDataPtr->AreaClassToIdMap.Find(Modifier->GetAreaClass());
@@ -2556,7 +2556,7 @@ void FRecastTileGenerator::MarkStaticAreas(class FNavMeshBuildContext* BuildCont
 					TArray<float> ConvexCoords;
 					ConvexCoords.AddZeroed(ConvexVerts.Num() * 3);
 
-					float* ItCoord = ConvexCoords.GetTypedData();
+					float* ItCoord = ConvexCoords.GetData();
 					for (int32 i = 0; i < ConvexVerts.Num(); i++)
 					{
 						const FVector RecastV = Unreal2RecastPoint(ConvexVerts[i]);
@@ -2565,7 +2565,7 @@ void FRecastTileGenerator::MarkStaticAreas(class FNavMeshBuildContext* BuildCont
 						*ItCoord = RecastV.Z; ItCoord++;
 					}
 
-					rcMarkConvexPolyArea(BuildContext, ConvexCoords.GetTypedData(), ConvexVerts.Num(),
+					rcMarkConvexPolyArea(BuildContext, ConvexCoords.GetData(), ConvexVerts.Num(),
 						ConvexData.MinZ, ConvexData.MaxZ, *AreaID, CompactHF);
 				}
 			}
@@ -2798,7 +2798,7 @@ bool FRecastTileGenerator::GenerateNavigationData(class FNavMeshBuildContext* Bu
 				OffMeshData.FlagsPerArea = AdditionalCachedDataPtr->FlagsPerOffMeshLinkArea;
 				const uint32 AgentMask = (1 << TileConfig.AgentIndex);
 
-				const FSimpleLinkNavModifier* LinkModifier = PtrOffmeshLinks->GetTypedData();
+				const FSimpleLinkNavModifier* LinkModifier = PtrOffmeshLinks->GetData();
 				for (int32 LinkModifierIndex = 0; LinkModifierIndex < PtrOffmeshLinks->Num(); ++LinkModifierIndex, ++LinkModifier)
 				{
 					OffMeshData.AddLinks(LinkModifier->Links, LinkModifier->LocalToWorld, AgentMask);
@@ -2832,7 +2832,7 @@ bool FRecastTileGenerator::GenerateNavigationData(class FNavMeshBuildContext* Bu
 				Params.detailTris = GenerationContext.DetailMesh->tris;
 				Params.detailTriCount = GenerationContext.DetailMesh->ntris;
 			}
-			Params.offMeshCons = OffMeshData.LinkParams.GetTypedData();
+			Params.offMeshCons = OffMeshData.LinkParams.GetData();
 			Params.offMeshConCount = OffMeshData.LinkParams.Num();
 			Params.walkableHeight = TileConfig.AgentHeight;
 			Params.walkableRadius = TileConfig.AgentRadius;
@@ -2910,7 +2910,7 @@ void FRecastTileGenerator::MarkDynamicAreas(struct dtTileCacheLayer& Layer, cons
 	const float* LayerRecastOrig = Layer.header->bmin;
 	const FBox LayerUnrealBounds = Recast2UnrealBox(Layer.header->bmin, Layer.header->bmax);
 
-	const FAreaNavModifier* Modifier = CombinedAreas.GetTypedData();
+	const FAreaNavModifier* Modifier = CombinedAreas.GetData();
 	for (int32 ModifierIndex = 0; ModifierIndex < CombinedAreas.Num(); ++ModifierIndex, ++Modifier)
 	{
 		const int32* AreaID = AdditionalCachedDataPtr->AreaClassToIdMap.Find(Modifier->GetAreaClass());
@@ -2980,7 +2980,7 @@ void FRecastTileGenerator::MarkDynamicAreas(struct dtTileCacheLayer& Layer, cons
 					TArray<float> ConvexCoords;
 					ConvexCoords.AddZeroed(ConvexVerts.Num() * 3);
 
-					float* ItCoord = ConvexCoords.GetTypedData();
+					float* ItCoord = ConvexCoords.GetData();
 					for (int32 i = 0; i < ConvexVerts.Num(); i++)
 					{
 						const FVector RecastV = Unreal2RecastPoint(ConvexVerts[i]);
@@ -2990,7 +2990,7 @@ void FRecastTileGenerator::MarkDynamicAreas(struct dtTileCacheLayer& Layer, cons
 					}
 
 					dtMarkConvexArea(Layer, LayerRecastOrig, TileConfig.cs, TileConfig.ch,
-						ConvexCoords.GetTypedData(), ConvexVerts.Num(), ConvexData.MinZ, ConvexData.MaxZ, *AreaID);
+						ConvexCoords.GetData(), ConvexVerts.Num(), ConvexData.MinZ, ConvexData.MaxZ, *AreaID);
 				}
 			}
 			break;
@@ -3010,7 +3010,7 @@ uint32 FRecastTileGenerator::GetUsedMemCount() const
 	TotalMemory += GeomCoords.GetAllocatedSize();
 	TotalMemory += GeomIndices.GetAllocatedSize();
 	
-	const FSimpleLinkNavModifier* SimpleLink = OffmeshLinks.GetTypedData();
+	const FSimpleLinkNavModifier* SimpleLink = OffmeshLinks.GetData();
 	for (int32 Index = 0; Index < OffmeshLinks.Num(); ++Index, ++SimpleLink)
 	{
 		TotalMemory += SimpleLink->Links.GetAllocatedSize();
@@ -3561,7 +3561,7 @@ bool FRecastNavMeshGenerator::Generate()
 
 void FRecastNavMeshGenerator::RebuildDirtyAreas(const TArray<FNavigationDirtyArea>& InDirtyAreas)
 {
-	const FNavigationDirtyArea* DirtyArea = InDirtyAreas.GetTypedData();
+	const FNavigationDirtyArea* DirtyArea = InDirtyAreas.GetData();
 	bool bOverlaps = false;
 	for (int32 DirtyIndex = 0; DirtyIndex < InDirtyAreas.Num(); ++DirtyIndex, ++DirtyArea)
 	{
@@ -3801,7 +3801,7 @@ void FRecastNavMeshGenerator::RemoveTileLayers(const int32 TileX, const int32 Ti
 
 	TArray<dtMeshTile*> Tiles;
 	Tiles.AddZeroed(NumLayers);
-	DetourMesh->getTilesAt(TileX, TileY, (const dtMeshTile**)Tiles.GetTypedData(), NumLayers);
+	DetourMesh->getTilesAt(TileX, TileY, (const dtMeshTile**)Tiles.GetData(), NumLayers);
 	
 	for (int32 i = 0; i < NumLayers; i++)
 	{
@@ -4025,7 +4025,7 @@ void FRecastNavMeshGenerator::MarkDirtyGenerators()
 	TSet<int32> DirtyIndices;
 
 	// find all tiles that need regeneration:
-	const FNavigationDirtyArea* DirtyArea = DirtyAreasCopy.GetTypedData();
+	const FNavigationDirtyArea* DirtyArea = DirtyAreasCopy.GetData();
 	for (int32 i = 0; i < DirtyAreasCopy.Num(); ++i, ++DirtyArea)
 	{
 		const FBox AdjustedBounds = GrowBoundingBox(DirtyArea->Bounds, DirtyArea->HasFlag(ENavigationDirtyFlag::UseAgentHeight));
@@ -4235,7 +4235,7 @@ void FRecastNavMeshGenerator::FillGeneratorData(FRecastTileGenerator* TileGenera
 					{
 						// rasterize
 						TileGenerator->PrepareVoxelCache(Element.Data.CollisionData, SpanData);
-						CachedVoxels = SpanData.GetTypedData();
+						CachedVoxels = SpanData.GetData();
 						NumCachedVoxels = SpanData.Num();
 
 						// encode
@@ -4279,7 +4279,7 @@ void FRecastNavMeshGenerator::UpdateTileGenerationWorkers(int32 TileId)
 {
 #if RECAST_ASYNC_REBUILDING
 	bool bRequestRenderingDirty = false;
-	FRecastTileGenerator** CurrentGenerator = ActiveGenerators.GetTypedData();
+	FRecastTileGenerator** CurrentGenerator = ActiveGenerators.GetData();
 	int32 QueueIndex = 0;
 
 	for (int32 i = 0; i < ActiveGenerators.Num(); ++i, ++CurrentGenerator)
@@ -4476,7 +4476,7 @@ bool FRecastNavMeshGenerator::AreAnyTilesBeingBuilt(bool bCheckDirtyToo) const
 {
 #if RECAST_ASYNC_REBUILDING	
 	bool bRet = false;
-	const FRecastTileGenerator* TileGenerator = TileGenerators.GetTypedData();
+	const FRecastTileGenerator* TileGenerator = TileGenerators.GetData();
 
 	for (int32 i = 0; i < TileGenerators.Num(); ++i, ++TileGenerator)
 	{
@@ -4498,7 +4498,7 @@ bool FRecastNavMeshGenerator::AreAnyTilesBeingBuilt(bool bCheckDirtyToo) const
 bool FRecastNavMeshGenerator::IsAsyncBuildInProgress() const
 {
 #if RECAST_ASYNC_REBUILDING	
-	const FRecastTileGenerator* TileGenerator = TileGenerators.GetTypedData();
+	const FRecastTileGenerator* TileGenerator = TileGenerators.GetData();
 	for (int32 i = 0; i < TileGenerators.Num(); ++i, ++TileGenerator)
 	{
 		if (TileGenerator->IsAsyncBuildInProgress())
@@ -4592,7 +4592,7 @@ void FRecastNavMeshGenerator::ExportNavigationData(const FString& FileName) cons
 
 				if (bExportGeometry && Element.Data.CollisionData.Num())
 				{
-					FRecastGeometryCache CachedGeometry(Element.Data.CollisionData.GetTypedData());
+					FRecastGeometryCache CachedGeometry(Element.Data.CollisionData.GetData());
 					IndexBuffer.Reserve( IndexBuffer.Num() + (CachedGeometry.Header.NumFaces * 3 ));
 					CoordBuffer.Reserve( CoordBuffer.Num() + (CachedGeometry.Header.NumVerts * 3 ));
 					for (int32 i = 0; i < CachedGeometry.Header.NumFaces * 3; i++)
diff --git a/Engine/Source/Runtime/Engine/Private/ActorReplication.cpp b/Engine/Source/Runtime/Engine/Private/ActorReplication.cpp
--- a/Engine/Source/Runtime/Engine/Private/ActorReplication.cpp
+++ b/Engine/Source/Runtime/Engine/Private/ActorReplication.cpp
@@ -273,7 +273,7 @@ void AActor::GetSubobjectsWithStableNamesForNetworking(TArray<UObject*> &ObjList
 		}
 	};
 
-	Sort( ObjList.GetTypedData(), ObjList.Num(), FCompareComponentNames() );
+	Sort( ObjList.GetData(), ObjList.Num(), FCompareComponentNames() );
 }
 
 void AActor::OnSubobjectCreatedFromReplication(UObject *NewSubobject)
diff --git a/Engine/Source/Runtime/Engine/Private/Animation/AnimCompress.cpp b/Engine/Source/Runtime/Engine/Private/Animation/AnimCompress.cpp
--- a/Engine/Source/Runtime/Engine/Private/Animation/AnimCompress.cpp
+++ b/Engine/Source/Runtime/Engine/Private/Animation/AnimCompress.cpp
@@ -18,7 +18,7 @@ DEFINE_LOG_CATEGORY(LogAnimationCompression);
 #define AC_UnalignedWriteToStream( Src, Len )										\
 	{																				\
 		const int32 Ofs = Seq->CompressedByteStream.AddUninitialized( Len );						\
-		FMemory::Memcpy( Seq->CompressedByteStream.GetTypedData()+Ofs, (Src), (Len) );	\
+		FMemory::Memcpy( Seq->CompressedByteStream.GetData()+Ofs, (Src), (Len) );	\
 	}
 
 static const uint8 AnimationPadSentinel = 85; //(1<<1)+(1<<3)+(1<<5)+(1<<7)
diff --git a/Engine/Source/Runtime/Engine/Private/Animation/AnimCompress_PerTrackCompression.cpp b/Engine/Source/Runtime/Engine/Private/Animation/AnimCompress_PerTrackCompression.cpp
--- a/Engine/Source/Runtime/Engine/Private/Animation/AnimCompress_PerTrackCompression.cpp
+++ b/Engine/Source/Runtime/Engine/Private/Animation/AnimCompress_PerTrackCompression.cpp
@@ -96,7 +96,7 @@ protected:
 	void AppendBytes(const void* Data, int32 Length)
 	{
 		const int32 Offset = CompressedBytes.AddUninitialized(Length);
-		FMemory::Memcpy(CompressedBytes.GetTypedData() + Offset, Data, Length);
+		FMemory::Memcpy(CompressedBytes.GetData() + Offset, Data, Length);
 	}
 
 	void CompressTranslation_Identity(const FTranslationTrack& TranslationData)
@@ -122,7 +122,7 @@ protected:
 		const int32 NumKeys = TranslationData.PosKeys.Num();
 
 		// Determine the bounds
-		const FBox KeyBounds(TranslationData.PosKeys.GetTypedData(), NumKeys);
+		const FBox KeyBounds(TranslationData.PosKeys.GetData(), NumKeys);
 		const bool bHasX = (FMath::Abs(KeyBounds.Max.X) >= ZeroingThreshold) || (FMath::Abs(KeyBounds.Min.X) >= ZeroingThreshold);
 		const bool bHasY = (FMath::Abs(KeyBounds.Max.Y) >= ZeroingThreshold) || (FMath::Abs(KeyBounds.Min.Y) >= ZeroingThreshold);
 		const bool bHasZ = (FMath::Abs(KeyBounds.Max.Z) >= ZeroingThreshold) || (FMath::Abs(KeyBounds.Min.Z) >= ZeroingThreshold);
@@ -179,7 +179,7 @@ protected:
 		const int32 NumKeys = TranslationData.PosKeys.Num();
 
 		// Determine the bounds
-		const FBox KeyBounds(TranslationData.PosKeys.GetTypedData(), NumKeys);
+		const FBox KeyBounds(TranslationData.PosKeys.GetData(), NumKeys);
 		const bool bHasX = (FMath::Abs(KeyBounds.Max.X) >= ZeroingThreshold) || (FMath::Abs(KeyBounds.Min.X) >= ZeroingThreshold);
 		const bool bHasY = (FMath::Abs(KeyBounds.Max.Y) >= ZeroingThreshold) || (FMath::Abs(KeyBounds.Min.Y) >= ZeroingThreshold);
 		const bool bHasZ = (FMath::Abs(KeyBounds.Max.Z) >= ZeroingThreshold) || (FMath::Abs(KeyBounds.Min.Z) >= ZeroingThreshold);
@@ -225,7 +225,7 @@ protected:
 		const int32 NumKeys = TranslationData.PosKeys.Num();
 
 		// Determine the bounds
-		const FBox KeyBounds(TranslationData.PosKeys.GetTypedData(), NumKeys);
+		const FBox KeyBounds(TranslationData.PosKeys.GetData(), NumKeys);
 		const bool bHasX = (FMath::Abs(KeyBounds.Max.X) >= ZeroingThreshold) || (FMath::Abs(KeyBounds.Min.X) >= ZeroingThreshold);
 		const bool bHasY = (FMath::Abs(KeyBounds.Max.Y) >= ZeroingThreshold) || (FMath::Abs(KeyBounds.Min.Y) >= ZeroingThreshold);
 		const bool bHasZ = (FMath::Abs(KeyBounds.Max.Z) >= ZeroingThreshold) || (FMath::Abs(KeyBounds.Min.Z) >= ZeroingThreshold);
@@ -371,7 +371,7 @@ protected:
 		const int32 NumKeys = RotationData.RotKeys.Num();
 
 		// Determine the bounds
-		const FBox KeyBounds = CalculateQuatACF96Bounds(RotationData.RotKeys.GetTypedData(), NumKeys);
+		const FBox KeyBounds = CalculateQuatACF96Bounds(RotationData.RotKeys.GetData(), NumKeys);
 		const bool bHasX = (FMath::Abs(KeyBounds.Max.X) >= ZeroingThreshold) || (FMath::Abs(KeyBounds.Min.X) >= ZeroingThreshold);
 		const bool bHasY = (FMath::Abs(KeyBounds.Max.Y) >= ZeroingThreshold) || (FMath::Abs(KeyBounds.Min.Y) >= ZeroingThreshold);
 		const bool bHasZ = (FMath::Abs(KeyBounds.Max.Z) >= ZeroingThreshold) || (FMath::Abs(KeyBounds.Min.Z) >= ZeroingThreshold);
@@ -459,7 +459,7 @@ protected:
 		const int32 NumKeys = RotationData.RotKeys.Num();
 
 		// Determine the bounds
-		const FBox KeyBounds = CalculateQuatACF96Bounds(RotationData.RotKeys.GetTypedData(), NumKeys);
+		const FBox KeyBounds = CalculateQuatACF96Bounds(RotationData.RotKeys.GetData(), NumKeys);
 		FVector Range(KeyBounds.Max - KeyBounds.Min);
 
 		const bool bHasX = (FMath::Abs(KeyBounds.Max.X) >= ZeroingThreshold) || (FMath::Abs(KeyBounds.Min.X) >= ZeroingThreshold);
@@ -607,7 +607,7 @@ protected:
 		const int32 NumKeys = ScaleData.ScaleKeys.Num();
 
 		// Determine the bounds
-		const FBox KeyBounds(ScaleData.ScaleKeys.GetTypedData(), NumKeys);
+		const FBox KeyBounds(ScaleData.ScaleKeys.GetData(), NumKeys);
 		const bool bHasX = (FMath::Abs(KeyBounds.Max.X) >= ZeroingThreshold) || (FMath::Abs(KeyBounds.Min.X) >= ZeroingThreshold);
 		const bool bHasY = (FMath::Abs(KeyBounds.Max.Y) >= ZeroingThreshold) || (FMath::Abs(KeyBounds.Min.Y) >= ZeroingThreshold);
 		const bool bHasZ = (FMath::Abs(KeyBounds.Max.Z) >= ZeroingThreshold) || (FMath::Abs(KeyBounds.Min.Z) >= ZeroingThreshold);
@@ -664,7 +664,7 @@ protected:
 		const int32 NumKeys = ScaleData.ScaleKeys.Num();
 
 		// Determine the bounds
-		const FBox KeyBounds(ScaleData.ScaleKeys.GetTypedData(), NumKeys);
+		const FBox KeyBounds(ScaleData.ScaleKeys.GetData(), NumKeys);
 		const bool bHasX = (FMath::Abs(KeyBounds.Max.X) >= ZeroingThreshold) || (FMath::Abs(KeyBounds.Min.X) >= ZeroingThreshold);
 		const bool bHasY = (FMath::Abs(KeyBounds.Max.Y) >= ZeroingThreshold) || (FMath::Abs(KeyBounds.Min.Y) >= ZeroingThreshold);
 		const bool bHasZ = (FMath::Abs(KeyBounds.Max.Z) >= ZeroingThreshold) || (FMath::Abs(KeyBounds.Min.Z) >= ZeroingThreshold);
@@ -710,7 +710,7 @@ protected:
 		const int32 NumKeys = ScaleData.ScaleKeys.Num();
 
 		// Determine the bounds
-		const FBox KeyBounds(ScaleData.ScaleKeys.GetTypedData(), NumKeys);
+		const FBox KeyBounds(ScaleData.ScaleKeys.GetData(), NumKeys);
 		const bool bHasX = (FMath::Abs(KeyBounds.Max.X) >= ZeroingThreshold) || (FMath::Abs(KeyBounds.Min.X) >= ZeroingThreshold);
 		const bool bHasY = (FMath::Abs(KeyBounds.Max.Y) >= ZeroingThreshold) || (FMath::Abs(KeyBounds.Min.Y) >= ZeroingThreshold);
 		const bool bHasZ = (FMath::Abs(KeyBounds.Max.Z) >= ZeroingThreshold) || (FMath::Abs(KeyBounds.Min.Z) >= ZeroingThreshold);
diff --git a/Engine/Source/Runtime/Engine/Private/Animation/AnimEncoding.cpp b/Engine/Source/Runtime/Engine/Private/Animation/AnimEncoding.cpp
--- a/Engine/Source/Runtime/Engine/Private/Animation/AnimEncoding.cpp
+++ b/Engine/Source/Runtime/Engine/Private/Animation/AnimEncoding.cpp
@@ -251,13 +251,13 @@ void AnimEncodingLegacyBase::GetBoneAtom(
 	OutAtom.SetIdentity();
 
 	// Use the CompressedTrackOffsets stream to find the data addresses
-	const int32* RESTRICT TrackData= Seq.CompressedTrackOffsets.GetTypedData() + (TrackIndex*4);
+	const int32* RESTRICT TrackData= Seq.CompressedTrackOffsets.GetData() + (TrackIndex*4);
 	int32 TransKeysOffset = *(TrackData+0);
 	int32 NumTransKeys	= *(TrackData+1);
 	int32 RotKeysOffset	= *(TrackData+2);
 	int32 NumRotKeys		= *(TrackData+3);
-	const uint8* RESTRICT TransStream	= Seq.CompressedByteStream.GetTypedData()+TransKeysOffset;
-	const uint8* RESTRICT RotStream		= Seq.CompressedByteStream.GetTypedData()+RotKeysOffset;
+	const uint8* RESTRICT TransStream	= Seq.CompressedByteStream.GetData()+TransKeysOffset;
+	const uint8* RESTRICT RotStream		= Seq.CompressedByteStream.GetData()+RotKeysOffset;
 
 	const float RelativePos = Time / (float)Seq.SequenceLength;
 
@@ -275,7 +275,7 @@ void AnimEncodingLegacyBase::GetBoneAtom(
 	{
 		int32 ScaleKeyOffset = Seq.CompressedScaleOffsets.GetOffsetData(TrackIndex, 0);
 		int32 NumScaleKeys = Seq.CompressedScaleOffsets.GetOffsetData(TrackIndex, 1);
-		const uint8* RESTRICT ScaleStream		= Seq.CompressedByteStream.GetTypedData()+ScaleKeyOffset;
+		const uint8* RESTRICT ScaleStream		= Seq.CompressedByteStream.GetData()+ScaleKeyOffset;
 		// decompress the rotation component using the proper method
 		checkSlow(Seq.ScaleCodec != NULL);
 		((AnimEncodingLegacyBase*)Seq.ScaleCodec)->GetBoneAtomScale(OutAtom, Seq, ScaleStream, NumScaleKeys, Time, RelativePos);
@@ -306,7 +306,7 @@ void AnimEncodingLegacyBase::ByteSwapIn(
 	Seq.CompressedByteStream.AddUninitialized(OriginalNumBytes);
 
 	// Read and swap
-	uint8* StreamBase = Seq.CompressedByteStream.GetTypedData();
+	uint8* StreamBase = Seq.CompressedByteStream.GetData();
 	bool bHasValidScale = Seq.CompressedScaleOffsets.IsValid();
 
 	for ( int32 TrackIndex = 0; TrackIndex < NumTracks; ++TrackIndex )
@@ -374,7 +374,7 @@ void AnimEncodingLegacyBase::ByteSwapOut(
 	FMemoryWriter MemoryWriter( SerializedData, true );
 	MemoryWriter.SetByteSwapping( ForceByteSwapping );
 
-	uint8* StreamBase		= Seq.CompressedByteStream.GetTypedData();
+	uint8* StreamBase		= Seq.CompressedByteStream.GetData();
 	const int32 NumTracks		= Seq.CompressedTrackOffsets.Num()/4;
 
 	bool bHasValidScale = Seq.CompressedScaleOffsets.IsValid();
@@ -583,7 +583,7 @@ void AnimationFormat_GetStats(
 				}
 				else
 				{
-					const int32 Header = *((int32*)(Seq->CompressedByteStream.GetTypedData() + TransOffset));
+					const int32 Header = *((int32*)(Seq->CompressedByteStream.GetData() + TransOffset));
 
 					int32 KeyFormat;
 					int32 FormatFlags;
@@ -616,7 +616,7 @@ void AnimationFormat_GetStats(
 				}
 				else
 				{
-					const int32 Header = *((int32*)(Seq->CompressedByteStream.GetTypedData() + RotOffset));
+					const int32 Header = *((int32*)(Seq->CompressedByteStream.GetData() + RotOffset));
 
 					int32 KeyFormat;
 					int32 FormatFlags;
@@ -649,7 +649,7 @@ void AnimationFormat_GetStats(
 				}
 				else
 				{
-					const int32 Header = *((int32*)(Seq->CompressedByteStream.GetTypedData() + ScaleOffset));
+					const int32 Header = *((int32*)(Seq->CompressedByteStream.GetData() + ScaleOffset));
 
 					int32 KeyFormat;
 					int32 FormatFlags;
diff --git a/Engine/Source/Runtime/Engine/Private/Animation/AnimEncoding_PerTrackCompression.cpp b/Engine/Source/Runtime/Engine/Private/Animation/AnimEncoding_PerTrackCompression.cpp
--- a/Engine/Source/Runtime/Engine/Private/Animation/AnimEncoding_PerTrackCompression.cpp
+++ b/Engine/Source/Runtime/Engine/Private/Animation/AnimEncoding_PerTrackCompression.cpp
@@ -295,7 +295,7 @@ void AEFPerTrackCompressionCodec::ByteSwapOneTrack(UAnimSequence& Seq, FMemoryAr
 	{
 		checkSlow( (Offset % 4) == 0 && "CompressedByteStream not aligned to four bytes" );
 
-		uint8* TrackData = Seq.CompressedByteStream.GetTypedData() + Offset;
+		uint8* TrackData = Seq.CompressedByteStream.GetData() + Offset;
 
 		// Read the header
 		AC_UnalignedSwap(MemoryStream, TrackData, sizeof(int32));
@@ -462,7 +462,7 @@ void AEFPerTrackCompressionCodec::GetBoneAtom(
 	OutAtom.SetIdentity();
 
 	// Use the CompressedTrackOffsets stream to find the data addresses
-	const int32* RESTRICT TrackData= Seq.CompressedTrackOffsets.GetTypedData() + (TrackIndex * 2);
+	const int32* RESTRICT TrackData= Seq.CompressedTrackOffsets.GetData() + (TrackIndex * 2);
 	const int32 TransKeysOffset = TrackData[0];
 	const int32 RotKeysOffset = TrackData[1];
 	const float RelativePos = Time / (float)Seq.SequenceLength;
@@ -491,8 +491,8 @@ void AEFPerTrackCompressionCodec::GetBoneAtomRotation(
 {
 	if (Offset != INDEX_NONE)
 	{
-		const uint8* RESTRICT TrackData = Seq.CompressedByteStream.GetTypedData() + Offset + 4;
-		const int32 Header = *((int32*)(Seq.CompressedByteStream.GetTypedData() + Offset));
+		const uint8* RESTRICT TrackData = Seq.CompressedByteStream.GetData() + Offset + 4;
+		const int32 Header = *((int32*)(Seq.CompressedByteStream.GetData() + Offset));
 
 		int32 KeyFormat;
 		int32 NumKeys;
@@ -582,8 +582,8 @@ void AEFPerTrackCompressionCodec::GetBoneAtomTranslation(
 {
 	if (Offset != INDEX_NONE)
 	{
-		const uint8* RESTRICT TrackData = Seq.CompressedByteStream.GetTypedData() + Offset + 4;
-		const int32 Header = *((int32*)(Seq.CompressedByteStream.GetTypedData() + Offset));
+		const uint8* RESTRICT TrackData = Seq.CompressedByteStream.GetData() + Offset + 4;
+		const int32 Header = *((int32*)(Seq.CompressedByteStream.GetData() + Offset));
 
 		int32 KeyFormat;
 		int32 NumKeys;
@@ -667,8 +667,8 @@ void AEFPerTrackCompressionCodec::GetBoneAtomScale(
 {
 	if (Offset != INDEX_NONE)
 	{
-		const uint8* RESTRICT TrackData = Seq.CompressedByteStream.GetTypedData() + Offset + 4;
-		const int32 Header = *((int32*)(Seq.CompressedByteStream.GetTypedData() + Offset));
+		const uint8* RESTRICT TrackData = Seq.CompressedByteStream.GetData() + Offset + 4;
+		const int32 Header = *((int32*)(Seq.CompressedByteStream.GetData() + Offset));
 
 		int32 KeyFormat;
 		int32 NumKeys;
@@ -769,7 +769,7 @@ void AEFPerTrackCompressionCodec::GetPoseRotations(
 		const int32 AtomIndex = Pair.AtomIndex;
 		FTransform& BoneAtom = Atoms[AtomIndex];
 
-		const int32* RESTRICT TrackData = Seq.CompressedTrackOffsets.GetTypedData() + (TrackIndex * 2);
+		const int32* RESTRICT TrackData = Seq.CompressedTrackOffsets.GetData() + (TrackIndex * 2);
 		const int32 RotKeysOffset = *(TrackData + 1);
 
 		GetBoneAtomRotation( BoneAtom, Seq, RotKeysOffset, Time, RelativePos );
@@ -801,7 +801,7 @@ void AEFPerTrackCompressionCodec::GetPoseTranslations(
 		const int32 AtomIndex = Pair.AtomIndex;
 		FTransform& BoneAtom = Atoms[AtomIndex];
 
-		const int32* RESTRICT TrackData = Seq.CompressedTrackOffsets.GetTypedData() + (TrackIndex * 2);
+		const int32* RESTRICT TrackData = Seq.CompressedTrackOffsets.GetData() + (TrackIndex * 2);
 		const int32 PosKeysOffset = *(TrackData + 0);
 
 		GetBoneAtomTranslation( BoneAtom, Seq, PosKeysOffset, Time, RelativePos );
diff --git a/Engine/Source/Runtime/Engine/Private/Animation/AnimationRuntime.cpp b/Engine/Source/Runtime/Engine/Private/Animation/AnimationRuntime.cpp
--- a/Engine/Source/Runtime/Engine/Private/Animation/AnimationRuntime.cpp
+++ b/Engine/Source/Runtime/Engine/Private/Animation/AnimationRuntime.cpp
@@ -304,8 +304,8 @@ void FAnimationRuntime::LerpBoneTransforms(TArray<FTransform> & A, const TArray<
 	}
 	else if( Alpha > ZERO_ANIMWEIGHT_THRESH )
 	{
-		FTransform * ATransformData = A.GetTypedData(); 
-		const FTransform * BTransformData = B.GetTypedData();
+		FTransform * ATransformData = A.GetData(); 
+		const FTransform * BTransformData = B.GetData();
 		const ScalarRegister VAlpha(Alpha);
 		const ScalarRegister VOneMinusAlpha(1.f - Alpha);
 
@@ -729,8 +729,8 @@ void FAnimationRuntime::ConvertPoseToMeshSpace(const TArray<FTransform> & LocalT
 	check( NumBones == LocalTransforms.Num() );
 	check( NumBones == MeshSpaceTransforms.Num() );
 
-	const FTransform * LocalTransformsData = LocalTransforms.GetTypedData(); 
-	FTransform * SpaceBasesData = MeshSpaceTransforms.GetTypedData();
+	const FTransform * LocalTransformsData = LocalTransforms.GetData(); 
+	FTransform * SpaceBasesData = MeshSpaceTransforms.GetData();
 	const TArray<FBoneIndexType> & RequiredBoneIndexArray = RequiredBones.GetBoneIndicesArray();
 
 	// First bone is always root bone, and it doesn't have a parent.
@@ -960,7 +960,7 @@ void FAnimationRuntime::UpdateDesiredBoneWeight(const TArray<FPerBoneBlendWeight
 	// in the future, cache this outside
 	ensure (TargetBoneBlendWeights.Num() == SrcBoneBlendWeights.Num());
 
-	FMemory::Memset(TargetBoneBlendWeights.GetTypedData(), 0, TargetBoneBlendWeights.Num() * sizeof(FPerBoneBlendWeight));
+	FMemory::Memset(TargetBoneBlendWeights.GetData(), 0, TargetBoneBlendWeights.Num() * sizeof(FPerBoneBlendWeight));
 
 	// go through skeleton tree requiredboneindices
 	const TArray<FBoneIndexType> & RequiredBoneIndices = RequiredBones.GetBoneIndicesArray();
diff --git a/Engine/Source/Runtime/Engine/Private/ArrowComponent.cpp b/Engine/Source/Runtime/Engine/Private/ArrowComponent.cpp
--- a/Engine/Source/Runtime/Engine/Private/ArrowComponent.cpp
+++ b/Engine/Source/Runtime/Engine/Private/ArrowComponent.cpp
@@ -27,7 +27,7 @@ public:
 
 		// Copy the vertex data into the vertex buffer.
 		void* VertexBufferData = RHILockVertexBuffer(VertexBufferRHI, 0, Vertices.Num() * sizeof(FDynamicMeshVertex), RLM_WriteOnly);
-		FMemory::Memcpy(VertexBufferData, Vertices.GetTypedData(), Vertices.Num() * sizeof(FDynamicMeshVertex));
+		FMemory::Memcpy(VertexBufferData, Vertices.GetData(), Vertices.Num() * sizeof(FDynamicMeshVertex));
 		RHIUnlockVertexBuffer(VertexBufferRHI);
 	}
 
@@ -46,7 +46,7 @@ public:
 
 		// Write the indices to the index buffer.
 		void* Buffer = RHILockIndexBuffer(IndexBufferRHI, 0, Indices.Num() * sizeof(int32), RLM_WriteOnly);
-		FMemory::Memcpy(Buffer, Indices.GetTypedData(), Indices.Num() * sizeof(int32));
+		FMemory::Memcpy(Buffer, Indices.GetData(), Indices.Num() * sizeof(int32));
 		RHIUnlockIndexBuffer(IndexBufferRHI);
 	}
 };
diff --git a/Engine/Source/Runtime/Engine/Private/Atmosphere/Atmosphere.cpp b/Engine/Source/Runtime/Engine/Private/Atmosphere/Atmosphere.cpp
--- a/Engine/Source/Runtime/Engine/Private/Atmosphere/Atmosphere.cpp
+++ b/Engine/Source/Runtime/Engine/Private/Atmosphere/Atmosphere.cpp
@@ -134,7 +134,7 @@ void UAtmosphericFogComponent::PostLoad()
 				int32 TotalByte = OutData.Num() * sizeof(FColor);
 				TransmittanceData.Lock(LOCK_READ_WRITE);
 				FColor* TextureData = (FColor*)TransmittanceData.Realloc(TotalByte);
-				FMemory::Memcpy(TextureData, OutData.GetTypedData(), TotalByte);
+				FMemory::Memcpy(TextureData, OutData.GetData(), TotalByte);
 				TransmittanceData.Unlock();
 
 				TransmittanceTexture_DEPRECATED = NULL;
@@ -162,7 +162,7 @@ void UAtmosphericFogComponent::PostLoad()
 				int32 TotalByte = OutData.Num() * sizeof(FColor);
 				IrradianceData.Lock(LOCK_READ_WRITE);
 				FColor* TextureData = (FColor*)IrradianceData.Realloc(TotalByte);
-				FMemory::Memcpy(TextureData, OutData.GetTypedData(), TotalByte);
+				FMemory::Memcpy(TextureData, OutData.GetData(), TotalByte);
 				IrradianceData.Unlock();
 
 				IrradianceTexture_DEPRECATED = NULL;
diff --git a/Engine/Source/Runtime/Engine/Private/AudioDerivedData.cpp b/Engine/Source/Runtime/Engine/Private/AudioDerivedData.cpp
--- a/Engine/Source/Runtime/Engine/Private/AudioDerivedData.cpp
+++ b/Engine/Source/Runtime/Engine/Private/AudioDerivedData.cpp
@@ -246,7 +246,7 @@ class FStreamedAudioCacheDerivedDataWorker : public FNonAbandonableTask
 			TArray<uint8> CompressedBuffer;
 			CompressedBuffer.Empty(CompressedData->GetBulkDataSize());
 			CompressedBuffer.AddUninitialized(CompressedData->GetBulkDataSize());
-			void* BufferData = CompressedBuffer.GetTypedData();
+			void* BufferData = CompressedBuffer.GetData();
 			CompressedData->GetCopy(&BufferData, false);
 			TArray<TArray<uint8>> ChunkBuffers;
 
@@ -258,7 +258,7 @@ class FStreamedAudioCacheDerivedDataWorker : public FNonAbandonableTask
 					NewChunk->DataSize = ChunkBuffers[ChunkIndex].Num();
 					NewChunk->BulkData.Lock(LOCK_READ_WRITE);
 					void* NewChunkData = NewChunk->BulkData.Realloc(ChunkBuffers[ChunkIndex].Num());
-					FMemory::Memcpy(NewChunkData, ChunkBuffers[ChunkIndex].GetTypedData(), ChunkBuffers[ChunkIndex].Num());
+					FMemory::Memcpy(NewChunkData, ChunkBuffers[ChunkIndex].GetData(), ChunkBuffers[ChunkIndex].Num());
 					NewChunk->BulkData.Unlock();
 				}
 			}
@@ -269,7 +269,7 @@ class FStreamedAudioCacheDerivedDataWorker : public FNonAbandonableTask
 				NewChunk->DataSize = CompressedBuffer.Num();
 				NewChunk->BulkData.Lock(LOCK_READ_WRITE);
 				void* NewChunkData = NewChunk->BulkData.Realloc(CompressedBuffer.Num());
-				FMemory::Memcpy(NewChunkData, CompressedBuffer.GetTypedData(), CompressedBuffer.Num());
+				FMemory::Memcpy(NewChunkData, CompressedBuffer.GetData(), CompressedBuffer.Num());
 				NewChunk->BulkData.Unlock();
 			}
 
@@ -671,7 +671,7 @@ static void CookSimpleWave(USoundWave* SoundWave, FName FormatName, const IAudio
 		else
 		{
 			Input.AddUninitialized(WaveInfo.SampleDataSize);
-			FMemory::Memcpy(Input.GetTypedData(), WaveInfo.SampleDataStart, WaveInfo.SampleDataSize);
+			FMemory::Memcpy(Input.GetData(), WaveInfo.SampleDataStart, WaveInfo.SampleDataSize);
 		}
 	}
 
@@ -743,7 +743,7 @@ void CookSurroundWave( USoundWave* SoundWave, FName FormatName, const IAudioForm
 		{
 			TArray<uint8>& Input = *new (SourceBuffers) TArray<uint8>;
 			Input.AddUninitialized(WaveInfo.SampleDataSize);
-			FMemory::Memcpy(Input.GetTypedData(), WaveInfo.SampleDataStart, WaveInfo.SampleDataSize);
+			FMemory::Memcpy(Input.GetData(), WaveInfo.SampleDataStart, WaveInfo.SampleDataSize);
 		}
 
 		SampleDataSize = WaveInfo.SampleDataSize;
@@ -759,7 +759,7 @@ void CookSurroundWave( USoundWave* SoundWave, FName FormatName, const IAudioForm
 					ChannelCount++;
 					TArray<uint8>& Input = *new (SourceBuffers) TArray<uint8>;
 					Input.AddUninitialized(WaveInfoInner.SampleDataSize);
-					FMemory::Memcpy(Input.GetTypedData(), WaveInfoInner.SampleDataStart, WaveInfoInner.SampleDataSize);
+					FMemory::Memcpy(Input.GetData(), WaveInfoInner.SampleDataStart, WaveInfoInner.SampleDataSize);
 					SampleDataSize = FMath::Min<uint32>(WaveInfoInner.SampleDataSize, SampleDataSize);
 				}
 			}
diff --git a/Engine/Source/Runtime/Engine/Private/BatchedElements.cpp b/Engine/Source/Runtime/Engine/Private/BatchedElements.cpp
--- a/Engine/Source/Runtime/Engine/Private/BatchedElements.cpp
+++ b/Engine/Source/Runtime/Engine/Private/BatchedElements.cpp
@@ -881,7 +881,7 @@ bool FBatchedElements::Draw(FRHICommandList& RHICmdList, ERHIFeatureLevel::Type
 						const int32 VertexCount = SpriteList.Num();
 						const int32 PrimCount = VertexCount / 3;
 						PrepareShaders(RHICmdList, FeatureLevel, CurrentBlendMode, Transform, bNeedToSwitchVerticalAxis, BatchedElementParameters, CurrentTexture, bHitTesting, Gamma, NULL, View, DepthTexture);
-						DrawPrimitiveUP(RHICmdList, PT_TriangleList, PrimCount, SpriteList.GetTypedData(), sizeof(FSimpleElementVertex));
+						DrawPrimitiveUP(RHICmdList, PT_TriangleList, PrimCount, SpriteList.GetData(), sizeof(FSimpleElementVertex));
 
 						SpriteList.Empty(6);
 						CurrentTexture = Sprite.Texture;
@@ -889,7 +889,7 @@ bool FBatchedElements::Draw(FRHICommandList& RHICmdList, ERHIFeatureLevel::Type
 					}
 
 					int32 SpriteListIndex = SpriteList.AddUninitialized(6);
-					FSimpleElementVertex* Vertex = SpriteList.GetTypedData();
+					FSimpleElementVertex* Vertex = SpriteList.GetData();
 
 					// Compute the sprite vertices.
 					const FVector WorldSpriteX = CameraX * Sprite.SizeX;
@@ -921,7 +921,7 @@ bool FBatchedElements::Draw(FRHICommandList& RHICmdList, ERHIFeatureLevel::Type
 					const int32 VertexCount = SpriteList.Num();
 					const int32 PrimCount = VertexCount / 3;
 					PrepareShaders(RHICmdList, FeatureLevel, CurrentBlendMode, Transform, bNeedToSwitchVerticalAxis, BatchedElementParameters, CurrentTexture, bHitTesting, Gamma, NULL, View, DepthTexture);
-					DrawPrimitiveUP(RHICmdList, PT_TriangleList, PrimCount, SpriteList.GetTypedData(), sizeof(FSimpleElementVertex));
+					DrawPrimitiveUP(RHICmdList, PT_TriangleList, PrimCount, SpriteList.GetData(), sizeof(FSimpleElementVertex));
 				}
 			}
 		}
@@ -947,7 +947,7 @@ bool FBatchedElements::Draw(FRHICommandList& RHICmdList, ERHIFeatureLevel::Type
 						0,
 						MeshElement.MaxVertex - MeshElement.MinVertex + 1,
 						MeshElement.Indices.Num() / 3,
-						MeshElement.Indices.GetTypedData(),
+						MeshElement.Indices.GetData(),
 						sizeof(uint16),
 						&MeshVertices[MeshElement.MinVertex],
 						sizeof(FSimpleElementVertex)
@@ -975,7 +975,7 @@ bool FBatchedElements::Draw(FRHICommandList& RHICmdList, ERHIFeatureLevel::Type
 						RHICmdList,
 						PT_QuadList,
 						MeshElement.Vertices.Num() / 4,
-						MeshElement.Vertices.GetTypedData(),
+						MeshElement.Vertices.GetData(),
 						sizeof(FSimpleElementVertex)
 						);
 				}
diff --git a/Engine/Source/Runtime/Engine/Private/BoxComponent.cpp b/Engine/Source/Runtime/Engine/Private/BoxComponent.cpp
--- a/Engine/Source/Runtime/Engine/Private/BoxComponent.cpp
+++ b/Engine/Source/Runtime/Engine/Private/BoxComponent.cpp
@@ -42,7 +42,7 @@ void UBoxComponent::UpdateBodySetup()
 	}
 
 	check(ShapeBodySetup->AggGeom.BoxElems.Num() == 1);
-	FKBoxElem* se = ShapeBodySetup->AggGeom.BoxElems.GetTypedData();
+	FKBoxElem* se = ShapeBodySetup->AggGeom.BoxElems.GetData();
 
 	// @todo UE4 do we allow this now?
 	// check for malformed values
diff --git a/Engine/Source/Runtime/Engine/Private/BrushComponent.cpp b/Engine/Source/Runtime/Engine/Private/BrushComponent.cpp
--- a/Engine/Source/Runtime/Engine/Private/BrushComponent.cpp
+++ b/Engine/Source/Runtime/Engine/Private/BrushComponent.cpp
@@ -589,7 +589,7 @@ FBoxSphereBounds UBrushComponent::CalcBounds(const FTransform & LocalToWorld) co
 				Points.Add(Brush->Polys->Element[i].Vertices[j]);
 			}
 		}
-		return FBoxSphereBounds( Points.GetTypedData(), Points.Num() ).TransformBy(LocalToWorld);
+		return FBoxSphereBounds( Points.GetData(), Points.Num() ).TransformBy(LocalToWorld);
 	}
 	else 
 #endif // WITH_EDITOR
diff --git a/Engine/Source/Runtime/Engine/Private/CapsuleComponent.cpp b/Engine/Source/Runtime/Engine/Private/CapsuleComponent.cpp
--- a/Engine/Source/Runtime/Engine/Private/CapsuleComponent.cpp
+++ b/Engine/Source/Runtime/Engine/Private/CapsuleComponent.cpp
@@ -168,7 +168,7 @@ void UCapsuleComponent::UpdateBodySetup()
 	}
 
 	check(ShapeBodySetup->AggGeom.SphylElems.Num() == 1);
-	FKSphylElem* SE = ShapeBodySetup->AggGeom.SphylElems.GetTypedData();
+	FKSphylElem* SE = ShapeBodySetup->AggGeom.SphylElems.GetData();
 	
 	SE->SetTransform(FTransform::Identity);
 	SE->Radius = CapsuleRadius;
diff --git a/Engine/Source/Runtime/Engine/Private/ContentStreaming.cpp b/Engine/Source/Runtime/Engine/Private/ContentStreaming.cpp
--- a/Engine/Source/Runtime/Engine/Private/ContentStreaming.cpp
+++ b/Engine/Source/Runtime/Engine/Private/ContentStreaming.cpp
@@ -2328,7 +2328,7 @@ void Renderthread_StreamOutTextureData(FRHICommandListImmediate& RHICmdList, TAr
 	// Makes sure that texture memory can get freed up right away.
 	RHICmdList.BlockUntilGPUIdle();
 
-	FTextureSortElement* CandidateTextures = InCandidateTextures->GetTypedData();
+	FTextureSortElement* CandidateTextures = InCandidateTextures->GetData();
 
 	// Sort the candidates.
 	InCandidateTextures->Sort( FTextureStreamingCompare() );
diff --git a/Engine/Source/Runtime/Engine/Private/DataChannel.cpp b/Engine/Source/Runtime/Engine/Private/DataChannel.cpp
--- a/Engine/Source/Runtime/Engine/Private/DataChannel.cpp
+++ b/Engine/Source/Runtime/Engine/Private/DataChannel.cpp
@@ -1130,7 +1130,7 @@ void UControlChannel::QueueMessage(const FOutBunch* Bunch)
 	{
 		int32 Index = QueuedMessages.AddZeroed();
 		QueuedMessages[Index].AddUninitialized(Bunch->GetNumBytes());
-		FMemory::Memcpy(QueuedMessages[Index].GetTypedData(), Bunch->GetData(), Bunch->GetNumBytes());
+		FMemory::Memcpy(QueuedMessages[Index].GetData(), Bunch->GetData(), Bunch->GetNumBytes());
 	}
 }
 
diff --git a/Engine/Source/Runtime/Engine/Private/DataReplication.cpp b/Engine/Source/Runtime/Engine/Private/DataReplication.cpp
--- a/Engine/Source/Runtime/Engine/Private/DataReplication.cpp
+++ b/Engine/Source/Runtime/Engine/Private/DataReplication.cpp
@@ -753,7 +753,7 @@ void FObjectReplicator::PostReceivedBunch()
 			}
 			else if (RepNotifyFunc->NumParms == 1)
 			{
-				Object->ProcessEvent(RepNotifyFunc, RepProperty->ContainerPtrToValuePtr<uint8>(RepState->StaticBuffer.GetTypedData()) );
+				Object->ProcessEvent(RepNotifyFunc, RepProperty->ContainerPtrToValuePtr<uint8>(RepState->StaticBuffer.GetData()) );
 			}
 			else if (RepNotifyFunc->NumParms == 2)
 			{
@@ -772,7 +772,7 @@ void FObjectReplicator::PostReceivedBunch()
 				TFieldIterator<UProperty> Itr(RepNotifyFunc);
 				check(Itr);
 				
-				Itr->CopyCompleteValue( Itr->ContainerPtrToValuePtr<void>(Parms), RepProperty->ContainerPtrToValuePtr<uint8>(RepState->StaticBuffer.GetTypedData()) );
+				Itr->CopyCompleteValue( Itr->ContainerPtrToValuePtr<void>(Parms), RepProperty->ContainerPtrToValuePtr<uint8>(RepState->StaticBuffer.GetData()) );
 				++Itr;
 				check(Itr);
 
diff --git a/Engine/Source/Runtime/Engine/Private/Debug/DebugDrawService.cpp b/Engine/Source/Runtime/Engine/Private/Debug/DebugDrawService.cpp
--- a/Engine/Source/Runtime/Engine/Private/Debug/DebugDrawService.cpp
+++ b/Engine/Source/Runtime/Engine/Private/Debug/DebugDrawService.cpp
@@ -32,7 +32,7 @@ void UDebugDrawService::Unregister(const FDebugDrawDelegate& DelegateToRemove)
 {
 	check(IsInGameThread());
 
-	TArray<FDebugDrawDelegate>* DelegatesArray = Delegates.GetTypedData();
+	TArray<FDebugDrawDelegate>* DelegatesArray = Delegates.GetData();
 	for (int32 Flag = 0; Flag < Delegates.Num(); ++Flag, ++DelegatesArray)
 	{
 		check(DelegatesArray); //it shouldn't happen, but to be sure
diff --git a/Engine/Source/Runtime/Engine/Private/DestructibleComponent.cpp b/Engine/Source/Runtime/Engine/Private/DestructibleComponent.cpp
--- a/Engine/Source/Runtime/Engine/Private/DestructibleComponent.cpp
+++ b/Engine/Source/Runtime/Engine/Private/DestructibleComponent.cpp
@@ -929,8 +929,8 @@ bool UDestructibleComponent::DoCustomNavigableGeometryExport(FNavigableGeometryE
 				{
 					Shapes.AddUninitialized(ShapesCount - Shapes.Num());
 				}
-				const PxU32 RetrievedShapesCount = PActor->getShapes(Shapes.GetTypedData(), Shapes.Num());
-				PxShape* const* ShapePtr = Shapes.GetTypedData();
+				const PxU32 RetrievedShapesCount = PActor->getShapes(Shapes.GetData(), Shapes.Num());
+				PxShape* const* ShapePtr = Shapes.GetData();
 				for (PxU32 ShapeIndex = 0; ShapeIndex < RetrievedShapesCount; ++ShapeIndex, ++ShapePtr)
 				{
 					if (*ShapePtr != NULL)
@@ -1160,7 +1160,7 @@ void UDestructibleComponent::SetCollisionResponseForActor(PxRigidDynamic* Actor,
 		TArray<PxShape*> Shapes;
 		Shapes.AddUninitialized(Actor->getNbShapes());
 
-		int ShapeCount = Actor->getShapes(Shapes.GetTypedData(), Shapes.Num());
+		int ShapeCount = Actor->getShapes(Shapes.GetData(), Shapes.Num());
 
 		for (int32 i=0; i < ShapeCount; ++i)
 		{
diff --git a/Engine/Source/Runtime/Engine/Private/DestructibleFractureSettings.cpp b/Engine/Source/Runtime/Engine/Private/DestructibleFractureSettings.cpp
--- a/Engine/Source/Runtime/Engine/Private/DestructibleFractureSettings.cpp
+++ b/Engine/Source/Runtime/Engine/Private/DestructibleFractureSettings.cpp
@@ -40,7 +40,7 @@ static void BuildApexRenderMesh(NxRenderMeshAssetAuthoring& RenderMeshAssetAutho
 	TArray<NxRenderMeshAssetAuthoring::SubmeshDesc> SubmeshDescs;
 	SubmeshDescs.Init(NxRenderMeshAssetAuthoring::SubmeshDesc(), HMesh.submeshCount());
 	MeshDesc.m_numSubmeshes = (PxU32)SubmeshDescs.Num();
-	MeshDesc.m_submeshes = SubmeshDescs.GetTypedData();
+	MeshDesc.m_submeshes = SubmeshDescs.GetData();
 
 	// Need to transpose the submesh/part arrays.  The outer arrays on the following are indexed by submesh:
 	TArray< TArray<NxVertex> > SubmeshVertices;
@@ -105,7 +105,7 @@ static void BuildApexRenderMesh(NxRenderMeshAssetAuthoring& RenderMeshAssetAutho
 			{
 				TArray<PxU32> Map;
 				Map.Init(PartVertices.Num());
-				const PxU32 ReducedPartVertexCount = RenderMeshAssetAuthor.createReductionMap(Map.GetTypedData(), PartVertices.GetTypedData(), NULL, (PxU32)PartVertices.Num(), PxVec3( 0.0001f ), 0.001f, 1.0f/256.01f);
+				const PxU32 ReducedPartVertexCount = RenderMeshAssetAuthor.createReductionMap(Map.GetData(), PartVertices.GetData(), NULL, (PxU32)PartVertices.Num(), PxVec3( 0.0001f ), 0.001f, 1.0f/256.01f);
 				const PxU32 VertexPartStart = (PxU32)Vertices.Num();
 				Resize(Vertices, NxVertex(), VertexPartStart + ReducedPartVertexCount);
 				Resize(Indices, (uint32)0, PartIndices[PartIndex] + PartVertices.Num());
@@ -122,9 +122,9 @@ static void BuildApexRenderMesh(NxRenderMeshAssetAuthoring& RenderMeshAssetAutho
 
 		SubmeshDesc.m_numVertices = (PxU32)Vertices.Num();
 		SubmeshDesc.m_numParts = (PxU32)PartIndices.Num();
-		SubmeshDesc.m_partIndices = PartIndices.GetTypedData();
+		SubmeshDesc.m_partIndices = PartIndices.GetData();
 		SubmeshDesc.m_numIndices = (PxU32)Indices.Num();
-		SubmeshDesc.m_vertexIndices = Indices.GetTypedData();
+		SubmeshDesc.m_vertexIndices = Indices.GetData();
 
 		if (SubmeshDesc.m_numParts == 0 || SubmeshDesc.m_vertexIndices == 0)
 		{
@@ -133,44 +133,44 @@ static void BuildApexRenderMesh(NxRenderMeshAssetAuthoring& RenderMeshAssetAutho
 
 		if (VertexFormat.mHasStaticPositions || VertexFormat.mHasDynamicPositions)
 		{
-			vb.setSemanticData( NxRenderVertexSemantic::POSITION, &Vertices.GetTypedData()->position, sizeof( NxVertex ), NxRenderDataFormat::FLOAT3);
+			vb.setSemanticData( NxRenderVertexSemantic::POSITION, &Vertices.GetData()->position, sizeof( NxVertex ), NxRenderDataFormat::FLOAT3);
 		}
 		if (VertexFormat.mHasStaticNormals || VertexFormat.mHasDynamicNormals)
 		{
-			vb.setSemanticData( NxRenderVertexSemantic::NORMAL, &Vertices.GetTypedData()->normal, sizeof( NxVertex ), VertexNormalFormat, NxRenderDataFormat::FLOAT3);
+			vb.setSemanticData( NxRenderVertexSemantic::NORMAL, &Vertices.GetData()->normal, sizeof( NxVertex ), VertexNormalFormat, NxRenderDataFormat::FLOAT3);
 		}
 		if (VertexFormat.mHasStaticTangents || VertexFormat.mHasDynamicTangents)
 		{
-			vb.setSemanticData( NxRenderVertexSemantic::TANGENT, &Vertices.GetTypedData()->tangent, sizeof( NxVertex ), VertexNormalFormat, NxRenderDataFormat::FLOAT3);
+			vb.setSemanticData( NxRenderVertexSemantic::TANGENT, &Vertices.GetData()->tangent, sizeof( NxVertex ), VertexNormalFormat, NxRenderDataFormat::FLOAT3);
 		}
 		if (VertexFormat.mHasStaticBinormals || VertexFormat.mHasDynamicBinormals)
 		{
-			vb.setSemanticData( NxRenderVertexSemantic::BINORMAL, &Vertices.GetTypedData()->binormal, sizeof( NxVertex ), VertexNormalFormat, NxRenderDataFormat::FLOAT3);
+			vb.setSemanticData( NxRenderVertexSemantic::BINORMAL, &Vertices.GetData()->binormal, sizeof( NxVertex ), VertexNormalFormat, NxRenderDataFormat::FLOAT3);
 		}
 		if (VertexFormat.mHasStaticColors || VertexFormat.mHasDynamicColors)
 		{
-			vb.setSemanticData( NxRenderVertexSemantic::COLOR, &Vertices.GetTypedData()->color, sizeof( NxVertex ), NxRenderDataFormat::R8G8B8A8, NxRenderDataFormat::R32G32B32A32_FLOAT);
+			vb.setSemanticData( NxRenderVertexSemantic::COLOR, &Vertices.GetData()->color, sizeof( NxVertex ), NxRenderDataFormat::R8G8B8A8, NxRenderDataFormat::R32G32B32A32_FLOAT);
 		}
 		for (PxU32 uvNum = 0; uvNum < VertexFormat.mUVCount; ++uvNum)
 		{
-			vb.setSemanticData( (NxRenderVertexSemantic::Enum)(uvNum+NxRenderVertexSemantic::TEXCOORD0), &Vertices.GetTypedData()->uv[uvNum], sizeof( NxVertex ), NxRenderDataFormat::FLOAT2);
+			vb.setSemanticData( (NxRenderVertexSemantic::Enum)(uvNum+NxRenderVertexSemantic::TEXCOORD0), &Vertices.GetData()->uv[uvNum], sizeof( NxVertex ), NxRenderDataFormat::FLOAT2);
 		}
 		switch (VertexFormat.mBonesPerVertex)
 		{
 		case 1:
-			vb.setSemanticData( NxRenderVertexSemantic::BONE_INDEX, &Vertices.GetTypedData()->boneIndices[0], sizeof( NxVertex ), NxRenderDataFormat::USHORT1);
+			vb.setSemanticData( NxRenderVertexSemantic::BONE_INDEX, &Vertices.GetData()->boneIndices[0], sizeof( NxVertex ), NxRenderDataFormat::USHORT1);
 			break;
 		case 2:
-			vb.setSemanticData( NxRenderVertexSemantic::BONE_INDEX, &Vertices.GetTypedData()->boneIndices[0], sizeof( NxVertex ), NxRenderDataFormat::USHORT2);
-			vb.setSemanticData( NxRenderVertexSemantic::BONE_WEIGHT, &Vertices.GetTypedData()->boneWeights[0], sizeof( NxVertex ), NxRenderDataFormat::FLOAT2);
+			vb.setSemanticData( NxRenderVertexSemantic::BONE_INDEX, &Vertices.GetData()->boneIndices[0], sizeof( NxVertex ), NxRenderDataFormat::USHORT2);
+			vb.setSemanticData( NxRenderVertexSemantic::BONE_WEIGHT, &Vertices.GetData()->boneWeights[0], sizeof( NxVertex ), NxRenderDataFormat::FLOAT2);
 			break;
 		case 3:
-			vb.setSemanticData( NxRenderVertexSemantic::BONE_INDEX, &Vertices.GetTypedData()->boneIndices[0], sizeof( NxVertex ), NxRenderDataFormat::USHORT3);
-			vb.setSemanticData( NxRenderVertexSemantic::BONE_WEIGHT, &Vertices.GetTypedData()->boneWeights[0], sizeof( NxVertex ), NxRenderDataFormat::FLOAT3);
+			vb.setSemanticData( NxRenderVertexSemantic::BONE_INDEX, &Vertices.GetData()->boneIndices[0], sizeof( NxVertex ), NxRenderDataFormat::USHORT3);
+			vb.setSemanticData( NxRenderVertexSemantic::BONE_WEIGHT, &Vertices.GetData()->boneWeights[0], sizeof( NxVertex ), NxRenderDataFormat::FLOAT3);
 			break;
 		case 4:
-			vb.setSemanticData( NxRenderVertexSemantic::BONE_INDEX, &Vertices.GetTypedData()->boneIndices[0], sizeof( NxVertex ), NxRenderDataFormat::USHORT4);
-			vb.setSemanticData( NxRenderVertexSemantic::BONE_WEIGHT, &Vertices.GetTypedData()->boneWeights[0], sizeof( NxVertex ), NxRenderDataFormat::FLOAT4);
+			vb.setSemanticData( NxRenderVertexSemantic::BONE_INDEX, &Vertices.GetData()->boneIndices[0], sizeof( NxVertex ), NxRenderDataFormat::USHORT4);
+			vb.setSemanticData( NxRenderVertexSemantic::BONE_WEIGHT, &Vertices.GetData()->boneWeights[0], sizeof( NxVertex ), NxRenderDataFormat::FLOAT4);
 			break;
 		}
 	}
@@ -336,7 +336,7 @@ void UDestructibleFractureSettings::BuildDestructibleAssetCookingDesc(NxDestruct
 
 	// Set up chunk desc array
 	ChunkDescs.Init(physx::NxDestructibleChunkDesc(), HMesh.chunkCount());
-	DestructibleAssetCookingDesc.chunkDescs = ChunkDescs.GetTypedData();
+	DestructibleAssetCookingDesc.chunkDescs = ChunkDescs.GetData();
 	DestructibleAssetCookingDesc.chunkDescCount = ChunkDescs.Num();
 	for (uint32 ChunkIndex = 0; ChunkIndex < DestructibleAssetCookingDesc.chunkDescCount; ++ChunkIndex)
 	{
@@ -360,7 +360,7 @@ void UDestructibleFractureSettings::BuildDestructibleAssetCookingDesc(NxDestruct
 
 	// Set up geometry desc array
 	GeometryDescs.Init(physx::NxDestructibleGeometryDesc(), HMesh.partCount());
-	DestructibleAssetCookingDesc.geometryDescs = GeometryDescs.GetTypedData();
+	DestructibleAssetCookingDesc.geometryDescs = GeometryDescs.GetData();
 	DestructibleAssetCookingDesc.geometryDescCount = GeometryDescs.Num();
 	for (uint32 GeometryIndex = 0; GeometryIndex < DestructibleAssetCookingDesc.geometryDescCount; ++GeometryIndex)
 	{
@@ -377,8 +377,8 @@ bool UDestructibleFractureSettings::SetRootMesh(const TArray<NxExplicitRenderTri
 
 	if (ApexDestructibleAssetAuthoring != NULL)
 	{
-		Success = ApexDestructibleAssetAuthoring->setRootMesh(MeshTriangles.GetTypedData(), MeshTriangles.Num(), SubmeshData.GetTypedData(), 
-															  SubmeshData.Num(), (uint32*)MeshPartition.GetTypedData(), MeshPartition.Num(), bFirstPartitionIsDepthZero);
+		Success = ApexDestructibleAssetAuthoring->setRootMesh(MeshTriangles.GetData(), MeshTriangles.Num(), SubmeshData.GetData(), 
+															  SubmeshData.Num(), (uint32*)MeshPartition.GetData(), MeshPartition.Num(), bFirstPartitionIsDepthZero);
 		if (Success)
 		{
 			NxCollisionDesc CollisionDesc;
@@ -489,7 +489,7 @@ void UDestructibleFractureSettings::CreateVoronoiSitesInRootMesh()
 		// Progress listener for reporting progress - for now, just a dummy
 		FProgressListener ProgressListener;
 		check(sizeof(FVector) == sizeof(PxVec3));
-		ApexDestructibleAssetAuthoring->createVoronoiSitesInsideMesh((PxVec3*)VoronoiSites.GetTypedData(), VoronoiSites.Num(), (PxU32*)&RandomSeed, NULL, ProgressListener);
+		ApexDestructibleAssetAuthoring->createVoronoiSitesInsideMesh((PxVec3*)VoronoiSites.GetData(), VoronoiSites.Num(), (PxU32*)&RandomSeed, NULL, ProgressListener);
 	}
 }
 
@@ -507,7 +507,7 @@ bool UDestructibleFractureSettings::VoronoiSplitMesh()
 		FractureTools::FractureVoronoiDesc FTFractureVoronoiDesc;
 		FTFractureVoronoiDesc.siteCount = VoronoiSites.Num();
 		check(sizeof(FVector) == sizeof(PxVec3));
-		FTFractureVoronoiDesc.sites = (PxVec3*)VoronoiSites.GetTypedData();
+		FTFractureVoronoiDesc.sites = (PxVec3*)VoronoiSites.GetData();
 
 		// Material descriptor
 		FractureMaterialDesc.FillNxFractureMaterialDesc(FTFractureVoronoiDesc.materialDesc);
diff --git a/Engine/Source/Runtime/Engine/Private/DestructibleMesh.cpp b/Engine/Source/Runtime/Engine/Private/DestructibleMesh.cpp
--- a/Engine/Source/Runtime/Engine/Private/DestructibleMesh.cpp
+++ b/Engine/Source/Runtime/Engine/Private/DestructibleMesh.cpp
@@ -210,7 +210,7 @@ void UDestructibleMesh::Serialize(FArchive& Ar)
 				// Read the stream data into our buffer for UE serialzation
 				Size = Stream->getFileLength();
 				Buffer.AddUninitialized( Size );
-				Stream->read( Buffer.GetTypedData(), Size );
+				Stream->read( Buffer.GetData(), Size );
 			}
 		}
 		Ar << Size;
diff --git a/Engine/Source/Runtime/Engine/Private/Distributions.cpp b/Engine/Source/Runtime/Engine/Private/Distributions.cpp
--- a/Engine/Source/Runtime/Engine/Private/Distributions.cpp
+++ b/Engine/Source/Runtime/Engine/Private/Distributions.cpp
@@ -298,7 +298,7 @@ static void ScaleLookupTableByConstants( FDistributionLookupTable* Table, const
 	const int32 SubEntryCount = (Table->SubEntryStride > 0) ? 2 : 1;
 	const int32 Stride = Table->EntryStride;
 	const int32 SubEntryStride = Table->SubEntryStride;
-	float* RESTRICT Values = Table->Values.GetTypedData();
+	float* RESTRICT Values = Table->Values.GetData();
 
 	for ( int32 Index = 0; Index < EntryCount; ++Index )
 	{
@@ -330,7 +330,7 @@ static void AddConstantToLookupTable( FDistributionLookupTable* Table, const flo
 	const int32 SubEntryCount = (Table->SubEntryStride > 0) ? 2 : 1;
 	const int32 Stride = Table->EntryStride;
 	const int32 SubEntryStride = Table->SubEntryStride;
-	float* RESTRICT Values = Table->Values.GetTypedData();
+	float* RESTRICT Values = Table->Values.GetData();
 
 	for ( int32 Index = 0; Index < EntryCount; ++Index )
 	{
@@ -1291,7 +1291,7 @@ void FComposableDistribution::QuantizeVector4(
 	const FDistributionLookupTable& Table = Distribution.LookupTable;
 	const int32 EntryCount = Table.EntryCount;
 	const int32 EntryStride = Table.EntryStride;
-	const float* RESTRICT Values = Table.Values.GetTypedData();
+	const float* RESTRICT Values = Table.Values.GetData();
 	
 	// First find the minimum and maximum values for each channel at each sample.
 	for ( int32 EntryIndex = 0; EntryIndex < EntryCount; ++EntryIndex )
@@ -1330,8 +1330,8 @@ void FComposableDistribution::QuantizeVector4(
 	// Now construct the quantized samples.
 	OutQuantizedSamples.Empty( EntryCount );
 	OutQuantizedSamples.AddUninitialized( EntryCount );
-	FColor* RESTRICT QuantizedValues = OutQuantizedSamples.GetTypedData();
-	Values = Table.Values.GetTypedData();
+	FColor* RESTRICT QuantizedValues = OutQuantizedSamples.GetData();
+	Values = Table.Values.GetData();
 	for ( int32 EntryIndex = 0; EntryIndex < EntryCount; ++EntryIndex )
 	{
 		QuantizedValues->R = FMath::Clamp<int32>( FMath::TruncToInt( (Values[0] - Bias.X) * InvScale.X ), 0, 255 );
diff --git a/Engine/Source/Runtime/Engine/Private/DynamicMeshBuilder.cpp b/Engine/Source/Runtime/Engine/Private/DynamicMeshBuilder.cpp
--- a/Engine/Source/Runtime/Engine/Private/DynamicMeshBuilder.cpp
+++ b/Engine/Source/Runtime/Engine/Private/DynamicMeshBuilder.cpp
@@ -199,7 +199,7 @@ public:
 		
 		// Write the indices to the index buffer.
 		void* Buffer = RHILockIndexBuffer(IndexBufferRHI,0,Indices.Num() * sizeof(int32),RLM_WriteOnly);
-		FMemory::Memcpy(Buffer,Indices.GetTypedData(),Indices.Num() * sizeof(int32));
+		FMemory::Memcpy(Buffer,Indices.GetData(),Indices.Num() * sizeof(int32));
 		RHIUnlockIndexBuffer(IndexBufferRHI);
 	}
 	
@@ -248,7 +248,7 @@ public:
 		
 		// Copy the vertex data into the vertex buffer.
 		void* VertexBufferData = RHILockVertexBuffer(VertexBufferRHI,0,SizeInBytes,RLM_WriteOnly);
-		FMemory::Memcpy(VertexBufferData,Vertices.GetTypedData(),SizeInBytes);
+		FMemory::Memcpy(VertexBufferData,Vertices.GetData(),SizeInBytes);
 		RHIUnlockVertexBuffer(VertexBufferRHI);
 	}
 	
diff --git a/Engine/Source/Runtime/Engine/Private/GPUSkinVertexFactory.cpp b/Engine/Source/Runtime/Engine/Private/GPUSkinVertexFactory.cpp
--- a/Engine/Source/Runtime/Engine/Private/GPUSkinVertexFactory.cpp
+++ b/Engine/Source/Runtime/Engine/Private/GPUSkinVertexFactory.cpp
@@ -138,14 +138,14 @@ void FGPUBaseSkinVertexFactory::ShaderDataType::UpdateBoneData()
 			check(VectorArraySize == NumBones * sizeof(BoneMatrices[0]));
 			if (GRHIThread)
 			{
-				GRHICommandList.GetImmediateCommandList().UpdateVertexBuffer(BoneBuffer.VertexBufferRHI, BoneMatrices.GetTypedData(), NumBones * sizeof(BoneMatrices[0]));
+				GRHICommandList.GetImmediateCommandList().UpdateVertexBuffer(BoneBuffer.VertexBufferRHI, BoneMatrices.GetData(), NumBones * sizeof(BoneMatrices[0]));
 			}
 			else
 #endif
 			{
 				float* Data = (float*)RHILockVertexBuffer(BoneBuffer.VertexBufferRHI, 0, VectorArraySize, RLM_WriteOnly);
 				checkSlow(Data);
-				FMemory::Memcpy(Data, BoneMatrices.GetTypedData(), NumBones * sizeof(BoneMatrices[0]));
+				FMemory::Memcpy(Data, BoneMatrices.GetData(), NumBones * sizeof(BoneMatrices[0]));
 				RHIUnlockVertexBuffer(BoneBuffer.VertexBufferRHI);
 			}
 		}
@@ -153,7 +153,7 @@ void FGPUBaseSkinVertexFactory::ShaderDataType::UpdateBoneData()
 	else
 	{
 		check(NumBones * BoneMatrices.GetTypeSize() <= sizeof(GBoneUniformStruct));
-		FMemory::Memcpy(&GBoneUniformStruct, BoneMatrices.GetTypedData(), NumBones * sizeof(BoneMatrices[0]));
+		FMemory::Memcpy(&GBoneUniformStruct, BoneMatrices.GetData(), NumBones * sizeof(BoneMatrices[0]));
 		UniformBuffer = RHICreateUniformBuffer(&GBoneUniformStruct, FBoneMatricesUniformShaderParameters::StaticStruct.GetLayout(), UniformBuffer_MultiFrame);
 	}
 }
@@ -446,7 +446,7 @@ public:
 					// copy the bone data and tell the instance where it can pick it up next frame
 					check(IsInRenderingThread());
 					// append data to a buffer we bind next frame to read old matrix data for motion blur
-					uint32 OldBoneDataStartIndex = GPrevPerBoneMotionBlur.AppendData(ShaderData.BoneMatrices.GetTypedData(), ShaderData.BoneMatrices.Num());
+					uint32 OldBoneDataStartIndex = GPrevPerBoneMotionBlur.AppendData(ShaderData.BoneMatrices.GetData(), ShaderData.BoneMatrices.Num());
 					GPUVertexFactory->SetOldBoneDataStartIndex(View.FrameNumber, OldBoneDataStartIndex);
 				}
 			}
@@ -721,12 +721,12 @@ void FGPUBaseSkinAPEXClothVertexFactory::ClothShaderType::UpdateClothSimulData(c
 		{
 			float* Data = (float*)RHILockVertexBuffer(ClothSimulPositionBuffer.VertexBufferRHI, 0, VectorArraySize, RLM_WriteOnly);
 			checkSlow(Data);
-			FMemory::Memcpy(Data, InSimulPositions.GetTypedData(), NumSimulVerts * sizeof(FVector4));
+			FMemory::Memcpy(Data, InSimulPositions.GetData(), NumSimulVerts * sizeof(FVector4));
 			RHIUnlockVertexBuffer(ClothSimulPositionBuffer.VertexBufferRHI);
 
 			Data = (float*)RHILockVertexBuffer(ClothSimulNormalBuffer.VertexBufferRHI, 0, VectorArraySize, RLM_WriteOnly);
 			checkSlow(Data);
-			FMemory::Memcpy(Data, InSimulNormals.GetTypedData(), NumSimulVerts * sizeof(FVector4));
+			FMemory::Memcpy(Data, InSimulNormals.GetData(), NumSimulVerts * sizeof(FVector4));
 			RHIUnlockVertexBuffer(ClothSimulNormalBuffer.VertexBufferRHI);
 		}
 	}
