TArray documentation and a coupld of methods deprecation.

diff --git a/Engine/Source/Runtime/Engine/Private/GPUSort.cpp b/Engine/Source/Runtime/Engine/Private/GPUSort.cpp
--- a/Engine/Source/Runtime/Engine/Private/GPUSort.cpp
+++ b/Engine/Source/Runtime/Engine/Private/GPUSort.cpp
@@ -131,7 +131,7 @@ public:
 		OutOffsets.Empty( OffsetsCount );
 		OutOffsets.AddUninitialized( OffsetsCount );
 		uint32* MappedOffsets = (uint32*)RHILockVertexBuffer( Buffers[BufferIndex], 0, OffsetsBufferSize, RLM_ReadOnly );
-		FMemory::Memcpy( OutOffsets.GetTypedData(), MappedOffsets, OffsetsBufferSize );
+		FMemory::Memcpy( OutOffsets.GetData(), MappedOffsets, OffsetsBufferSize );
 		RHIUnlockVertexBuffer( Buffers[BufferIndex] );
 	}
 
@@ -884,10 +884,10 @@ static bool RunGPUSortTest(FRHICommandListImmediate& RHICmdList, int32 TestSize)
 		uint32* Buffer;
 
 		Buffer = (uint32*)RHICmdList.LockVertexBuffer(KeysBufferRHI[0], /*Offset=*/ 0, BufferSize, RLM_WriteOnly);
-		FMemory::Memcpy(Buffer, Keys.GetTypedData(), BufferSize);
+		FMemory::Memcpy(Buffer, Keys.GetData(), BufferSize);
 		RHICmdList.UnlockVertexBuffer(KeysBufferRHI[0]);
 		Buffer = (uint32*)RHICmdList.LockVertexBuffer(ValuesBufferRHI[0], /*Offset=*/ 0, BufferSize, RLM_WriteOnly);
-		FMemory::Memcpy(Buffer, Keys.GetTypedData(), BufferSize);
+		FMemory::Memcpy(Buffer, Keys.GetData(), BufferSize);
 		RHICmdList.UnlockVertexBuffer(ValuesBufferRHI[0]);
 	}
 
@@ -911,10 +911,10 @@ static bool RunGPUSortTest(FRHICommandListImmediate& RHICmdList, int32 TestSize)
 		SortedValues.AddUninitialized(TestSize);
 
 		Buffer = (uint32*)RHICmdList.LockVertexBuffer(KeysBufferRHI[ResultBufferIndex], /*Offset=*/ 0, BufferSize, RLM_ReadOnly);
-		FMemory::Memcpy(SortedKeys.GetTypedData(), Buffer, BufferSize);
+		FMemory::Memcpy(SortedKeys.GetData(), Buffer, BufferSize);
 		RHICmdList.UnlockVertexBuffer(KeysBufferRHI[ResultBufferIndex]);
 		Buffer = (uint32*)RHICmdList.LockVertexBuffer(ValuesBufferRHI[ResultBufferIndex], /*Offset=*/ 0, BufferSize, RLM_ReadOnly);
-		FMemory::Memcpy(SortedValues.GetTypedData(), Buffer, BufferSize);
+		FMemory::Memcpy(SortedValues.GetData(), Buffer, BufferSize);
 		RHICmdList.UnlockVertexBuffer(ValuesBufferRHI[ResultBufferIndex]);
 	}
 
diff --git a/Engine/Source/Runtime/Engine/Private/GameViewportClient.cpp b/Engine/Source/Runtime/Engine/Private/GameViewportClient.cpp
--- a/Engine/Source/Runtime/Engine/Private/GameViewportClient.cpp
+++ b/Engine/Source/Runtime/Engine/Private/GameViewportClient.cpp
@@ -1161,7 +1161,7 @@ void UGameViewportClient::ProcessScreenShots(FViewport* InViewport)
 						bWriteAlpha = GetHighResScreenshotConfig().MergeMaskIntoAlpha(Bitmap);
 						SourceRect = GetHighResScreenshotConfig().CaptureRegion;
 					}
-					FFileHelper::CreateBitmap(*ScreenShotName, InViewport->GetSizeXY().X, InViewport->GetSizeXY().Y, Bitmap.GetTypedData(), &SourceRect, &IFileManager::Get(), NULL, bWriteAlpha);
+					FFileHelper::CreateBitmap(*ScreenShotName, InViewport->GetSizeXY().X, InViewport->GetSizeXY().Y, Bitmap.GetData(), &SourceRect, &IFileManager::Get(), NULL, bWriteAlpha);
 				}
 			}
 		}
diff --git a/Engine/Source/Runtime/Engine/Private/HighResScreenshot.cpp b/Engine/Source/Runtime/Engine/Private/HighResScreenshot.cpp
--- a/Engine/Source/Runtime/Engine/Private/HighResScreenshot.cpp
+++ b/Engine/Source/Runtime/Engine/Private/HighResScreenshot.cpp
@@ -121,7 +121,7 @@ FString FHighResScreenshotConfig::SaveImage(const FString& File, const TArray<TP
 		{
 			const TArray<uint8>& CompressedData = ImageCompressor->GetCompressed();
 			int32 CompressedSize = CompressedData.Num();
-			Ar->Serialize((void*)CompressedData.GetTypedData(), CompressedSize);
+			Ar->Serialize((void*)CompressedData.GetData(), CompressedSize);
 			delete Ar;
 		}
 		else
diff --git a/Engine/Source/Runtime/Engine/Private/Landscape/LandscapeCollision.cpp b/Engine/Source/Runtime/Engine/Private/Landscape/LandscapeCollision.cpp
--- a/Engine/Source/Runtime/Engine/Private/Landscape/LandscapeCollision.cpp
+++ b/Engine/Source/Runtime/Engine/Private/Landscape/LandscapeCollision.cpp
@@ -139,7 +139,7 @@ void ULandscapeHeightfieldCollisionComponent::CreatePhysicsState()
 
 				// Create the sync scene actor
 				PxRigidStatic* HeightFieldActorSync = GPhysXSDK->createRigidStatic(PhysXLandscapeComponentTransform);
-				PxShape* HeightFieldShapeSync = HeightFieldActorSync->createShape(LandscapeComponentGeom, HeightfieldRef->UsedPhysicalMaterialArray.GetTypedData(), HeightfieldRef->UsedPhysicalMaterialArray.Num());
+				PxShape* HeightFieldShapeSync = HeightFieldActorSync->createShape(LandscapeComponentGeom, HeightfieldRef->UsedPhysicalMaterialArray.GetData(), HeightfieldRef->UsedPhysicalMaterialArray.Num());
 				check(HeightFieldShapeSync);
 
 				// Setup filtering
@@ -186,7 +186,7 @@ void ULandscapeHeightfieldCollisionComponent::CreatePhysicsState()
 				{
 					// Create the async scene actor
 					HeightFieldActorAsync = GPhysXSDK->createRigidStatic(PhysXLandscapeComponentTransform);
-					PxShape* HeightFieldShapeAsync = HeightFieldActorAsync->createShape(LandscapeComponentGeom, HeightfieldRef->UsedPhysicalMaterialArray.GetTypedData(), HeightfieldRef->UsedPhysicalMaterialArray.Num());
+					PxShape* HeightFieldShapeAsync = HeightFieldActorAsync->createShape(LandscapeComponentGeom, HeightfieldRef->UsedPhysicalMaterialArray.GetData(), HeightfieldRef->UsedPhysicalMaterialArray.Num());
 
 					HeightFieldShapeAsync->setQueryFilterData(PQueryFilterData);
 					HeightFieldShapeAsync->setSimulationFilterData(PSimFilterData);
@@ -674,7 +674,7 @@ void ULandscapeMeshCollisionComponent::CreatePhysicsState()
 
 				// Create the sync scene actor
 				PxRigidStatic* MeshActorSync = GPhysXSDK->createRigidStatic(PhysXLandscapeComponentTransform);
-				PxShape* MeshShapeSync = MeshActorSync->createShape(PTriMeshGeom, MeshRef->UsedPhysicalMaterialArray.GetTypedData(), MeshRef->UsedPhysicalMaterialArray.Num());
+				PxShape* MeshShapeSync = MeshActorSync->createShape(PTriMeshGeom, MeshRef->UsedPhysicalMaterialArray.GetData(), MeshRef->UsedPhysicalMaterialArray.Num());
 				check(MeshShapeSync);
 
 				// Setup filtering
@@ -697,7 +697,7 @@ void ULandscapeMeshCollisionComponent::CreatePhysicsState()
 				{
 					// Create the async scene actor
 					MeshActorAsync = GPhysXSDK->createRigidStatic(PhysXLandscapeComponentTransform);
-					PxShape* MeshShapeAsync = MeshActorAsync->createShape(PTriMeshGeom, MeshRef->UsedPhysicalMaterialArray.GetTypedData(), MeshRef->UsedPhysicalMaterialArray.Num());
+					PxShape* MeshShapeAsync = MeshActorAsync->createShape(PTriMeshGeom, MeshRef->UsedPhysicalMaterialArray.GetData(), MeshRef->UsedPhysicalMaterialArray.Num());
 					check(MeshShapeAsync);
 
 					MeshShapeAsync->setQueryFilterData(PQueryFilterData);
diff --git a/Engine/Source/Runtime/Engine/Private/Landscape/LandscapeEdit.cpp b/Engine/Source/Runtime/Engine/Private/Landscape/LandscapeEdit.cpp
--- a/Engine/Source/Runtime/Engine/Private/Landscape/LandscapeEdit.cpp
+++ b/Engine/Source/Runtime/Engine/Private/Landscape/LandscapeEdit.cpp
@@ -525,7 +525,7 @@ void ULandscapeComponent::UpdateCollisionHeightData(const FColor* HeightmapTextu
 				DominantLayerData.AddUninitialized(FMath::Square(CollisionSizeVerts));
 
 				const uint8* SrcDominantLayerData = (uint8*)CollisionComp->DominantLayerData.Lock(LOCK_READ_ONLY);
-				FMemory::Memcpy(DominantLayerData.GetTypedData(), SrcDominantLayerData, sizeof(uint8)*FMath::Square(CollisionSizeVerts));
+				FMemory::Memcpy(DominantLayerData.GetData(), SrcDominantLayerData, sizeof(uint8)*FMath::Square(CollisionSizeVerts));
 				CollisionComp->DominantLayerData.Unlock();
 			}
 
@@ -574,7 +574,7 @@ void ULandscapeComponent::UpdateCollisionHeightData(const FColor* HeightmapTextu
 			{
 				MeshCollisionComponent->DominantLayerData.Lock(LOCK_READ_WRITE);
 				uint8* DestDominantLayerData = (uint8*)MeshCollisionComponent->DominantLayerData.Realloc(FMath::Square(CollisionSizeVerts));
-				FMemory::Memcpy(DestDominantLayerData, DominantLayerData.GetTypedData(), sizeof(uint8)*FMath::Square(CollisionSizeVerts));
+				FMemory::Memcpy(DestDominantLayerData, DominantLayerData.GetData(), sizeof(uint8)*FMath::Square(CollisionSizeVerts));
 				MeshCollisionComponent->DominantLayerData.Unlock();
 			}
 
@@ -918,7 +918,7 @@ void ULandscapeComponent::UpdateCollisionLayerData()
 	{
 		TArray<uint8>* MipData = new(CachedWeightmapTextureMipData)TArray<uint8>();
 		WeightmapTextures[Idx]->Source.GetMipData(*MipData, CollisionMipLevel);
-		WeightmapTextureMipData.Add((FColor*)MipData->GetTypedData());
+		WeightmapTextureMipData.Add((FColor*)MipData->GetData());
 	}
 
 	UpdateCollisionLayerData(WeightmapTextureMipData);
@@ -1193,7 +1193,7 @@ void ULandscapeComponent::CreateEmptyTextureMips(UTexture2D* Texture, bool bClea
 		Texture->Source.Init2DWithMipChain(WeightmapSizeU, WeightmapSizeV, WeightmapFormat);
 		int32 NumMips = Texture->Source.GetNumMips();
 		uint8* MipData = Texture->Source.LockMip(0);
-		FMemory::Memcpy(MipData, TopMipData.GetTypedData(), TopMipData.Num());
+		FMemory::Memcpy(MipData, TopMipData.GetData(), TopMipData.Num());
 		Texture->Source.UnlockMip(0);
 	}
 }
@@ -2319,7 +2319,7 @@ bool ALandscapeProxy::ExportToRawMesh(int32 InExportLOD, FRawMesh& OutRawMesh) c
 
 		const int32 VisThreshold = 170;
 		const int32 WeightMapSize = (SubsectionSizeQuadsLOD + 1) * Component->NumSubsections;
-		uint32* Faces = OutRawMesh.WedgeIndices.GetTypedData() + IndicesOffset;
+		uint32* Faces = OutRawMesh.WedgeIndices.GetData() + IndicesOffset;
 
 		// Export verts
 		int32 VertexIdx = VerticesOffset;
@@ -2566,7 +2566,7 @@ void ULandscapeInfo::Export(const TArray<ULandscapeLayerInfoObject*>& LayerInfos
 
 	TArray<uint8> HeightData;
 	HeightData.AddZeroed((1 + MaxX - MinX)*(1 + MaxY - MinY)*sizeof(uint16));
-	LandscapeEdit.GetHeightDataFast(MinX, MinY, MaxX, MaxY, (uint16*)HeightData.GetTypedData(), 0);
+	LandscapeEdit.GetHeightDataFast(MinX, MinY, MaxX, MaxY, (uint16*)HeightData.GetData(), 0);
 	FFileHelper::SaveArrayToFile(HeightData, *Filenames[0]);
 
 	for (int32 i = 1; i < Filenames.Num(); i++)
@@ -2578,7 +2578,7 @@ void ULandscapeInfo::Export(const TArray<ULandscapeLayerInfoObject*>& LayerInfos
 			ULandscapeLayerInfoObject* LayerInfo = LayerInfos[i - 1];
 			if (LayerInfo)
 			{
-				LandscapeEdit.GetWeightDataFast(LayerInfo, MinX, MinY, MaxX, MaxY, WeightData.GetTypedData(), 0);
+				LandscapeEdit.GetWeightDataFast(LayerInfo, MinX, MinY, MaxX, MaxY, WeightData.GetData(), 0);
 			}
 			FFileHelper::SaveArrayToFile(WeightData, *Filenames[i]);
 		}
@@ -2862,7 +2862,7 @@ bool ULandscapeInfo::ApplySplinesInternal(bool bOnlySelected, ALandscapeProxy* L
 			int32 ValidMinY = MinY;
 			int32 ValidMaxX = MaxX;
 			int32 ValidMaxY = MaxY;
-			LandscapeEdit.GetHeightData(ValidMinX, ValidMinY, ValidMaxX, ValidMaxY, Data.GetTypedData(), 0);
+			LandscapeEdit.GetHeightData(ValidMinX, ValidMinY, ValidMaxX, ValidMaxY, Data.GetData(), 0);
 
 			if (ValidMinX > ValidMaxX || ValidMinY > ValidMaxY)
 			{
@@ -2907,7 +2907,7 @@ bool ULandscapeInfo::ApplySplinesInternal(bool bOnlySelected, ALandscapeProxy* L
 				Rasterizer.DrawTriangle(FalloffRight0, Left1, FalloffLeft1, FalloffRight0Pos, Left1Pos, FalloffLeft1Pos, false);
 			}
 
-			LandscapeEdit.SetHeightData(MinX, MinY, MaxX, MaxY, Data.GetTypedData(), 0, true);
+			LandscapeEdit.SetHeightData(MinX, MinY, MaxX, MaxY, Data.GetData(), 0, true);
 			LandscapeEdit.GetComponentsInRegion(MinX, MinY, MaxX, MaxY, &ModifiedComponents);
 		}
 
@@ -2922,7 +2922,7 @@ bool ULandscapeInfo::ApplySplinesInternal(bool bOnlySelected, ALandscapeProxy* L
 			int32 ValidMinY = MinY;
 			int32 ValidMaxX = MaxX;
 			int32 ValidMaxY = MaxY;
-			LandscapeEdit.GetWeightData(LayerInfo, ValidMinX, ValidMinY, ValidMaxX, ValidMaxY, Data.GetTypedData(), 0);
+			LandscapeEdit.GetWeightData(LayerInfo, ValidMinX, ValidMinY, ValidMaxX, ValidMaxY, Data.GetData(), 0);
 
 			if (ValidMinX > ValidMaxX || ValidMinY > ValidMaxY)
 			{
@@ -2968,7 +2968,7 @@ bool ULandscapeInfo::ApplySplinesInternal(bool bOnlySelected, ALandscapeProxy* L
 				Rasterizer.DrawTriangle(FalloffRight0, Left1, FalloffLeft1, FalloffRight0Pos, Left1Pos, FalloffLeft1Pos, false);
 			}
 
-			LandscapeEdit.SetAlphaData(LayerInfo, MinX, MinY, MaxX, MaxY, Data.GetTypedData(), 0, ELandscapeLayerPaintingRestriction::None, true, false);
+			LandscapeEdit.SetAlphaData(LayerInfo, MinX, MinY, MaxX, MaxY, Data.GetData(), 0, ELandscapeLayerPaintingRestriction::None, true, false);
 			LandscapeEdit.GetComponentsInRegion(MinX, MinY, MaxX, MaxY, &ModifiedComponents);
 		}
 	}
@@ -3016,7 +3016,7 @@ bool ULandscapeInfo::ApplySplinesInternal(bool bOnlySelected, ALandscapeProxy* L
 			int32 ValidMinY = MinY;
 			int32 ValidMaxX = MaxX;
 			int32 ValidMaxY = MaxY;
-			LandscapeEdit.GetHeightData(ValidMinX, ValidMinY, ValidMaxX, ValidMaxY, Data.GetTypedData(), 0);
+			LandscapeEdit.GetHeightData(ValidMinX, ValidMinY, ValidMaxX, ValidMaxY, Data.GetData(), 0);
 
 			if (ValidMinX > ValidMaxX || ValidMinY > ValidMaxY)
 			{
@@ -3066,7 +3066,7 @@ bool ULandscapeInfo::ApplySplinesInternal(bool bOnlySelected, ALandscapeProxy* L
 				Rasterizer.DrawTriangle(FalloffRight0, Right1, FalloffRight1, FalloffRight0Pos, Right1Pos, FalloffRight1Pos, false);
 			}
 
-			LandscapeEdit.SetHeightData(MinX, MinY, MaxX, MaxY, Data.GetTypedData(), 0, true);
+			LandscapeEdit.SetHeightData(MinX, MinY, MaxX, MaxY, Data.GetData(), 0, true);
 			LandscapeEdit.GetComponentsInRegion(MinX, MinY, MaxX, MaxY, &ModifiedComponents);
 		}
 
@@ -3081,7 +3081,7 @@ bool ULandscapeInfo::ApplySplinesInternal(bool bOnlySelected, ALandscapeProxy* L
 			int32 ValidMinY = MinY;
 			int32 ValidMaxX = MaxX;
 			int32 ValidMaxY = MaxY;
-			LandscapeEdit.GetWeightData(LayerInfo, ValidMinX, ValidMinY, ValidMaxX, ValidMaxY, Data.GetTypedData(), 0);
+			LandscapeEdit.GetWeightData(LayerInfo, ValidMinX, ValidMinY, ValidMaxX, ValidMaxY, Data.GetData(), 0);
 
 			if (ValidMinX > ValidMaxX || ValidMinY > ValidMaxY)
 			{
@@ -3132,7 +3132,7 @@ bool ULandscapeInfo::ApplySplinesInternal(bool bOnlySelected, ALandscapeProxy* L
 				Rasterizer.DrawTriangle(FalloffRight0, Right1, FalloffRight1, FalloffRight0Pos, Right1Pos, FalloffRight1Pos, false);
 			}
 
-			LandscapeEdit.SetAlphaData(LayerInfo, MinX, MinY, MaxX, MaxY, Data.GetTypedData(), 0, ELandscapeLayerPaintingRestriction::None, true, false);
+			LandscapeEdit.SetAlphaData(LayerInfo, MinX, MinY, MaxX, MaxY, Data.GetData(), 0, ELandscapeLayerPaintingRestriction::None, true, false);
 			LandscapeEdit.GetComponentsInRegion(MinX, MinY, MaxX, MaxY, &ModifiedComponents);
 		}
 	}
@@ -3452,7 +3452,7 @@ void ALandscapeProxy::RecreateCollisionComponents()
 			Comp->CollisionMipLevel = CollisionMipLevel;
 			TArray<uint8> CollisionMipData;
 			Comp->HeightmapTexture->Source.GetMipData(CollisionMipData, CollisionMipLevel);
-			Comp->UpdateCollisionHeightData((FColor*)CollisionMipData.GetTypedData(), 0, 0, MAX_int32, MAX_int32, true, NULL, true); // Rebuild for new CollisionMipLevel
+			Comp->UpdateCollisionHeightData((FColor*)CollisionMipData.GetData(), 0, 0, MAX_int32, MAX_int32, true, NULL, true); // Rebuild for new CollisionMipLevel
 		}
 	}
 }
@@ -3955,7 +3955,7 @@ void ULandscapeComponent::PostEditChangeProperty(FPropertyChangedEvent& Property
 		CollisionMipLevel = FMath::Clamp<int32>(CollisionMipLevel, 0, FMath::CeilLogTwo(SubsectionSizeQuads + 1) - 1);
 		TArray<uint8> CollisionMipData;
 		HeightmapTexture->Source.GetMipData(CollisionMipData, CollisionMipLevel);
-		UpdateCollisionHeightData((FColor*)CollisionMipData.GetTypedData(), 0, 0, MAX_int32, MAX_int32, true, NULL, true); // Rebuild for new CollisionMipLevel
+		UpdateCollisionHeightData((FColor*)CollisionMipData.GetData(), 0, 0, MAX_int32, MAX_int32, true, NULL, true); // Rebuild for new CollisionMipLevel
 	}
 }
 
@@ -4403,7 +4403,7 @@ void ULandscapeComponent::InitHeightmapData(TArray<FColor>& Heights, bool bUpdat
 		FColor* HeightmapTextureData = (FColor*)HeightmapTexture->Source.LockMip(Mip);
 		if (Mip == 0)
 		{
-			FMemory::Memcpy(HeightmapTextureData, Heights.GetTypedData(), MipSizeU*MipSizeV*sizeof(FColor));
+			FMemory::Memcpy(HeightmapTextureData, Heights.GetData(), MipSizeU*MipSizeV*sizeof(FColor));
 		}
 		else
 		{
@@ -4954,7 +4954,7 @@ void ULandscapeComponent::GeneratePlatformVertexData()
 	check(VertexOrder.Num() == FMath::Square(SubsectionSizeVerts) * FMath::Square(NumSubsections));
 
 	// Fill in the vertices in the specified order
-	FLandscapeMobileVertex* DstVert = (FLandscapeMobileVertex*)NewPlatformData.GetTypedData();
+	FLandscapeMobileVertex* DstVert = (FLandscapeMobileVertex*)NewPlatformData.GetData();
 	for (int32 Idx = 0; Idx < VertexOrder.Num(); Idx++)
 	{
 		int32 X = VertexOrder[Idx].X;
diff --git a/Engine/Source/Runtime/Engine/Private/Landscape/LandscapeEditInterface.cpp b/Engine/Source/Runtime/Engine/Private/Landscape/LandscapeEditInterface.cpp
--- a/Engine/Source/Runtime/Engine/Private/Landscape/LandscapeEditInterface.cpp
+++ b/Engine/Source/Runtime/Engine/Private/Landscape/LandscapeEditInterface.cpp
@@ -3576,7 +3576,7 @@ void FLandscapeEditDataInterface::CopyTextureFromHeightmap(UTexture2D* Dest, int
 		DataInterface.GetHeightmapTextureData(Heightmap);
 
 		uint8* DestTextureData = (uint8*)DestDataInfo->GetMipData(MipIdx) + ChannelOffsets[DestChannel];
-		uint8* SrcTextureData = (uint8*)Heightmap.GetTypedData() + ChannelOffsets[SrcChannel];
+		uint8* SrcTextureData = (uint8*)Heightmap.GetData() + ChannelOffsets[SrcChannel];
 
 		for( int32 i=0;i<FMath::Square(MipSize);i++ )
 		{
diff --git a/Engine/Source/Runtime/Engine/Private/Landscape/LandscapeGizmoActor.cpp b/Engine/Source/Runtime/Engine/Private/Landscape/LandscapeGizmoActor.cpp
--- a/Engine/Source/Runtime/Engine/Private/Landscape/LandscapeGizmoActor.cpp
+++ b/Engine/Source/Runtime/Engine/Private/Landscape/LandscapeGizmoActor.cpp
@@ -1009,7 +1009,7 @@ void ALandscapeGizmoActiveActor::Export(int32 Index, TArray<FString>& Filenames)
 			{
 				HeightData.AddZeroed((1+MaxX-MinX)*(1+MaxY-MinY)*sizeof(uint16));
 			}
-			uint16* pHeightData = (uint16*)HeightData.GetTypedData();
+			uint16* pHeightData = (uint16*)HeightData.GetData();
 
 			TArray<TArray<uint8> > WeightDatas;
 			for( int32 i=1;i<Filenames.Num();i++ )
diff --git a/Engine/Source/Runtime/Engine/Private/Landscape/LandscapeRender.cpp b/Engine/Source/Runtime/Engine/Private/Landscape/LandscapeRender.cpp
--- a/Engine/Source/Runtime/Engine/Private/Landscape/LandscapeRender.cpp
+++ b/Engine/Source/Runtime/Engine/Private/Landscape/LandscapeRender.cpp
@@ -672,7 +672,7 @@ FLandscapeComponentSceneProxy::FLandscapeComponentSceneProxy(ULandscapeComponent
 	{
 		for( int32 SubX=0;SubX<NumSubsections;SubX++ )
 		{
-			FMeshBatchElement* BatchElement = (SubX==0 && SubY==0) ? DynamicMesh.Elements.GetTypedData() : new(DynamicMesh.Elements) FMeshBatchElement;
+			FMeshBatchElement* BatchElement = (SubX==0 && SubY==0) ? DynamicMesh.Elements.GetData() : new(DynamicMesh.Elements) FMeshBatchElement;
 			BatchElement->PrimitiveUniformBufferResource = &GetUniformBuffer();
 			FLandscapeBatchElementParams* BatchElementParams = new(DynamicMeshBatchParamArray) FLandscapeBatchElementParams;
 			BatchElementParams->LocalToWorldNoScalingPtr = &LocalToWorldNoScaling;
@@ -685,7 +685,7 @@ FLandscapeComponentSceneProxy::FLandscapeComponentSceneProxy(ULandscapeComponent
 			BatchElementParams->SubY = SubY;
 
 #if WITH_EDITOR
-			FMeshBatchElement* ToolBatchElement = (SubX==0 && SubY==0) ? DynamicMeshTools.Elements.GetTypedData() : new(DynamicMeshTools.Elements) FMeshBatchElement;
+			FMeshBatchElement* ToolBatchElement = (SubX==0 && SubY==0) ? DynamicMeshTools.Elements.GetData() : new(DynamicMeshTools.Elements) FMeshBatchElement;
 			*ToolBatchElement = *BatchElement;
 #endif
 		}
@@ -1418,7 +1418,7 @@ void FLandscapeComponentSceneProxy::GetDynamicMeshElements(const TArray<const FS
 					int32 SubSectionIdx = SubX + SubY*NumSubsections;
 					int32 CurrentLOD = CalcLODForSubsection(*View, SubX, SubY, CameraLocalPos);
 
-					FMeshBatchElement& BatchElement = (SubX==0 && SubY==0) ? *Mesh.Elements.GetTypedData() : *(new(Mesh.Elements) FMeshBatchElement);
+					FMeshBatchElement& BatchElement = (SubX==0 && SubY==0) ? *Mesh.Elements.GetData() : *(new(Mesh.Elements) FMeshBatchElement);
 					BatchElement.PrimitiveUniformBufferResource = &GetUniformBuffer();
 					FLandscapeBatchElementParams& BatchElementParams = ParameterArray.ElementParams[SubSectionIdx];
 					BatchElementParams.LocalToWorldNoScalingPtr = &LocalToWorldNoScaling;
@@ -1449,7 +1449,7 @@ void FLandscapeComponentSceneProxy::GetDynamicMeshElements(const TArray<const FS
 					BatchElement.MaxVertexIndex = SharedBuffers->IndexRanges[CurrentLOD].MaxIndex[SubX][SubY];
 					
 #if WITH_EDITOR
-					FMeshBatchElement& BatchElementTools = (SubX==0 && SubY==0) ? *MeshTools.Elements.GetTypedData() : *(new(MeshTools.Elements) FMeshBatchElement);
+					FMeshBatchElement& BatchElementTools = (SubX==0 && SubY==0) ? *MeshTools.Elements.GetData() : *(new(MeshTools.Elements) FMeshBatchElement);
 					BatchElementTools.PrimitiveUniformBufferResource = &GetUniformBuffer();
 					BatchElementTools.UserData = &BatchElementParams;
 
diff --git a/Engine/Source/Runtime/Engine/Private/Landscape/LandscapeRenderMobile.cpp b/Engine/Source/Runtime/Engine/Private/Landscape/LandscapeRenderMobile.cpp
--- a/Engine/Source/Runtime/Engine/Private/Landscape/LandscapeRenderMobile.cpp
+++ b/Engine/Source/Runtime/Engine/Private/Landscape/LandscapeRenderMobile.cpp
@@ -205,7 +205,7 @@ void FLandscapeComponentSceneProxyMobile::CreateRenderThreadResources()
 
 	int32 VertexBufferSize = PlatformData.Num();
 	// Copy platform data into vertex buffer
-	VertexBuffer = new FLandscapeVertexBufferMobile(PlatformData.GetTypedData(), VertexBufferSize);
+	VertexBuffer = new FLandscapeVertexBufferMobile(PlatformData.GetData(), VertexBufferSize);
 
 	FLandscapeVertexFactoryMobile* LandscapeVertexFactory = new FLandscapeVertexFactoryMobile();
 	LandscapeVertexFactory->MobileData.PositionComponent = FVertexStreamComponent(VertexBuffer, STRUCT_OFFSET(FLandscapeMobileVertex,Position), sizeof(FLandscapeMobileVertex), VET_UByte4N);
diff --git a/Engine/Source/Runtime/Engine/Private/Level.cpp b/Engine/Source/Runtime/Engine/Private/Level.cpp
--- a/Engine/Source/Runtime/Engine/Private/Level.cpp
+++ b/Engine/Source/Runtime/Engine/Private/Level.cpp
@@ -771,7 +771,7 @@ static void SortActorsHierarchy(TTransArray<AActor*>& Actors)
 	
 	// Unfortunately TArray.StableSort assumes no null entries in the array
 	// So it forces me to use internal unrestricted version
-	StableSortInternal(Actors.GetTypedData(), Actors.Num(), [&](AActor* L, AActor* R) {
+	StableSortInternal(Actors.GetData(), Actors.Num(), [&](AActor* L, AActor* R) {
 			return CalcAttachDepth(L) < CalcAttachDepth(R);
 	});
 }
diff --git a/Engine/Source/Runtime/Engine/Private/LightMap.cpp b/Engine/Source/Runtime/Engine/Private/LightMap.cpp
--- a/Engine/Source/Runtime/Engine/Private/LightMap.cpp
+++ b/Engine/Source/Runtime/Engine/Private/LightMap.cpp
@@ -1142,7 +1142,7 @@ void FLightMap2D::EncodeTextures( UWorld* InWorld, bool bLightingSuccessful, boo
 		PendingLightMapSize = 0;
 
 		// Sort the light-maps in descending order by size.
-		Sort( PendingLightMaps.GetTypedData(), PendingLightMaps.Num(), FCompareLightmaps() );
+		Sort(PendingLightMaps.GetData(), PendingLightMaps.Num(), FCompareLightmaps());
 
 		// Allocate texture space for each light-map.
 		TArray<FLightMapPendingTexture*> PendingTextures;
diff --git a/Engine/Source/Runtime/Engine/Private/Model.cpp b/Engine/Source/Runtime/Engine/Private/Model.cpp
--- a/Engine/Source/Runtime/Engine/Private/Model.cpp
+++ b/Engine/Source/Runtime/Engine/Private/Model.cpp
@@ -537,7 +537,7 @@ void UModel::BuildBound()
 		for( int32 i=0; i<Polys->Element.Num(); i++ )
 			for( int32 j=0; j<Polys->Element[i].Vertices.Num(); j++ )
 				NewPoints.Add(Polys->Element[i].Vertices[j]);
-		Bounds = FBoxSphereBounds( NewPoints.GetTypedData(), NewPoints.Num() );
+		Bounds = FBoxSphereBounds( NewPoints.GetData(), NewPoints.Num() );
 	}
 }
 
diff --git a/Engine/Source/Runtime/Engine/Private/ModelComponent.cpp b/Engine/Source/Runtime/Engine/Private/ModelComponent.cpp
--- a/Engine/Source/Runtime/Engine/Private/ModelComponent.cpp
+++ b/Engine/Source/Runtime/Engine/Private/ModelComponent.cpp
@@ -353,7 +353,7 @@ void UModelComponent::GetStreamingTextureInfo(TArray<FStreamingTexturePrimitiveI
 						SurfaceVertices.Add(WorldVertex);
 					}
 				}
-				const FSphere SurfaceBoundingSphere(SurfaceVertices.GetTypedData(),SurfaceVertices.Num());
+				const FSphere SurfaceBoundingSphere(SurfaceVertices.GetData(),SurfaceVertices.Num());
 
 				// Compute the surface's texture scaling factor.
 				const float BspTexelsPerNormalizedTexel = UModel::GetGlobalBSPTexelScale();
diff --git a/Engine/Source/Runtime/Engine/Private/NavMeshRenderingComponent.cpp b/Engine/Source/Runtime/Engine/Private/NavMeshRenderingComponent.cpp
--- a/Engine/Source/Runtime/Engine/Private/NavMeshRenderingComponent.cpp
+++ b/Engine/Source/Runtime/Engine/Private/NavMeshRenderingComponent.cpp
@@ -263,7 +263,7 @@ void UNavMeshRenderingComponent::GatherData(struct FNavMeshSceneProxyData* Curre
 					const FNavigationOctreeElement& Element = *ElementIt;
 					if (Element.ShouldUseGeometry(&NavMesh->NavDataConfig) && Element.Data.CollisionData.Num())
 					{
-						const FRecastGeometryCache CachedGeometry(Element.Data.CollisionData.GetTypedData());
+						const FRecastGeometryCache CachedGeometry(Element.Data.CollisionData.GetData());
 						AppendGeometry(CurrentData->PathCollidingGeomVerts, CurrentData->PathCollidingGeomIndices,
 							CachedGeometry.Verts, CachedGeometry.Header.NumVerts, CachedGeometry.Indices, CachedGeometry.Header.NumFaces);
 					}
diff --git a/Engine/Source/Runtime/Engine/Private/NavMeshRenderingHelpers.h b/Engine/Source/Runtime/Engine/Private/NavMeshRenderingHelpers.h
--- a/Engine/Source/Runtime/Engine/Private/NavMeshRenderingHelpers.h
+++ b/Engine/Source/Runtime/Engine/Private/NavMeshRenderingHelpers.h
@@ -181,7 +181,7 @@ public:
 
 			// Write the indices to the index buffer.
 			void* Buffer = RHILockIndexBuffer(IndexBufferRHI, 0, Indices.Num() * sizeof(int32), RLM_WriteOnly);
-			FMemory::Memcpy(Buffer, Indices.GetTypedData(), Indices.Num() * sizeof(int32));
+			FMemory::Memcpy(Buffer, Indices.GetData(), Indices.Num() * sizeof(int32));
 			RHIUnlockIndexBuffer(IndexBufferRHI);
 		}
 	}
@@ -202,7 +202,7 @@ public:
 
 			// Copy the vertex data into the vertex buffer.
 			void* VertexBufferData = RHILockVertexBuffer(VertexBufferRHI, 0, Vertices.Num() * sizeof(FDynamicMeshVertex), RLM_WriteOnly);
-			FMemory::Memcpy(VertexBufferData, Vertices.GetTypedData(), Vertices.Num() * sizeof(FDynamicMeshVertex));
+			FMemory::Memcpy(VertexBufferData, Vertices.GetData(), Vertices.Num() * sizeof(FDynamicMeshVertex));
 			RHIUnlockVertexBuffer(VertexBufferRHI);
 		}
 	}
@@ -678,7 +678,7 @@ public:
 		Canvas->SetDrawColor(FColor::White);
 		const FSceneView* View = Canvas->SceneView;
 		UFont* Font = GEngine->GetSmallFont();
-		const FNavMeshSceneProxyData::FDebugText* DebugText = ProxyData.DebugLabels.GetTypedData();
+		const FNavMeshSceneProxyData::FDebugText* DebugText = ProxyData.DebugLabels.GetData();
 		for (int32 i = 0 ; i < ProxyData.DebugLabels.Num(); ++i, ++DebugText)
 		{
 			if (PointInView(DebugText->Location, View))
@@ -739,7 +739,7 @@ FORCEINLINE void AppendGeometry(TNavStatArray<FVector>& OutVertexBuffer, TArray<
 
 	const int32 FirstNewFaceVertexIndex = OutIndexBuffer.Num();
 	OutIndexBuffer.AddUninitialized(NumFaces * 3);
-	int32* DestFaceVertIndex = OutIndexBuffer.GetTypedData() + FirstNewFaceVertexIndex;
+	int32* DestFaceVertIndex = OutIndexBuffer.GetData() + FirstNewFaceVertexIndex;
 	const int32* SrcFaceVertIndex = Faces;
 	// copy with offset
 	for (int32 Index = 0; Index < NumFaces * 3; ++Index, ++DestFaceVertIndex, ++SrcFaceVertIndex)
diff --git a/Engine/Source/Runtime/Engine/Private/NetConnection.cpp b/Engine/Source/Runtime/Engine/Private/NetConnection.cpp
--- a/Engine/Source/Runtime/Engine/Private/NetConnection.cpp
+++ b/Engine/Source/Runtime/Engine/Private/NetConnection.cpp
@@ -500,7 +500,7 @@ void UNetConnection::FlushNet(bool bIgnoreSimulation)
 		{
 			DelayedPacket& B = *(new(Delayed)DelayedPacket);
 			B.Data.AddUninitialized( SendBuffer.GetNumBytes() );
-			FMemory::Memcpy( B.Data.GetTypedData(), SendBuffer.GetData(), SendBuffer.GetNumBytes() );
+			FMemory::Memcpy( B.Data.GetData(), SendBuffer.GetData(), SendBuffer.GetNumBytes() );
 
 			for( int32 i=Delayed.Num()-1; i>=0; i-- )
 			{
@@ -524,7 +524,7 @@ void UNetConnection::FlushNet(bool bIgnoreSimulation)
 			{
 				DelayedPacket& B = *(new(Delayed)DelayedPacket);
 				B.Data.AddUninitialized( SendBuffer.GetNumBytes() );
-				FMemory::Memcpy( B.Data.GetTypedData(), SendBuffer.GetData(), SendBuffer.GetNumBytes() );
+				FMemory::Memcpy( B.Data.GetData(), SendBuffer.GetData(), SendBuffer.GetNumBytes() );
 				B.SendTime = FPlatformTime::Seconds() + (double(PacketSimulationSettings.PktLag)  + 2.0f * (FMath::FRand() - 0.5f) * double(PacketSimulationSettings.PktLagVariance))/ 1000.f;
 			}
 		}
diff --git a/Engine/Source/Runtime/Engine/Private/OpusAudioInfo.cpp b/Engine/Source/Runtime/Engine/Private/OpusAudioInfo.cpp
--- a/Engine/Source/Runtime/Engine/Private/OpusAudioInfo.cpp
+++ b/Engine/Source/Runtime/Engine/Private/OpusAudioInfo.cpp
@@ -410,7 +410,7 @@ int32 FOpusAudioInfo::DecompressToPCMBuffer(uint16 FrameSize)
 	const uint8* SrcPtr = SrcBufferData + SrcBufferOffset;
 	SrcBufferOffset += FrameSize;
 	LastPCMOffset = 0;
-	return OpusDecoderWrapper->Decode(SrcPtr, FrameSize, (int16*)LastDecodedPCM.GetTypedData(), MaxFrameSizeSamples);
+	return OpusDecoderWrapper->Decode(SrcPtr, FrameSize, (int16*)LastDecodedPCM.GetData(), MaxFrameSizeSamples);
 }
 
 uint32 FOpusAudioInfo::IncrementCurrentSampleCount(uint32 NewSamples)
@@ -432,7 +432,7 @@ uint32 FOpusAudioInfo::WriteFromDecodedPCM(uint8* Destination, uint32 BufferSize
 	uint32 BytesToCopy = FMath::Min(BufferSize, LastPCMByteSize - LastPCMOffset);
 	if (BytesToCopy > 0)
 	{
-		FMemory::Memcpy(Destination, LastDecodedPCM.GetTypedData() + LastPCMOffset, BytesToCopy);
+		FMemory::Memcpy(Destination, LastDecodedPCM.GetData() + LastPCMOffset, BytesToCopy);
 		LastPCMOffset += BytesToCopy;
 		if (LastPCMOffset >= LastPCMByteSize)
 		{
diff --git a/Engine/Source/Runtime/Engine/Private/PackageMapClient.cpp b/Engine/Source/Runtime/Engine/Private/PackageMapClient.cpp
--- a/Engine/Source/Runtime/Engine/Private/PackageMapClient.cpp
+++ b/Engine/Source/Runtime/Engine/Private/PackageMapClient.cpp
@@ -252,7 +252,7 @@ bool UPackageMapClient::SerializeNewActor(FArchive& Ar, class UActorChannel *Cha
 			if ( Channel && Ar.IsSaving() )
 			{
 				FObjectReplicator * RepData = &Channel->GetActorReplicationData();
-				uint8 * Recent = RepData && RepData->RepState != NULL && RepData->RepState->StaticBuffer.Num() ? RepData->RepState->StaticBuffer.GetTypedData() : NULL;
+				uint8 * Recent = RepData && RepData->RepState != NULL && RepData->RepState->StaticBuffer.Num() ? RepData->RepState->StaticBuffer.GetData() : NULL;
 				if ( Recent )
 				{
 					((AActor*)Recent)->ReplicatedMovement.Location = Location;
diff --git a/Engine/Source/Runtime/Engine/Private/Particles/ParticleCurveTexture.cpp b/Engine/Source/Runtime/Engine/Private/Particles/ParticleCurveTexture.cpp
--- a/Engine/Source/Runtime/Engine/Private/Particles/ParticleCurveTexture.cpp
+++ b/Engine/Source/Runtime/Engine/Private/Particles/ParticleCurveTexture.cpp
@@ -544,7 +544,7 @@ FTexelAllocation FParticleCurveTexture::AddCurve( const TArray<FColor>& CurveSam
 			FCurveSamples* PendingCurve = new(PendingCurves) FCurveSamples;
 			PendingCurve->TexelAllocation = TexelAllocation;
 			PendingCurve->Samples = (FColor*)FMemory::Malloc( TexelAllocation.Size * sizeof(FColor) );
-			FMemory::Memcpy( PendingCurve->Samples, CurveSamples.GetTypedData(), TexelAllocation.Size * sizeof(FColor) );
+			FMemory::Memcpy( PendingCurve->Samples, CurveSamples.GetData(), TexelAllocation.Size * sizeof(FColor) );
 			return TexelAllocation;
 		}
 		UE_LOG(LogParticles, Warning, TEXT("FParticleCurveTexture: Failed to allocate %d texels for a curve."), CurveSamples.Num() );
diff --git a/Engine/Source/Runtime/Engine/Private/Particles/ParticleGpuSimulation.cpp b/Engine/Source/Runtime/Engine/Private/Particles/ParticleGpuSimulation.cpp
--- a/Engine/Source/Runtime/Engine/Private/Particles/ParticleGpuSimulation.cpp
+++ b/Engine/Source/Runtime/Engine/Private/Particles/ParticleGpuSimulation.cpp
@@ -1147,7 +1147,7 @@ static void BuildTileVertexBuffer( FParticleBufferParamRef TileOffsetsRef, const
  */
 static void BuildTileVertexBuffer( FParticleBufferParamRef TileOffsetsRef, const TArray<uint32>& Tiles )
 {
-	BuildTileVertexBuffer( TileOffsetsRef, Tiles.GetTypedData(), Tiles.Num() );
+	BuildTileVertexBuffer( TileOffsetsRef, Tiles.GetData(), Tiles.Num() );
 }
 
 /**
@@ -1313,7 +1313,7 @@ void ClearTiles(FRHICommandList& RHICmdList, ERHIFeatureLevel::Type FeatureLevel
 	{
 		// Copy new particles in to the vertex buffer.
 		const int32 TilesThisDrawCall = FMath::Min<int32>( TileCount, MaxTilesPerDrawCall );
-		const uint32* TilesPtr = Tiles.GetTypedData() + FirstTile;
+		const uint32* TilesPtr = Tiles.GetData() + FirstTile;
 		BuildTileVertexBuffer( BufferParam, TilesPtr, TilesThisDrawCall );
 		
 		VertexShader->SetParameters(RHICmdList, ShaderParam);
@@ -1531,7 +1531,7 @@ void InjectNewParticles(FRHICommandList& RHICmdList, ERHIFeatureLevel::Type Feat
 	{
 		// Copy new particles in to the vertex buffer.
 		const int32 ParticlesThisDrawCall = FMath::Min<int32>( ParticleCount, MaxParticlesPerDrawCall );
-		const void* Src = NewParticles.GetTypedData() + FirstParticle;
+		const void* Src = NewParticles.GetData() + FirstParticle;
 #if PLATFORM_SUPPORTS_RHI_THREAD
 		if (GRHIThread)
 		{
@@ -3271,13 +3271,13 @@ public:
 			if (BurstInfo.Count > 0)
 			{
 				// Spawn burst particles.
-				BuildNewParticles(NewParticles.GetTypedData() + FirstBurstParticleIndex, BurstInfo);
+				BuildNewParticles(NewParticles.GetData() + FirstBurstParticleIndex, BurstInfo);
 			}
 
 			if (SpawnInfo.Count > 0)
 			{
 				// Spawn normal particles.
-				BuildNewParticles(NewParticles.GetTypedData() + FirstSpawnParticleIndex, SpawnInfo);
+				BuildNewParticles(NewParticles.GetData() + FirstSpawnParticleIndex, SpawnInfo);
 			}
 
 			int32 NewParticleCount = BurstInfo.Count + SpawnInfo.Count;
@@ -3298,7 +3298,7 @@ public:
 			if (Component && Component->bWarmingUp)
 			{
 				SimulateWarmupParticles(
-					NewParticles.GetTypedData() + (NewParticles.Num() - NewParticleCount),
+					NewParticles.GetData() + (NewParticles.Num() - NewParticleCount),
 					NewParticleCount,
 					Component->WarmupTime - SecondsSinceCreation );
 			}
diff --git a/Engine/Source/Runtime/Engine/Private/Particles/ParticleSystemRender.cpp b/Engine/Source/Runtime/Engine/Private/Particles/ParticleSystemRender.cpp
--- a/Engine/Source/Runtime/Engine/Private/Particles/ParticleSystemRender.cpp
+++ b/Engine/Source/Runtime/Engine/Private/Particles/ParticleSystemRender.cpp
@@ -1899,7 +1899,7 @@ void FDynamicMeshEmitterData::GetDynamicMeshElementsEmitter(const FParticleSyste
 				}
 
 				// Fill instance buffer.
-				GetInstanceData((void*)InstanceVerticesCPU->InstanceDataAllocationsCPU.GetTypedData(), (void*)InstanceVerticesCPU->DynamicParameterDataAllocationsCPU.GetTypedData(), Proxy, View);
+				GetInstanceData((void*)InstanceVerticesCPU->InstanceDataAllocationsCPU.GetData(), (void*)InstanceVerticesCPU->DynamicParameterDataAllocationsCPU.GetData(), Proxy, View);
 			}
 
 			Proxy->UpdateWorldSpacePrimitiveUniformBuffer();
@@ -1971,8 +1971,8 @@ void FDynamicMeshEmitterData::GetDynamicMeshElementsEmitter(const FParticleSyste
 					if (!bInstanced)
 					{
 						FMeshParticleVertexFactory::FBatchParametersCPU& BatchParameters = Collector.AllocateOneFrameResource<FMeshParticleVertexFactory::FBatchParametersCPU>();
-						BatchParameters.InstanceBuffer = InstanceVerticesCPU->InstanceDataAllocationsCPU.GetTypedData();
-						BatchParameters.DynamicParameterBuffer = InstanceVerticesCPU->DynamicParameterDataAllocationsCPU.GetTypedData();
+						BatchParameters.InstanceBuffer = InstanceVerticesCPU->InstanceDataAllocationsCPU.GetData();
+						BatchParameters.DynamicParameterBuffer = InstanceVerticesCPU->DynamicParameterDataAllocationsCPU.GetData();
 						BatchElement.UserData = &BatchParameters;
 						BatchElement.UserIndex = 0;
 
@@ -2138,7 +2138,7 @@ void FDynamicMeshEmitterData::PreRenderView(FParticleSystemSceneProxy* Proxy, co
 			DynamicParameterDataAllocationsCPU.AddUninitialized(ParticleCount);
 
 			// Fill instance buffer.
-			GetInstanceData((void*) InstanceDataAllocationsCPU.GetTypedData(), (void*) DynamicParameterDataAllocationsCPU.GetTypedData(), Proxy, ViewFamily->Views[ViewIndex]);
+			GetInstanceData((void*) InstanceDataAllocationsCPU.GetData(), (void*) DynamicParameterDataAllocationsCPU.GetData(), Proxy, ViewFamily->Views[ViewIndex]);
 		}
 
 		// Instance data is not calculated per-view when the emitter is unsorted.
@@ -2239,8 +2239,8 @@ void FDynamicMeshEmitterData::PreRenderView(FParticleSystemSceneProxy* Proxy, co
 					if(!bInstanced)
 					{
 						FMeshParticleVertexFactory::FBatchParametersCPU* BatchParameters = new(MeshBatchParameters) FMeshParticleVertexFactory::FBatchParametersCPU();
-						BatchParameters->InstanceBuffer = InstanceDataAllocationsCPU.GetTypedData();
-						BatchParameters->DynamicParameterBuffer = DynamicParameterDataAllocationsCPU.GetTypedData();
+						BatchParameters->InstanceBuffer = InstanceDataAllocationsCPU.GetData();
+						BatchParameters->DynamicParameterBuffer = DynamicParameterDataAllocationsCPU.GetData();
 						BatchElement.UserData = BatchParameters;
 						BatchElement.UserIndex = 0;
 
diff --git a/Engine/Source/Runtime/Engine/Private/PhysicsEngine/BodyInstance.cpp b/Engine/Source/Runtime/Engine/Private/PhysicsEngine/BodyInstance.cpp
--- a/Engine/Source/Runtime/Engine/Private/PhysicsEngine/BodyInstance.cpp
+++ b/Engine/Source/Runtime/Engine/Private/PhysicsEngine/BodyInstance.cpp
@@ -1095,7 +1095,7 @@ void FBodyInstance::InitBody(UBodySetup* Setup, const FTransform& Transform, UPr
 						}
 
 						b2PolygonShape ConvexPoly;
-						ConvexPoly.Set(Verts.GetTypedData(), Verts.Num());
+						ConvexPoly.Set(Verts.GetData(), Verts.Num());
 
 						b2FixtureDef FixtureDef;
 						FixtureDef.shape = &ConvexPoly;
@@ -1251,7 +1251,7 @@ void FBodyInstance::InitBody(UBodySetup* Setup, const FTransform& Transform, UPr
 	{
 		Setup->AddShapesToRigidActor(PNewActorSync, Scale3D);
 		PNewActorSync->userData = &PhysxUserData; // Store pointer to owning bodyinstance.
-		PNewActorSync->setName( CharDebugName.IsValid() ? CharDebugName->GetTypedData() : NULL );
+		PNewActorSync->setName( CharDebugName.IsValid() ? CharDebugName->GetData() : NULL );
 
 		check(FPhysxUserData::Get<FBodyInstance>(PNewActorSync->userData) == this && FPhysxUserData::Get<FBodyInstance>(PNewActorSync->userData)->OwnerComponent != NULL);
 	}
@@ -1263,7 +1263,7 @@ void FBodyInstance::InitBody(UBodySetup* Setup, const FTransform& Transform, UPr
 
 		Setup->AddShapesToRigidActor(PNewActorAsync, Scale3D);
 		PNewActorAsync->userData = &PhysxUserData; // Store pointer to owning bodyinstance.
-		PNewActorAsync->setName( CharDebugName.IsValid() ? CharDebugName->GetTypedData() : NULL );
+		PNewActorAsync->setName( CharDebugName.IsValid() ? CharDebugName->GetData() : NULL );
 
 		check(FPhysxUserData::Get<FBodyInstance>(PNewActorAsync->userData) == this && FPhysxUserData::Get<FBodyInstance>(PNewActorAsync->userData)->OwnerComponent != NULL);
 	}
diff --git a/Engine/Source/Runtime/Engine/Private/PhysicsEngine/BodySetup.cpp b/Engine/Source/Runtime/Engine/Private/PhysicsEngine/BodySetup.cpp
--- a/Engine/Source/Runtime/Engine/Private/PhysicsEngine/BodySetup.cpp
+++ b/Engine/Source/Runtime/Engine/Private/PhysicsEngine/BodySetup.cpp
@@ -841,7 +841,7 @@ FByteBulkData* UBodySetup::GetCookedData(FName Format)
 			if (OutData.Num())
 			{
 				Result->Lock(LOCK_READ_WRITE);
-				FMemory::Memcpy(Result->Realloc(OutData.Num()), OutData.GetTypedData(), OutData.Num());
+				FMemory::Memcpy(Result->Realloc(OutData.Num()), OutData.GetData(), OutData.Num());
 				Result->Unlock();
 			}
 		}
diff --git a/Engine/Source/Runtime/Engine/Private/PhysicsEngine/PhysDrawing.cpp b/Engine/Source/Runtime/Engine/Private/PhysicsEngine/PhysDrawing.cpp
--- a/Engine/Source/Runtime/Engine/Private/PhysicsEngine/PhysDrawing.cpp
+++ b/Engine/Source/Runtime/Engine/Private/PhysicsEngine/PhysDrawing.cpp
@@ -478,7 +478,7 @@ void FConvexCollisionVertexBuffer::InitRHI()
 
 	// Copy the vertex data into the vertex buffer.
 	void* VertexBufferData = RHILockVertexBuffer(VertexBufferRHI,0,Vertices.Num() * sizeof(FDynamicMeshVertex), RLM_WriteOnly);
-	FMemory::Memcpy(VertexBufferData,Vertices.GetTypedData(),Vertices.Num() * sizeof(FDynamicMeshVertex));
+	FMemory::Memcpy(VertexBufferData,Vertices.GetData(),Vertices.Num() * sizeof(FDynamicMeshVertex));
 	RHIUnlockVertexBuffer(VertexBufferRHI);
 }
 
@@ -489,7 +489,7 @@ void FConvexCollisionIndexBuffer::InitRHI()
 
 	// Write the indices to the index buffer.
 	void* Buffer = RHILockIndexBuffer(IndexBufferRHI,0,Indices.Num() * sizeof(int32),RLM_WriteOnly);
-	FMemory::Memcpy(Buffer,Indices.GetTypedData(),Indices.Num() * sizeof(int32));
+	FMemory::Memcpy(Buffer,Indices.GetData(),Indices.Num() * sizeof(int32));
 	RHIUnlockIndexBuffer(IndexBufferRHI);
 }
 
diff --git a/Engine/Source/Runtime/Engine/Private/Polygon.cpp b/Engine/Source/Runtime/Engine/Private/Polygon.cpp
--- a/Engine/Source/Runtime/Engine/Private/Polygon.cpp
+++ b/Engine/Source/Runtime/Engine/Private/Polygon.cpp
@@ -373,10 +373,10 @@ int32 FPoly::SplitWithPlaneFast
 		// Split.
 		if( FrontPoly )
 		{
-			const FVector *V  = Vertices.GetTypedData();
+			const FVector *V  = Vertices.GetData();
 			const FVector *W  = V + Vertices.Num()-1;
-			FVector *V1       = FrontPoly->Vertices.GetTypedData();
-			FVector *V2       = BackPoly ->Vertices.GetTypedData();
+			FVector *V1       = FrontPoly->Vertices.GetData();
+			FVector *V2       = BackPoly ->Vertices.GetData();
 			StatusPtr         = &VertStatus        [0];
 			PrevStatus        = VertStatus         [Vertices.Num()-1];
 
diff --git a/Engine/Source/Runtime/Engine/Private/PoseableMeshComponent.cpp b/Engine/Source/Runtime/Engine/Private/PoseableMeshComponent.cpp
--- a/Engine/Source/Runtime/Engine/Private/PoseableMeshComponent.cpp
+++ b/Engine/Source/Runtime/Engine/Private/PoseableMeshComponent.cpp
@@ -87,8 +87,8 @@ void UPoseableMeshComponent::FillSpaceBases()
 	BoneProcessed.AddZeroed(NumBones);
 #endif
 	// Build in 3 passes.
-	FTransform* LocalTransformsData = LocalAtoms.GetTypedData(); 
-	FTransform* SpaceBasesData = SpaceBases.GetTypedData();
+	FTransform* LocalTransformsData = LocalAtoms.GetData(); 
+	FTransform* SpaceBasesData = SpaceBases.GetData();
 	
 	SpaceBases[0] = LocalAtoms[0];
 #if (UE_BUILD_DEBUG || UE_BUILD_DEVELOPMENT)
diff --git a/Engine/Source/Runtime/Engine/Private/RawIndexBuffer.cpp b/Engine/Source/Runtime/Engine/Private/RawIndexBuffer.cpp
--- a/Engine/Source/Runtime/Engine/Private/RawIndexBuffer.cpp
+++ b/Engine/Source/Runtime/Engine/Private/RawIndexBuffer.cpp
@@ -44,7 +44,7 @@ void FRawIndexBuffer::InitRHI()
 
 		// Initialize the buffer.
 		void* Buffer = RHILockIndexBuffer(IndexBufferRHI, 0, Size, RLM_WriteOnly);
-		FMemory::Memcpy(Buffer,Indices.GetTypedData(),Size);
+		FMemory::Memcpy(Buffer,Indices.GetData(),Size);
 		RHIUnlockIndexBuffer(IndexBufferRHI);
 	}
 }
@@ -83,7 +83,7 @@ void FRawIndexBuffer16or32::InitRHI()
 
 		// Initialize the buffer.
 		void* Buffer = RHILockIndexBuffer(IndexBufferRHI, 0, Size, RLM_WriteOnly);
-		FMemory::Memcpy(Buffer,Indices.GetTypedData(),Size);
+		FMemory::Memcpy(Buffer,Indices.GetData(),Size);
 		RHIUnlockIndexBuffer(IndexBufferRHI);
 	}
 
diff --git a/Engine/Source/Runtime/Engine/Private/RepLayout.cpp b/Engine/Source/Runtime/Engine/Private/RepLayout.cpp
--- a/Engine/Source/Runtime/Engine/Private/RepLayout.cpp
+++ b/Engine/Source/Runtime/Engine/Private/RepLayout.cpp
@@ -158,23 +158,23 @@ static void SerializeReadWritePropertyChecksum( const FRepLayoutCmd & Cmd, const
 		uint32 * Guard = (uint32*)&TempPropMemory[TempPropMemory.Num() - 4];
 		const uint32 TAG_VALUE = 0xABADF00D;
 		*Guard = TAG_VALUE;
-		Cmd.Property->InitializeValue( TempPropMemory.GetTypedData() );
+		Cmd.Property->InitializeValue( TempPropMemory.GetData() );
 		check( *Guard == TAG_VALUE );
 
 		// Read it back in and then write it out to produce what the client will produce
 		FBitReader Reader( Writer.GetData(), Writer.GetNumBits() );
-		Cmd.Property->NetSerializeItem( Reader, NULL, TempPropMemory.GetTypedData() );
+		Cmd.Property->NetSerializeItem( Reader, NULL, TempPropMemory.GetData() );
 		check( Reader.AtEnd() && !Reader.IsError() );
 		check( *Guard == TAG_VALUE );
 
 		// Write it back out for a final time
 		Writer.Reset();
 
-		Cmd.Property->NetSerializeItem( Writer, NULL, TempPropMemory.GetTypedData() );
+		Cmd.Property->NetSerializeItem( Writer, NULL, TempPropMemory.GetData() );
 		check( *Guard == TAG_VALUE );
 
 		// Destroy temp memory
-		Cmd.Property->DestroyValue( TempPropMemory.GetTypedData() );
+		Cmd.Property->DestroyValue( TempPropMemory.GetData() );
 
 		// Restore the static array size
 		Cmd.Property->ArrayDim = OriginalDim;
@@ -288,7 +288,7 @@ bool FRepLayout::CompareProperties(
 {
 	bool PropertyChanged = false;
 
-	const uint16 * RESTRICT FirstProp	= PropertyList.GetTypedData();
+	const uint16 * RESTRICT FirstProp	= PropertyList.GetData();
 	const uint16 * RESTRICT LastProp	= FirstProp + PropertyList.Num();
 
 	for ( const uint16 * RESTRICT pLifeProp = FirstProp; pLifeProp < LastProp; ++pLifeProp )
@@ -368,8 +368,8 @@ void FRepLayout::LogChangeListMismatches(
 			FString StringValue3;
 
 			Parents[Index].Property->ExportText_InContainer( Parents[Index].ArrayIndex, StringValue1, Data, NULL, NULL, PPF_DebugDump );
-			Parents[Index].Property->ExportText_InContainer( Parents[Index].ArrayIndex, StringValue2, RepState1->StaticBuffer.GetTypedData(), NULL, NULL, PPF_DebugDump );
-			Parents[Index].Property->ExportText_InContainer( Parents[Index].ArrayIndex, StringValue3, RepState2->StaticBuffer.GetTypedData(), NULL, NULL, PPF_DebugDump );
+			Parents[Index].Property->ExportText_InContainer( Parents[Index].ArrayIndex, StringValue2, RepState1->StaticBuffer.GetData(), NULL, NULL, PPF_DebugDump );
+			Parents[Index].Property->ExportText_InContainer( Parents[Index].ArrayIndex, StringValue3, RepState2->StaticBuffer.GetData(), NULL, NULL, PPF_DebugDump );
 
 			UE_LOG( LogNet, Warning, TEXT( "      Values: %s, %s, %s" ), *StringValue1, *StringValue2, *StringValue3 );
 		}
@@ -397,7 +397,7 @@ void FRepLayout::SanityCheckShadowStateAgainstChangeList(
 	FRepState *							OtherRepState,
 	const TArray< FRepChangedParent > & OtherChangedParents ) const
 {
-	const uint8 *	CompareData	= RepState->StaticBuffer.GetTypedData();
+	const uint8 *	CompareData	= RepState->StaticBuffer.GetData();
 	UObject *		Object		= (UObject*)Data;
 
 	TArray< FRepChangedParent > LocalChangedParents;
@@ -450,7 +450,7 @@ bool FRepLayout::ReplicateProperties(
 	UObject *						Object			= (UObject*)Data;
 	const UNetDriver *				NetDriver		= OwningChannel->Connection->Driver;
 	FRepChangedPropertyTracker *	ChangeTracker	= RepState->RepChangedPropertyTracker.Get();
-	const uint8 *					CompareData		= RepState->StaticBuffer.GetTypedData();
+	const uint8 *					CompareData		= RepState->StaticBuffer.GetData();
 
 	// Rebuild conditional properties if needed
 	if ( RepState->RepFlags.Value != RepFlags.Value || RepState->ActiveStatusChanged != ChangeTracker->ActiveStatusChanged )
@@ -608,7 +608,7 @@ bool FRepLayout::ReplicateProperties(
 
 		if ( bIsAllAcked )
 		{
-			ValidateWithChecksum( RepState->StaticBuffer.GetTypedData(), Writer, false );
+			ValidateWithChecksum( RepState->StaticBuffer.GetData(), Writer, false );
 		}
 #endif
 		
@@ -990,7 +990,7 @@ void FRepLayout::SendProperties(
 	Writer.WriteBit( bDoChecksum ? 1 : 0 );
 #endif
 
-	SendProperties_r( RepState, RepFlags, WriterState, 0, Cmds.Num() - 1, RepState->StaticBuffer.GetTypedData(), Data, 0 );
+	SendProperties_r( RepState, RepFlags, WriterState, 0, Cmds.Num() - 1, RepState->StaticBuffer.GetData(), Data, 0 );
 
 	WritePropertyHandle( Writer, 0, bDoChecksum );
 }
@@ -1244,7 +1244,7 @@ bool FRepLayout::ReceiveProperties( UClass * InObjectClass, FRepState * RESTRICT
 	ReadNextHandle( ReaderState );
 
 	// Read all properties
-	if ( !ReceiveProperties_r( ReaderState, &RepState->UnmappedGuids, 0, 0, Cmds.Num() - 1, RepState->StaticBuffer.GetTypedData(), (uint8*)Data ) )
+	if ( !ReceiveProperties_r( ReaderState, &RepState->UnmappedGuids, 0, 0, Cmds.Num() - 1, RepState->StaticBuffer.GetData(), (uint8*)Data ) )
 	{
 		return false;
 	}
@@ -1259,7 +1259,7 @@ bool FRepLayout::ReceiveProperties( UClass * InObjectClass, FRepState * RESTRICT
 #ifdef ENABLE_SUPER_CHECKSUMS
 	if ( InBunch.ReadBit() == 1 )
 	{
-		ValidateWithChecksum( RepState->StaticBuffer.GetTypedData(), InBunch );
+		ValidateWithChecksum( RepState->StaticBuffer.GetData(), InBunch );
 	}
 #endif
 
@@ -1387,7 +1387,7 @@ void FRepLayout::CallRepNotifies( FRepState * RepState, UObject * Object ) const
 		}
 		else if (RepNotifyFunc->NumParms == 1 )
 		{
-			Object->ProcessEvent( RepNotifyFunc, RepProperty->ContainerPtrToValuePtr<uint8>( RepState->StaticBuffer.GetTypedData() ) );
+			Object->ProcessEvent( RepNotifyFunc, RepProperty->ContainerPtrToValuePtr<uint8>( RepState->StaticBuffer.GetData() ) );
 		}
 
 		// Store the property we just received
@@ -1456,7 +1456,7 @@ void FRepLayout::ValidateWithChecksum( const void * RESTRICT Data, FArchive & Ar
 uint32 FRepLayout::GenerateChecksum( const FRepState * RepState ) const
 {
 	FBitWriter Writer( 1024, true );
-	ValidateWithChecksum_r( 0, Cmds.Num() - 1, (const uint8*)RepState->StaticBuffer.GetTypedData(), Writer );
+	ValidateWithChecksum_r( 0, Cmds.Num() - 1, (const uint8*)RepState->StaticBuffer.GetData(), Writer );
 
 	return FCrc::MemCrc32( Writer.GetData(), Writer.GetNumBytes(), 0 );
 }
@@ -1785,7 +1785,7 @@ void FRepLayout::DiffProperties_r(
 bool FRepLayout::DiffProperties( FRepState * RepState, const void * RESTRICT Data, const bool bSync ) const
 {
 	bool bDifferent = false;
-	DiffProperties_r( RepState, 0, Cmds.Num() - 1, RepState->StaticBuffer.GetTypedData(), (uint8*)Data, bSync, bDifferent );
+	DiffProperties_r( RepState, 0, Cmds.Num() - 1, RepState->StaticBuffer.GetData(), (uint8*)Data, bSync, bDifferent );
 	return bDifferent;
 }
 
@@ -1974,7 +1974,7 @@ int32 FRepLayout::InitFromProperty_r( UProperty * Property, int32 Offset, int32
 			}
 		};
 
-		Sort( NetProperties.GetTypedData(), NetProperties.Num(), FCompareUFieldOffsets() );
+		Sort( NetProperties.GetData(), NetProperties.Num(), FCompareUFieldOffsets() );
 
 		for ( int32 i = 0; i < NetProperties.Num(); i++ )
 		{
@@ -2460,7 +2460,7 @@ void FRepLayout::InitRepState(
 
 void FRepLayout::ConstructProperties( FRepState * RepState ) const
 {
-	uint8 * StoredData = RepState->StaticBuffer.GetTypedData();
+	uint8 * StoredData = RepState->StaticBuffer.GetData();
 
 	// Construct all items
 	for ( int32 i = 0; i < Parents.Num(); i++ )
@@ -2478,7 +2478,7 @@ void FRepLayout::ConstructProperties( FRepState * RepState ) const
 
 void FRepLayout::InitProperties( FRepState * RepState, uint8 * Src ) const
 {
-	uint8 * StoredData = RepState->StaticBuffer.GetTypedData();
+	uint8 * StoredData = RepState->StaticBuffer.GetData();
 
 	// Init all items
 	for ( int32 i = 0; i < Parents.Num(); i++ )
@@ -2496,7 +2496,7 @@ void FRepLayout::InitProperties( FRepState * RepState, uint8 * Src ) const
 
 void FRepLayout::DestructProperties( FRepState * RepState ) const
 {
-	uint8 * StoredData = RepState->StaticBuffer.GetTypedData();
+	uint8 * StoredData = RepState->StaticBuffer.GetData();
 
 	// Destruct all items
 	for ( int32 i = 0; i < Parents.Num(); i++ )
diff --git a/Engine/Source/Runtime/Engine/Private/ShadowMap.cpp b/Engine/Source/Runtime/Engine/Private/ShadowMap.cpp
--- a/Engine/Source/Runtime/Engine/Private/ShadowMap.cpp
+++ b/Engine/Source/Runtime/Engine/Private/ShadowMap.cpp
@@ -414,7 +414,7 @@ void FShadowMap2D::EncodeTextures(UWorld* InWorld , bool bLightingSuccessful )
 		// Reset the pending shadow-map size.
 		PendingShadowMapSize = 0;
 
-		Sort(PendingShadowMaps.GetTypedData(), PendingShadowMaps.Num(), FCompareShadowMaps());
+		Sort(PendingShadowMaps.GetData(), PendingShadowMaps.Num(), FCompareShadowMaps());
 
 		// Allocate texture space for each light-map.
 		TIndirectArray<FShadowMapPendingTexture> PendingTextures;
diff --git a/Engine/Source/Runtime/Engine/Private/SkeletalMesh.cpp b/Engine/Source/Runtime/Engine/Private/SkeletalMesh.cpp
--- a/Engine/Source/Runtime/Engine/Private/SkeletalMesh.cpp
+++ b/Engine/Source/Runtime/Engine/Private/SkeletalMesh.cpp
@@ -98,7 +98,7 @@ static bool SaveApexClothingAssetToBlob(const NxClothingAsset *InAsset, TArray<u
 			// Read the stream data into our buffer for UE serialzation
 			Size = Stream->getFileLength();
 			OutBuffer.AddUninitialized( Size );
-			Stream->read( OutBuffer.GetTypedData(), Size );
+			Stream->read( OutBuffer.GetData(), Size );
 			bResult = true;
 		}
 
@@ -1389,7 +1389,7 @@ void FStaticLODModel::GetVertices(TArray<FSoftSkinVertex>& Vertices) const
 			}
 			DestVertex++;
 		}
-		FMemory::Memcpy(DestVertex,Chunk.SoftVertices.GetTypedData(),Chunk.SoftVertices.Num() * sizeof(FSoftSkinVertex));
+		FMemory::Memcpy(DestVertex,Chunk.SoftVertices.GetData(),Chunk.SoftVertices.Num() * sizeof(FSoftSkinVertex));
 		DestVertex += Chunk.SoftVertices.Num();
 	}
 }
@@ -1519,7 +1519,7 @@ void FStaticLODModel::SortTriangles( FVector SortCenter, bool bUseSortCenter, in
 		{
 			TArray<uint32> Indices;
 			MultiSizeIndexContainer.GetIndexBuffer( Indices );
-			SortTriangles_None( Section.NumTriangles, Vertices.GetTypedData(), Indices.GetData() + Section.BaseIndex );
+			SortTriangles_None( Section.NumTriangles, Vertices.GetData(), Indices.GetData() + Section.BaseIndex );
 			MultiSizeIndexContainer.CopyIndexBuffer( Indices );
 		}
 		break;
@@ -1529,11 +1529,11 @@ void FStaticLODModel::SortTriangles( FVector SortCenter, bool bUseSortCenter, in
 			MultiSizeIndexContainer.GetIndexBuffer( Indices );
 			if (bUseSortCenter)
 			{
-				SortTriangles_CenterRadialDistance( SortCenter, Section.NumTriangles, Vertices.GetTypedData(), Indices.GetData() + Section.BaseIndex );
+				SortTriangles_CenterRadialDistance( SortCenter, Section.NumTriangles, Vertices.GetData(), Indices.GetData() + Section.BaseIndex );
 			}
 			else
 			{
-				SortTriangles_CenterRadialDistance( Section.NumTriangles, Vertices.GetTypedData(), Indices.GetData() + Section.BaseIndex );
+				SortTriangles_CenterRadialDistance( Section.NumTriangles, Vertices.GetData(), Indices.GetData() + Section.BaseIndex );
 			}
 			MultiSizeIndexContainer.CopyIndexBuffer( Indices );
 		}
@@ -1542,7 +1542,7 @@ void FStaticLODModel::SortTriangles( FVector SortCenter, bool bUseSortCenter, in
 		{
 			TArray<uint32> Indices;
 			MultiSizeIndexContainer.GetIndexBuffer( Indices );
-			SortTriangles_Random( Section.NumTriangles, Vertices.GetTypedData(), Indices.GetData() + Section.BaseIndex );
+			SortTriangles_Random( Section.NumTriangles, Vertices.GetData(), Indices.GetData() + Section.BaseIndex );
 			MultiSizeIndexContainer.CopyIndexBuffer( Indices );
 		}
 		break;
@@ -1550,7 +1550,7 @@ void FStaticLODModel::SortTriangles( FVector SortCenter, bool bUseSortCenter, in
 		{
 			TArray<uint32> Indices;
 			MultiSizeIndexContainer.GetIndexBuffer( Indices );
-			SortTriangles_MergeContiguous( Section.NumTriangles, NumVertices, Vertices.GetTypedData(), Indices.GetData() + Section.BaseIndex );
+			SortTriangles_MergeContiguous( Section.NumTriangles, NumVertices, Vertices.GetData(), Indices.GetData() + Section.BaseIndex );
 			MultiSizeIndexContainer.CopyIndexBuffer( Indices );
 		}
 		break;
diff --git a/Engine/Source/Runtime/Engine/Private/SkeletalMeshComponent.cpp b/Engine/Source/Runtime/Engine/Private/SkeletalMeshComponent.cpp
--- a/Engine/Source/Runtime/Engine/Private/SkeletalMeshComponent.cpp
+++ b/Engine/Source/Runtime/Engine/Private/SkeletalMeshComponent.cpp
@@ -609,8 +609,8 @@ void USkeletalMeshComponent::FillSpaceBases(const USkeletalMesh* InSkeletalMesh,
 	BoneProcessed.AddZeroed(NumBones);
 #endif
 
-	const FTransform * LocalTransformsData = SourceAtoms.GetTypedData();
-	FTransform * SpaceBasesData = DestSpaceBases.GetTypedData();
+	const FTransform * LocalTransformsData = SourceAtoms.GetData();
+	FTransform * SpaceBasesData = DestSpaceBases.GetData();
 
 	// First bone is always root bone, and it doesn't have a parent.
 	{
diff --git a/Engine/Source/Runtime/Engine/Private/SkeletalMeshComponentPhysics.cpp b/Engine/Source/Runtime/Engine/Private/SkeletalMeshComponentPhysics.cpp
--- a/Engine/Source/Runtime/Engine/Private/SkeletalMeshComponentPhysics.cpp
+++ b/Engine/Source/Runtime/Engine/Private/SkeletalMeshComponentPhysics.cpp
@@ -2116,7 +2116,7 @@ bool USkeletalMeshComponent::CreateClothingActor(int32 AssetIndex, TSharedPtr<FC
 		}
 		NxParameterized::Handle md(*ActorDesc, "morphDisplacements");
 		md.resizeArray(PxBlendedData.Num());
-		md.setParamVec3Array(PxBlendedData.GetTypedData(), PxBlendedData.Num());
+		md.setParamVec3Array(PxBlendedData.GetData(), PxBlendedData.Num());
 	}
 
 	FPhysScene* PhysScene = NULL;
@@ -2738,7 +2738,7 @@ void USkeletalMeshComponent::CreateInternalClothCollisions(TArray<FApexClothColl
 						ClothingPlanes[PlaneIdx] = Actor->createCollisionPlane(PPlane);
 					}
 
-					NxClothingConvex* Convex = Actor->createCollisionConvex(ClothingPlanes.GetTypedData(), ClothingPlanes.Num());
+					NxClothingConvex* Convex = Actor->createCollisionConvex(ClothingPlanes.GetData(), ClothingPlanes.Num());
 
 					OutCollisions.Add(Convex);
 				}
@@ -2925,7 +2925,7 @@ FApexClothCollisionInfo* USkeletalMeshComponent::CreateNewClothingCollsions(UPri
 						ClothingPlanes[PlaneIdx] = Actor->createCollisionPlane(PPlane);
 					}
 
-					ClothCol = Actor->createCollisionConvex(ClothingPlanes.GetTypedData(), ClothingPlanes.Num());
+					ClothCol = Actor->createCollisionConvex(ClothingPlanes.GetData(), ClothingPlanes.Num());
 
 					if(ClothCol)
 					{
@@ -3243,7 +3243,7 @@ void USkeletalMeshComponent::ChangeClothMorphTargetMapping(FClothMorphTargetData
 
 			NxClothingAsset* ClothingAsset = Asset.ApexClothingAsset->GetAsset();
 			float Epsilon = 0.0f;
-			uint32 NumMapped = ClothingAsset->prepareMorphTargetMapping(ClothOriginalPosArray.GetTypedData(), NumOriginPos, Epsilon);
+			uint32 NumMapped = ClothingAsset->prepareMorphTargetMapping(ClothOriginalPosArray.GetData(), NumOriginPos, Epsilon);
 
 			if ((int32)NumMapped < NumOriginPos)
 			{
@@ -3352,7 +3352,7 @@ void USkeletalMeshComponent::PrepareClothMorphTargets()
 				{
 					NxClothingAsset* ClothingAsset = SkeletalMesh->ClothingAssets[AssetIdx].ApexClothingAsset->GetAsset();
 					float Epsilon = 0.0f;
-					uint32 NumMapped = ClothingAsset->prepareMorphTargetMapping(ClothOriginalPosArray[AssetIdx].GetTypedData(), ClothOriginalPosArray[AssetIdx].Num(), Epsilon);
+					uint32 NumMapped = ClothingAsset->prepareMorphTargetMapping(ClothOriginalPosArray[AssetIdx].GetData(), ClothOriginalPosArray[AssetIdx].Num(), Epsilon);
 
 					int32 NumOriginPos = ClothOriginalPosArray[AssetIdx].Num();
 					if ((int32)NumMapped < NumOriginPos)
@@ -3570,7 +3570,7 @@ void USkeletalMeshComponent::UpdateClothState(float DeltaTime)
 		// if bUseInternalboneOrder is set, "NumUsedBones" works, otherwise have to use "getNumBones" 
 		ClothingActor->updateState(
 			PxGlobalPose, 
-			BoneMatrices.GetTypedData(), 
+			BoneMatrices.GetData(), 
 			sizeof(physx::PxMat44), 
 			NumUsedBones,
 			CurTeleportMode);
diff --git a/Engine/Source/Runtime/Engine/Private/SkeletalMeshSorting.cpp b/Engine/Source/Runtime/Engine/Private/SkeletalMeshSorting.cpp
--- a/Engine/Source/Runtime/Engine/Private/SkeletalMeshSorting.cpp
+++ b/Engine/Source/Runtime/Engine/Private/SkeletalMeshSorting.cpp
@@ -314,7 +314,7 @@ void SortTriangles_CenterRadialDistance( FVector SortCenter, int32 NumTriangles,
 		}
 		CacheOptimizeSortStrip( &NewIndices[StripStartIndex], Strips[s].Triangles.Num()*3 );
 	}
-	FMemory::Memcpy( Indices, NewIndices.GetTypedData(), NewIndices.Num() * sizeof(uint32) );
+	FMemory::Memcpy( Indices, NewIndices.GetData(), NewIndices.Num() * sizeof(uint32) );
 }
 
 
@@ -337,7 +337,7 @@ void SortTriangles_Random( int32 NumTriangles, const FSoftSkinVertex* Vertices,
 		NewIndices.Add(Indices[tri*3+2]);	
 	}
 
-	FMemory::Memcpy( Indices, NewIndices.GetTypedData(), NewIndices.Num() * sizeof(uint32) );
+	FMemory::Memcpy( Indices, NewIndices.GetData(), NewIndices.Num() * sizeof(uint32) );
 }
 
 
diff --git a/Engine/Source/Runtime/Engine/Private/SkeletalRenderCPUSkin.cpp b/Engine/Source/Runtime/Engine/Private/SkeletalRenderCPUSkin.cpp
--- a/Engine/Source/Runtime/Engine/Private/SkeletalRenderCPUSkin.cpp
+++ b/Engine/Source/Runtime/Engine/Private/SkeletalRenderCPUSkin.cpp
@@ -252,14 +252,14 @@ void FSkeletalMeshObjectCPUSkin::CacheVertices(int32 LODIndex, bool bForce) cons
 		const FSkelMeshObjectLODInfo& MeshLODInfo = LODInfo[LODIndex];
 
 		// bone matrices
-		FMatrix* ReferenceToLocal = DynamicData->ReferenceToLocal.GetTypedData();
+		FMatrix* ReferenceToLocal = DynamicData->ReferenceToLocal.GetData();
 
 		int32 CachedFinalVerticesNum = LOD.NumVertices;
 		CachedFinalVertices.Empty(CachedFinalVerticesNum);
 		CachedFinalVertices.AddUninitialized(CachedFinalVerticesNum);
 
 		// final cached verts
-		FFinalSkinVertex* DestVertex = CachedFinalVertices.GetTypedData();
+		FFinalSkinVertex* DestVertex = CachedFinalVertices.GetData();
 
 		if (DestVertex)
 		{
@@ -301,7 +301,7 @@ void FSkeletalMeshObjectCPUSkin::CacheVertices(int32 LODIndex, bool bForce) cons
 		CachedVertexLOD = LODIndex;
 
 		check(LOD.NumVertices == CachedFinalVertices.Num());
-		MeshLOD.UpdateFinalSkinVertexBuffer( CachedFinalVertices.GetTypedData(), LOD.NumVertices * sizeof(FFinalSkinVertex) );
+		MeshLOD.UpdateFinalSkinVertexBuffer( CachedFinalVertices.GetData(), LOD.NumVertices * sizeof(FFinalSkinVertex) );
 	}
 }
 
@@ -622,7 +622,7 @@ static void SkinVertexChunk( FFinalSkinVertex*& DestVertex, TArray<FVertexAnimEv
 	VertexType  VertexCopy;
 
 	// Prefetch all bone indices
-	const FBoneIndexType* BoneMap = Chunk.BoneMap.GetTypedData();
+	const FBoneIndexType* BoneMap = Chunk.BoneMap.GetData();
 	FPlatformMisc::Prefetch( BoneMap );
 	FPlatformMisc::Prefetch( BoneMap, CACHE_LINE_SIZE );
 
@@ -926,7 +926,7 @@ static FORCEINLINE void CalculateChunkBoneWeights(FFinalSkinVertex*& DestVertex,
 	int32 VertexBufferBaseIndex = 0;
 
 	//array of bone mapping
-	FBoneIndexType* BoneMap = Chunk.BoneMap.GetTypedData();
+	FBoneIndexType* BoneMap = Chunk.BoneMap.GetData();
 
 	for(int32 VertexIndex = VertexBufferBaseIndex;VertexIndex < Chunk.GetNumRigidVertices();VertexIndex++,DestVertex++)
 	{
diff --git a/Engine/Source/Runtime/Engine/Private/SkeletalRenderGPUSkin.cpp b/Engine/Source/Runtime/Engine/Private/SkeletalRenderGPUSkin.cpp
--- a/Engine/Source/Runtime/Engine/Private/SkeletalRenderGPUSkin.cpp
+++ b/Engine/Source/Runtime/Engine/Private/SkeletalRenderGPUSkin.cpp
@@ -270,10 +270,10 @@ void FSkeletalMeshObjectGPUSkin::UpdateDynamicData_RenderThread(FRHICommandListI
 			const int32 NumReferenceToLocal = ReferenceToLocalMatrices.Num();
 			for( int32 BoneIdx=0; BoneIdx < NumBones; BoneIdx++ )
 			{
-				FPlatformMisc::Prefetch( ChunkMatrices.GetTypedData() + BoneIdx + PreFetchStride ); 
-				FPlatformMisc::Prefetch( ChunkMatrices.GetTypedData() + BoneIdx + PreFetchStride, CACHE_LINE_SIZE ); 
-				FPlatformMisc::Prefetch( ReferenceToLocalMatrices.GetTypedData() + BoneIdx + PreFetchStride );
-				FPlatformMisc::Prefetch( ReferenceToLocalMatrices.GetTypedData() + BoneIdx + PreFetchStride, CACHE_LINE_SIZE );
+				FPlatformMisc::Prefetch( ChunkMatrices.GetData() + BoneIdx + PreFetchStride ); 
+				FPlatformMisc::Prefetch( ChunkMatrices.GetData() + BoneIdx + PreFetchStride, CACHE_LINE_SIZE ); 
+				FPlatformMisc::Prefetch( ReferenceToLocalMatrices.GetData() + BoneIdx + PreFetchStride );
+				FPlatformMisc::Prefetch( ReferenceToLocalMatrices.GetData() + BoneIdx + PreFetchStride, CACHE_LINE_SIZE );
 	
 				FBoneSkinning& BoneMat = ChunkMatrices[BoneIdx];
 				const FBoneIndexType RefToLocalIdx = Chunk.BoneMap[BoneIdx];
diff --git a/Engine/Source/Runtime/Engine/Private/SphereComponent.cpp b/Engine/Source/Runtime/Engine/Private/SphereComponent.cpp
--- a/Engine/Source/Runtime/Engine/Private/SphereComponent.cpp
+++ b/Engine/Source/Runtime/Engine/Private/SphereComponent.cpp
@@ -33,7 +33,7 @@ void USphereComponent::UpdateBodySetup()
 	}
 
 	check (ShapeBodySetup->AggGeom.SphereElems.Num() == 1);
-	FKSphereElem* se = ShapeBodySetup->AggGeom.SphereElems.GetTypedData();
+	FKSphereElem* se = ShapeBodySetup->AggGeom.SphereElems.GetData();
 
 	// check for mal formed values
 	float Radius = SphereRadius;
diff --git a/Engine/Source/Runtime/Engine/Private/StaticMesh.cpp b/Engine/Source/Runtime/Engine/Private/StaticMesh.cpp
--- a/Engine/Source/Runtime/Engine/Private/StaticMesh.cpp
+++ b/Engine/Source/Runtime/Engine/Private/StaticMesh.cpp
@@ -1084,7 +1084,7 @@ FString BuildStaticMeshDerivedDataKey(UStaticMesh* Mesh, const FStaticMeshLODGro
 		Ar << FinalReductionSettings;
 
 		// Now convert the raw bytes to a string.
-		const uint8* SettingsAsBytes = TempBytes.GetTypedData();
+		const uint8* SettingsAsBytes = TempBytes.GetData();
 		KeySuffix.Reserve(KeySuffix.Len() + TempBytes.Num() + 1);
 		for (int32 ByteIndex = 0; ByteIndex < TempBytes.Num(); ++ByteIndex)
 		{
@@ -1294,7 +1294,7 @@ bool UStaticMesh::HasValidRenderData() const
 {
 	return RenderData != NULL
 		&& RenderData->LODResources.Num() > 0
-		&& RenderData->LODResources.GetTypedData() != NULL
+		&& RenderData->LODResources.GetData() != NULL
 		&& RenderData->LODResources[0].VertexBuffer.GetNumVertices() > 0;
 }
 
diff --git a/Engine/Source/Runtime/Engine/Private/TextRenderComponent.cpp b/Engine/Source/Runtime/Engine/Private/TextRenderComponent.cpp
--- a/Engine/Source/Runtime/Engine/Private/TextRenderComponent.cpp
+++ b/Engine/Source/Runtime/Engine/Private/TextRenderComponent.cpp
@@ -111,7 +111,7 @@ public:
 
 		// Copy the vertex data into the vertex buffer.
 		void* VertexBufferData = RHILockVertexBuffer(VertexBufferRHI,0,Vertices.Num() * sizeof(FDynamicMeshVertex), RLM_WriteOnly);
-		FMemory::Memcpy(VertexBufferData,Vertices.GetTypedData(),Vertices.Num() * sizeof(FDynamicMeshVertex));
+		FMemory::Memcpy(VertexBufferData,Vertices.GetData(),Vertices.Num() * sizeof(FDynamicMeshVertex));
 		RHIUnlockVertexBuffer(VertexBufferRHI);
 	}
 };
@@ -129,7 +129,7 @@ public:
 
 		// Copy the index data into the index buffer.
 		void* Buffer = RHILockIndexBuffer(IndexBufferRHI, 0, Indices.Num() * sizeof(uint16), RLM_WriteOnly);
-		FMemory::Memcpy(Buffer, Indices.GetTypedData(), Indices.Num() * sizeof(uint16));
+		FMemory::Memcpy(Buffer, Indices.GetData(), Indices.Num() * sizeof(uint16));
 		RHIUnlockIndexBuffer(IndexBufferRHI);
 	}
 };
diff --git a/Engine/Source/Runtime/Engine/Private/Texture.cpp b/Engine/Source/Runtime/Engine/Private/Texture.cpp
--- a/Engine/Source/Runtime/Engine/Private/Texture.cpp
+++ b/Engine/Source/Runtime/Engine/Private/Texture.cpp
@@ -603,7 +603,7 @@ void FTextureSource::Compress()
 			if ( CompressedData.Num() > 0 )
 			{
 				BulkDataPtr = (uint8*)BulkData.Realloc(CompressedData.Num());
-				FMemory::Memcpy(BulkDataPtr, CompressedData.GetTypedData(), CompressedData.Num());
+				FMemory::Memcpy(BulkDataPtr, CompressedData.GetData(), CompressedData.Num());
 				BulkData.Unlock();
 				bPNGCompressed = true;
 
@@ -649,7 +649,7 @@ uint8* FTextureSource::LockMip(int32 MipIndex)
 						if (RawData->Num() > 0)
 						{
 							LockedMipData = (uint8*)FMemory::Malloc(RawData->Num());
-							FMemory::Memcpy(LockedMipData, RawData->GetTypedData(), RawData->Num());
+							FMemory::Memcpy(LockedMipData, RawData->GetData(), RawData->Num());
 						}
 					}
 					if (RawData == NULL || RawData->Num() == 0)
@@ -753,7 +753,7 @@ bool FTextureSource::GetMipData(TArray<uint8>& OutMipData, int32 MipIndex)
 				OutMipData.Empty(MipSize);
 				OutMipData.AddUninitialized(MipSize);
 				FMemory::Memcpy(
-					OutMipData.GetTypedData(),
+					OutMipData.GetData(),
 					(uint8*)RawSourceData + MipOffset,
 					MipSize
 					);
@@ -875,12 +875,12 @@ void FTextureSource::ConvertFromLegacyDDS()
 	{
 		void* LockedDDSData = NULL;
 		BulkData.GetCopy(&LockedDDSData);
-		FMemory::Memcpy(RawDDS.GetTypedData(), LockedDDSData, RawDDS.Num());
+		FMemory::Memcpy(RawDDS.GetData(), LockedDDSData, RawDDS.Num());
 		FMemory::Free(LockedDDSData);
 		BulkData.RemoveBulkData();
 	}
 
-	FDDSLoadHelper DDS(RawDDS.GetTypedData(), RawDDS.Num());
+	FDDSLoadHelper DDS(RawDDS.GetData(), RawDDS.Num());
 	if ((DDS.IsValid2DTexture() || DDS.IsValidCubemapTexture()) && DDS.ComputeSourceFormat() != TSF_Invalid)
 	{
 		Init(
diff --git a/Engine/Source/Runtime/Engine/Private/Texture2D.cpp b/Engine/Source/Runtime/Engine/Private/Texture2D.cpp
--- a/Engine/Source/Runtime/Engine/Private/Texture2D.cpp
+++ b/Engine/Source/Runtime/Engine/Private/Texture2D.cpp
@@ -383,7 +383,7 @@ float UTexture2D::GetAverageBrightness(bool bIgnoreTrueBlack, bool bUseGrayscale
 		int32 SizeY = Source.GetSizeY();
 		double PixelSum = 0.0f;
 		int32 Divisor = SizeX * SizeY;
-		FColor* ColorData = (FColor*)RawData.GetTypedData();
+		FColor* ColorData = (FColor*)RawData.GetData();
 		FLinearColor CurrentColor;
 		for (int32 Y = 0; Y < SizeY; Y++)
 		{
diff --git a/Engine/Source/Runtime/Engine/Private/TextureDerivedData.cpp b/Engine/Source/Runtime/Engine/Private/TextureDerivedData.cpp
--- a/Engine/Source/Runtime/Engine/Private/TextureDerivedData.cpp
+++ b/Engine/Source/Runtime/Engine/Private/TextureDerivedData.cpp
@@ -203,7 +203,7 @@ static void GetTextureDerivedDataKeySuffix(
 	SerializeForKey(Ar, BuildSettings);
 
 	// Now convert the raw bytes to a string.
-	const uint8* SettingsAsBytes = TempBytes.GetTypedData();
+	const uint8* SettingsAsBytes = TempBytes.GetData();
 	OutKeySuffix.Reserve(OutKeySuffix.Len() + TempBytes.Num());
 	for (int32 ByteIndex = 0; ByteIndex < TempBytes.Num(); ++ByteIndex)
 	{
@@ -607,7 +607,7 @@ static float ComputePSNR(const FImage& SrcImage, const FCompressedImage2D& Compr
 {
 	double SquaredError = 0.0;
 	int32 NumErrors = 0;
-	const uint8* CompressedData = CompressedImage.RawData.GetTypedData();
+	const uint8* CompressedData = CompressedImage.RawData.GetData();
 
 	if (SrcImage.Format == ERawImageFormat::BGRA8 && (CompressedImage.PixelFormat == PF_DXT1 || CompressedImage.PixelFormat == PF_DXT5))
 	{
@@ -769,7 +769,7 @@ class FTextureCacheDerivedDataWorker : public FNonAbandonableTask
 					NewMip->BulkData.Lock(LOCK_READ_WRITE);
 					check(CompressedImage.RawData.GetTypeSize() == 1);
 					void* NewMipData = NewMip->BulkData.Realloc(CompressedImage.RawData.Num());
-					FMemory::Memcpy(NewMipData, CompressedImage.RawData.GetTypedData(), CompressedImage.RawData.Num());
+					FMemory::Memcpy(NewMipData, CompressedImage.RawData.GetData(), CompressedImage.RawData.Num());
 					NewMip->BulkData.Unlock();
 
 					if (MipIndex == 0)
diff --git a/Engine/Source/Runtime/Engine/Private/TextureRenderTarget2D.cpp b/Engine/Source/Runtime/Engine/Private/TextureRenderTarget2D.cpp
--- a/Engine/Source/Runtime/Engine/Private/TextureRenderTarget2D.cpp
+++ b/Engine/Source/Runtime/Engine/Private/TextureRenderTarget2D.cpp
@@ -197,7 +197,7 @@ UTexture2D* UTextureRenderTarget2D::ConstructTexture2D(UObject* Outer, const FSt
 		}
 		// copy the 2d surface data to the first mip of the static 2d texture
 		check(TextureDataSize == SurfData.Num()*sizeof(FColor));
-		FMemory::Memcpy(TextureData, SurfData.GetTypedData(), TextureDataSize);
+		FMemory::Memcpy(TextureData, SurfData.GetData(), TextureDataSize);
 	}
 	else if (TextureFormat == TSF_RGBA16F)
 	{
@@ -230,7 +230,7 @@ UTexture2D* UTextureRenderTarget2D::ConstructTexture2D(UObject* Outer, const FSt
 		}
 		// copy the 2d surface data to the first mip of the static 2d texture
 		check(TextureDataSize == SurfData.Num()*sizeof(FFloat16Color));
-		FMemory::Memcpy(TextureData, SurfData.GetTypedData(), TextureDataSize);
+		FMemory::Memcpy(TextureData, SurfData.GetData(), TextureDataSize);
 	}
 	Result->Source.UnlockMip(0);
 	// if render target gamma used was 1.0 then disable SRGB for the static texture
diff --git a/Engine/Source/Runtime/Engine/Private/TextureRenderTargetCube.cpp b/Engine/Source/Runtime/Engine/Private/TextureRenderTargetCube.cpp
--- a/Engine/Source/Runtime/Engine/Private/TextureRenderTargetCube.cpp
+++ b/Engine/Source/Runtime/Engine/Private/TextureRenderTargetCube.cpp
@@ -155,7 +155,7 @@ UTextureCube* UTextureRenderTargetCube::ConstructTextureCube(
 					{
 						if (CubeResource->ReadPixels(OutputBuffer, FReadSurfaceDataFlags(RCM_UNorm, (ECubeFace)SliceIndex)))
 						{
-							FMemory::Memcpy((FColor*)(SliceData + SliceIndex * MipSize), OutputBuffer.GetTypedData(), MipSize);
+							FMemory::Memcpy((FColor*)(SliceData + SliceIndex * MipSize), OutputBuffer.GetData(), MipSize);
 						}
 					}
 				}
@@ -167,7 +167,7 @@ UTextureCube* UTextureRenderTargetCube::ConstructTextureCube(
 					{
 						if (CubeResource->ReadPixels(OutputBuffer, FReadSurfaceDataFlags(RCM_UNorm, (ECubeFace)SliceIndex)))
 						{
-							FMemory::Memcpy((FFloat16Color*)(SliceData + SliceIndex * MipSize), OutputBuffer.GetTypedData(), MipSize);
+							FMemory::Memcpy((FFloat16Color*)(SliceData + SliceIndex * MipSize), OutputBuffer.GetData(), MipSize);
 						}
 					}
 				}
diff --git a/Engine/Source/Runtime/Engine/Private/TickTaskManager.cpp b/Engine/Source/Runtime/Engine/Private/TickTaskManager.cpp
--- a/Engine/Source/Runtime/Engine/Private/TickTaskManager.cpp
+++ b/Engine/Source/Runtime/Engine/Private/TickTaskManager.cpp
@@ -653,8 +653,8 @@ public:
 					{
 						check(Task < NumTasks);
 						FGraphEventArray Setup;
-						new (Setup) FGraphEventRef(TGraphTask<FQueueTickTasks>::CreateTask(NULL,ENamedThreads::GameThread).ConstructAndDispatchWhenReady(AllTickFunctions.GetTypedData() + Start, AllCompletionEvents.GetTypedData() + Start, NumThisTask, &Context));
-						new (QueueTickTasks) FGraphEventRef(TGraphTask<FPostTickTasks>::CreateTask(&Setup,ENamedThreads::GameThread).ConstructAndDispatchWhenReady(AllCompletionEvents.GetTypedData() + Start, NumThisTask));
+						new (Setup) FGraphEventRef(TGraphTask<FQueueTickTasks>::CreateTask(NULL,ENamedThreads::GameThread).ConstructAndDispatchWhenReady(AllTickFunctions.GetData() + Start, AllCompletionEvents.GetData() + Start, NumThisTask, &Context));
+						new (QueueTickTasks) FGraphEventRef(TGraphTask<FPostTickTasks>::CreateTask(&Setup,ENamedThreads::GameThread).ConstructAndDispatchWhenReady(AllCompletionEvents.GetData() + Start, NumThisTask));
 						Start += NumThisTask;
 						NumTicksSoFar = 0;
 
diff --git a/Engine/Source/Runtime/Engine/Private/UnrealClient.cpp b/Engine/Source/Runtime/Engine/Private/UnrealClient.cpp
--- a/Engine/Source/Runtime/Engine/Private/UnrealClient.cpp
+++ b/Engine/Source/Runtime/Engine/Private/UnrealClient.cpp
@@ -143,7 +143,7 @@ bool FRenderTarget::ReadFloat16Pixels(FFloat16Color* OutImageData,ECubeFace Cube
 	const int32 ImageHeight = GetSizeXY().Y;
 	for (int32 Y = 0; Y < ImageHeight; Y++)
 	{
-		FFloat16Color* SourceData = (FFloat16Color*)SurfaceData.GetTypedData() + Y * ImageWidth;
+		FFloat16Color* SourceData = (FFloat16Color*)SurfaceData.GetData() + Y * ImageWidth;
 		for (int32 X = 0; X < ImageWidth; X++)
 		{
 			OutImageColors[ Y * ImageWidth + X ] = SourceData[X];
diff --git a/Engine/Source/Runtime/Engine/Private/UnrealEngine.cpp b/Engine/Source/Runtime/Engine/Private/UnrealEngine.cpp
--- a/Engine/Source/Runtime/Engine/Private/UnrealEngine.cpp
+++ b/Engine/Source/Runtime/Engine/Private/UnrealEngine.cpp
@@ -10314,7 +10314,7 @@ void UEngine::HandleScreenshotCaptured(int32 Width, int32 Height, const TArray<F
 				const FString Filename = Local::GenerateScreenshotFilename(TEXT("bmp"));
 				if (Filename.Len() > 0)
 				{
-					FFileHelper::CreateBitmap(*Filename, Width, Height, Colors.GetTypedData());
+					FFileHelper::CreateBitmap(*Filename, Width, Height, Colors.GetData());
 				}
 			}
 			break;
diff --git a/Engine/Source/Runtime/Engine/Private/Vehicles/PhysXVehicleManager.cpp b/Engine/Source/Runtime/Engine/Private/Vehicles/PhysXVehicleManager.cpp
--- a/Engine/Source/Runtime/Engine/Private/Vehicles/PhysXVehicleManager.cpp
+++ b/Engine/Source/Runtime/Engine/Private/Vehicles/PhysXVehicleManager.cpp
@@ -190,8 +190,8 @@ void FPhysXVehicleManager::SetUpBatchedSceneQuery()
 		}
 		
 		PxBatchQueryDesc SqDesc(NumWheels, 0, 0);
-		SqDesc.queryMemory.userRaycastResultBuffer = WheelQueryResults.GetTypedData();
-		SqDesc.queryMemory.userRaycastTouchBuffer = WheelHitResults.GetTypedData();
+		SqDesc.queryMemory.userRaycastResultBuffer = WheelQueryResults.GetData();
+		SqDesc.queryMemory.userRaycastTouchBuffer = WheelHitResults.GetData();
 		SqDesc.queryMemory.raycastTouchBufferSize = WheelHitResults.Num();
 		SqDesc.preFilterShader = WheelRaycastPreFilter;
 
@@ -268,7 +268,7 @@ void FPhysXVehicleManager::Update( float DeltaTime )
 	}
 
 	// Suspension raycasts
-	PxVehicleSuspensionRaycasts( WheelRaycastBatchQuery, PVehicles.Num(), PVehicles.GetTypedData(), WheelQueryResults.Num(), WheelQueryResults.GetTypedData() );
+	PxVehicleSuspensionRaycasts( WheelRaycastBatchQuery, PVehicles.Num(), PVehicles.GetData(), WheelQueryResults.Num(), WheelQueryResults.GetData() );
 
 	// Tick vehicles
 	for ( int32 i = Vehicles.Num() - 1; i >= 0; --i )
@@ -304,7 +304,7 @@ void FPhysXVehicleManager::PreTick( float DeltaTime )
 
 void FPhysXVehicleManager::UpdateVehicles( float DeltaTime )
 {
-	PxVehicleUpdates( DeltaTime, GetSceneGravity(), *SurfaceTirePairs, PVehicles.Num(), PVehicles.GetTypedData(), PVehiclesWheelsStates.GetTypedData());
+	PxVehicleUpdates( DeltaTime, GetSceneGravity(), *SurfaceTirePairs, PVehicles.Num(), PVehicles.GetData(), PVehiclesWheelsStates.GetData());
 }
 
 PxVec3 FPhysXVehicleManager::GetSceneGravity()
@@ -374,7 +374,7 @@ void FPhysXVehicleManager::UpdateVehiclesWithTelemetry( float DeltaTime )
 
 	if ( TelemetryData )
 	{
-		PxVehicleUpdateSingleVehicleAndStoreTelemetryData( DeltaTime, GetSceneGravity(), *SurfaceTirePairs, TelemetryVehicle, PVehiclesWheelsStates.GetTypedData(), *TelemetryData );
+		PxVehicleUpdateSingleVehicleAndStoreTelemetryData( DeltaTime, GetSceneGravity(), *SurfaceTirePairs, TelemetryVehicle, PVehiclesWheelsStates.GetData(), *TelemetryData );
 
 		if ( PVehicles.Num() > 1 )
 		{
@@ -385,7 +385,7 @@ void FPhysXVehicleManager::UpdateVehiclesWithTelemetry( float DeltaTime )
 	{
 		UE_LOG( LogPhysics, Warning, TEXT("Cannot record telemetry for vehicle, it does not have 4 wheels") );
 
-		PxVehicleUpdates( DeltaTime, GetSceneGravity(), *SurfaceTirePairs, PVehicles.Num(), PVehicles.GetTypedData(), PVehiclesWheelsStates.GetTypedData() );
+		PxVehicleUpdates( DeltaTime, GetSceneGravity(), *SurfaceTirePairs, PVehicles.Num(), PVehicles.GetData(), PVehiclesWheelsStates.GetData() );
 	}
 }
 
diff --git a/Engine/Source/Runtime/Engine/Private/Vehicles/WheeledVehicleMovementComponent.cpp b/Engine/Source/Runtime/Engine/Private/Vehicles/WheeledVehicleMovementComponent.cpp
--- a/Engine/Source/Runtime/Engine/Private/Vehicles/WheeledVehicleMovementComponent.cpp
+++ b/Engine/Source/Runtime/Engine/Private/Vehicles/WheeledVehicleMovementComponent.cpp
@@ -341,7 +341,7 @@ void UWheeledVehicleMovementComponent::SetupWheels( PxVehicleWheelsSimData* PWhe
 	TArray<PxShape*> Shapes;
 	Shapes.AddZeroed(NumShapes);
 
-	PVehicleActor->getShapes( Shapes.GetTypedData(), NumShapes );
+	PVehicleActor->getShapes( Shapes.GetData(), NumShapes );
 
 	for ( int32 WheelIdx = 0; WheelIdx < WheelSetups.Num(); ++WheelIdx )
 	{
diff --git a/Engine/Source/Runtime/Engine/Private/VisualLog.cpp b/Engine/Source/Runtime/Engine/Private/VisualLog.cpp
--- a/Engine/Source/Runtime/Engine/Private/VisualLog.cpp
+++ b/Engine/Source/Runtime/Engine/Private/VisualLog.cpp
@@ -23,7 +23,7 @@ FVisLogEntry::FVisLogEntry(const class AActor* InActor, TArray<TWeakObjectPtr<UO
 		InActor->GrabDebugSnapshot(this);
 		if (Children != NULL)
 		{
-			TWeakObjectPtr<UObject>* WeakActorPtr = Children->GetTypedData();
+			TWeakObjectPtr<UObject>* WeakActorPtr = Children->GetData();
 			for (int32 Index = 0; Index < Children->Num(); ++Index, ++WeakActorPtr)
 			{
 				if (WeakActorPtr->IsValid())
@@ -191,13 +191,13 @@ FVisLogEntry::FVisLogEntry(TSharedPtr<FJsonValue> FromJson)
 				// uncompress decoded data to get final TArray<uint8>
 				int32 UncompressedSize = 0;
 				const int32 HeaderSize = sizeof(int32);
-				uint8* SrcBuffer = (uint8*)CompressedDataBlock.GetTypedData();
+				uint8* SrcBuffer = (uint8*)CompressedDataBlock.GetData();
 				FMemory::Memcpy(&UncompressedSize, SrcBuffer, HeaderSize);
 				SrcBuffer += HeaderSize;
 				const int32 CompressedSize = CompressedDataBlock.Num() - HeaderSize;
 
 				Sample.Data.AddZeroed(UncompressedSize);
-				FCompression::UncompressMemory((ECompressionFlags)(COMPRESS_ZLIB), (void*)Sample.Data.GetTypedData(), UncompressedSize, SrcBuffer, CompressedSize);
+				FCompression::UncompressMemory((ECompressionFlags)(COMPRESS_ZLIB), (void*)Sample.Data.GetData(), UncompressedSize, SrcBuffer, CompressedSize);
 			}
 		}
 	}
@@ -212,7 +212,7 @@ TSharedPtr<FJsonValue> FVisLogEntry::ToJson() const
 	JsonEntryObject->SetStringField(VisualLogJson::TAG_LOCATION, Location.ToString());
 
 	const int32 StatusCategoryCount = Status.Num();
-	const FVisLogEntry::FStatusCategory* StatusCategory = Status.GetTypedData();
+	const FVisLogEntry::FStatusCategory* StatusCategory = Status.GetData();
 	TArray< TSharedPtr<FJsonValue> > JsonStatus;
 	JsonStatus.AddZeroed(StatusCategoryCount);
 
@@ -234,7 +234,7 @@ TSharedPtr<FJsonValue> FVisLogEntry::ToJson() const
 	JsonEntryObject->SetArrayField(VisualLogJson::TAG_STATUS, JsonStatus);
 
 	const int32 LogLineCount = LogLines.Num();
-	const FVisLogEntry::FLogLine* LogLine = LogLines.GetTypedData();
+	const FVisLogEntry::FLogLine* LogLine = LogLines.GetData();
 	TArray< TSharedPtr<FJsonValue> > JsonLines;
 	JsonLines.AddZeroed(LogLineCount);
 
@@ -252,7 +252,7 @@ TSharedPtr<FJsonValue> FVisLogEntry::ToJson() const
 	JsonEntryObject->SetArrayField(VisualLogJson::TAG_LOGLINES, JsonLines);
 
 	const int32 ElementsToDrawCount = ElementsToDraw.Num();
-	const FVisLogEntry::FElementToDraw* Element = ElementsToDraw.GetTypedData();
+	const FVisLogEntry::FElementToDraw* Element = ElementsToDraw.GetData();
 	TArray< TSharedPtr<FJsonValue> > JsonElementsToDraw;
 	JsonElementsToDraw.AddZeroed(ElementsToDrawCount);
 
@@ -270,7 +270,7 @@ TSharedPtr<FJsonValue> FVisLogEntry::ToJson() const
 		const int32 PointsCount = Element->Points.Num();
 		TArray< TSharedPtr<FJsonValue> > JsonStringPoints;
 		JsonStringPoints.AddZeroed(PointsCount);
-		const FVector* PointToDraw = Element->Points.GetTypedData();
+		const FVector* PointToDraw = Element->Points.GetData();
 		for (int32 PointIndex = 0; PointIndex < PointsCount; ++PointIndex, ++PointToDraw)
 		{
 			JsonStringPoints[PointIndex] = MakeShareable(new FJsonValueString(PointToDraw->ToString()));
@@ -283,7 +283,7 @@ TSharedPtr<FJsonValue> FVisLogEntry::ToJson() const
 	JsonEntryObject->SetArrayField(VisualLogJson::TAG_ELEMENTSTODRAW, JsonElementsToDraw);
 	
 	const int32 HistogramSamplesCount = HistogramSamples.Num();
-	const FVisLogEntry::FHistogramSample* Sample = HistogramSamples.GetTypedData();
+	const FVisLogEntry::FHistogramSample* Sample = HistogramSamples.GetData();
 	TArray<TSharedPtr<FJsonValue> > JsonHistogramSamples;
 	JsonHistogramSamples.AddZeroed(HistogramSamplesCount);
 
@@ -315,7 +315,7 @@ TSharedPtr<FJsonValue> FVisLogEntry::ToJson() const
 		CompressedData.Init(0, HeaderSize + FMath::TruncToInt(1.1f * UncompressedSize));
 
 		int32 CompressedSize = CompressedData.Num() - HeaderSize;
-		uint8* DestBuffer = CompressedData.GetTypedData();
+		uint8* DestBuffer = CompressedData.GetData();
 		FMemory::Memcpy(DestBuffer, &UncompressedSize, HeaderSize);
 		DestBuffer += HeaderSize;
 
@@ -477,7 +477,7 @@ TSharedPtr<FJsonValue> FActorsVisLog::ToJson() const
 	JsonLogObject->SetStringField(VisualLogJson::TAG_NAME, Name.ToString());
 	JsonLogObject->SetStringField(VisualLogJson::TAG_FULLNAME, FullName);
 
-	const TSharedPtr<FVisLogEntry>* Entry = Entries.GetTypedData();
+	const TSharedPtr<FVisLogEntry>* Entry = Entries.GetData();
 	const int32 EntryCount = Entries.Num();
 	TArray< TSharedPtr<FJsonValue> > JsonLogEntries;
 	JsonLogEntries.AddZeroed(EntryCount);
@@ -723,7 +723,7 @@ void FVisualLog::Redirect(UObject* Source, const AActor* NewRedirection)
 		TArray<TWeakObjectPtr<UObject> >* Children = RedirectsMap.Find(SourceAsActor);
 		if (Children != NULL)
 		{
-			TWeakObjectPtr<UObject>* WeakActorPtr = Children->GetTypedData();
+			TWeakObjectPtr<UObject>* WeakActorPtr = Children->GetData();
 			for (int32 Index = 0; Index < Children->Num(); ++Index)
 			{
 				if (WeakActorPtr->IsValid())
diff --git a/Engine/Source/Runtime/Engine/Public/AnimEncoding_ConstantKeyLerp.h b/Engine/Source/Runtime/Engine/Public/AnimEncoding_ConstantKeyLerp.h
--- a/Engine/Source/Runtime/Engine/Public/AnimEncoding_ConstantKeyLerp.h
+++ b/Engine/Source/Runtime/Engine/Public/AnimEncoding_ConstantKeyLerp.h
@@ -399,10 +399,10 @@ inline void AEFConstantKeyLerp<FORMAT>::GetPoseRotations(
 		const int32 AtomIndex = Pair.AtomIndex;
 		FTransform& BoneAtom = Atoms[AtomIndex];
 
-		const int32* RESTRICT TrackData = Seq.CompressedTrackOffsets.GetTypedData() + (TrackIndex*4);
+		const int32* RESTRICT TrackData = Seq.CompressedTrackOffsets.GetData() + (TrackIndex*4);
 		const int32 RotKeysOffset	= *(TrackData+2);
 		const int32 NumRotKeys	= *(TrackData+3);
-		const uint8* RESTRICT RotStream		= Seq.CompressedByteStream.GetTypedData()+RotKeysOffset;
+		const uint8* RESTRICT RotStream		= Seq.CompressedByteStream.GetData()+RotKeysOffset;
 
 		// call the decoder directly (not through the vtable)
 		AEFConstantKeyLerp<FORMAT>::GetBoneAtomRotation(BoneAtom, Seq, RotStream, NumRotKeys, Time, RelativePos);
@@ -435,7 +435,7 @@ inline void AEFConstantKeyLerp<FORMAT>::GetPoseTranslations(
 	for (int32 PairIndex=0; PairIndex<PairCount; ++PairIndex)
 	{
 		const BoneTrackPair& Pair = DesiredPairs[PairIndex];
-		FPlatformMisc::Prefetch(Atoms.GetTypedData() + Pair.AtomIndex);
+		FPlatformMisc::Prefetch(Atoms.GetData() + Pair.AtomIndex);
 	}
 
 	for (int32 PairIndex=0; PairIndex<PairCount; ++PairIndex)
@@ -443,7 +443,7 @@ inline void AEFConstantKeyLerp<FORMAT>::GetPoseTranslations(
 		int32 PrefetchIndex = PairIndex + PrefetchCount;
 		if (PrefetchIndex < PairCount)
 		{
-			FPlatformMisc::Prefetch(Atoms.GetTypedData() + DesiredPairs[PrefetchIndex].AtomIndex);
+			FPlatformMisc::Prefetch(Atoms.GetData() + DesiredPairs[PrefetchIndex].AtomIndex);
 		}
 
 		const BoneTrackPair& Pair = DesiredPairs[PairIndex];
@@ -451,10 +451,10 @@ inline void AEFConstantKeyLerp<FORMAT>::GetPoseTranslations(
 		const int32 AtomIndex = Pair.AtomIndex;
 		FTransform& BoneAtom = Atoms[AtomIndex];
 
-		const int32* RESTRICT TrackData = Seq.CompressedTrackOffsets.GetTypedData() + (TrackIndex*4);
+		const int32* RESTRICT TrackData = Seq.CompressedTrackOffsets.GetData() + (TrackIndex*4);
 		const int32 TransKeysOffset	= *(TrackData+0);
 		const int32 NumTransKeys		= *(TrackData+1);
-		const uint8* RESTRICT TransStream = Seq.CompressedByteStream.GetTypedData()+TransKeysOffset;
+		const uint8* RESTRICT TransStream = Seq.CompressedByteStream.GetData()+TransKeysOffset;
 
 		// call the decoder directly (not through the vtable)
 		AEFConstantKeyLerp<FORMAT>::GetBoneAtomTranslation(BoneAtom, Seq, TransStream, NumTransKeys, Time, RelativePos);
@@ -489,7 +489,7 @@ inline void AEFConstantKeyLerp<FORMAT>::GetPoseScales(
 	for (int32 PairIndex=0; PairIndex<PairCount; ++PairIndex)
 	{
 		const BoneTrackPair& Pair = DesiredPairs[PairIndex];
-		FPlatformMisc::Prefetch(Atoms.GetTypedData() + Pair.AtomIndex);
+		FPlatformMisc::Prefetch(Atoms.GetData() + Pair.AtomIndex);
 	}
 
 	for (int32 PairIndex=0; PairIndex<PairCount; ++PairIndex)
@@ -497,7 +497,7 @@ inline void AEFConstantKeyLerp<FORMAT>::GetPoseScales(
 		int32 PrefetchIndex = PairIndex + PrefetchCount;
 		if (PrefetchIndex < PairCount)
 		{
-			FPlatformMisc::Prefetch(Atoms.GetTypedData() + DesiredPairs[PrefetchIndex].AtomIndex);
+			FPlatformMisc::Prefetch(Atoms.GetData() + DesiredPairs[PrefetchIndex].AtomIndex);
 		}
 
 		const BoneTrackPair& Pair = DesiredPairs[PairIndex];
@@ -507,7 +507,7 @@ inline void AEFConstantKeyLerp<FORMAT>::GetPoseScales(
 
 		const int32 ScaleKeysOffset		= Seq.CompressedScaleOffsets.GetOffsetData(TrackIndex, 0);
 		const int32 NumScaleKeys		= Seq.CompressedScaleOffsets.GetOffsetData(TrackIndex, 1);
-		const uint8* RESTRICT ScaleStream = Seq.CompressedByteStream.GetTypedData()+ScaleKeysOffset;
+		const uint8* RESTRICT ScaleStream = Seq.CompressedByteStream.GetData()+ScaleKeysOffset;
 
 		// call the decoder directly (not through the vtable)
 		AEFConstantKeyLerp<FORMAT>::GetBoneAtomScale(BoneAtom, Seq, ScaleStream, NumScaleKeys, Time, RelativePos);
diff --git a/Engine/Source/Runtime/Engine/Public/AnimEncoding_VariableKeyLerp.h b/Engine/Source/Runtime/Engine/Public/AnimEncoding_VariableKeyLerp.h
--- a/Engine/Source/Runtime/Engine/Public/AnimEncoding_VariableKeyLerp.h
+++ b/Engine/Source/Runtime/Engine/Public/AnimEncoding_VariableKeyLerp.h
@@ -409,10 +409,10 @@ FORCEINLINE_DEBUGGABLE void AEFVariableKeyLerp<FORMAT>::GetPoseRotations(
 		const int32 AtomIndex = Pair.AtomIndex;
 		FTransform& BoneAtom = Atoms[AtomIndex];
 
-		const int32* RESTRICT TrackData = Seq.CompressedTrackOffsets.GetTypedData() + (TrackIndex*4);
+		const int32* RESTRICT TrackData = Seq.CompressedTrackOffsets.GetData() + (TrackIndex*4);
 		const int32 RotKeysOffset	= *(TrackData+2);
 		const int32 NumRotKeys	= *(TrackData+3);
-		const uint8* RESTRICT RotStream		= Seq.CompressedByteStream.GetTypedData()+RotKeysOffset;
+		const uint8* RESTRICT RotStream		= Seq.CompressedByteStream.GetData()+RotKeysOffset;
 
 		// call the decoder directly (not through the vtable)
 		AEFVariableKeyLerp<FORMAT>::GetBoneAtomRotation(BoneAtom, Seq, RotStream, NumRotKeys, Time, RelativePos);
@@ -445,10 +445,10 @@ FORCEINLINE_DEBUGGABLE void AEFVariableKeyLerp<FORMAT>::GetPoseTranslations(
 		const int32 AtomIndex = Pair.AtomIndex;
 		FTransform& BoneAtom = Atoms[AtomIndex];
 
-		const int32* RESTRICT TrackData = Seq.CompressedTrackOffsets.GetTypedData() + (TrackIndex*4);
+		const int32* RESTRICT TrackData = Seq.CompressedTrackOffsets.GetData() + (TrackIndex*4);
 		const int32 TransKeysOffset	= *(TrackData+0);
 		const int32 NumTransKeys		= *(TrackData+1);
-		const uint8* RESTRICT TransStream = Seq.CompressedByteStream.GetTypedData()+TransKeysOffset;
+		const uint8* RESTRICT TransStream = Seq.CompressedByteStream.GetData()+TransKeysOffset;
 
 		// call the decoder directly (not through the vtable)
 		AEFVariableKeyLerp<FORMAT>::GetBoneAtomTranslation(BoneAtom, Seq, TransStream, NumTransKeys, Time, RelativePos);
@@ -485,7 +485,7 @@ FORCEINLINE_DEBUGGABLE void AEFVariableKeyLerp<FORMAT>::GetPoseScales(
 
 		const int32 ScaleKeysOffset	= Seq.CompressedScaleOffsets.GetOffsetData(TrackIndex, 0);
 		const int32 NumScaleKeys	= Seq.CompressedScaleOffsets.GetOffsetData(TrackIndex, 1);
-		const uint8* RESTRICT ScaleStream = Seq.CompressedByteStream.GetTypedData()+ScaleKeysOffset;
+		const uint8* RESTRICT ScaleStream = Seq.CompressedByteStream.GetData()+ScaleKeysOffset;
 
 		// call the decoder directly (not through the vtable)
 		AEFVariableKeyLerp<FORMAT>::GetBoneAtomScale(BoneAtom, Seq, ScaleStream, NumScaleKeys, Time, RelativePos);
diff --git a/Engine/Source/Runtime/Engine/Public/Distributions.h b/Engine/Source/Runtime/Engine/Public/Distributions.h
--- a/Engine/Source/Runtime/Engine/Public/Distributions.h
+++ b/Engine/Source/Runtime/Engine/Public/Distributions.h
@@ -128,7 +128,7 @@ struct FDistributionLookupTable
 		if ( EntryCount > 0 )
 		{
 			const int32 ValuesPerEntry = GetValuesPerEntry();
-			const float* Entry = Values.GetTypedData();
+			const float* Entry = Values.GetData();
 
 			// Initialize to the first entry in the table.
 			for ( int32 ValueIndex = 0; ValueIndex < ValuesPerEntry; ++ValueIndex )
diff --git a/Engine/Source/Runtime/Engine/Public/Landscape/LandscapeDataAccess.h b/Engine/Source/Runtime/Engine/Public/Landscape/LandscapeDataAccess.h
--- a/Engine/Source/Runtime/Engine/Public/Landscape/LandscapeDataAccess.h
+++ b/Engine/Source/Runtime/Engine/Public/Landscape/LandscapeDataAccess.h
@@ -87,7 +87,7 @@ public:
 		}
 		(*MipInfo)[MipLevel].LockCount++;
 
-		return (*MipInfo)[MipLevel].MipData.GetTypedData();
+		return (*MipInfo)[MipLevel].MipData.GetData();
 	}
 
 	void UnlockMip(UTexture2D* Texture, int32 MipLevel)
diff --git a/Engine/Source/Runtime/IOS/MetalRHI/Private/MetalShaders.cpp b/Engine/Source/Runtime/IOS/MetalRHI/Private/MetalShaders.cpp
--- a/Engine/Source/Runtime/IOS/MetalRHI/Private/MetalShaders.cpp
+++ b/Engine/Source/Runtime/IOS/MetalRHI/Private/MetalShaders.cpp
@@ -65,7 +65,7 @@ TMetalBaseShader<BaseResourceType>::TMetalBaseShader(const TArray<uint8>& InCode
 
 	// remember where the header ended and code (precompiled or source) begins
 	int32 CodeOffset = Ar.Tell();
-	const ANSICHAR* SourceCode = (ANSICHAR*)InCode.GetTypedData() + CodeOffset;
+	const ANSICHAR* SourceCode = (ANSICHAR*)InCode.GetData() + CodeOffset;
 
 	uint32 CodeLength, CodeCRC;
 	if (OfflineCompiledFlag)
@@ -73,7 +73,7 @@ TMetalBaseShader<BaseResourceType>::TMetalBaseShader(const TArray<uint8>& InCode
 		CodeLength = InCode.Num() - CodeOffset;
 
 		// CRC the compiled code
-		CodeCRC = FCrc::MemCrc_DEPRECATED(InCode.GetTypedData() + CodeOffset, InCode.Num() - CodeOffset);
+		CodeCRC = FCrc::MemCrc_DEPRECATED(InCode.GetData() + CodeOffset, InCode.Num() - CodeOffset);
 	}
 	else
 	{
@@ -98,7 +98,7 @@ TMetalBaseShader<BaseResourceType>::TMetalBaseShader(const TArray<uint8>& InCode
 			//		dispatch_data_t GCDBuffer = dispatch_data_create(InCode.GetTypedData() + CodeOffset, InCode.Num() - CodeOffset, nil, DISPATCH_DATA_DESTRUCTOR_DEFAULT);
 			uint32 BufferSize = InCode.Num() - CodeOffset;
 			void* Buffer = FMemory::Malloc( BufferSize );
-			FMemory::Memcpy( Buffer, InCode.GetTypedData() + CodeOffset, BufferSize );
+			FMemory::Memcpy( Buffer, InCode.GetData() + CodeOffset, BufferSize );
 			dispatch_data_t GCDBuffer = dispatch_data_create(Buffer, BufferSize, dispatch_get_main_queue(), ^(void) { FMemory::Free(Buffer); } );
 
 			// load up the already compiled shader
@@ -141,8 +141,8 @@ TMetalBaseShader<BaseResourceType>::TMetalBaseShader(const TArray<uint8>& InCode
 
 			GlslCode.Empty(CodeLength + 1);
 			GlslCode.AddUninitialized(CodeLength + 1);
-			FMemory::Memcpy(GlslCode.GetTypedData(), SourceCode, CodeLength + 1);
-			GlslCodeString = (ANSICHAR*)GlslCode.GetTypedData();
+			FMemory::Memcpy(GlslCode.GetData(), SourceCode, CodeLength + 1);
+			GlslCodeString = (ANSICHAR*)GlslCode.GetData();
 			GlslCodeNSString = ShaderSource;
 
 #if !UE_BUILD_SHIPPING
diff --git a/Engine/Source/Runtime/IOS/MetalRHI/Public/MetalShaderResources.h b/Engine/Source/Runtime/IOS/MetalRHI/Public/MetalShaderResources.h
--- a/Engine/Source/Runtime/IOS/MetalRHI/Public/MetalShaderResources.h
+++ b/Engine/Source/Runtime/IOS/MetalRHI/Public/MetalShaderResources.h
@@ -33,7 +33,7 @@ struct FMetalShaderResourceTable : public FBaseShaderResourceTable
 		{
 			return false;
 		}
-		if (FMemory::Memcmp(A.TextureMap.GetTypedData(), B.TextureMap.GetTypedData(), A.TextureMap.GetTypeSize()*A.TextureMap.Num()) != 0)
+		if (FMemory::Memcmp(A.TextureMap.GetData(), B.TextureMap.GetData(), A.TextureMap.GetTypeSize()*A.TextureMap.Num()) != 0)
 		{
 			return false;
 		}
@@ -88,14 +88,14 @@ struct FMetalShaderBindings
 			return bEqual;
 		}
 
-		bEqual &= FMemory::Memcmp(A.PackedGlobalArrays.GetTypedData(), B.PackedGlobalArrays.GetTypedData(), A.PackedGlobalArrays.GetTypeSize()*A.PackedGlobalArrays.Num()) == 0;
+		bEqual &= FMemory::Memcmp(A.PackedGlobalArrays.GetData(), B.PackedGlobalArrays.GetData(), A.PackedGlobalArrays.GetTypeSize()*A.PackedGlobalArrays.Num()) == 0;
 
 		for (int32 Item = 0; Item < A.PackedUniformBuffers.Num(); Item++)
 		{
 			const TArray<CrossCompiler::FPackedArrayInfo> &ArrayA = A.PackedUniformBuffers[Item];
 			const TArray<CrossCompiler::FPackedArrayInfo> &ArrayB = B.PackedUniformBuffers[Item];
 
-			bEqual &= FMemory::Memcmp(ArrayA.GetTypedData(), ArrayB.GetTypedData(), ArrayA.GetTypeSize()*ArrayA.Num()) == 0;
+			bEqual &= FMemory::Memcmp(ArrayA.GetData(), ArrayB.GetData(), ArrayA.GetTypeSize()*ArrayA.Num()) == 0;
 		}
 
 		return bEqual;
@@ -109,12 +109,12 @@ struct FMetalShaderBindings
 		Hash |= Binding.NumUniformBuffers << 24;
 		Hash ^= Binding.NumUAVs;
 		Hash ^= Binding.bHasRegularUniformBuffers << 8;
-		Hash ^= FCrc::MemCrc_DEPRECATED(Binding.PackedGlobalArrays.GetTypedData(), Binding.PackedGlobalArrays.GetTypeSize()*Binding.PackedGlobalArrays.Num());
+		Hash ^= FCrc::MemCrc_DEPRECATED(Binding.PackedGlobalArrays.GetData(), Binding.PackedGlobalArrays.GetTypeSize()*Binding.PackedGlobalArrays.Num());
 
 		for (int32 Item = 0; Item < Binding.PackedUniformBuffers.Num(); Item++)
 		{
 			const TArray<CrossCompiler::FPackedArrayInfo> &Array = Binding.PackedUniformBuffers[Item];
-			Hash ^= FCrc::MemCrc_DEPRECATED(Array.GetTypedData(), Array.GetTypeSize()* Array.Num());
+			Hash ^= FCrc::MemCrc_DEPRECATED(Array.GetData(), Array.GetTypeSize()* Array.Num());
 		}
 		return Hash;
 	}
diff --git a/Engine/Source/Runtime/Launch/Private/IOS/IOSAppDelegateConsoleHandling.cpp b/Engine/Source/Runtime/Launch/Private/IOS/IOSAppDelegateConsoleHandling.cpp
--- a/Engine/Source/Runtime/Launch/Private/IOS/IOSAppDelegateConsoleHandling.cpp
+++ b/Engine/Source/Runtime/Launch/Private/IOS/IOSAppDelegateConsoleHandling.cpp
@@ -59,8 +59,8 @@
 			TArray<TCHAR> Ch;
 			Ch.AddZeroed([ConsoleCommand length]);
 
-			FPlatformString::CFStringToTCHAR((CFStringRef)ConsoleCommand, Ch.GetTypedData());
-			new(GEngine->DeferredCommands) FString(Ch.GetTypedData());
+			FPlatformString::CFStringToTCHAR((CFStringRef)ConsoleCommand, Ch.GetData());
+			new(GEngine->DeferredCommands) FString(Ch.GetData());
 		}
 
 		uint32 ExistingCommand = [self.ConsoleHistoryValues indexOfObjectPassingTest:
diff --git a/Engine/Source/Runtime/NetworkFileSystem/Private/NetworkFileServerConnection.cpp b/Engine/Source/Runtime/NetworkFileSystem/Private/NetworkFileServerConnection.cpp
--- a/Engine/Source/Runtime/NetworkFileSystem/Private/NetworkFileServerConnection.cpp
+++ b/Engine/Source/Runtime/NetworkFileSystem/Private/NetworkFileServerConnection.cpp
@@ -877,7 +877,7 @@ void FNetworkFileServerClientConnection::PackageFile( FString& Filename, FArchiv
 		{
 			// read it
 			Contents.AddUninitialized(File->Size());
-			File->Read(Contents.GetTypedData(), Contents.Num());
+			File->Read(Contents.GetData(), Contents.Num());
 		}
 
 		// close it
diff --git a/Engine/Source/Runtime/Networking/Public/Common/UdpSocketSender.h b/Engine/Source/Runtime/Networking/Public/Common/UdpSocketSender.h
--- a/Engine/Source/Runtime/Networking/Public/Common/UdpSocketSender.h
+++ b/Engine/Source/Runtime/Networking/Public/Common/UdpSocketSender.h
@@ -150,7 +150,7 @@ public:
 					int32 Sent = 0;
 
 					SendQueue.Dequeue(Packet);
-					Socket->SendTo(Packet.Data->GetTypedData(), Packet.Data->Num(), Sent, *Packet.Recipient.ToInternetAddr());
+					Socket->SendTo(Packet.Data->GetData(), Packet.Data->Num(), Sent, *Packet.Recipient.ToInternetAddr());
 					
 					if (Sent != Packet.Data->Num())
 					{
