Fixed compilation warrning - changed deprecated GetTypedData() to GetData() #ue4

diff --git a/Engine/Source/Runtime/Engine/Private/VisualLog.cpp b/Engine/Source/Runtime/Engine/Private/VisualLog.cpp
--- a/Engine/Source/Runtime/Engine/Private/VisualLog.cpp
+++ b/Engine/Source/Runtime/Engine/Private/VisualLog.cpp
@@ -45,7 +45,7 @@ FVisLogEntry::FVisLogEntry(float InTimeStamp, FVector InLocation, const UObject*
 	Location = InLocation;
 	if (Children != NULL)
 	{
-		TWeakObjectPtr<UObject>* WeakActorPtr = Children->GetTypedData();
+		TWeakObjectPtr<UObject>* WeakActorPtr = Children->GetData();
 		for (int32 Index = 0; Index < Children->Num(); ++Index, ++WeakActorPtr)
 		{
 			if (WeakActorPtr->IsValid())
