Clean up opposing normal code and fix case where PHit contains a shape that's not convex,trimesh, or heightfield

diff --git a/Engine/Source/Runtime/Engine/Private/Collision/CollisionConversions.cpp b/Engine/Source/Runtime/Engine/Private/Collision/CollisionConversions.cpp
--- a/Engine/Source/Runtime/Engine/Private/Collision/CollisionConversions.cpp
+++ b/Engine/Source/Runtime/Engine/Private/Collision/CollisionConversions.cpp
@@ -80,50 +80,104 @@ static PxVec3 TransformNormalToShapeSpace(const PxMeshScale& meshScale, const Px
 	}
 }
 
+static bool FindSimpleOpposingNormal(const PxLocationHit& PHit, const FVector& Direction, FVector& OutNormal)
+{
+	bool bNormalData = PHit.flags & PxHitFlag::eNORMAL;
+	OutNormal = bNormalData ? P2UVector(PHit.normal) : -Direction;
+	return true;
+}
 
+static bool FindHeightFieldOpposingNormal(const PxLocationHit& PHit, const FVector& Direction, FVector& OutNormal)
+{
 
-/**
- * Util to find the normal of the face that we hit.
- * @param PHit - incoming hit from PhysX
- * @param Direction - direction of sweep test (not normalized)
- * @param OutNormal - normal we may recompute based on the faceIndex of the hit
- * @return true if we compute a new normal for the geometry.
- */
-static bool FindGeomOpposingNormal(const PxLocationHit& PHit, const FVector& Direction, FVector& OutNormal)
+	PxHeightFieldGeometry PHeightFieldGeom;
+	bool bSuccess = PHit.shape->getHeightFieldGeometry(PHeightFieldGeom);
+	check(bSuccess);	//we should only call this function when we have a heightfield
+	if (PHeightFieldGeom.heightField)
+	{
+		const PxU32 TriIndex = PHit.faceIndex;
+		const PxTransform PShapeWorldPose = PxShapeExt::getGlobalPose(*PHit.shape, *PHit.actor);
+
+		PxTriangle Tri;
+		PxMeshQuery::getTriangle(PHeightFieldGeom, PShapeWorldPose, TriIndex, Tri);
+
+		PxVec3 TriNormal;
+		Tri.normal(TriNormal);
+		OutNormal = P2UVector(TriNormal);
+
+		return true;
+	}
+
+	return false;
+}
+
+static bool FindConvexMeshOpposingNormal(const PxLocationHit& PHit, const FVector& Direction, FVector& OutNormal)
 {
-	checkf(InvalidQueryHit.faceIndex == 0xFFFFffff, TEXT("Engine code needs fixing: PhysX invalid face index sentinel has changed or is not part of default PxQueryHit!"));
-	if (PHit.faceIndex == InvalidQueryHit.faceIndex)
+	PxConvexMeshGeometry PConvexMeshGeom;
+	bool bSuccess = PHit.shape->getConvexMeshGeometry(PConvexMeshGeom);
+	check(bSuccess);	//should only call this function when we have a convex mesh
+
+	if (PConvexMeshGeom.convexMesh)
 	{
-		return false;
+		check(PHit.faceIndex < PConvexMeshGeom.convexMesh->getNbPolygons());
+
+		const PxU32 PolyIndex = PHit.faceIndex;
+		PxHullPolygon PPoly;
+		bool bSuccessData = PConvexMeshGeom.convexMesh->getPolygonData(PolyIndex, PPoly);
+		if (bSuccessData)
+		{
+			// Account for non-uniform scale in local space normal.
+			const PxVec3 PPlaneNormal(PPoly.mPlane[0], PPoly.mPlane[1], PPoly.mPlane[2]);
+			const PxVec3 PLocalPolyNormal = TransformNormalToShapeSpace(PConvexMeshGeom.scale, PPlaneNormal.getNormalized());
+
+			// Convert to world space
+			const PxTransform PShapeWorldPose = PxShapeExt::getGlobalPose(*PHit.shape, *PHit.actor);
+			const PxVec3 PWorldPolyNormal = PShapeWorldPose.rotate(PLocalPolyNormal);
+			OutNormal = P2UVector(PWorldPolyNormal);
+
+#if !(UE_BUILD_SHIPPING || UE_BUILD_TEST)
+			if (!OutNormal.IsNormalized())
+			{
+				UE_LOG(LogPhysics, Warning, TEXT("Non-normalized Normal (Hit shape is ConvexMesh): %s (LocalPolyNormal:%s)"), *OutNormal.ToString(), *P2UVector(PLocalPolyNormal).ToString());
+				UE_LOG(LogPhysics, Warning, TEXT("WorldTransform \n: %s"), *P2UTransform(PShapeWorldPose).ToString());
+			}
+#endif
+			return true;
+		}
 	}
 
+	return false;
+}
+
+static bool FindTriMeshOpposingNormal(const PxLocationHit& PHit, const FVector& Direction, FVector& OutNormal)
+{
 	PxTriangleMeshGeometry PTriMeshGeom;
-	PxConvexMeshGeometry PConvexMeshGeom;
-	PxHeightFieldGeometry PHeightFieldGeom;
+	bool bSuccess = PHit.shape->getTriangleMeshGeometry(PTriMeshGeom);
+	check(bSuccess);	//this function should only be called when we have a trimesh
 
-	if(	PHit.shape->getTriangleMeshGeometry(PTriMeshGeom) && 
-		PTriMeshGeom.triangleMesh != NULL &&
-		PHit.faceIndex < PTriMeshGeom.triangleMesh->getNbTriangles() )
+	if (PTriMeshGeom.triangleMesh)
 	{
+		check(PHit.faceIndex < PTriMeshGeom.triangleMesh->getNbTriangles());
+
 		const PxU32 TriIndex = PHit.faceIndex;
 		const void* Triangles = PTriMeshGeom.triangleMesh->getTriangles();
 
 		// Grab triangle indices that we hit
 		int32 I0, I1, I2;
 
-		if(PTriMeshGeom.triangleMesh->getTriangleMeshFlags() & PxTriangleMeshFlag::eHAS_16BIT_TRIANGLE_INDICES) 
+		if (PTriMeshGeom.triangleMesh->getTriangleMeshFlags() & PxTriangleMeshFlag::eHAS_16BIT_TRIANGLE_INDICES)
 		{
 			PxU16* P16BitIndices = (PxU16*)Triangles;
-			I0 = P16BitIndices[(TriIndex*3)+0];
-			I1 = P16BitIndices[(TriIndex*3)+1];
-			I2 = P16BitIndices[(TriIndex*3)+2];
+			I0 = P16BitIndices[(TriIndex * 3) + 0];
+			I1 = P16BitIndices[(TriIndex * 3) + 1];
+			I2 = P16BitIndices[(TriIndex * 3) + 2];
 		}
 		else
 		{
 			PxU32* P32BitIndices = (PxU32*)Triangles;
-			I0 = P32BitIndices[(TriIndex*3)+0];
-			I1 = P32BitIndices[(TriIndex*3)+1];
-			I2 = P32BitIndices[(TriIndex*3)+2];
+			I0 = P32BitIndices[(TriIndex * 3) + 0];
+			I1 = P32BitIndices[(TriIndex * 3) + 1];
+			I2 = P32BitIndices[(TriIndex * 3) + 2];
 		}
 
 		// Get verts we hit (local space)
@@ -135,7 +189,7 @@ static bool FindGeomOpposingNormal(const PxLocationHit& PHit, const FVector& Dir
 		// Find normal of triangle (local space), and account for non-uniform scale
 		const PxVec3 PTempNormal = ((V1 - V0).cross(V2 - V0)).getNormalized();
 		const PxVec3 PLocalTriNormal = TransformNormalToShapeSpace(PTriMeshGeom.scale, PTempNormal);
-	
+
 		// Convert to world space
 		const PxTransform PShapeWorldPose = PxShapeExt::getGlobalPose(*PHit.shape, *PHit.actor);
 		const PxVec3 PWorldTriNormal = PShapeWorldPose.rotate(PLocalTriNormal);
@@ -157,48 +211,35 @@ static bool FindGeomOpposingNormal(const PxLocationHit& PHit, const FVector& Dir
 #endif
 		return true;
 	}
-	else if( PHit.shape->getConvexMeshGeometry(PConvexMeshGeom) && 
-		PConvexMeshGeom.convexMesh != NULL &&
-		PHit.faceIndex < PConvexMeshGeom.convexMesh->getNbPolygons() )
-	{
-		const PxU32 PolyIndex = PHit.faceIndex;
-		PxHullPolygon PPoly;
-		bool bSuccess = PConvexMeshGeom.convexMesh->getPolygonData(PolyIndex, PPoly);
-		if(bSuccess)
-		{
-			// Account for non-uniform scale in local space normal.
-			const PxVec3 PPlaneNormal(PPoly.mPlane[0], PPoly.mPlane[1], PPoly.mPlane[2]);
-			const PxVec3 PLocalPolyNormal = TransformNormalToShapeSpace(PConvexMeshGeom.scale, PPlaneNormal.getNormalized());			
 
-			// Convert to world space
-			const PxTransform PShapeWorldPose = PxShapeExt::getGlobalPose(*PHit.shape, *PHit.actor); 
-			const PxVec3 PWorldPolyNormal = PShapeWorldPose.rotate(PLocalPolyNormal);
-			OutNormal = P2UVector(PWorldPolyNormal);
+	return false;
+}
 
-#if !(UE_BUILD_SHIPPING || UE_BUILD_TEST)
-			if (!OutNormal.IsNormalized())
-			{
-				UE_LOG(LogPhysics, Warning, TEXT("Non-normalized Normal (Hit shape is ConvexMesh): %s (LocalPolyNormal:%s)"), *OutNormal.ToString(), *P2UVector(PLocalPolyNormal).ToString());
-				UE_LOG(LogPhysics, Warning, TEXT("WorldTransform \n: %s"), *P2UTransform(PShapeWorldPose).ToString());
-			}
-#endif
-			return true;
-		}
-	}
-	else if( PHit.shape->getHeightFieldGeometry(PHeightFieldGeom) &&
-		PHeightFieldGeom.heightField != NULL)
+/**
+ * Util to find the normal of the face that we hit.
+ * @param PHit - incoming hit from PhysX
+ * @param Direction - direction of sweep test (not normalized)
+ * @param OutNormal - normal we may recompute based on the faceIndex of the hit
+ * @return true if we compute a new normal for the geometry.
+ */
+static bool FindGeomOpposingNormal(const PxLocationHit& PHit, const FVector& Direction, FVector& OutNormal)
+{
+	checkf(InvalidQueryHit.faceIndex == 0xFFFFffff, TEXT("Engine code needs fixing: PhysX invalid face index sentinel has changed or is not part of default PxQueryHit!"));
+	if (PHit.faceIndex == InvalidQueryHit.faceIndex)
 	{
-		const PxU32 TriIndex = PHit.faceIndex;
-		const PxTransform PShapeWorldPose = PxShapeExt::getGlobalPose(*PHit.shape, *PHit.actor);
-		
-		PxTriangle Tri;
-		PxMeshQuery::getTriangle(PHeightFieldGeom, PShapeWorldPose, TriIndex, Tri);
-
-		PxVec3 TriNormal;
-		Tri.normal(TriNormal);
-		OutNormal = P2UVector(TriNormal);
+		return false;
+	}
 
-		return true;
+	PxGeometryType::Enum GeomType = PHit.shape->getGeometryType();
+	switch (GeomType)
+	{
+		case PxGeometryType::eSPHERE:
+		case PxGeometryType::eBOX:
+		case PxGeometryType::eCAPSULE:		return FindSimpleOpposingNormal(PHit, Direction, OutNormal);
+		case PxGeometryType::eCONVEXMESH:	return FindConvexMeshOpposingNormal(PHit, Direction, OutNormal);
+		case PxGeometryType::eHEIGHTFIELD:	return FindHeightFieldOpposingNormal(PHit, Direction, OutNormal);
+		case PxGeometryType::eTRIANGLEMESH:	return FindTriMeshOpposingNormal(PHit, Direction, OutNormal);
+		default: check(false);	//unsupported geom type
 	}
 
 	return false;
@@ -327,17 +368,11 @@ void ConvertQueryImpactHit(const PxLocationHit& PHit, FHitResult& OutResult, flo
 	CheckHitResultNormal(OutResult, TEXT("Invalid Normal from ConvertQueryImpactHit"), StartLoc, EndLoc, Geom);
 #endif
 
+	PxGeometryType::Enum GeometryType = Geom ? Geom->getType() : PxGeometryType::eINVALID;
+
 	// Special handling for swept-capsule results
-	if(Geom && Geom->getType() == PxGeometryType::eCAPSULE)
-	{
-		// Compute better ImpactNormal. This is the same as Normal except when we hit convex/trimesh, and then its the most opposing normal of the geom at point of impact.
-		const PxCapsuleGeometry* Capsule = static_cast<const PxCapsuleGeometry*>(Geom);
-		FindGeomOpposingNormal(PHit, TraceDir, OutResult.ImpactNormal);
-	}
-	else if (Geom && Geom->getType() == PxGeometryType::eSPHERE)
+	if (GeometryType == PxGeometryType::eCAPSULE || GeometryType == PxGeometryType::eSPHERE)
 	{
-		// Compute better ImpactNormal. This is the same as Normal except when we hit convex/trimesh, and then its the most opposing normal of the geom at point of impact.
-		const PxSphereGeometry* Sphere = static_cast<const PxSphereGeometry*>(Geom);
 		FindGeomOpposingNormal(PHit, TraceDir, OutResult.ImpactNormal);
 	}
 	
@@ -593,28 +628,9 @@ static bool ConvertOverlappedShapeToImpactHit(const PxLocationHit& PHit, const F
 		PxTriangleMeshGeometry PTriMeshGeom;
 		PxHeightFieldGeometry PHeightfieldGeom;
 
-		if (PShape->getTriangleMeshGeometry(PTriMeshGeom))
-		{
-			if (FindGeomOpposingNormal(PHit, TraceDir, OutResult.ImpactNormal))
-			{
-				// Success
-			}
-			else if (ComputeInflatedMTD(MtdInflation, PHit, OutResult, QueryTM, Geom, PShapeWorldPose))
-			{
-				// Success
-			}
-			else
-			{
-				PxU32 HitTris[64];
-				bool bOverflow = false;
-				const int32 NumTrisHit = PxMeshQuery::findOverlapTriangleMesh(Geom, QueryTM, PTriMeshGeom, PShapeWorldPose, HitTris, ARRAY_COUNT(HitTris), 0, bOverflow);
-				if (NumTrisHit > 0)
-				{
-					OutResult.ImpactNormal = FindBestOverlappingNormal(Geom, QueryTM, PTriMeshGeom, PShapeWorldPose, HitTris, NumTrisHit);
-				}
-			}
-		}
-		else if (PShape->getHeightFieldGeometry(PHeightfieldGeom))
+		PxGeometryType::Enum GeometryType = PShape->getGeometryType();
+
+		if (GeometryType == PxGeometryType::eTRIANGLEMESH || GeometryType == PxGeometryType::eHEIGHTFIELD)
 		{
 			if (FindGeomOpposingNormal(PHit, TraceDir, OutResult.ImpactNormal))
 			{
@@ -626,12 +642,22 @@ static bool ConvertOverlappedShapeToImpactHit(const PxLocationHit& PHit, const F
 			}
 			else
 			{
+				const bool bIsTriMesh = GeometryType == PxGeometryType::eTRIANGLEMESH;
 				PxU32 HitTris[64];
 				bool bOverflow = false;
-				const int32 NumTrisHit = PxMeshQuery::findOverlapHeightField(Geom, QueryTM, PHeightfieldGeom, PShapeWorldPose, HitTris, ARRAY_COUNT(HitTris), 0, bOverflow);
+				const int32 NumTrisHit =  bIsTriMesh? PxMeshQuery::findOverlapTriangleMesh(Geom, QueryTM, PTriMeshGeom,     PShapeWorldPose, HitTris, ARRAY_COUNT(HitTris), 0, bOverflow) :
+													  PxMeshQuery::findOverlapHeightField( Geom, QueryTM, PHeightfieldGeom, PShapeWorldPose, HitTris, ARRAY_COUNT(HitTris), 0, bOverflow);
 				if (NumTrisHit > 0)
 				{
-					OutResult.ImpactNormal = FindBestOverlappingNormal(Geom, QueryTM, PHeightfieldGeom, PShapeWorldPose, HitTris, NumTrisHit);
+					if (bIsTriMesh)
+					{
+						OutResult.ImpactNormal = FindBestOverlappingNormal(Geom, QueryTM, PTriMeshGeom, PShapeWorldPose, HitTris, NumTrisHit);
+					}
+					else
+					{
+						OutResult.ImpactNormal = FindBestOverlappingNormal(Geom, QueryTM, PHeightfieldGeom, PShapeWorldPose, HitTris, NumTrisHit);
+					}
+
 				}
 			}
 		}
