Expose PhysicsThrusterComponent to game modules

diff --git a/Engine/Source/Runtime/Engine/Classes/PhysicsEngine/PhysicsThrusterComponent.h b/Engine/Source/Runtime/Engine/Classes/PhysicsEngine/PhysicsThrusterComponent.h
--- a/Engine/Source/Runtime/Engine/Classes/PhysicsEngine/PhysicsThrusterComponent.h
+++ b/Engine/Source/Runtime/Engine/Classes/PhysicsEngine/PhysicsThrusterComponent.h
@@ -8,7 +8,7 @@
  *	Used with objects that have physics to apply a force down the negative-X direction
  *	ie. point X in the direction you want the thrust in.
  */
-UCLASS(hidecategories=(Object, Mobility, LOD), ClassGroup=Physics, showcategories=Trigger, meta=(BlueprintSpawnableComponent))
+UCLASS(hidecategories=(Object, Mobility, LOD), ClassGroup=Physics, showcategories=Trigger, MinimalAPI, meta=(BlueprintSpawnableComponent))
 class UPhysicsThrusterComponent : public USceneComponent
 {
 	GENERATED_UCLASS_BODY()
@@ -21,7 +21,6 @@ class UPhysicsThrusterComponent : public USceneComponent
 	UPROPERTY(BlueprintReadWrite, interp, Category=Physics)
 	float ThrustStrength;
 
-
 	// Begin UActorComponent Interface
 	virtual void TickComponent(float DeltaTime, enum ELevelTick TickType, FActorComponentTickFunction *ThisTickFunction) override;
 	// End UActorComponent Interface
