Fix actors with native root component moving to 0,0,0 when undoing a transaction that has an annotation for any blueprint added component

diff --git a/Engine/Source/Runtime/Engine/Private/ActorConstruction.cpp b/Engine/Source/Runtime/Engine/Private/ActorConstruction.cpp
--- a/Engine/Source/Runtime/Engine/Private/ActorConstruction.cpp
+++ b/Engine/Source/Runtime/Engine/Private/ActorConstruction.cpp
@@ -151,6 +151,8 @@ void AActor::RerunConstructionScripts()
 		AActor* Parent = NULL;
 		USceneComponent* ParentComponent = NULL;
 
+		bool bUseRootComponentProperties = true;
+
 		// Struct to store info about attached actors
 		struct FAttachedActorInfo
 		{
@@ -193,6 +195,8 @@ void AActor::RerunConstructionScripts()
 						AttachedActor->DetachRootComponentFromParent();
 					}
 				}
+
+				bUseRootComponentProperties = false;
 			}
 		}
 		else
@@ -222,28 +226,28 @@ void AActor::RerunConstructionScripts()
 					EachRoot->DetachFromParent(true);
 				}
 			}
+		}
 
-			if (RootComponent != NULL)
+		if (bUseRootComponentProperties && RootComponent != nullptr)
+		{
+			// Do not need to detach if root component is not going away
+			if (RootComponent->AttachParent != NULL && RootComponent->bCreatedByConstructionScript)
 			{
-				// Do not need to detach if root component is not going away
-				if (RootComponent->AttachParent != NULL && RootComponent->bCreatedByConstructionScript)
+				Parent = RootComponent->AttachParent->GetOwner();
+				// Root component should never be attached to another component in the same actor!
+				if (Parent == this)
 				{
-					Parent = RootComponent->AttachParent->GetOwner();
-					// Root component should never be attached to another component in the same actor!
-					if (Parent == this)
-					{
-						UE_LOG(LogActor, Warning, TEXT("RerunConstructionScripts: RootComponent (%s) attached to another component in this Actor (%s)."), *RootComponent->GetPathName(), *Parent->GetPathName());
-						Parent = NULL;
-					}
-					ParentComponent = RootComponent->AttachParent;
-					SocketName = RootComponent->AttachSocketName;
-					//detach it to remove any scaling 
-					RootComponent->DetachFromParent(true);
+					UE_LOG(LogActor, Warning, TEXT("RerunConstructionScripts: RootComponent (%s) attached to another component in this Actor (%s)."), *RootComponent->GetPathName(), *Parent->GetPathName());
+					Parent = NULL;
 				}
-
-				OldTransform = RootComponent->ComponentToWorld;
-				OldTransform.SetTranslation(RootComponent->GetComponentLocation()); // take into account any custom location
+				ParentComponent = RootComponent->AttachParent;
+				SocketName = RootComponent->AttachSocketName;
+				//detach it to remove any scaling 
+				RootComponent->DetachFromParent(true);
 			}
+
+			OldTransform = RootComponent->ComponentToWorld;
+			OldTransform.SetTranslation(RootComponent->GetComponentLocation()); // take into account any custom location
 		}
 
 		// Destroy existing components
