Added ability to ignore mismatched package versions when referring to objects over the network

diff --git a/Engine/Source/Runtime/Engine/Private/PackageMapClient.cpp b/Engine/Source/Runtime/Engine/Private/PackageMapClient.cpp
--- a/Engine/Source/Runtime/Engine/Private/PackageMapClient.cpp
+++ b/Engine/Source/Runtime/Engine/Private/PackageMapClient.cpp
@@ -23,6 +23,7 @@ static const int INTERNAL_LOAD_OBJECT_RECURSION_LIMIT = 16;
 
 static TAutoConsoleVariable<int32> CVarAllowAsyncLoading( TEXT( "net.AllowAsyncLoading" ), 0, TEXT( "Allow async loading" ) );
 static TAutoConsoleVariable<int32> CVarSimulateAsyncLoading( TEXT( "net.SimulateAsyncLoading" ), 0, TEXT( "Simulate async loading" ) );
+static TAutoConsoleVariable<int32> CVarIgnorePackageMismatch( TEXT( "net.IgnorePackageMismatch" ), 0, TEXT( "Ignore when package versions are different" ) );
 
 /*-----------------------------------------------------------------------------
 	UPackageMapClient implementation.
@@ -282,7 +283,7 @@ bool UPackageMapClient::SerializeNewActor(FArchive& Ar, class UActorChannel *Cha
 				}
 				else
 				{
-					UE_LOG(LogNetPackageMap, Warning, TEXT("UPackageMapClient::SerializeNewActor Unable to read Archetype for NetGUID %s / %s"), *NetGUID.ToString(), *ArchetypeNetGUID.ToString() );
+					UE_LOG(LogNetPackageMap, Error, TEXT("UPackageMapClient::SerializeNewActor Unable to read Archetype for NetGUID %s / %s"), *NetGUID.ToString(), *ArchetypeNetGUID.ToString() );
 					ensure(Archetype);
 				}
 			}
@@ -628,7 +629,7 @@ FNetworkGUID UPackageMapClient::InternalLoadObject( FArchive & Ar, UObject *& Ob
 					return NetGUID;
 				}
 
-				if ( Package->GetGuid() != PackageGuid )
+				if ( Package->GetGuid() != PackageGuid && CVarIgnorePackageMismatch.GetValueOnGameThread() == 0 )
 				{
 					UE_LOG( LogNetPackageMap, Error, TEXT( "UPackageMapClient::InternalLoadObject: Default object package guid mismatch! PathName: %s, ObjOuter: %s, GUID1: %s, GUID2: %s " ), *PathName, ObjOuter != NULL ? *ObjOuter->GetPathName() : TEXT( "NULL" ), *Package->GetGuid().ToString(), *PackageGuid.ToString() );
 					Object = NULL;
@@ -1663,7 +1664,7 @@ void FNetGUIDCache::AsyncPackageCallback( const FString & PackageName, UPackage
 		CacheObject->bIsBroken = true;
 		UE_LOG( LogNetPackageMap, Error, TEXT( "AsyncPackageCallback: Package FAILED to load. Path: %s, NetGUID: %s" ), *PackageName, *NetGUID.ToString() );
 	}
-	else if ( Package->GetGuid() != CacheObject->PackageGuid )
+	else if ( Package->GetGuid() != CacheObject->PackageGuid && CVarIgnorePackageMismatch.GetValueOnGameThread() == 0 )
 	{
 		CacheObject->bIsBroken = true;
 		UE_LOG( LogNetPackageMap, Error, TEXT( "AsyncPackageCallback: Package GUID mismatch! Path: %s, NetGUID: %s, GUID1: %s, GUID2: %s" ), *PackageName, *NetGUID.ToString(), *Package->GetGuid().ToString(), *CacheObject->PackageGuid.ToString() );
@@ -1868,7 +1869,7 @@ UObject * FNetGUIDCache::GetObjectFromNetGUID( const FNetworkGUID & NetGUID, con
 			return NULL;
 		}
 
-		if ( Package->GetGuid() != CacheObjectPtr->PackageGuid )
+		if ( Package->GetGuid() != CacheObjectPtr->PackageGuid && CVarIgnorePackageMismatch.GetValueOnGameThread() == 0 )
 		{
 			// If the package guid doesn't match, don't allow it to load
 			CacheObjectPtr->bIsBroken = true;
