nicer UIf ro PostProcess settings

diff --git a/Engine/Source/Runtime/Engine/Classes/Engine/Scene.h b/Engine/Source/Runtime/Engine/Classes/Engine/Scene.h
--- a/Engine/Source/Runtime/Engine/Classes/Engine/Scene.h
+++ b/Engine/Source/Runtime/Engine/Classes/Engine/Scene.h
@@ -340,9 +340,9 @@ struct FPostProcessSettings
 	// -----------------------------------------------------------------------
 
 
-	UPROPERTY(interp, BlueprintReadWrite, Category=Film, meta=(editcondition = "bOverride_FilmWhitePoint", DisplayName = "Tint"))
+	UPROPERTY(interp, BlueprintReadWrite, Category=Film, meta=(editcondition = "bOverride_FilmWhitePoint", DisplayName = "Tint", HideAlphaChannel))
 	FLinearColor FilmWhitePoint;
-	UPROPERTY(interp, BlueprintReadWrite, Category=Film, AdvancedDisplay, meta=(editcondition = "bOverride_FilmShadowTint", DisplayName = "Tint Shadow"))
+	UPROPERTY(interp, BlueprintReadWrite, Category=Film, AdvancedDisplay, meta=(editcondition = "bOverride_FilmShadowTint", DisplayName = "Tint Shadow", HideAlphaChannel))
 	FLinearColor FilmShadowTint;
 	UPROPERTY(interp, BlueprintReadWrite, Category=Film, AdvancedDisplay, meta=(UIMin = "0.0", UIMax = "1.0", editcondition = "bOverride_FilmShadowTintBlend", DisplayName = "Tint Shadow Blend"))
 	float FilmShadowTintBlend;
@@ -351,11 +351,11 @@ struct FPostProcessSettings
 
 	UPROPERTY(interp, BlueprintReadWrite, Category=Film, meta=(UIMin = "0.0", UIMax = "2.0", editcondition = "bOverride_FilmSaturation", DisplayName = "Saturation"))
 	float FilmSaturation;
-	UPROPERTY(interp, BlueprintReadWrite, Category=Film, AdvancedDisplay, meta=(editcondition = "bOverride_FilmChannelMixerRed", DisplayName = "Channel Mixer Red"))
+	UPROPERTY(interp, BlueprintReadWrite, Category=Film, AdvancedDisplay, meta=(editcondition = "bOverride_FilmChannelMixerRed", DisplayName = "Channel Mixer Red", HideAlphaChannel))
 	FLinearColor FilmChannelMixerRed;
-	UPROPERTY(interp, BlueprintReadWrite, Category=Film, AdvancedDisplay, meta=(editcondition = "bOverride_FilmChannelMixerGreen", DisplayName = "Channel Mixer Green"))
+	UPROPERTY(interp, BlueprintReadWrite, Category=Film, AdvancedDisplay, meta=(editcondition = "bOverride_FilmChannelMixerGreen", DisplayName = "Channel Mixer Green", HideAlphaChannel))
 	FLinearColor FilmChannelMixerGreen;
-	UPROPERTY(interp, BlueprintReadWrite, Category=Film, AdvancedDisplay, meta=(editcondition = "bOverride_FilmChannelMixerBlue", DisplayName = " Channel Mixer Blue"))
+	UPROPERTY(interp, BlueprintReadWrite, Category=Film, AdvancedDisplay, meta=(editcondition = "bOverride_FilmChannelMixerBlue", DisplayName = " Channel Mixer Blue", HideAlphaChannel))
 	FLinearColor FilmChannelMixerBlue;
 
 	UPROPERTY(interp, BlueprintReadWrite, Category=Film, meta=(UIMin = "0.0", UIMax = "1.0", editcondition = "bOverride_FilmContrast", DisplayName = "Contrast"))
@@ -449,7 +449,7 @@ struct FPostProcessSettings
 	float BloomDirtMaskIntensity;
 
 	/** BloomDirtMask tint color */
-	UPROPERTY(interp, BlueprintReadWrite, Category=Bloom, AdvancedDisplay, meta=(editcondition = "bOverride_BloomDirtMaskTint", DisplayName = "Dirt Mask Tint"))
+	UPROPERTY(interp, BlueprintReadWrite, Category=Bloom, AdvancedDisplay, meta=(editcondition = "bOverride_BloomDirtMaskTint", DisplayName = "Dirt Mask Tint", HideAlphaChannel))
 	FLinearColor BloomDirtMaskTint;
 
 	/**
@@ -494,7 +494,7 @@ struct FPostProcessSettings
 	float LPVTransmissionIntensity;
 
 	/** AmbientCubemap tint color */
-	UPROPERTY(interp, BlueprintReadWrite, Category=AmbientCubemap, AdvancedDisplay, meta=(editcondition = "bOverride_AmbientCubemapTint", DisplayName = "Tint"))
+	UPROPERTY(interp, BlueprintReadWrite, Category=AmbientCubemap, AdvancedDisplay, meta=(editcondition = "bOverride_AmbientCubemapTint", DisplayName = "Tint", HideAlphaChannel))
 	FLinearColor AmbientCubemapTint;
 	/**
 	 * To scale the Ambient cubemap brightness
@@ -574,7 +574,7 @@ struct FPostProcessSettings
 	float LensFlareIntensity;
 
 	/** Tint color for the image based lens flares. */
-	UPROPERTY(interp, BlueprintReadWrite, Category=LensFlares, AdvancedDisplay, meta=(editcondition = "bOverride_LensFlareTint", DisplayName = "Tint"))
+	UPROPERTY(interp, BlueprintReadWrite, Category=LensFlares, AdvancedDisplay, meta=(editcondition = "bOverride_LensFlareTint", DisplayName = "Tint", HideAlphaChannel))
 	FLinearColor LensFlareTint;
 
 	/** Size of the Lens Blur (in percent of the view width) that is done with the Bokeh texture (note: performance cost is radius*radius) */
@@ -598,7 +598,7 @@ struct FPostProcessSettings
 	float VignetteIntensity;
 
 	/** Vignette color. */
-	UPROPERTY(interp, BlueprintReadWrite, Category=SceneColor, AdvancedDisplay, meta=(editcondition = "bOverride_VignetteColor", DisplayName = "Vignette Color"))
+	UPROPERTY(interp, BlueprintReadWrite, Category=SceneColor, AdvancedDisplay, meta=(editcondition = "bOverride_VignetteColor", DisplayName = "Vignette Color", HideAlphaChannel))
 	FLinearColor VignetteColor;
 
 	/** 0..1 grain jitter */
@@ -662,7 +662,7 @@ struct FPostProcessSettings
 	float AmbientOcclusionMipThreshold;
 
 	/** Adjusts indirect lighting color. (1,1,1) is default. (0,0,0) to disable GI. The show flag 'Global Illumination' must be enabled to use this property. */
-	UPROPERTY(interp, BlueprintReadWrite, Category=GlobalIllumination, meta=(editcondition = "bOverride_IndirectLightingColor", DisplayName = "Indirect Lighting Color"))
+	UPROPERTY(interp, BlueprintReadWrite, Category=GlobalIllumination, meta=(editcondition = "bOverride_IndirectLightingColor", DisplayName = "Indirect Lighting Color", HideAlphaChannel))
 	FLinearColor IndirectLightingColor;
 
 	/** Scales the indirect lighting contribution. A value of 0 disables GI. Default is 1. The show flag 'Global Illumination' must be enabled to use this property. */
