minor optimization - don't need to set the RefractionBias on materials that don't have translucency

diff --git a/Engine/Source/Runtime/Engine/Private/Materials/HLSLMaterialTranslator.h b/Engine/Source/Runtime/Engine/Private/Materials/HLSLMaterialTranslator.h
--- a/Engine/Source/Runtime/Engine/Private/Materials/HLSLMaterialTranslator.h
+++ b/Engine/Source/Runtime/Engine/Private/Materials/HLSLMaterialTranslator.h
@@ -359,6 +359,7 @@ public:
 			Chunk[MP_ClearCoatRoughness]			= Material->CompilePropertyAndSetMaterialProperty(MP_ClearCoatRoughness    ,this);
 			Chunk[MP_AmbientOcclusion]				= Material->CompilePropertyAndSetMaterialProperty(MP_AmbientOcclusion      ,this);
 
+			if(IsTranslucentBlendMode(Material->GetBlendMode()))
 			{
 				int32 UserRefraction = ForceCast(Material->CompilePropertyAndSetMaterialProperty(MP_Refraction, this), MCT_Float2);
 				int32 RefractionDepthBias = ForceCast(ScalarParameter(FName(TEXT("RefractionDepthBias")), Material->GetRefractionDepthBiasValue()), MCT_Float1);
