added parameters to postprocess materila tonemapper pass, more can be exposed later
 vector: Engine.FilmWhitePoint
 scalar: Engine.FilmSaturation, Engine.FilmContrast

diff --git a/Engine/Shaders/ShaderVersion.usf b/Engine/Shaders/ShaderVersion.usf
--- a/Engine/Shaders/ShaderVersion.usf
+++ b/Engine/Shaders/ShaderVersion.usf
@@ -2,4 +2,4 @@
 // Each time the console command is executed it generates a new GUID. As this file is included
 // in common.usf (which should be included in any shader) it allows to invalidate the shader DDC.
 // 
-// GUID = 1B1A4F2441A3CDA2433BB820BFD049B1
\ No newline at end of file
+// GUID = 55058CB5415760DB2475E8ABED874656
\ No newline at end of file
diff --git a/Engine/Source/Runtime/Engine/Classes/Engine/BlendableInterface.h b/Engine/Source/Runtime/Engine/Classes/Engine/BlendableInterface.h
--- a/Engine/Source/Runtime/Engine/Classes/Engine/BlendableInterface.h
+++ b/Engine/Source/Runtime/Engine/Classes/Engine/BlendableInterface.h
@@ -16,7 +16,9 @@ enum EBlendableLocation
 	BL_BeforeTonemapping UMETA(DisplayName="Before Tonemapping"),
 	// Input0:former pass color, Input1:SeparateTranslucency
 	BL_BeforeTranslucency UMETA(DisplayName="Before Translucency"),
-	// Input0:former pass color, Input1:SeparateTranslucency, Input2: BloomOutput, Input3: EyeAdaptation
+	// Input0:former pass color, Input1:SeparateTranslucency, Input2: BloomOutput
+	// vector parameters: Engine.FilmWhitePoint
+	// scalar parameters: Engine.FilmSaturation, Engine.FilmContrast
 	BL_ReplacingTonemapper UMETA(DisplayName="Replacing the Tonemapper"),
 //	BL_AfterOpaque,
 //	BL_AfterFog,
diff --git a/Engine/Source/Runtime/Engine/Classes/Materials/MaterialInstance.h b/Engine/Source/Runtime/Engine/Classes/Materials/MaterialInstance.h
--- a/Engine/Source/Runtime/Engine/Classes/Materials/MaterialInstance.h
+++ b/Engine/Source/Runtime/Engine/Classes/Materials/MaterialInstance.h
@@ -383,6 +383,8 @@ protected:
 	// to share code between PostLoad() and PostEditChangeProperty()
 	void PropagateDataToMaterialProxy();
 
+	void SetTonemapperPostprocessMaterialSettings(class FSceneView& View, UMaterialInstanceDynamic& MID) const;
+
 	/** Allow resource to access private members. */
 	friend class FMaterialInstanceResource;
 	/** Editor-only access to private members. */
diff --git a/Engine/Source/Runtime/Engine/Private/Materials/MaterialInstance.cpp b/Engine/Source/Runtime/Engine/Private/Materials/MaterialInstance.cpp
--- a/Engine/Source/Runtime/Engine/Private/Materials/MaterialInstance.cpp
+++ b/Engine/Source/Runtime/Engine/Private/Materials/MaterialInstance.cpp
@@ -942,6 +942,8 @@ void UMaterialInstance::CopyMaterialInstanceParameters(UMaterialInterface* Mater
 {
 	if(MaterialInterface)
 	{
+		UMaterial* Material = GetMaterial();
+
 		//first, clear out all the parameter values
 		ClearParameterValuesInternal();
 
@@ -950,8 +952,8 @@ void UMaterialInstance::CopyMaterialInstanceParameters(UMaterialInterface* Mater
 		TArray<FGuid> Guids;
 
 		//handle all the fonts
-		GetMaterial()->GetAllFontParameterNames(Names, Guids);
-		for(int32 i = 0; i < Names.Num(); ++i)
+		Material->GetAllFontParameterNames(Names, Guids);
+		for(int32 i = 0, Size = Names.Num(); i < Size; ++i)
 		{
 			FName ParameterName = Names[i];
 			UFont* FontValue = NULL;
@@ -970,8 +972,8 @@ void UMaterialInstance::CopyMaterialInstanceParameters(UMaterialInterface* Mater
 		//now do the scalar params
 		Names.Reset();
 		Guids.Reset();
-		GetMaterial()->GetAllScalarParameterNames(Names, Guids);
-		for(int32 i = 0; i < Names.Num(); ++i)
+		Material->GetAllScalarParameterNames(Names, Guids);
+		for(int32 i = 0, Size = Names.Num(); i < Size; ++i)
 		{
 			FName ParameterName = Names[i];
 			float ScalarValue = 1.0f;
@@ -987,8 +989,8 @@ void UMaterialInstance::CopyMaterialInstanceParameters(UMaterialInterface* Mater
 		//now do the vector params
 		Names.Reset();
 		Guids.Reset();
-		GetMaterial()->GetAllVectorParameterNames(Names, Guids);
-		for(int32 i = 0; i < Names.Num(); ++i)
+		Material->GetAllVectorParameterNames(Names, Guids);
+		for(int32 i = 0, Size = Names.Num(); i < Size; ++i)
 		{
 			FName ParameterName = Names[i];
 			FLinearColor VectorValue;
@@ -1004,8 +1006,8 @@ void UMaterialInstance::CopyMaterialInstanceParameters(UMaterialInterface* Mater
 		//now do the texture params
 		Names.Reset();
 		Guids.Reset();
-		GetMaterial()->GetAllTextureParameterNames(Names, Guids);
-		for(int32 i = 0; i < Names.Num(); ++i)
+		Material->GetAllTextureParameterNames(Names, Guids);
+		for(int32 i = 0, Size = Names.Num(); i < Size; ++i)
 		{
 			FName ParameterName = Names[i];
 			UTexture* TextureValue = NULL;
@@ -2326,14 +2328,16 @@ void UMaterialInstance::OverrideBlendableSettings(class FSceneView& View, float
 
 		// a material already exists, blend with existing ones
 		DestMID->K2_InterpolateMaterialInstanceParams(DestMID, SrcMID, Weight);
+		SetTonemapperPostprocessMaterialSettings(View, *DestMID);
 	}
 	else
 	{
-		UMaterialInstanceDynamic* MID = View.State->GetReusableMID((UMaterialInterface*)Material);//, (UMaterialInterface*)this);
+		UMaterialInstanceDynamic* MID = View.State->GetReusableMID((UMaterialInterface*)Material);
 
 		if(MID)
 		{
 			MID->K2_CopyMaterialInstanceParameters((UMaterialInterface*)this);
+			SetTonemapperPostprocessMaterialSettings(View, *MID);
 
 			FPostProcessMaterialNode NewNode(MID, Material->BlendableLocation, Material->BlendablePriority);
 
@@ -2343,6 +2347,29 @@ void UMaterialInstance::OverrideBlendableSettings(class FSceneView& View, float
 	}
 }
 
+void UMaterialInstance::SetTonemapperPostprocessMaterialSettings(FSceneView& View, UMaterialInstanceDynamic& MID) const
+{
+	UMaterial* Material = Parent->GetMaterial();
+
+	if(Material->BlendableLocation == BL_ReplacingTonemapper)
+	{
+		// can expose more if needed
+		{
+			static FName Name("Engine.FilmWhitePoint");
+			MID.SetVectorParameterValue(Name, View.FinalPostProcessSettings.FilmWhitePoint);
+		}
+		{
+			static FName Name("Engine.FilmSaturation");
+			MID.SetScalarParameterValue(Name, View.FinalPostProcessSettings.FilmSaturation);
+		}
+		{
+			static FName Name("Engine.FilmContrast");
+			MID.SetScalarParameterValue(Name, View.FinalPostProcessSettings.FilmContrast);
+		}
+	}
+}
+
+
 void UMaterialInstance::AllMaterialsCacheResourceShadersForRendering()
 {
 	for (TObjectIterator<UMaterialInstance> It; It; ++It)
diff --git a/Engine/Source/Runtime/Renderer/Private/PostProcess/PostProcessMaterial.h b/Engine/Source/Runtime/Renderer/Private/PostProcess/PostProcessMaterial.h
--- a/Engine/Source/Runtime/Renderer/Private/PostProcess/PostProcessMaterial.h
+++ b/Engine/Source/Runtime/Renderer/Private/PostProcess/PostProcessMaterial.h
@@ -11,9 +11,8 @@
 // derives from TRenderingCompositePassBase<InputCount, OutputCount>
 // ePId_Input0: former pass
 // ePId_Input1: optional, depends on EBlendableLocation
-// ePId_Input2: optional, depends on EBlendableLocation 
-// ePId_Input3: optional, depends on EBlendableLocation
-class FRCPassPostProcessMaterial : public TRenderingCompositePassBase<4,1>
+// ePId_Input2: optional, depends on EBlendableLocation
+class FRCPassPostProcessMaterial : public TRenderingCompositePassBase<3,1>
 {
 public:
 	// constructor
diff --git a/Engine/Source/Runtime/Renderer/Private/PostProcess/PostProcessing.cpp b/Engine/Source/Runtime/Renderer/Private/PostProcess/PostProcessing.cpp
--- a/Engine/Source/Runtime/Renderer/Private/PostProcess/PostProcessing.cpp
+++ b/Engine/Source/Runtime/Renderer/Private/PostProcess/PostProcessing.cpp
@@ -1117,7 +1117,6 @@ void FPostProcessing::Process(FRHICommandListImmediate& RHICmdList, FViewInfo& V
 					// TODO: Check if material actually uses this texture and only bind if needed.
 					Node->SetInput(ePId_Input1, SeparateTranslucency);
 					Node->SetInput(ePId_Input2, BloomOutputCombined);
-					Node->SetInput(ePId_Input3, EyeAdaptation);
 					Context.FinalOutput = Node;
 				}
 				else
