Allow cvar to control demo record rate

diff --git a/Engine/Source/Runtime/Engine/Private/DemoNetDriver.cpp b/Engine/Source/Runtime/Engine/Private/DemoNetDriver.cpp
--- a/Engine/Source/Runtime/Engine/Private/DemoNetDriver.cpp
+++ b/Engine/Source/Runtime/Engine/Private/DemoNetDriver.cpp
@@ -13,6 +13,8 @@
 
 DEFINE_LOG_CATEGORY_STATIC( LogDemo, Log, All );
 
+static TAutoConsoleVariable<float> CVarDemoRecordHz( TEXT( "demo.RecordHz" ), 10, TEXT( "Number of demo frames recorded per second" ) );
+
 /*-----------------------------------------------------------------------------
 	UDemoNetDriver.
 -----------------------------------------------------------------------------*/
@@ -311,8 +313,6 @@ void UDemoNetDriver::StopDemo()
 Demo Recording tick.
 -----------------------------------------------------------------------------*/
 
-//static int GNumReplicatedActors = 0;
-
 static void DemoReplicateActor(AActor* Actor, UNetConnection* Connection, bool IsNetClient)
 {
 	// All actors marked for replication are assumed to be relevant for demo recording.
@@ -348,7 +348,6 @@ static void DemoReplicateActor(AActor* Actor, UNetConnection* Connection, bool I
 			if ( Channel->IsNetReady( 0 ) )
 			{
 				Channel->ReplicateActor();
-				//GNumReplicatedActors++;
 			}
 		}
 	}
@@ -370,7 +369,8 @@ void UDemoNetDriver::TickDemoRecord( float DeltaSeconds )
 
 	const double CurrentSeconds = FPlatformTime::Seconds();
 
-	const double RECORD_DELAY = 1.0 / 10.0;
+	const double RECORD_HZ		= CVarDemoRecordHz.GetValueOnGameThread();
+	const double RECORD_DELAY	= 1.0 / RECORD_HZ;
 
 	if ( CurrentSeconds - LastRecordTime < RECORD_DELAY )
 	{
@@ -379,8 +379,6 @@ void UDemoNetDriver::TickDemoRecord( float DeltaSeconds )
 
 	LastRecordTime = CurrentSeconds;
 
-	//GNumReplicatedActors = 0;
-
 	// Save out a frame
 	DemoFrameNum++;
 	ReplicationFrame++;
@@ -414,8 +412,6 @@ void UDemoNetDriver::TickDemoRecord( float DeltaSeconds )
 	int32 EndCount = 0;
 
 	*FileAr << EndCount;
-
-	//UE_LOG( LogDemo, Log, TEXT( "Num replicated actors: %i, %i"), GNumReplicatedActors, ReplicationFrame );
 }
 
 bool UDemoNetDriver::ReadDemoFrame()
