More error handling during demo playback failure cases ability to do checksums in various places

diff --git a/Engine/Source/Runtime/Engine/Private/DemoNetDriver.cpp b/Engine/Source/Runtime/Engine/Private/DemoNetDriver.cpp
--- a/Engine/Source/Runtime/Engine/Private/DemoNetDriver.cpp
+++ b/Engine/Source/Runtime/Engine/Private/DemoNetDriver.cpp
@@ -15,6 +15,10 @@ DEFINE_LOG_CATEGORY_STATIC( LogDemo, Log, All );
 
 static TAutoConsoleVariable<float> CVarDemoRecordHz( TEXT( "demo.RecordHz" ), 10, TEXT( "Number of demo frames recorded per second" ) );
 
+static const int32 MAX_DEMO_READ_WRITE_BUFFER = 1024 * 2;
+
+#define DEMO_CHECKSUMS 0
+
 /*-----------------------------------------------------------------------------
 	UDemoNetDriver.
 -----------------------------------------------------------------------------*/
@@ -245,7 +249,7 @@ void UDemoNetDriver::StopDemo()
 {
 	if ( !ServerConnection && ClientConnections.Num() == 0 )
 	{
-		UE_LOG( LogDemo, Log, TEXT( "StopDemo: No demo is playing" ) );
+		UE_LOG( LogDemo, Warning, TEXT( "StopDemo: No demo is playing" ) );
 		return;
 	}
 
@@ -386,6 +390,11 @@ void UDemoNetDriver::TickDemoRecord( float DeltaSeconds )
 	// Save elapsed game time
 	*FileAr << DemoDeltaTime;
 
+#if DEMO_CHECKSUMS == 1
+	uint32 DeltaTimeChecksum = FCrc::MemCrc32( &DemoDeltaTime, sizeof( DemoDeltaTime ), 0 );
+	*FileAr << DeltaTimeChecksum;
+#endif
+
 	DemoDeltaTime = 0;
 
 	// Make sure we don't have anything in the buffer for this new frame
@@ -429,9 +438,25 @@ bool UDemoNetDriver::ReadDemoFrame()
 	// Peek at the next demo delta time, and see if we should process this frame
 	*FileAr << ServerDeltaTime;
 
+#if DEMO_CHECKSUMS == 1
+	{
+		uint32 ServerDeltaTimeCheksum = 0;
+		*FileAr << ServerDeltaTimeCheksum;
+
+		const uint32 DeltaTimeChecksum = FCrc::MemCrc32( &ServerDeltaTime, sizeof( ServerDeltaTime ), 0 );
+
+		if ( DeltaTimeChecksum != ServerDeltaTimeCheksum )
+		{
+			UE_LOG( LogDemo, Error, TEXT( "UDemoNetDriver::ReadDemoFrame: DeltaTimeChecksum != ServerDeltaTimeCheksum" ) );
+			StopDemo();
+			return false;
+		}
+	}
+#endif
+
 	if ( FileAr->IsError() )
 	{
-		UE_LOG( LogDemo, Error, TEXT( "Failed to read demo ServerDeltaTime" ) );
+		UE_LOG( LogDemo, Error, TEXT( "UDemoNetDriver::ReadDemoFrame: Failed to read demo ServerDeltaTime" ) );
 		StopDemo();
 		return false;
 	}
@@ -447,7 +472,7 @@ bool UDemoNetDriver::ReadDemoFrame()
 
 	while ( true )
 	{
-		uint8 Data[ 512 + 8 ];
+		uint8 ReadBuffer[ MAX_DEMO_READ_WRITE_BUFFER ];
 
 		int32 PacketBytes;
 
@@ -455,7 +480,7 @@ bool UDemoNetDriver::ReadDemoFrame()
 
 		if ( FileAr->IsError() )
 		{
-			UE_LOG( LogDemo, Error, TEXT( "Failed to read demo PacketBytes" ) );
+			UE_LOG( LogDemo, Error, TEXT( "UDemoNetDriver::ReadDemoFrame: Failed to read demo PacketBytes" ) );
 			StopDemo();
 			return false;
 		}
@@ -465,22 +490,46 @@ bool UDemoNetDriver::ReadDemoFrame()
 			break;
 		}
 
+		if ( PacketBytes > sizeof( ReadBuffer ) )
+		{
+			UE_LOG( LogDemo, Error, TEXT( "UDemoNetDriver::ReadDemoFrame: PacketBytes > sizeof( ReadBuffer )" ) );
+			StopDemo();
+			return false;
+		}
+
 		// Read data from file.
-		FileAr->Serialize( Data, PacketBytes );
+		FileAr->Serialize( ReadBuffer, PacketBytes );
 
 		if ( FileAr->IsError() )
 		{
-			UE_LOG( LogDemo, Error, TEXT( "Failed to read demo file packet" ) );
+			UE_LOG( LogDemo, Error, TEXT( "UDemoNetDriver::ReadDemoFrame: Failed to read demo file packet" ) );
 			StopDemo();
 			return false;
 		}
 
+#if DEMO_CHECKSUMS == 1
+		{
+			uint32 ServerChecksum = 0;
+			*FileAr << ServerChecksum;
+
+			const uint32 Checksum = FCrc::MemCrc32( ReadBuffer, PacketBytes, 0 );
+
+			if ( Checksum != ServerChecksum )
+			{
+				UE_LOG( LogDemo, Error, TEXT( "UDemoNetDriver::ReadDemoFrame: Checksum != ServerChecksum" ) );
+				StopDemo();
+				return false;
+			}
+		}
+#endif
+
 		// Process incoming packet.
-		ServerConnection->ReceivedRawPacket( Data, PacketBytes );
+		ServerConnection->ReceivedRawPacket( ReadBuffer, PacketBytes );
 
 		if ( ServerConnection == NULL || ServerConnection->State == USOCK_Closed )
 		{
-			// something we received resulted in the demo being stopped
+			// Something we received resulted in the demo being stopped
+			UE_LOG( LogDemo, Error, TEXT( "UDemoNetDriver::ReadDemoFrame: ReceivedRawPacket closed connection" ) );
 			StopDemo();
 			return false;
 		}
@@ -501,11 +550,13 @@ void UDemoNetDriver::TickDemoPlayback( float DeltaSeconds )
 
 	while ( true )
 	{
-		// Read demo frames until we are cought up
+		// Read demo frames until we are caught up
 		if ( !ReadDemoFrame() )
 		{
 			break;
 		}
+
+		DemoFrameNum++;
 	}
 }
 
@@ -565,10 +616,20 @@ void UDemoNetConnection::LowLevelSend( void* Data, int32 Count )
 		return;
 	}
 
+	if ( Count > MAX_DEMO_READ_WRITE_BUFFER )
+	{
+		UE_LOG( LogDemo, Fatal, TEXT( "UDemoNetConnection::LowLevelSend: Count > MAX_DEMO_READ_WRITE_BUFFER." ) );
+	}
+
 	if ( !GetDriver()->ServerConnection && GetDriver()->FileAr )
 	{
 		*GetDriver()->FileAr << Count;
 		GetDriver()->FileAr->Serialize( Data, Count );
+		
+#if DEMO_CHECKSUMS == 1
+		uint32 Checksum = FCrc::MemCrc32( Data, Count, 0 );
+		*GetDriver()->FileAr << Checksum;
+#endif
 	}
 }
 
