Queue packets that were written outside of a demo frame so we can save them within the frame

diff --git a/Engine/Source/Runtime/Engine/Classes/Engine/DemoNetConnection.h b/Engine/Source/Runtime/Engine/Classes/Engine/DemoNetConnection.h
--- a/Engine/Source/Runtime/Engine/Classes/Engine/DemoNetConnection.h
+++ b/Engine/Source/Runtime/Engine/Classes/Engine/DemoNetConnection.h
@@ -7,6 +7,11 @@
 #pragma once
 #include "DemoNetConnection.generated.h"
 
+struct FQueuedDemoPacket
+{
+	TArray< uint8 > Data;
+};
+
 UCLASS(transient, config=Engine)
 class UDemoNetConnection : public UNetConnection
 {
@@ -25,6 +30,8 @@ class UDemoNetConnection : public UNetConnection
 
 	/** @return The DemoRecording driver object */
 	FORCEINLINE class UDemoNetDriver* GetDriver() { return (UDemoNetDriver*)Driver; }
+
+	TArray< FQueuedDemoPacket > QueuedDemoPackets;
 };
 
 
diff --git a/Engine/Source/Runtime/Engine/Classes/Engine/DemoNetDriver.h b/Engine/Source/Runtime/Engine/Classes/Engine/DemoNetDriver.h
--- a/Engine/Source/Runtime/Engine/Classes/Engine/DemoNetDriver.h
+++ b/Engine/Source/Runtime/Engine/Classes/Engine/DemoNetDriver.h
@@ -30,6 +30,9 @@ class UDemoNetDriver : public UNetDriver
 	/** during playback, set to total number of frames recorded in the demo */
 	int32			PlaybackTotalFrames;
 
+	/** True if we're in the middle of recording a frame */
+	bool			bIsRecordingDemoFrame;
+
 	UPROPERTY( config )
 	FString			DemoSpectatorClass;
 
diff --git a/Engine/Source/Runtime/Engine/Private/DemoNetDriver.cpp b/Engine/Source/Runtime/Engine/Private/DemoNetDriver.cpp
--- a/Engine/Source/Runtime/Engine/Private/DemoNetDriver.cpp
+++ b/Engine/Source/Runtime/Engine/Private/DemoNetDriver.cpp
@@ -32,9 +32,10 @@ bool UDemoNetDriver::InitBase( bool bInitAsClient, FNetworkNotify* InNotify, con
 {
 	if ( Super::InitBase( bInitAsClient, InNotify, URL, bReuseAddressAndPort, Error ) )
 	{
-		DemoFilename	= URL.Map;
-		Time			= 0;
-		DemoFrameNum	= 0;
+		DemoFilename			= URL.Map;
+		Time					= 0;
+		DemoFrameNum			= 0;
+		bIsRecordingDemoFrame	= false;
 
 		return true;
 	}
@@ -400,6 +401,18 @@ void UDemoNetDriver::TickDemoRecord( float DeltaSeconds )
 	// Make sure we don't have anything in the buffer for this new frame
 	check( ClientConnections[0]->SendBuffer.GetNumBits() == 0 );
 
+	bIsRecordingDemoFrame = true;
+
+	// Dump any queued packets
+	UDemoNetConnection * ClientDemoConnection = CastChecked< UDemoNetConnection >( ClientConnections[0] );
+
+	for ( int32 i = 0; i < ClientDemoConnection->QueuedDemoPackets.Num(); i++ )
+	{
+		ClientDemoConnection->LowLevelSend( (char*)&ClientDemoConnection->QueuedDemoPackets[i].Data[0], ClientDemoConnection->QueuedDemoPackets[i].Data.Num() );
+	}
+
+	ClientDemoConnection->QueuedDemoPackets.Empty();
+
 	const bool IsNetClient = ( GetWorld()->GetNetDriver() != NULL && GetWorld()->GetNetDriver()->GetNetMode() == NM_Client );
 
 	DemoReplicateActor( World->GetWorldSettings(), ClientConnections[0], IsNetClient );
@@ -417,6 +430,8 @@ void UDemoNetDriver::TickDemoRecord( float DeltaSeconds )
 
 	check( ClientConnections[0]->SendBuffer.GetNumBits() == 0 );
 
+	bIsRecordingDemoFrame = false;
+
 	// Write a count of 0 to signal the end of the frame
 	int32 EndCount = 0;
 
@@ -623,6 +638,15 @@ void UDemoNetConnection::LowLevelSend( void* Data, int32 Count )
 
 	if ( !GetDriver()->ServerConnection && GetDriver()->FileAr )
 	{
+		// If we're outside of an official demo frame, we need to queue this up or it will throw off the stream
+		if ( !GetDriver()->bIsRecordingDemoFrame )
+		{
+			FQueuedDemoPacket & B = *( new( QueuedDemoPackets )FQueuedDemoPacket );
+			B.Data.AddUninitialized( Count );
+			FMemory::Memcpy( B.Data.GetData(), Data, Count );
+			return;
+		}
+
 		*GetDriver()->FileAr << Count;
 		GetDriver()->FileAr->Serialize( Data, Count );
 		
