Notify demo driver when actors are destroyed

diff --git a/Engine/Source/Runtime/Engine/Private/LevelActor.cpp b/Engine/Source/Runtime/Engine/Private/LevelActor.cpp
--- a/Engine/Source/Runtime/Engine/Private/LevelActor.cpp
+++ b/Engine/Source/Runtime/Engine/Private/LevelActor.cpp
@@ -5,6 +5,7 @@
 #include "EditorSupportDelegates.h"
 #include "Net/UnrealNetwork.h"
 #include "Collision.h"
+#include "Engine/DemoNetDriver.h"
 
 #if WITH_PHYSX
 	#include "PhysicsEngine/PhysXSupport.h"
@@ -535,6 +536,11 @@ bool UWorld::DestroyActor( AActor* ThisActor, bool bNetForce, bool bShouldModify
 		NetDriver->NotifyActorDestroyed( ThisActor );
 	}
 
+	if ( DemoNetDriver )
+	{
+		DemoNetDriver->NotifyActorDestroyed( ThisActor );
+	}
+
 	// Remove the actor from the actor list.
 	RemoveActor( ThisActor, bShouldModifyLevel );
 	
