Few minor rendering optimizations for NavMesh #ue4

diff --git a/Engine/Source/Runtime/Engine/Private/AI/Navigation/RecastNavMeshGenerator.cpp b/Engine/Source/Runtime/Engine/Private/AI/Navigation/RecastNavMeshGenerator.cpp
--- a/Engine/Source/Runtime/Engine/Private/AI/Navigation/RecastNavMeshGenerator.cpp
+++ b/Engine/Source/Runtime/Engine/Private/AI/Navigation/RecastNavMeshGenerator.cpp
@@ -4182,6 +4182,7 @@ void FRecastNavMeshGenerator::StartDirtyGenerators()
 		DirtyGenerators.KeySort(SortFunc);
 
 		FRecastTileDirtyState DirtyState;
+		const bool bRequestDrawingUpdate = DirtyGenerators.Num() > 0;
 		for (TMap<int32,FRecastTileDirtyState>::TIterator It(DirtyGenerators); It; ++It)
 		{
 			const int32 GenIdx = It.Key();
@@ -4265,6 +4266,10 @@ void FRecastNavMeshGenerator::StartDirtyGenerators()
 
 		// start processing generator queue
 		UpdateTileGenerationWorkers(INDEX_NONE);
+		if (bRequestDrawingUpdate)
+		{
+			DestNavMesh->RequestDrawingUpdate();
+		}
 	}
 
 	INC_FLOAT_STAT_BY(STAT_Navigation_CumulativeBuildTime,(float)ThisTime*1000);
@@ -4381,8 +4386,6 @@ void FRecastNavMeshGenerator::UpdateTileGenerationWorkers(int32 TileId)
 		// let navmesh drawing mark tiles being built
 		DestNavMesh->RequestDrawingUpdate();
 	}
-#else
-	DestNavMesh->RequestDrawingUpdate();
 #endif // RECAST_ASYNC_REBUILDING
 
 	// prepare next batch of dirty generators in next tick
diff --git a/Engine/Source/Runtime/Engine/Private/NavMeshRenderingHelpers.h b/Engine/Source/Runtime/Engine/Private/NavMeshRenderingHelpers.h
--- a/Engine/Source/Runtime/Engine/Private/NavMeshRenderingHelpers.h
+++ b/Engine/Source/Runtime/Engine/Private/NavMeshRenderingHelpers.h
@@ -385,7 +385,8 @@ public:
 	{
 		QUICK_SCOPE_CYCLE_COUNTER( STAT_RecastRenderingSceneProxy_DrawDynamicElements );
 
-		if (!ProxyData.bEnableDrawing) //check if we have any data to render
+		const bool bVisible = !!View->Family->EngineShowFlags.Navigation || bForceRendering;
+		if (!ProxyData.bEnableDrawing || !bVisible) //check if we have any data to render
 		{
 			return;
 		}
@@ -530,6 +531,11 @@ public:
 			if (VisibilityMap & (1 << ViewIndex))
 			{
 				const FSceneView* View = Views[ViewIndex];
+				const bool bVisible = !!View->Family->EngineShowFlags.Navigation || bForceRendering;
+				if (!bVisible)
+				{
+					continue;
+				}
 				FPrimitiveDrawInterface* PDI = Collector.GetPDI(ViewIndex);
 
 				const bool bSkipDistanceCheck = GIsEditor && (GEngine->GetDebugLocalPlayer() == NULL);
@@ -669,7 +675,8 @@ public:
 
 	void DrawDebugLabels(UCanvas* Canvas, APlayerController*)
 	{
-		if (ProxyData.bNeedsNewData == true || ProxyData.bEnableDrawing == false)
+		const bool bVisible = (Canvas && Canvas->SceneView && !!Canvas->SceneView->Family->EngineShowFlags.Navigation) || bForceRendering;
+		if (ProxyData.bNeedsNewData == true || ProxyData.bEnableDrawing == false || !bVisible || ProxyData.DebugLabels.Num() == 0)
 		{
 			return;
 		}
