Change GetBoneLocation, GetBoneQuaternion to use space enum instead of ints

#TTP 342772 - Anim: GetBoneLocation, etc... use an int32 to indicate their space instead of one of the existing enums

diff --git a/Engine/Config/BaseEngine.ini b/Engine/Config/BaseEngine.ini
--- a/Engine/Config/BaseEngine.ini
+++ b/Engine/Config/BaseEngine.ini
@@ -1305,6 +1305,9 @@ bUseStreamingPause=false
 ;Removal of EmitterSpawnable
 +ActiveClassRedirects=(OldClassName="EmitterSpawnable", NewClassName="Emitter")
 
+; Moving BoneSpace enum so more classes can make use of it
++ActiveClassRedirects=(OldClassName="EBoneSpaces",NewClassName="/Script/Engine.SkinnedMeshComponent:EBoneSpaces")
+
 ;AdBanner redirects
 +K2FieldRedirects=(OldFieldName="KismetSystemLibrary.EXPERIMENTAL_ShowAdBanner",NewFieldName="KismetSystemLibrary.ShowAdBanner")
 +K2FieldRedirects=(OldFieldName="KismetSystemLibrary.EXPERIMENTAL_HideAdBanner",NewFieldName="KismetSystemLibrary.HideAdBanner")
diff --git a/Engine/Source/Runtime/Engine/Classes/Components/PoseableMeshComponent.h b/Engine/Source/Runtime/Engine/Classes/Components/PoseableMeshComponent.h
--- a/Engine/Source/Runtime/Engine/Classes/Components/PoseableMeshComponent.h
+++ b/Engine/Source/Runtime/Engine/Classes/Components/PoseableMeshComponent.h
@@ -4,20 +4,6 @@
 #include "Animation/AnimationAsset.h"
 #include "PoseableMeshComponent.generated.h"
 
-UENUM()
-namespace EBoneSpaces
-{
-	enum Type
-	{
-		/** Set absolute position of bone in world space */
-		WorldSpace		UMETA( DisplayName = "World Space" ),
-		/** Set position of bone in components reference frame */
-		ComponentSpace	UMETA( DisplayName = "Component Space" ),
-		/** Set position of bone relative to parent bone */
-		//LocalSpace		UMETA( DisplayName = "Parent Bone Space" ),
-	};
-}
-
 /**
  *	UPoseableMeshComponent that allows bone transforms to be driven by blueprint.
  */
diff --git a/Engine/Source/Runtime/Engine/Classes/Components/SkinnedMeshComponent.h b/Engine/Source/Runtime/Engine/Classes/Components/SkinnedMeshComponent.h
--- a/Engine/Source/Runtime/Engine/Classes/Components/SkinnedMeshComponent.h
+++ b/Engine/Source/Runtime/Engine/Classes/Components/SkinnedMeshComponent.h
@@ -59,6 +59,21 @@ namespace EMeshComponentUpdateFlag
 	};
 }
 
+/** Flag for specifying bone space */
+UENUM()
+namespace EBoneSpaces
+{
+	enum Type
+	{
+		/** Set absolute position of bone in world space */
+		WorldSpace		UMETA(DisplayName = "World Space"),
+		/** Set position of bone in components reference frame */
+		ComponentSpace	UMETA(DisplayName = "Component Space"),
+		/** Set position of bone relative to parent bone */
+		//LocalSpace		UMETA( DisplayName = "Parent Bone Space" ),
+	};
+}
+
 //
 // Vertex Anim
 //
@@ -674,7 +689,7 @@ public:
 	 * 
 	 * @return Quaternion of the bone
 	 */
-	FQuat GetBoneQuaternion( FName BoneName, int32 Space = 0 ) const; 
+	FQuat GetBoneQuaternion(FName BoneName, EBoneSpaces::Type Space = EBoneSpaces::WorldSpace) const;
 
 	/** Get Bone Location
 	 *
@@ -683,7 +698,7 @@ public:
 	 * 
 	 * @return Vector of the bone
 	 */
-	FVector GetBoneLocation( FName BoneName, int32 Space = 0 ) const; 
+	FVector GetBoneLocation( FName BoneName, EBoneSpaces::Type Space = EBoneSpaces::WorldSpace) const; 
 
 	/** 
 	 * Fills the given array with the names of all the bones in this component's current SkeletalMesh 
diff --git a/Engine/Source/Runtime/Engine/Private/SkinnedMeshComponent.cpp b/Engine/Source/Runtime/Engine/Private/SkinnedMeshComponent.cpp
--- a/Engine/Source/Runtime/Engine/Private/SkinnedMeshComponent.cpp
+++ b/Engine/Source/Runtime/Engine/Private/SkinnedMeshComponent.cpp
@@ -1050,7 +1050,7 @@ FName USkinnedMeshComponent::GetSocketBoneName(FName InSocketName)
 }
 
 
-FQuat USkinnedMeshComponent::GetBoneQuaternion(FName BoneName, int32 Space) const
+FQuat USkinnedMeshComponent::GetBoneQuaternion(FName BoneName, EBoneSpaces::Type Space) const
 {
 	int32 BoneIndex = GetBoneIndex(BoneName);
 
@@ -1061,7 +1061,7 @@ FQuat USkinnedMeshComponent::GetBoneQuaternion(FName BoneName, int32 Space) cons
 	}
 
 	FTransform BoneTransform;
-	if( Space == 1 )
+	if( Space == EBoneSpaces::ComponentSpace )
 	{
 		if(MasterPoseComponent.IsValid())
 		{
@@ -1099,7 +1099,7 @@ FQuat USkinnedMeshComponent::GetBoneQuaternion(FName BoneName, int32 Space) cons
 }
 
 
-FVector USkinnedMeshComponent::GetBoneLocation( FName BoneName, int32 Space ) const
+FVector USkinnedMeshComponent::GetBoneLocation(FName BoneName, EBoneSpaces::Type Space) const
 {
 	int32 BoneIndex = GetBoneIndex(BoneName);
 	if( BoneIndex == INDEX_NONE )
@@ -1108,8 +1108,7 @@ FVector USkinnedMeshComponent::GetBoneLocation( FName BoneName, int32 Space ) co
 		return FVector::ZeroVector;
 	}
 
-	// If space == Local
-	if( Space == 1 )
+	if( Space == EBoneSpaces::ComponentSpace )
 	{
 		if(MasterPoseComponent.IsValid())
 		{
@@ -1132,11 +1131,16 @@ FVector USkinnedMeshComponent::GetBoneLocation( FName BoneName, int32 Space ) co
 			return SpaceBases[BoneIndex].GetLocation();
 		}
 	}
-	else
+	else if (Space == EBoneSpaces::WorldSpace)
 	{
 		// To support non-uniform scale (via LocalToWorld), use GetBoneMatrix
 		return GetBoneMatrix(BoneIndex).GetOrigin();
 	}
+	else
+	{
+		check(false); // Unknown BoneSpace
+		return FVector::ZeroVector;
+	}
 }
 
 
