Corrected widget scale behaviour in the case that any scale component is negative.
Const-corrected FVector::GetSignVector.
#ttp 347989 - Editor: Uniform scaling is broken on mehes where one of scaling values is negative
#branch UE4
#proj Runtime.Engine, Runtime.Core

diff --git a/Engine/Source/Runtime/Core/Public/Math/Vector.h b/Engine/Source/Runtime/Core/Public/Math/Vector.h
--- a/Engine/Source/Runtime/Core/Public/Math/Vector.h
+++ b/Engine/Source/Runtime/Core/Public/Math/Vector.h
@@ -467,7 +467,7 @@ public:
 	 *
 	 * @param A copy of the vector with each component set to +1 or -1
 	 */
-	FORCEINLINE FVector GetSignVector();
+	FORCEINLINE FVector GetSignVector() const;
 
 	/**
 	 * Projects 2D components of vector based on Z.
@@ -1376,7 +1376,7 @@ FORCEINLINE void FVector::ToDirectionAndLength(FVector &OutDir, float &OutLength
 	}
 }
 
-FORCEINLINE FVector FVector::GetSignVector()
+FORCEINLINE FVector FVector::GetSignVector() const
 {
 	return FVector
 		(
diff --git a/Engine/Source/Runtime/Engine/Private/ActorEditor.cpp b/Engine/Source/Runtime/Engine/Private/ActorEditor.cpp
--- a/Engine/Source/Runtime/Engine/Private/ActorEditor.cpp
+++ b/Engine/Source/Runtime/Engine/Private/ActorEditor.cpp
@@ -377,14 +377,16 @@ void AActor::EditorApplyScale( const FVector& DeltaScale, const FVector* PivotLo
 {
 	if( RootComponent != NULL )
 	{
+		const FVector CurrentScale = GetRootComponent()->RelativeScale3D;
+
 		// @todo: Remove this hack once we have decided on the scaling method to use.
 		if( AActor::bUsePercentageBasedScaling )
 		{
-			GetRootComponent()->SetRelativeScale3D( GetRootComponent()->RelativeScale3D + DeltaScale * GetRootComponent()->RelativeScale3D );
+			GetRootComponent()->SetRelativeScale3D(CurrentScale + DeltaScale * CurrentScale);
 		}
 		else
 		{
-			GetRootComponent()->SetRelativeScale3D( GetRootComponent()->RelativeScale3D + DeltaScale );
+			GetRootComponent()->SetRelativeScale3D(CurrentScale + DeltaScale * CurrentScale.GetSignVector());
 		}
 
 		if( PivotLocation )
