virtualize GameViewportClient::Tick

diff --git a/Engine/Source/Runtime/Engine/Classes/Engine/GameViewportClient.h b/Engine/Source/Runtime/Engine/Classes/Engine/GameViewportClient.h
--- a/Engine/Source/Runtime/Engine/Classes/Engine/GameViewportClient.h
+++ b/Engine/Source/Runtime/Engine/Classes/Engine/GameViewportClient.h
@@ -345,7 +345,7 @@ public:
 	 * Called every frame to allow the game viewport to update time based state.
 	 * @param	DeltaTime	The time since the last call to Tick.
 	 */
-	void Tick( float DeltaTime );
+	virtual void Tick( float DeltaTime );
 
 	/**
 	 * Determines whether this viewport client should receive calls to InputAxis() if the game's window is not currently capturing the mouse.
