Only force skeletal mesh to update physics if the skeletal mesh is actually simulating. This fixes major performance issues coming out of Matinee

diff --git a/Engine/Source/Runtime/Engine/Private/PhysicsEngine/PhysAnim.cpp b/Engine/Source/Runtime/Engine/Private/PhysicsEngine/PhysAnim.cpp
--- a/Engine/Source/Runtime/Engine/Private/PhysicsEngine/PhysAnim.cpp
+++ b/Engine/Source/Runtime/Engine/Private/PhysicsEngine/PhysAnim.cpp
@@ -254,7 +254,7 @@ void USkeletalMeshComponent::UpdateKinematicBonesToPhysics(bool bTeleport, bool
 	const bool bUpdateKinematics = (KinematicBonesUpdateType != EKinematicBonesUpdateToPhysics::SkipAllBones);
 
 	// If desired, update physics bodies associated with skeletal mesh component to match.
-	if( !bUpdateKinematics && !bForceUpdate)
+	if( !bUpdateKinematics && !(bForceUpdate && IsAnySimulatingPhysics()))
 	{
 		// nothing to do 
 		return;
