Fixed crash with cooked static mesh

diff --git a/Engine/Source/Runtime/Engine/Private/StaticMesh.cpp b/Engine/Source/Runtime/Engine/Private/StaticMesh.cpp
--- a/Engine/Source/Runtime/Engine/Private/StaticMesh.cpp
+++ b/Engine/Source/Runtime/Engine/Private/StaticMesh.cpp
@@ -627,6 +627,11 @@ void FStaticMeshRenderData::Serialize(FArchive& Ar, UStaticMesh* Owner, bool bCo
 
 			if (bValid)
 			{
+				if (!LOD.DistanceFieldData)
+				{
+					LOD.DistanceFieldData = new FDistanceFieldVolumeData();
+				}
+
 				Ar << *(LOD.DistanceFieldData);
 			}
 		}
@@ -1153,8 +1158,12 @@ void FStaticMeshRenderData::Cache(UStaticMesh* Owner, const FStaticMeshLODSettin
 	if (CVar->GetValueOnGameThread() != 0)
 	{
 		FString DistanceFieldKey = BuildDistanceFieldDerivedDataKey(DerivedDataKey);
-		check(!LODResources[0].DistanceFieldData);
-		LODResources[0].DistanceFieldData = new FDistanceFieldVolumeData();
+		
+		if (!LODResources[0].DistanceFieldData)
+		{
+			LODResources[0].DistanceFieldData = new FDistanceFieldVolumeData();
+		}
+
 		LODResources[0].DistanceFieldData->CacheDerivedData(DistanceFieldKey, Owner);
 	}
 }
