World::SpawnActor casts restored.

diff --git a/Engine/Source/Runtime/Engine/Classes/Engine/World.h b/Engine/Source/Runtime/Engine/Classes/Engine/World.h
--- a/Engine/Source/Runtime/Engine/Classes/Engine/World.h
+++ b/Engine/Source/Runtime/Engine/Classes/Engine/World.h
@@ -2066,27 +2066,25 @@ public:
 	template< class T >
 	T* SpawnActor( const FActorSpawnParameters& SpawnParameters = FActorSpawnParameters() )
 	{
-		return static_cast<T*>(SpawnActor(T::StaticClass(), NULL, NULL, SpawnParameters));
+		return CastChecked<T>(SpawnActor(T::StaticClass(), NULL, NULL, SpawnParameters),ECastCheckedType::NullAllowed);
 	}
 
 	template< class T >
 	T* SpawnActor( FVector const& Location, FRotator const& Rotation, const FActorSpawnParameters& SpawnParameters = FActorSpawnParameters() )
 	{
-		return static_cast<T*>(SpawnActor(T::StaticClass(), &Location, &Rotation, SpawnParameters));
+		return CastChecked<T>(SpawnActor(T::StaticClass(), &Location, &Rotation, SpawnParameters),ECastCheckedType::NullAllowed);
 	}
 	
 	template< class T >
 	T* SpawnActor( UClass* Class, const FActorSpawnParameters& SpawnParameters = FActorSpawnParameters() )
 	{
-		check(Class->IsChildOf(T::StaticClass()));
-		return static_cast<T*>(SpawnActor(Class, NULL, NULL, SpawnParameters));
+		return CastChecked<T>(SpawnActor(Class, NULL, NULL, SpawnParameters),ECastCheckedType::NullAllowed);
 	}
 
 	template< class T >
 	T* SpawnActor( UClass* Class, FVector const& Location, FRotator const& Rotation, const FActorSpawnParameters& SpawnParameters = FActorSpawnParameters() )
 	{
-		check(Class->IsChildOf(T::StaticClass()));
-		return static_cast<T*>(SpawnActor(Class, &Location, &Rotation, SpawnParameters));
+		return CastChecked<T>(SpawnActor(Class, &Location, &Rotation, SpawnParameters),ECastCheckedType::NullAllowed);
 	}
 
 
