renamed ReattachComponents to Reattach.Components and added Reattach.Materials and Reattach.MatrialInstances

diff --git a/Engine/Source/Runtime/Engine/Classes/Engine/GameEngine.h b/Engine/Source/Runtime/Engine/Classes/Engine/GameEngine.h
--- a/Engine/Source/Runtime/Engine/Classes/Engine/GameEngine.h
+++ b/Engine/Source/Runtime/Engine/Classes/Engine/GameEngine.h
@@ -103,7 +103,6 @@ public:
 	 * Exec command handlers
 	 */
 
-	bool HandleReattachComponentsCommand( const TCHAR* Cmd, FOutputDevice& Ar );
 	bool HandleCommand( const TCHAR* Cmd, FOutputDevice& Ar );
 	bool HandleExitCommand( const TCHAR* Cmd, FOutputDevice& Ar );
 	bool HandleMinimizeCommand( const TCHAR *Cmd, FOutputDevice &Ar );
diff --git a/Engine/Source/Runtime/Engine/Private/Engine.cpp b/Engine/Source/Runtime/Engine/Private/Engine.cpp
--- a/Engine/Source/Runtime/Engine/Private/Engine.cpp
+++ b/Engine/Source/Runtime/Engine/Private/Engine.cpp
@@ -62,3 +62,107 @@ void ResetCachedRendererModule()
 {
 	CachedRendererModule = NULL;
 }
+
+
+
+#if !(UE_BUILD_SHIPPING || UE_BUILD_TEST)
+
+void ReattachMaterialInstances(const TArray<FString>& Args)
+{
+	FMaterialUpdateContext MaterialUpdateContext;
+
+	UE_LOG(LogConsoleResponse, Display, TEXT("Reattach.MaterialInstances:"));
+
+	// Clear the parents out of combination material instances
+	for( TObjectIterator<UMaterialInstanceConstant> MaterialIt; MaterialIt; ++MaterialIt )
+	{
+		if(Args.Num() == 1)
+		{
+			UE_LOG(LogConsoleResponse, Display, TEXT("   %s"), *MaterialIt->GetName());
+
+			if(MaterialIt->GetName() == Args[0])
+			{
+				MaterialUpdateContext.AddMaterialInstance(*MaterialIt);
+			}
+		}
+	}
+	UE_LOG(LogConsoleResponse, Display, TEXT(""));
+}
+
+FAutoConsoleCommand ReattachMaterialInstancesCmd(
+	TEXT("Reattach.MaterialInstances"),
+	TEXT("Useful for debugging, reattaches all materials. Optional parameter can be a materialinstance name (e.g. DecoStatue_Subsurface0)."),
+	FConsoleCommandWithArgsDelegate::CreateStatic(&ReattachMaterialInstances)
+	);
+
+void ReattachMaterials(const TArray<FString>& Args)
+{
+	FMaterialUpdateContext MaterialUpdateContext;
+
+	UE_LOG(LogConsoleResponse, Display, TEXT("Reattach.Materials:"));
+
+	// Clear the parents out of combination material
+	for( TObjectIterator<UMaterial> MaterialIt; MaterialIt; ++MaterialIt )
+	{
+		if(Args.Num() == 1)
+		{
+			UE_LOG(LogConsoleResponse, Display, TEXT("   %s"), *MaterialIt->GetName());
+
+			if(MaterialIt->GetName() == Args[0])
+			{
+				MaterialUpdateContext.AddMaterial(*MaterialIt);
+			}
+		}
+	}
+	UE_LOG(LogConsoleResponse, Display, TEXT(""));
+}
+
+FAutoConsoleCommand ReattachMaterialsCmd(
+	TEXT("Reattach.Materials"),
+	TEXT("Useful for debugging, reattaches all materials. Optional parameter can be a material name (e.g. DecoStatue_Subsurface0_Inst)."),
+	FConsoleCommandWithArgsDelegate::CreateStatic(&ReattachMaterials)
+	);
+
+
+void ReattachComponents(const TArray<FString>& Args)
+{
+	if(Args.Num() != 1)
+	{
+		UE_LOG(LogConsoleResponse, Warning, TEXT("Reattach.Components: missing class name parameter"));
+		return;
+	}
+
+	UE_LOG(LogConsoleResponse, Display, TEXT("Reattach.Components:"));
+
+	UClass* Class=NULL;
+	if( ParseObject<UClass>( *Args[0], TEXT("CLASS="), Class, ANY_PACKAGE ) &&
+		Class->IsChildOf(UActorComponent::StaticClass()) )
+	{
+		for( FObjectIterator It(Class); It; ++It )
+		{
+			UActorComponent* ActorComponent = Cast<UActorComponent>(*It);
+			if( ActorComponent )
+			{
+				UE_LOG(LogConsoleResponse, Display, TEXT("   Component: %s"), 
+					*ActorComponent->GetName());
+
+				FComponentReregisterContext Reregister(ActorComponent);
+			}
+		}
+		UE_LOG(LogConsoleResponse, Display, TEXT(""));
+	}
+	else
+	{
+		UE_LOG(LogConsoleResponse, Warning, TEXT("Reattach.Components: No objects with the class name '%s' found"), *Args[0]);
+	}
+
+}
+
+FAutoConsoleCommand ReattachComponentsCmd(
+	TEXT("Reattach.Components"),
+	TEXT("Useful for debugging, reattaches all components. Parameter needs to be the class name."),
+	FConsoleCommandWithArgsDelegate::CreateStatic(&ReattachComponents)
+	);
+
+
+#endif // !(UE_BUILD_SHIPPING || UE_BUILD_TEST)
diff --git a/Engine/Source/Runtime/Engine/Private/GameEngine.cpp b/Engine/Source/Runtime/Engine/Private/GameEngine.cpp
--- a/Engine/Source/Runtime/Engine/Private/GameEngine.cpp
+++ b/Engine/Source/Runtime/Engine/Private/GameEngine.cpp
@@ -543,7 +543,8 @@ bool UGameEngine::Exec( UWorld* InWorld, const TCHAR* Cmd, FOutputDevice& Ar )
 {
 	if( FParse::Command( &Cmd,TEXT("REATTACHCOMPONENTS")) || FParse::Command( &Cmd,TEXT("REREGISTERCOMPONENTS")))
 	{
-		return HandleReattachComponentsCommand( Cmd, Ar );
+		UE_LOG(LogConsoleResponse, Warning, TEXT("Deprectated command! Please use 'Reattach.Components' instead."));
+		return true;
 	}
 	else if( FParse::Command( &Cmd,TEXT("EXIT")) || FParse::Command(&Cmd,TEXT("QUIT")))
 	{
@@ -619,26 +620,6 @@ bool UGameEngine::Exec( UWorld* InWorld, const TCHAR* Cmd, FOutputDevice& Ar )
 	}
 }
 
-
-bool UGameEngine::HandleReattachComponentsCommand( const TCHAR* Cmd, FOutputDevice& Ar )
-{
-	UClass* Class=NULL;
-	if( ParseObject<UClass>( Cmd, TEXT("CLASS="), Class, ANY_PACKAGE ) &&
-		Class->IsChildOf(UActorComponent::StaticClass()) )
-	{
-		for( FObjectIterator It(Class); It; ++It )
-		{
-			UActorComponent* ActorComponent = Cast<UActorComponent>(*It);
-			if( ActorComponent )
-			{
-				FComponentReregisterContext Reregister(ActorComponent);
-			}
-		}
-	}
-	return true;
-}
-
-
 bool UGameEngine::HandleExitCommand( const TCHAR* Cmd, FOutputDevice& Ar )
 {
 	for (int32 WorldIndex = 0; WorldIndex < WorldList.Num(); ++WorldIndex)
