Engine: Fix memory leak in FPhysicsScene2D

diff --git a/Engine/Source/Runtime/Engine/Private/PhysicsEngine2D/Box2DIntegration.cpp b/Engine/Source/Runtime/Engine/Private/PhysicsEngine2D/Box2DIntegration.cpp
--- a/Engine/Source/Runtime/Engine/Private/PhysicsEngine2D/Box2DIntegration.cpp
+++ b/Engine/Source/Runtime/Engine/Private/PhysicsEngine2D/Box2DIntegration.cpp
@@ -126,6 +126,8 @@ FPhysicsScene2D::~FPhysicsScene2D()
 {
 	StartPhysicsTickFunction.UnRegisterTickFunction();
 	EndPhysicsTickFunction.UnRegisterTickFunction();
+
+	delete World;
 }
 
 bool FPhysicsScene2D::ShouldCollide(b2Fixture* FixtureA, b2Fixture* FixtureB)
