Added geometry and current time to ISlateViewport Tick function

This will be needed by the web browser so that its viewport can pass on the size of the viewport before it tries to render it.

diff --git a/Engine/Source/Runtime/Engine/Private/Slate/SceneViewport.cpp b/Engine/Source/Runtime/Engine/Private/Slate/SceneViewport.cpp
--- a/Engine/Source/Runtime/Engine/Private/Slate/SceneViewport.cpp
+++ b/Engine/Source/Runtime/Engine/Private/Slate/SceneViewport.cpp
@@ -1116,7 +1116,7 @@ void FSceneViewport::EndRenderFrame(FRHICommandListImmediate& RHICmdList, bool b
 	}
 }
 
-void FSceneViewport::Tick( float DeltaTime )
+void FSceneViewport::Tick( const FGeometry& AllottedGeometry, double InCurrentTime, float DeltaTime )
 {
 	ProcessInput( DeltaTime );
 }
diff --git a/Engine/Source/Runtime/Engine/Public/Slate/SceneViewport.h b/Engine/Source/Runtime/Engine/Public/Slate/SceneViewport.h
--- a/Engine/Source/Runtime/Engine/Public/Slate/SceneViewport.h
+++ b/Engine/Source/Runtime/Engine/Public/Slate/SceneViewport.h
@@ -120,7 +120,7 @@ public:
 	/**
 	 * Ticks the viewport
 	 */
-	virtual void Tick( float DeltaTime );
+	virtual void Tick( const FGeometry& AllottedGeometry, double InCurrentTime, float DeltaTime ) override;
 
 	/**
 	 * Performs a resize when in swapping viewports while viewing the play world.
diff --git a/Engine/Source/Runtime/Slate/Private/Widgets/SViewport.cpp b/Engine/Source/Runtime/Slate/Private/Widgets/SViewport.cpp
--- a/Engine/Source/Runtime/Slate/Private/Widgets/SViewport.cpp
+++ b/Engine/Source/Runtime/Slate/Private/Widgets/SViewport.cpp
@@ -112,7 +112,7 @@ void SViewport::Tick( const FGeometry& AllottedGeometry, const double InCurrentT
 {
 	if(ViewportInterface.IsValid())
 	{
-		ViewportInterface.Pin()->Tick(InDeltaTime);
+		ViewportInterface.Pin()->Tick(AllottedGeometry, InCurrentTime, InDeltaTime);
 	}
 }
 
diff --git a/Engine/Source/Runtime/SlateCore/Public/Rendering/RenderingCommon.h b/Engine/Source/Runtime/SlateCore/Public/Rendering/RenderingCommon.h
--- a/Engine/Source/Runtime/SlateCore/Public/Rendering/RenderingCommon.h
+++ b/Engine/Source/Runtime/SlateCore/Public/Rendering/RenderingCommon.h
@@ -232,7 +232,7 @@ public:
 	/**
 	 * Performs any ticking necessary by this handle                   
 	 */
-	virtual void Tick( float DeltaTime )
+	virtual void Tick( const FGeometry& AllottedGeometry, double InCurrentTime, float DeltaTime )
 	{
 	}
 
diff --git a/Engine/Source/Runtime/UMG/Private/Components/Viewport.cpp b/Engine/Source/Runtime/UMG/Private/Components/Viewport.cpp
--- a/Engine/Source/Runtime/UMG/Private/Components/Viewport.cpp
+++ b/Engine/Source/Runtime/UMG/Private/Components/Viewport.cpp
@@ -469,7 +469,7 @@ class SAutoRefreshViewport : public SViewport
 		Viewport->Invalidate();
 		Viewport->InvalidateDisplay();
 
-		Viewport->Tick(InDeltaTime);
+		Viewport->Tick(AllottedGeometry, InCurrentTime, InDeltaTime);
 
 		ViewportClient->Tick(InDeltaTime);
 	}
