Created an accessor and mutator for ForwardAxis in SplineMeshComponent.
#ttp 347332 - SplineComponent: can't change forward axis in construction script
#branch UE4
#proj Runtime.Engine

diff --git a/Engine/Source/Runtime/Engine/Classes/Components/SplineMeshComponent.h b/Engine/Source/Runtime/Engine/Classes/Components/SplineMeshComponent.h
--- a/Engine/Source/Runtime/Engine/Classes/Components/SplineMeshComponent.h
+++ b/Engine/Source/Runtime/Engine/Classes/Components/SplineMeshComponent.h
@@ -234,7 +234,13 @@ class ENGINE_API USplineMeshComponent : public UStaticMeshComponent, public IInt
 	UFUNCTION(BlueprintCallable, Category = SplineMesh)
 	void SetEndOffset(FVector2D EndOffset);
 
+	/** Get the forward axis */
+	UFUNCTION(BlueprintCallable, Category = SplineMesh)
+	ESplineMeshAxis::Type GetForwardAxis() const;
 
+	/** Set the forward axis */
+	UFUNCTION(BlueprintCallable, Category = SplineMesh)
+	void SetForwardAxis(ESplineMeshAxis::Type InForwardAxis);
 
 	// Destroys the body setup, used to clear collision if the mesh goes missing
 	void DestroyBodySetup();
diff --git a/Engine/Source/Runtime/Engine/Private/SplineMeshComponent.cpp b/Engine/Source/Runtime/Engine/Private/SplineMeshComponent.cpp
--- a/Engine/Source/Runtime/Engine/Private/SplineMeshComponent.cpp
+++ b/Engine/Source/Runtime/Engine/Private/SplineMeshComponent.cpp
@@ -354,6 +354,18 @@ void USplineMeshComponent::SetEndOffset(FVector2D EndOffset)
 	MarkSplineParamsDirty();
 }
 
+ESplineMeshAxis::Type USplineMeshComponent::GetForwardAxis() const
+{
+	return ForwardAxis;
+}
+
+void USplineMeshComponent::SetForwardAxis(ESplineMeshAxis::Type InForwardAxis)
+{
+	ForwardAxis = InForwardAxis;
+	MarkSplineParamsDirty();
+}
+
+
 void USplineMeshComponent::MarkSplineParamsDirty()
 {
 	MarkRenderStateDirty();
