Fix TBasePlayerControllerIterator.
Github pull request https://github.com/EpicGames/UnrealEngine/pull/403

diff --git a/Engine/Source/Runtime/Engine/Public/UnrealEngine.h b/Engine/Source/Runtime/Engine/Public/UnrealEngine.h
--- a/Engine/Source/Runtime/Engine/Public/UnrealEngine.h
+++ b/Engine/Source/Runtime/Engine/Public/UnrealEngine.h
@@ -91,7 +91,7 @@ public:
 		: Iter(InWorld->GetPlayerControllerIterator())
 	{
 		check(!LocalOnly || InWorld->GetNetMode() != NM_Client);	// You should only iterate on non local player controllers if you are the server
-		GetCurrent();
+		AdvanceCurrent();
 	}
 
 	void operator++()
@@ -116,7 +116,7 @@ protected:
 	FConstPlayerControllerIterator Iter;
 	T* Current;
 
-	T* GetCurrent()
+	void AdvanceCurrent()
 	{
 		// Look at current Iter
 		Current = Cast<T>(*Iter);
@@ -125,17 +125,17 @@ protected:
 		while(Iter && (!Current || (LocalOnly && !Current->IsLocalController())))
 		{
 			++Iter;
-			Current = Cast<T>*Iter;
+			Current = Cast<T>(*Iter);
 		}
 	}
 
-	T* Next()
+	void Next()
 	{
 		// Advance one
 		++Iter;
 
 		// Update Current
-		GetCurrent();
+		AdvanceCurrent();
 	}
 };
 
