Removed hard reference to AIController from APawn's constructor #UE4

- made default AIController class ini-configurable

diff --git a/Engine/Config/BaseEngine.ini b/Engine/Config/BaseEngine.ini
--- a/Engine/Config/BaseEngine.ini
+++ b/Engine/Config/BaseEngine.ini
@@ -39,6 +39,7 @@ LevelScriptActorClassName=/Script/Engine.LevelScriptActor
 DefaultBlueprintBaseClassName=/Script/Engine.Actor
 DataStoreClientClassName=/Script/Engine.DataStoreClient
 GameUserSettingsClassName=/Script/Engine.GameUserSettings
+AIControllerClassName=/Script/AIModule.AIController
 bAllowMatureLanguage=false
 GameEngine=/Script/Engine.GameEngine
 EditorEngine=/Script/UnrealEd.EditorEngine
diff --git a/Engine/Source/Runtime/Engine/Classes/Engine/Engine.h b/Engine/Source/Runtime/Engine/Classes/Engine/Engine.h
--- a/Engine/Source/Runtime/Engine/Classes/Engine/Engine.h
+++ b/Engine/Source/Runtime/Engine/Classes/Engine/Engine.h
@@ -658,11 +658,14 @@ public:
 	UPROPERTY()
 	TSubclassOf<class UGameUserSettings> GameUserSettingsClass;
 
+	/** name of Controller class to be used as default AIController class for pawns */
+	UPROPERTY(globalconfig, noclear, meta = (MetaClass = "AI", DisplayName = "Default AIController class for all Pawns"))
+	FStringClassReference AIControllerClassName;
+
 	/** Global instance of the user game settings */
 	UPROPERTY()
 	class UGameUserSettings* GameUserSettings;
 
-
 	/** @todo document */
 	UPROPERTY()
 	TSubclassOf<class ALevelScriptActor>  LevelScriptActorClass;
diff --git a/Engine/Source/Runtime/Engine/Private/Pawn.cpp b/Engine/Source/Runtime/Engine/Private/Pawn.cpp
--- a/Engine/Source/Runtime/Engine/Private/Pawn.cpp
+++ b/Engine/Source/Runtime/Engine/Private/Pawn.cpp
@@ -15,7 +15,6 @@
 #include "DisplayDebugHelpers.h"
 #include "NetworkingDistanceConstants.h"
 #include "VisualLog.h"
-#include "AIController.h"
 #include "Engine/InputDelegateBinding.h"
 
 DEFINE_LOG_CATEGORY(LogDamage);
@@ -31,7 +30,10 @@ APawn::APawn(const class FPostConstructInitializeProperties& PCIP)
 {
 	PrimaryActorTick.bCanEverTick = true;
 	PrimaryActorTick.TickGroup = TG_PrePhysics;
-	AIControllerClass = AAIController::StaticClass();
+	if (HasAnyFlags(RF_ClassDefaultObject) && GetClass() == APawn::StaticClass())
+	{
+		AIControllerClass = LoadClass<AController>(NULL, *((UEngine*)(UEngine::StaticClass()->GetDefaultObject()))->AIControllerClassName.ToString(), NULL, LOAD_None, NULL);
+	}
 	bCanBeDamaged = true;
 	
 	SetRemoteRoleForBackwardsCompat(ROLE_SimulatedProxy);
