Expose text render functions from component on TextRenderActor
Remove TextRenderActor.h from Engine/h and EnginePrivate.h

diff --git a/Engine/Source/Editor/UnrealEd/Private/Factories/ActorFactory.cpp b/Engine/Source/Editor/UnrealEd/Private/Factories/ActorFactory.cpp
--- a/Engine/Source/Editor/UnrealEd/Private/Factories/ActorFactory.cpp
+++ b/Engine/Source/Editor/UnrealEd/Private/Factories/ActorFactory.cpp
@@ -31,6 +31,7 @@ ActorFactory.cpp:
 #include "Engine/TriggerBox.h"
 #include "Engine/TriggerSphere.h"
 #include "Engine/TriggerCapsule.h"
+#include "Engine/TextRenderActor.h"
 
 DEFINE_LOG_CATEGORY_STATIC(LogActorFactory, Log, All);
 
diff --git a/Engine/Source/Runtime/Engine/Classes/Engine/TextRenderActor.h b/Engine/Source/Runtime/Engine/Classes/Engine/TextRenderActor.h
--- a/Engine/Source/Runtime/Engine/Classes/Engine/TextRenderActor.h
+++ b/Engine/Source/Runtime/Engine/Classes/Engine/TextRenderActor.h
@@ -16,7 +16,7 @@ class ATextRenderActor : public AActor
 	friend class UActorFactoryTextRender;
 
 	/** Component to render a text in 3d with a font */
-	UPROPERTY(Category=TextRenderActor, VisibleAnywhere, BlueprintReadOnly)
+	UPROPERTY(Category=TextRenderActor, VisibleAnywhere, BlueprintReadOnly, meta=(ExposeFunctionCategories="Rendering|Components|TextRender"))
 	TSubobjectPtr<class UTextRenderComponent> TextRender;
 
 #if WITH_EDITORONLY_DATA
diff --git a/Engine/Source/Runtime/Engine/Private/EnginePrivate.h b/Engine/Source/Runtime/Engine/Private/EnginePrivate.h
--- a/Engine/Source/Runtime/Engine/Private/EnginePrivate.h
+++ b/Engine/Source/Runtime/Engine/Private/EnginePrivate.h
@@ -80,7 +80,6 @@
 #include "PhysicsEngine/PhysicsConstraintComponent.h"
 #include "Components/SkeletalMeshComponent.h"
 #include "Engine/TargetPoint.h"
-#include "Engine/TextRenderActor.h"
 #include "Engine/DataAsset.h"
 #include "Components/PawnNoiseEmitterComponent.h"
 #include "PhysicsEngine/PhysicsHandleComponent.h"
diff --git a/Engine/Source/Runtime/Engine/Private/TextRenderComponent.cpp b/Engine/Source/Runtime/Engine/Private/TextRenderComponent.cpp
--- a/Engine/Source/Runtime/Engine/Private/TextRenderComponent.cpp
+++ b/Engine/Source/Runtime/Engine/Private/TextRenderComponent.cpp
@@ -2,7 +2,7 @@
 
 #include "EnginePrivate.h"
 #include "DynamicMeshBuilder.h"
-#include "../../../../Source/Runtime/Engine/Classes/Engine/TextRenderActor.h"
+#include "Engine/TextRenderActor.h"
 #include "LocalVertexFactory.h"
 
 ATextRenderActor::ATextRenderActor(const class FPostConstructInitializeProperties& PCIP)
diff --git a/Engine/Source/Runtime/Engine/Public/Engine.h b/Engine/Source/Runtime/Engine/Public/Engine.h
--- a/Engine/Source/Runtime/Engine/Public/Engine.h
+++ b/Engine/Source/Runtime/Engine/Public/Engine.h
@@ -115,7 +115,6 @@
 #include "PhysicsEngine/PhysicsConstraintComponent.h"
 #include "Engine/StaticMeshActor.h"
 #include "Engine/TargetPoint.h"
-#include "Engine/TextRenderActor.h"
 #include "VectorField/VectorFieldVolume.h"
 #include "Engine/DataAsset.h"
 #include "GameFramework/SpectatorPawnMovement.h"
