Added a closed loop property to SplineComponent.  Also fixed small rendering bug (local space widget) in SplineComponent visualizer.
#ttp 338259 - SplineComponent: Add an option flag to close the spline
#branch UE4
#proj Editor.ComponentVisualizers, Runtime.Engine

diff --git a/Engine/Source/Editor/ComponentVisualizers/Private/SplineComponentVisualizer.cpp b/Engine/Source/Editor/ComponentVisualizers/Private/SplineComponentVisualizer.cpp
--- a/Engine/Source/Editor/ComponentVisualizers/Private/SplineComponentVisualizer.cpp
+++ b/Engine/Source/Editor/ComponentVisualizers/Private/SplineComponentVisualizer.cpp
@@ -83,7 +83,7 @@ void FSplineComponentVisualizer::OnRegister()
 	SplineComponentVisualizerActions->MapAction(
 		Commands.DeleteKey,
 		FExecuteAction::CreateSP(this, &FSplineComponentVisualizer::OnDeleteKey),
-		FCanExecuteAction::CreateSP(this, &FSplineComponentVisualizer::IsKeySelectionValid));
+		FCanExecuteAction::CreateSP(this, &FSplineComponentVisualizer::CanDeleteKey));
 
 	SplineComponentVisualizerActions->MapAction(
 		Commands.DuplicateKey,
@@ -148,7 +148,7 @@ void FSplineComponentVisualizer::DrawVisualization(const UActorComponent* Compon
 
 		FVector OldKeyPos(0);
 		float OldKeyTime = 0.f;
-		for (int32 KeyIdx = 0; KeyIdx<SplineInfo.Points.Num(); KeyIdx++)
+		for (int32 KeyIdx = 0; KeyIdx < SplineInfo.Points.Num(); KeyIdx++)
 		{
 			float NewKeyTime = SplineInfo.Points[KeyIdx].InVal;
 			FVector NewKeyPos = SplineComp->ComponentToWorld.TransformPosition( SplineInfo.Eval(NewKeyTime, FVector::ZeroVector) );
@@ -156,9 +156,12 @@ void FSplineComponentVisualizer::DrawVisualization(const UActorComponent* Compon
 			const FColor KeyColor = (SplineComp == EditedSplineComp && KeyIdx == SelectedKeyIndex) ? SelectedColor : NormalColor;
 
 			// Draw the keypoint
-			PDI->SetHitProxy(new HSplineKeyProxy(Component, KeyIdx));
-			PDI->DrawPoint(NewKeyPos, KeyColor, GrabHandleSize, SDPG_Foreground);
-			PDI->SetHitProxy(NULL);
+			if (!SplineComp->IsClosedLoop() || KeyIdx < SplineInfo.Points.Num() - 1)
+			{
+				PDI->SetHitProxy(new HSplineKeyProxy(Component, KeyIdx));
+				PDI->DrawPoint(NewKeyPos, KeyColor, GrabHandleSize, SDPG_Foreground);
+				PDI->SetHitProxy(NULL);
+			}
 
 			// If not the first keypoint, draw a line to the previous keypoint.
 			if (KeyIdx>0)
@@ -303,10 +306,8 @@ bool FSplineComponentVisualizer::GetCustomInputCoordinateSystem(const FEditorVie
 			if (SelectedKeyIndex < SplineComp->SplineInfo.Points.Num())
 			{
 				const auto& Point = SplineComp->SplineInfo.Points[SelectedKeyIndex];
-				const FVector Tangent = Point.ArriveTangent.SafeNormal();
-				const FVector Bitangent = (Tangent.Z == 1.0f) ?
-					FVector(1.0f, 0.0f, 0.0f) :
-					FVector(-Tangent.Y, Tangent.X, 0.0f).SafeNormal();
+				const FVector Tangent = Point.ArriveTangent.IsNearlyZero() ? FVector(1.0f, 0.0f, 0.0f) : Point.ArriveTangent.SafeNormal();
+				const FVector Bitangent = (Tangent.Z == 1.0f) ? FVector(1.0f, 0.0f, 0.0f) : FVector(-Tangent.Y, Tangent.X, 0.0f).SafeNormal();
 				const FVector Normal = FVector::CrossProduct(Tangent, Bitangent);
 
 				OutMatrix = FMatrix(Tangent, Bitangent, Normal, FVector::ZeroVector) * FQuatRotationTranslationMatrix(SplineComp->ComponentToWorld.GetRotation(), FVector::ZeroVector);
@@ -322,9 +323,11 @@ bool FSplineComponentVisualizer::GetCustomInputCoordinateSystem(const FEditorVie
 bool FSplineComponentVisualizer::HandleInputDelta(FEditorViewportClient* ViewportClient, FViewport* Viewport, FVector& DeltaTranslate, FRotator& DeltaRotate, FVector& DeltaScale)
 {
 	USplineComponent* SplineComp = GetEditedSplineComponent();
-	if( SplineComp && 
+	const int32 NumPoints = SplineComp->SplineInfo.Points.Num();
+
+	if (SplineComp && 
 		SelectedKeyIndex != INDEX_NONE &&
-		SelectedKeyIndex < SplineComp->SplineInfo.Points.Num() )
+		SelectedKeyIndex < NumPoints)
 	{
 		if (ViewportClient->IsAltPressed() && bAllowDuplication)
 		{
@@ -382,6 +385,14 @@ bool FSplineComponentVisualizer::HandleInputDelta(FEditorViewportClient* Viewpor
 			EditedPoint.ArriveTangent = NewTangent;
 		}
 
+		// Duplicate first point to last point if it is a closed loop
+		if (SplineComp->IsClosedLoop() && SelectedKeyIndex == 0 && NumPoints > 0)
+		{
+			SplineComp->SplineInfo.Points[NumPoints - 1].OutVal = EditedPoint.OutVal;
+			SplineComp->SplineInfo.Points[NumPoints - 1].LeaveTangent = EditedPoint.LeaveTangent;
+			SplineComp->SplineInfo.Points[NumPoints - 1].ArriveTangent = EditedPoint.ArriveTangent;
+		}
+
 		NotifyComponentModified();
 
 		return true;
@@ -427,6 +438,9 @@ void FSplineComponentVisualizer::OnDuplicateKey()
 		Owner->Modify();
 	}
 
+	const bool bWasLoop = SplineComp->IsClosedLoop();
+	SplineComp->SetClosedLoop(false);
+
 	FInterpCurvePoint<FVector> KeyToDupe = SplineComp->SplineInfo.Points[SelectedKeyIndex];
 	KeyToDupe.InterpMode = CIM_CurveAuto;
 	int32 NewKeyIndex = SplineComp->SplineInfo.Points.Insert(KeyToDupe, SelectedKeyIndex);
@@ -436,6 +450,8 @@ void FSplineComponentVisualizer::OnDuplicateKey()
 	// Update Input value for all keys
 	SplineComp->RefreshSplineInputs();
 
+	SplineComp->SetClosedLoop(bWasLoop);
+
 	NotifyComponentModified();
 }
 
@@ -459,6 +475,9 @@ void FSplineComponentVisualizer::OnAddKey()
 		Owner->Modify();
 	}
 
+	const bool bWasLoop = SplineComp->IsClosedLoop();
+	SplineComp->SetClosedLoop(false);
+
 	FInterpCurvePoint<FVector> NewKey;
 	NewKey.InVal = SelectedSegmentIndex;
 	NewKey.OutVal = SplineComp->ComponentToWorld.InverseTransformPosition(SelectedSplinePosition);
@@ -471,6 +490,8 @@ void FSplineComponentVisualizer::OnAddKey()
 	// Update Input value for all keys
 	SplineComp->RefreshSplineInputs();
 
+	SplineComp->SetClosedLoop(bWasLoop);
+
 	NotifyComponentModified();
 }
 
@@ -485,9 +506,14 @@ void FSplineComponentVisualizer::OnDeleteKey()
 		Owner->Modify();
 	}
 
+	const bool bWasLoop = SplineComp->IsClosedLoop();
+	SplineComp->SetClosedLoop(false);
+
 	SplineComp->SplineInfo.Points.RemoveAt(SelectedKeyIndex);
 	SplineComp->RefreshSplineInputs(); // update input value for each key
 
+	SplineComp->SetClosedLoop(bWasLoop);
+
 	SelectedKeyIndex = INDEX_NONE; // deselect any keys
 	SelectedSegmentIndex = INDEX_NONE;
 
@@ -495,6 +521,17 @@ void FSplineComponentVisualizer::OnDeleteKey()
 }
 
 
+bool FSplineComponentVisualizer::CanDeleteKey() const
+{
+	USplineComponent* SplineComp = GetEditedSplineComponent();
+	return (SplineComp &&
+		SplineComp->SplineInfo.Points.Num() >= 2 &&
+		SelectedSegmentIndex == INDEX_NONE &&
+		SelectedKeyIndex != INDEX_NONE &&
+		SelectedKeyIndex < SplineComp->SplineInfo.Points.Num());
+}
+
+
 bool FSplineComponentVisualizer::IsKeySelectionValid() const
 {
 	USplineComponent* SplineComp = GetEditedSplineComponent();
@@ -517,7 +554,10 @@ void FSplineComponentVisualizer::OnResetToAutomaticTangent(EInterpCurveMode Mode
 			Owner->Modify();
 		}
 
+		const bool bWasLoop = SplineComp->IsClosedLoop();
+		SplineComp->SetClosedLoop(false);
 		SplineComp->SplineInfo.Points[SelectedKeyIndex].InterpMode = Mode;
+		SplineComp->SetClosedLoop(bWasLoop);
 		NotifyComponentModified();
 	}
 }
@@ -546,7 +586,10 @@ void FSplineComponentVisualizer::OnSetKeyType(EInterpCurveMode Mode)
 			Owner->Modify();
 		}
 
+		const bool bWasLoop = SplineComp->IsClosedLoop();
+		SplineComp->SetClosedLoop(false);
 		SplineComp->SplineInfo.Points[SelectedKeyIndex].InterpMode = Mode;
+		SplineComp->SetClosedLoop(bWasLoop);
 		NotifyComponentModified();
 	}
 }
diff --git a/Engine/Source/Editor/ComponentVisualizers/Private/SplineComponentVisualizer.h b/Engine/Source/Editor/ComponentVisualizers/Private/SplineComponentVisualizer.h
--- a/Engine/Source/Editor/ComponentVisualizers/Private/SplineComponentVisualizer.h
+++ b/Engine/Source/Editor/ComponentVisualizers/Private/SplineComponentVisualizer.h
@@ -68,6 +68,8 @@ private:
 	void NotifyComponentModified();
 
 	void OnDeleteKey();
+	bool CanDeleteKey() const;
+
 	void OnDuplicateKey();
 	bool IsKeySelectionValid() const;
 
diff --git a/Engine/Source/Runtime/Engine/Classes/Components/SplineComponent.h b/Engine/Source/Runtime/Engine/Classes/Components/SplineComponent.h
--- a/Engine/Source/Runtime/Engine/Classes/Components/SplineComponent.h
+++ b/Engine/Source/Runtime/Engine/Classes/Components/SplineComponent.h
@@ -32,9 +32,18 @@ class ENGINE_API USplineComponent : public USceneComponent
 	float Duration;
 
 	/** Whether the endpoints of the spline are considered stationary when traversing the spline at non-constant velocity.  Essentially this sets the endpoints' tangents to zero vectors. */
-	UPROPERTY(EditAnywhere, BlueprintReadWrite, AdvancedDisplay, Category = Spline)
+	UPROPERTY(EditAnywhere, BlueprintReadWrite, AdvancedDisplay, Category = Spline, meta=(EditCondition="!bClosedLoop"))
 	bool bStationaryEndpoints;
 
+private:
+	/**
+	 * Whether the spline is to be considered as a closed loop.
+	 * Use SetClosedLoop() to set this property, and IsClosedLoop() to read it.
+	 */
+	UPROPERTY(EditAnywhere, Category = Spline)
+	bool bClosedLoop;
+
+public:
 	// Begin UActorComponent interface.
 	virtual TSharedPtr<FComponentInstanceDataBase> GetComponentInstanceData() const override;
 	virtual FName GetComponentInstanceDataType() const override;
@@ -44,6 +53,14 @@ class ENGINE_API USplineComponent : public USceneComponent
 	/** Update the spline tangents and SplineReparamTable */
 	void UpdateSpline();
 
+	/** Specify whether the spline is a closed loop or not */
+	UFUNCTION(BlueprintCallable, Category = Spline)
+	void SetClosedLoop(bool bInClosedLoop);
+
+	/** Check whether the spline is a closed loop or not */
+	UFUNCTION(BlueprintCallable, Category = Spline)
+	bool IsClosedLoop() const;
+
 	/** Clears all the points in the spline */
 	UFUNCTION(BlueprintCallable, Category=Spline)
 	void ClearSplinePoints();
@@ -124,6 +141,7 @@ class ENGINE_API USplineComponent : public USceneComponent
 	virtual void PostLoad() override;
 	virtual void PostEditImport() override;
 #if WITH_EDITOR
+	virtual void PreEditChange(UProperty* PropertyAboutToChange) override;
 	virtual void PostEditChangeProperty(FPropertyChangedEvent& PropertyChangedEvent) override;
 	// End of UObject interface
 #endif
@@ -132,8 +150,17 @@ class ENGINE_API USplineComponent : public USceneComponent
 	void RefreshSplineInputs();
 
 private:
+	/** Returns the length of the specified spline segment up to the parametric value given */
 	float GetSegmentLength(const int32 Index, const float Param = 1.0f) const;
+
+	/** Returns the parametric value t which would result in a spline segment of the given length between S(0)...S(t) */
 	float GetSegmentParamFromLength(const int32 Index, const float Length, const float SegmentLength) const;
+
+	/** Adds an additional point to the end of the spline, whose OutValue and tangents coincides with the first point */
+	void AddLoopEndpoint();
+
+	/** Removes the final point from the spline, assuming it to be a duplicate of the first point */
+	void RemoveLoopEndpoint();
 };
 
 
diff --git a/Engine/Source/Runtime/Engine/Private/SplineComponent.cpp b/Engine/Source/Runtime/Engine/Private/SplineComponent.cpp
--- a/Engine/Source/Runtime/Engine/Private/SplineComponent.cpp
+++ b/Engine/Source/Runtime/Engine/Private/SplineComponent.cpp
@@ -15,6 +15,7 @@ USplineComponent::USplineComponent(const class FPostConstructInitializePropertie
 	, ReparamStepsPerSegment(10)
 	, Duration(1.0f)
 	, bStationaryEndpoints(false)
+	, bClosedLoop(false)
 {
 	SplineInfo.Points.Reset(10);
 
@@ -45,16 +46,45 @@ void USplineComponent::PostEditImport()
 
 void USplineComponent::UpdateSpline()
 {
+	const int32 NumPoints = SplineInfo.Points.Num();
+	check(!bClosedLoop || NumPoints == 0 || (NumPoints >= 2 && SplineInfo.Points[0].OutVal == SplineInfo.Points[NumPoints - 1].OutVal));
+
 	// Automatically set the tangents on any CurveAuto keys
 	SplineInfo.AutoSetTangents(0.0f, bStationaryEndpoints);
 
-	// Nothing to do if no points
-	if(SplineInfo.Points.Num() < 2)
+	// Nothing else to do if less than 2 points
+	if (NumPoints < 2)
 	{
 		return;
 	}
 
-	const int32 NumSegments = SplineInfo.Points.Num() - 1;
+	// Adjust auto tangents for first and last keys to take into account the looping
+	if (bClosedLoop)
+	{
+		auto& FirstPoint = SplineInfo.Points[0];
+		auto& LastPoint = SplineInfo.Points[NumPoints - 1];
+		const auto& SecondPoint = SplineInfo.Points[1];
+		const auto& PenultimatePoint = SplineInfo.Points[NumPoints - 2];
+
+		if (FirstPoint.InterpMode == CIM_CurveAuto || FirstPoint.InterpMode == CIM_CurveAutoClamped)
+		{
+			FVector Tangent;
+			ComputeCurveTangent(
+				PenultimatePoint.InVal - LastPoint.InVal, PenultimatePoint.OutVal,
+				FirstPoint.InVal, FirstPoint.OutVal,
+				SecondPoint.InVal, SecondPoint.OutVal,
+				0.0f,
+				FirstPoint.InterpMode == CIM_CurveAutoClamped,
+				Tangent);
+
+			FirstPoint.LeaveTangent = Tangent;
+			FirstPoint.ArriveTangent = Tangent;
+			LastPoint.LeaveTangent = Tangent;
+			LastPoint.ArriveTangent = Tangent;
+		}
+	}
+
+	const int32 NumSegments = NumPoints - 1;
 
 	// Start by clearing it
 	SplineReparamTable.Points.Reset(NumSegments * ReparamStepsPerSegment + 1);
@@ -174,6 +204,55 @@ float USplineComponent::GetSegmentParamFromLength(const int32 Index, const float
 }
 
 
+void USplineComponent::SetClosedLoop(bool bInClosedLoop)
+{
+	if (bClosedLoop != bInClosedLoop)
+	{
+		bClosedLoop = bInClosedLoop;
+
+		if (bClosedLoop)
+		{
+			AddLoopEndpoint();
+		}
+		else
+		{
+			RemoveLoopEndpoint();
+		}
+	}
+}
+
+
+bool USplineComponent::IsClosedLoop() const
+{
+	return bClosedLoop;
+}
+
+
+void USplineComponent::RemoveLoopEndpoint()
+{
+	const int32 NumPoints = SplineInfo.Points.Num();
+	check(!bClosedLoop);
+	check(NumPoints == 0 || (NumPoints >= 2 && SplineInfo.Points[0].OutVal == SplineInfo.Points[NumPoints - 1].OutVal));
+	if (NumPoints > 0)
+	{
+		SplineInfo.Points.RemoveAt(NumPoints - 1, 1, false);
+	}
+}
+
+
+void USplineComponent::AddLoopEndpoint()
+{
+	check(bClosedLoop);
+	const int32 NumPoints = SplineInfo.Points.Num();
+	if (NumPoints > 0)
+	{
+		FInterpCurvePoint<FVector> EndPoint(SplineInfo.Points[0]);
+		EndPoint.InVal = static_cast<float>(NumPoints);
+		SplineInfo.Points.Add(EndPoint);
+	}
+}
+
+
 void USplineComponent::ClearSplinePoints()
 {
 	SplineInfo.Reset();
@@ -183,27 +262,45 @@ void USplineComponent::ClearSplinePoints()
 
 void USplineComponent::AddSplineWorldPoint(const FVector& Position)
 {
+	// If it's a closed loop, remove the endpoint before adding a new point
+	const bool bWasLoop = IsClosedLoop();
+	SetClosedLoop(false);
+
 	float InputKey = static_cast<float>(SplineInfo.Points.Num());
-	int32 PointIndex = SplineInfo.AddPoint(InputKey, ComponentToWorld.InverseTransformPosition(Position));
+	const int32 PointIndex = SplineInfo.AddPoint(InputKey, ComponentToWorld.InverseTransformPosition(Position));
 	SplineInfo.Points[PointIndex].InterpMode = CIM_CurveAuto;
 
+	// Then re-close the spline if required
+	SetClosedLoop(bWasLoop);
+
 	UpdateSpline();
 }
 
 
 void USplineComponent::AddSplineLocalPoint(const FVector& Position)
 {
+	// If it's a closed loop, remove the endpoint before adding a new point
+	const bool bWasLoop = IsClosedLoop();
+	SetClosedLoop(false);
+
 	float InputKey = static_cast<float>(SplineInfo.Points.Num());
-	int32 PointIndex = SplineInfo.AddPoint(InputKey, Position);
+	const int32 PointIndex = SplineInfo.AddPoint(InputKey, Position);
 	SplineInfo.Points[PointIndex].InterpMode = CIM_CurveAuto;
 
+	// Then re-close the spline if required
+	SetClosedLoop(bWasLoop);
+
 	UpdateSpline();
 }
 
 
 void USplineComponent::SetSplineWorldPoints(const TArray<FVector>& Points)
 {
-	SplineInfo.Points.Reset(Points.Num());
+	// If it's a closed loop, mark it as not closed before setting a new array of points
+	const bool bWasLoop = IsClosedLoop();
+	SetClosedLoop(false);
+
+	SplineInfo.Points.Reset(bWasLoop ? Points.Num() + 1 : Points.Num());
 	float InputKey = 0.0f;
 	for (const auto& Point : Points)
 	{
@@ -212,13 +309,18 @@ void USplineComponent::SetSplineWorldPoints(const TArray<FVector>& Points)
 		InputKey += 1.0f;
 	}
 
+	SetClosedLoop(bWasLoop);
+
 	UpdateSpline();
 }
 
 
 void USplineComponent::SetSplineLocalPoints(const TArray<FVector>& Points)
 {
-	SplineInfo.Points.Reset(Points.Num());
+	const bool bWasLoop = IsClosedLoop();
+	SetClosedLoop(false);
+
+	SplineInfo.Points.Reset(bWasLoop ? Points.Num() + 1 : Points.Num());
 	float InputKey = 0.0f;
 	for (const auto& Point : Points)
 	{
@@ -227,15 +329,33 @@ void USplineComponent::SetSplineLocalPoints(const TArray<FVector>& Points)
 		InputKey += 1.0f;
 	}
 
+	SetClosedLoop(bWasLoop);
+
 	UpdateSpline();
 }
 
 
 void USplineComponent::SetWorldLocationAtSplinePoint(int32 PointIndex, const FVector& InLocation)
 {
-	if ((PointIndex >= 0) && (PointIndex < SplineInfo.Points.Num()))
+	const int32 NumPoints = SplineInfo.Points.Num();
+
+	if ((PointIndex >= 0) && (PointIndex < NumPoints))
 	{
 		SplineInfo.Points[PointIndex].OutVal = ComponentToWorld.InverseTransformPosition(InLocation);
+
+		if (IsClosedLoop())
+		{
+			// In a closed loop, the first and last points are tied, so update one with the other
+			if (PointIndex == 0)
+			{
+				SplineInfo.Points[NumPoints - 1].OutVal = InLocation;
+			}
+			else if (PointIndex == NumPoints - 1)
+			{
+				SplineInfo.Points[0].OutVal = InLocation;
+			}
+		}
+
 		UpdateSpline();
 	}
 }
@@ -397,9 +517,7 @@ public:
 	}
 
 	FInterpCurveVector SplineInfo;
-	int32 ReparamStepsPerSegment;
-	float Duration;
-	bool bStationaryEndpoints;
+	bool bClosedLoop;
 };
 
 FName USplineComponent::GetComponentInstanceDataType() const
@@ -415,9 +533,7 @@ TSharedPtr<FComponentInstanceDataBase> USplineComponent::GetComponentInstanceDat
 	{
 		SplineInstanceData = MakeShareable(new FSplineInstanceData(this));
 		SplineInstanceData->SplineInfo = SplineInfo;
-		SplineInstanceData->ReparamStepsPerSegment = ReparamStepsPerSegment;
-		SplineInstanceData->Duration = Duration;
-		SplineInstanceData->bStationaryEndpoints = bStationaryEndpoints;
+		SplineInstanceData->bClosedLoop = bClosedLoop;
 	}
 	return SplineInstanceData;
 }
@@ -430,15 +546,44 @@ void USplineComponent::ApplyComponentInstanceData(TSharedPtr<FComponentInstanceD
 		if (bAllowSplineEditingPerInstance)
 		{
 			SplineInfo = SplineInstanceData->SplineInfo;
-			ReparamStepsPerSegment = SplineInstanceData->ReparamStepsPerSegment;
-			Duration = SplineInstanceData->Duration;
-			bStationaryEndpoints = SplineInstanceData->bStationaryEndpoints;
+
+			// If the construction script changed bClosedLoop, amend the applied points accordingly
+			if (SplineInstanceData->bClosedLoop != bClosedLoop)
+			{
+				if (bClosedLoop)
+				{
+					AddLoopEndpoint();
+				}
+				else
+				{
+					RemoveLoopEndpoint();
+				}
+			}
+
 			UpdateSpline();
 		}
 	}
 }
 
 #if WITH_EDITOR
+void USplineComponent::PreEditChange(UProperty* PropertyAboutToChange)
+{
+	Super::PreEditChange(PropertyAboutToChange);
+
+	if (PropertyAboutToChange != nullptr && PropertyAboutToChange->GetFName() == GET_MEMBER_NAME_CHECKED(USplineComponent, bClosedLoop))
+	{
+		if (bClosedLoop)
+		{
+			// We need to detect changes in state in order to correctly open or close the loop.
+			// Since we don't know what state is about to be set, we ensure that the spline is marked as not closed here, so that either:
+			//  a) Nothing needs to be done if the property is set to false, or
+			//  b) The spline is already set as not looped, so that setting the property to true will work correctly in PostEditChangeProperty.
+			bClosedLoop = false;
+			RemoveLoopEndpoint();
+		}
+	}
+}
+
 void USplineComponent::PostEditChangeProperty(FPropertyChangedEvent& PropertyChangedEvent)
 {
 	if (PropertyChangedEvent.Property != nullptr)
@@ -449,6 +594,18 @@ void USplineComponent::PostEditChangeProperty(FPropertyChangedEvent& PropertyCha
 		{
 			UpdateSpline();
 		}
+
+		if (PropertyName == GET_MEMBER_NAME_CHECKED(USplineComponent, bClosedLoop))
+		{
+			if (bClosedLoop)
+			{
+				// Spline is guaranteed to be non-looping when we get here (due to PreEditChange).
+				// Now just force the loop endpoint to be added if bClosedLoop == true.
+				AddLoopEndpoint();
+			}
+
+			UpdateSpline();
+		}
 	}
 
 	Super::PostEditChangeProperty(PropertyChangedEvent);
