Remove dead code

diff --git a/Engine/Source/Runtime/Engine/Private/RepLayout.cpp b/Engine/Source/Runtime/Engine/Private/RepLayout.cpp
--- a/Engine/Source/Runtime/Engine/Private/RepLayout.cpp
+++ b/Engine/Source/Runtime/Engine/Private/RepLayout.cpp
@@ -2280,13 +2280,6 @@ void FRepLayout::SendPropertiesForRPC( UObject * Object, UFunction * Function, U
 
 	for ( int32 i = 0; i < Parents.Num(); i++ )
 	{
-		/*
-		if ( !Writer.PackageMap->SupportsObject( Parents[i].Property ) )
-		{
-			continue;
-		}
-		*/
-
 		bool Send = true;
 
 		if ( !Cast<UBoolProperty>( Parents[i].Property ) )
@@ -2326,13 +2319,6 @@ void FRepLayout::ReceivePropertiesForRPC( UObject * Object, UFunction * Function
 
 	for ( int32 i = 0; i < Parents.Num(); i++ )
 	{
-		/*
-		if ( !Reader.PackageMap->SupportsObject( Parents[i].Property ) )
-		{
-			continue;
-		}
-		*/
-
 		if ( Cast<UBoolProperty>( Parents[i].Property ) || Reader.ReadBit() ) 
 		{
 			bool bHasUnmapped = false;
@@ -2359,13 +2345,6 @@ void FRepLayout::SerializePropertiesForStruct( UStruct * Struct, FArchive & Ar,
 
 	for ( int32 i = 0; i < Parents.Num(); i++ )
 	{
-		/*
-		if ( Map != NULL && !Map->SupportsObject( Parents[i].Property ) )
-		{
-			continue;
-		}
-		*/
-
 		SerializeProperties_r( Ar, Map, Parents[i].CmdStart, Parents[i].CmdEnd, Data, bHasUnmapped );
 
 		if ( Ar.IsError() )
