UE4 - Compile fix
[BuildWatcher]

diff --git a/Engine/Source/Runtime/Engine/Classes/Animation/AnimNode_LayeredBoneBlend.h b/Engine/Source/Runtime/Engine/Classes/Animation/AnimNode_LayeredBoneBlend.h
--- a/Engine/Source/Runtime/Engine/Classes/Animation/AnimNode_LayeredBoneBlend.h
+++ b/Engine/Source/Runtime/Engine/Classes/Animation/AnimNode_LayeredBoneBlend.h
@@ -38,7 +38,9 @@ protected:
 public:	
 	FAnimNode_LayeredBoneBlend()
 	{
+#if WITH_EDITOR
 		AddPose();
+#endif
 	}
 
 	// FAnimNode_Base interface
