#docs - FTransform class comments expanded to include description of how transformations are applied and composed. Similar changes to KismetMathLibrary::ComposeTransforms() as well.

diff --git a/Engine/Source/Runtime/Core/Public/Math/Transform.h b/Engine/Source/Runtime/Core/Public/Math/Transform.h
--- a/Engine/Source/Runtime/Core/Public/Math/Transform.h
+++ b/Engine/Source/Runtime/Core/Public/Math/Transform.h
@@ -4,8 +4,21 @@
 #pragma once
 
 /** 
- * Transform composed of Quat/Translation/Scale.
+ * Transform composed of Scale, Rotation (as a quaternion), and Translation.
+ *
+ * Transforms can be used to convert from one space to another, for example by transforming
+ * positions and directions from local space to world space.
+ *
+ * Transformation of position vectors is applied in the order:  Scale -> Rotate -> Translate.
+ * Transformation of direction vectors is applied in the order: Scale -> Rotate.
+ *
+ * Order matters when composing transforms: C = A * B will yield a transform C that logically
+ * first applies B then A to any subsequent transformation.
+ *
+ * Example: LocalToWorld = (LocalToWorld * DeltaRotation) will change rotation in local space by DeltaRotation.
+ * Example: LocalToWorld = (DeltaRotation * LocalToWorld) will change rotation in world space by DeltaRotation.
  */
+
 class FTransform
 {
 #if !defined(COREUOBJECT_API)
@@ -438,10 +451,40 @@ public:
 		return *this;
 	}
 
-	FORCEINLINE FTransform		operator*(const FTransform& Other) const;
-	FORCEINLINE void			operator*=(const FTransform& Other);
-	FORCEINLINE FTransform		operator*(const FQuat& Other) const;
-	FORCEINLINE void			operator*=(const FQuat& Other);
+	/**
+	 * Return a transform that is the result of this multiplied by another transform.
+	 * Order matters when composing transforms : C = A * B will yield a transform C that logically first applies B then A to any subsequent transformation.
+	 * 
+	 * @param  Other other transform by which to multiply.
+	 * @return new transform: this * Other
+	 */
+	FORCEINLINE FTransform operator*(const FTransform& Other) const;
+
+
+	/**
+	 * Sets this transform to the result of this multiplied by another transform.
+	 * Order matters when composing transforms : C = A * B will yield a transform C that logically first applies B then A to any subsequent transformation.
+	 *
+	 * @param  Other other transform by which to multiply.
+	 */
+	FORCEINLINE void operator*=(const FTransform& Other);
+	
+	/**
+	* Return a transform that is the result of this multiplied by another transform (made only from a rotation).
+	* Order matters when composing transforms : C = A * B will yield a transform C that logically first applies B then A to any subsequent transformation.
+	*
+	* @param  Other other quaternion rotation by which to multiply.
+	* @return new transform: this * FTransform(Other)
+	*/
+	FORCEINLINE FTransform operator*(const FQuat& Other) const;
+	
+	/**
+	* Sets this transform to the result of this multiplied by another transform (made only from a rotation).
+	* Order matters when composing transforms : C = A * B will yield a transform C that logically first applies B then A to any subsequent transformation.
+	*
+	* @param  Other other quaternion rotation by which to multiply.
+	*/
+	FORCEINLINE void operator*=(const FQuat& Other);
 
 	FORCEINLINE void ScaleTranslation(const FVector& Scale3D);
 	FORCEINLINE void ScaleTranslation(const float& Scale);
@@ -573,6 +616,15 @@ public:
 		return Rotation.Equals(Other.Rotation, Tolerance) && Translation.Equals(Other.Translation, Tolerance);
 	}
 
+	/**
+	 * Create a new transform: OutTransform = A * B.
+	 *
+	 * Order matters when composing transforms : A * B will yield a transform that logically first applies B then A to any subsequent transformation.
+	 * 
+	 * @param  OutTransform pointer to transform that will store the result of A * B.
+	 * @param  A Transform A.
+	 * @param  B Transform B.
+	 */
 	FORCEINLINE static void Multiply(FTransform * OutTransform, const FTransform * A, const FTransform * B);
 
 	/**
diff --git a/Engine/Source/Runtime/Core/Public/Math/TransformVectorized.h b/Engine/Source/Runtime/Core/Public/Math/TransformVectorized.h
--- a/Engine/Source/Runtime/Core/Public/Math/TransformVectorized.h
+++ b/Engine/Source/Runtime/Core/Public/Math/TransformVectorized.h
@@ -6,9 +6,22 @@
 /** This returns Quaternion Inverse of X **/
 #define MAKE_QUATINV_VECTORREGISTER(X) VectorMultiply(GlobalVectorConstants::QINV_SIGN_MASK, X)
 
-/** 
- * FTransform class for Quat/Translation/Scale.
- */
+/**
+* Transform composed of Scale, Rotation (as a quaternion), and Translation.
+*
+* Transforms can be used to convert from one space to another, for example by transforming
+* positions and directions from local space to world space.
+*
+* Transformation of position vectors is applied in the order:  Scale -> Rotate -> Translate.
+* Transformation of direction vectors is applied in the order: Scale -> Rotate.
+*
+* Order matters when composing transforms: C = A * B will yield a transform C that logically
+* first applies B then A to any subsequent transformation.
+*
+* Example: LocalToWorld = (LocalToWorld * DeltaRotation) will change rotation in local space by DeltaRotation.
+* Example: LocalToWorld = (DeltaRotation * LocalToWorld) will change rotation in world space by DeltaRotation.
+*/
+
 MS_ALIGN(16) class FTransform
 {
 #if !defined(COREUOBJECT_API)
diff --git a/Engine/Source/Runtime/Engine/Classes/Kismet/KismetMathLibrary.h b/Engine/Source/Runtime/Engine/Classes/Kismet/KismetMathLibrary.h
--- a/Engine/Source/Runtime/Engine/Classes/Kismet/KismetMathLibrary.h
+++ b/Engine/Source/Runtime/Engine/Classes/Kismet/KismetMathLibrary.h
@@ -1190,7 +1190,7 @@ class UKismetMathLibrary : public UBlueprintFunctionLibrary
 	//
 	
 	/** 
-	 *	Transform a position by the supplied transform 
+	 *	Transform a position by the supplied transform.
 	 *	For example, if T was an object's transform, would transform a position from local space to world space.
 	 */
 	UFUNCTION(BlueprintPure, Category="Math|Transform", meta=(Keywords="location"))
@@ -1204,25 +1204,35 @@ class UKismetMathLibrary : public UBlueprintFunctionLibrary
 	static FVector TransformDirection(const FTransform& T, FVector Direction);
 
 	/** 
-	 *	Transform a position by the inverse of the supplied transform 
+	 *	Transform a position by the inverse of the supplied transform.
 	 *	For example, if T was an object's transform, would transform a position from world space to local space.
 	 */
 	UFUNCTION(BlueprintPure, Category="Math|Transform", meta=(Keywords="location"))
 	static FVector InverseTransformLocation(const FTransform& T, FVector Location);
 
 	/** 
-	 *	Transform a direction vector by the inverse of the supplied transform - will not change its length 
+	 *	Transform a direction vector by the inverse of the supplied transform - will not change its length.
 	 *	For example, if T was an object's transform, would transform a direction from world space to local space.
 	 */
 	UFUNCTION(BlueprintPure, Category="Math|Transform")
 	static FVector InverseTransformDirection(const FTransform& T, FVector Direction);
 
-	/** Multiply two transforms in order */
-	UFUNCTION(BlueprintPure, meta=(ToolTip = "Multiply two transforms in order: A * B", Keywords="add"), Category="Math|Transform")
+	/**
+	 * Compose two transforms in order: A * B.
+	 *
+	 * Order matters when composing transforms:
+	 * A * B will yield a transform that logically first applies B then A to any subsequent transformation.
+	 *
+	 * Example: LocalToWorld = ComposeTransforms(LocalToWorld, DeltaRotation) will change rotation in local space by DeltaRotation.
+	 * Example: LocalToWorld = ComposeTransforms(DeltaRotation, LocalToWorld) will change rotation in world space by DeltaRotation.
+	 *
+	 * @return New transform: A * B
+	 */
+	UFUNCTION(BlueprintPure, meta=(Keywords="multiply *"), Category="Math|Transform")
 	static FTransform ComposeTransforms(const FTransform& A, const FTransform& B);
 
 	/** 
-	 *  Convert a world-transform from world-space into local-space
+	 *  Convert a world-transform from world-space into local-space.
 	 *  @param		WorldTransform	The transform you wish to convert
 	 *  @param		LocalTransform	The transform the conversion is relative to
 	 *  @return		A new relative transform
@@ -1230,15 +1240,15 @@ class UKismetMathLibrary : public UBlueprintFunctionLibrary
 	UFUNCTION(BlueprintPure, Category="Math|Transform", meta=(Keywords="cast convert"))
 	static FTransform ConvertTransformToRelative(const FTransform& WorldTransform, const FTransform& LocalTransform);
 
-	/* Linearly interpolates between A and B based on Alpha (100% of A when Alpha=0 and 100% of B when Alpha=1) */
+	/* Linearly interpolates between A and B based on Alpha (100% of A when Alpha=0 and 100% of B when Alpha=1). */
 	UFUNCTION(BlueprintPure, meta=(FriendlyName = "Lerp (Transform)"), Category="Math|Transform")
 	static FTransform TLerp(const FTransform& A, const FTransform& B, float Alpha);
 
-	/* Easeing  between A and B using a specified easing function */
+	/** Ease between A and B using a specified easing function. */
 	UFUNCTION(BlueprintPure, meta = (FriendlyName = "Ease (Transform)", BlueprintInternalUseOnly = "true"), Category = "Math|Interpolation")
 	static FTransform TEase(const FTransform& A, const FTransform& B, float Alpha, TEnumAsByte<EEasingFunc::Type> EasingFunc, float BlendExp = 2, int32 Steps = 2);
 
-	/** Tries to reach a target transform */
+	/** Tries to reach a target transform. */
 	UFUNCTION(BlueprintPure, Category="Math|Interpolation")
 	static FTransform TInterpTo(const FTransform& Current, const FTransform& Target, float DeltaTime, float InterpSpeed);
 
@@ -1246,27 +1256,27 @@ class UKismetMathLibrary : public UBlueprintFunctionLibrary
 	// Vector2D functions
 	//
 
-	//Returns addition of Vector A and Vector B (A + B)
+	/** Returns addition of Vector A and Vector B (A + B) */
 	UFUNCTION(BlueprintPure, meta=(FriendlyName = "vector2d + vector2d", CompactNodeTitle = "+", Keywords = "+ add plus", CommutativeAssociativeBinaryOperator = "true"), Category="Math|Vector2D")
 	static FVector2D Add_Vector2DVector2D(FVector2D A, FVector2D B);
 
-	//Returns subtraction of Vector B from Vector A (A - B)
+	/** Returns subtraction of Vector B from Vector A (A - B) */
 	UFUNCTION(BlueprintPure, meta=(FriendlyName = "vector2d - vector2d", CompactNodeTitle = "-", Keywords = "- subtract minus"), Category="Math|Vector2D")
 	static FVector2D Subtract_Vector2DVector2D(FVector2D A, FVector2D B);
 
-	//Returns Vector A scaled by B
+	/** Returns Vector A scaled by B */
 	UFUNCTION(BlueprintPure, meta=(FriendlyName = "vector2d * float", CompactNodeTitle = "*", Keywords = "* multiply"), Category="Math|Vector2D")
 	static FVector2D Multiply_Vector2DFloat(FVector2D A, float B);
 
-	//Returns Vector A divided by B
+	/** Returns Vector A divided by B */
 	UFUNCTION(BlueprintPure, meta=(FriendlyName = "vector2d / float", CompactNodeTitle = "/", Keywords = "/ divide division"), Category="Math|Vector2D")
 	static FVector2D Divide_Vector2DFloat(FVector2D A, float B);
 
-	//Returns Vector A added by B
+	/** Returns Vector A added by B */
 	UFUNCTION(BlueprintPure, meta=(FriendlyName = "vector2d + float", CompactNodeTitle = "+", Keywords = "+ add plus"), Category="Math|Vector2D")
 	static FVector2D Add_Vector2DFloat(FVector2D A, float B);
 
-	//Returns Vector A subtracted by B
+	/** Returns Vector A subtracted by B */
 	UFUNCTION(BlueprintPure, meta=(FriendlyName = "vector2d - float", CompactNodeTitle = "-", Keywords = "- subtract minus"), Category="Math|Vector2D")
 	static FVector2D Subtract_Vector2DFloat(FVector2D A, float B);
 
