Add demo time dilation support to world settings

diff --git a/Engine/Source/Runtime/Engine/Classes/GameFramework/WorldSettings.h b/Engine/Source/Runtime/Engine/Classes/GameFramework/WorldSettings.h
--- a/Engine/Source/Runtime/Engine/Classes/GameFramework/WorldSettings.h
+++ b/Engine/Source/Runtime/Engine/Classes/GameFramework/WorldSettings.h
@@ -381,6 +381,10 @@ class ENGINE_API AWorldSettings : public AInfo, public IInterface_AssetUserData
 	UPROPERTY(replicated)
 	float MatineeTimeDilation;
 
+	// Additional TimeDilation used to control demo playback speed
+	UPROPERTY()
+	float DemoPlayTimeDilation;
+
 	// If paused, FName of person pausing the game.
 	UPROPERTY(replicated)
 	class APlayerState* Pauser;
@@ -436,7 +440,7 @@ public:
 
 	float GetEffectiveTimeDilation() const
 	{
-		return TimeDilation * MatineeTimeDilation;
+		return TimeDilation * MatineeTimeDilation * DemoPlayTimeDilation;
 	}
 
 	/** 
diff --git a/Engine/Source/Runtime/Engine/Private/WorldSettings.cpp b/Engine/Source/Runtime/Engine/Private/WorldSettings.cpp
--- a/Engine/Source/Runtime/Engine/Private/WorldSettings.cpp
+++ b/Engine/Source/Runtime/Engine/Private/WorldSettings.cpp
@@ -43,6 +43,7 @@ AWorldSettings::AWorldSettings(const class FPostConstructInitializeProperties& P
 	bAlwaysRelevant = true;
 	TimeDilation = 1.0f;
 	MatineeTimeDilation = 1.0f;
+	DemoPlayTimeDilation = 1.0f;
 	PackedLightAndShadowMapTextureSize = 1024;
 #if WITH_EDITORONLY_DATA
 	bHiddenEd = true;
