Add time dilation support to demo playback using cvar

diff --git a/Engine/Source/Runtime/Engine/Classes/Engine/DemoNetDriver.h b/Engine/Source/Runtime/Engine/Classes/Engine/DemoNetDriver.h
--- a/Engine/Source/Runtime/Engine/Classes/Engine/DemoNetDriver.h
+++ b/Engine/Source/Runtime/Engine/Classes/Engine/DemoNetDriver.h
@@ -13,28 +13,31 @@ class UDemoNetDriver : public UNetDriver
 	// Variables.
 
 	/** Name of the file to read/write from */
-	FString			DemoFilename;
+	FString				DemoFilename;
 
 	/** Handle to the archive that will read/write network packets */
-	FArchive*		FileAr;
+	FArchive*			FileAr;
 
 	/** @todo document */
-	int32			DemoFrameNum;
+	int32				DemoFrameNum;
 
 	/** Last time (in real seconds) that we recorded a frame */
-	double			LastRecordTime;
+	double				LastRecordTime;
 
 	/** Time (in game seconds) that have elapsed between recorded frames */
-	float			DemoDeltaTime;
+	float				DemoDeltaTime;
 
 	/** during playback, set to total number of frames recorded in the demo */
-	int32			PlaybackTotalFrames;
+	int32				PlaybackTotalFrames;
 
 	/** True if we're in the middle of recording a frame */
-	bool			bIsRecordingDemoFrame;
+	bool				bIsRecordingDemoFrame;
+
+	/** This is our spectator controller that is used to view the demo world from */
+	APlayerController*	SpectatorController;
 
 	UPROPERTY( config )
-	FString			DemoSpectatorClass;
+	FString				DemoSpectatorClass;
 
 	// Begin UNetDriver interface.
 	virtual bool InitBase( bool bInitAsClient, FNetworkNotify* InNotify, const FURL& URL, bool bReuseAddressAndPort, FString& Error ) override;
diff --git a/Engine/Source/Runtime/Engine/Private/DemoNetDriver.cpp b/Engine/Source/Runtime/Engine/Private/DemoNetDriver.cpp
--- a/Engine/Source/Runtime/Engine/Private/DemoNetDriver.cpp
+++ b/Engine/Source/Runtime/Engine/Private/DemoNetDriver.cpp
@@ -14,6 +14,7 @@
 DEFINE_LOG_CATEGORY_STATIC( LogDemo, Log, All );
 
 static TAutoConsoleVariable<float> CVarDemoRecordHz( TEXT( "demo.RecordHz" ), 10, TEXT( "Number of demo frames recorded per second" ) );
+static TAutoConsoleVariable<float> CVarDemoTimeDilation( TEXT( "demo.TimeDilation" ), 1.0f, TEXT( "Override time dilation during demo playback" ) );
 
 static const int32 MAX_DEMO_READ_WRITE_BUFFER = 1024 * 2;
 
@@ -231,6 +232,17 @@ void UDemoNetDriver::TickFlush( float DeltaSeconds )
 		}
 		else if ( ServerConnection != NULL )
 		{
+			World->GetWorldSettings()->DemoPlayTimeDilation = CVarDemoTimeDilation.GetValueOnGameThread();
+
+			// Update time dilation on spectator pawn to compensate for any demo dilation 
+			//	(we want to continue to fly around in real-time)
+			if ( SpectatorController != NULL )
+			{
+				if ( SpectatorController->GetSpectatorPawn() != NULL )
+				{
+					SpectatorController->GetSpectatorPawn()->CustomTimeDilation = 1.0f / World->GetWorldSettings()->DemoPlayTimeDilation;
+				}
+			}
 			TickDemoPlayback( DeltaSeconds );
 		}
 	}
@@ -685,6 +697,9 @@ void UDemoNetConnection::HandleClientPlayer( APlayerController* PC, UNetConnecti
 {
 	Super::HandleClientPlayer( PC, NetConnection );
 
+	// Assume this is our special spectator controller
+	GetDriver()->SpectatorController = PC;
+
 	for ( FActorIterator It( Driver->World ); It; ++It)
 	{
 		if ( It->IsA( APlayerStart::StaticClass() ) )
