use true instead of 1

diff --git a/Engine/Source/Runtime/Engine/Private/World.cpp b/Engine/Source/Runtime/Engine/Private/World.cpp
--- a/Engine/Source/Runtime/Engine/Private/World.cpp
+++ b/Engine/Source/Runtime/Engine/Private/World.cpp
@@ -2957,8 +2957,8 @@ void UWorld::InitializeActorsForPlay(const FURL& InURL, bool bResetTime)
 		}
 
 		// Enable actor script calls.
-		bStartup = 1;
-		bActorsInitialized = 1;
+		bStartup = true;
+		bActorsInitialized = true;
 
 		// Spawn server actors
 		ENetMode CurNetMode = GEngine != NULL ? GEngine->GetNetMode(this) : NM_Standalone;
