Fix destructibles needing hit events to fracture on collision. This flag is coupled on PhysX side but we don't want this on UE side.

diff --git a/Engine/Source/Runtime/Engine/Private/DestructibleComponent.cpp b/Engine/Source/Runtime/Engine/Private/DestructibleComponent.cpp
--- a/Engine/Source/Runtime/Engine/Private/DestructibleComponent.cpp
+++ b/Engine/Source/Runtime/Engine/Private/DestructibleComponent.cpp
@@ -243,7 +243,7 @@ void UDestructibleComponent::CreatePhysicsState()
 
 	// Passing AssetInstanceID = 0 so we'll have self-collision
 	AActor* Owner = GetOwner();
-	CreateShapeFilterData(MoveChannel, (Owner ? Owner->GetUniqueID() : 0), CollResponse, 0, 0, PQueryFilterData, PSimFilterData, BodyInstance.bUseCCD, BodyInstance.bNotifyRigidBodyCollision, false);
+	CreateShapeFilterData(MoveChannel, (Owner ? Owner->GetUniqueID() : 0), CollResponse, 0, 0, PQueryFilterData, PSimFilterData, BodyInstance.bUseCCD, true, false);
 
 	// Build filterData variations for complex and simple
 	PSimFilterData.word3 |= EPDF_SimpleCollision | EPDF_ComplexCollision;
@@ -1151,7 +1151,7 @@ void UDestructibleComponent::SetCollisionResponseForActor(PxRigidDynamic* Actor,
 		AActor* Owner = GetOwner();
 		bool bLargeChunk = IsChunkLarge(ChunkIdx);
 		const FCollisionResponse& ColResponse = bLargeChunk ? LargeChunkCollisionResponse : SmallChunkCollisionResponse;
-		CreateShapeFilterData(MoveChannel, (Owner ? Owner->GetUniqueID() : 0), ColResponse.GetResponseContainer(), 0, ChunkIdxToBoneIdx(ChunkIdx), PQueryFilterData, PSimFilterData, BodyInstance.bUseCCD, BodyInstance.bNotifyRigidBodyCollision, false);
+		CreateShapeFilterData(MoveChannel, (Owner ? Owner->GetUniqueID() : 0), ColResponse.GetResponseContainer(), 0, ChunkIdxToBoneIdx(ChunkIdx), PQueryFilterData, PSimFilterData, BodyInstance.bUseCCD, true, false);
 		
 		PQueryFilterData.word3 |= EPDF_SimpleCollision | EPDF_ComplexCollision;
 
diff --git a/Engine/Source/Runtime/Engine/Private/PhysicsEngine/PhysXSupport.cpp b/Engine/Source/Runtime/Engine/Private/PhysicsEngine/PhysXSupport.cpp
--- a/Engine/Source/Runtime/Engine/Private/PhysicsEngine/PhysXSupport.cpp
+++ b/Engine/Source/Runtime/Engine/Private/PhysicsEngine/PhysXSupport.cpp
@@ -419,11 +419,14 @@ void FPhysXSimEventCallback::onContact(const PxContactPairHeader& PairHeader, co
 	const FBodyInstance* BodyInst0 = FPhysxUserData::Get<FBodyInstance>(PActor0->userData);
 	const FBodyInstance* BodyInst1 = FPhysxUserData::Get<FBodyInstance>(PActor1->userData);
 	
+	bool bEitherDestructible = false;
+
 	// check if it's a destructible actor
 	if (BodyInst0 == NULL)
 	{
 		if (const FDestructibleChunkInfo* DestructibleChunkInfo = FPhysxUserData::Get<FDestructibleChunkInfo>(PActor0->userData))
 		{
+			bEitherDestructible = true;
 			BodyInst0 = DestructibleChunkInfo->OwningComponent.IsValid() ? &DestructibleChunkInfo->OwningComponent->BodyInstance : NULL;
 		}
 	}
@@ -432,6 +435,7 @@ void FPhysXSimEventCallback::onContact(const PxContactPairHeader& PairHeader, co
 	{
 		if (const FDestructibleChunkInfo* DestructibleChunkInfo = FPhysxUserData::Get<FDestructibleChunkInfo>(PActor1->userData))
 		{
+			bEitherDestructible = true;
 			BodyInst1 = DestructibleChunkInfo->OwningComponent.IsValid() ? &DestructibleChunkInfo->OwningComponent->BodyInstance : NULL;
 		}
 	}
@@ -443,6 +447,16 @@ void FPhysXSimEventCallback::onContact(const PxContactPairHeader& PairHeader, co
 		return;
 	}
 
+	//destruction applies damage when it hits something. Unfortunately it relies on the same flag that generates onContact.
+	//We only want onContact events to happen if the user actually selected bNotifyRigidBodyCollision so we have to check if this is the case
+	if (bEitherDestructible)
+	{
+		if (BodyInst0->bNotifyRigidBodyCollision == false && BodyInst1->bNotifyRigidBodyCollision == false)
+		{
+			return;
+		}
+	}
+
 	// Get the Scene. 
 	const PxScene* PScene = PActor0->getScene();
 	check(PScene == PActor1->getScene());
