Apply time dilation manually in demo net driver since real-time delta seconds is passed in

diff --git a/Engine/Source/Runtime/Engine/Private/DemoNetDriver.cpp b/Engine/Source/Runtime/Engine/Private/DemoNetDriver.cpp
--- a/Engine/Source/Runtime/Engine/Private/DemoNetDriver.cpp
+++ b/Engine/Source/Runtime/Engine/Private/DemoNetDriver.cpp
@@ -234,6 +234,9 @@ void UDemoNetDriver::TickFlush( float DeltaSeconds )
 		{
 			World->GetWorldSettings()->DemoPlayTimeDilation = CVarDemoTimeDilation.GetValueOnGameThread();
 
+			// We need to compensate for the fact that DeltaSeconds is real-time for net drivers
+			DeltaSeconds *= World->GetWorldSettings()->GetEffectiveTimeDilation();
+
 			// Update time dilation on spectator pawn to compensate for any demo dilation 
 			//	(we want to continue to fly around in real-time)
 			if ( SpectatorController != NULL )
