Bumped reflection capture DDC guid to fix TTP.

diff --git a/Engine/Source/Runtime/Engine/Private/ReflectionCaptureComponent.cpp b/Engine/Source/Runtime/Engine/Private/ReflectionCaptureComponent.cpp
--- a/Engine/Source/Runtime/Engine/Private/ReflectionCaptureComponent.cpp
+++ b/Engine/Source/Runtime/Engine/Private/ReflectionCaptureComponent.cpp
@@ -197,7 +197,7 @@ APlaneReflectionCapture::APlaneReflectionCapture(const class FPostConstructIniti
 // Generate a new guid to force a recache of all reflection derived data
 // Note: changing this will cause saved capture data in maps to be discarded
 // A resave of those maps will be required to guarantee valid reflections when cooking for ES2
-FGuid ReflectionCaptureDDCVer(0x299db8a0, 0x54254535, 0xa088cc40, 0x1ea245eb);
+FGuid ReflectionCaptureDDCVer(0x0c669396, 0x9cb849ae, 0x9f4120ff, 0x5812f4d2);
 
 FString FReflectionCaptureFullHDRDerivedData::GetDDCKeyString(const FGuid& StateId)
 {
@@ -566,7 +566,7 @@ void FReflectionCaptureEncodedHDRDerivedData::GenerateFromDerivedDataSource(cons
 }
 
 // Generate a new guid to force a recache of all encoded HDR derived data
-#define REFLECTIONCAPTURE_ENCODED_DERIVEDDATA_VER TEXT("96DFC088836B48889143E9DF484C3296")
+#define REFLECTIONCAPTURE_ENCODED_DERIVEDDATA_VER TEXT("96DFC022836B48889143E9DF484C3296")
 
 FString FReflectionCaptureEncodedHDRDerivedData::GetDDCKeyString(const FGuid& StateId)
 {
