UE4 - Temporary compile workaround
- Fix OpenGL warning

diff --git a/Engine/Source/Runtime/Engine/Classes/Materials/MaterialInterface.h b/Engine/Source/Runtime/Engine/Classes/Materials/MaterialInterface.h
--- a/Engine/Source/Runtime/Engine/Classes/Materials/MaterialInterface.h
+++ b/Engine/Source/Runtime/Engine/Classes/Materials/MaterialInterface.h
@@ -589,6 +589,18 @@ public:
 	ENGINE_API static uint32 GetFeatureLevelsToCompileForAllMaterials() { return FeatureLevelsForAllMaterials | (1 << GMaxRHIFeatureLevel); }
 
 	/** Iterate over all feature levels currently marked as active */
+#if IS_MONOLITHIC
+	template <typename TFunction>
+	ENGINE_API static void IterateOverActiveFeatureLevels(TFunction InHandler) 
+	{  
+		uint32 FeatureLevels = GetFeatureLevelsToCompileForAllMaterials();
+		while (FeatureLevels != 0)
+		{
+			InHandler((ERHIFeatureLevel::Type)FBitSet::GetAndClearNextBit(FeatureLevels));
+		}
+	}
+#else
+	//@todo: Can't export template functions...
 	ENGINE_API static void IterateOverActiveFeatureLevels(std::function<void(ERHIFeatureLevel::Type FeatureLevel)> InHandler) 
 	{  
 		uint32 FeatureLevels = GetFeatureLevelsToCompileForAllMaterials();
@@ -597,6 +609,7 @@ public:
 			InHandler((ERHIFeatureLevel::Type)FBitSet::GetAndClearNextBit(FeatureLevels));
 		}
 	}
+#endif
 
 protected:
 
diff --git a/Engine/Source/Runtime/OpenGLDrv/Private/OpenGLRenderTarget.cpp b/Engine/Source/Runtime/OpenGLDrv/Private/OpenGLRenderTarget.cpp
--- a/Engine/Source/Runtime/OpenGLDrv/Private/OpenGLRenderTarget.cpp
+++ b/Engine/Source/Runtime/OpenGLDrv/Private/OpenGLRenderTarget.cpp
@@ -785,7 +785,7 @@ void FOpenGLDynamicRHI::RHIReadSurfaceFloatData(FTextureRHIParamRef TextureRHI,F
 		// 4 float components per texel (RGBA)
 		FloatData.AddUninitialized(SizeX * SizeY);
 		glReadPixels(Rect.Min.X, Rect.Min.Y, SizeX, SizeY, GL_RGBA, GL_FLOAT, FloatData.GetData());
-		FLinearColor* FloatDataPtr = FloatData.GetTypedData();
+		FLinearColor* FloatDataPtr = FloatData.GetData();
 		for (uint32 Index = 0; Index < SizeX * SizeY; ++Index, ++FloatDataPtr)
 		{
 			OutData[Index] = FFloat16Color(*FloatDataPtr);
