Renamed new cvars, now 'r.Shaders.Optimize' and 'r.Shaders.KeepDebugInfo'
#PR 457

diff --git a/Engine/Source/Developer/Windows/ShaderFormatD3D/Private/D3D11ShaderCompiler.cpp b/Engine/Source/Developer/Windows/ShaderFormatD3D/Private/D3D11ShaderCompiler.cpp
--- a/Engine/Source/Developer/Windows/ShaderFormatD3D/Private/D3D11ShaderCompiler.cpp
+++ b/Engine/Source/Developer/Windows/ShaderFormatD3D/Private/D3D11ShaderCompiler.cpp
@@ -633,7 +633,7 @@ void CompileD3D11Shader(const FShaderCompilerInput& Input,FShaderCompilerOutput&
 
 			TRefCountPtr<ID3DBlob> CompressedData;
 
-			if (Input.Environment.CompilerFlags.Contains(CFLAG_DontStripDebugInfo))
+			if (Input.Environment.CompilerFlags.Contains(CFLAG_KeepDebugInfo))
 			{
 				CompressedData = Shader;
 			}
diff --git a/Engine/Source/Runtime/Engine/Private/ShaderCompiler/ShaderCompiler.cpp b/Engine/Source/Runtime/Engine/Private/ShaderCompiler/ShaderCompiler.cpp
--- a/Engine/Source/Runtime/Engine/Private/ShaderCompiler/ShaderCompiler.cpp
+++ b/Engine/Source/Runtime/Engine/Private/ShaderCompiler/ShaderCompiler.cpp
@@ -32,14 +32,14 @@ static FAutoConsoleVariableRef CVarDumpShaderDebugInfo(
 	TEXT("On iOS, if the PowerVR graphics SDK is installed to the default path, the PowerVR shader compiler will be called and errors will be reported during the cook.")
 	);
 
-static TAutoConsoleVariable<int32> CVarStripShaderDebugData(
-	TEXT("r.StripShaderDebugData"),
-	1,
-	TEXT("Whether to strip shader reflection and debug data from shader bytecode.  When using graphical debuggers like Nsight it can be useful to disable this on startup."),
+static TAutoConsoleVariable<int32> CVarKeepShaderDebugData(
+	TEXT("r.Shaders.KeepDebugInfo"),
+	0,
+	TEXT("Whether to keep shader reflection and debug data from shader bytecode, default is to strip.  When using graphical debuggers like Nsight it can be useful to enable this on startup."),
 	ECVF_ReadOnly);
 
 static TAutoConsoleVariable<int32> CVarOptimizeShaders(
-	TEXT("r.OptimizeShaders"),
+	TEXT("r.Shaders.Optimize"),
 	1,
 	TEXT("Whether to optimize shaders.  When using graphical debuggers like Nsight it can be useful to disable this on startup."),
 	ECVF_ReadOnly);
@@ -2075,7 +2075,7 @@ void GlobalBeginCompileShader(
 #endif
 
 	{
-		static const auto CVar = IConsoleManager::Get().FindConsoleVariable(TEXT("r.OptimizeShaders"));
+		static const auto CVar = IConsoleManager::Get().FindConsoleVariable(TEXT("r.Shaders.Optimize"));
 
 		if (CVar->GetInt() == 0)
 		{
@@ -2084,11 +2084,11 @@ void GlobalBeginCompileShader(
 	}
 	
 	{
-		static const auto CVar = IConsoleManager::Get().FindConsoleVariable(TEXT("r.StripShaderDebugData"));
+		static const auto CVar = IConsoleManager::Get().FindConsoleVariable(TEXT("r.Shaders.KeepDebugInfo"));
 
-		if (CVar->GetInt() == 0)
+		if (CVar->GetInt() != 0)
 		{
-			Input.Environment.CompilerFlags.Add(CFLAG_DontStripDebugInfo);
+			Input.Environment.CompilerFlags.Add(CFLAG_KeepDebugInfo);
 		}
 	}
 
diff --git a/Engine/Source/Runtime/ShaderCore/Private/Shader.cpp b/Engine/Source/Runtime/ShaderCore/Private/Shader.cpp
--- a/Engine/Source/Runtime/ShaderCore/Private/Shader.cpp
+++ b/Engine/Source/Runtime/ShaderCore/Private/Shader.cpp
@@ -1159,12 +1159,12 @@ void ShaderMapAppendKeyString(EShaderPlatform Platform, FString& KeyString)
 	}
 
 	{
-		static const auto CVar = IConsoleManager::Get().FindConsoleVariable(TEXT("r.StripShaderDebugData"));
-		KeyString += (CVar && CVar->GetInt() != 0) ? TEXT("") : TEXT("_NoStrip");
+		static const auto CVar = IConsoleManager::Get().FindConsoleVariable(TEXT("r.Shaders.KeepDebugInfo"));
+		KeyString += (CVar && CVar->GetInt() != 0) ? TEXT("_NoStrip") : TEXT("");
 	}
 
 	{
-		static const auto CVar = IConsoleManager::Get().FindConsoleVariable(TEXT("r.OptimizeShaders"));
+		static const auto CVar = IConsoleManager::Get().FindConsoleVariable(TEXT("r.Shaders.Optimize"));
 		KeyString += (CVar && CVar->GetInt() != 0) ? TEXT("") : TEXT("_NoOpt");
 	}
 
diff --git a/Engine/Source/Runtime/ShaderCore/Public/ShaderCore.h b/Engine/Source/Runtime/ShaderCore/Public/ShaderCore.h
--- a/Engine/Source/Runtime/ShaderCore/Public/ShaderCore.h
+++ b/Engine/Source/Runtime/ShaderCore/Public/ShaderCore.h
@@ -111,7 +111,7 @@ enum ECompilerFlags
 	CFLAG_StandardOptimization,
 	/** Shader should use on chip memory instead of main memory ring buffer memory. */
 	CFLAG_OnChip,
-	CFLAG_DontStripDebugInfo
+	CFLAG_KeepDebugInfo
 };
 
 /**
