Add physics stats from PxSimulationStatistics

diff --git a/Engine/Source/Runtime/Engine/Private/PhysicsEngine/PhysLevel.cpp b/Engine/Source/Runtime/Engine/Private/PhysicsEngine/PhysLevel.cpp
--- a/Engine/Source/Runtime/Engine/Private/PhysicsEngine/PhysLevel.cpp
+++ b/Engine/Source/Runtime/Engine/Private/PhysicsEngine/PhysLevel.cpp
@@ -26,6 +26,25 @@ DEFINE_STAT(STAT_PhysicsFetchDynamicsTime);
 DEFINE_STAT(STAT_PhysicsEventTime);
 DEFINE_STAT(STAT_SetBodyTransform);
 
+DEFINE_STAT(STAT_NumBroadphaseAdds);
+DEFINE_STAT(STAT_NumBroadphaseRemoves);
+DEFINE_STAT(STAT_NumActiveConstraints);
+DEFINE_STAT(STAT_NumActiveSimulatedBodies);
+DEFINE_STAT(STAT_NumActiveKinematicBodies);
+DEFINE_STAT(STAT_NumMobileBodies);
+DEFINE_STAT(STAT_NumStaticBodies);
+DEFINE_STAT(STAT_NumShapes);
+
+DEFINE_STAT(STAT_NumBroadphaseAddsAsync);
+DEFINE_STAT(STAT_NumBroadphaseRemovesAsync);
+DEFINE_STAT(STAT_NumActiveConstraintsAsync);
+DEFINE_STAT(STAT_NumActiveSimulatedBodiesAsync);
+DEFINE_STAT(STAT_NumActiveKinematicBodiesAsync);
+DEFINE_STAT(STAT_NumMobileBodiesAsync);
+DEFINE_STAT(STAT_NumStaticBodiesAsync);
+DEFINE_STAT(STAT_NumShapesAsync);
+
+
 FPhysCommandHandler * GPhysCommandHandler = NULL;
 
 // CVars
diff --git a/Engine/Source/Runtime/Engine/Private/PhysicsEngine/PhysScene.cpp b/Engine/Source/Runtime/Engine/Private/PhysicsEngine/PhysScene.cpp
--- a/Engine/Source/Runtime/Engine/Private/PhysicsEngine/PhysScene.cpp
+++ b/Engine/Source/Runtime/Engine/Private/PhysicsEngine/PhysScene.cpp
@@ -411,6 +411,64 @@ void FPhysScene::FlushDeferredCollisionDisableTableQueue()
 	DeferredCollisionDisableTableQueue.Empty();
 }
 
+void GatherPhysXStats(PxScene* PScene, uint32 SceneType)
+{
+	/** Gather PhysX stats */
+	if (SceneType == 0)
+	{
+		if (PScene)
+		{
+			PxSimulationStatistics SimStats;
+			PScene->getSimulationStatistics(SimStats);
+
+			SET_DWORD_STAT(STAT_NumActiveConstraints, SimStats.nbActiveConstraints);
+			SET_DWORD_STAT(STAT_NumActiveSimulatedBodies, SimStats.nbActiveDynamicBodies);
+			SET_DWORD_STAT(STAT_NumActiveKinematicBodies, SimStats.nbActiveKinematicBodies);
+			SET_DWORD_STAT(STAT_NumStaticBodies, SimStats.nbStaticBodies);
+			SET_DWORD_STAT(STAT_NumMobileBodies, SimStats.nbDynamicBodies);
+			
+			SET_DWORD_STAT(STAT_NumBroadphaseAdds, SimStats.getNbBroadPhaseAdds(PxSimulationStatistics::VolumeType::eRIGID_BODY));
+			SET_DWORD_STAT(STAT_NumBroadphaseRemoves, SimStats.getNbBroadPhaseRemoves(PxSimulationStatistics::VolumeType::eRIGID_BODY));
+
+			uint32 NumShapes = 0;
+			for (int32 GeomType = 0; GeomType < PxGeometryType::eGEOMETRY_COUNT; ++GeomType)
+			{
+				NumShapes += SimStats.nbShapes[GeomType];
+			}
+
+			SET_DWORD_STAT(STAT_NumShapes, NumShapes);
+
+		}
+
+	}
+
+#if 0	//this is not reliable right now
+	else if (SceneType == 1 & UPhysicsSettings::Get()->bEnableAsyncScene)
+	{
+		//Having to duplicate because of macros. In theory we can fix this but need to get this quickly
+		PxSimulationStatistics SimStats;
+		PScene->getSimulationStatistics(SimStats);
+
+		SET_DWORD_STAT(STAT_NumActiveConstraintsAsync, SimStats.nbActiveConstraints);
+		SET_DWORD_STAT(STAT_NumActiveSimulatedBodiesAsync, SimStats.nbActiveDynamicBodies);
+		SET_DWORD_STAT(STAT_NumActiveKinematicBodiesAsync, SimStats.nbActiveKinematicBodies);
+		SET_DWORD_STAT(STAT_NumStaticBodiesAsync, SimStats.nbStaticBodies);
+		SET_DWORD_STAT(STAT_NumMobileBodiesAsync, SimStats.nbDynamicBodies);
+
+		SET_DWORD_STAT(STAT_NumBroadphaseAddsAsync, SimStats.getNbBroadPhaseAdds(PxSimulationStatistics::VolumeType::eRIGID_BODY));
+		SET_DWORD_STAT(STAT_NumBroadphaseRemovesAsync, SimStats.getNbBroadPhaseRemoves(PxSimulationStatistics::VolumeType::eRIGID_BODY));
+
+		uint32 NumShapes = 0;
+		for (int32 GeomType = 0; GeomType < PxGeometryType::eGEOMETRY_COUNT; ++GeomType)
+		{
+			NumShapes += SimStats.nbShapes[GeomType];
+		}
+
+		SET_DWORD_STAT(STAT_NumShapesAsync, NumShapes);
+	}
+#endif
+}
+
 /** Exposes ticking of physics-engine scene outside Engine. */
 void FPhysScene::TickPhysScene(uint32 SceneType, FGraphEventRef& InOutCompletionEvent)
 {
@@ -454,6 +512,12 @@ void FPhysScene::TickPhysScene(uint32 SceneType, FGraphEventRef& InOutCompletion
 		return;
 	}
 
+#if WITH_PHYSX
+	GatherPhysXStats(GetPhysXScene(SceneType), SceneType);
+#endif
+
+
+
 	/**
 	* Weight frame time according to PhysScene settings.
 	*/
@@ -641,10 +705,12 @@ void FPhysScene::UpdateActiveTransforms(uint32 SceneType)
 		return;
 	}
 
+
 	PxScene* PScene = GetPhysXScene(SceneType);
 	check(PScene);
 	SCENE_LOCK_READ(PScene);
 	PxU32 NumTransforms = 0;
+
 	const PxActiveTransform* PActiveTransforms = PScene->getActiveTransforms(NumTransforms);
 	SCENE_UNLOCK_READ(PScene);
 
@@ -659,6 +725,7 @@ void FPhysScene::UpdateActiveTransforms(uint32 SceneType)
 
 		ensure(!RigidActor->userData || !FPhysxUserData::IsGarbage(RigidActor->userData));
 	}
+
 }
 #endif
 
diff --git a/Engine/Source/Runtime/Engine/Public/PhysicsPublic.h b/Engine/Source/Runtime/Engine/Public/PhysicsPublic.h
--- a/Engine/Source/Runtime/Engine/Public/PhysicsPublic.h
+++ b/Engine/Source/Runtime/Engine/Public/PhysicsPublic.h
@@ -18,6 +18,23 @@ DECLARE_CYCLE_STAT_EXTERN(TEXT("Fetch Results Time"),STAT_PhysicsFetchDynamicsTi
 DECLARE_CYCLE_STAT_EXTERN(TEXT("Phys Events Time"),STAT_PhysicsEventTime,STATGROUP_Physics, );
 DECLARE_CYCLE_STAT_EXTERN(TEXT("Phys SetBodyTransform"),STAT_SetBodyTransform,STATGROUP_Physics, );
 
+DECLARE_DWORD_COUNTER_STAT_EXTERN(TEXT("Broadphase Adds"), STAT_NumBroadphaseAdds, STATGROUP_Physics, );
+DECLARE_DWORD_COUNTER_STAT_EXTERN(TEXT("Broadphase Removes"), STAT_NumBroadphaseRemoves, STATGROUP_Physics, );
+DECLARE_DWORD_COUNTER_STAT_EXTERN(TEXT("Active Constraints"), STAT_NumActiveConstraints, STATGROUP_Physics, );
+DECLARE_DWORD_COUNTER_STAT_EXTERN(TEXT("Active Simulated Bodies"), STAT_NumActiveSimulatedBodies, STATGROUP_Physics, );
+DECLARE_DWORD_COUNTER_STAT_EXTERN(TEXT("Active Kinematic Bodies"), STAT_NumActiveKinematicBodies, STATGROUP_Physics, );
+DECLARE_DWORD_COUNTER_STAT_EXTERN(TEXT("Mobile Bodies"), STAT_NumMobileBodies, STATGROUP_Physics, );
+DECLARE_DWORD_COUNTER_STAT_EXTERN(TEXT("Static Bodies"), STAT_NumStaticBodies, STATGROUP_Physics, );
+DECLARE_DWORD_COUNTER_STAT_EXTERN(TEXT("Shapes"), STAT_NumShapes, STATGROUP_Physics, );
+
+DECLARE_DWORD_COUNTER_STAT_EXTERN(TEXT("(ASync) Broadphase Adds"), STAT_NumBroadphaseAddsAsync, STATGROUP_Physics, );
+DECLARE_DWORD_COUNTER_STAT_EXTERN(TEXT("(ASync) Broadphase Removes"), STAT_NumBroadphaseRemovesAsync, STATGROUP_Physics, );
+DECLARE_DWORD_COUNTER_STAT_EXTERN(TEXT("(ASync) Active Constraints"), STAT_NumActiveConstraintsAsync, STATGROUP_Physics, );
+DECLARE_DWORD_COUNTER_STAT_EXTERN(TEXT("(ASync) Active Simulated Bodies"), STAT_NumActiveSimulatedBodiesAsync, STATGROUP_Physics, );
+DECLARE_DWORD_COUNTER_STAT_EXTERN(TEXT("(ASync) Active Kinematic Bodies"), STAT_NumActiveKinematicBodiesAsync, STATGROUP_Physics, );
+DECLARE_DWORD_COUNTER_STAT_EXTERN(TEXT("(ASync) Mobile Bodies"), STAT_NumMobileBodiesAsync, STATGROUP_Physics, );
+DECLARE_DWORD_COUNTER_STAT_EXTERN(TEXT("(ASync) Static Bodies"), STAT_NumStaticBodiesAsync, STATGROUP_Physics, );
+DECLARE_DWORD_COUNTER_STAT_EXTERN(TEXT("(ASync) Shapes"), STAT_NumShapesAsync, STATGROUP_Physics, );
 
 #if WITH_PHYSX
 
