fix for landscape crash on mobile devices
TTP348285
#ue4
#ios

diff --git a/Engine/Source/Runtime/Engine/Private/Landscape/LandscapeRender.cpp b/Engine/Source/Runtime/Engine/Private/Landscape/LandscapeRender.cpp
--- a/Engine/Source/Runtime/Engine/Private/Landscape/LandscapeRender.cpp
+++ b/Engine/Source/Runtime/Engine/Private/Landscape/LandscapeRender.cpp
@@ -2796,7 +2796,7 @@ void ULandscapeComponent::GetStreamingTextureInfo(TArray<FStreamingTexturePrimit
 
 	// Heightmap
 	{
-		float HeightmapTexelFactor = TexelFactor*HeightmapTexture->GetSizeY()/(ComponentSizeQuads+1);
+		float HeightmapTexelFactor = TexelFactor * ((HeightmapTexture == NULL) ? 1.0f : (HeightmapTexture->GetSizeY()/(ComponentSizeQuads+1)));
 		FStreamingTexturePrimitiveInfo& StreamingHeightmap = *new(OutStreamingTextures) FStreamingTexturePrimitiveInfo;
 		StreamingHeightmap.Bounds = BoundingSphere;
 		StreamingHeightmap.TexelFactor = ForcedLOD >= 0 ? -13+ForcedLOD : HeightmapTexelFactor; // Minus Value indicate ForcedLOD, 13 for 8k texture
