Switched some API to be more extensible. Hopefully we can avoid deprecating an API this way.

diff --git a/Engine/Source/Runtime/Engine/Classes/GameFramework/PlayerController.h b/Engine/Source/Runtime/Engine/Classes/GameFramework/PlayerController.h
--- a/Engine/Source/Runtime/Engine/Classes/GameFramework/PlayerController.h
+++ b/Engine/Source/Runtime/Engine/Classes/GameFramework/PlayerController.h
@@ -27,17 +27,6 @@ DECLARE_DELEGATE_ThreeParams(FGetAudioListenerPos, FVector& /*Location*/, FVecto
 
 DECLARE_LOG_CATEGORY_EXTERN(LogPlayerController, Log, All);
 
-UENUM()
-namespace EInputMode
-{
-	enum Type
-	{
-		UIOnly,
-		UIAndGame,
-		GameOnly,
-	};
-}
-
 UENUM()
 namespace EDynamicForceFeedbackAction
 {
@@ -69,6 +58,78 @@ struct FDynamicForceFeedbackDetails
 	void Update(FForceFeedbackValues& Values) const;
 };
 
+/** Abstract base class for Input Mode structures */
+struct ENGINE_API FInputModeDataBase
+{
+protected:
+	/** Derived classes override this function to apply the necessary settings for the desired input mode */
+	virtual void ApplyInputMode(FReply& SlateOperations, UGameViewportClient& GameViewportClient) const = 0;
+
+	/** Utility functions for derived classes. */
+	void SetFocusAndLocking(FReply& SlateOperations, TSharedPtr<SWidget> InWidgetToFocus, bool bLockMouseToViewport, TSharedRef<class SViewport> InViewportWidget) const;
+
+	friend class APlayerController;
+};
+
+/** Data structure used to setup an input mode that allows only the UI to respond to user input. */
+struct ENGINE_API FInputModeUIOnly : public FInputModeDataBase
+{
+	/** Widget to focus */
+	FInputModeUIOnly& SetWidgetToFocus(TSharedPtr<SWidget> InWidgetToFocus) { WidgetToFocus = InWidgetToFocus; return *this; }
+
+	/** Whether to lock the mouse to the viewport */
+	FInputModeUIOnly& SetLockMouseToViewport(bool InLockMouseToViewport) { bLockMouseToViewport = InLockMouseToViewport; return *this; }
+
+	FInputModeUIOnly()
+		: WidgetToFocus()
+		, bLockMouseToViewport(false)
+	{}
+
+protected:
+	TSharedPtr<SWidget> WidgetToFocus;
+	bool bLockMouseToViewport;
+
+	virtual void ApplyInputMode(FReply& SlateOperations, UGameViewportClient& GameViewportClient) const override;
+};
+
+/** Data structure used to setup an input mode that allows the UI to respond to user input, and if the UI doesn't handle it player input / player controller gets a chance. */
+struct ENGINE_API FInputModeGameAndUI : public FInputModeDataBase
+{
+	/** Widget to focus */
+	FInputModeGameAndUI& SetWidgetToFocus(TSharedPtr<SWidget> InWidgetToFocus) { WidgetToFocus = InWidgetToFocus; return *this; }
+
+	/** Whether to lock the mouse to the viewport */
+	FInputModeGameAndUI& SetLockMouseToViewport(bool InLockMouseToViewport) { bLockMouseToViewport = InLockMouseToViewport; return *this; }
+
+	/** Whether to hide the cursor during temporary mouse capture caused by a mouse down */
+	FInputModeGameAndUI& SetHideCursorDuringCapture(bool InHideCursorDuringCapture) { bHideCursorDuringCapture = InHideCursorDuringCapture; return *this; }
+
+	FInputModeGameAndUI()
+		: WidgetToFocus()
+		, bLockMouseToViewport(false)
+		, bHideCursorDuringCapture(true)
+	{}
+
+protected:
+
+	TSharedPtr<SWidget> WidgetToFocus;
+	bool bLockMouseToViewport;
+	bool bHideCursorDuringCapture;
+
+	virtual void ApplyInputMode(FReply& SlateOperations, UGameViewportClient& GameViewportClient) const override;
+};
+
+/** Data structure used to setup an input mode that allows only the player input / player controller to respond to user input. */
+struct ENGINE_API FInputModeGameOnly : public FInputModeDataBase
+{
+	FInputModeGameOnly()
+	{}
+
+protected:
+	virtual void ApplyInputMode(FReply& SlateOperations, UGameViewportClient& GameViewportClient) const override;
+};
+
+
 //=============================================================================
 /**
  * PlayerControllers are used by human players to control Pawns.
@@ -974,9 +1035,14 @@ public:
 	virtual void SetVirtualJoystickVisibility(bool bVisible);
 
 public:
-	/** Set the Input Mode. */
-	UFUNCTION(BlueprintCallable, Category = "Input")
-	virtual void SetInputMode(EInputMode::Type Mode);
+
+	/** Setup an input mode. */
+	virtual void SetInputMode(const FInputModeDataBase& InData);
+
+protected:
+
+	/** Utility function for Input Mode functions. */
+	void SetFocusAndLocking(TSharedPtr<SWidget> InWidgetToFocus, bool bLockMouseToViewport, TSharedPtr<class SViewport> InViewportWidget);
 
 public:
 	/**
diff --git a/Engine/Source/Runtime/Engine/Private/PlayerController.cpp b/Engine/Source/Runtime/Engine/Private/PlayerController.cpp
--- a/Engine/Source/Runtime/Engine/Private/PlayerController.cpp
+++ b/Engine/Source/Runtime/Engine/Private/PlayerController.cpp
@@ -2081,7 +2081,7 @@ void APlayerController::FlushPressedKeys()
 bool APlayerController::InputKey(FKey Key, EInputEvent EventType, float AmountDepressed, bool bGamepad)
 {
 	bool bResult = false;
-
+	
 	if (GEngine->HMDDevice.IsValid() && GEngine->IsStereoscopic3D())
 	{
 		bResult = GEngine->HMDDevice->HandleInputKey(PlayerInput, Key, EventType, AmountDepressed, bGamepad);
@@ -4253,41 +4253,76 @@ void APlayerController::SetVirtualJoystickVisibility(bool bVisible)
 	}
 }
 
-void APlayerController::SetInputMode(EInputMode::Type Mode)
+void FInputModeDataBase::SetFocusAndLocking(FReply& SlateOperations, TSharedPtr<SWidget> InWidgetToFocus, bool bLockMouseToViewport, TSharedRef<SViewport> InViewportWidget) const
+{
+	if (InWidgetToFocus.IsValid())
+	{
+		SlateOperations.CaptureJoystick(InWidgetToFocus.ToSharedRef());
+	}
+	else
+	{
+		SlateOperations.ReleaseJoystickCapture();
+	}
+
+	if (bLockMouseToViewport)
+	{
+		SlateOperations.LockMouseToWidget(InViewportWidget);
+	}
+	else
+	{
+		SlateOperations.ReleaseMouseLock();
+	}
+}
+
+void FInputModeUIOnly::ApplyInputMode(FReply& SlateOperations, UGameViewportClient& GameViewportClient) const
+{
+	TSharedPtr<SViewport> ViewportWidget = GameViewportClient.GetGameViewportWidget();
+	if (ViewportWidget.IsValid())
+	{
+		SetFocusAndLocking(SlateOperations, WidgetToFocus, bLockMouseToViewport, ViewportWidget.ToSharedRef());
+
+		SlateOperations.ReleaseMouseCapture();
+
+		GameViewportClient.SetIgnoreInput(true);
+		GameViewportClient.SetCaptureMouseOnClick(EMouseCaptureMode::NoCapture);
+	}
+}
+
+void FInputModeGameAndUI::ApplyInputMode(FReply& SlateOperations, UGameViewportClient& GameViewportClient) const
+{
+	TSharedPtr<SViewport> ViewportWidget = GameViewportClient.GetGameViewportWidget();
+	if (ViewportWidget.IsValid())
+	{
+		SetFocusAndLocking(SlateOperations, WidgetToFocus, bLockMouseToViewport, ViewportWidget.ToSharedRef());
+
+		SlateOperations.ReleaseMouseCapture();
+
+		GameViewportClient.SetIgnoreInput(false);
+		GameViewportClient.SetCaptureMouseOnClick(EMouseCaptureMode::CaptureDuringMouseDown);
+	}
+}
+
+void FInputModeGameOnly::ApplyInputMode(FReply& SlateOperations, UGameViewportClient& GameViewportClient) const
+{
+	TSharedPtr<SViewport> ViewportWidget = GameViewportClient.GetGameViewportWidget();
+	if (ViewportWidget.IsValid())
+	{
+		TSharedRef<SViewport> ViewportWidgetRef = ViewportWidget.ToSharedRef();
+		SlateOperations.UseHighPrecisionMouseMovement(ViewportWidgetRef);
+		SlateOperations.CaptureJoystick(ViewportWidgetRef);
+		SlateOperations.LockMouseToWidget(ViewportWidgetRef);
+		SlateOperations.SetKeyboardFocus(ViewportWidgetRef, EKeyboardFocusCause::SetDirectly);
+		GameViewportClient.SetIgnoreInput(false);
+		GameViewportClient.SetCaptureMouseOnClick(EMouseCaptureMode::CapturePermanently);
+	}
+}
+
+void APlayerController::SetInputMode(const FInputModeDataBase& InData)
 {
 	UGameViewportClient* GameViewportClient = GetWorld()->GetGameViewport();
 	if (GameViewportClient != nullptr)
 	{
-		TSharedPtr<SViewport> ViewportWidget = GameViewportClient->GetGameViewportWidget();
-		if (ViewportWidget.IsValid())
-		{
-			switch (Mode)
-			{
-			case EInputMode::UIOnly:
-				SlateOperations.ReleaseMouseCapture();
-				SlateOperations.ReleaseJoystickCapture();
-				SlateOperations.ReleaseMouseLock();
-				GameViewportClient->SetIgnoreInput(true);
-				GameViewportClient->SetCaptureMouseOnClick(EMouseCaptureMode::NoCapture);
-				break;
-			case EInputMode::UIAndGame:
-				SlateOperations.ReleaseMouseCapture();
-				SlateOperations.ReleaseJoystickCapture();
-				SlateOperations.ReleaseMouseLock();
-				GameViewportClient->SetIgnoreInput(false);
-				GameViewportClient->SetCaptureMouseOnClick(EMouseCaptureMode::CaptureDuringMouseDown);
-				break;
-			case EInputMode::GameOnly:
-				TSharedRef<SViewport> pViewportWidgetRef = ViewportWidget.ToSharedRef();
-				SlateOperations.CaptureMouse(pViewportWidgetRef);
-				SlateOperations.CaptureJoystick(pViewportWidgetRef);
-				SlateOperations.LockMouseToWidget(pViewportWidgetRef);
-				SlateOperations.SetKeyboardFocus(pViewportWidgetRef, EKeyboardFocusCause::SetDirectly);
-				GameViewportClient->SetIgnoreInput(false);
-				GameViewportClient->SetCaptureMouseOnClick(EMouseCaptureMode::CapturePermanently);
-				break;
-			}
-		}
+		InData.ApplyInputMode(SlateOperations, *GameViewportClient);
 	}
 }
 
diff --git a/Engine/Source/Runtime/UMG/Private/WidgetBlueprintLibrary.cpp b/Engine/Source/Runtime/UMG/Private/WidgetBlueprintLibrary.cpp
--- a/Engine/Source/Runtime/UMG/Private/WidgetBlueprintLibrary.cpp
+++ b/Engine/Source/Runtime/UMG/Private/WidgetBlueprintLibrary.cpp
@@ -39,6 +39,46 @@ UDragDropOperation* UWidgetBlueprintLibrary::CreateDragDropOperation(TSubclassOf
 	}
 }
 
+void UWidgetBlueprintLibrary::SetInputMode_UIOnly(APlayerController* Target, UWidget* InWidgetToFocus, bool bLockMouseToViewport)
+{
+	if (Target != nullptr)
+	{
+		FInputModeUIOnly InputMode;
+		InputMode.SetLockMouseToViewport(bLockMouseToViewport);
+
+		if (InWidgetToFocus != nullptr)
+		{
+			InputMode.SetWidgetToFocus(InWidgetToFocus->TakeWidget());
+		}
+		Target->SetInputMode(InputMode);
+	}
+}
+
+void UWidgetBlueprintLibrary::SetInputMode_GameAndUI(APlayerController* Target, UWidget* InWidgetToFocus, bool bLockMouseToViewport, bool bHideCursorDuringCapture)
+{
+	if (Target != nullptr)
+	{
+		FInputModeGameAndUI InputMode;
+		InputMode.SetLockMouseToViewport(bLockMouseToViewport);
+		InputMode.SetHideCursorDuringCapture(bHideCursorDuringCapture);
+
+		if (InWidgetToFocus != nullptr)
+		{
+			InputMode.SetWidgetToFocus(InWidgetToFocus->TakeWidget());
+		}
+		Target->SetInputMode(InputMode);
+	}
+}
+
+void UWidgetBlueprintLibrary::SetInputMode_GameOnly(APlayerController* Target)
+{
+	if (Target != nullptr)
+	{
+		FInputModeGameOnly InputMode;
+		Target->SetInputMode(InputMode);
+	}
+}
+
 void UWidgetBlueprintLibrary::SetFocusToGameViewport()
 {
 	FSlateApplication::Get().SetFocusToGameViewport();
diff --git a/Engine/Source/Runtime/UMG/Public/Blueprint/WidgetBlueprintLibrary.h b/Engine/Source/Runtime/UMG/Public/Blueprint/WidgetBlueprintLibrary.h
--- a/Engine/Source/Runtime/UMG/Public/Blueprint/WidgetBlueprintLibrary.h
+++ b/Engine/Source/Runtime/UMG/Public/Blueprint/WidgetBlueprintLibrary.h
@@ -17,6 +17,19 @@ class UWidgetBlueprintLibrary : public UBlueprintFunctionLibrary
 	static UDragDropOperation* CreateDragDropOperation(TSubclassOf<UDragDropOperation> OperationClass);
 	
 
+	/** Setup an input mode that allows only the UI to respond to user input. */
+	UFUNCTION(BlueprintCallable, Category = "Input")
+	static void SetInputMode_UIOnly(APlayerController* Target, UWidget* InWidgetToFocus = nullptr, bool bLockMouseToViewport = false);
+
+	/** Setup an input mode that allows only the UI to respond to user input, and if the UI doesn't handle it player input / player controller gets a chance. */
+	UFUNCTION(BlueprintCallable, Category = "Input")
+	static void SetInputMode_GameAndUI(APlayerController* Target, UWidget* InWidgetToFocus = nullptr, bool bLockMouseToViewport = false, bool bHideCursorDuringCapture = true);
+
+	/** Setup an input mode that allows only player input / player controller to respond to user input. */
+	UFUNCTION(BlueprintCallable, Category = "Input")
+	static void SetInputMode_GameOnly(APlayerController* Target);
+
+
 	UFUNCTION(BlueprintCallable, Category="Focus")
 	static void SetFocusToGameViewport();
 
