#ue4 - Consolidate all CharacterMovementComponent::PhysFalling acceleration handling into GetFallingLateralAcceleration().

diff --git a/Engine/Source/Runtime/Engine/Classes/GameFramework/CharacterMovementComponent.h b/Engine/Source/Runtime/Engine/Classes/GameFramework/CharacterMovementComponent.h
--- a/Engine/Source/Runtime/Engine/Classes/GameFramework/CharacterMovementComponent.h
+++ b/Engine/Source/Runtime/Engine/Classes/GameFramework/CharacterMovementComponent.h
@@ -995,8 +995,8 @@ public:
 	// Helpers for PhysFalling
 
 	/**
-	 * Get the acceleration to use during falling movement. Base implementation simply returns current Acceleration.
-	 * The Z component of the result is ignored, and the actual value of acceleration used while falling will be modified by the result of GetAirControl().
+	 * Get the lateral acceleration to use during falling movement. The Z component of the result is ignored.
+	 * Default implementation returns current Acceleration value modified by GetAirControl(), with Z component removed.
 	 * This function is used internally by PhysFalling().
 	 *
 	 * @param DeltaTime Time step for the current update.
@@ -1008,13 +1008,13 @@ public:
 	 * Get the air control to use during falling movement. This is a multiplier applied to the acceleration used when falling.
 	 * Given an initial air control (TickAirControl), applies the result of BoostAirControl(). If non-zero and FindAirControlImpact()
 	 * returns true, we then apply the result of LimitAirControl().
-	 * This function is used internally by PhysFalling().
+	 * This function is used internally by GetFallingLateralAcceleration().
 	 *
 	 * @param DeltaTime			Time step for the current update.
 	 * @param TickAirControl	Current air control value.
 	 * @param FallAcceleration	Acceleration used during movement.
 	 * @return Air control to use during falling movement.
-	 * @see AirControl, BoostAirControl(), LimitAirControl()
+	 * @see AirControl, BoostAirControl(), LimitAirControl(), GetFallingLateralAcceleration()
 	 */
 	virtual float GetAirControl(float DeltaTime, float TickAirControl, const FVector& FallAcceleration);
 
diff --git a/Engine/Source/Runtime/Engine/Private/CharacterMovementComponent.cpp b/Engine/Source/Runtime/Engine/Private/CharacterMovementComponent.cpp
--- a/Engine/Source/Runtime/Engine/Private/CharacterMovementComponent.cpp
+++ b/Engine/Source/Runtime/Engine/Private/CharacterMovementComponent.cpp
@@ -2809,7 +2809,18 @@ void UCharacterMovementComponent::NotifyJumpApex()
 FVector UCharacterMovementComponent::GetFallingLateralAcceleration(float DeltaTime)
 {
 	// No acceleration in Z
-	return FVector(Acceleration.X, Acceleration.Y, 0.f);
+	FVector FallAcceleration = FVector(Acceleration.X, Acceleration.Y, 0.f);
+
+	// bound acceleration, falling object has minimal ability to impact acceleration
+	if (!HasRootMotion() && FallAcceleration.SizeSquared2D() > 0.f)
+	{
+		const float TickAirControl = GetAirControl(DeltaTime, AirControl, FallAcceleration);
+
+		const float MaxAccel = GetMaxAcceleration() * TickAirControl;
+		FallAcceleration = FallAcceleration.ClampMaxSize(MaxAccel);
+	}
+
+	return FallAcceleration;
 }
 
 
@@ -2857,15 +2868,6 @@ void UCharacterMovementComponent::PhysFalling(float deltaTime, int32 Iterations)
 	FVector FallAcceleration = GetFallingLateralAcceleration(deltaTime);
 	FallAcceleration.Z = 0.f;
 
-	// bound acceleration, falling object has minimal ability to impact acceleration
-	if (!HasRootMotion() && FallAcceleration.SizeSquared() > 0.f)
-	{
-		const float TickAirControl = GetAirControl(deltaTime, AirControl, FallAcceleration);
-
-		const float MaxAccel = GetMaxAcceleration() * TickAirControl;
-		FallAcceleration = FallAcceleration.ClampMaxSize(MaxAccel);
-	}
-
 	float remainingTime = deltaTime;
 	while( (remainingTime >= MIN_TICK_TIME) && (Iterations < MaxSimulationIterations) )
 	{
