Forward declair to fix mac build.

diff --git a/Engine/Source/Runtime/Engine/Classes/GameFramework/PlayerController.h b/Engine/Source/Runtime/Engine/Classes/GameFramework/PlayerController.h
--- a/Engine/Source/Runtime/Engine/Classes/GameFramework/PlayerController.h
+++ b/Engine/Source/Runtime/Engine/Classes/GameFramework/PlayerController.h
@@ -63,7 +63,7 @@ struct ENGINE_API FInputModeDataBase
 {
 protected:
 	/** Derived classes override this function to apply the necessary settings for the desired input mode */
-	virtual void ApplyInputMode(FReply& SlateOperations, UGameViewportClient& GameViewportClient) const = 0;
+	virtual void ApplyInputMode(FReply& SlateOperations, class UGameViewportClient& GameViewportClient) const = 0;
 
 	/** Utility functions for derived classes. */
 	void SetFocusAndLocking(FReply& SlateOperations, TSharedPtr<SWidget> InWidgetToFocus, bool bLockMouseToViewport, TSharedRef<class SViewport> InViewportWidget) const;
@@ -89,7 +89,7 @@ protected:
 	TSharedPtr<SWidget> WidgetToFocus;
 	bool bLockMouseToViewport;
 
-	virtual void ApplyInputMode(FReply& SlateOperations, UGameViewportClient& GameViewportClient) const override;
+	virtual void ApplyInputMode(FReply& SlateOperations, class UGameViewportClient& GameViewportClient) const override;
 };
 
 /** Data structure used to setup an input mode that allows the UI to respond to user input, and if the UI doesn't handle it player input / player controller gets a chance. */
@@ -116,7 +116,7 @@ protected:
 	bool bLockMouseToViewport;
 	bool bHideCursorDuringCapture;
 
-	virtual void ApplyInputMode(FReply& SlateOperations, UGameViewportClient& GameViewportClient) const override;
+	virtual void ApplyInputMode(FReply& SlateOperations, class UGameViewportClient& GameViewportClient) const override;
 };
 
 /** Data structure used to setup an input mode that allows only the player input / player controller to respond to user input. */
@@ -126,7 +126,7 @@ struct ENGINE_API FInputModeGameOnly : public FInputModeDataBase
 	{}
 
 protected:
-	virtual void ApplyInputMode(FReply& SlateOperations, UGameViewportClient& GameViewportClient) const override;
+	virtual void ApplyInputMode(FReply& SlateOperations, class UGameViewportClient& GameViewportClient) const override;
 };
 
 
diff --git a/Engine/Source/Runtime/Engine/Private/PlayerController.cpp b/Engine/Source/Runtime/Engine/Private/PlayerController.cpp
--- a/Engine/Source/Runtime/Engine/Private/PlayerController.cpp
+++ b/Engine/Source/Runtime/Engine/Private/PlayerController.cpp
@@ -4274,7 +4274,7 @@ void FInputModeDataBase::SetFocusAndLocking(FReply& SlateOperations, TSharedPtr<
 	}
 }
 
-void FInputModeUIOnly::ApplyInputMode(FReply& SlateOperations, UGameViewportClient& GameViewportClient) const
+void FInputModeUIOnly::ApplyInputMode(FReply& SlateOperations, class UGameViewportClient& GameViewportClient) const
 {
 	TSharedPtr<SViewport> ViewportWidget = GameViewportClient.GetGameViewportWidget();
 	if (ViewportWidget.IsValid())
@@ -4288,7 +4288,7 @@ void FInputModeUIOnly::ApplyInputMode(FReply& SlateOperations, UGameViewportClie
 	}
 }
 
-void FInputModeGameAndUI::ApplyInputMode(FReply& SlateOperations, UGameViewportClient& GameViewportClient) const
+void FInputModeGameAndUI::ApplyInputMode(FReply& SlateOperations, class UGameViewportClient& GameViewportClient) const
 {
 	TSharedPtr<SViewport> ViewportWidget = GameViewportClient.GetGameViewportWidget();
 	if (ViewportWidget.IsValid())
@@ -4302,7 +4302,7 @@ void FInputModeGameAndUI::ApplyInputMode(FReply& SlateOperations, UGameViewportC
 	}
 }
 
-void FInputModeGameOnly::ApplyInputMode(FReply& SlateOperations, UGameViewportClient& GameViewportClient) const
+void FInputModeGameOnly::ApplyInputMode(FReply& SlateOperations, class UGameViewportClient& GameViewportClient) const
 {
 	TSharedPtr<SViewport> ViewportWidget = GameViewportClient.GetGameViewportWidget();
 	if (ViewportWidget.IsValid())
