Documentation fix, RandomUnitVectorInCone expects radians, and it has always expected radians! TTP#346699
Also reverting a change to the friendly name of float * float, looks like it was not intentional

diff --git a/Engine/Source/Runtime/Engine/Classes/Kismet/KismetMathLibrary.h b/Engine/Source/Runtime/Engine/Classes/Kismet/KismetMathLibrary.h
--- a/Engine/Source/Runtime/Engine/Classes/Kismet/KismetMathLibrary.h
+++ b/Engine/Source/Runtime/Engine/Classes/Kismet/KismetMathLibrary.h
@@ -237,7 +237,7 @@ class UKismetMathLibrary : public UBlueprintFunctionLibrary
 	static float MultiplyMultiply_FloatFloat(float Base, float Exp);
 
 	/* Multiplication (A * B) */
-	UFUNCTION(BlueprintPure, meta=(FriendlyName = "float* float", CompactNodeTitle = "*", Keywords = "* multiply", CommutativeAssociativeBinaryOperator = "true"), Category="Math|Float")
+	UFUNCTION(BlueprintPure, meta=(FriendlyName = "float * float", CompactNodeTitle = "*", Keywords = "* multiply", CommutativeAssociativeBinaryOperator = "true"), Category="Math|Float")
 	static float Multiply_FloatFloat(float A, float B);
 
 	/* Multiplication (A * B) */
@@ -572,7 +572,7 @@ class UKismetMathLibrary : public UBlueprintFunctionLibrary
 	/** 
 	 * Returns a random vector with length of 1, within the specified cone, with uniform random distribution. 
 	 * @param ConeDir	The base "center" direction of the cone.
-	 * @param ConeHalfAngle		The half-angle of the cone (from ConeDir to edge), in degrees.
+	 * @param ConeHalfAngle		The half-angle of the cone (from ConeDir to edge), in radians.
 	 */
 	UFUNCTION(BlueprintPure, Category = "Math|Random")
 	static FVector RandomUnitVectorInCone(FVector ConeDir, float ConeHalfAngle);
