Add SetRenderInMainPass function to set bRenderInMainPass property and update rendering state.

Solves issues highlighted in UDN 'bRenderInMainPass change at runtime not working'

diff --git a/Engine/Source/Runtime/Engine/Classes/Components/PrimitiveComponent.h b/Engine/Source/Runtime/Engine/Classes/Components/PrimitiveComponent.h
--- a/Engine/Source/Runtime/Engine/Classes/Components/PrimitiveComponent.h
+++ b/Engine/Source/Runtime/Engine/Classes/Components/PrimitiveComponent.h
@@ -206,7 +206,7 @@ public:
 	uint32 bRenderCustomDepth:1;
 
 	/** If true, this component will be rendered in the main pass (z prepass, basepass, transparency) */
-	UPROPERTY(EditAnywhere, AdvancedDisplay, BlueprintReadWrite, Category=Rendering)
+	UPROPERTY(EditAnywhere, AdvancedDisplay, BlueprintReadOnly, Category = Rendering)
 	uint32 bRenderInMainPass:1;
 
 	/** Whether to completely hide the primitive in the game; if true, the primitive is not drawn, does not cast a shadow. */
@@ -845,6 +845,10 @@ public:
 	UFUNCTION(BlueprintCallable, Category="Rendering")
 	void SetRenderCustomDepth(bool bValue);
 
+	/** Sets bRenderInMainPass property and marks the render state dirty. */
+	UFUNCTION(BlueprintCallable, Category = "Rendering")
+	void SetRenderInMainPass(bool bValue);
+
 public:
 	static int32 CurrentTag;
 
diff --git a/Engine/Source/Runtime/Engine/Private/PrimitiveComponent.cpp b/Engine/Source/Runtime/Engine/Private/PrimitiveComponent.cpp
--- a/Engine/Source/Runtime/Engine/Private/PrimitiveComponent.cpp
+++ b/Engine/Source/Runtime/Engine/Private/PrimitiveComponent.cpp
@@ -2509,6 +2509,15 @@ void UPrimitiveComponent::SetRenderCustomDepth(bool bValue)
 	}
 }
 
+void UPrimitiveComponent::SetRenderInMainPass(bool bValue)
+{
+	if (bRenderInMainPass != bValue)
+	{
+		bRenderInMainPass = bValue;
+		MarkRenderStateDirty();
+	}
+}
+
 #if WITH_EDITOR
 const int32 UPrimitiveComponent::GetNumUncachedStaticLightingInteractions() const
 {
