PS4 - Initial support for rendering to specific slices of a cubemap
* Fixes crash when rendering reflection captures

diff --git a/Engine/Shaders/ShaderVersion.usf b/Engine/Shaders/ShaderVersion.usf
--- a/Engine/Shaders/ShaderVersion.usf
+++ b/Engine/Shaders/ShaderVersion.usf
@@ -2,4 +2,4 @@
 // Each time the console command is executed it generates a new GUID. As this file is included
 // in common.usf (which should be included in any shader) it allows to invalidate the shader DDC.
 // 
-// GUID = DF48532F444DC92D3D0A8593C3F4749F
\ No newline at end of file
+// GUID = 98C7AD7080954AAFA80A0922AE68EC87
\ No newline at end of file
diff --git a/Engine/Source/Runtime/Engine/Public/ClearReplacementShaders.h b/Engine/Source/Runtime/Engine/Public/ClearReplacementShaders.h
--- a/Engine/Source/Runtime/Engine/Public/ClearReplacementShaders.h
+++ b/Engine/Source/Runtime/Engine/Public/ClearReplacementShaders.h
@@ -90,6 +90,16 @@ public:
 		return IsFeatureLevelSupported(Platform, ERHIFeatureLevel::SM5);
 	}
 
+	const FShaderParameter& GetClearColorParameter()
+	{
+		return ClearColor;
+	}
+
+	const FShaderResourceParameter& GetClearTextureRWParameter()
+	{
+		return ClearTextureRW;
+	}
+
 protected:
 	FShaderParameter ClearColor;
 	FShaderResourceParameter ClearTextureRW;
