Core - Creating a copy of a mesh with ConstructObject copies tick function references (fixed)

diff --git a/Engine/Source/Runtime/Engine/Classes/Engine/EngineBaseTypes.h b/Engine/Source/Runtime/Engine/Classes/Engine/EngineBaseTypes.h
--- a/Engine/Source/Runtime/Engine/Classes/Engine/EngineBaseTypes.h
+++ b/Engine/Source/Runtime/Engine/Classes/Engine/EngineBaseTypes.h
@@ -294,6 +294,15 @@ private:
 	friend class FTickTaskLevel;
 };
 
+template<>
+struct TStructOpsTypeTraits<FTickFunction> : public TStructOpsTypeTraitsBase
+{
+	enum
+	{
+		WithCopy = false
+	};
+};
+
 /** 
 * Tick function that calls AActor::TickActor
 **/
@@ -317,6 +326,15 @@ struct FActorTickFunction : public FTickFunction
 	ENGINE_API virtual FString DiagnosticMessage();
 };
 
+template<>
+struct TStructOpsTypeTraits<FActorTickFunction> : public TStructOpsTypeTraitsBase
+{
+	enum
+	{
+		WithCopy = false
+	};
+};
+
 /** 
 * Tick function that calls UActorComponent::ConditionalTick
 **/
@@ -340,6 +358,15 @@ struct FActorComponentTickFunction : public FTickFunction
 	ENGINE_API virtual FString DiagnosticMessage();
 };
 
+
+template<>
+struct TStructOpsTypeTraits<FActorComponentTickFunction> : public TStructOpsTypeTraitsBase
+{
+	enum
+	{
+		WithCopy = false
+	};
+};
 /** 
 * Tick function that calls UPrimitiveComponent::PostPhysicsTick
 **/
@@ -363,7 +390,14 @@ struct FPrimitiveComponentPostPhysicsTickFunction : public FTickFunction
 	virtual FString DiagnosticMessage();
 };
 
-
+template<>
+struct TStructOpsTypeTraits<FPrimitiveComponentPostPhysicsTickFunction> : public TStructOpsTypeTraitsBase
+{
+	enum
+	{
+		WithCopy = false
+	};
+};
 
 /** Types of network failures broadcast from the engine */
 namespace ENetworkFailure
