Correct issues with calling order of cleanup functions for components
* BeginDestroy calls ExecuteUnregisterEvents before OnComponentDestroyed
* UninitializeComponent now gets called from DestroyComponent
* InitializeComponent will be called by RegisterComponentWithWorld if its Owner Actor is initialized
* SceneComponent now overrides OnComponentDestroyed instead of DestroyComponent to do ScopedMovementStack and Detachment ensuring UninitializeComponent is consistently called before UnregisterComponent

diff --git a/Engine/Source/Runtime/AIModule/Private/AIController.cpp b/Engine/Source/Runtime/AIModule/Private/AIController.cpp
--- a/Engine/Source/Runtime/AIModule/Private/AIController.cpp
+++ b/Engine/Source/Runtime/AIModule/Private/AIController.cpp
@@ -728,13 +728,6 @@ bool AAIController::RunBehaviorTree(UBehaviorTree* BTAsset)
 
 			BrainComponent = BTComp = ConstructObject<UBehaviorTreeComponent>(UBehaviorTreeComponent::StaticClass(), this, TEXT("BTComponent"));
 			BrainComponent->RegisterComponent();
-
-			if (BrainComponent->bWantsInitializeComponent)
-			{
-				// make sure that newly created component is initialized before running BT
-				// both blackboard and BT to must exist before calling it!
-				BrainComponent->InitializeComponent();
-			}
 		}
 
 		check(BTComp != NULL);
@@ -762,7 +755,6 @@ bool AAIController::UseBlackboard(UBlackboardData* BlackboardAsset)
 		{
 			BlackboardComp->InitializeBlackboard(BlackboardAsset);
 			BlackboardComp->RegisterComponent();
-			BlackboardComp->InitializeComponent();
 		}
 
 	}
diff --git a/Engine/Source/Runtime/Engine/Classes/Components/SceneComponent.h b/Engine/Source/Runtime/Engine/Classes/Components/SceneComponent.h
--- a/Engine/Source/Runtime/Engine/Classes/Components/SceneComponent.h
+++ b/Engine/Source/Runtime/Engine/Classes/Components/SceneComponent.h
@@ -465,7 +465,7 @@ public:
 	// Begin ActorComponent interface
 	virtual void OnRegister() override;
 	virtual void UpdateComponentToWorld(bool bSkipPhysicsMove = false) override final;
-	virtual void DestroyComponent() override;
+	virtual void OnComponentDestroyed() override;
 	virtual void ApplyWorldOffset(const FVector& InOffset, bool bWorldShift) override;
 	// End ActorComponent interface
 
diff --git a/Engine/Source/Runtime/Engine/Private/ActorComponent.cpp b/Engine/Source/Runtime/Engine/Private/ActorComponent.cpp
--- a/Engine/Source/Runtime/Engine/Private/ActorComponent.cpp
+++ b/Engine/Source/Runtime/Engine/Private/ActorComponent.cpp
@@ -289,16 +289,14 @@ void UActorComponent::BeginDestroy()
 		UninitializeComponent();
 	}
 
+	ExecuteUnregisterEvents();
+
 	// Ensure that we call OnComponentDestroyed before we destroy this component
 	if(bHasBeenCreated)
 	{
-#if !(UE_BUILD_SHIPPING || UE_BUILD_TEST)
-		//UE_LOG(LogActor, Warning, TEXT("UChildActorComponent::BeginDestroy called on '%s' before without OnComponentDestroyed being called."), *this->GetPathName());
-#endif
 		OnComponentDestroyed();
 	}
 
-	ExecuteUnregisterEvents();
 	World = NULL;
 
 	// Remove from the parent's OwnedComponents list
@@ -449,6 +447,7 @@ void UActorComponent::OnUnregister()
 void UActorComponent::InitializeComponent()
 {
 	check(bRegistered);
+	check(!bHasBeenInitialized);
 	bHasBeenInitialized = true;
 }
 
@@ -604,7 +603,7 @@ void UActorComponent::RegisterComponentWithWorld(UWorld* InWorld)
 		// can happen with undo because the owner will be restored "next"
 		//checkf(!Owner->IsPendingKill(), TEXT("%s"), *GetFullName());
 
-		if(InWorld != GetOwner()->GetWorld())
+		if(InWorld != Owner->GetWorld())
 		{
 			// The only time you should specify a scene that is not GetOwner()->GetWorld() is when you don't have an Actor
 			UE_LOG(LogActorComponent, Log, TEXT("RegisterComponentWithWorld: (%s) Specifying a world, but an Owner Actor found, and InWorld is not GetOwner()->GetWorld()"), *GetPathName());
@@ -616,6 +615,14 @@ void UActorComponent::RegisterComponentWithWorld(UWorld* InWorld)
 
 	ExecuteRegisterEvents();
 	RegisterAllComponentTickFunctions(true);
+
+	if (Owner == nullptr || Owner->bActorInitialized)
+	{
+		if (bWantsInitializeComponent)
+		{
+			InitializeComponent();
+		}
+	}
 }
 
 void UActorComponent::RegisterComponent()
@@ -655,7 +662,13 @@ void UActorComponent::UnregisterComponent()
 
 void UActorComponent::DestroyComponent()
 {
-	// First, unregister if registered
+	// Ensure that we call UninitializeComponent before we destroy this component
+	if (bHasBeenInitialized)
+	{
+		UninitializeComponent();
+	}
+
+	// Unregister if registered
 	if(IsRegistered())
 	{
 		UnregisterComponent();
diff --git a/Engine/Source/Runtime/Engine/Private/Controller.cpp b/Engine/Source/Runtime/Engine/Private/Controller.cpp
--- a/Engine/Source/Runtime/Engine/Private/Controller.cpp
+++ b/Engine/Source/Runtime/Engine/Private/Controller.cpp
@@ -569,32 +569,18 @@ void AController::UpdateNavigationComponents()
 
 void AController::InitNavigationControl(UNavigationComponent*& PathFindingComp, UPathFollowingComponent*& PathFollowingComp)
 {
-	bool bInitPathFinding = false;
 	PathFindingComp = FindComponentByClass<UNavigationComponent>();
 	if (PathFindingComp == NULL)
 	{
 		PathFindingComp = NewObject<UNavigationComponent>(this);
 		PathFindingComp->RegisterComponentWithWorld(GetWorld());
-		bInitPathFinding = true;
 	}
 
-	bool bInitPathFollowing = false;
 	PathFollowingComp = FindComponentByClass<UPathFollowingComponent>();
 	if (PathFollowingComp == NULL)
 	{
 		PathFollowingComp = NewObject<UPathFollowingComponent>(this);
 		PathFollowingComp->RegisterComponentWithWorld(GetWorld());
-		bInitPathFollowing = true;
-	}
-
-	if (bInitPathFinding)
-	{
-		PathFindingComp->InitializeComponent();
-	}
-
-	if (bInitPathFollowing)
-	{
-		PathFollowingComp->InitializeComponent();
 	}
 }
 
diff --git a/Engine/Source/Runtime/Engine/Private/SceneComponent.cpp b/Engine/Source/Runtime/Engine/Private/SceneComponent.cpp
--- a/Engine/Source/Runtime/Engine/Private/SceneComponent.cpp
+++ b/Engine/Source/Runtime/Engine/Private/SceneComponent.cpp
@@ -248,19 +248,14 @@ void USceneComponent::EndScopedMovementUpdate(class FScopedMovementUpdate& Compl
 }
 
 
-void USceneComponent::DestroyComponent()
+void USceneComponent::OnComponentDestroyed()
 {
-	ScopedMovementStack.Reset();
+	Super::OnComponentDestroyed();
 
-	if (IsRegistered())
-	{
-		UnregisterComponent();
-	}
+	ScopedMovementStack.Reset();
 
 	// Ensure we are detached before destroying
 	DetachFromParent();
-
-	Super::DestroyComponent();
 }
 
 FBoxSphereBounds USceneComponent::CalcBounds(const FTransform& LocalToWorld) const
