GRHIFeatureLevel -> GMaxRHIFeatureLevel

diff --git a/Engine/Source/Runtime/Engine/Private/UnrealClient.cpp b/Engine/Source/Runtime/Engine/Private/UnrealClient.cpp
--- a/Engine/Source/Runtime/Engine/Private/UnrealClient.cpp
+++ b/Engine/Source/Runtime/Engine/Private/UnrealClient.cpp
@@ -1631,7 +1631,7 @@ FDummyViewport::FDummyViewport(FViewportClient* InViewportClient)
 	, DebugCanvas(NULL)
 {
 	UWorld* CurWorld = (InViewportClient != NULL ? InViewportClient->GetWorld() : NULL);
-	DebugCanvas = new FCanvas(this, NULL, CurWorld, (CurWorld != NULL ? CurWorld->FeatureLevel : GRHIFeatureLevel));
+	DebugCanvas = new FCanvas(this, NULL, CurWorld, (CurWorld != NULL ? CurWorld->FeatureLevel : GMaxRHIFeatureLevel));
 		
 	DebugCanvas->SetAllowedModes(0);
 }
