Global CVar for simple scaling of emitter spawn rates.
Bool property to enable/disable in spawn module. Defaults to enabled.
Existing data with constant spawn rates of <= 1.0 have this scaling disabled.

diff --git a/Engine/Config/BaseScalability.ini b/Engine/Config/BaseScalability.ini
--- a/Engine/Config/BaseScalability.ini
+++ b/Engine/Config/BaseScalability.ini
@@ -193,6 +193,7 @@ r.TranslucencyVolumeBlur=0
 r.MaterialQualityLevel=0
 r.SSS.Scale = 0
 r.SSS.SampleSet=0
+r.EmitterSpawnRateScale=1.0
 
 [EffectsQuality@1]
 r.TranslucencyLightingVolumeDim=32
@@ -204,6 +205,7 @@ r.TranslucencyVolumeBlur=0
 r.MaterialQualityLevel=1
 r.SSS.Scale = 0.75
 r.SSS.SampleSet=0
+r.EmitterSpawnRateScale=1.0
 
 [EffectsQuality@2]
 r.TranslucencyLightingVolumeDim=48
@@ -215,6 +217,7 @@ r.TranslucencyVolumeBlur=1
 r.MaterialQualityLevel=1
 r.SSS.Scale = 1
 r.SSS.SampleSet=1
+r.EmitterSpawnRateScale=1.0
 
 [EffectsQuality@3]
 r.TranslucencyLightingVolumeDim=64
@@ -225,4 +228,5 @@ r.DetailMode=2
 r.TranslucencyVolumeBlur=1
 r.MaterialQualityLevel=1
 r.SSS.Scale = 1
-r.SSS.SampleSet=2
\ No newline at end of file
+r.SSS.SampleSet=2
+r.EmitterSpawnRateScale=1.0
\ No newline at end of file
diff --git a/Engine/Source/Runtime/Core/Private/HAL/ConsoleManager.cpp b/Engine/Source/Runtime/Core/Private/HAL/ConsoleManager.cpp
--- a/Engine/Source/Runtime/Core/Private/HAL/ConsoleManager.cpp
+++ b/Engine/Source/Runtime/Core/Private/HAL/ConsoleManager.cpp
@@ -1936,3 +1936,9 @@ static TAutoConsoleVariable<int32> CVarVerifyPeer(
 	TEXT("  0 = disable (allows self-signed certificates)\n")
 	TEXT("  1 = enable [default]"),
 	ECVF_ReadOnly);
+
+static TAutoConsoleVariable<float> CVarEmitterSpawnRateScale(
+	TEXT("r.EmitterSpawnRateScale"),
+	1.0,
+	TEXT("A global scale upon the spawn rate of emitters. Emitters can choose to apply or ignore it via their bApplyGlobalSpawnRateScale property."),
+	ECVF_Scalability | ECVF_RenderThreadSafe);
\ No newline at end of file
diff --git a/Engine/Source/Runtime/Core/Public/UObject/ObjectVersion.h b/Engine/Source/Runtime/Core/Public/UObject/ObjectVersion.h
--- a/Engine/Source/Runtime/Core/Public/UObject/ObjectVersion.h
+++ b/Engine/Source/Runtime/Core/Public/UObject/ObjectVersion.h
@@ -620,6 +620,8 @@ enum EUnrealEngineObjectUE4Version
 	VER_UE4_RENAME_CAMERA_COMPONENT_CONTROL_ROTATION,
 	// Fix bad refraction material attribute masks
 	VER_UE4_FIX_REFRACTION_INPUT_MASKING,
+	// A global spawn rate for emitters.
+	VER_UE4_GLOBAL_EMITTER_SPAWN_RATE_SCALE,
 
 	// -----<new versions can be added before this line>-------------------------------------------------
 	// - this needs to be the last line (see note below)
diff --git a/Engine/Source/Runtime/Engine/Classes/Particles/Spawn/ParticleModuleSpawn.h b/Engine/Source/Runtime/Engine/Classes/Particles/Spawn/ParticleModuleSpawn.h
--- a/Engine/Source/Runtime/Engine/Classes/Particles/Spawn/ParticleModuleSpawn.h
+++ b/Engine/Source/Runtime/Engine/Classes/Particles/Spawn/ParticleModuleSpawn.h
@@ -33,6 +33,10 @@ class UParticleModuleSpawn : public UParticleModuleSpawnBase
 	UPROPERTY(EditAnywhere, Category=Burst)
 	struct FRawDistributionFloat BurstScale;
 
+	/**	If true, the SpawnRate will be scaled by the global CVar r.EmitterSpawnRateScale */
+	UPROPERTY(EditAnywhere, Category=Spawn)
+	uint32 bApplyGlobalSpawnRateScale : 1;
+
 	/** Initializes the default values for this property */
 	void InitializeDefaults();
 
@@ -41,6 +45,7 @@ class UParticleModuleSpawn : public UParticleModuleSpawnBase
 	virtual void	PostEditChangeProperty(FPropertyChangedEvent& PropertyChangedEvent) override;
 #endif // WITH_EDITOR
 	virtual void	PostInitProperties() override;
+	virtual void	PostLoad() override;
 	virtual void	Serialize(FArchive& Ar) override;
 	// End UObject Interface
 
@@ -55,6 +60,8 @@ class UParticleModuleSpawn : public UParticleModuleSpawnBase
 	virtual float GetEstimatedSpawnRate() override;
 	virtual int32 GetMaximumBurstCount() override;
 	// End UParticleModuleSpawnBase Interface
+
+	float GetGlobalRateScale()const;
 };
 
 
diff --git a/Engine/Source/Runtime/Engine/Private/Particles/ParticleEmitterInstances.cpp b/Engine/Source/Runtime/Engine/Private/Particles/ParticleEmitterInstances.cpp
--- a/Engine/Source/Runtime/Engine/Private/Particles/ParticleEmitterInstances.cpp
+++ b/Engine/Source/Runtime/Engine/Private/Particles/ParticleEmitterInstances.cpp
@@ -1758,7 +1758,7 @@ float FParticleEmitterInstance::Spawn(float DeltaTime)
 		// Figure out spawn rate for this tick.
 		if (bProcessSpawnRate)
 		{
-			float RateScale = LODLevel->SpawnModule->RateScale.GetValue(EmitterTime, Component);
+			float RateScale = LODLevel->SpawnModule->RateScale.GetValue(EmitterTime, Component) * LODLevel->SpawnModule->GetGlobalRateScale();
 			SpawnRate += LODLevel->SpawnModule->Rate.GetValue(EmitterTime, Component) * RateScale;
 			SpawnRate = FMath::Max<float>(0.0f, SpawnRate);
 		}
diff --git a/Engine/Source/Runtime/Engine/Private/Particles/ParticleGpuSimulation.cpp b/Engine/Source/Runtime/Engine/Private/Particles/ParticleGpuSimulation.cpp
--- a/Engine/Source/Runtime/Engine/Private/Particles/ParticleGpuSimulation.cpp
+++ b/Engine/Source/Runtime/Engine/Private/Particles/ParticleGpuSimulation.cpp
@@ -3648,7 +3648,7 @@ private:
 
 		// Determine spawn rate.
 		check(SpawnModule && RequiredModule);
-		const float RateScale = CurrentLODLevel->SpawnModule->RateScale.GetValue(EmitterTime, Component);
+		const float RateScale = CurrentLODLevel->SpawnModule->RateScale.GetValue(EmitterTime, Component) * CurrentLODLevel->SpawnModule->GetGlobalRateScale();
 		float SpawnRate = CurrentLODLevel->SpawnModule->Rate.GetValue(EmitterTime, Component) * RateScale;
 		SpawnRate = FMath::Max<float>(0.0f, SpawnRate);
 
diff --git a/Engine/Source/Runtime/Engine/Private/Particles/ParticleModules_Spawn.cpp b/Engine/Source/Runtime/Engine/Private/Particles/ParticleModules_Spawn.cpp
--- a/Engine/Source/Runtime/Engine/Private/Particles/ParticleModules_Spawn.cpp
+++ b/Engine/Source/Runtime/Engine/Private/Particles/ParticleModules_Spawn.cpp
@@ -31,6 +31,7 @@ UParticleModuleSpawn::UParticleModuleSpawn(const class FPostConstructInitializeP
 {
 	bProcessSpawnRate = true;
 	LODDuplicate = false;
+	bApplyGlobalSpawnRateScale = true;
 }
 
 void UParticleModuleSpawn::InitializeDefaults()
@@ -66,6 +67,20 @@ void UParticleModuleSpawn::PostInitProperties()
 	}
 }
 
+void UParticleModuleSpawn::PostLoad()
+{
+	Super::PostLoad();
+
+	if (GetLinkerUE4Version() < VER_UE4_GLOBAL_EMITTER_SPAWN_RATE_SCALE && Rate.Distribution)
+	{
+		Rate.Distribution->ConditionalPostLoad();//Ensure the distribution is loaded.
+		if (Rate.Distribution->IsA(UDistributionFloatConstant::StaticClass()) && Rate.GetValue() <= 1.0f)
+		{
+			bApplyGlobalSpawnRateScale = false;
+		}
+	}
+}
+
 void UParticleModuleSpawn::Serialize(FArchive& Ar)
 {
 	Super::Serialize(Ar);
@@ -246,6 +261,12 @@ int32 UParticleModuleSpawn::GetMaximumBurstCount()
 	return MaxBurst;
 }
 
+float UParticleModuleSpawn::GetGlobalRateScale()const
+{
+	static const auto EmitterRateScaleCVar = IConsoleManager::Get().FindTConsoleVariableDataFloat(TEXT("r.EmitterSpawnRateScale"));
+	return bApplyGlobalSpawnRateScale ? EmitterRateScaleCVar->GetValueOnAnyThread() : 1.0f;
+}
+
 /*-----------------------------------------------------------------------------
 	UParticleModuleSpawnPerUnit implementation.
 -----------------------------------------------------------------------------*/
diff --git a/Engine/Source/Runtime/Engine/Private/Particles/ParticleTrail2EmitterInstance.cpp b/Engine/Source/Runtime/Engine/Private/Particles/ParticleTrail2EmitterInstance.cpp
--- a/Engine/Source/Runtime/Engine/Private/Particles/ParticleTrail2EmitterInstance.cpp
+++ b/Engine/Source/Runtime/Engine/Private/Particles/ParticleTrail2EmitterInstance.cpp
@@ -1216,7 +1216,7 @@ float FParticleRibbonEmitterInstance::Spawn(float DeltaTime)
 	// Figure out spawn rate for this tick.
 	if (bProcessSpawnRate)
 	{
-		float RateScale = LODLevel->SpawnModule->RateScale.GetValue(EmitterTime, Component);
+		float RateScale = LODLevel->SpawnModule->RateScale.GetValue(EmitterTime, Component) * LODLevel->SpawnModule->GetGlobalRateScale();
 		float QualityMult = SpriteTemplate->GetQualityLevelSpawnRateMult();
 		SpawnRate += LODLevel->SpawnModule->Rate.GetValue(EmitterTime, Component) * FMath::Clamp<float>(QualityMult, 0.0f, 1.0);
 	}
@@ -3293,7 +3293,7 @@ float FParticleAnimTrailEmitterInstance::Spawn(float DeltaTime)
 	// Figure out spawn rate for this tick.
 	if (bProcessSpawnRate)
 	{
-		float RateScale = LODLevel->SpawnModule->RateScale.GetValue(EmitterTime, Component);
+		float RateScale = LODLevel->SpawnModule->RateScale.GetValue(EmitterTime, Component) * LODLevel->SpawnModule->GetGlobalRateScale();
 		float QualityMult = 0.25f * (1 << Scalability::GetQualityLevels().EffectsQuality);
 		RateScale *= SpriteTemplate->QualityLevelSpawnRateScale*QualityMult;
 		SpawnRate += LODLevel->SpawnModule->Rate.GetValue(EmitterTime, Component) * FMath::Clamp<float>(RateScale, 0.0f, RateScale);
