Fixed bug in applying component instance data in the case that the component has been added in a user construction script.
#ttp 342068 - EDITOR: BP: Unable to undo deletion of a spline key
#branch UE4
#proj Runtime.Engine

diff --git a/Engine/Source/Runtime/Engine/Classes/GameFramework/Actor.h b/Engine/Source/Runtime/Engine/Classes/GameFramework/Actor.h
--- a/Engine/Source/Runtime/Engine/Classes/GameFramework/Actor.h
+++ b/Engine/Source/Runtime/Engine/Classes/GameFramework/Actor.h
@@ -2117,7 +2117,7 @@ private:
 public:
 
 	/** Array of ActorComponents that is actually serialized per-instance. */
-	UPROPERTY(TextExportTransient)
+	UPROPERTY(TextExportTransient, NonTransactional)
 	TArray<UActorComponent*> SerializedComponents;
 
 public:
