Cleanup comments and get rid of implicit behavior of destructible mesh fracture settings.
Refactor code so that each destructible setting struct fills its own settings

diff --git a/Engine/Source/Runtime/Engine/Classes/Engine/DestructibleMesh.h b/Engine/Source/Runtime/Engine/Classes/Engine/DestructibleMesh.h
--- a/Engine/Source/Runtime/Engine/Classes/Engine/DestructibleMesh.h
+++ b/Engine/Source/Runtime/Engine/Classes/Engine/DestructibleMesh.h
@@ -49,6 +49,9 @@ struct FDestructibleDepthParameters
 		: ImpactDamageOverride(IDO_None)
 		{}
 	
+	void FillDestructibleActorDesc(NxParameterized::Interface* Params, const char* OverrideName, const char* OverrideValueName) const;
+	void LoadDefaultDestructibleParametersFromApexAsset(const NxParameterized::Interface* Params, const char* OverrideName, const char* OverrideValueName);
+
 };
 
 /** Flags that apply to a destructible actor. */
@@ -150,6 +153,9 @@ struct FDestructibleParametersFlag
 	{
 	}
 
+	void FillDestructibleActorDesc(NxParameterized::Interface* Params) const;
+	void LoadDefaultDestructibleParametersFromApexAsset(const NxParameterized::Interface* Params);
+
 };
 
 /** Parameters that pertain to chunk damage. */
@@ -159,59 +165,65 @@ struct FDestructibleDamageParameters
 	GENERATED_USTRUCT_BODY()
 
 	/**
-		The damage amount which will cause a chunk to fracture (break free) from the destructible.
-		This is obtained from the damage value passed into the NxDestructibleActor::applyDamage,
-		or NxDestructibleActor::applyRadiusDamage, or via impact (see 'forceToDamage', below).
+		The damage amount which will cause a chunk to fracture (break free).
 	*/
 	UPROPERTY(EditAnywhere, Category=DestructibleDamageParameters)
 	float DamageThreshold;
 
-	/**
-		Controls the distance into the destructible to propagate damage.  The damage applied to the chunk
-		is multiplied by DamageSpread, to get the propagation distance.  All chunks within the radius
-		will have damage applied to them.  The damage applied to each chunk varies with distance to the damage
-		application position.  Full damage is taken at zero distance, and zero damage at the damage radius.
-	*/
+	/** Controls how easily damage spreads. DamageRadius = Damage*DamageSpread. All chunks within DamageRadius will take damage. Full damage is taken at zero distance, and zero damage at the DamageRadius.*/
 	UPROPERTY(EditAnywhere, Category=DestructibleDamageParameters)
 	float DamageSpread;
 
 	/**
-		If a chunk is at a depth which has impact damage set (see DepthParameters),
-		then when a chunk has a collision in the NxScene, it will take damage equal to ImpactDamage mulitplied by
-		the impact force.
-		The default value is zero, which effectively disables impact damage.
+		Whether to apply damage to destructible when colliding with an object.
+		@see ImpactDamage
 	*/
-	UPROPERTY(EditAnywhere, Category=DestructibleDamageParameters)
-	float ImpactDamage;
+	UPROPERTY(EditAnywhere, Category = DestructibleDamageParameters)
+	bool bEnableImpactDamage;
 
 	/**
-		When a chunk takes impact damage due to physical contact (see DepthParameters), this parameter
-		is the maximum impulse the contact can generate.  Weak materials such as glass may have this set to a low value, so that
-		heavier objects will pass through them during fracture.
-		N.B.: Setting this parameter to 0 disables the impulse cap; that is, zero is interpreted as infinite.
-		Default value = 0.0f.
+		Controls how much damage is applied upon collision. Damage = ImpactDamage * ImpactForce. The default value is zero, which means impact damage is disabled.
+		@see DepthParameters for per level control of ImpactDamage
 	*/
-	UPROPERTY(EditAnywhere, Category=DestructibleDamageParameters)
-	float ImpactResistance;
+	UPROPERTY(EditAnywhere, Category=DestructibleDamageParameters, meta = (editcondition = "bEnableImpactDamage") )
+	float ImpactDamage;
 
 	/**
-		By default, impact damage will only be taken to this depth.  For a particular depth, this
-		default may be overridden in the DepthParameters.  If negative, impact damage
-		is disabled.
+		Max depth level where impact damage is enabled. @see DepthParameters for per level control of ImpactDamage
+		If negative, impact damage is disabled
 	*/
-	UPROPERTY(EditAnywhere, Category=DestructibleDamageParameters)
+	UPROPERTY(EditAnywhere, Category = DestructibleDamageParameters, meta = (ClampMin = "0", UIMin = "0", editcondition = "bEnableImpactDamage"))
 	int32 DefaultImpactDamageDepth;
 
+	/**
+		By default, objects that collide with destructibles will bounce back. Custom resistance allows for finer control of how much a destructible "pushes back" against a colliding object.
+		@see ImpactResistance
+	*/
+	UPROPERTY(EditAnywhere, Category = DestructibleDamageParameters)
+	bool bCustomImpactResistance;
 
+	/**
+		Controls how much resistance is applied to colliding objects. Weak materials like glass should set this to a low value so that objects will pass right through them during fracture.
+		@see DepthParameters for per level control of ImpactResistance
+	*/
+	UPROPERTY(EditAnywhere, meta = (ClampMin = "0.001", UIMin = "0", editcondition = "bCustomImpactResistance"), Category = DestructibleDamageParameters)
+	float ImpactResistance;
 
-		FDestructibleDamageParameters()
-		: DamageThreshold(1.0f)
-		, DamageSpread(0.1f)
-		, ImpactDamage(0.0f)
-		, ImpactResistance(0.0f)
-		, DefaultImpactDamageDepth(-1)
-		{
-		}
+
+	FDestructibleDamageParameters()
+	: DamageThreshold(1.0f)
+	, DamageSpread(0.1f)
+	, bEnableImpactDamage(false)
+	, ImpactDamage(0.0f)
+	, DefaultImpactDamageDepth(-1)
+	, bCustomImpactResistance(false)
+	, ImpactResistance(0.0f)
+	
+	{
+	}
+
+	void FillDestructibleActorDesc(NxParameterized::Interface* Params, UPhysicalMaterial* PhysMat) const;
+	void LoadDefaultDestructibleParametersFromApexAsset(const NxParameterized::Interface* Params);
 	
 };
 
@@ -267,6 +279,9 @@ struct FDestructibleDebrisParameters
 		, ValidBounds(FVector(-500000.0f), FVector(500000.0f))
 		{
 		}
+
+	void FillDestructibleActorDesc(NxParameterized::Interface* Params) const;
+	void LoadDefaultDestructibleParametersFromApexAsset(const NxParameterized::Interface* Params);
 	
 };
 
@@ -315,6 +330,9 @@ struct FDestructibleAdvancedParameters
 	{
 	}
 
+	void FillDestructibleActorDesc(NxParameterized::Interface* Params) const;
+	void LoadDefaultDestructibleParametersFromApexAsset(const NxParameterized::Interface* Params);
+
 };
 
 /** Special hierarchy depths for various behaviors. */
@@ -338,13 +356,17 @@ struct FDestructibleSpecialHierarchyDepths
 	UPROPERTY(EditAnywhere, Category=DestructibleSpecialHierarchyDepths)
 	int32 MinimumFractureDepth;
 
+
+	/** 
+		Enables debris at a specific depth level.
+		@see DebrisDepth */
+	UPROPERTY(EditAnywhere, Category = DestructibleSpecialHierarchyDepths)
+	bool bEnableDebris;
+
 	/**
-		The chunk hierarchy depth at which chunks are considered to be "debris."  Chunks at this depth or
-		below will be considered for various debris settings, such as debrisLifetime.
-		Negative values indicate that no chunk depth is considered debris.
-		Default value is -1.
+		The hierarchy depth at which chunks are considered to be "debris."  
 	*/
-	UPROPERTY(EditAnywhere, Category=DestructibleSpecialHierarchyDepths)
+	UPROPERTY(EditAnywhere, Category = DestructibleSpecialHierarchyDepths, meta = (UIMin = "0", editcondition = "bEnableDebris"))
 	int32 DebrisDepth;
 
 	/**
@@ -358,14 +380,17 @@ struct FDestructibleSpecialHierarchyDepths
 
 
 
-		FDestructibleSpecialHierarchyDepths()
-		: SupportDepth(0)
-		, MinimumFractureDepth(0)
-		, DebrisDepth(-1)
-		, EssentialDepth(0)
-		{
-		}
+	FDestructibleSpecialHierarchyDepths()
+	: SupportDepth(0)
+	, MinimumFractureDepth(0)
+	, bEnableDebris(false)
+	, DebrisDepth(-1)
+	, EssentialDepth(0)
+	{
+	}
 	
+	void FillDestructibleActorDesc(NxParameterized::Interface* Params) const;
+	void LoadDefaultDestructibleParametersFromApexAsset(const NxParameterized::Interface* Params);
 };
 
 /** Parameters that apply to a destructible actor. */
@@ -401,7 +426,6 @@ struct FDestructibleParameters
 	/** A collection of flags defined in DestructibleParametersFlag. */
 	UPROPERTY(EditAnywhere, Category=DestructibleParameters)
 	struct FDestructibleParametersFlag Flags;
-
 };
 
 UCLASS(hidecategories=(Object, Mesh, LevelOfDetail, Mirroring, Physics, Reimport, Clothing), MinimalAPI)
diff --git a/Engine/Source/Runtime/Engine/Private/DestructibleMesh.cpp b/Engine/Source/Runtime/Engine/Private/DestructibleMesh.cpp
--- a/Engine/Source/Runtime/Engine/Private/DestructibleMesh.cpp
+++ b/Engine/Source/Runtime/Engine/Private/DestructibleMesh.cpp
@@ -7,6 +7,7 @@
 #include "EnginePrivate.h"
 #include "PhysicsPublic.h"
 #include "StaticMeshResources.h"
+#include "Engine/DestructibleMesh.h"
 #include "Engine/DestructibleFractureSettings.h"
 #include "PhysicsEngine/PhysXSupport.h"
 
@@ -83,6 +84,48 @@ void UDestructibleMesh::PreEditChange(UProperty* PropertyAboutToChange)
 }
 #endif
 
+float ImpactResistanceToAPEX(bool bCustomResistance, float ImpactResistance)
+{
+	return bCustomResistance ? ImpactResistance : 0.f;
+}
+
+
+void APEXToImpactResistance(bool& bCustomImpactResistance, float& ImpactResistance)
+{
+	//apex interprets 0 as disabled, but we want custom flag
+	bCustomImpactResistance = ImpactResistance != 0.f;
+	if (ImpactResistance == 0.f)
+	{
+		ImpactResistance = 1.f;
+	}
+}
+
+int32 DefaultImpactDamageDepthToAPEX(bool bEnableImpactDamage, int32 DefaultImpactDamageDepth)
+{
+	return bEnableImpactDamage ? DefaultImpactDamageDepth : -1.f;
+}
+
+void APEXToDefaultImpactDamageDepth(bool& bEnableImpactDamage, int32& DefaultImpactDamageDepth)
+{
+	//apex interprets -1 as disabled, but we want custom flag
+	bEnableImpactDamage = DefaultImpactDamageDepth != -1;
+	DefaultImpactDamageDepth = 1;
+}
+
+int32 DebrisDepthToAPEX(bool bEnableDebris, int32 DebrisDepth)
+{
+	return bEnableDebris ? DebrisDepth : -1;
+}
+
+void APEXToDebrisDepth(bool& bEnableDebris, int32& DebrisDepth)
+{
+	bEnableDebris = DebrisDepth != -1;
+	if (DebrisDepth == -1)
+	{
+		DebrisDepth = 0;
+	}
+}
+
 void UDestructibleMesh::Serialize(FArchive& Ar)
 {
 	Super::Serialize( Ar );
@@ -226,6 +269,13 @@ void UDestructibleMesh::Serialize(FArchive& Ar)
 		Ar << size;
 #endif
 	}
+
+	if (Ar.UE4Ver() < VER_UE4_CLEAN_DESTRUCTIBLE_SETTINGS)
+	{
+		APEXToImpactResistance(DefaultDestructibleParameters.DamageParameters.bCustomImpactResistance, DefaultDestructibleParameters.DamageParameters.ImpactResistance);
+		APEXToDefaultImpactDamageDepth(DefaultDestructibleParameters.DamageParameters.bEnableImpactDamage, DefaultDestructibleParameters.DamageParameters.DefaultImpactDamageDepth);
+		APEXToDebrisDepth(DefaultDestructibleParameters.SpecialHierarchyDepths.bEnableDebris, DefaultDestructibleParameters.SpecialHierarchyDepths.DebrisDepth);
+	}
 }
 
 void UDestructibleMesh::FinishDestroy()
@@ -242,6 +292,176 @@ void UDestructibleMesh::FinishDestroy()
 }
 
 #if WITH_APEX
+
+void FDestructibleDamageParameters::FillDestructibleActorDesc(NxParameterized::Interface* Params, UPhysicalMaterial* PhysMat) const
+{
+	// Damage parameters
+	verify(NxParameterized::setParamF32(*Params, "defaultBehaviorGroup.damageThreshold", DamageThreshold * PhysMat->DestructibleDamageThresholdScale));
+	verify(NxParameterized::setParamF32(*Params, "defaultBehaviorGroup.damageToRadius", DamageSpread));
+	verify(NxParameterized::setParamF32(*Params, "destructibleParameters.forceToDamage", ImpactDamage));
+	verify(NxParameterized::setParamF32(*Params, "defaultBehaviorGroup.materialStrength", ImpactResistanceToAPEX(bCustomImpactResistance, ImpactResistance)));
+	verify(NxParameterized::setParamI32(*Params, "destructibleParameters.impactDamageDefaultDepth", DefaultImpactDamageDepthToAPEX(bEnableImpactDamage, DefaultImpactDamageDepth)));
+}
+
+void FDestructibleDamageParameters::LoadDefaultDestructibleParametersFromApexAsset(const NxParameterized::Interface* Params)
+{
+	// Damage parameters
+	verify(NxParameterized::getParamF32(*Params, "defaultBehaviorGroup.damageThreshold", DamageThreshold));
+	verify(NxParameterized::getParamF32(*Params, "defaultBehaviorGroup.damageToRadius", DamageSpread));
+	verify(NxParameterized::getParamF32(*Params, "destructibleParameters.forceToDamage", ImpactDamage));
+	verify(NxParameterized::getParamF32(*Params, "defaultBehaviorGroup.materialStrength",ImpactResistance));
+	verify(NxParameterized::getParamI32(*Params, "destructibleParameters.impactDamageDefaultDepth", DefaultImpactDamageDepth));
+
+	APEXToImpactResistance(bCustomImpactResistance, ImpactResistance);
+	APEXToDefaultImpactDamageDepth(bEnableImpactDamage, DefaultImpactDamageDepth);
+}
+
+void FDestructibleSpecialHierarchyDepths::FillDestructibleActorDesc(NxParameterized::Interface* Params) const
+{
+	verify(NxParameterized::setParamU32(*Params, "supportDepth", SupportDepth));
+	verify(NxParameterized::setParamU32(*Params, "destructibleParameters.minimumFractureDepth", MinimumFractureDepth));
+	verify(NxParameterized::setParamI32(*Params, "destructibleParameters.debrisDepth", DebrisDepthToAPEX(bEnableDebris,DebrisDepth)) );
+	verify(NxParameterized::setParamU32(*Params, "destructibleParameters.essentialDepth", EssentialDepth));
+}
+
+void FDestructibleSpecialHierarchyDepths::LoadDefaultDestructibleParametersFromApexAsset(const NxParameterized::Interface* Params)
+{
+	physx::PxU32 PSupportDepth;
+	verify(NxParameterized::getParamU32(*Params, "supportDepth", PSupportDepth));
+	SupportDepth = (int32)PSupportDepth;
+	physx::PxU32 PMinimumFractureDepth;
+	verify(NxParameterized::getParamU32(*Params, "destructibleParameters.minimumFractureDepth", PMinimumFractureDepth));
+	MinimumFractureDepth = (int32)PMinimumFractureDepth;
+	verify(NxParameterized::getParamI32(*Params, "destructibleParameters.debrisDepth", DebrisDepth));
+	DebrisDepthToAPEX(bEnableDebris, DebrisDepth);
+
+	physx::PxU32 PEssentialDepth;
+	verify(NxParameterized::getParamU32(*Params, "destructibleParameters.essentialDepth", PEssentialDepth));
+	EssentialDepth = (int32)PEssentialDepth;
+}
+
+void FDestructibleAdvancedParameters::FillDestructibleActorDesc(NxParameterized::Interface* Params) const
+{
+	verify(NxParameterized::setParamF32(*Params, "destructibleParameters.damageCap", DamageCap));
+	verify(NxParameterized::setParamF32(*Params, "destructibleParameters.impactVelocityThreshold", ImpactVelocityThreshold));
+	verify(NxParameterized::setParamF32(*Params, "destructibleParameters.maxChunkSpeed", MaxChunkSpeed));
+	verify(NxParameterized::setParamF32(*Params, "destructibleParameters.fractureImpulseScale", FractureImpulseScale));
+}
+
+void FDestructibleAdvancedParameters::LoadDefaultDestructibleParametersFromApexAsset(const NxParameterized::Interface* Params)
+{
+	verify(NxParameterized::getParamF32(*Params, "destructibleParameters.damageCap", DamageCap));
+	verify(NxParameterized::getParamF32(*Params, "destructibleParameters.impactVelocityThreshold", ImpactVelocityThreshold));
+	verify(NxParameterized::getParamF32(*Params, "destructibleParameters.maxChunkSpeed", MaxChunkSpeed));
+	verify(NxParameterized::getParamF32(*Params, "destructibleParameters.fractureImpulseScale", FractureImpulseScale));
+}
+
+void FDestructibleDebrisParameters::FillDestructibleActorDesc(NxParameterized::Interface* Params) const
+{
+	verify(NxParameterized::setParamF32(*Params, "destructibleParameters.debrisLifetimeMin", DebrisLifetimeMin));
+	verify(NxParameterized::setParamF32(*Params, "destructibleParameters.debrisLifetimeMax", DebrisLifetimeMax));
+	verify(NxParameterized::setParamF32(*Params, "destructibleParameters.debrisMaxSeparationMin", DebrisMaxSeparationMin));
+	verify(NxParameterized::setParamF32(*Params, "destructibleParameters.debrisMaxSeparationMax", DebrisMaxSeparationMax));
+	physx::PxBounds3 PValidBounds;
+	PValidBounds.minimum.x = ValidBounds.Min.X;
+	PValidBounds.minimum.y = ValidBounds.Min.Y;
+	PValidBounds.minimum.z = ValidBounds.Min.Z;
+	PValidBounds.maximum.x = ValidBounds.Max.X;
+	PValidBounds.maximum.y = ValidBounds.Max.Y;
+	PValidBounds.maximum.z = ValidBounds.Max.Z;
+	verify(NxParameterized::setParamBounds3(*Params, "destructibleParameters.validBounds", PValidBounds));
+}
+
+void FDestructibleDebrisParameters::LoadDefaultDestructibleParametersFromApexAsset(const NxParameterized::Interface* Params)
+{
+	verify(NxParameterized::getParamF32(*Params, "destructibleParameters.debrisLifetimeMin", DebrisLifetimeMin));
+	verify(NxParameterized::getParamF32(*Params, "destructibleParameters.debrisLifetimeMax", DebrisLifetimeMax));
+	verify(NxParameterized::getParamF32(*Params, "destructibleParameters.debrisMaxSeparationMin", DebrisMaxSeparationMin));
+	verify(NxParameterized::getParamF32(*Params, "destructibleParameters.debrisMaxSeparationMax", DebrisMaxSeparationMax));
+	physx::PxBounds3 PValidBounds;
+	verify(NxParameterized::getParamBounds3(*Params, "destructibleParameters.validBounds", PValidBounds));
+	ValidBounds.Min.X = PValidBounds.minimum.x;
+	ValidBounds.Min.Y = PValidBounds.minimum.y;
+	ValidBounds.Min.Z = PValidBounds.minimum.z;
+	ValidBounds.Max.X = PValidBounds.maximum.x;
+	ValidBounds.Max.Y = PValidBounds.maximum.y;
+	ValidBounds.Max.Z = PValidBounds.maximum.z;
+
+}
+
+void FDestructibleParametersFlag::FillDestructibleActorDesc(NxParameterized::Interface* Params) const
+{
+	verify(NxParameterized::setParamBool(*Params, "destructibleParameters.flags.ACCUMULATE_DAMAGE", bAccumulateDamage));
+	verify(NxParameterized::setParamBool(*Params, "useAssetDefinedSupport", bAssetDefinedSupport));
+	verify(NxParameterized::setParamBool(*Params, "useWorldSupport", bWorldSupport));
+	verify(NxParameterized::setParamBool(*Params, "destructibleParameters.flags.DEBRIS_TIMEOUT", bDebrisTimeout));
+	verify(NxParameterized::setParamBool(*Params, "destructibleParameters.flags.DEBRIS_MAX_SEPARATION", bDebrisMaxSeparation));
+	verify(NxParameterized::setParamBool(*Params, "destructibleParameters.flags.CRUMBLE_SMALLEST_CHUNKS", bCrumbleSmallestChunks));
+	verify(NxParameterized::setParamBool(*Params, "destructibleParameters.flags.ACCURATE_RAYCASTS", bAccurateRaycasts));
+	verify(NxParameterized::setParamBool(*Params, "destructibleParameters.flags.USE_VALID_BOUNDS", bUseValidBounds));
+	verify(NxParameterized::setParamBool(*Params, "formExtendedStructures", bFormExtendedStructures));
+}
+
+void FDestructibleParametersFlag::LoadDefaultDestructibleParametersFromApexAsset(const NxParameterized::Interface* Params)
+{
+	bool bFlag;
+	verify(NxParameterized::getParamBool(*Params, "destructibleParameters.flags.ACCUMULATE_DAMAGE", bFlag));
+	bAccumulateDamage = bFlag ? 1 : 0;
+	verify(NxParameterized::getParamBool(*Params, "useAssetDefinedSupport", bFlag));
+	bAssetDefinedSupport = bFlag ? 1 : 0;
+	verify(NxParameterized::getParamBool(*Params, "useWorldSupport", bFlag));
+	bWorldSupport = bFlag ? 1 : 0;
+	verify(NxParameterized::getParamBool(*Params, "destructibleParameters.flags.DEBRIS_TIMEOUT", bFlag));
+	bDebrisTimeout = bFlag ? 1 : 0;
+	verify(NxParameterized::getParamBool(*Params, "destructibleParameters.flags.DEBRIS_MAX_SEPARATION", bFlag));
+	bDebrisMaxSeparation = bFlag ? 1 : 0;
+	verify(NxParameterized::getParamBool(*Params, "destructibleParameters.flags.CRUMBLE_SMALLEST_CHUNKS", bFlag));
+	bCrumbleSmallestChunks = bFlag ? 1 : 0;
+	verify(NxParameterized::getParamBool(*Params, "destructibleParameters.flags.ACCURATE_RAYCASTS", bFlag));
+	bAccurateRaycasts = bFlag ? 1 : 0;
+	verify(NxParameterized::getParamBool(*Params, "destructibleParameters.flags.USE_VALID_BOUNDS", bFlag));
+	bUseValidBounds = bFlag ? 1 : 0;
+	verify(NxParameterized::getParamBool(*Params, "formExtendedStructures", bFlag));
+	bFormExtendedStructures = bFlag ? 1 : 0;
+}
+
+void FDestructibleDepthParameters::FillDestructibleActorDesc(NxParameterized::Interface* Params, const char* OverrideName, const char* OverrideValueName) const
+{
+	switch (ImpactDamageOverride)
+	{
+	case IDO_None:
+		verify(NxParameterized::setParamBool(*Params, OverrideName, false));
+		break;
+	case IDO_On:
+		verify(NxParameterized::setParamBool(*Params, OverrideName, true));
+		verify(NxParameterized::setParamBool(*Params, OverrideValueName, true));
+		break;
+	case IDO_Off:
+		verify(NxParameterized::setParamBool(*Params, OverrideName, true));
+		verify(NxParameterized::setParamBool(*Params, OverrideValueName, false));
+		break;
+	default:
+		break;
+	}
+}
+
+void FDestructibleDepthParameters::LoadDefaultDestructibleParametersFromApexAsset(const NxParameterized::Interface* Params, const char* OverrideName, const char* OverrideValueName)
+{
+	bool bOverride;
+	bool bOverrideValue;
+	verify(NxParameterized::getParamBool(*Params, OverrideName, bOverride));
+	verify(NxParameterized::getParamBool(*Params, OverrideValueName, bOverrideValue));
+
+	if (!bOverride)
+	{
+		ImpactDamageOverride = IDO_None;
+	}
+	else
+	{
+		ImpactDamageOverride = bOverrideValue ? IDO_On : IDO_Off;
+	}
+}
+
 NxParameterized::Interface* UDestructibleMesh::GetDestructibleActorDesc(UPhysicalMaterial* PhysMat)
 {
 	NxParameterized::Interface* Params = NULL;
@@ -252,51 +472,13 @@ NxParameterized::Interface* UDestructibleMesh::GetDestructibleActorDesc(UPhysica
 	}
 	if (Params != NULL)
 	{
-		// Damage parameters
-		verify( NxParameterized::setParamF32(*Params,"defaultBehaviorGroup.damageThreshold",				DefaultDestructibleParameters.DamageParameters.DamageThreshold * PhysMat->DestructibleDamageThresholdScale) );
-		verify( NxParameterized::setParamF32(*Params,"defaultBehaviorGroup.damageToRadius",					DefaultDestructibleParameters.DamageParameters.DamageSpread) );
-		verify( NxParameterized::setParamF32(*Params,"destructibleParameters.forceToDamage",				DefaultDestructibleParameters.DamageParameters.ImpactDamage) );
-		verify( NxParameterized::setParamF32(*Params,"defaultBehaviorGroup.materialStrength",				DefaultDestructibleParameters.DamageParameters.ImpactResistance) );
-		verify( NxParameterized::setParamI32(*Params,"destructibleParameters.impactDamageDefaultDepth",		DefaultDestructibleParameters.DamageParameters.DefaultImpactDamageDepth) );
-
-		// Special depths
-		verify( NxParameterized::setParamU32(*Params,"supportDepth",										DefaultDestructibleParameters.SpecialHierarchyDepths.SupportDepth) );
-		verify( NxParameterized::setParamU32(*Params,"destructibleParameters.minimumFractureDepth",			DefaultDestructibleParameters.SpecialHierarchyDepths.MinimumFractureDepth) );
-		verify( NxParameterized::setParamI32(*Params,"destructibleParameters.debrisDepth",					DefaultDestructibleParameters.SpecialHierarchyDepths.DebrisDepth) );
-		verify( NxParameterized::setParamU32(*Params,"destructibleParameters.essentialDepth",				DefaultDestructibleParameters.SpecialHierarchyDepths.EssentialDepth) );
-
-		// Advanced parameters
-		verify( NxParameterized::setParamF32(*Params,"destructibleParameters.damageCap",					DefaultDestructibleParameters.AdvancedParameters.DamageCap) );
-		verify( NxParameterized::setParamF32(*Params,"destructibleParameters.impactVelocityThreshold",		DefaultDestructibleParameters.AdvancedParameters.ImpactVelocityThreshold) );
-		verify( NxParameterized::setParamF32(*Params,"destructibleParameters.maxChunkSpeed",				DefaultDestructibleParameters.AdvancedParameters.MaxChunkSpeed) );
-		verify( NxParameterized::setParamF32(*Params,"destructibleParameters.fractureImpulseScale",			DefaultDestructibleParameters.AdvancedParameters.FractureImpulseScale) );
-
-		// Debris parameters
-		verify( NxParameterized::setParamF32(*Params,"destructibleParameters.debrisLifetimeMin",			DefaultDestructibleParameters.DebrisParameters.DebrisLifetimeMin) );
-		verify( NxParameterized::setParamF32(*Params,"destructibleParameters.debrisLifetimeMax",			DefaultDestructibleParameters.DebrisParameters.DebrisLifetimeMax) );
-		verify( NxParameterized::setParamF32(*Params,"destructibleParameters.debrisMaxSeparationMin",		DefaultDestructibleParameters.DebrisParameters.DebrisMaxSeparationMin) );
-		verify( NxParameterized::setParamF32(*Params,"destructibleParameters.debrisMaxSeparationMax",		DefaultDestructibleParameters.DebrisParameters.DebrisMaxSeparationMax) );
-		physx::PxBounds3 ValidBounds;
-		ValidBounds.minimum.x = DefaultDestructibleParameters.DebrisParameters.ValidBounds.Min.X;
-		ValidBounds.minimum.y = DefaultDestructibleParameters.DebrisParameters.ValidBounds.Min.Y;
-		ValidBounds.minimum.z = DefaultDestructibleParameters.DebrisParameters.ValidBounds.Min.Z;															   
-		ValidBounds.maximum.x = DefaultDestructibleParameters.DebrisParameters.ValidBounds.Max.X;
-		ValidBounds.maximum.y = DefaultDestructibleParameters.DebrisParameters.ValidBounds.Max.Y;
-		ValidBounds.maximum.z = DefaultDestructibleParameters.DebrisParameters.ValidBounds.Max.Z;
-		verify( NxParameterized::setParamBounds3(*Params,"destructibleParameters.validBounds",					ValidBounds) );
-
-		// Flags
-		verify( NxParameterized::setParamBool(*Params,"destructibleParameters.flags.ACCUMULATE_DAMAGE",			DefaultDestructibleParameters.Flags.bAccumulateDamage) );
-		verify( NxParameterized::setParamBool(*Params,"useAssetDefinedSupport",									DefaultDestructibleParameters.Flags.bAssetDefinedSupport) );
-		verify( NxParameterized::setParamBool(*Params,"useWorldSupport",										DefaultDestructibleParameters.Flags.bWorldSupport) );
-		verify( NxParameterized::setParamBool(*Params,"destructibleParameters.flags.DEBRIS_TIMEOUT",			DefaultDestructibleParameters.Flags.bDebrisTimeout) );
-		verify( NxParameterized::setParamBool(*Params,"destructibleParameters.flags.DEBRIS_MAX_SEPARATION",		DefaultDestructibleParameters.Flags.bDebrisMaxSeparation) );
-		verify( NxParameterized::setParamBool(*Params,"destructibleParameters.flags.CRUMBLE_SMALLEST_CHUNKS",	DefaultDestructibleParameters.Flags.bCrumbleSmallestChunks) );
-		verify( NxParameterized::setParamBool(*Params,"destructibleParameters.flags.ACCURATE_RAYCASTS",			DefaultDestructibleParameters.Flags.bAccurateRaycasts) );
-		verify( NxParameterized::setParamBool(*Params,"destructibleParameters.flags.USE_VALID_BOUNDS",			DefaultDestructibleParameters.Flags.bUseValidBounds) );
-		verify( NxParameterized::setParamBool(*Params,"formExtendedStructures",									DefaultDestructibleParameters.Flags.bFormExtendedStructures) );
-
-		// Depth parameters
+		DefaultDestructibleParameters.DamageParameters.FillDestructibleActorDesc(Params, PhysMat);
+		DefaultDestructibleParameters.SpecialHierarchyDepths.FillDestructibleActorDesc(Params);
+		DefaultDestructibleParameters.AdvancedParameters.FillDestructibleActorDesc(Params);
+		DefaultDestructibleParameters.DebrisParameters.FillDestructibleActorDesc(Params);
+		DefaultDestructibleParameters.Flags.FillDestructibleActorDesc(Params);
+		
+		//depth params done per level
 		for (int32 Depth = 0; Depth < DefaultDestructibleParameters.DepthParameters.Num(); ++Depth)
 		{
 			char OverrideName[MAX_SPRINTF];
@@ -304,22 +486,7 @@ NxParameterized::Interface* UDestructibleMesh::GetDestructibleActorDesc(UPhysica
 			char OverrideValueName[MAX_SPRINTF];
 			FCStringAnsi::Sprintf(OverrideValueName, "depthParameters[%d].OVERRIDE_IMPACT_DAMAGE_VALUE", Depth);
 
-			switch (DefaultDestructibleParameters.DepthParameters[Depth].ImpactDamageOverride)
-			{
-			case IDO_None:
-				verify( NxParameterized::setParamBool(*Params, OverrideName, false) );
-				break;
-			case IDO_On:
-				verify( NxParameterized::setParamBool(*Params, OverrideName, true) );
-				verify( NxParameterized::setParamBool(*Params, OverrideValueName, true) );
-				break;
-			case IDO_Off:
-				verify( NxParameterized::setParamBool(*Params, OverrideName, true) );
-				verify( NxParameterized::setParamBool(*Params, OverrideValueName, false) );
-				break;
-			default:
-				break;
-			}
+			DefaultDestructibleParameters.DepthParameters[Depth].FillDestructibleActorDesc(Params, OverrideName, OverrideValueName);
 		}
 	}
 
@@ -327,7 +494,6 @@ NxParameterized::Interface* UDestructibleMesh::GetDestructibleActorDesc(UPhysica
 }
 #endif // WITH_APEX
 
-
 void UDestructibleMesh::LoadDefaultDestructibleParametersFromApexAsset()
 {
 #if WITH_APEX
@@ -335,65 +501,11 @@ void UDestructibleMesh::LoadDefaultDestructibleParametersFromApexAsset()
 
 	if (Params != NULL)
 	{
-		// Damage parameters
-		verify( NxParameterized::getParamF32(*Params,"defaultBehaviorGroup.damageThreshold",				DefaultDestructibleParameters.DamageParameters.DamageThreshold) );
-		verify( NxParameterized::getParamF32(*Params,"defaultBehaviorGroup.damageToRadius",				DefaultDestructibleParameters.DamageParameters.DamageSpread) );
-		verify( NxParameterized::getParamF32(*Params,"destructibleParameters.forceToDamage",				DefaultDestructibleParameters.DamageParameters.ImpactDamage) );
-		verify( NxParameterized::getParamF32(*Params,"defaultBehaviorGroup.materialStrength",				DefaultDestructibleParameters.DamageParameters.ImpactResistance) );
-		verify( NxParameterized::getParamI32(*Params,"destructibleParameters.impactDamageDefaultDepth",		DefaultDestructibleParameters.DamageParameters.DefaultImpactDamageDepth) );
-
-		// Special depths
-		physx::PxU32 SupportDepth;
-		verify( NxParameterized::getParamU32(*Params,"supportDepth",										SupportDepth) );
-		DefaultDestructibleParameters.SpecialHierarchyDepths.SupportDepth = (int32)SupportDepth;
-		physx::PxU32 MinimumFractureDepth;
-		verify( NxParameterized::getParamU32(*Params,"destructibleParameters.minimumFractureDepth",			MinimumFractureDepth) );
-		DefaultDestructibleParameters.SpecialHierarchyDepths.MinimumFractureDepth = (int32)MinimumFractureDepth;
-		verify( NxParameterized::getParamI32(*Params,"destructibleParameters.debrisDepth",					DefaultDestructibleParameters.SpecialHierarchyDepths.DebrisDepth) );
-		physx::PxU32 EssentialDepth;
-		verify( NxParameterized::getParamU32(*Params,"destructibleParameters.essentialDepth",				EssentialDepth) );
-		DefaultDestructibleParameters.SpecialHierarchyDepths.EssentialDepth = (int32)EssentialDepth;
-
-		// Advanced parameters
-		verify( NxParameterized::getParamF32(*Params,"destructibleParameters.damageCap",					DefaultDestructibleParameters.AdvancedParameters.DamageCap) );
-		verify( NxParameterized::getParamF32(*Params,"destructibleParameters.impactVelocityThreshold",		DefaultDestructibleParameters.AdvancedParameters.ImpactVelocityThreshold) );
-		verify( NxParameterized::getParamF32(*Params,"destructibleParameters.maxChunkSpeed",				DefaultDestructibleParameters.AdvancedParameters.MaxChunkSpeed) );
-		verify( NxParameterized::getParamF32(*Params,"destructibleParameters.fractureImpulseScale",			DefaultDestructibleParameters.AdvancedParameters.FractureImpulseScale) );
-
-		// Debris parameters
-		verify( NxParameterized::getParamF32(*Params,"destructibleParameters.debrisLifetimeMin",			DefaultDestructibleParameters.DebrisParameters.DebrisLifetimeMin) );
-		verify( NxParameterized::getParamF32(*Params,"destructibleParameters.debrisLifetimeMax",			DefaultDestructibleParameters.DebrisParameters.DebrisLifetimeMax) );
-		verify( NxParameterized::getParamF32(*Params,"destructibleParameters.debrisMaxSeparationMin",		DefaultDestructibleParameters.DebrisParameters.DebrisMaxSeparationMin) );
-		verify( NxParameterized::getParamF32(*Params,"destructibleParameters.debrisMaxSeparationMax",		DefaultDestructibleParameters.DebrisParameters.DebrisMaxSeparationMax) );
-		physx::PxBounds3 ValidBounds;
-		verify( NxParameterized::getParamBounds3(*Params,"destructibleParameters.validBounds",				ValidBounds) );
-		DefaultDestructibleParameters.DebrisParameters.ValidBounds.Min.X = ValidBounds.minimum.x;
-		DefaultDestructibleParameters.DebrisParameters.ValidBounds.Min.Y = ValidBounds.minimum.y;
-		DefaultDestructibleParameters.DebrisParameters.ValidBounds.Min.Z = ValidBounds.minimum.z;															   
-		DefaultDestructibleParameters.DebrisParameters.ValidBounds.Max.X = ValidBounds.maximum.x;
-		DefaultDestructibleParameters.DebrisParameters.ValidBounds.Max.Y = ValidBounds.maximum.y;
-		DefaultDestructibleParameters.DebrisParameters.ValidBounds.Max.Z = ValidBounds.maximum.z;
-
-		// Flags
-		bool bFlag;
-		verify( NxParameterized::getParamBool(*Params,"destructibleParameters.flags.ACCUMULATE_DAMAGE",			bFlag) );
-		DefaultDestructibleParameters.Flags.bAccumulateDamage = bFlag ? 1 : 0;
-		verify( NxParameterized::getParamBool(*Params,"useAssetDefinedSupport",									bFlag) );
-		DefaultDestructibleParameters.Flags.bAssetDefinedSupport = bFlag ? 1 : 0;
-		verify( NxParameterized::getParamBool(*Params,"useWorldSupport",										bFlag) );
-		DefaultDestructibleParameters.Flags.bWorldSupport = bFlag ? 1 : 0;
-		verify( NxParameterized::getParamBool(*Params,"destructibleParameters.flags.DEBRIS_TIMEOUT",			bFlag) );
-		DefaultDestructibleParameters.Flags.bDebrisTimeout = bFlag ? 1 : 0;
-		verify( NxParameterized::getParamBool(*Params,"destructibleParameters.flags.DEBRIS_MAX_SEPARATION",		bFlag) );
-		DefaultDestructibleParameters.Flags.bDebrisMaxSeparation = bFlag ? 1 : 0;
-		verify( NxParameterized::getParamBool(*Params,"destructibleParameters.flags.CRUMBLE_SMALLEST_CHUNKS",	bFlag) );
-		DefaultDestructibleParameters.Flags.bCrumbleSmallestChunks = bFlag ? 1 : 0;
-		verify( NxParameterized::getParamBool(*Params,"destructibleParameters.flags.ACCURATE_RAYCASTS",			bFlag) );
-		DefaultDestructibleParameters.Flags.bAccurateRaycasts = bFlag ? 1 : 0;
-		verify( NxParameterized::getParamBool(*Params,"destructibleParameters.flags.USE_VALID_BOUNDS",			bFlag) );
-		DefaultDestructibleParameters.Flags.bUseValidBounds = bFlag ? 1 : 0;
-		verify( NxParameterized::getParamBool(*Params,"formExtendedStructures",									bFlag) );
-		DefaultDestructibleParameters.Flags.bFormExtendedStructures = bFlag ? 1 : 0;
+		DefaultDestructibleParameters.DebrisParameters.LoadDefaultDestructibleParametersFromApexAsset(Params);
+		DefaultDestructibleParameters.SpecialHierarchyDepths.LoadDefaultDestructibleParametersFromApexAsset(Params);
+		DefaultDestructibleParameters.AdvancedParameters.LoadDefaultDestructibleParametersFromApexAsset(Params);
+		DefaultDestructibleParameters.DebrisParameters.LoadDefaultDestructibleParametersFromApexAsset(Params);
+		DefaultDestructibleParameters.Flags.LoadDefaultDestructibleParametersFromApexAsset(Params);
 
 		// Depth parameters
 		for (int32 Depth = 0; Depth < DefaultDestructibleParameters.DepthParameters.Num(); ++Depth)
@@ -403,19 +515,7 @@ void UDestructibleMesh::LoadDefaultDestructibleParametersFromApexAsset()
 			char OverrideValueName[MAX_SPRINTF];
 			FCStringAnsi::Sprintf(OverrideValueName, "depthParameters[%d].OVERRIDE_IMPACT_DAMAGE_VALUE", Depth);
 
-			bool bOverride;
-			bool bOverrideValue;
-			verify( NxParameterized::getParamBool(*Params, OverrideName, bOverride) );
-			verify( NxParameterized::getParamBool(*Params, OverrideValueName, bOverrideValue) );
-
-			if (!bOverride)
-			{
-				DefaultDestructibleParameters.DepthParameters[Depth].ImpactDamageOverride = IDO_None;
-			}
-			else
-			{
-				DefaultDestructibleParameters.DepthParameters[Depth].ImpactDamageOverride = bOverrideValue ? IDO_On : IDO_Off;
-			}
+			DefaultDestructibleParameters.DepthParameters[Depth].LoadDefaultDestructibleParametersFromApexAsset(Params, OverrideName, OverrideValueName);
 		}
 	}
 #endif
