fix non apex builds

diff --git a/Engine/Source/Runtime/Engine/Classes/Engine/DestructibleMesh.h b/Engine/Source/Runtime/Engine/Classes/Engine/DestructibleMesh.h
--- a/Engine/Source/Runtime/Engine/Classes/Engine/DestructibleMesh.h
+++ b/Engine/Source/Runtime/Engine/Classes/Engine/DestructibleMesh.h
@@ -49,8 +49,10 @@ struct FDestructibleDepthParameters
 		: ImpactDamageOverride(IDO_None)
 		{}
 	
+#if WITH_APEX
 	void FillDestructibleActorDesc(NxParameterized::Interface* Params, const char* OverrideName, const char* OverrideValueName) const;
 	void LoadDefaultDestructibleParametersFromApexAsset(const NxParameterized::Interface* Params, const char* OverrideName, const char* OverrideValueName);
+#endif
 
 };
 
@@ -153,8 +155,10 @@ struct FDestructibleParametersFlag
 	{
 	}
 
+#if WITH_APEX
 	void FillDestructibleActorDesc(NxParameterized::Interface* Params) const;
 	void LoadDefaultDestructibleParametersFromApexAsset(const NxParameterized::Interface* Params);
+#endif
 
 };
 
@@ -222,8 +226,10 @@ struct FDestructibleDamageParameters
 	{
 	}
 
+#if WITH_APEX
 	void FillDestructibleActorDesc(NxParameterized::Interface* Params, UPhysicalMaterial* PhysMat) const;
 	void LoadDefaultDestructibleParametersFromApexAsset(const NxParameterized::Interface* Params);
+#endif
 	
 };
 
@@ -280,8 +286,10 @@ struct FDestructibleDebrisParameters
 		{
 		}
 
+#if WITH_APEX
 	void FillDestructibleActorDesc(NxParameterized::Interface* Params) const;
 	void LoadDefaultDestructibleParametersFromApexAsset(const NxParameterized::Interface* Params);
+#endif
 	
 };
 
@@ -330,8 +338,10 @@ struct FDestructibleAdvancedParameters
 	{
 	}
 
+#if WITH_APEX
 	void FillDestructibleActorDesc(NxParameterized::Interface* Params) const;
 	void LoadDefaultDestructibleParametersFromApexAsset(const NxParameterized::Interface* Params);
+#endif
 
 };
 
@@ -389,8 +399,10 @@ struct FDestructibleSpecialHierarchyDepths
 	{
 	}
 	
+#if WITH_APEX
 	void FillDestructibleActorDesc(NxParameterized::Interface* Params) const;
 	void LoadDefaultDestructibleParametersFromApexAsset(const NxParameterized::Interface* Params);
+#endif
 };
 
 /** Parameters that apply to a destructible actor. */
