Update ReplicatedComponents in parallel with OwnedComponents

diff --git a/Engine/Source/Runtime/Engine/Classes/GameFramework/Actor.h b/Engine/Source/Runtime/Engine/Classes/GameFramework/Actor.h
--- a/Engine/Source/Runtime/Engine/Classes/GameFramework/Actor.h
+++ b/Engine/Source/Runtime/Engine/Classes/GameFramework/Actor.h
@@ -236,10 +236,6 @@ public:
 	/** Dormancy setting for actor to take itself off of the replication list without being destroyed on clients. */
 	TEnumAsByte<enum ENetDormancy> NetDormancy;
 
-	/** List of replicated components. Cooked builds won't rebuild this array. */
-	UPROPERTY()
-	TArray< TWeakObjectPtr<class UActorComponent> >	ReplicatedComponents;
-
 	/** Automatically registers this actor to receive input from a player. */
 	UPROPERTY(EditAnywhere, Category=Input)
 	TEnumAsByte<EAutoReceiveInput::Type> AutoReceiveInput;
@@ -2107,13 +2103,24 @@ public:
 	 */
 	void ResetOwnedComponents();
 
+	/** Called when the replicated state of a component changes to update the Actor's cached ReplicatedComponents array
+	 */
+	void UpdateReplicatedComponent(UActorComponent* Component);
+
+	/** Returns a constant reference to the replicated components array
+	 */
+	const TArray<UActorComponent*>& GetReplicatedComponents() const;
+
 private:
 	/**
 	 * All ActorComponents owned by this Actor.
 	 * @see GetComponents()
 	 */
 	TArray<UActorComponent*> OwnedComponents;
-	
+
+	/** List of replicated components. */
+	TArray<UActorComponent*> ReplicatedComponents;
+
 public:
 
 	/** Array of ActorComponents that is actually serialized per-instance. */
diff --git a/Engine/Source/Runtime/Engine/Private/Actor.cpp b/Engine/Source/Runtime/Engine/Private/Actor.cpp
--- a/Engine/Source/Runtime/Engine/Private/Actor.cpp
+++ b/Engine/Source/Runtime/Engine/Private/Actor.cpp
@@ -187,6 +187,7 @@ void AActor::ResetOwnedComponents()
 {
 	TArray<UObject*> Children;
 	OwnedComponents.Empty();
+	ReplicatedComponents.Empty();
 	GetObjectsWithOuter(this, Children, true, RF_PendingKill);
 
 	for (UObject* Child : Children)
@@ -195,6 +196,11 @@ void AActor::ResetOwnedComponents()
 		if (Component)
 		{
 			OwnedComponents.Add(Component);
+
+			if (Component->GetIsReplicated())
+			{
+				ReplicatedComponents.Add(Component);
+			}
 		}
 	}
 }
@@ -1962,6 +1968,11 @@ void AActor::AddOwnedComponent(UActorComponent* Component)
 {
 	check(Component->GetOwner() == this);
 	OwnedComponents.AddUnique(Component);
+
+	if (Component->GetIsReplicated())
+	{
+		ReplicatedComponents.AddUnique(Component);
+	}
 }
 
 void AActor::RemoveOwnedComponent(UActorComponent* Component)
@@ -1972,6 +1983,8 @@ void AActor::RemoveOwnedComponent(UActorComponent* Component)
 		// property system so take the time to pull them out now
 		OwnedComponents.RemoveSwap(NULL);
 	}
+
+	ReplicatedComponents.RemoveSwap(Component);
 }
 
 #if DO_CHECK
@@ -1981,6 +1994,24 @@ bool AActor::OwnsComponent(UActorComponent* Component) const
 }
 #endif
 
+const TArray<UActorComponent*>& AActor::GetReplicatedComponents() const
+{
+	return ReplicatedComponents;
+}
+
+void AActor::UpdateReplicatedComponent(UActorComponent* Component)
+{
+	check(OwnsComponent(Component));
+	if (Component->GetIsReplicated())
+	{
+		ReplicatedComponents.AddUnique(Component);
+	}
+	else
+	{
+		ReplicatedComponents.RemoveSwap(Component);
+	}
+}
+
 UActorComponent* AActor::FindComponentByClass(const TSubclassOf<UActorComponent> ComponentClass) const
 {
 	UActorComponent* FoundComponent = NULL;
diff --git a/Engine/Source/Runtime/Engine/Private/ActorComponent.cpp b/Engine/Source/Runtime/Engine/Private/ActorComponent.cpp
--- a/Engine/Source/Runtime/Engine/Private/ActorComponent.cpp
+++ b/Engine/Source/Runtime/Engine/Private/ActorComponent.cpp
@@ -1153,13 +1153,10 @@ void UActorComponent::SetIsReplicated(bool ShouldReplicate)
 	check(GetComponentClassCanReplicate()); // Only certain component classes can replicate!
 	bReplicates = ShouldReplicate;
 
-	if (ShouldReplicate)
+	AActor* Owner = GetOwner();
+	if (Owner)
 	{
-		AActor* Owner = GetOwner();
-		if (Owner)
-		{
-			Owner->ReplicatedComponents.AddUnique( this );
-		}
+		Owner->UpdateReplicatedComponent( this );
 	}
 }
 
diff --git a/Engine/Source/Runtime/Engine/Private/ActorReplication.cpp b/Engine/Source/Runtime/Engine/Private/ActorReplication.cpp
--- a/Engine/Source/Runtime/Engine/Private/ActorReplication.cpp
+++ b/Engine/Source/Runtime/Engine/Private/ActorReplication.cpp
@@ -240,9 +240,8 @@ bool AActor::ReplicateSubobjects(UActorChannel *Channel, FOutBunch *Bunch, FRepl
 
 	bool WroteSomething = false;
 
-	for (int32 CompIdx =0; CompIdx < ReplicatedComponents.Num(); ++CompIdx )
+	for (UActorComponent* ActorComp : ReplicatedComponents)
 	{
-		UActorComponent * ActorComp = ReplicatedComponents[CompIdx].Get();
 		if (ActorComp && ActorComp->GetIsReplicated())
 		{
 			WroteSomething |= ActorComp->ReplicateSubobjects(Channel, Bunch, RepFlags);		// Lets the component add subobjects before replicating its own properties.
diff --git a/Engine/Source/Runtime/Engine/Private/NetConnection.cpp b/Engine/Source/Runtime/Engine/Private/NetConnection.cpp
--- a/Engine/Source/Runtime/Engine/Private/NetConnection.cpp
+++ b/Engine/Source/Runtime/Engine/Private/NetConnection.cpp
@@ -1631,9 +1631,8 @@ void UNetConnection::FlushDormancy( class AActor* Actor )
 	{
 		FlushDormancyForObject( Actor );
 
-		for ( int32 i = 0; i < Actor->ReplicatedComponents.Num(); ++i )
+		for ( UActorComponent* ActorComp : Actor->GetReplicatedComponents() )
 		{
-			UActorComponent * ActorComp = Actor->ReplicatedComponents[i].Get();
 			if ( ActorComp && ActorComp->GetIsReplicated() )
 			{
 				FlushDormancyForObject( ActorComp );
