#ue4 - Change ProjectileMovementComponent::ComputeMoveDelta() to use CalculateVelocity() rather than ComputeAcceleration() to update the position.

- This allows custom implementations of changes in velocity to directly affect the simulation, without the need to also compute the acceleration.
- CalculateVelocity() is now const.
- Removed unnecessary TmpVelocity local variable (same as OldVelocity) in TickComponent.

diff --git a/Engine/Source/Runtime/Engine/Classes/GameFramework/ProjectileMovementComponent.h b/Engine/Source/Runtime/Engine/Classes/GameFramework/ProjectileMovementComponent.h
--- a/Engine/Source/Runtime/Engine/Classes/GameFramework/ProjectileMovementComponent.h
+++ b/Engine/Source/Runtime/Engine/Classes/GameFramework/ProjectileMovementComponent.h
@@ -135,8 +135,8 @@ class ENGINE_API UProjectileMovementComponent : public UMovementComponent
 
 	bool ShouldApplyGravity() const { return ProjectileGravityScale != 0.f; }
 
-	/** @returns the velocity after DeltaTime */
-	virtual FVector CalculateVelocity(FVector OldVelocity, float DeltaTime, bool bGravityEnabled);
+	/** @return the velocity after DeltaTime */
+	virtual FVector CalculateVelocity(FVector OldVelocity, float DeltaTime, bool bGravityEnabled) const;
 
 	/** Clears the reference to UpdatedComponent, fires stop event, and stops ticking. */
 	UFUNCTION(BlueprintCallable, Category="Game|Components|ProjectileMovement")
diff --git a/Engine/Source/Runtime/Engine/Private/ProjectileMovementComponent.cpp b/Engine/Source/Runtime/Engine/Private/ProjectileMovementComponent.cpp
--- a/Engine/Source/Runtime/Engine/Private/ProjectileMovementComponent.cpp
+++ b/Engine/Source/Runtime/Engine/Private/ProjectileMovementComponent.cpp
@@ -117,10 +117,9 @@ void UProjectileMovementComponent::TickComponent(float DeltaTime, enum ELevelTic
 
 		Hit.Time = 1.f;
 		const FVector OldVelocity = Velocity;
-		FVector MoveDelta = ComputeMoveDelta(Velocity, TimeTick, !bSliding);
+		FVector MoveDelta = ComputeMoveDelta(OldVelocity, TimeTick, !bSliding);
 
-		const FVector TmpVelocity = Velocity;
-		const FRotator NewRotation = bRotationFollowsVelocity ? Velocity.Rotation() : ActorOwner->GetActorRotation();
+		const FRotator NewRotation = bRotationFollowsVelocity ? OldVelocity.Rotation() : ActorOwner->GetActorRotation();
 
 		// Move the component
 		if (bShouldBounce)
@@ -145,20 +144,28 @@ void UProjectileMovementComponent::TickComponent(float DeltaTime, enum ELevelTic
 		if (Hit.Time == 1.f)
 		{
 			bSliding = false;
-			Velocity = CalculateVelocity(Velocity, TimeTick, !bSliding);
+
+			// Only calculate new velocity if events didn't change it during the movement update.
+			if (Velocity == OldVelocity)
+			{
+				Velocity = CalculateVelocity(Velocity, TimeTick, !bSliding);
+			}
 		}
 		else
 		{
-			if ( Velocity == TmpVelocity )
+			// Only calculate new velocity if events didn't change it during the movement update.
+			if (Velocity == OldVelocity)
 			{
 				// re-calculate end velocity for partial time
 				Velocity = CalculateVelocity(OldVelocity, TimeTick*Hit.Time, !bSliding);
 			}
-			if ( HandleHitWall(Hit, TimeTick, MoveDelta) )
+
+			if (HandleHitWall(Hit, TimeTick, MoveDelta))
 			{
 				break;
 			}
-			if( NumBounces < 2 )
+
+			if (NumBounces < 2)
 			{
 				RemainingTime += TimeTick * (1.f - Hit.Time);
 			}
@@ -173,7 +180,7 @@ void UProjectileMovementComponent::TickComponent(float DeltaTime, enum ELevelTic
 					FVector NewDir = (Hit.Normal ^ OldHitNormal);
 					NewDir = NewDir.SafeNormal();
 					Velocity = (Velocity | NewDir) * NewDir;
-					if ((TmpVelocity | Velocity) < 0.f)
+					if ((OldVelocity | Velocity) < 0.f)
 					{
 						Velocity *= -1.f;
 					}
@@ -210,12 +217,11 @@ void UProjectileMovementComponent::SetVelocityInLocalSpace(FVector NewVelocity)
 }
 
 
-FVector UProjectileMovementComponent::CalculateVelocity(FVector OldVelocity, float DeltaTime, bool bGravityEnabled)
+FVector UProjectileMovementComponent::CalculateVelocity(FVector OldVelocity, float DeltaTime, bool bGravityEnabled) const
 {
-	FVector NewVelocity = OldVelocity;
-
+	// v = v0 + a*t
 	const FVector Acceleration = ComputeAcceleration(OldVelocity, DeltaTime, bGravityEnabled);
-	NewVelocity += Acceleration * DeltaTime;
+	FVector NewVelocity = OldVelocity + (Acceleration * DeltaTime);
 
 	return LimitVelocity(NewVelocity);
 }
@@ -234,25 +240,16 @@ FVector UProjectileMovementComponent::LimitVelocity(FVector NewVelocity) const
 
 FVector UProjectileMovementComponent::ComputeMoveDelta(const FVector& InVelocity, float DeltaTime, bool bGravityEnabled) const
 {
-	// p = p0 + v*t
-	FVector Delta = InVelocity * DeltaTime;
-	
-	const FVector Acceleration = ComputeAcceleration(InVelocity, DeltaTime, bGravityEnabled);
-	if (!Acceleration.IsZero())
-	{
-		// p = p0 + v*t (above) + 1/2*a*t^2 (below)
-		Delta += 0.5f * Acceleration * FMath::Square(DeltaTime);
+	// Velocity Verlet integration (http://en.wikipedia.org/wiki/Verlet_integration#Velocity_Verlet)
+	// The addition of p0 is done outside this method, we are just computing the delta.
+	// p = p0 + v0*t + 1/2*a*t^2
 
-		// limit velocity, else acceleration will push this result over the allowed velocity constraint during this timestep.
-		const FVector EffectiveVelocity = Delta / DeltaTime;
-		const FVector ClampedVelocity = LimitVelocity(EffectiveVelocity);
-		if (ClampedVelocity != EffectiveVelocity)
-		{
-			// Maintain direction but change magnitude
-			Delta = Delta.SafeNormal() * (GetMaxSpeed() * DeltaTime);
-		}
-	}
+	// We use CalculateVelocity() here to infer the acceleration, to make it easier to apply custom velocities.
+	// p = p0 + v0*t + 1/2*((v1-v0)/t)*t^2
+	// p = p0 + v0*t + 1/2*((v1-v0))*t
 
+	const FVector NewVelocity = CalculateVelocity(InVelocity, DeltaTime, bGravityEnabled);
+	const FVector Delta = (InVelocity * DeltaTime) + (NewVelocity - InVelocity) * (0.5f * DeltaTime);
 	return Delta;
 }
 
