#ue4 - ProjectileMovementComponent.MaxSimulationIterations is no longer an advanced display property, in line with the other related settings (MaxSimulationTimeStep, bForceSubStepping).

diff --git a/Engine/Source/Runtime/Engine/Classes/GameFramework/ProjectileMovementComponent.h b/Engine/Source/Runtime/Engine/Classes/GameFramework/ProjectileMovementComponent.h
--- a/Engine/Source/Runtime/Engine/Classes/GameFramework/ProjectileMovementComponent.h
+++ b/Engine/Source/Runtime/Engine/Classes/GameFramework/ProjectileMovementComponent.h
@@ -180,7 +180,7 @@ class ENGINE_API UProjectileMovementComponent : public UMovementComponent
 	 *
 	 * WARNING: if (MaxSimulationTimeStep * MaxSimulationIterations) is too low for the min framerate, the last simulation step may exceed MaxSimulationTimeStep to complete the simulation.
 	 */
-	UPROPERTY(EditAnywhere, BlueprintReadWrite, AdvancedDisplay, meta=(ClampMin="1", ClampMax="25", UIMin="1", UIMax="25"), Category=ProjectileSimulation)
+	UPROPERTY(EditAnywhere, BlueprintReadWrite, meta=(ClampMin="1", ClampMax="25", UIMin="1", UIMax="25"), Category=ProjectileSimulation)
 	int32 MaxSimulationIterations;
 
 protected:
