Remove check (for now)

diff --git a/Engine/Source/Runtime/Engine/Private/Actor.cpp b/Engine/Source/Runtime/Engine/Private/Actor.cpp
--- a/Engine/Source/Runtime/Engine/Private/Actor.cpp
+++ b/Engine/Source/Runtime/Engine/Private/Actor.cpp
@@ -2001,7 +2001,6 @@ const TArray<UActorComponent*>& AActor::GetReplicatedComponents() const
 
 void AActor::UpdateReplicatedComponent(UActorComponent* Component)
 {
-	check(OwnsComponent(Component));
 	if (Component->GetIsReplicated())
 	{
 		ReplicatedComponents.AddUnique(Component);
