fixed TTP 348407 UE4: Rendering: Objects that do not render in main pass still affect motion blur.

diff --git a/Engine/Source/Runtime/Engine/Public/PrimitiveViewRelevance.h b/Engine/Source/Runtime/Engine/Public/PrimitiveViewRelevance.h
--- a/Engine/Source/Runtime/Engine/Public/PrimitiveViewRelevance.h
+++ b/Engine/Source/Runtime/Engine/Public/PrimitiveViewRelevance.h
@@ -18,7 +18,7 @@ struct FPrimitiveViewRelevance
 	uint32 bShadowRelevance : 1;
 	/** The primitive should render to the custom depth pass. */
 	uint32 bRenderCustomDepth : 1;
-	/** The primitive should render to the main depth pass. */
+	/** The primitive should render to the base pass / normal depth / velocity rendering. */
 	uint32 bRenderInMainPass : 1;
 	/** The primitive is drawn only in the editor and composited onto the scene after post processing */
 	uint32 bEditorPrimitiveRelevance : 1;
diff --git a/Engine/Source/Runtime/Renderer/Private/PrimitiveSceneInfo.cpp b/Engine/Source/Runtime/Renderer/Private/PrimitiveSceneInfo.cpp
--- a/Engine/Source/Runtime/Renderer/Private/PrimitiveSceneInfo.cpp
+++ b/Engine/Source/Runtime/Renderer/Private/PrimitiveSceneInfo.cpp
@@ -434,7 +434,7 @@ bool FPrimitiveSceneInfo::ShouldRenderVelocity(const FViewInfo& View, bool bChec
 	const FPrimitiveViewRelevance& PrimitiveViewRelevance = View.PrimitiveViewRelevanceMap[PrimitiveId];
 
 	// !Skip translucent objects as they don't support velocities and in the case of particles have a significant CPU overhead.
-	if(!PrimitiveViewRelevance.bOpaqueRelevance)
+	if(!PrimitiveViewRelevance.bOpaqueRelevance || !PrimitiveViewRelevance.bRenderInMainPass)
 	{
 		return false;
 	}
