#docs - Fix comments for ProjectileMovement substepping.

diff --git a/Engine/Source/Runtime/Engine/Classes/GameFramework/ProjectileMovementComponent.h b/Engine/Source/Runtime/Engine/Classes/GameFramework/ProjectileMovementComponent.h
--- a/Engine/Source/Runtime/Engine/Classes/GameFramework/ProjectileMovementComponent.h
+++ b/Engine/Source/Runtime/Engine/Classes/GameFramework/ProjectileMovementComponent.h
@@ -165,20 +165,21 @@ class ENGINE_API UProjectileMovementComponent : public UMovementComponent
 	virtual bool ShouldUseSubStepping() const;
 
 	/**
-	 * Max time delta for each discrete simulation step. This is used if the object is under the effect of 
+	 * Max time delta for each discrete simulation step.
 	 * Lowering this value can address issues with fast-moving objects or complex collision scenarios, at the cost of performance.
 	 *
 	 * WARNING: if (MaxSimulationTimeStep * MaxSimulationIterations) is too low for the min framerate, the last simulation step may exceed MaxSimulationTimeStep to complete the simulation.
-	 * @see MaxSimulationIterations
+	 * @see MaxSimulationIterations, bForceSubStepping
 	 */
 	UPROPERTY(EditAnywhere, BlueprintReadWrite, meta=(ClampMin="0.0166", ClampMax="0.50", UIMin="0.0166", UIMax="0.50"), Category=ProjectileSimulation)
 	float MaxSimulationTimeStep;
 
 	/**
-	 * Max number of iterations used for each discrete simulation step
+	 * Max number of iterations used for each discrete simulation step.
 	 * Increasing this value can address issues with fast-moving objects or complex collision scenarios, at the cost of performance.
 	 *
 	 * WARNING: if (MaxSimulationTimeStep * MaxSimulationIterations) is too low for the min framerate, the last simulation step may exceed MaxSimulationTimeStep to complete the simulation.
+	 * @see MaxSimulationTimeStep, bForceSubStepping
 	 */
 	UPROPERTY(EditAnywhere, BlueprintReadWrite, meta=(ClampMin="1", ClampMax="25", UIMin="1", UIMax="25"), Category=ProjectileSimulation)
 	int32 MaxSimulationIterations;
