Remove std::function based instance of UMaterialInterface::IterateOverActiveFeatureLevels, and also remove the module export declaration (not necessary on an inline function)

diff --git a/Engine/Source/Runtime/Engine/Classes/Materials/MaterialInterface.h b/Engine/Source/Runtime/Engine/Classes/Materials/MaterialInterface.h
--- a/Engine/Source/Runtime/Engine/Classes/Materials/MaterialInterface.h
+++ b/Engine/Source/Runtime/Engine/Classes/Materials/MaterialInterface.h
@@ -6,7 +6,6 @@
 #include "Runtime/RHI/Public/RHIDefinitions.h"
 #include "Runtime/RenderCore/Public/RenderCommandFence.h"
 #include "SceneTypes.h"
-#include <functional>
 #include "MaterialInterface.generated.h"
 
 class UMaterial;
@@ -589,9 +588,8 @@ public:
 	ENGINE_API static uint32 GetFeatureLevelsToCompileForAllMaterials() { return FeatureLevelsForAllMaterials | (1 << GMaxRHIFeatureLevel); }
 
 	/** Iterate over all feature levels currently marked as active */
-#if IS_MONOLITHIC
 	template <typename TFunction>
-	ENGINE_API static void IterateOverActiveFeatureLevels(TFunction InHandler) 
+	static void IterateOverActiveFeatureLevels(TFunction InHandler) 
 	{  
 		uint32 FeatureLevels = GetFeatureLevelsToCompileForAllMaterials();
 		while (FeatureLevels != 0)
@@ -599,17 +597,6 @@ public:
 			InHandler((ERHIFeatureLevel::Type)FBitSet::GetAndClearNextBit(FeatureLevels));
 		}
 	}
-#else
-	//@todo: Can't export template functions...
-	ENGINE_API static void IterateOverActiveFeatureLevels(std::function<void(ERHIFeatureLevel::Type FeatureLevel)> InHandler) 
-	{  
-		uint32 FeatureLevels = GetFeatureLevelsToCompileForAllMaterials();
-		while (FeatureLevels != 0)
-		{
-			InHandler((ERHIFeatureLevel::Type)FBitSet::GetAndClearNextBit(FeatureLevels));
-		}
-	}
-#endif
 
 protected:
 
