Fixed wrong condition being checked to determine if navigation needs to be build #UE4

TTP#348227

diff --git a/Engine/Source/Runtime/Engine/Private/AI/Navigation/NavigationSystem.cpp b/Engine/Source/Runtime/Engine/Private/AI/Navigation/NavigationSystem.cpp
--- a/Engine/Source/Runtime/Engine/Private/AI/Navigation/NavigationSystem.cpp
+++ b/Engine/Source/Runtime/Engine/Private/AI/Navigation/NavigationSystem.cpp
@@ -1139,14 +1139,13 @@ TSharedPtr<FNavigationQueryFilter> UNavigationSystem::CreateDefaultQueryFilterCo
 
 bool UNavigationSystem::IsNavigationBuilt(const AWorldSettings* Settings) const
 {
-	if (IsThereAnywhereToBuildNavigation() == false)
+	if (Settings == nullptr || Settings->bEnableNavigationSystem == false || IsThereAnywhereToBuildNavigation() == false)
 	{
 		return true;
 	}
 
 	bool bIsBuilt = true;
-	// look for all ANavigationDatas that belong to specified world, and check if they're dirty
-	// trigger navmesh update
+
 	for (int32 NavDataIndex = 0; NavDataIndex < NavDataSet.Num(); ++NavDataIndex)
 	{
 		ANavigationData* NavData = NavDataSet[NavDataIndex];
@@ -1154,7 +1153,11 @@ bool UNavigationSystem::IsNavigationBuilt(const AWorldSettings* Settings) const
 		{
 #if WITH_NAVIGATION_GENERATOR
 			FNavDataGenerator* Generator = NavData->GetGenerator(FNavigationSystem::DontCreate);
-			if ( NavData->bRebuildAtRuntime == false && ( Generator == NULL || Generator->IsBuildInProgress(true) == true ) )
+			if ((NavData->bRebuildAtRuntime == true || 
+#if WITH_EDITOR
+				GEditor != NULL
+#endif // WITH_EDITOR
+				) && (Generator == NULL || Generator->IsBuildInProgress(/*bCheckDirtyToo=*/true) == true))
 			{
 				bIsBuilt = false;
 				break;
