Adds the Seamless Travel's LoadedWorld as an exception to the invalid world check

Partial fix towards TTP #346855 (still working on)

This is needed when running multiple PIE processes as the server and client are loaded concurrently and the client may not have finished loading when the server does (and performs the check), so it needs to be OK with this.

diff --git a/Engine/Source/Editor/UnrealEd/Private/Editor.cpp b/Engine/Source/Editor/UnrealEd/Private/Editor.cpp
--- a/Engine/Source/Editor/UnrealEd/Private/Editor.cpp
+++ b/Engine/Source/Editor/UnrealEd/Private/Editor.cpp
@@ -6852,10 +6852,17 @@ void UEditorEngine::VerifyLoadMapWorldCleanup()
 			bool ValidWorld = false;
 			for (int32 idx=0; idx < WorldList.Num(); ++idx)
 			{
-				if (WorldList[idx].World())
+				FWorldContext& WorldContext = WorldList[idx];
+				if (World == WorldContext.SeamlessTravelHandler.GetLoadedWorld())
+				{
+					// World valid, but not loaded yet
+					ValidWorld = true;
+					break;
+				}
+				else if (WorldContext.World())
 				{
 					TArray<UWorld*> OtherWorlds;
-					EditorLevelUtils::GetWorlds(WorldList[idx].World(), OtherWorlds, true, false);
+					EditorLevelUtils::GetWorlds(WorldContext.World(), OtherWorlds, true, false);
 
 					if (OtherWorlds.Contains(World))
 					{
diff --git a/Engine/Source/Runtime/Engine/Classes/Engine/World.h b/Engine/Source/Runtime/Engine/Classes/Engine/World.h
--- a/Engine/Source/Runtime/Engine/Classes/Engine/World.h
+++ b/Engine/Source/Runtime/Engine/Classes/Engine/World.h
@@ -171,21 +171,26 @@ public:
 	bool StartTravel(UWorld* InCurrentWorld, const FURL& InURL, const FGuid& InGuid);
 
 	/** @return whether a transition is already in progress */
-	FORCEINLINE bool IsInTransition()
+	FORCEINLINE bool IsInTransition() const
 	{
 		return bTransitionInProgress;
 	}
 	/** @return if current transition has switched to the default map; returns false if no transition is in progress */
-	FORCEINLINE bool HasSwitchedToDefaultMap()
+	FORCEINLINE bool HasSwitchedToDefaultMap() const
 	{
 		return IsInTransition() && bSwitchedToDefaultMap;
 	}
 
-	inline FString GetDestinationMapName()
+	inline FString GetDestinationMapName() const
 	{
 		return (IsInTransition() ? PendingTravelURL.Map : TEXT(""));
 	}
 
+	inline const UWorld* GetLoadedWorld() const
+	{
+		return LoadedWorld;
+	}
+
 	/** cancels transition in progress */
 	void CancelTravel();
 
