Clamp delta time when playing back demos. This is useful when for whtatever reason a long frame occurs, and you generally don't want all that time to pass for the demo.

diff --git a/Engine/Source/Runtime/Engine/Private/DemoNetDriver.cpp b/Engine/Source/Runtime/Engine/Private/DemoNetDriver.cpp
--- a/Engine/Source/Runtime/Engine/Private/DemoNetDriver.cpp
+++ b/Engine/Source/Runtime/Engine/Private/DemoNetDriver.cpp
@@ -235,6 +235,10 @@ void UDemoNetDriver::TickFlush( float DeltaSeconds )
 		{
 			World->GetWorldSettings()->DemoPlayTimeDilation = CVarDemoTimeDilation.GetValueOnGameThread();
 
+			// Clamp time between 1000 hz, and 2 hz 
+			// (this is useful when debugging and you set a breakpoint, you don't want all that time to pass in one frame)
+			DeltaSeconds = FMath::Clamp( DeltaSeconds, 1.0f / 1000.0f, 1.0f / 2.0f );
+
 			// We need to compensate for the fact that DeltaSeconds is real-time for net drivers
 			DeltaSeconds *= World->GetWorldSettings()->GetEffectiveTimeDilation();
 
