Fix crash when endmap is called and world is null (before PendingKill() is true).

diff --git a/Engine/Source/Runtime/Engine/Private/PrimitiveComponentPhysics.cpp b/Engine/Source/Runtime/Engine/Private/PrimitiveComponentPhysics.cpp
--- a/Engine/Source/Runtime/Engine/Private/PrimitiveComponentPhysics.cpp
+++ b/Engine/Source/Runtime/Engine/Private/PrimitiveComponentPhysics.cpp
@@ -615,7 +615,8 @@ void UPrimitiveComponent::WeldTo(USceneComponent* InParent, FName InSocketName /
 void UPrimitiveComponent::UnWeldFromParent()
 {
 	FBodyInstance* NewRootBI = GetBodyInstance(NAME_None, false);
-	if (NewRootBI == NULL || NewRootBI->bWelded == false || IsPendingKill())
+	UWorld* CurrentWorld = GetWorld();
+	if (NewRootBI == NULL || NewRootBI->bWelded == false || CurrentWorld == nullptr || IsPendingKill())
 	{
 		return;
 	}
@@ -639,7 +640,7 @@ void UPrimitiveComponent::UnWeldFromParent()
 			{
 				bool bPrevAutoWeld = NewRootBI->bAutoWeld;
 				NewRootBI->bAutoWeld = false;
-				NewRootBI->InitBody(GetBodySetup(), GetComponentToWorld(), this, GetWorld()->GetPhysicsScene());
+				NewRootBI->InitBody(GetBodySetup(), GetComponentToWorld(), this, CurrentWorld->GetPhysicsScene());
 				NewRootBI->bAutoWeld = bPrevAutoWeld;
 			}
 
