Disable collision on demo spectator

diff --git a/Engine/Source/Runtime/Engine/Private/DemoNetDriver.cpp b/Engine/Source/Runtime/Engine/Private/DemoNetDriver.cpp
--- a/Engine/Source/Runtime/Engine/Private/DemoNetDriver.cpp
+++ b/Engine/Source/Runtime/Engine/Private/DemoNetDriver.cpp
@@ -248,6 +248,10 @@ void UDemoNetDriver::TickFlush( float DeltaSeconds )
 			{
 				if ( SpectatorController->GetSpectatorPawn() != NULL )
 				{
+					// Disable collision on the spectator
+					SpectatorController->GetSpectatorPawn()->SetActorEnableCollision( false );
+					
+					// Apply time dilation on spectator to reverse the effects of global dilation
 					SpectatorController->GetSpectatorPawn()->CustomTimeDilation = 1.0f / World->GetWorldSettings()->DemoPlayTimeDilation;
 				}
 			}
