Add some keywords to Set/GetActorLocation, AddActorWorld/LocalOffset, USceneComponent::AddLocal/WorldOffset

diff --git a/Engine/Source/Runtime/Engine/Classes/Components/SceneComponent.h b/Engine/Source/Runtime/Engine/Classes/Components/SceneComponent.h
--- a/Engine/Source/Runtime/Engine/Classes/Components/SceneComponent.h
+++ b/Engine/Source/Runtime/Engine/Classes/Components/SceneComponent.h
@@ -247,7 +247,7 @@ public:
 	void AddRelativeRotation(FRotator DeltaRotation, bool bSweep=false);
 
 	/** Adds a delta to the location of this component in its local reference frame */
-	UFUNCTION(BlueprintCallable, Category="Utilities|Transformation")
+	UFUNCTION(BlueprintCallable, Category="Utilities|Transformation", meta=(Keywords="location position"))
 	void AddLocalOffset(FVector DeltaLocation, bool bSweep=false);
 	
 	/** Adds a delta to the rotation of this component in its local reference frame */
@@ -275,7 +275,7 @@ public:
 	void SetWorldTransform(const FTransform& NewTransform, bool bSweep=false);
 
 	/** Adds a delta to the location of this component in world space. */
-	UFUNCTION(BlueprintCallable, Category="Utilities|Transformation")
+	UFUNCTION(BlueprintCallable, Category="Utilities|Transformation", meta=(Keywords="location position"))
 	void AddWorldOffset(FVector DeltaLocation, bool bSweep=false);
 
 	/** Adds a delta to the rotation of this component in world space. */
diff --git a/Engine/Source/Runtime/Engine/Classes/GameFramework/Actor.h b/Engine/Source/Runtime/Engine/Classes/GameFramework/Actor.h
--- a/Engine/Source/Runtime/Engine/Classes/GameFramework/Actor.h
+++ b/Engine/Source/Runtime/Engine/Classes/GameFramework/Actor.h
@@ -579,7 +579,7 @@ public:
 	FTransform ActorToWorld() const;
 
 	/** Returns the location of the RootComponent of this Actor */
-	UFUNCTION(BlueprintCallable, meta=(FriendlyName = "GetActorLocation"), Category="Utilities|Transformation")
+	UFUNCTION(BlueprintCallable, meta=(FriendlyName = "GetActorLocation", Keywords="position"), Category="Utilities|Transformation")
 	FVector K2_GetActorLocation() const;
 
 	/** 
@@ -588,7 +588,7 @@ public:
 	 *	@param bSweep		Should we sweep to the destination location. If true, will stop short of the target if blocked by something.
 	 *	@return	Whether the location was successfully set (if not swept), or whether movement occurred at all (if swept).
 	 */
-	UFUNCTION(BlueprintCallable, meta=(FriendlyName = "SetActorLocation"), Category="Utilities|Transformation")
+	UFUNCTION(BlueprintCallable, meta=(FriendlyName = "SetActorLocation", Keywords="position"), Category="Utilities|Transformation")
 	bool K2_SetActorLocation(FVector NewLocation, bool bSweep=false);
 
 	/** Returns rotation of the RootComponent of this Actor. */
@@ -677,7 +677,7 @@ public:
 	float GetHorizontalDotProductTo(AActor* OtherActor);
 
 	/** Adds a delta to the location of this actor in world space. */
-	UFUNCTION(BlueprintCallable, Category="Utilities|Transformation")
+	UFUNCTION(BlueprintCallable, Category="Utilities|Transformation", meta=(Keywords="location position"))
 	void AddActorWorldOffset(FVector DeltaLocation, bool bSweep=false);
 
 	/** Adds a delta to the rotation of this actor in world space. */
@@ -696,7 +696,7 @@ public:
 	bool SetActorTransform(const FTransform& NewTransform, bool bSweep=false);
 
 	/** Adds a delta to the location of this component in its local reference frame */
-	UFUNCTION(BlueprintCallable, Category="Utilities|Transformation")
+	UFUNCTION(BlueprintCallable, Category="Utilities|Transformation", meta=(Keywords="location position"))
 	void AddActorLocalOffset(FVector DeltaLocation, bool bSweep=false);
 
 	/** Adds a delta to the rotation of this component in its local reference frame */
