Fix case where we try to inherit from owner when there is no owner and decide to set it to no collision.

diff --git a/Engine/Source/Runtime/Engine/Private/PhysicsEngine/BodyInstance.cpp b/Engine/Source/Runtime/Engine/Private/PhysicsEngine/BodyInstance.cpp
--- a/Engine/Source/Runtime/Engine/Private/PhysicsEngine/BodyInstance.cpp
+++ b/Engine/Source/Runtime/Engine/Private/PhysicsEngine/BodyInstance.cpp
@@ -793,14 +793,18 @@ void FBodyInstance::UpdatePhysicsFilterData(bool bForceSimpleAsComplex)
 
 		// In skeletal case, collision enable/disable/movement should be overriden by mesh component
 		// being in the physics asset, and not having collision is a waste and it can cause a bug where disconnected bodies
-		if (Owner && Owner->GetActorEnableCollision())	//we only override if actor has collision enabled
+		if (Owner)
 		{
-			UseCollisionEnabled = SkelMeshComp->BodyInstance.CollisionEnabled;
-		}
-		else
-		{	//actor has collision disabled so disable regardless of component setting
-			UseCollisionEnabled = ECollisionEnabled::NoCollision;
+			if (Owner->GetActorEnableCollision())	//we only override if actor has collision enabled
+			{
+				UseCollisionEnabled = SkelMeshComp->BodyInstance.CollisionEnabled;
+			}
+			else
+			{	//actor has collision disabled so disable regardless of component setting
+				UseCollisionEnabled = ECollisionEnabled::NoCollision;
+			}
 		}
+		
 
 		ObjectType = SkelMeshComp->GetCollisionObjectType();
 		bUseCollisionEnabledOverride = true;
