Fixed some shader compilation timeout code.

diff --git a/Engine/Source/Runtime/Engine/Private/Materials/MaterialShader.cpp b/Engine/Source/Runtime/Engine/Private/Materials/MaterialShader.cpp
--- a/Engine/Source/Runtime/Engine/Private/Materials/MaterialShader.cpp
+++ b/Engine/Source/Runtime/Engine/Private/Materials/MaterialShader.cpp
@@ -1241,7 +1241,7 @@ bool FMaterialShaderMap::ProcessCompilationResults(const TArray<FShaderCompileJo
 {
 	check(InOutJobIndex < InCompilationResults.Num());
 
-	const double StartTime = FPlatformTime::Seconds();
+	double StartTime = FPlatformTime::Seconds();
 
 	FSHAHash MaterialShaderMapHash;
 	ShaderMapId.GetMaterialHash(MaterialShaderMapHash);
@@ -1286,10 +1286,12 @@ bool FMaterialShaderMap::ProcessCompilationResults(const TArray<FShaderCompileJo
 		}
 
 		InOutJobIndex++;
-
-		TimeBudget -= FPlatformTime::Seconds() - StartTime;
+		
+		double NewStartTime = FPlatformTime::Seconds();
+		TimeBudget -= NewStartTime - StartTime;
+		StartTime = NewStartTime;
 	}
-	while (TimeBudget > 0 && InOutJobIndex < InCompilationResults.Num());
+	while ((TimeBudget > 0.0f) && (InOutJobIndex < InCompilationResults.Num()));
 
 	if (InOutJobIndex == InCompilationResults.Num())
 	{
