Fix case where bAutoRegister false components got duplicate IniitalizeComponent calls

diff --git a/Engine/Source/Runtime/Engine/Private/ActorComponent.cpp b/Engine/Source/Runtime/Engine/Private/ActorComponent.cpp
--- a/Engine/Source/Runtime/Engine/Private/ActorComponent.cpp
+++ b/Engine/Source/Runtime/Engine/Private/ActorComponent.cpp
@@ -618,7 +618,7 @@ void UActorComponent::RegisterComponentWithWorld(UWorld* InWorld)
 
 	if (Owner == nullptr || Owner->bActorInitialized)
 	{
-		if (bWantsInitializeComponent)
+		if (!bHasBeenInitialized && bWantsInitializeComponent)
 		{
 			InitializeComponent();
 		}
