Network Demos: Streaming level support

diff --git a/Engine/Source/Runtime/Engine/Private/DemoNetDriver.cpp b/Engine/Source/Runtime/Engine/Private/DemoNetDriver.cpp
--- a/Engine/Source/Runtime/Engine/Private/DemoNetDriver.cpp
+++ b/Engine/Source/Runtime/Engine/Private/DemoNetDriver.cpp
@@ -141,6 +141,35 @@ bool UDemoNetDriver::InitConnect( FNetworkNotify* InNotify, const FURL& ConnectU
 	WorldContext->PendingNetGame = NULL;
 #endif
 
+	int32 NumStreamingLevels = 0;
+
+	(*FileAr) << NumStreamingLevels;
+
+	for ( int32 i = 0; i < NumStreamingLevels; ++i )
+	{
+		ULevelStreamingKismet* StreamingLevel = static_cast<ULevelStreamingKismet*>(StaticConstructObject(ULevelStreamingKismet::StaticClass(), GetWorld(), NAME_None, RF_NoFlags, NULL ) );
+
+		StreamingLevel->bShouldBeLoaded		= true;
+		StreamingLevel->bShouldBeVisible	= true;
+		StreamingLevel->bShouldBlockOnLoad	= false;
+		StreamingLevel->bInitiallyLoaded	= true;
+		StreamingLevel->bInitiallyVisible	= true;
+
+		FString PackageName;
+		FString PackageNameToLoad;
+
+		(*FileAr) << PackageName;
+		(*FileAr) << PackageNameToLoad;
+		(*FileAr) << StreamingLevel->LevelTransform;
+
+		StreamingLevel->PackageNameToLoad = FName( *PackageNameToLoad );
+		StreamingLevel->SetWorldAssetByPackageName( FName( *PackageName ) );
+
+		GetWorld()->StreamingLevels.Add( StreamingLevel );
+
+		UE_LOG( LogDemo, Log, TEXT( "  Loading streamingLevel: %s, %s" ), *PackageName, *PackageNameToLoad );
+	}
+
 	DemoDeltaTime = 0;
 
 	return true;
@@ -207,6 +236,35 @@ bool UDemoNetDriver::InitListen( FNetworkNotify* InNotify, FURL& ListenURL, bool
 	(*FileAr) << LevelName;
 #endif
 
+	// Save out any levels that are in the streamed level list
+	// This needs some work, but for now, to try and get games that use heavy streaming working
+	int32 NumStreamingLevels = 0;
+
+	for ( int32 i = 0; i < World->StreamingLevels.Num(); ++i )
+	{
+		if ( World->StreamingLevels[i] != NULL )
+		{
+			NumStreamingLevels++;
+		}
+	}
+
+	(*FileAr) << NumStreamingLevels;
+
+	for ( int32 i = 0; i < World->StreamingLevels.Num(); ++i )
+	{
+		if ( World->StreamingLevels[i] != NULL )
+		{
+			FString PackageName = World->StreamingLevels[i]->GetWorldAssetPackageName();
+			FString PackageNameToLoad = World->StreamingLevels[i]->PackageNameToLoad.ToString();
+
+			UE_LOG( LogDemo, Log, TEXT( "  StreamingLevel: %s, %s" ), *PackageName, *PackageNameToLoad );
+
+			(*FileAr) << PackageName;
+			(*FileAr) << PackageNameToLoad;
+			(*FileAr) << World->StreamingLevels[i]->LevelTransform;
+		}
+	}
+
 	// Spawn the demo recording spectator.
 	SpawnDemoRecSpectator( Connection );
 
@@ -233,6 +291,17 @@ void UDemoNetDriver::TickFlush( float DeltaSeconds )
 		}
 		else if ( ServerConnection != NULL )
 		{
+			// Wait until all levels are streamed in
+			for ( int32 i = 0; i < World->StreamingLevels.Num(); ++i )
+			{
+				ULevelStreaming * StreamingLevel = World->StreamingLevels[i];
+				if ( StreamingLevel != NULL && ( !StreamingLevel->IsLevelLoaded() || !StreamingLevel->GetLoadedLevel()->GetOutermost()->IsFullyLoaded() || !StreamingLevel->IsLevelVisible() ) )
+				{
+					// Abort, we have more streaming levels to load
+					return;
+				}
+			}
+
 			World->GetWorldSettings()->DemoPlayTimeDilation = CVarDemoTimeDilation.GetValueOnGameThread();
 
 			// Clamp time between 1000 hz, and 2 hz 
