NetworkDemos: AActor::GetNetMode now returns the proper result when playing back demos

diff --git a/Engine/Source/Runtime/Engine/Private/Actor.cpp b/Engine/Source/Runtime/Engine/Private/Actor.cpp
--- a/Engine/Source/Runtime/Engine/Private/Actor.cpp
+++ b/Engine/Source/Runtime/Engine/Private/Actor.cpp
@@ -16,6 +16,7 @@
 #include "Particles/ParticleSystemComponent.h"
 #include "VisualLog.h"
 #include "Animation/AnimInstance.h"
+#include "Engine/DemoNetDriver.h"
 
 //DEFINE_LOG_CATEGORY_STATIC(LogActor, Log, All);
 DEFINE_LOG_CATEGORY(LogActor);
@@ -2694,7 +2695,18 @@ void AActor::PlaySoundAtLocation(USoundCue* InSoundCue, FVector SoundLocation, f
 ENetMode AActor::GetNetMode() const
 {
 	UNetDriver *NetDriver = GetNetDriver();
-	return NetDriver ? NetDriver->GetNetMode() : NM_Standalone;
+
+	if ( NetDriver != NULL )
+	{
+		return NetDriver->GetNetMode();
+	}
+
+	if ( GetWorld() != NULL && GetWorld()->DemoNetDriver != NULL )
+	{
+		return GetWorld()->DemoNetDriver->GetNetMode();
+	}
+
+	return NM_Standalone;
 }
 
 UNetDriver* AActor::GetNetDriver() const
