Added PreFeatureLevelSwitch() to scene component
Reflection captures work with feature level preview switching.
removed 'r.NoTiledReflections' support.

diff --git a/Engine/Source/Runtime/Engine/Classes/Components/ReflectionCaptureComponent.h b/Engine/Source/Runtime/Engine/Classes/Components/ReflectionCaptureComponent.h
--- a/Engine/Source/Runtime/Engine/Classes/Components/ReflectionCaptureComponent.h
+++ b/Engine/Source/Runtime/Engine/Classes/Components/ReflectionCaptureComponent.h
@@ -124,6 +124,7 @@ class UReflectionCaptureComponent : public USceneComponent
 #if WITH_EDITOR
 	virtual void PostEditChangeProperty(FPropertyChangedEvent& PropertyChangedEvent) override;
 	virtual void PostEditImport();
+	virtual void PreFeatureLevelChange(ERHIFeatureLevel::Type PendingFeatureLevel) override;
 #endif // WITH_EDITOR
 	virtual void BeginDestroy() override;
 	virtual bool IsReadyForFinishDestroy() override;
diff --git a/Engine/Source/Runtime/Engine/Classes/Components/SceneComponent.h b/Engine/Source/Runtime/Engine/Classes/Components/SceneComponent.h
--- a/Engine/Source/Runtime/Engine/Classes/Components/SceneComponent.h
+++ b/Engine/Source/Runtime/Engine/Classes/Components/SceneComponent.h
@@ -3,6 +3,7 @@
 #pragma once
 #include "Components/ActorComponent.h"
 #include "Engine/EngineTypes.h"
+#include "RHIDefinitions.h"
 #include "SceneComponent.generated.h"
 
 /** Overlap info consisting of the primitive and the body that is overlapping */
@@ -649,6 +650,8 @@ public:
 	virtual bool CanEditChange( const UProperty* Property ) const override;
 
 	virtual const int32 GetNumUncachedStaticLightingInteractions() const;
+
+	virtual void PreFeatureLevelChange(ERHIFeatureLevel::Type PendingFeatureLevel) {}
 #endif // WITH_EDITOR
 
 protected:
diff --git a/Engine/Source/Runtime/Engine/Classes/Engine/World.h b/Engine/Source/Runtime/Engine/Classes/Engine/World.h
--- a/Engine/Source/Runtime/Engine/Classes/Engine/World.h
+++ b/Engine/Source/Runtime/Engine/Classes/Engine/World.h
@@ -642,9 +642,13 @@ public:
 	/** Creates a new FX system for this world */
 	void CreateFXSystem();
 
+#if WITH_EDITOR
+
 	/** Change the feature level that this world is current rendering with */
 	void ChangeFeatureLevel(ERHIFeatureLevel::Type InFeatureLevel);
 
+#endif // WITH_EDITOR
+
 private:
 
 	/** List of all the controllers in the world. */
diff --git a/Engine/Source/Runtime/Engine/Private/ReflectionCaptureComponent.cpp b/Engine/Source/Runtime/Engine/Private/ReflectionCaptureComponent.cpp
--- a/Engine/Source/Runtime/Engine/Private/ReflectionCaptureComponent.cpp
+++ b/Engine/Source/Runtime/Engine/Private/ReflectionCaptureComponent.cpp
@@ -20,6 +20,12 @@ ENGINE_API int32 GReflectionCaptureSize = 128;
 
 void UWorld::UpdateAllReflectionCaptures()
 {
+	if ( FeatureLevel < ERHIFeatureLevel::SM4)
+	{
+		UE_LOG(LogMaterial, Warning, TEXT("Update reflection captures only works with an active feature level of SM4 or greater."));
+		return;
+	}
+
 	TArray<UReflectionCaptureComponent*> UpdatedComponents;
 
 	for (TObjectIterator<UReflectionCaptureComponent> It; It; ++It)
@@ -961,6 +967,8 @@ void UReflectionCaptureComponent::PostLoad()
 	Super::PostLoad();
 
 	bool bRetainAllFeatureLevelData = GIsEditor && GMaxRHIFeatureLevel >= ERHIFeatureLevel::SM4;
+	bool bEncodedDataRequired = GMaxRHIFeatureLevel == ERHIFeatureLevel::ES2;
+	bool bFullDataRequired = GMaxRHIFeatureLevel >= ERHIFeatureLevel::SM4;
 
 	// If we're loading on a platform that doesn't require cooked data, attempt to load missing data from the DDC
 	if (!FPlatformProperties::RequiresCookedData())
@@ -994,39 +1002,32 @@ void UReflectionCaptureComponent::PostLoad()
 		// If we have full HDR data but not encoded HDR data, generate the encoded data now
 		if (FullHDRDerivedData 
 			&& !EncodedHDRDerivedData 
-			&& (GMaxRHIFeatureLevel == ERHIFeatureLevel::ES2 || bRetainAllFeatureLevelData))
+			&& (bEncodedDataRequired))
 		{
 			EncodedHDRDerivedData = FReflectionCaptureEncodedHDRDerivedData::GenerateEncodedHDRData(*FullHDRDerivedData, StateId, Brightness);
 		}
 	}
 	
-	// Initialize rendering resources for the current feature level, and toss data only needed by other feature levels
-	if (FullHDRDerivedData && (GMaxRHIFeatureLevel >= ERHIFeatureLevel::SM4 || bRetainAllFeatureLevelData))
+	// Initialize rendering resources for the current feature level, and toss data only needed by other feature levels (unless in editor mode, in which all feature level data should be resident.)
+	if (FullHDRDerivedData && bFullDataRequired)
 	{
 		// Don't need encoded HDR data for rendering on this feature level
 		INC_MEMORY_STAT_BY(STAT_ReflectionCaptureMemory, FullHDRDerivedData->CompressedCapturedData.GetAllocatedSize());
 
-		// command line option: -NoTiledReflections
-		// useful to test this code path and to run the potentially more efficient path (PS vs CS) on SM5 hardware
-		// not a cvar as it also needs to change the data when loading the captures
-		static const auto CVar = IConsoleManager::Get().FindTConsoleVariableDataInt(TEXT("r.NoTiledReflections"));
-		check(CVar);
-		const bool bNoTiledReflections = (CVar->GetValueOnAnyThread() == 1);
-
-		if ((GMaxRHIFeatureLevel == ERHIFeatureLevel::SM4 || bRetainAllFeatureLevelData || bNoTiledReflections))
+		if (GMaxRHIFeatureLevel == ERHIFeatureLevel::SM4)
 		{
 			SM4FullHDRCubemapTexture = new FReflectionTextureCubeResource();
 			SM4FullHDRCubemapTexture->SetupParameters(GReflectionCaptureSize, FMath::CeilLogTwo(GReflectionCaptureSize) + 1, PF_FloatRGBA, &FullHDRDerivedData->GetCapturedDataForSM4Load());
 			BeginInitResource(SM4FullHDRCubemapTexture);
 		}
 
-		if (!bRetainAllFeatureLevelData)
+		if (!bEncodedDataRequired)
 		{
 			EncodedHDRDerivedData = NULL;
 		}
 	}
 
-	if (EncodedHDRDerivedData && (GMaxRHIFeatureLevel == ERHIFeatureLevel::ES2 || bRetainAllFeatureLevelData))
+	if (EncodedHDRDerivedData && bEncodedDataRequired)
 	{
 		if (FPlatformProperties::RequiresCookedData())
 		{
@@ -1038,12 +1039,9 @@ void UReflectionCaptureComponent::PostLoad()
 		EncodedHDRCubemapTexture->SetupParameters(GReflectionCaptureSize, FMath::CeilLogTwo(GReflectionCaptureSize) + 1, PF_B8G8R8A8, &EncodedHDRDerivedData->CapturedData);
 		BeginInitResource(EncodedHDRCubemapTexture);
 
-		if (!bRetainAllFeatureLevelData)
-		{
-			// Don't need the full HDR data for rendering on this feature level
-			delete FullHDRDerivedData;
-			FullHDRDerivedData = NULL;
-		}
+		// Don't need the full HDR data for rendering on this feature level
+		delete FullHDRDerivedData;
+		FullHDRDerivedData = NULL;
 	}
 }
 
@@ -1336,6 +1334,76 @@ void UReflectionCaptureComponent::UpdateReflectionCaptureContents(UWorld* WorldT
 	}
 }
 
+#if WITH_EDITOR
+void UReflectionCaptureComponent::PreFeatureLevelChange(ERHIFeatureLevel::Type PendingFeatureLevel)
+{
+	if (PendingFeatureLevel == ERHIFeatureLevel::ES2)
+	{
+		// generate encoded hdr data for ES2 preview mode.
+		if (World != nullptr)
+		{
+			// Capture full hdr derived data first.
+			ReadbackFromGPUAndSaveDerivedData(World);
+		}
+		if (FullHDRDerivedData)
+		{
+			EncodedHDRDerivedData = FReflectionCaptureEncodedHDRDerivedData::GenerateEncodedHDRData(*FullHDRDerivedData, StateId, Brightness);
+			if (FPlatformProperties::RequiresCookedData() && EncodedHDRDerivedData)
+			{
+				INC_MEMORY_STAT_BY(STAT_ReflectionCaptureMemory, EncodedHDRDerivedData->CapturedData.GetAllocatedSize());
+			}
+
+			if (EncodedHDRCubemapTexture == nullptr)
+			{
+				EncodedHDRCubemapTexture = new FReflectionTextureCubeResource();
+			}
+			EncodedHDRCubemapTexture->SetupParameters(GReflectionCaptureSize, FMath::CeilLogTwo(GReflectionCaptureSize) + 1, PF_B8G8R8A8, &EncodedHDRDerivedData->CapturedData);
+			BeginInitResource(EncodedHDRCubemapTexture);
+		}
+	}
+	else
+	{
+		// Release ES2's unused textures & data.
+		if (FPlatformProperties::RequiresCookedData() && EncodedHDRDerivedData)
+		{
+			DEC_MEMORY_STAT_BY(STAT_ReflectionCaptureMemory, EncodedHDRDerivedData->CapturedData.GetAllocatedSize());
+		}
+		EncodedHDRDerivedData = nullptr;
+		if (EncodedHDRCubemapTexture)
+		{
+			BeginReleaseResource(EncodedHDRCubemapTexture);
+			FlushRenderingCommands();
+			delete EncodedHDRCubemapTexture;
+			EncodedHDRCubemapTexture = nullptr;
+		}
+
+		// for >= SM4 capture should be updated.
+		SetCaptureIsDirty();
+	}
+
+	if (PendingFeatureLevel == ERHIFeatureLevel::SM4)
+	{
+		if (SM4FullHDRCubemapTexture == nullptr)
+		{
+			SM4FullHDRCubemapTexture = new FReflectionTextureCubeResource();
+			SM4FullHDRCubemapTexture->SetupParameters(GReflectionCaptureSize, FMath::CeilLogTwo(GReflectionCaptureSize) + 1, PF_FloatRGBA, NULL);
+			BeginInitResource(SM4FullHDRCubemapTexture);
+		}
+	}
+	else
+	{
+		// release SM4 texture
+		if (SM4FullHDRCubemapTexture != nullptr)
+		{
+			BeginReleaseResource(SM4FullHDRCubemapTexture);
+			FlushRenderingCommands();
+			delete SM4FullHDRCubemapTexture;
+			SM4FullHDRCubemapTexture = nullptr;
+		}
+	}
+}
+#endif // WITH_EDITOR
+
 USphereReflectionCaptureComponent::USphereReflectionCaptureComponent(const class FPostConstructInitializeProperties& PCIP)
 	: Super(PCIP)
 {
diff --git a/Engine/Source/Runtime/Engine/Private/World.cpp b/Engine/Source/Runtime/Engine/Private/World.cpp
--- a/Engine/Source/Runtime/Engine/Private/World.cpp
+++ b/Engine/Source/Runtime/Engine/Private/World.cpp
@@ -5319,14 +5319,25 @@ void UWorld::GetLandscapeTexturesAndMaterials(ULevel* Level, TArray<UObject*>& O
 	}
 }
 
+#if WITH_EDITOR
 void UWorld::ChangeFeatureLevel(ERHIFeatureLevel::Type InFeatureLevel)
 {
 	if (InFeatureLevel != FeatureLevel)
 	{
 		FlushRenderingCommands();
 
+		// Give all scene components the opportunity to prepare for pending feature level change.
+		for (TObjectIterator<USceneComponent> It; It; ++It)
+		{
+			USceneComponent* SceneComponent = *It;
+			if (SceneComponent->GetWorld() == this)
+			{
+				SceneComponent->PreFeatureLevelChange(InFeatureLevel);
+			}
+		}
+
 		FGlobalComponentReregisterContext RecreateComponents;
-		
+
 		// Decrement refcount on old feature level
 		UMaterialInterface::SetGlobalRequiredFeatureLevel(InFeatureLevel, true);
 
@@ -5362,6 +5373,7 @@ void UWorld::ChangeFeatureLevel(ERHIFeatureLevel::Type InFeatureLevel)
 		RecompileShaders(TEXT("CHANGED"), Ar);
 	}
 }
+#endif // WITH_EDITOR
 
 /**
 * Dump visible actors in current world.
diff --git a/Engine/Source/Runtime/Renderer/Private/ReflectionEnvironment.cpp b/Engine/Source/Runtime/Renderer/Private/ReflectionEnvironment.cpp
--- a/Engine/Source/Runtime/Renderer/Private/ReflectionEnvironment.cpp
+++ b/Engine/Source/Runtime/Renderer/Private/ReflectionEnvironment.cpp
@@ -46,22 +46,6 @@ static TAutoConsoleVariable<int32> CVarHalfResReflections(
 	TEXT(" 0 is off (default), 1 is on"),
 	ECVF_RenderThreadSafe);
 
-// NoTiledReflections option, useful to fall back to SM4 style reflections on some SM5 HW
-static TAutoConsoleVariable<int32> CVarNoTiledReflections(
-	TEXT("r.NoTiledReflections"),
-	0,
-	TEXT("Use SM4 style PS reflections on SM5 HW.\n")
-	TEXT(" 0 is off (default), 1 is on"),
-	ECVF_ReadOnly);
-
-int32 GNumCapturesBeforeUsingTiledDeferred = 0;
-static FAutoConsoleVariableRef CVarNumCapturesBeforeUsingTiledDeferred(
-	TEXT("r.TiledReflectionEnvironmentMinimumCount"),
-	GNumCapturesBeforeUsingTiledDeferred,
-	TEXT("Number of applicable reflection captures that must be in the scene before switching to tiled deferred."),
-	ECVF_RenderThreadSafe
-	);
-
 // to avoid having direct access from many places
 static int GetReflectionEnvironmentCVar()
 {
@@ -611,7 +595,7 @@ public:
 		}
 		else
 		{
-		SetTextureParameter(RHICmdList, ShaderRHI, ReflectionEnvironmentColorTexture, ReflectionEnvironmentColorSampler, TStaticSamplerState<SF_Trilinear,AM_Clamp,AM_Clamp,AM_Clamp>::GetRHI(), SortData.SM4FullHDRCubemap->TextureRHI);
+			SetTextureParameter(RHICmdList, ShaderRHI, ReflectionEnvironmentColorTexture, ReflectionEnvironmentColorSampler, TStaticSamplerState<SF_Trilinear,AM_Clamp,AM_Clamp,AM_Clamp>::GetRHI(), SortData.SM4FullHDRCubemap->TextureRHI);
 		}
 
 		DeferredParameters.Set(RHICmdList, ShaderRHI, View);
@@ -984,19 +968,6 @@ void FDeferredShadingSceneRenderer::RenderStandardDeferredImageBasedReflections(
 	}
 }
 
-bool ShouldUseTiledDeferredReflectionEnvironment(FScene* Scene)
-{
-	// command line option: -NoTiledReflections
-	// useful to test this code path and to run the potentially more efficient path (PS vs CS) on SM5 hardware
-	// not a cvar as it also needs to change the data when loading the captures
-	static const auto CVar = IConsoleManager::Get().FindTConsoleVariableDataInt(TEXT("r.NoTiledReflections"));
-	check(CVar);
-	const bool bNoTiledReflections = (CVar->GetValueOnAnyThread() == 1);
-
-	const ERHIFeatureLevel::Type FeatureLevel = Scene->GetFeatureLevel();
-	return FeatureLevel >= ERHIFeatureLevel::SM5 && Scene->ReflectionSceneData.RegisteredReflectionCaptures.Num() > GNumCapturesBeforeUsingTiledDeferred && !bNoTiledReflections;
-}
-
 void FDeferredShadingSceneRenderer::RenderDeferredReflections(FRHICommandListImmediate& RHICmdList, const TRefCountPtr<IPooledRenderTarget>& DynamicBentNormalAO)
 {
 	if (IsSimpleDynamicLightingEnabled() || ViewFamily.EngineShowFlags.VisualizeLightCulling)
@@ -1020,9 +991,9 @@ void FDeferredShadingSceneRenderer::RenderDeferredReflections(FRHICommandListImm
 	else
 	{
 		const bool bReflectionEnvironment = ShouldDoReflectionEnvironment();
-		const bool bWantsTiledReflections = ShouldUseTiledDeferredReflectionEnvironment(Scene);
+		const bool bReflectionsWithCompute = (FeatureLevel >= ERHIFeatureLevel::SM5) && bReflectionEnvironment && Scene->ReflectionSceneData.CubemapArray.IsValid();
 
-		if (bWantsTiledReflections && bReflectionEnvironment)
+		if (bReflectionsWithCompute)
 		{
 			RenderTiledDeferredImageBasedReflections(RHICmdList, DynamicBentNormalAO);
 		}
