Preventing Blueprint options/details from showing for BSP brush actors.

diff --git a/Engine/Source/Runtime/Engine/Classes/Engine/Brush.h b/Engine/Source/Runtime/Engine/Classes/Engine/Brush.h
--- a/Engine/Source/Runtime/Engine/Classes/Engine/Brush.h
+++ b/Engine/Source/Runtime/Engine/Classes/Engine/Brush.h
@@ -72,7 +72,7 @@ struct FGeomSelection
 	}
 };
 
-UCLASS(hidecategories=(Object, Collision, Display, Rendering, Physics, Input), showcategories=("Input|MouseInput", "Input|TouchInput"), NotBlueprintable, ConversionRoot)
+UCLASS(hidecategories=(Object, Collision, Display, Rendering, Physics, Input, Blueprint), showcategories=("Input|MouseInput", "Input|TouchInput"), NotBlueprintable, ConversionRoot)
 class ENGINE_API ABrush : public AActor
 {
 	GENERATED_UCLASS_BODY()
