Added functionality to hide the mouse during a temporary capture due to click.

diff --git a/Engine/Source/Runtime/Engine/Classes/Engine/GameViewportClient.h b/Engine/Source/Runtime/Engine/Classes/Engine/GameViewportClient.h
--- a/Engine/Source/Runtime/Engine/Classes/Engine/GameViewportClient.h
+++ b/Engine/Source/Runtime/Engine/Classes/Engine/GameViewportClient.h
@@ -679,6 +679,22 @@ public:
 		return MouseCaptureMode;
 	}
 
+	/**
+	 * Sets whether or not the cursor is hidden when the viewport captures the mouse
+	 */
+	void SetHideCursorDuringCapture(bool InHideCursorDuringCapture)
+	{
+		bHideCursorDuringCapture = InHideCursorDuringCapture;
+	}
+
+	/**
+	 * Gets whether or not the cursor is hidden when the viewport captures the mouse
+	 */
+	virtual bool HideCursorDuringCapture() override
+	{
+		return bHideCursorDuringCapture;
+	}
+
 private:
 	/**
 	 * Set a specific stat to either enabled or disabled (returns the number of remaining enabled stats)
@@ -767,6 +783,10 @@ private:
 
 	/** Mouse capture behavior when the viewport is clicked */
 	EMouseCaptureMode::Type MouseCaptureMode;
+
+	/** Whether or not the cursor is hidden when the viewport captures the mouse */
+	bool bHideCursorDuringCapture;
+
 };
 
 
diff --git a/Engine/Source/Runtime/Engine/Private/GameViewportClient.cpp b/Engine/Source/Runtime/Engine/Private/GameViewportClient.cpp
--- a/Engine/Source/Runtime/Engine/Private/GameViewportClient.cpp
+++ b/Engine/Source/Runtime/Engine/Private/GameViewportClient.cpp
@@ -84,6 +84,7 @@ UGameViewportClient::UGameViewportClient(const class FPostConstructInitializePro
 	, HighResScreenshotDialog(NULL)
 	, bIgnoreInput(false)
 	, MouseCaptureMode(EMouseCaptureMode::CapturePermanently)
+	, bHideCursorDuringCapture(false)
 {
 
 	TitleSafeZone.MaxPercentX = 0.9f;
diff --git a/Engine/Source/Runtime/Engine/Private/Slate/SceneViewport.cpp b/Engine/Source/Runtime/Engine/Private/Slate/SceneViewport.cpp
--- a/Engine/Source/Runtime/Engine/Private/Slate/SceneViewport.cpp
+++ b/Engine/Source/Runtime/Engine/Private/Slate/SceneViewport.cpp
@@ -31,6 +31,7 @@ FSceneViewport::FSceneViewport( FViewportClient* InViewportClient, TSharedPtr<SV
 	, bIsResizing( false )
 	, bPlayInEditorGetsMouseControl( true )
 	, bPlayInEditorIsSimulate( false )
+	, bCursorHiddenDueToCapture( false )
 {
 	bIsSlateViewport = true;
 }
@@ -308,6 +309,11 @@ bool FSceneViewport::IsForegroundWindow() const
 
 FCursorReply FSceneViewport::OnCursorQuery( const FGeometry& MyGeometry, const FPointerEvent& CursorEvent )
 {
+	if (bCursorHiddenDueToCapture)
+	{
+		return FCursorReply::Cursor(EMouseCursor::None);
+	}
+
 	EMouseCursor::Type MouseCursorToUse = EMouseCursor::Default;
 
 	// If the cursor should be hidden, use EMouseCursor::None,
@@ -364,7 +370,13 @@ FReply FSceneViewport::OnMouseButtonDown( const FGeometry& InGeometry, const FPo
 			{
 				CurrentReplyState.CaptureMouse(ViewportWidgetRef);
 				CurrentReplyState.LockMouseToWidget(ViewportWidgetRef);
-				if (!World->GetFirstPlayerController()->ShouldShowMouseCursor())
+
+				bool bShouldShowMouseCursor = World->GetFirstPlayerController()->ShouldShowMouseCursor();
+				if (ViewportClient->HideCursorDuringCapture() && bShouldShowMouseCursor)
+				{
+					bCursorHiddenDueToCapture = true;
+				}
+				if (bCursorHiddenDueToCapture || !bShouldShowMouseCursor)
 				{
 					CurrentReplyState.UseHighPrecisionMouseMovement(ViewportWidgetRef);
 				}
@@ -411,6 +423,8 @@ FReply FSceneViewport::OnMouseButtonUp( const FGeometry& InGeometry, const FPoin
 		// as long as the left or right mouse buttons are not still down
 		if( !InMouseEvent.IsMouseButtonDown( EKeys::RightMouseButton ) && !InMouseEvent.IsMouseButtonDown( EKeys::LeftMouseButton ))
 		{
+			bCursorHiddenDueToCapture = false;
+
 			CurrentReplyState.ReleaseMouseCapture();
 			if (bIsCursorVisible)
 			{
diff --git a/Engine/Source/Runtime/Engine/Public/Slate/SceneViewport.h b/Engine/Source/Runtime/Engine/Public/Slate/SceneViewport.h
--- a/Engine/Source/Runtime/Engine/Public/Slate/SceneViewport.h
+++ b/Engine/Source/Runtime/Engine/Public/Slate/SceneViewport.h
@@ -311,6 +311,8 @@ private:
 	bool bPlayInEditorGetsMouseControl;
 	/** Whether the PIE viewport is currently in simulate in editor mode */
 	bool bPlayInEditorIsSimulate;
+	/** Whether or not the cursor is hidden when the viewport captures the mouse */
+	bool bCursorHiddenDueToCapture;
 };
 
 
diff --git a/Engine/Source/Runtime/Engine/Public/UnrealClient.h b/Engine/Source/Runtime/Engine/Public/UnrealClient.h
--- a/Engine/Source/Runtime/Engine/Public/UnrealClient.h
+++ b/Engine/Source/Runtime/Engine/Public/UnrealClient.h
@@ -877,6 +877,11 @@ public:
 	 * Gets the mouse capture behavior when the viewport is clicked
 	 */
 	virtual EMouseCaptureMode::Type CaptureMouseOnClick() { return EMouseCaptureMode::CapturePermanently; }
+
+	/**
+	 * Gets whether or not the cursor is hidden when the viewport captures the mouse
+	 */
+	virtual bool HideCursorDuringCapture() { return false; }
 };
 
 /** Tracks the viewport client that should process the stat command, can be NULL */
diff --git a/Engine/Source/Runtime/UMG/Private/WidgetBlueprintLibrary.cpp b/Engine/Source/Runtime/UMG/Private/WidgetBlueprintLibrary.cpp
--- a/Engine/Source/Runtime/UMG/Private/WidgetBlueprintLibrary.cpp
+++ b/Engine/Source/Runtime/UMG/Private/WidgetBlueprintLibrary.cpp
@@ -54,13 +54,13 @@ void UWidgetBlueprintLibrary::SetInputMode_UIOnly(APlayerController* Target, UWi
 	}
 }
 
-void UWidgetBlueprintLibrary::SetInputMode_GameAndUI(APlayerController* Target, UWidget* InWidgetToFocus, bool bLockMouseToViewport)//, bool bHideCursorDuringCapture)
+void UWidgetBlueprintLibrary::SetInputMode_GameAndUI(APlayerController* Target, UWidget* InWidgetToFocus, bool bLockMouseToViewport, bool bHideCursorDuringCapture)
 {
 	if (Target != nullptr)
 	{
 		FInputModeGameAndUI InputMode;
 		InputMode.SetLockMouseToViewport(bLockMouseToViewport);
-		//InputMode.SetHideCursorDuringCapture(bHideCursorDuringCapture);
+		InputMode.SetHideCursorDuringCapture(bHideCursorDuringCapture);
 
 		if (InWidgetToFocus != nullptr)
 		{
diff --git a/Engine/Source/Runtime/UMG/Public/Blueprint/WidgetBlueprintLibrary.h b/Engine/Source/Runtime/UMG/Public/Blueprint/WidgetBlueprintLibrary.h
--- a/Engine/Source/Runtime/UMG/Public/Blueprint/WidgetBlueprintLibrary.h
+++ b/Engine/Source/Runtime/UMG/Public/Blueprint/WidgetBlueprintLibrary.h
@@ -23,7 +23,7 @@ class UWidgetBlueprintLibrary : public UBlueprintFunctionLibrary
 
 	/** Setup an input mode that allows only the UI to respond to user input, and if the UI doesn't handle it player input / player controller gets a chance. */
 	UFUNCTION(BlueprintCallable, Category = "Input")
-	static void SetInputMode_GameAndUI(APlayerController* Target, UWidget* InWidgetToFocus = nullptr, bool bLockMouseToViewport = false);//, bool bHideCursorDuringCapture = true);
+	static void SetInputMode_GameAndUI(APlayerController* Target, UWidget* InWidgetToFocus = nullptr, bool bLockMouseToViewport = false, bool bHideCursorDuringCapture = true);
 
 	/** Setup an input mode that allows only player input / player controller to respond to user input. */
 	UFUNCTION(BlueprintCallable, Category = "Input")
