Improve physics stats timing to show sub-stepping and named events for fphysxtask

diff --git a/Engine/Source/Runtime/Engine/Private/PhysicsEngine/PhysSubstepTasks.cpp b/Engine/Source/Runtime/Engine/Private/PhysicsEngine/PhysSubstepTasks.cpp
--- a/Engine/Source/Runtime/Engine/Private/PhysicsEngine/PhysSubstepTasks.cpp
+++ b/Engine/Source/Runtime/Engine/Private/PhysicsEngine/PhysSubstepTasks.cpp
@@ -289,8 +289,13 @@ void FPhysSubstepTask::StepSimulation(PhysXCompletionTask * Task)
 }
 #endif
 
+DEFINE_STAT(STAT_SubstepSimulationStart);
+DEFINE_STAT(STAT_SubstepSimulationEnd);
+
 void FPhysSubstepTask::SubstepSimulationStart()
 {
+	SCOPE_CYCLE_COUNTER(STAT_TotalPhysicsTime);
+	SCOPE_CYCLE_COUNTER(STAT_SubstepSimulationStart);
 #if WITH_PHYSX
 	check(SubTime > 0.f);
 	check(DeltaSeconds > 0.f);
@@ -349,13 +354,20 @@ void FPhysSubstepTask::SubstepSimulationEnd(ENamedThreads::Type CurrentThread, c
 	if (CurrentSubStep < NumSubsteps)
 	{
 		uint32 OutErrorCode = 0;
+		{
+			SCOPE_CYCLE_COUNTER(STAT_TotalPhysicsTime);
+			SCOPE_CYCLE_COUNTER(STAT_SubstepSimulationEnd);
+
 #if WITH_APEX
-		PAScene->fetchResults(true, &OutErrorCode);
+			PAScene->fetchResults(true, &OutErrorCode);
 #else
-		PAScene->lockWrite();
-		PAScene->fetchResults(true, &OutErrorCode);
-		PAScene->unlockWrite();
+			PAScene->lockWrite();
+			PAScene->fetchResults(true, &OutErrorCode);
+			PAScene->unlockWrite();
 #endif
+
+		}
+
 		if (OutErrorCode != 0)
 		{
 			UE_LOG(LogPhysics, Log, TEXT("PHYSX FETCHRESULTS ERROR: %d"), OutErrorCode);
@@ -365,6 +377,9 @@ void FPhysSubstepTask::SubstepSimulationEnd(ENamedThreads::Type CurrentThread, c
 	}
 	else
 	{
+		SCOPE_CYCLE_COUNTER(STAT_TotalPhysicsTime);
+		SCOPE_CYCLE_COUNTER(STAT_SubstepSimulationEnd);
+
 		//final step we call fetch on in game thread
 		FullSimulationTask->removeReference();
 	}
diff --git a/Engine/Source/Runtime/Engine/Private/PhysicsEngine/PhysXSupport.cpp b/Engine/Source/Runtime/Engine/Private/PhysicsEngine/PhysXSupport.cpp
--- a/Engine/Source/Runtime/Engine/Private/PhysicsEngine/PhysXSupport.cpp
+++ b/Engine/Source/Runtime/Engine/Private/PhysicsEngine/PhysXSupport.cpp
@@ -567,7 +567,7 @@ public:
 
 	FORCEINLINE TStatId GetStatId() const
 	{
-		RETURN_QUICK_DECLARE_CYCLE_STAT(FPhysXTask, STATGROUP_TaskGraphTasks);
+		RETURN_QUICK_DECLARE_CYCLE_STAT(FPhysXTask, STATGROUP_Physics);
 	}
 	static ENamedThreads::Type GetDesiredThread()
 	{
@@ -580,7 +580,9 @@ public:
 
 	void DoTask(ENamedThreads::Type CurrentThread, const FGraphEventRef& MyCompletionGraphEvent)
 	{
+		FPlatformMisc::BeginNamedEvent(FColor::Black, Task.getName());
 		Task.run();
+		FPlatformMisc::EndNamedEvent();
 	}
 };
 
diff --git a/Engine/Source/Runtime/Engine/Public/PhysicsPublic.h b/Engine/Source/Runtime/Engine/Public/PhysicsPublic.h
--- a/Engine/Source/Runtime/Engine/Public/PhysicsPublic.h
+++ b/Engine/Source/Runtime/Engine/Public/PhysicsPublic.h
@@ -12,11 +12,13 @@
 /**
  * Physics stats
  */
-DECLARE_CYCLE_STAT_EXTERN(TEXT("Total Physics Time"),STAT_TotalPhysicsTime,STATGROUP_Physics, );
+DECLARE_CYCLE_STAT_EXTERN(TEXT("FetchAndStart Time"),STAT_TotalPhysicsTime,STATGROUP_Physics, );
 DECLARE_CYCLE_STAT_EXTERN(TEXT("Start Physics Time"),STAT_PhysicsKickOffDynamicsTime,STATGROUP_Physics, );
 DECLARE_CYCLE_STAT_EXTERN(TEXT("Fetch Results Time"),STAT_PhysicsFetchDynamicsTime,STATGROUP_Physics, );
 DECLARE_CYCLE_STAT_EXTERN(TEXT("Phys Events Time"),STAT_PhysicsEventTime,STATGROUP_Physics, );
 DECLARE_CYCLE_STAT_EXTERN(TEXT("Phys SetBodyTransform"),STAT_SetBodyTransform,STATGROUP_Physics, );
+DECLARE_CYCLE_STAT_EXTERN(TEXT("Phys SubstepStart"), STAT_SubstepSimulationStart,STATGROUP_Physics, );
+DECLARE_CYCLE_STAT_EXTERN(TEXT("Phys SubstepEnd"), STAT_SubstepSimulationEnd, STATGROUP_Physics, );
 
 DECLARE_DWORD_COUNTER_STAT_EXTERN(TEXT("Broadphase Adds"), STAT_NumBroadphaseAdds, STATGROUP_Physics, );
 DECLARE_DWORD_COUNTER_STAT_EXTERN(TEXT("Broadphase Removes"), STAT_NumBroadphaseRemoves, STATGROUP_Physics, );
