Fix crash with having mutliple meshes within same skeleton

diff --git a/Engine/Source/Runtime/Engine/Private/Animation/AnimSequence.cpp b/Engine/Source/Runtime/Engine/Private/Animation/AnimSequence.cpp
--- a/Engine/Source/Runtime/Engine/Private/Animation/AnimSequence.cpp
+++ b/Engine/Source/Runtime/Engine/Private/Animation/AnimSequence.cpp
@@ -1571,7 +1571,7 @@ bool UAnimSequence::IsValidAdditive() const
 
 #if WITH_EDITOR
 
-void FillUpTransformBasedOnRig(const USkeleton* Skeleton, TArray<FTransform> & NodeSpaceBases, TArray<FTransform> &Rotations, TArray<FTransform> & Translations)
+void FillUpTransformBasedOnRig(USkeleton* Skeleton, TArray<FTransform> & NodeSpaceBases, TArray<FTransform> &Rotations, TArray<FTransform> & Translations)
 {
 	TArray<FTransform> SpaceBases;
 	FAnimationRuntime::FillUpSpaceBasesRetargetBasePose(Skeleton, SpaceBases);
@@ -1593,12 +1593,21 @@ void FillUpTransformBasedOnRig(const USkeleton* Skeleton, TArray<FTransform> & N
 		Translations.Empty(NodeNum);
 		Translations.AddUninitialized(NodeNum);
 
+		const USkeletalMesh* PreviewMesh = Skeleton->GetPreviewMesh();
+
 		for ( int32 Index = 0; Index < NodeNum; ++Index )
 		{
 			const FName NodeName = Rig->GetNodeName(Index);
 			const FName& BoneName = Skeleton->GetRigBoneMapping(NodeName);
-			const int32& BoneIndex = Skeleton->GetReferenceSkeleton().FindBoneIndex(BoneName);
+			const int32& SkeletonBoneIndex = Skeleton->GetReferenceSkeleton().FindBoneIndex(BoneName);
+
+			int32 BoneIndex = INDEX_NONE;
 
+			if (SkeletonBoneIndex != INDEX_NONE)
+			{
+				BoneIndex = Skeleton->GetMeshBoneIndexFromSkeletonBoneIndex(PreviewMesh, SkeletonBoneIndex);
+			}
+			
 			if (BoneIndex == INDEX_NONE)
 			{
 				// add identity
