Added code to replace \\ with /  in documentationactor

diff --git a/Engine/Source/Runtime/Engine/Private/DocumentationActor.cpp b/Engine/Source/Runtime/Engine/Private/DocumentationActor.cpp
--- a/Engine/Source/Runtime/Engine/Private/DocumentationActor.cpp
+++ b/Engine/Source/Runtime/Engine/Private/DocumentationActor.cpp
@@ -69,6 +69,7 @@ void ADocumentationActor::PostEditChangeProperty(FPropertyChangedEvent& Property
 void ADocumentationActor::UpdateLinkType()
 {
 #if WITH_EDITORONLY_DATA
+	DocumentLink = DocumentLink.Replace(TEXT("\\"), TEXT("/"));	
 	if (DocumentLink.IsEmpty() == true)
 	{
 		LinkType = EDocumentationActorType::None;
