Added description of bone as well as socket.

diff --git a/Engine/Source/Runtime/Engine/Classes/Components/SceneComponent.h b/Engine/Source/Runtime/Engine/Classes/Components/SceneComponent.h
--- a/Engine/Source/Runtime/Engine/Classes/Components/SceneComponent.h
+++ b/Engine/Source/Runtime/Engine/Classes/Components/SceneComponent.h
@@ -377,7 +377,7 @@ public:
 
 	/** 
 	 * Get world-space socket transform.
-	 * @param InSocketName Name of the socket to get the transform 
+	 * @param InSocketName Name of the socket or the bone to get the transform 
 	 * @return Socket transform in world space if socket if found. Otherwise it will return component's transform in world space.
 	 */
 	UFUNCTION(BlueprintCallable, Category="Utilities|Transformation", meta=(Keywords="Bone"))
@@ -385,7 +385,7 @@ public:
 
 	/** 
 	 * Get world-space socket or bone location.
-	 * @param InSocketName Name of the socket to get the transform 
+	 * @param InSocketName Name of the socket or the bone to get the transform 
 	 * @return Socket transform in world space if socket if found. Otherwise it will return component's transform in world space.
 	 */
 	UFUNCTION(BlueprintCallable, Category="Utilities|Transformation", meta=(Keywords="Bone"))
@@ -393,7 +393,7 @@ public:
 
 	/** 
 	 * Get world-space socket or bone  FRotator rotation.
-	 * @param InSocketName Name of the socket to get the transform 
+	 * @param InSocketName Name of the socket or the bone to get the transform 
 	 * @return Socket transform in world space if socket if found. Otherwise it will return component's transform in world space.
 	 */
 	UFUNCTION(BlueprintCallable, Category="Utilities|Transformation", meta=(Keywords="Bone"))
@@ -401,7 +401,7 @@ public:
 
 	/** 
 	 * Get world-space socket or bone FQuat rotation.
-	 * @param InSocketName Name of the socket to get the transform 
+	 * @param InSocketName Name of the socket or the bone to get the transform 
 	 * @return Socket transform in world space if socket if found. Otherwise it will return component's transform in world space.
 	 */
 	UFUNCTION(BlueprintCallable, Category="Utilities|Transformation", meta=(Keywords="Bone"))
@@ -409,7 +409,7 @@ public:
 
 	/** 
 	 * return true if socket with the given name exists
-	 * @param InSocketName Name of the socket to get the transform 
+	 * @param InSocketName Name of the socket or the bone to get the transform 
 	 * @return true if the socket with the given name exists. Otherwise, return false
 	 */
 	UFUNCTION(BlueprintCallable, Category="Utilities|Transformation", meta=(Keywords="Bone"))
