Add the option to turn substepping off for async scene

diff --git a/Engine/Source/Runtime/Engine/Classes/PhysicsEngine/PhysicsSettings.h b/Engine/Source/Runtime/Engine/Classes/PhysicsEngine/PhysicsSettings.h
--- a/Engine/Source/Runtime/Engine/Classes/PhysicsEngine/PhysicsSettings.h
+++ b/Engine/Source/Runtime/Engine/Classes/PhysicsEngine/PhysicsSettings.h
@@ -123,6 +123,10 @@ class ENGINE_API UPhysicsSettings
 	UPROPERTY(config, EditAnywhere, Category = Framerate)
 	bool bSubstepping;
 
+	/** Whether to substep the async physics simulation. This feature is still experimental. Certain functionality might not work correctly*/
+	UPROPERTY(config, EditAnywhere, Category = Framerate)
+	bool bSubsteppingAsync;
+
 	/** Max delta time for an individual substep simulation. */
 	UPROPERTY(config, EditAnywhere, meta = (ClampMin = "0.0013", UIMin = "0.0013", ClampMax = "1.0", UIMax = "1.0", editcondition = "bSubStepping"), Category=Framerate)
 	float MaxSubstepDeltaTime;
diff --git a/Engine/Source/Runtime/Engine/Private/PhysicsEngine/PhysScene.cpp b/Engine/Source/Runtime/Engine/Private/PhysicsEngine/PhysScene.cpp
--- a/Engine/Source/Runtime/Engine/Private/PhysicsEngine/PhysScene.cpp
+++ b/Engine/Source/Runtime/Engine/Private/PhysicsEngine/PhysScene.cpp
@@ -99,8 +99,8 @@ FPhysScene::FPhysScene()
 
 #if WITH_SUBSTEPPING
 	bSubstepping = PhysSetting->bSubstepping;
+	bSubsteppingAsync = PhysSetting->bSubsteppingAsync;
 #endif
-
 	bAsyncSceneEnabled = PhysSetting->bEnableAsyncScene;
 	NumPhysScenes = bAsyncSceneEnabled ? PST_Async + 1 : PST_Cloth + 1;
 
@@ -174,9 +174,10 @@ bool FPhysScene::GetKinematicTarget(const FBodyInstance* BodyInstance, FTransfor
 	if (PxRigidDynamic * PRigidDynamic = BodyInstance->GetPxRigidDynamic())
 	{
 #if WITH_SUBSTEPPING
-		if (IsSubstepping())
+		uint32 BodySceneType = SceneType(BodyInstance);
+		if (IsSubstepping(BodySceneType))
 		{
-			FPhysSubstepTask * PhysSubStepper = PhysSubSteppers[SceneType(BodyInstance)];
+			FPhysSubstepTask * PhysSubStepper = PhysSubSteppers[BodySceneType];
 			return PhysSubStepper->GetKinematicTarget(BodyInstance, OutTM);
 		}
 		else
@@ -206,9 +207,10 @@ void FPhysScene::SetKinematicTarget(FBodyInstance* BodyInstance, const FTransfor
 	if (PxRigidDynamic * PRigidDynamic = BodyInstance->GetPxRigidDynamic())
 	{
 #if WITH_SUBSTEPPING
-		if (bAllowSubstepping && IsSubstepping())
+		uint32 BodySceneType = SceneType(BodyInstance);
+		if (bAllowSubstepping && IsSubstepping(BodySceneType))
 		{
-			FPhysSubstepTask * PhysSubStepper = PhysSubSteppers[SceneType(BodyInstance)];
+			FPhysSubstepTask * PhysSubStepper = PhysSubSteppers[BodySceneType];
 			PhysSubStepper->SetKinematicTarget(BodyInstance, TargetTransform);
 		}
 		else
@@ -231,9 +233,10 @@ void FPhysScene::AddForce(FBodyInstance* BodyInstance, const FVector& Force, boo
 	if (PxRigidDynamic * PRigidDynamic = BodyInstance->GetPxRigidDynamic())
 	{
 #if WITH_SUBSTEPPING
-		if (bAllowSubstepping && IsSubstepping())
+		uint32 BodySceneType = SceneType(BodyInstance);
+		if (bAllowSubstepping && IsSubstepping(BodySceneType))
 		{
-			FPhysSubstepTask * PhysSubStepper = PhysSubSteppers[SceneType(BodyInstance)];
+			FPhysSubstepTask * PhysSubStepper = PhysSubSteppers[BodySceneType];
 			PhysSubStepper->AddForce(BodyInstance, Force);
 		}
 		else
@@ -253,9 +256,10 @@ void FPhysScene::AddForceAtPosition(FBodyInstance* BodyInstance, const FVector&
 	if (PxRigidDynamic * PRigidDynamic = BodyInstance->GetPxRigidDynamic())
 	{
 #if WITH_SUBSTEPPING
-		if (bAllowSubstepping && IsSubstepping())
+		uint32 BodySceneType = SceneType(BodyInstance);
+		if (bAllowSubstepping && IsSubstepping(BodySceneType))
 		{
-			FPhysSubstepTask * PhysSubStepper = PhysSubSteppers[SceneType(BodyInstance)];
+			FPhysSubstepTask * PhysSubStepper = PhysSubSteppers[BodySceneType];
 			PhysSubStepper->AddForceAtPosition(BodyInstance, Force, Position);
 		}
 		else
@@ -275,9 +279,10 @@ void FPhysScene::AddTorque(FBodyInstance* BodyInstance, const FVector& Torque, b
 	if (PxRigidDynamic * PRigidDynamic = BodyInstance->GetPxRigidDynamic())
 	{
 #if WITH_SUBSTEPPING
-		if (bAllowSubstepping && IsSubstepping())
+		uint32 BodySceneType = SceneType(BodyInstance);
+		if (bAllowSubstepping && IsSubstepping(BodySceneType))
 		{
-			FPhysSubstepTask * PhysSubStepper = PhysSubSteppers[SceneType(BodyInstance)];
+			FPhysSubstepTask * PhysSubStepper = PhysSubSteppers[BodySceneType];
 			PhysSubStepper->AddTorque(BodyInstance, Torque);
 		}
 		else
@@ -485,7 +490,7 @@ void FPhysScene::TickPhysScene(uint32 SceneType, FGraphEventRef& InOutCompletion
 	}
 
 #if WITH_SUBSTEPPING
-	if (IsSubstepping() && SceneType != PST_Cloth)	//we don't bother sub-stepping cloth
+	if (IsSubstepping(SceneType))	//we don't bother sub-stepping cloth
 	{
 		//We're about to start stepping so swap buffers. Might want to find a better place for this?
 		PhysSubSteppers[SceneType]->SwapBuffers();
@@ -538,14 +543,14 @@ void FPhysScene::TickPhysScene(uint32 SceneType, FGraphEventRef& InOutCompletion
 	{
 		float TickTime = AveragedFrameTime[SceneType];
 #if WITH_SUBSTEPPING
-		if (IsSubstepping())
+		if (IsSubstepping(SceneType))
 		{
 			TickTime = UseSyncTime(SceneType) ? SyncDeltaSeconds : DeltaSeconds;
 		}
 #endif
 		VehicleManager->PreTick(TickTime);
 #if WITH_SUBSTEPPING
-		if (IsSubstepping() == false)
+		if (IsSubstepping(SceneType) == false)
 #endif
 		{
 			VehicleManager->Update(AveragedFrameTime[SceneType]);
@@ -570,7 +575,7 @@ void FPhysScene::TickPhysScene(uint32 SceneType, FGraphEventRef& InOutCompletion
 	if(ApexScene && UseDelta > 0.f)
 	{
 #if WITH_SUBSTEPPING
-		if (IsSubstepping() && SceneType != PST_Cloth) //we don't bother sub-stepping cloth
+		if (IsSubstepping(SceneType)) //we don't bother sub-stepping cloth
 		{
 			bTaskOutstanding = SubstepSimulation(SceneType, InOutCompletionEvent);
 		}else
@@ -588,10 +593,9 @@ void FPhysScene::TickPhysScene(uint32 SceneType, FGraphEventRef& InOutCompletion
 	{
 		InOutCompletionEvent->DispatchSubsequents(); // nothing to do, so nothing to wait for
 	}
-
 #if WITH_SUBSTEPPING
-	//check if substepping settings have changed
 	bSubstepping = UPhysicsSettings::Get()->bSubstepping;
+	bSubsteppingAsync = UPhysicsSettings::Get()->bSubsteppingAsync;
 #endif
 }
 
diff --git a/Engine/Source/Runtime/Engine/Private/PhysicsEngine/PhysicsSettings.cpp b/Engine/Source/Runtime/Engine/Private/PhysicsEngine/PhysicsSettings.cpp
--- a/Engine/Source/Runtime/Engine/Private/PhysicsEngine/PhysicsSettings.cpp
+++ b/Engine/Source/Runtime/Engine/Private/PhysicsEngine/PhysicsSettings.cpp
@@ -12,6 +12,7 @@ UPhysicsSettings::UPhysicsSettings(const class FPostConstructInitializePropertie
 	, bSimulateSkeletalMeshOnDedicatedServer(true)
 	, MaxPhysicsDeltaTime(1.f / 30.f)
 	, bSubstepping(false)
+	, bSubsteppingAsync(false)
 	, MaxSubstepDeltaTime(1.f / 60.f)
 	, MaxSubsteps(6)
 	, SyncSceneSmoothingFactor(0.0f)
diff --git a/Engine/Source/Runtime/Engine/Public/PhysicsPublic.h b/Engine/Source/Runtime/Engine/Public/PhysicsPublic.h
--- a/Engine/Source/Runtime/Engine/Public/PhysicsPublic.h
+++ b/Engine/Source/Runtime/Engine/Public/PhysicsPublic.h
@@ -247,6 +247,9 @@ public:
 	/** Indicates whether the scene is using substepping */
 	bool							bSubstepping;
 
+	/** Indicates whether the scene is using substepping */
+	bool							bSubsteppingAsync;
+
 	/** Stores the number of valid scenes we are working with. This will be PST_MAX or PST_Async, 
 		depending on whether the async scene is enabled or not*/
 	uint32							NumPhysScenes;
@@ -366,10 +369,12 @@ public:
 	static bool SupportsOriginShifting() { return true; }
 
 	/** @return Whether physics scene is using substepping */
-	bool IsSubstepping() const
+	bool IsSubstepping(uint32 SceneType) const
 	{
 #if WITH_SUBSTEPPING
-		return bSubstepping;
+		if (SceneType == PST_Sync) return bSubstepping;
+		if (SceneType == PST_Async) return bSubsteppingAsync;
+		return false;
 #else
 		return false;
 #endif
