fix for including T_EditorHelp.uasset in cooked games
TTP348789
#ue4

diff --git a/Engine/Source/Runtime/Engine/Private/DocumentationActor.cpp b/Engine/Source/Runtime/Engine/Private/DocumentationActor.cpp
--- a/Engine/Source/Runtime/Engine/Private/DocumentationActor.cpp
+++ b/Engine/Source/Runtime/Engine/Private/DocumentationActor.cpp
@@ -14,11 +14,14 @@ ADocumentationActor::ADocumentationActor(const class FPostConstructInitializePro
 	RootComponent = SceneComponent;	
 
 #if WITH_EDITORONLY_DATA
- 	static ConstructorHelpers::FObjectFinder<UMaterial> MaterialAsset(TEXT("/Engine/EditorMaterials/HelpActorMaterial"));
- 	// Create a Material billboard to represent our actor
+	// Create a Material billboard to represent our actor
 	Billboard = PCIP.CreateDefaultSubobject<UMaterialBillboardComponent>(this, TEXT("BillboardComponent"));
-	Billboard->AddElement(MaterialAsset.Object, nullptr, false, 32.0f, 32.0f, nullptr);
-	Billboard->AttachParent = RootComponent;
+	if (!IsRunningCommandlet() && (Billboard != NULL))
+	{
+		static ConstructorHelpers::FObjectFinder<UMaterial> MaterialAsset(TEXT("/Engine/EditorMaterials/HelpActorMaterial"));
+		Billboard->AddElement(MaterialAsset.Object, nullptr, false, 32.0f, 32.0f, nullptr);
+		Billboard->AttachParent = RootComponent;
+	}
 #endif //WITH_EDITORONLY_DATA
 }
 
