Make sure compression exists when post edit change is called.

diff --git a/Engine/Source/Runtime/Engine/Private/Animation/AnimSequence.cpp b/Engine/Source/Runtime/Engine/Private/Animation/AnimSequence.cpp
--- a/Engine/Source/Runtime/Engine/Private/Animation/AnimSequence.cpp
+++ b/Engine/Source/Runtime/Engine/Private/Animation/AnimSequence.cpp
@@ -494,6 +494,12 @@ void UAnimSequence::PostEditChangeProperty(FPropertyChangedEvent& PropertyChange
 			RefPoseSeq = NULL;
 		}
 	}
+
+	// @Todo fix me: This is temporary fix to make sure they always have compressed data
+	if (CompressedTrackOffsets.Num() == 0 && RawAnimationData.Num() > 0)
+	{
+		PostProcessSequence();
+	}
 }
 #endif // WITH_EDITOR
 
