Fixed regression with streaming levels renaming for PIE
#TTP 348698

diff --git a/Engine/Source/Runtime/Engine/Classes/Engine/World.h b/Engine/Source/Runtime/Engine/Classes/Engine/World.h
--- a/Engine/Source/Runtime/Engine/Classes/Engine/World.h
+++ b/Engine/Source/Runtime/Engine/Classes/Engine/World.h
@@ -2448,6 +2448,7 @@ public:
 
 public:
 	static FString ConvertToPIEPackageName(const FString& PackageName, int32 PIEInstanceID);
+	static FString StripPIEPrefixFromPackageName(const FString& PackageName, const FString& Prefix);
 	static FString BuildPIEPackagePrefix(int32 PIEInstanceID);
 	static UWorld* DuplicateWorldForPIE(const FString& PackageName, UWorld* OwningWorld);
 	static FString RemovePIEPrefix(const FString &Source);
diff --git a/Engine/Source/Runtime/Engine/Private/LevelStreaming.cpp b/Engine/Source/Runtime/Engine/Private/LevelStreaming.cpp
--- a/Engine/Source/Runtime/Engine/Private/LevelStreaming.cpp
+++ b/Engine/Source/Runtime/Engine/Private/LevelStreaming.cpp
@@ -400,11 +400,7 @@ bool ULevelStreaming::RequestLevel(UWorld* PersistentWorld, bool bAllowLevelLoad
 		}
 		PIEInstanceID = PersistentWorld->GetOutermost()->PIEInstanceID;
 #endif
-		const FString PrefixedLevelName = DesiredPackageName.ToString();
-		const FString ShortPrefixedLevelName = FPackageName::GetLongPackageAssetName(PrefixedLevelName);
-		const FString LevelPath = FPackageName::GetLongPackagePath(PrefixedLevelName);
-		// Rebuild the original NonPrefixedLevelName so we can find and duplicate it
-		const FString NonPrefixedLevelName = LevelPath + "/" + ShortPrefixedLevelName.RightChop(PersistentWorld->StreamingLevelsPrefix.Len());
+		const FString NonPrefixedLevelName = UWorld::StripPIEPrefixFromPackageName(DesiredPackageName.ToString(), PersistentWorld->StreamingLevelsPrefix);
 					
 		// Do the duplication
 		UWorld* PIELevelWorld = UWorld::DuplicateWorldForPIE(NonPrefixedLevelName, PersistentWorld);
@@ -746,7 +742,10 @@ void ULevelStreaming::RenameForPIE(int32 PIEInstanceID)
 		// Store original name 
 		if (PackageNameToLoad == NAME_None)
 		{
-			PackageNameToLoad = GetWorldAssetPackageFName();
+			FString NonPrefixedName = UWorld::StripPIEPrefixFromPackageName(
+				GetWorldAssetPackageName(), 
+				UWorld::BuildPIEPackagePrefix(PIEInstanceID));
+			PackageNameToLoad = FName(*NonPrefixedName);
 		}
 		const FString PlayWorldSteamingLevelName = UWorld::ConvertToPIEPackageName(GetWorldAssetPackageName(), PIEInstanceID);
 		SetWorldAssetByPackageName(FName(*PlayWorldSteamingLevelName));
diff --git a/Engine/Source/Runtime/Engine/Private/World.cpp b/Engine/Source/Runtime/Engine/Private/World.cpp
--- a/Engine/Source/Runtime/Engine/Private/World.cpp
+++ b/Engine/Source/Runtime/Engine/Private/World.cpp
@@ -2081,6 +2081,23 @@ FString UWorld::ConvertToPIEPackageName(const FString& PackageName, int32 PIEIns
 	}
 }
 
+FString UWorld::StripPIEPrefixFromPackageName(const FString& PrefixedName, const FString& Prefix)
+{
+	FString ResultName;
+	FString ShortPrefixedName = FPackageName::GetLongPackageAssetName(PrefixedName);
+	if (ShortPrefixedName.StartsWith(Prefix))
+	{
+		FString NamePath = FPackageName::GetLongPackagePath(PrefixedName);
+		ResultName = NamePath + "/" + ShortPrefixedName.RightChop(Prefix.Len());
+	}
+	else
+	{
+		ResultName = PrefixedName;
+	}
+
+	return ResultName;
+}
+
 FString UWorld::BuildPIEPackagePrefix(int PIEInstanceID)
 {
 	check(PIEInstanceID != -1);
