Visual Loggger refactor - first pass. #ue4
- new interface. Old FVisualLog class is used only as a devce for visual logger - only to store logs and to communicate with LogVisualizer for now (this class is going to be removed soon).
- added first basic automation tests
- added first simple Blueprint interface
- kept old functionality. ToggleAILogging cheat and EnableAILogging command line works too.

diff --git a/Engine/Source/Developer/LogVisualizer/Private/LogVisualizer.cpp b/Engine/Source/Developer/LogVisualizer/Private/LogVisualizer.cpp
--- a/Engine/Source/Developer/LogVisualizer/Private/LogVisualizer.cpp
+++ b/Engine/Source/Developer/LogVisualizer/Private/LogVisualizer.cpp
@@ -19,7 +19,7 @@ void FLogVisualizer::SummonUI(UWorld* InWorld)
 		}
 
 		World = InWorld;
-		FVisualLog& VisualLog = FVisualLog::Get();
+		FVisualLog& VisualLog = FVisualLog::GetStatic();
 		VisualLog.RegisterNewLogsObserver(FVisualLog::FOnNewLogCreatedDelegate::CreateRaw(this, &FLogVisualizer::OnNewLog));
 		PullDataFromVisualLog(VisualLog);
 
@@ -77,7 +77,7 @@ bool FLogVisualizer::IsOpenUI(UWorld* InWorld)
 
 void FLogVisualizer::CleanUp()
 {
-	FVisualLog::Get().ClearNewLogsObserver();
+	FVisualLog::GetStatic().ClearNewLogsObserver();
 }
 
 class AActor* FLogVisualizer::GetHelperActor(class UWorld* InWorld)
@@ -142,12 +142,12 @@ void FLogVisualizer::AddLoadedLog(TSharedPtr<FActorsVisLog> Log)
 
 bool FLogVisualizer::IsRecording()
 {
-	return FVisualLog::Get().IsRecording();
+	return FVisualLogger::Get().IsRecording();
 }
 
 void FLogVisualizer::SetIsRecording(bool bNewRecording)
 {
-	FVisualLog::Get().SetIsRecording(bNewRecording);
+	FVisualLogger::Get().SetIsRecording(bNewRecording);
 }
 
 int32 FLogVisualizer::GetLogIndexForActor(const AActor* Actor)
diff --git a/Engine/Source/Developer/LogVisualizer/Private/SLogBar.cpp b/Engine/Source/Developer/LogVisualizer/Private/SLogBar.cpp
--- a/Engine/Source/Developer/LogVisualizer/Private/SLogBar.cpp
+++ b/Engine/Source/Developer/LogVisualizer/Private/SLogBar.cpp
@@ -68,7 +68,7 @@ int32 SLogBar::OnPaint( const FPaintArgs& Args, const FGeometry& AllottedGeometr
 	{
 		float CurrentStartX, CurrentEndX;
 		{
-			TSharedPtr<FVisLogEntry> Entry = Entries[EntryIndex];
+			TSharedPtr<FVisualLogEntry> Entry = Entries[EntryIndex];
 			if (!Entry.IsValid() || !CalculateEntryGeometry(Entry.Get(), AllottedGeometry, CurrentStartX, CurrentEndX))
 			{
 				EntryIndex++;
@@ -83,7 +83,7 @@ int32 SLogBar::OnPaint( const FPaintArgs& Args, const FGeometry& AllottedGeometr
 		float LastEndX = MAX_FLT;
 		for (; StartIndex < Entries.Num(); ++StartIndex)
 		{
-			TSharedPtr<FVisLogEntry> Entry = Entries[StartIndex];
+			TSharedPtr<FVisualLogEntry> Entry = Entries[StartIndex];
 			float StartX, EndX;
 			if (!Entry.IsValid() || !CalculateEntryGeometry(Entry.Get(), AllottedGeometry, StartX, EndX))
 			{
@@ -149,7 +149,7 @@ int32 SLogBar::OnPaint( const FPaintArgs& Args, const FGeometry& AllottedGeometr
 			const float PosX = (float)(LogBarWidth * RelativePos);
 			const float ClampedSize = FMath::Clamp(TimeUnit / TotalTime * Zoom, 0.0f, 1.0f);
 
-			TSharedPtr<FVisLogEntry> Entry = Entries[EntryIndex];
+			TSharedPtr<FVisualLogEntry> Entry = Entries[EntryIndex];
 			float StartX, EndX;
 			CalculateEntryGeometry(Entry.Get(), AllottedGeometry, StartX, EndX);
 			const float WindowHalfWidth = LogBarWidth * HistogramPreviewWindow * 0.01;
@@ -189,7 +189,7 @@ int32 SLogBar::OnPaint( const FPaintArgs& Args, const FGeometry& AllottedGeometr
 		EntryIndex = CurrentEntryIndex.IsBound() ? CurrentEntryIndex.Execute() : INDEX_NONE;
 		if (Entries.IsValidIndex(EntryIndex) == true)
 		{
-			TSharedPtr<FVisLogEntry> Entry = Entries[EntryIndex];
+			TSharedPtr<FVisualLogEntry> Entry = Entries[EntryIndex];
 			float StartX, EndX;
 			if (CalculateEntryGeometry( Entry.Get(), AllottedGeometry, StartX, EndX ) )
 			{
@@ -244,7 +244,7 @@ void SLogBar::SelectEntry(const FGeometry& MyGeometry, const float ClickX)
 	{
 		for (EntryIndex = 0; EntryIndex < Entries.Num(); ++EntryIndex)
 		{
-			TSharedPtr<FVisLogEntry> Entry = Entries[EntryIndex];
+			TSharedPtr<FVisualLogEntry> Entry = Entries[EntryIndex];
 			float StartX, EndX;
 			if (CalculateEntryGeometry(Entry.Get(), MyGeometry, StartX, EndX))
 			{
@@ -275,7 +275,7 @@ int32 SLogBar::GetEntryIndexAtTime(float Time) const
 	int32 EntryIndex = INDEX_NONE;
 	for( EntryIndex = 0; EntryIndex < Entries.Num(); ++EntryIndex )
 	{
-		TSharedPtr<FVisLogEntry> Entry = Entries[EntryIndex];
+		TSharedPtr<FVisualLogEntry> Entry = Entries[EntryIndex];
 		const float TimeDiff = FMath::Abs(Entry->TimeStamp - Time);
 
 		if (TimeDiff < BestTimeDiff)
@@ -346,7 +346,7 @@ FReply SLogBar::OnMouseMove( const FGeometry& MyGeometry, const FPointerEvent& M
 
 	for( int32 EntryIndex = 0; EntryIndex < Entries.Num(); EntryIndex++ )
 	{
-		TSharedPtr<FVisLogEntry> Entry = Entries[EntryIndex];
+		TSharedPtr<FVisualLogEntry> Entry = Entries[EntryIndex];
 
 		if( HoverX >= Entry->TimeStamp && Entry->TimeStamp <= HoverX + 1)
 		{
@@ -381,7 +381,7 @@ FVector2D SLogBar::ComputeDesiredSize() const
 	return FVector2D( 500.0f, 24.0f );
 }
 
-void SLogBar::SetEntries(const TArray<TSharedPtr<FVisLogEntry> >& InEntries, float InStartTime, float InTotalTime)
+void SLogBar::SetEntries(const TArray<TSharedPtr<FVisualLogEntry> >& InEntries, float InStartTime, float InTotalTime)
 {
 	Entries = InEntries;
 	StartTime = InStartTime;
@@ -391,7 +391,7 @@ void SLogBar::SetEntries(const TArray<TSharedPtr<FVisLogEntry> >& InEntries, flo
 	int32 EntryIndex = 0;
 	while (EntryIndex < Entries.Num())
 	{
-		TSharedPtr<FVisLogEntry> Entry = Entries[EntryIndex];
+		TSharedPtr<FVisualLogEntry> Entry = Entries[EntryIndex];
 		if (!Entry.IsValid())
 		{
 			EntryIndex++;
diff --git a/Engine/Source/Developer/LogVisualizer/Private/SLogBar.h b/Engine/Source/Developer/LogVisualizer/Private/SLogBar.h
--- a/Engine/Source/Developer/LogVisualizer/Private/SLogBar.h
+++ b/Engine/Source/Developer/LogVisualizer/Private/SLogBar.h
@@ -6,7 +6,7 @@
 
 struct FLogEntryBarProxy
 {
-	FLogEntryBarProxy(const FVisLogEntry* LogEntry)
+	FLogEntryBarProxy(const FVisualLogEntry* LogEntry)
 	{
 
 	}
@@ -24,7 +24,7 @@ public:
 
 	/** Delegate used allotted geometry changes */
 	DECLARE_DELEGATE_OneParam(FOnGeometryChanged, FGeometry);
-	DECLARE_DELEGATE_OneParam(FOnSelectionChanged, TSharedPtr<FVisLogEntry>);
+	DECLARE_DELEGATE_OneParam(FOnSelectionChanged, TSharedPtr<FVisualLogEntry>);
 	DECLARE_DELEGATE_RetVal(bool, FShouldDrawSelection);
 	DECLARE_DELEGATE_RetVal(float, FCurrentDisplayedTime);
 	DECLARE_DELEGATE_RetVal(int32, FCurrentEntryIndex);
@@ -66,7 +66,7 @@ public:
 	 *
 	 * @param InEvent	Events to draw
 	 */
-	void SetEntries(const TArray<TSharedPtr<FVisLogEntry> >& InEntries, float InStartTime, float InTotalTime);
+	void SetEntries(const TArray<TSharedPtr<FVisualLogEntry> >& InEntries, float InStartTime, float InTotalTime);
 	
 	void SelectEntry(const FGeometry& MyGeometry, const float ClickX);
 
@@ -119,7 +119,7 @@ private:
 
 	void SelectEntryAtIndex(const int32 Index);
 
-	FORCEINLINE_DEBUGGABLE bool CalculateEntryGeometry(const FVisLogEntry* LogEntry, const FGeometry& InGeometry, float& OutStartX, float& OutEndX) const
+	FORCEINLINE_DEBUGGABLE bool CalculateEntryGeometry(const FVisualLogEntry* LogEntry, const FGeometry& InGeometry, float& OutStartX, float& OutEndX) const
 	{
 		const float EventStart = (LogEntry->TimeStamp - StartTime) / TotalTime;
 		const float EventDuration = TimeUnit / TotalTime;
@@ -149,7 +149,7 @@ private:
 	const FSlateBrush* TimeMarkImage;
 
 	/** List of all events to draw */
-	TArray<TSharedPtr<FVisLogEntry> > Entries;
+	TArray<TSharedPtr<FVisualLogEntry> > Entries;
 
 	/** Start time (0.0 - 1.0) */
 	float StartTime;
diff --git a/Engine/Source/Developer/LogVisualizer/Private/SLogVisualizer.cpp b/Engine/Source/Developer/LogVisualizer/Private/SLogVisualizer.cpp
--- a/Engine/Source/Developer/LogVisualizer/Private/SLogVisualizer.cpp
+++ b/Engine/Source/Developer/LogVisualizer/Private/SLogVisualizer.cpp
@@ -648,7 +648,7 @@ void SLogVisualizer::OnListDoubleClick(TSharedPtr<FLogsListItem> LogListItem)
 #endif
 }
 
-int32 SLogVisualizer::GetCurrentVisibleLogEntryIndex(const TArray<TSharedPtr<FVisLogEntry> >& InVisibleEntries)
+int32 SLogVisualizer::GetCurrentVisibleLogEntryIndex(const TArray<TSharedPtr<FVisualLogEntry> >& InVisibleEntries)
 {
 	if (LogVisualizer->Logs.IsValidIndex(SelectedLogIndex))
 	{
@@ -709,7 +709,7 @@ void SLogVisualizer::UpdateVisibleEntriesCache(const TSharedPtr<FActorsVisLog>&
 	{
 		for (int32 EntryIndex = 0; EntryIndex < Log->Entries.Num(); ++EntryIndex)
 		{
-			const TSharedPtr<FVisLogEntry> CurrentEntry = Log->Entries[EntryIndex];
+			const TSharedPtr<FVisualLogEntry> CurrentEntry = Log->Entries[EntryIndex];
 			if (CurrentEntry->TimeStamp <= LastTimestamp)
 			{
 				continue;
@@ -720,7 +720,7 @@ void SLogVisualizer::UpdateVisibleEntriesCache(const TSharedPtr<FActorsVisLog>&
 			{
 				for (int32 LogLineIndex = 0; LogLineIndex < CurrentEntry->LogLines.Num(); ++LogLineIndex)
 				{
-					const FVisLogEntry::FLogLine &CurrentLine = CurrentEntry->LogLines[LogLineIndex];
+					const FVisualLogEntry::FLogLine &CurrentLine = CurrentEntry->LogLines[LogLineIndex];
 					const FString CurrentCategory = CurrentLine.Category.ToString();
 					if (FilterListPtr->IsFilterEnabled(CurrentCategory, CurrentLine.Verbosity))
 					{
@@ -739,7 +739,7 @@ void SLogVisualizer::UpdateVisibleEntriesCache(const TSharedPtr<FActorsVisLog>&
 			for (int32 ElementIndex = 0; ElementIndex < CurrentEntry->ElementsToDraw.Num(); ++ElementIndex)
 			{
 				const FString CurrentCategory = CurrentEntry->ElementsToDraw[ElementIndex].Category.ToString();
-				FVisLogEntry::FElementToDraw &CurrentElement = CurrentEntry->ElementsToDraw[ElementIndex];
+				FVisualLogEntry::FElementToDraw &CurrentElement = CurrentEntry->ElementsToDraw[ElementIndex];
 				if (CurrentElement.Category != NAME_None && (FilterListPtr->IsFilterEnabled(CurrentCategory, CurrentElement.Verbosity)))
 				{
 					OutEntriesCached[CachedLogIndex].CachedEntries.AddUnique(CurrentEntry);
@@ -754,7 +754,7 @@ void SLogVisualizer::UpdateVisibleEntriesCache(const TSharedPtr<FActorsVisLog>&
 
 			for (int32 SampleIndex = 0; SampleIndex < CurrentEntry->HistogramSamples.Num(); ++SampleIndex)
 			{
-				const FVisLogEntry::FHistogramSample &CurrentSample = CurrentEntry->HistogramSamples[SampleIndex];
+				const FVisualLogEntry::FHistogramSample &CurrentSample = CurrentEntry->HistogramSamples[SampleIndex];
 				const FName CurrentCategory = CurrentSample.Category;
 				const FName CurrentDataName = CurrentSample.DataName;
 
@@ -796,7 +796,7 @@ void SLogVisualizer::UpdateVisibleEntriesCache(const TSharedPtr<FActorsVisLog>&
 	LogsList[CurrentIndex]->LastEndTimestamp = LogsList[CurrentIndex]->EndTimestamp;
 }
 
-void SLogVisualizer::GetVisibleEntries(const TSharedPtr<FActorsVisLog>& Log, TArray<TSharedPtr<FVisLogEntry> >& OutEntries)
+void SLogVisualizer::GetVisibleEntries(const TSharedPtr<FActorsVisLog>& Log, TArray<TSharedPtr<FVisualLogEntry> >& OutEntries)
 {
 	const int32 NumLogs = OutEntriesCached.Num();
 	int32 CachedLogIndex = INDEX_NONE;
@@ -826,12 +826,12 @@ void SLogVisualizer::GetVisibleEntries(const TSharedPtr<FActorsVisLog>& Log, TAr
 			{
 				bool bAddedEntry = false;
 
-				const TSharedPtr<FVisLogEntry> CurrentEntry = Log->Entries[i];
+				const TSharedPtr<FVisualLogEntry> CurrentEntry = Log->Entries[i];
 				if (!bAddedEntry)
 				{
 					for (int32 j = 0; j < CurrentEntry->LogLines.Num(); ++j)
 					{
-						const FVisLogEntry::FLogLine &CurrentLine = CurrentEntry->LogLines[j];
+						const FVisualLogEntry::FLogLine &CurrentLine = CurrentEntry->LogLines[j];
 						const FString CurrentCategory = CurrentLine.Category.ToString();
 						if (FilterListPtr->IsFilterEnabled(CurrentCategory, CurrentLine.Verbosity))
 						{
@@ -850,7 +850,7 @@ void SLogVisualizer::GetVisibleEntries(const TSharedPtr<FActorsVisLog>& Log, TAr
 				for (int32 j = 0; j < CurrentEntry->ElementsToDraw.Num(); ++j)
 				{
 					const FString CurrentCategory = CurrentEntry->ElementsToDraw[j].Category.ToString();
-					FVisLogEntry::FElementToDraw &CurrentElement = CurrentEntry->ElementsToDraw[j];
+					FVisualLogEntry::FElementToDraw &CurrentElement = CurrentEntry->ElementsToDraw[j];
 					if (CurrentElement.Category != NAME_None && (FilterListPtr->IsFilterEnabled(CurrentCategory, CurrentElement.Verbosity)))
 					{
 						OutEntriesCached[CachedLogIndex].CachedEntries.AddUnique(CurrentEntry);
@@ -865,7 +865,7 @@ void SLogVisualizer::GetVisibleEntries(const TSharedPtr<FActorsVisLog>& Log, TAr
 
 				for (int32 SampleIndex = 0; SampleIndex < CurrentEntry->HistogramSamples.Num(); ++SampleIndex)
 				{
-					const FVisLogEntry::FHistogramSample &CurrentSample = CurrentEntry->HistogramSamples[SampleIndex];
+					const FVisualLogEntry::FHistogramSample &CurrentSample = CurrentEntry->HistogramSamples[SampleIndex];
 					const FName CurrentCategory = CurrentSample.Category;
 					const FName CurrentDataName = CurrentSample.DataName;
 
@@ -913,7 +913,7 @@ void SLogVisualizer::GetVisibleEntries(const TSharedPtr<FActorsVisLog>& Log, TAr
 		const int32 NumEntries = OutEntriesCached[CachedLogIndex].CachedEntries.Num();
 		for (int32 Index = 0; Index < NumEntries; ++Index)
 		{
-			TSharedPtr<FVisLogEntry> LogEntry = OutEntriesCached[CachedLogIndex].CachedEntries[Index];
+			TSharedPtr<FVisualLogEntry> LogEntry = OutEntriesCached[CachedLogIndex].CachedEntries[Index];
 			if (!LogEntry.IsValid())
 			{
 				continue;
@@ -1514,7 +1514,7 @@ TSharedRef<ITableRow> SLogVisualizer::LogEntryLinesGenerateRow(TSharedPtr<FLogEn
 
 void SLogVisualizer::LogEntryLineSelectionChanged(TSharedPtr<FLogEntryItem> SelectedItem, ESelectInfo::Type SelectInfo)
 {
-	TMap<FName, FVisualLogExtensionInterface*>& AllExtensions = FVisualLog::Get().GetAllExtensions();
+	TMap<FName, FVisualLogExtensionInterface*>& AllExtensions = FVisualLogger::Get().GetAllExtensions();
 	for (auto Iterator = AllExtensions.CreateIterator(); Iterator; ++Iterator)
 	{
 		FVisualLogExtensionInterface* Extension = (*Iterator).Value;
@@ -1547,7 +1547,7 @@ bool SLogVisualizer::ShouldListLog(const TSharedPtr<FActorsVisLog>& Log)
 	}
 	else
 	{
-		static TArray<TSharedPtr<FVisLogEntry> > OutEntries;
+		static TArray<TSharedPtr<FVisualLogEntry> > OutEntries;
 		GetVisibleEntries(Log, OutEntries);
 		return !bIgnoreTrivialLogs || OutEntries.Num() >= 2;
 	}
@@ -1590,12 +1590,12 @@ void SLogVisualizer::InvalidateCanvas()
 #endif
 }
 
-void SLogVisualizer::RequestShowLogEntry(TSharedPtr<FLogsListItem> Item, TSharedPtr<FVisLogEntry> LogEntry)
+void SLogVisualizer::RequestShowLogEntry(TSharedPtr<FLogsListItem> Item, TSharedPtr<FVisualLogEntry> LogEntry)
 {
 	ShowLogEntry(Item, LogEntry);
 }
 
-void SLogVisualizer::ShowLogEntry(TSharedPtr<FLogsListItem> Item, TSharedPtr<FVisLogEntry> LogEntry)
+void SLogVisualizer::ShowLogEntry(TSharedPtr<FLogsListItem> Item, TSharedPtr<FVisualLogEntry> LogEntry)
 {
 	if(LogsListWidget->GetSelectedItems().Find(Item) == INDEX_NONE)
 	{
@@ -1671,7 +1671,7 @@ void SLogVisualizer::OnLogStatusGetChildren(TSharedPtr<FLogStatusItem> InItem, T
 	OutItems = InItem->Children;
 }
 
-void SLogVisualizer::UpdateStatusItems(const FVisLogEntry* LogEntry)
+void SLogVisualizer::UpdateStatusItems(const FVisualLogEntry* LogEntry)
 {
 	TArray<FString> ExpandedCategories;
 	for (int32 ItemIndex = 0; ItemIndex < StatusItems.Num(); ItemIndex++)
@@ -1727,12 +1727,12 @@ void SLogVisualizer::UpdateStatusItems(const FVisLogEntry* LogEntry)
 	}
 }
 
-void SLogVisualizer::ShowEntry(const FVisLogEntry* LogEntry)
+void SLogVisualizer::ShowEntry(const FVisualLogEntry* LogEntry)
 {
 	UpdateStatusItems(LogEntry);
 	LogEntryLines.Reset();
 	
-	const FVisLogEntry::FLogLine* LogLine = LogEntry->LogLines.GetData();
+	const FVisualLogEntry::FLogLine* LogLine = LogEntry->LogLines.GetData();
 	for (int LineIndex = 0; LineIndex < LogEntry->LogLines.Num(); ++LineIndex, ++LogLine)
 	{
 		bool bShowLine = true;
@@ -1768,7 +1768,7 @@ void SLogVisualizer::ShowEntry(const FVisLogEntry* LogEntry)
 	SetCurrentViewedTime(LogEntry->TimeStamp);
 
 	UWorld* World = GetWorld();
-	for (auto It = FVisualLog::Get().GetAllExtensions().CreateIterator(); It; ++It)
+	for (auto It = FVisualLogger::Get().GetAllExtensions().CreateIterator(); It; ++It)
 	{
 		(*It).Value->OnTimestampChange(LogEntry->TimeStamp, World, FLogVisualizerModule::Get()->GetHelperActor(World));
 	}
@@ -1993,11 +1993,11 @@ void SLogVisualizer::DrawOnCanvas(UCanvas* Canvas, APlayerController*)
 	if (World != NULL && LogVisualizer->Logs.IsValidIndex(SelectedLogIndex))
 	{
 		TSharedPtr<FActorsVisLog> Log = LogVisualizer->Logs[SelectedLogIndex];
-		const TArray<TSharedPtr<FVisLogEntry> >& Entries = Log->Entries;
+		const TArray<TSharedPtr<FVisualLogEntry> >& Entries = Log->Entries;
 
 		if (bDrawLogEntriesPath)
 		{
-			const TSharedPtr<FVisLogEntry>* Entry = Entries.GetData();
+			const TSharedPtr<FVisualLogEntry>* Entry = Entries.GetData();
 			FVector Location = (*Entry)->Location;
 			++Entry;
 
@@ -2012,7 +2012,7 @@ void SLogVisualizer::DrawOnCanvas(UCanvas* Canvas, APlayerController*)
 		if (Entries.IsValidIndex(LogEntryIndex))
 		{
 			// draw all additional data stored in current entry
-			const TSharedPtr<FVisLogEntry>& Entry = Entries[LogEntryIndex];
+			const TSharedPtr<FVisualLogEntry>& Entry = Entries[LogEntryIndex];
 
 			// mark current location
 			DrawDebugCone(World, Entry->Location, /*Direction*/FVector(0, 0, 1), /*Length*/200.f
@@ -2057,7 +2057,7 @@ void SLogVisualizer::DrawOnCanvas(UCanvas* Canvas, APlayerController*)
 			int32 ColorIndex = 0;
 			for (int32 EntryIndex = 0; EntryIndex < Entries.Num(); ++EntryIndex)
 			{
-				const TSharedPtr<FVisLogEntry>& CurrentEntry = Entries[EntryIndex];
+				const TSharedPtr<FVisualLogEntry>& CurrentEntry = Entries[EntryIndex];
 				if (HistogramPreviewWindow <= 0)
 				{
 					if (CurrentEntry->TimeStamp > Entry->TimeStamp)
@@ -2080,7 +2080,7 @@ void SLogVisualizer::DrawOnCanvas(UCanvas* Canvas, APlayerController*)
 				const int32 SamplesNum = CurrentEntry->HistogramSamples.Num();
 				for (int32 SampleIndex = 0; SampleIndex < SamplesNum; ++SampleIndex)
 				{
-					FVisLogEntry::FHistogramSample CurrentSample = CurrentEntry->HistogramSamples[SampleIndex];
+					FVisualLogEntry::FHistogramSample CurrentSample = CurrentEntry->HistogramSamples[SampleIndex];
 
 					const FName CurrentCategory = CurrentSample.Category;
 					const FName CurrentGraphName = CurrentSample.GraphName;
@@ -2210,7 +2210,7 @@ void SLogVisualizer::DrawOnCanvas(UCanvas* Canvas, APlayerController*)
 			{
 				const FName TagName = CurrentData.TagName;
 				const bool bIsValidByFilter = FilterListPtr->IsFilterEnabled(CurrentData.Category.ToString(), ELogVerbosity::All) && FilterListPtr->IsFilterEnabled(CurrentData.TagName.ToString(), ELogVerbosity::All);
-				FVisualLogExtensionInterface* Extension = FVisualLog::Get().GetExtensionForTag(TagName);
+				FVisualLogExtensionInterface* Extension = FVisualLogger::Get().GetExtensionForTag(TagName);
 				if (!Extension)
 				{
 					continue;
@@ -2226,7 +2226,7 @@ void SLogVisualizer::DrawOnCanvas(UCanvas* Canvas, APlayerController*)
 				}
 			}
 
-			const FVisLogEntry::FElementToDraw* ElementToDraw = Entry->ElementsToDraw.GetData();
+			const FVisualLogEntry::FElementToDraw* ElementToDraw = Entry->ElementsToDraw.GetData();
 			const int32 ElementsCount = Entry->ElementsToDraw.Num();
 			
 			for (int32 ElementIndex = 0; ElementIndex < ElementsCount; ++ElementIndex, ++ElementToDraw)
@@ -2241,7 +2241,7 @@ void SLogVisualizer::DrawOnCanvas(UCanvas* Canvas, APlayerController*)
 
 				switch(ElementToDraw->GetType())
 				{
-				case FVisLogEntry::FElementToDraw::SinglePoint:
+				case FVisualLogEntry::FElementToDraw::SinglePoint:
 					{
 						const float Radius = float(ElementToDraw->Radius);
 						const bool bDrawLabel = ElementToDraw->Description.IsEmpty() == false;
@@ -2257,7 +2257,7 @@ void SLogVisualizer::DrawOnCanvas(UCanvas* Canvas, APlayerController*)
 						}
 					}
 					break;
-				case FVisLogEntry::FElementToDraw::Segment:
+				case FVisualLogEntry::FElementToDraw::Segment:
 					{
 						const float Thickness = float(ElementToDraw->Thicknes);
 						const bool bDrawLabel = ElementToDraw->Description.IsEmpty() == false && ElementToDraw->Points.Num() > 2;
@@ -2287,7 +2287,7 @@ void SLogVisualizer::DrawOnCanvas(UCanvas* Canvas, APlayerController*)
 						}
 					}
 					break;
-				case FVisLogEntry::FElementToDraw::Path:
+				case FVisualLogEntry::FElementToDraw::Path:
 					{
 						const float Thickness = float(ElementToDraw->Thicknes);
 						FVector Location = ElementToDraw->Points[0];
@@ -2301,7 +2301,7 @@ void SLogVisualizer::DrawOnCanvas(UCanvas* Canvas, APlayerController*)
 						}
 					}
 					break;
-				case FVisLogEntry::FElementToDraw::Box:
+				case FVisualLogEntry::FElementToDraw::Box:
 					{
 						const float Thickness = float(ElementToDraw->Thicknes);
 						const bool bDrawLabel = ElementToDraw->Description.IsEmpty() == false && ElementToDraw->Points.Num() > 2;
@@ -2332,7 +2332,7 @@ void SLogVisualizer::DrawOnCanvas(UCanvas* Canvas, APlayerController*)
 						}
 					}
 					break;
-				case FVisLogEntry::FElementToDraw::Cone:
+				case FVisualLogEntry::FElementToDraw::Cone:
 					{
 						const float Thickness = float(ElementToDraw->Thicknes);
 						const bool bDrawLabel = ElementToDraw->Description.IsEmpty() == false;
@@ -2350,7 +2350,7 @@ void SLogVisualizer::DrawOnCanvas(UCanvas* Canvas, APlayerController*)
 						}
 					}
 					break;
-				case FVisLogEntry::FElementToDraw::Cylinder:
+				case FVisualLogEntry::FElementToDraw::Cylinder:
 					{
 						const float Thickness = float(ElementToDraw->Thicknes);
 						const bool bDrawLabel = ElementToDraw->Description.IsEmpty() == false;
@@ -2368,7 +2368,7 @@ void SLogVisualizer::DrawOnCanvas(UCanvas* Canvas, APlayerController*)
 						}
 					}
 					break;
-				case FVisLogEntry::FElementToDraw::Capsule:
+				case FVisualLogEntry::FElementToDraw::Capsule:
 					{
 						const float Thickness = float(ElementToDraw->Thicknes);
 						const bool bDrawLabel = ElementToDraw->Description.IsEmpty() == false;
@@ -2600,7 +2600,7 @@ void SLogVisualizer::SelectActor(AActor* SelectedActor)
 	{
 		return;
 	}
-	const AActor* LogOwner = FVisualLog::Get().GetVisualLogRedirection(SelectedActor);
+	const AActor* LogOwner = Cast<AActor>(FVisualLogger::Get().FindRedirection(SelectedActor));
 	const int32 LogIndex = LogVisualizer->GetLogIndexForActor(LogOwner);
 	if (LogVisualizer->Logs.IsValidIndex(LogIndex))
 	{
diff --git a/Engine/Source/Developer/LogVisualizer/Private/SLogVisualizer.h b/Engine/Source/Developer/LogVisualizer/Private/SLogVisualizer.h
--- a/Engine/Source/Developer/LogVisualizer/Private/SLogVisualizer.h
+++ b/Engine/Source/Developer/LogVisualizer/Private/SLogVisualizer.h
@@ -129,15 +129,15 @@ public:
 
 	void SetCurrentViewedTime(float NewTime, const bool bForce = false);
 
-	void RequestShowLogEntry(TSharedPtr<FLogsListItem> Item, TSharedPtr<FVisLogEntry> LogEntry);
+	void RequestShowLogEntry(TSharedPtr<FLogsListItem> Item, TSharedPtr<FVisualLogEntry> LogEntry);
 
 	void UpdateVisibleEntriesCache(const TSharedPtr<FActorsVisLog>& Log, int32 Index);
-	void GetVisibleEntries(const TSharedPtr<FActorsVisLog>& Log, TArray<TSharedPtr<FVisLogEntry> >& OutEntries);
+	void GetVisibleEntries(const TSharedPtr<FActorsVisLog>& Log, TArray<TSharedPtr<FVisualLogEntry> >& OutEntries);
 
-	int32 GetCurrentVisibleLogEntryIndex(const TArray<TSharedPtr<FVisLogEntry> >& InVisibleEntries);
+	int32 GetCurrentVisibleLogEntryIndex(const TArray<TSharedPtr<FVisualLogEntry> >& InVisibleEntries);
 
 protected:
-	void ShowLogEntry(TSharedPtr<FLogsListItem> Item, TSharedPtr<FVisLogEntry> LogEntry);
+	void ShowLogEntry(TSharedPtr<FLogsListItem> Item, TSharedPtr<FVisualLogEntry> LogEntry);
 
 	void SelectionChanged(class AActor* DebuggedActor, bool bIsBeingDebuggedNow);
 public:	
@@ -174,7 +174,7 @@ protected:
 
 private:
 	void InvalidateCanvas();
-	void UpdateStatusItems(const FVisLogEntry* LogEntry);
+	void UpdateStatusItems(const FVisualLogEntry* LogEntry);
 	TSharedRef<ITableRow> HandleGenerateLogStatus(TSharedPtr<FLogStatusItem> InItem, const TSharedRef<STableViewBase>& OwnerTable);
 	void OnLogStatusGetChildren(TSharedPtr<FLogStatusItem> InItem, TArray< TSharedPtr<FLogStatusItem> >& OutItems);
 
@@ -237,7 +237,7 @@ private:
 	/** See if a particular log passes the current filter and other display flags */
 	bool ShouldListLog(const TSharedPtr<FActorsVisLog>& Log);
 
-	void ShowEntry(const FVisLogEntry* LogEntry);
+	void ShowEntry(const FVisualLogEntry* LogEntry);
 
 	int32 FindIndexInLogsList(const int32 LogIndex) const;
 
@@ -303,7 +303,7 @@ private:
 	struct FCachedEntries
 	{
 		TSharedPtr<FActorsVisLog> Log;
-		TArray<TSharedPtr<FVisLogEntry> > CachedEntries;
+		TArray<TSharedPtr<FVisualLogEntry> > CachedEntries;
 	};
 	TArray<FCachedEntries> OutEntriesCached;
 		
diff --git a/Engine/Source/Developer/LogVisualizer/Private/SLogsTableRow.cpp b/Engine/Source/Developer/LogVisualizer/Private/SLogsTableRow.cpp
--- a/Engine/Source/Developer/LogVisualizer/Private/SLogsTableRow.cpp
+++ b/Engine/Source/Developer/LogVisualizer/Private/SLogsTableRow.cpp
@@ -114,7 +114,7 @@ bool SLogsTableRow::ShouldDrawSelection()
 	return OwnerTable.IsValid() && OwnerTable->Private_IsItemSelected(Item);
 }
 
-void SLogsTableRow::OnBarGraphSelectionChanged( TSharedPtr<FVisLogEntry> Selection )
+void SLogsTableRow::OnBarGraphSelectionChanged( TSharedPtr<FVisualLogEntry> Selection )
 {
 	if( Selection.IsValid() )
 	{
diff --git a/Engine/Source/Developer/LogVisualizer/Private/SLogsTableRow.h b/Engine/Source/Developer/LogVisualizer/Private/SLogsTableRow.h
--- a/Engine/Source/Developer/LogVisualizer/Private/SLogsTableRow.h
+++ b/Engine/Source/Developer/LogVisualizer/Private/SLogsTableRow.h
@@ -28,12 +28,12 @@ public:
 private:
 	int32 GetCurrentLogEntryIndex() const;
 
-	void OnBarGraphSelectionChanged(TSharedPtr<FVisLogEntry> Selection);
+	void OnBarGraphSelectionChanged(TSharedPtr<FVisualLogEntry> Selection);
 	void OnBarGeometryChanged(FGeometry Geometry);
 
 	bool ShouldDrawSelection();
 
-	TArray<TSharedPtr<FVisLogEntry> > VisibleEntries;
+	TArray<TSharedPtr<FVisualLogEntry> > VisibleEntries;
 
 	/** Tree item */
 	TSharedPtr<FLogsListItem> Item;
diff --git a/Engine/Source/Runtime/AIModule/Classes/AIController.h b/Engine/Source/Runtime/AIModule/Classes/AIController.h
--- a/Engine/Source/Runtime/AIModule/Classes/AIController.h
+++ b/Engine/Source/Runtime/AIModule/Classes/AIController.h
@@ -253,7 +253,7 @@ public:
 	virtual void DisplayDebug(class UCanvas* Canvas, const FDebugDisplayInfo& DebugDisplay, float& YL, float& YPos) override;
 
 #if ENABLE_VISUAL_LOG
-	virtual void GrabDebugSnapshot(struct FVisLogEntry* Snapshot) const override;
+	virtual void GrabDebugSnapshot(struct FVisualLogEntry* Snapshot) const override;
 #endif
 
 	virtual void Reset() override;
diff --git a/Engine/Source/Runtime/AIModule/Classes/Actions/PawnActionsComponent.h b/Engine/Source/Runtime/AIModule/Classes/Actions/PawnActionsComponent.h
--- a/Engine/Source/Runtime/AIModule/Classes/Actions/PawnActionsComponent.h
+++ b/Engine/Source/Runtime/AIModule/Classes/Actions/PawnActionsComponent.h
@@ -137,7 +137,7 @@ public:
 	FORCEINLINE UPawnAction* GetActiveAction(EAIRequestPriority::Type Priority) { return ActionStacks[Priority].GetTop(); }
 
 #if ENABLE_VISUAL_LOG
-	void DescribeSelfToVisLog(struct FVisLogEntry* Snapshot) const;
+	void DescribeSelfToVisLog(struct FVisualLogEntry* Snapshot) const;
 #endif // ENABLE_VISUAL_LOG
 
 	//----------------------------------------------------------------------//
diff --git a/Engine/Source/Runtime/AIModule/Classes/BehaviorTree/BehaviorTreeComponent.h b/Engine/Source/Runtime/AIModule/Classes/BehaviorTree/BehaviorTreeComponent.h
--- a/Engine/Source/Runtime/AIModule/Classes/BehaviorTree/BehaviorTreeComponent.h
+++ b/Engine/Source/Runtime/AIModule/Classes/BehaviorTree/BehaviorTreeComponent.h
@@ -174,7 +174,7 @@ public:
 	virtual FString DescribeActiveTrees() const;
 
 #if ENABLE_VISUAL_LOG
-	virtual void DescribeSelfToVisLog(struct FVisLogEntry* Snapshot) const override;
+	virtual void DescribeSelfToVisLog(struct FVisualLogEntry* Snapshot) const override;
 #endif
 
 protected:
diff --git a/Engine/Source/Runtime/AIModule/Classes/BehaviorTree/BlackboardComponent.h b/Engine/Source/Runtime/AIModule/Classes/BehaviorTree/BlackboardComponent.h
--- a/Engine/Source/Runtime/AIModule/Classes/BehaviorTree/BlackboardComponent.h
+++ b/Engine/Source/Runtime/AIModule/Classes/BehaviorTree/BlackboardComponent.h
@@ -205,7 +205,7 @@ class AIMODULE_API UBlackboardComponent : public UActorComponent
 
 #if ENABLE_VISUAL_LOG
 	/** prepare blackboard snapshot for logs */
-	virtual void DescribeSelfToVisLog(struct FVisLogEntry* Snapshot) const;
+	virtual void DescribeSelfToVisLog(struct FVisualLogEntry* Snapshot) const;
 #endif
 	
 protected:
diff --git a/Engine/Source/Runtime/AIModule/Classes/BrainComponent.h b/Engine/Source/Runtime/AIModule/Classes/BrainComponent.h
--- a/Engine/Source/Runtime/AIModule/Classes/BrainComponent.h
+++ b/Engine/Source/Runtime/AIModule/Classes/BrainComponent.h
@@ -137,7 +137,7 @@ public:
 	virtual bool IsPaused() const { return false; }
 
 #if ENABLE_VISUAL_LOG
-	virtual void DescribeSelfToVisLog(struct FVisLogEntry* Snapshot) const;
+	virtual void DescribeSelfToVisLog(struct FVisualLogEntry* Snapshot) const;
 #endif // ENABLE_VISUAL_LOG
 	
 	// IAIResourceInterface begin
diff --git a/Engine/Source/Runtime/AIModule/Classes/EnvironmentQuery/EnvQueryDebugHelpers.h b/Engine/Source/Runtime/AIModule/Classes/EnvironmentQuery/EnvQueryDebugHelpers.h
--- a/Engine/Source/Runtime/AIModule/Classes/EnvironmentQuery/EnvQueryDebugHelpers.h
+++ b/Engine/Source/Runtime/AIModule/Classes/EnvironmentQuery/EnvQueryDebugHelpers.h
@@ -141,14 +141,8 @@ FArchive& operator<<(FArchive& Ar, EQSDebug::FQueryData& Data)
 
 #define UE_VLOG_EQS(Query, CategoryName, Verbosity) \
 { \
-	SCOPE_CYCLE_COUNTER(STAT_VisualLog); \
-	static_assert((ELogVerbosity::Verbosity & ELogVerbosity::VerbosityMask) < ELogVerbosity::NumVerbosity && ELogVerbosity::Verbosity > 0, "Verbosity must be constant and in range."); \
-	if (FVisualLog::Get().IsRecording() && (!FVisualLog::Get().IsAllBlocked() || FVisualLog::Get().InWhitelist(CategoryName.GetCategoryName()))) \
-	{ \
-		const AActor* OwnerActor = FVisualLog::Get().GetVisualLogRedirection(Query->Owner.Get()); \
-			UEnvQueryDebugHelpers::LogQuery(OwnerActor ? OwnerActor : Query->Owner.Get(), Query, CategoryName.GetCategoryName(), ELogVerbosity::Verbosity, \
-			FString::Printf(TEXT("Executed EQS: \n - Name: '%s' (id=%d, option=%d),\n - All Items: %d,\n - ValidItems: %d"), *Query->QueryName, Query->QueryID, Query->OptionIndex, Query->ItemDetails.Num(), Query->NumValidItems)); \
-	} \
+	UEnvQueryDebugHelpers::LogQuery(Query->Owner.Get(), Query, CategoryName.GetCategoryName(), ELogVerbosity::Verbosity, \
+		FString::Printf(TEXT("Executed EQS: \n - Name: '%s' (id=%d, option=%d),\n - All Items: %d,\n - ValidItems: %d"), *Query->QueryName, Query->QueryID, Query->OptionIndex, Query->ItemDetails.Num(), Query->NumValidItems)); \
 }
 
 #else
diff --git a/Engine/Source/Runtime/AIModule/Classes/Navigation/CrowdFollowingComponent.h b/Engine/Source/Runtime/AIModule/Classes/Navigation/CrowdFollowingComponent.h
--- a/Engine/Source/Runtime/AIModule/Classes/Navigation/CrowdFollowingComponent.h
+++ b/Engine/Source/Runtime/AIModule/Classes/Navigation/CrowdFollowingComponent.h
@@ -90,7 +90,7 @@ class AIMODULE_API UCrowdFollowingComponent : public UPathFollowingComponent, pu
 
 	virtual void GetDebugStringTokens(TArray<FString>& Tokens, TArray<EPathFollowingDebugTokens::Type>& Flags) const;
 #if ENABLE_VISUAL_LOG
-	virtual void DescribeSelfToVisLog(struct FVisLogEntry* Snapshot) const;
+	virtual void DescribeSelfToVisLog(struct FVisualLogEntry* Snapshot) const;
 #endif // ENABLE_VISUAL_LOG
 
 protected:
diff --git a/Engine/Source/Runtime/AIModule/Classes/Navigation/NavigationComponent.h b/Engine/Source/Runtime/AIModule/Classes/Navigation/NavigationComponent.h
--- a/Engine/Source/Runtime/AIModule/Classes/Navigation/NavigationComponent.h
+++ b/Engine/Source/Runtime/AIModule/Classes/Navigation/NavigationComponent.h
@@ -190,7 +190,7 @@ public:
 	//----------------------------------------------------------------------//
 
 #if ENABLE_VISUAL_LOG
-	virtual void DescribeSelfToVisLog(struct FVisLogEntry* Snapshot) const;
+	virtual void DescribeSelfToVisLog(struct FVisualLogEntry* Snapshot) const;
 #endif // ENABLE_VISUAL_LOG
 
 	/** quick and dirty debug render. If Canvas not NULL additional data (if present) will be drawn */
diff --git a/Engine/Source/Runtime/AIModule/Classes/Navigation/PathFollowingComponent.h b/Engine/Source/Runtime/AIModule/Classes/Navigation/PathFollowingComponent.h
--- a/Engine/Source/Runtime/AIModule/Classes/Navigation/PathFollowingComponent.h
+++ b/Engine/Source/Runtime/AIModule/Classes/Navigation/PathFollowingComponent.h
@@ -230,7 +230,7 @@ class AIMODULE_API UPathFollowingComponent : public UActorComponent, public IAIR
 
 	virtual void DisplayDebug(class UCanvas* Canvas, const FDebugDisplayInfo& DebugDisplay, float& YL, float& YPos) const;
 #if ENABLE_VISUAL_LOG
-	virtual void DescribeSelfToVisLog(struct FVisLogEntry* Snapshot) const;
+	virtual void DescribeSelfToVisLog(struct FVisualLogEntry* Snapshot) const;
 #endif // ENABLE_VISUAL_LOG
 
 	/** called when moving agent collides with another actor */
diff --git a/Engine/Source/Runtime/AIModule/Classes/Perception/AIPerceptionComponent.h b/Engine/Source/Runtime/AIModule/Classes/Perception/AIPerceptionComponent.h
--- a/Engine/Source/Runtime/AIModule/Classes/Perception/AIPerceptionComponent.h
+++ b/Engine/Source/Runtime/AIModule/Classes/Perception/AIPerceptionComponent.h
@@ -8,7 +8,7 @@
 #include "AIPerceptionComponent.generated.h"
 
 class AAIController;
-struct FVisLogEntry;
+struct FVisualLogEntry;
 class UCanvas;
 class UAIPerceptionSystem;
 
@@ -206,7 +206,7 @@ public:
 	void DrawDebugInfo(UCanvas* Canvas);
 
 #if ENABLE_VISUAL_LOG
-	virtual void DescribeSelfToVisLog(FVisLogEntry* Snapshot) const;
+	virtual void DescribeSelfToVisLog(FVisualLogEntry* Snapshot) const;
 #endif // ENABLE_VISUAL_LOG
 	
 protected:
diff --git a/Engine/Source/Runtime/AIModule/Classes/VisualLoggerExtension.h b/Engine/Source/Runtime/AIModule/Classes/VisualLoggerExtension.h
--- a/Engine/Source/Runtime/AIModule/Classes/VisualLoggerExtension.h
+++ b/Engine/Source/Runtime/AIModule/Classes/VisualLoggerExtension.h
@@ -16,8 +16,8 @@ class FVisualLoggerExtension : public FVisualLogExtensionInterface
 public:
 	FVisualLoggerExtension();
 	virtual void OnTimestampChange(float Timestamp, class UWorld* InWorld, class AActor* HelperActor) override;
-	virtual void DrawData(class UWorld* InWorld, class UCanvas* Canvas, class AActor* HelperActor, const FName& TagName, const struct FVisLogEntry::FDataBlock& DataBlock, float Timestamp) override;
-	virtual void DisableDrawingForData(class UWorld* InWorld, class UCanvas* Canvas, class AActor* HelperActor, const FName& TagName, const FVisLogEntry::FDataBlock& DataBlock, float Timestamp) override;
+	virtual void DrawData(class UWorld* InWorld, class UCanvas* Canvas, class AActor* HelperActor, const FName& TagName, const struct FVisualLogEntry::FDataBlock& DataBlock, float Timestamp) override;
+	virtual void DisableDrawingForData(class UWorld* InWorld, class UCanvas* Canvas, class AActor* HelperActor, const FName& TagName, const FVisualLogEntry::FDataBlock& DataBlock, float Timestamp) override;
 	virtual void LogEntryLineSelectionChanged(TSharedPtr<struct FLogEntryItem> SelectedItem, int64 UserData, FName TagName) override;
 
 private:
diff --git a/Engine/Source/Runtime/AIModule/Private/AIController.cpp b/Engine/Source/Runtime/AIModule/Private/AIController.cpp
--- a/Engine/Source/Runtime/AIModule/Private/AIController.cpp
+++ b/Engine/Source/Runtime/AIModule/Private/AIController.cpp
@@ -109,9 +109,9 @@ void AAIController::DisplayDebug(UCanvas* Canvas, const FDebugDisplayInfo& Debug
 
 #if ENABLE_VISUAL_LOG
 
-void AAIController::GrabDebugSnapshot(FVisLogEntry* Snapshot) const
+void AAIController::GrabDebugSnapshot(FVisualLogEntry* Snapshot) const
 {
-	FVisLogEntry::FStatusCategory MyCategory;
+	FVisualLogEntry::FStatusCategory MyCategory;
 	MyCategory.Category = TEXT("AI Controller");
 	MyCategory.Add(TEXT("Pawn"), GetNameSafe(GetPawn()));
 	AActor* FocusActor = GetFocusActor();
@@ -145,7 +145,7 @@ void AAIController::GrabDebugSnapshot(FVisLogEntry* Snapshot) const
 	if (PerceptionComponent != NULL)
 	{
 		PerceptionComponent->DescribeSelfToVisLog(Snapshot);
-}
+	}
 }
 #endif // ENABLE_VISUAL_LOG
 
diff --git a/Engine/Source/Runtime/AIModule/Private/AIModule.cpp b/Engine/Source/Runtime/AIModule/Private/AIModule.cpp
--- a/Engine/Source/Runtime/AIModule/Private/AIModule.cpp
+++ b/Engine/Source/Runtime/AIModule/Private/AIModule.cpp
@@ -36,7 +36,7 @@ void FAIModule::StartupModule()
 #endif // WITH_EDITOR 
 
 #if WITH_EDITOR && ENABLE_VISUAL_LOG
-	FVisualLog::Get().RegisterExtension(*EVisLogTags::TAG_EQS, &VisualLoggerExtension);
+	FVisualLogger::Get().RegisterExtension(*EVisLogTags::TAG_EQS, &VisualLoggerExtension);
 #endif
 }
 
@@ -45,7 +45,7 @@ void FAIModule::ShutdownModule()
 	// This function may be called during shutdown to clean up your module.  For modules that support dynamic reloading,
 	// we call this function before unloading the module.
 #if WITH_EDITOR && ENABLE_VISUAL_LOG
-	FVisualLog::Get().UnregisterExtension(*EVisLogTags::TAG_EQS, &VisualLoggerExtension);
+	FVisualLogger::Get().UnregisterExtension(*EVisLogTags::TAG_EQS, &VisualLoggerExtension);
 #endif
 }
 
diff --git a/Engine/Source/Runtime/AIModule/Private/Actions/PawnActionsComponent.cpp b/Engine/Source/Runtime/AIModule/Private/Actions/PawnActionsComponent.cpp
--- a/Engine/Source/Runtime/AIModule/Private/Actions/PawnActionsComponent.cpp
+++ b/Engine/Source/Runtime/AIModule/Private/Actions/PawnActionsComponent.cpp
@@ -504,7 +504,7 @@ bool UPawnActionsComponent::PerformAction(APawn& Pawn, UPawnAction& Action, TEnu
 // debug
 //----------------------------------------------------------------------//
 #if ENABLE_VISUAL_LOG
-void UPawnActionsComponent::DescribeSelfToVisLog(FVisLogEntry* Snapshot) const
+void UPawnActionsComponent::DescribeSelfToVisLog(FVisualLogEntry* Snapshot) const
 {
 	static const FString Category = TEXT("PawnActions");
 
@@ -521,7 +521,7 @@ void UPawnActionsComponent::DescribeSelfToVisLog(FVisLogEntry* Snapshot) const
 			continue;
 		}
 
-		FVisLogEntry::FStatusCategory StatusCategory;
+		FVisualLogEntry::FStatusCategory StatusCategory;
 		StatusCategory.Category = Category + TEXT(": ") + GetPriorityName(PriorityIndex);
 
 		while (Action)
diff --git a/Engine/Source/Runtime/AIModule/Private/BehaviorTree/BehaviorTreeComponent.cpp b/Engine/Source/Runtime/AIModule/Private/BehaviorTree/BehaviorTreeComponent.cpp
--- a/Engine/Source/Runtime/AIModule/Private/BehaviorTree/BehaviorTreeComponent.cpp
+++ b/Engine/Source/Runtime/AIModule/Private/BehaviorTree/BehaviorTreeComponent.cpp
@@ -1614,7 +1614,7 @@ FString UBehaviorTreeComponent::DescribeActiveTrees() const
 }
 
 #if ENABLE_VISUAL_LOG
-void UBehaviorTreeComponent::DescribeSelfToVisLog(struct FVisLogEntry* Snapshot) const
+void UBehaviorTreeComponent::DescribeSelfToVisLog(struct FVisualLogEntry* Snapshot) const
 {
 	if (IsPendingKill())
 	{
@@ -1628,7 +1628,7 @@ void UBehaviorTreeComponent::DescribeSelfToVisLog(struct FVisLogEntry* Snapshot)
 		const FBehaviorTreeInstance& InstanceInfo = InstanceStack[InstanceIndex];
 		const FBehaviorTreeInstanceId& InstanceId = KnownInstances[InstanceInfo.InstanceIdIndex];
 
-		FVisLogEntry::FStatusCategory StatusCategory;
+		FVisualLogEntry::FStatusCategory StatusCategory;
 		StatusCategory.Category = FString::Printf(TEXT("BehaviorTree %d (asset: %s)"), InstanceIndex, *GetNameSafe(InstanceId.TreeAsset));
 		
 		TArray<FString> Descriptions;
diff --git a/Engine/Source/Runtime/AIModule/Private/BehaviorTree/BlackboardComponent.cpp b/Engine/Source/Runtime/AIModule/Private/BehaviorTree/BlackboardComponent.cpp
--- a/Engine/Source/Runtime/AIModule/Private/BehaviorTree/BlackboardComponent.cpp
+++ b/Engine/Source/Runtime/AIModule/Private/BehaviorTree/BlackboardComponent.cpp
@@ -299,9 +299,9 @@ FString UBlackboardComponent::DescribeKeyValue(FBlackboard::FKey KeyID, EBlackbo
 
 #if ENABLE_VISUAL_LOG
 
-void UBlackboardComponent::DescribeSelfToVisLog(struct FVisLogEntry* Snapshot) const
+void UBlackboardComponent::DescribeSelfToVisLog(struct FVisualLogEntry* Snapshot) const
 {
-	FVisLogEntry::FStatusCategory Category;
+	FVisualLogEntry::FStatusCategory Category;
 	Category.Category = FString::Printf(TEXT("Blackboard (asset: %s)"), *GetNameSafe(BlackboardAsset));
 
 	for (UBlackboardData* It = BlackboardAsset; It; It = It->Parent)
diff --git a/Engine/Source/Runtime/AIModule/Private/BrainComponent.cpp b/Engine/Source/Runtime/AIModule/Private/BrainComponent.cpp
--- a/Engine/Source/Runtime/AIModule/Private/BrainComponent.cpp
+++ b/Engine/Source/Runtime/AIModule/Private/BrainComponent.cpp
@@ -184,7 +184,7 @@ UBrainComponent::UBrainComponent(const FPostConstructInitializeProperties& PCIP)
 }
 
 #if ENABLE_VISUAL_LOG
-void UBrainComponent::DescribeSelfToVisLog(struct FVisLogEntry* Snapshot) const
+void UBrainComponent::DescribeSelfToVisLog(struct FVisualLogEntry* Snapshot) const
 {
 	const static UEnum* SourceEnum = FindObject<UEnum>(ANY_PACKAGE, TEXT("EAILockSource"));
 
@@ -193,7 +193,7 @@ void UBrainComponent::DescribeSelfToVisLog(struct FVisLogEntry* Snapshot) const
 		return;
 	}
 
-	FVisLogEntry::FStatusCategory StatusCategory;
+	FVisualLogEntry::FStatusCategory StatusCategory;
 	StatusCategory.Category = FString::Printf(TEXT("Resource lock: %s"), *ResourceLock.GetLockSourceName());
 	for (int32 LockLevel = 0; LockLevel < int32(EAILockSource::MAX); ++LockLevel)
 	{
diff --git a/Engine/Source/Runtime/AIModule/Private/EnvironmentQuery/EnvQueryDebugHelpers.cpp b/Engine/Source/Runtime/AIModule/Private/EnvironmentQuery/EnvQueryDebugHelpers.cpp
--- a/Engine/Source/Runtime/AIModule/Private/EnvironmentQuery/EnvQueryDebugHelpers.cpp
+++ b/Engine/Source/Runtime/AIModule/Private/EnvironmentQuery/EnvQueryDebugHelpers.cpp
@@ -131,14 +131,25 @@ void  UEnvQueryDebugHelpers::BlobArrayToDebugData(const TArray<uint8>& BlobArray
 
 void UEnvQueryDebugHelpers::LogQuery(const class UObject* QueryOwner, struct FEnvQueryInstance* Query, const FName& CategoryName, ELogVerbosity::Type Type, const FString& AdditionalLogInfo)
 {
+#if ENABLE_VISUAL_LOG
 	TArray<uint8> BlobArray;
 	UEnvQueryDebugHelpers::QueryToBlobArray(Query, BlobArray);
-	FVisLogEntry* EntryToWrite = FVisualLog::Get().GetEntryToWrite(QueryOwner);
+
+	UObject* LogOwner = FVisualLogger::Get().FindRedirection(QueryOwner);
+	UWorld* World = GEngine->GetWorldFromContextObject(QueryOwner);
+
+	FVisualLogEntry* EntryToWrite = FVisualLogger::Get().GetEntryToWrite(LogOwner, World->TimeSeconds);
+	const int32 UniqueId = FVisualLogger::Get().GetUniqueId(World->TimeSeconds);
 	if (EntryToWrite)
 	{
-		FVisualLog::Get().LogLine(QueryOwner, CategoryName, Type, AdditionalLogInfo, Query->QueryID, *EVisLogTags::TAG_EQS);
-		EntryToWrite->AddDataBlock(EVisLogTags::TAG_EQS, BlobArray, CategoryName);
+		FVisualLogEntry::FLogLine Line(CategoryName, Type, AdditionalLogInfo, Query->QueryID);
+		Line.TagName = *EVisLogTags::TAG_EQS;
+		Line.UniqueId = UniqueId;
+		EntryToWrite->LogLines.Add(Line);
+
+		EntryToWrite->AddDataBlock(EVisLogTags::TAG_EQS, BlobArray, CategoryName).UniqueId = UniqueId;
 	}
+#endif
 }
 
 #endif //ENABLE_VISUAL_LOG
diff --git a/Engine/Source/Runtime/AIModule/Private/Navigation/CrowdFollowingComponent.cpp b/Engine/Source/Runtime/AIModule/Private/Navigation/CrowdFollowingComponent.cpp
--- a/Engine/Source/Runtime/AIModule/Private/Navigation/CrowdFollowingComponent.cpp
+++ b/Engine/Source/Runtime/AIModule/Private/Navigation/CrowdFollowingComponent.cpp
@@ -847,7 +847,7 @@ void UCrowdFollowingComponent::GetDebugStringTokens(TArray<FString>& Tokens, TAr
 
 #if ENABLE_VISUAL_LOG
 
-void UCrowdFollowingComponent::DescribeSelfToVisLog(struct FVisLogEntry* Snapshot) const
+void UCrowdFollowingComponent::DescribeSelfToVisLog(struct FVisualLogEntry* Snapshot) const
 {
 	if (!bEnableCrowdSimulation)
 	{
@@ -855,7 +855,7 @@ void UCrowdFollowingComponent::DescribeSelfToVisLog(struct FVisLogEntry* Snapsho
 		return;
 	}
 
-	FVisLogEntry::FStatusCategory Category;
+	FVisualLogEntry::FStatusCategory Category;
 	Category.Category = TEXT("Path following");
 
 	if (DestinationActor.IsValid())
diff --git a/Engine/Source/Runtime/AIModule/Private/Navigation/NavigationComponent.cpp b/Engine/Source/Runtime/AIModule/Private/Navigation/NavigationComponent.cpp
--- a/Engine/Source/Runtime/AIModule/Private/Navigation/NavigationComponent.cpp
+++ b/Engine/Source/Runtime/AIModule/Private/Navigation/NavigationComponent.cpp
@@ -959,7 +959,7 @@ void UNavigationComponent::SwapCurrentMoveGoal(const AActor* NewGoalActor)
 //----------------------------------------------------------------------//
 // debug
 //----------------------------------------------------------------------//
-void UNavigationComponent::DescribeSelfToVisLog(FVisLogEntry* Snapshot) const
+void UNavigationComponent::DescribeSelfToVisLog(FVisualLogEntry* Snapshot) const
 {
 	// path will be drawn by path following
 }
diff --git a/Engine/Source/Runtime/AIModule/Private/Navigation/PathFollowingComponent.cpp b/Engine/Source/Runtime/AIModule/Private/Navigation/PathFollowingComponent.cpp
--- a/Engine/Source/Runtime/AIModule/Private/Navigation/PathFollowingComponent.cpp
+++ b/Engine/Source/Runtime/AIModule/Private/Navigation/PathFollowingComponent.cpp
@@ -54,11 +54,11 @@ UPathFollowingComponent::UPathFollowingComponent(const class FPostConstructIniti
 void LogPathHelper(AActor* LogOwner, FNavPathSharedPtr Path, const AActor* GoalActor, const int32 CurrentPathPointGoal)
 {
 #if ENABLE_VISUAL_LOG
-	FVisualLog& Vlog = FVisualLog::Get();
+	FVisualLogger& Vlog = FVisualLogger::Get();
 	if (Vlog.IsRecording() && 
 		Path.IsValid() && Path->IsValid() && Path->GetPathPoints().Num())
 	{
-		FVisLogEntry* Entry = Vlog.GetEntryToWrite(LogOwner);
+		FVisualLogEntry* Entry = Vlog.GetEntryToWrite(LogOwner, LogOwner->GetWorld()->TimeSeconds);
 		Path->DescribeSelfToVisLog(Entry);
 
 		const FVector PathEnd = Path->GetPathPoints().Last().Location;
@@ -1306,9 +1306,9 @@ void UPathFollowingComponent::DisplayDebug(class UCanvas* Canvas, const FDebugDi
 }
 
 #if ENABLE_VISUAL_LOG
-void UPathFollowingComponent::DescribeSelfToVisLog(struct FVisLogEntry* Snapshot) const
+void UPathFollowingComponent::DescribeSelfToVisLog(struct FVisualLogEntry* Snapshot) const
 {
-	FVisLogEntry::FStatusCategory Category;
+	FVisualLogEntry::FStatusCategory Category;
 	Category.Category = TEXT("Path following");
 
 	if (DestinationActor.IsValid())
diff --git a/Engine/Source/Runtime/AIModule/Private/Perception/AIPerceptionComponent.cpp b/Engine/Source/Runtime/AIModule/Private/Perception/AIPerceptionComponent.cpp
--- a/Engine/Source/Runtime/AIModule/Private/Perception/AIPerceptionComponent.cpp
+++ b/Engine/Source/Runtime/AIModule/Private/Perception/AIPerceptionComponent.cpp
@@ -469,7 +469,7 @@ void UAIPerceptionComponent::DrawDebugInfo(UCanvas* Canvas)
 }
 
 #if ENABLE_VISUAL_LOG
-void UAIPerceptionComponent::DescribeSelfToVisLog(FVisLogEntry* Snapshot) const
+void UAIPerceptionComponent::DescribeSelfToVisLog(FVisualLogEntry* Snapshot) const
 {
 
 }
diff --git a/Engine/Source/Runtime/AIModule/Private/VisualLoggerExtension.cpp b/Engine/Source/Runtime/AIModule/Private/VisualLoggerExtension.cpp
--- a/Engine/Source/Runtime/AIModule/Private/VisualLoggerExtension.cpp
+++ b/Engine/Source/Runtime/AIModule/Private/VisualLoggerExtension.cpp
@@ -27,7 +27,7 @@ void FVisualLoggerExtension::OnTimestampChange(float Timestamp, class UWorld* In
 #endif
 }
 
-void FVisualLoggerExtension::DisableDrawingForData(class UWorld* InWorld, class UCanvas* Canvas, class AActor* HelperActor, const FName& TagName, const FVisLogEntry::FDataBlock& DataBlock, float Timestamp)
+void FVisualLoggerExtension::DisableDrawingForData(class UWorld* InWorld, class UCanvas* Canvas, class AActor* HelperActor, const FName& TagName, const FVisualLogEntry::FDataBlock& DataBlock, float Timestamp)
 {
 #if USE_EQS_DEBUGGER
 	if (TagName == *EVisLogTags::TAG_EQS && CurrentTimestamp == Timestamp)
@@ -65,7 +65,7 @@ void FVisualLoggerExtension::LogEntryLineSelectionChanged(TSharedPtr<struct FLog
 
 extern RENDERCORE_API class FTexture* GWhiteTexture;
 
-void FVisualLoggerExtension::DrawData(class UWorld* InWorld, class UCanvas* Canvas, class AActor* HelperActor, const FName& TagName, const FVisLogEntry::FDataBlock& DataBlock, float Timestamp)
+void FVisualLoggerExtension::DrawData(class UWorld* InWorld, class UCanvas* Canvas, class AActor* HelperActor, const FName& TagName, const FVisualLogEntry::FDataBlock& DataBlock, float Timestamp)
 {
 #if USE_EQS_DEBUGGER
 	if (TagName == *EVisLogTags::TAG_EQS && HelperActor)
diff --git a/Engine/Source/Runtime/Engine/Classes/AI/Navigation/NavigationData.h b/Engine/Source/Runtime/Engine/Classes/AI/Navigation/NavigationData.h
--- a/Engine/Source/Runtime/Engine/Classes/AI/Navigation/NavigationData.h
+++ b/Engine/Source/Runtime/Engine/Classes/AI/Navigation/NavigationData.h
@@ -163,7 +163,7 @@ struct ENGINE_API FNavigationPath : public TSharedFromThis<FNavigationPath, ESPM
 
 	virtual void DebugDraw(const ANavigationData* NavData, FColor PathColor, class UCanvas* Canvas, bool bPersistent, const uint32 NextPathPointIndex = 0) const;
 #if ENABLE_VISUAL_LOG
-	virtual void DescribeSelfToVisLog(struct FVisLogEntry* Snapshot) const;
+	virtual void DescribeSelfToVisLog(struct FVisualLogEntry* Snapshot) const;
 	virtual FString GetDescription() const;
 #endif // ENABLE_VISUAL_LOG
 
diff --git a/Engine/Source/Runtime/Engine/Classes/AI/Navigation/RecastNavMesh.h b/Engine/Source/Runtime/Engine/Classes/AI/Navigation/RecastNavMesh.h
--- a/Engine/Source/Runtime/Engine/Classes/AI/Navigation/RecastNavMesh.h
+++ b/Engine/Source/Runtime/Engine/Classes/AI/Navigation/RecastNavMesh.h
@@ -144,7 +144,7 @@ struct ENGINE_API FNavMeshPath : public FNavigationPath
 	virtual bool DoesIntersectBox(const FBox& Box, uint32 StartingIndex = 0, int32* IntersectingSegmentIndex = NULL) const override;
 
 #if ENABLE_VISUAL_LOG
-	virtual void DescribeSelfToVisLog(struct FVisLogEntry* Snapshot) const override;
+	virtual void DescribeSelfToVisLog(struct FVisualLogEntry* Snapshot) const override;
 	virtual FString GetDescription() const override;
 #endif // ENABLE_VISUAL_LOG
 
diff --git a/Engine/Source/Runtime/Engine/Classes/GameFramework/Actor.h b/Engine/Source/Runtime/Engine/Classes/GameFramework/Actor.h
--- a/Engine/Source/Runtime/Engine/Classes/GameFramework/Actor.h
+++ b/Engine/Source/Runtime/Engine/Classes/GameFramework/Actor.h
@@ -2158,7 +2158,7 @@ public:
 	 *	Hook for Actors to supply visual logger with additional data.
 	 *	It's guaranteed that Snapshot != NULL
 	 */
-	virtual void GrabDebugSnapshot(struct FVisLogEntry* Snapshot) const {}
+	virtual void GrabDebugSnapshot(struct FVisualLogEntry* Snapshot) const {}
 
 private:
 	friend class FVisualLog;
diff --git a/Engine/Source/Runtime/Engine/Classes/GameFramework/CheatManager.h b/Engine/Source/Runtime/Engine/Classes/GameFramework/CheatManager.h
--- a/Engine/Source/Runtime/Engine/Classes/GameFramework/CheatManager.h
+++ b/Engine/Source/Runtime/Engine/Classes/GameFramework/CheatManager.h
@@ -67,6 +67,10 @@ class ENGINE_API UCheatManager : public UObject
 	UPROPERTY()
 	TSubclassOf<class ADebugCameraController>  DebugCameraControllerClass;
 
+	/** Holds information about VisualLogger on server - came from replication */
+	UPROPERTY(ReplicatedUsing = OnRep_VisualLoggerActiveOnServer)
+	bool bVisualLoggerActiveOnServer;
+
 	// Trace/Sweep debug start
 	/** If we should should perform a debug capsule trace and draw results. Toggled with DebugCapsuleSweep() */
 	uint32 bDebugCapsuleSweep:1;
@@ -80,7 +84,8 @@ class ENGINE_API UCheatManager : public UObject
 	/** Holds information if we used ToggleAILogging cheat to activate AI logging */
 	uint32 bToggleAILogging : 1;
 
-	/** How far debugf trace should go out from player viewpoint */
+
+	/** How far debug trace should go out from player viewpoint */
 	float DebugTraceDistance;
 
 	/** Half distance between debug capsule sphere ends. Total heigh of capsule is 2*(this + DebugCapsuleRadius). */
@@ -107,9 +112,6 @@ class ENGINE_API UCheatManager : public UObject
 	/** Index of the array for current trace to overwrite.  Whenever you capture, this index will be increased **/
 	int32 CurrentTracePawnIndex;
 
-	/** time interval for dumping AI loggs (vlogs) */
-	float DumpAILogsInterval;
-
 	/** Pause the game for Delay seconds. */
 	UFUNCTION(exec)
 	virtual void FreezeFrame(float Delay);
@@ -205,7 +207,9 @@ class ENGINE_API UCheatManager : public UObject
 	UFUNCTION(reliable, server, WithValidation)
 	virtual void ServerToggleAILogging();
 
-	void DumpAILogs();
+	/** after bVisualLoggerActiveOnServer replication callback */
+	UFUNCTION()
+	virtual void OnRep_VisualLoggerActiveOnServer();
 
 	/** Toggle capsule trace debugging. Will trace a capsule from current view point and show where it hits the world */
 	UFUNCTION(exec)
@@ -307,7 +311,6 @@ class ENGINE_API UCheatManager : public UObject
 	virtual void LogOutBugItGoToLogFile( const FString& InScreenShotDesc, const FString& InGoString, const FString& InLocString );
 
 #if !(UE_BUILD_SHIPPING || UE_BUILD_TEST)
-	
 	/** Do any trace debugging that is currently enabled */
 	void TickCollisionDebug();
 
diff --git a/Engine/Source/Runtime/Engine/Classes/VisualLogger/VisualLoggerAutomationTests.h b/Engine/Source/Runtime/Engine/Classes/VisualLogger/VisualLoggerAutomationTests.h
new file mode 100644
--- /dev/null
+++ b/Engine/Source/Runtime/Engine/Classes/VisualLogger/VisualLoggerAutomationTests.h
@@ -0,0 +1,13 @@
+// Copyright 1998-2014 Epic Games, Inc. All Rights Reserved.
+
+#pragma once
+
+#include "Misc/AutomationTest.h"
+#include "VisualLoggerAutomationTests.generated.h"
+
+UCLASS(NotBlueprintable, Transient, hidecategories = UObject, notplaceable)
+class UVisualLoggerAutomationTests : public UObject
+{
+	GENERATED_UCLASS_BODY()
+};
+
diff --git a/Engine/Source/Runtime/Engine/Classes/VisualLogger/VisualLoggerKismetLibrary.h b/Engine/Source/Runtime/Engine/Classes/VisualLogger/VisualLoggerKismetLibrary.h
new file mode 100644
--- /dev/null
+++ b/Engine/Source/Runtime/Engine/Classes/VisualLogger/VisualLoggerKismetLibrary.h
@@ -0,0 +1,20 @@
+// Copyright 1998-2014 Epic Games, Inc. All Rights Reserved.
+#include "VisualLoggerKismetLibrary.generated.h"
+
+UCLASS(MinimalAPI)
+class UVisualLoggerKismetLibrary : public UBlueprintFunctionLibrary
+{
+	GENERATED_UCLASS_BODY()
+
+	/** Logs simple text string with Visual Logger - recording for Visual Logs has to be enabled to record this data */
+	UFUNCTION(BlueprintCallable, Category = "Debug|VisualLogger", meta = (HidePin = "WorldContextObject", DefaultToSelf = "WorldContextObject", FriendlyName = "LogText"))
+	static void LogText(UObject* WorldContextObject, FString Text, FName LogCategory = TEXT("Blueprint Log"));
+
+	/** Logs location as sphere with given radius - recording for Visual Logs has to be enabled to record this data */
+	UFUNCTION(BlueprintCallable, Category = "Debug|VisualLogger", meta = (HidePin = "WorldContextObject", DefaultToSelf = "WorldContextObject", FriendlyName = "LogLocation"))
+	static void LogLocation(UObject* WorldContextObject, FVector Location, FString Text, FLinearColor ObjectColor = FLinearColor::White, float Radius = 10, FName LogCategory = TEXT("Blueprint Log"));
+
+	/** Logs box shape - recording for Visual Logs has to be enabled to record this data */
+	UFUNCTION(BlueprintCallable, Category = "Debug|VisualLogger", meta = (HidePin = "WorldContextObject", DefaultToSelf = "WorldContextObject", FriendlyName = "Log Box Shape"))
+	static void LogBox(UObject* WorldContextObject, FBox BoxShape, FString Text, FLinearColor ObjectColor = FLinearColor::White, FName LogCategory = TEXT("Blueprint Log"));
+};
diff --git a/Engine/Source/Runtime/Engine/Private/AI/Navigation/NavigationPath.cpp b/Engine/Source/Runtime/Engine/Private/AI/Navigation/NavigationPath.cpp
--- a/Engine/Source/Runtime/Engine/Private/AI/Navigation/NavigationPath.cpp
+++ b/Engine/Source/Runtime/Engine/Private/AI/Navigation/NavigationPath.cpp
@@ -264,10 +264,10 @@ bool FNavigationPath::DoesIntersectBox(const FBox& Box, uint32 StartingIndex, in
 
 #if ENABLE_VISUAL_LOG
 
-void FNavigationPath::DescribeSelfToVisLog(FVisLogEntry* Snapshot) const 
+void FNavigationPath::DescribeSelfToVisLog(FVisualLogEntry* Snapshot) const 
 {
 	const int32 NumPathVerts = PathPoints.Num();
-	FVisLogEntry::FElementToDraw Element(FVisLogEntry::FElementToDraw::Path);
+	FVisualLogEntry::FElementToDraw Element(FVisualLogEntry::FElementToDraw::Path);
 	Element.Category = LogNavigation.GetCategoryName();
 	Element.SetColor(FColorList::Green);
 	Element.Points.Reserve(NumPathVerts);
@@ -899,7 +899,7 @@ bool FNavMeshPath::DoesIntersectBox(const FBox& Box, uint32 StartingIndex, int32
 
 #if ENABLE_VISUAL_LOG
 
-void FNavMeshPath::DescribeSelfToVisLog(FVisLogEntry* Snapshot) const
+void FNavMeshPath::DescribeSelfToVisLog(FVisualLogEntry* Snapshot) const
 {
 	if (IsStringPulled())
 	{
@@ -909,7 +909,7 @@ void FNavMeshPath::DescribeSelfToVisLog(FVisLogEntry* Snapshot) const
 
 	// draw corridor
 #if WITH_RECAST
-	FVisLogEntry::FElementToDraw CorridorElem(FVisLogEntry::FElementToDraw::Segment);
+	FVisualLogEntry::FElementToDraw CorridorElem(FVisualLogEntry::FElementToDraw::Segment);
 	CorridorElem.SetColor(FColorList::Cyan);
 	CorridorElem.Category = LogNavigation.GetCategoryName();
 	CorridorElem.Points.Reserve(PathCorridor.Num() * 6);
@@ -940,7 +940,7 @@ void FNavMeshPath::DescribeSelfToVisLog(FVisLogEntry* Snapshot) const
 		const UClass* AreaClass = NavMesh->GetAreaClass(AreaID);
 		if (AreaClass && AreaClass != UNavigationSystem::GetDefaultWalkableArea())
 		{
-			FVisLogEntry::FElementToDraw AreaMarkElem(FVisLogEntry::FElementToDraw::Segment);
+			FVisualLogEntry::FElementToDraw AreaMarkElem(FVisualLogEntry::FElementToDraw::Segment);
 			AreaMarkElem.SetColor(FColorList::Orange);
 			AreaMarkElem.Category = LogNavigation.GetCategoryName();
 			AreaMarkElem.Thicknes = 2;
diff --git a/Engine/Source/Runtime/Engine/Private/CheatManager.cpp b/Engine/Source/Runtime/Engine/Private/CheatManager.cpp
--- a/Engine/Source/Runtime/Engine/Private/CheatManager.cpp
+++ b/Engine/Source/Runtime/Engine/Private/CheatManager.cpp
@@ -1,7 +1,7 @@
 // Copyright 1998-2014 Epic Games, Inc. All Rights Reserved.
 
 #include "EnginePrivate.h"
-
+#include "UnrealNetwork.h"
 #include "Slate.h"
 #include "SlateReflector.h"
 #include "NavDataGenerator.h"
@@ -19,8 +19,8 @@ DEFINE_LOG_CATEGORY_STATIC(LogCheatManager, Log, All);
 
 UCheatManager::UCheatManager(const class FPostConstructInitializeProperties& PCIP)
 	: Super(PCIP)
+	, bToggleAILogging(false)
 {
-	DumpAILogsInterval = 1;
 	DebugCameraControllerClass = ADebugCameraController::StaticClass();
 	DebugCapsuleHalfHeight = 23.0f;
 	DebugCapsuleRadius = 21.0f;
@@ -30,6 +30,14 @@ UCheatManager::UCheatManager(const class FPostConstructInitializeProperties& PCI
 	bDebugCapsuleTraceComplex = false;
 }
 
+void UCheatManager::GetLifetimeReplicatedProps(TArray< FLifetimeProperty > & OutLifetimeProps) const
+{
+	Super::GetLifetimeReplicatedProps(OutLifetimeProps);
+#if ENABLE_VISUAL_LOG
+	DOREPLIFETIME(UCheatManager, bVisualLoggerActiveOnServer);
+#endif
+}
+
 void UCheatManager::FreezeFrame(float delay)
 {
 	FCanUnpause DefaultCanUnpause;
@@ -561,18 +569,13 @@ void UCheatManager::InitCheatManager()
 void UCheatManager::BeginDestroy()
 {
 #if ENABLE_VISUAL_LOG
-	if (bToggleAILogging && FVisualLog::Get().IsRecording())
+	if (bToggleAILogging && FVisualLogger::Get().IsRecording())
 	{
-		UWorld *World = GetWorld();
-		if (World)
-		{
-			// clear timer, we'll dump all remaining logs
-			World->GetTimerManager().ClearTimer(this, &UCheatManager::DumpAILogs);
-		}
-
 		// stop recording and dump all remaining logs
-		FVisualLog::Get().SetIsRecording(false);
+		FVisualLogger::Get().SetIsRecording(false);
+		FVisualLogger::Get().SetIsRecordingToFile(false);
 		bToggleAILogging = false;
+		bVisualLoggerActiveOnServer = false;
 	}
 #endif
 	Super::BeginDestroy();
@@ -583,47 +586,25 @@ bool UCheatManager::ServerToggleAILogging_Validate()
 	return true;
 }
 
-void UCheatManager::DumpAILogs()
-{
-#if ENABLE_VISUAL_LOG
-	UWorld *World = GetWorld();
-	if (World == NULL)
-	{
-		return;
-	}
-
-	FVisualLog::Get().DumpRecordedLogs();
-	World->GetTimerManager().SetTimer(this, &UCheatManager::DumpAILogs, DumpAILogsInterval);
-#endif
-}
-
 void UCheatManager::ServerToggleAILogging_Implementation()
 {
 #if ENABLE_VISUAL_LOG
 	UWorld *World = GetWorld();
-	if (FVisualLog::Get().IsRecording())
+	if (FVisualLogger::Get().IsRecordingToFile())
 	{
-		// clear timer, we'll dump all remaining logs in a moment
-		if (World)
-		{
-			World->GetTimerManager().ClearTimer(this, &UCheatManager::DumpAILogs);
-		}
-
 		// stop recording and dump all remaining logs in a moment
-		FVisualLog::Get().SetIsRecording(false, true);
+		FVisualLogger::Get().SetIsRecordingToFile(false);
+		FVisualLogger::Get().SetIsRecording(false);
 		bToggleAILogging = false;
 	}
 	else
 	{
-		if (World)
-		{
-			World->GetTimerManager().SetTimer(this, &UCheatManager::DumpAILogs, DumpAILogsInterval);
-		}
-		FVisualLog::Get().SetIsRecording(true, true);
+		FVisualLogger::Get().SetIsRecordingToFile(true);
 		bToggleAILogging = true;
 	}
 
-	GetOuterAPlayerController()->ClientMessage(FString::Printf(TEXT("OK! VisLog recording is now %s"), FVisualLog::Get().IsRecording() ? TEXT("Enabled") : TEXT("Disabled")));
+	bVisualLoggerActiveOnServer = bToggleAILogging;
+	GetOuterAPlayerController()->ClientMessage(FString::Printf(TEXT("OK! VisLog recording is now %s"), FVisualLogger::Get().IsRecording() ? TEXT("Enabled") : TEXT("Disabled")));
 #endif
 }
 
@@ -639,8 +620,6 @@ void UCheatManager::ToggleAILogging()
 	UWorld *World = GetWorld();
 	if (World && World->GetNetMode() == NM_Client)
 	{
-		FVisualLog::Get().SetIsRecordingOnServer(!FVisualLog::Get().IsRecordingOnServer());
-		GetOuterAPlayerController()->ClientMessage(FString::Printf(TEXT("OK! VisLog recording is now %s"), FVisualLog::Get().IsRecordingOnServer() ? TEXT("Enabled") : TEXT("Disabled")));
 		PC->ServerToggleAILogging();
 	}
 	else
@@ -650,6 +629,13 @@ void UCheatManager::ToggleAILogging()
 #endif
 }
 
+void UCheatManager::OnRep_VisualLoggerActiveOnServer()
+{
+#if ENABLE_VISUAL_LOG
+	FVisualLogger::Get().SetIsRecordingOnServer(bVisualLoggerActiveOnServer);
+#endif
+}
+
 #define SAFE_TRACEINDEX_DECREASE(x) ((--x)<0)? 9:(x)
 
 #if !(UE_BUILD_SHIPPING || UE_BUILD_TEST)
diff --git a/Engine/Source/Runtime/Engine/Private/Debug/LogVisualizerHUD.cpp b/Engine/Source/Runtime/Engine/Private/Debug/LogVisualizerHUD.cpp
--- a/Engine/Source/Runtime/Engine/Private/Debug/LogVisualizerHUD.cpp
+++ b/Engine/Source/Runtime/Engine/Private/Debug/LogVisualizerHUD.cpp
@@ -47,7 +47,7 @@ void ALogVisualizerHUD::PostRender()
 			bool bHit = GetWorld()->LineTraceSingle(Hit, CamLoc, CamRot.Vector() * 100000.f + CamLoc, ECC_Pawn, TraceParams);
 			if( bHit )
 			{
-				TextItem.Text = FText::FromString(FString::Printf(TEXT("Under cursor: '%s'"), TEXT_NAME(Hit.GetActor())));
+				TextItem.Text = FText::FromString(FString::Printf(TEXT("Under cursor: '%s'"), *Hit.GetActor()->GetName()));
 				Canvas->DrawItem( TextItem, X, Y );
 				
 				DrawDebugLine( GetWorld(), Hit.Location, Hit.Location+Hit.Normal*30.f, FColor::White );
diff --git a/Engine/Source/Runtime/Engine/Private/UnrealEngine.cpp b/Engine/Source/Runtime/Engine/Private/UnrealEngine.cpp
--- a/Engine/Source/Runtime/Engine/Private/UnrealEngine.cpp
+++ b/Engine/Source/Runtime/Engine/Private/UnrealEngine.cpp
@@ -6731,7 +6731,7 @@ void DrawStatsHUD( UWorld* World, FViewport* Viewport, FCanvas* Canvas, UCanvas*
 			}
 
 #if ENABLE_VISUAL_LOG
-			if (FVisualLog::Get().IsRecording() || FVisualLog::Get().IsRecordingOnServer())
+			if (FVisualLogger::Get().IsRecording() || FVisualLogger::Get().IsRecordingOnServer() )
 			{
 				int32 XSize;
 				int32 YSize;
diff --git a/Engine/Source/Runtime/Engine/Private/VisualLog.cpp b/Engine/Source/Runtime/Engine/Private/VisualLog.cpp
--- a/Engine/Source/Runtime/Engine/Private/VisualLog.cpp
+++ b/Engine/Source/Runtime/Engine/Private/VisualLog.cpp
@@ -5,90 +5,170 @@
 #include "VisualLog.h"
 #include "Json.h"
 
-DEFINE_LOG_CATEGORY(LogVisual);
-
 #if ENABLE_VISUAL_LOG
-
 DEFINE_STAT(STAT_VisualLog);
 
-//----------------------------------------------------------------------//
-// FVisLogEntry
-//----------------------------------------------------------------------//
-FVisLogEntry::FVisLogEntry(const class AActor* InActor, TArray<TWeakObjectPtr<UObject> >* Children)
+TSharedPtr<FJsonValue> GenerateJson(FVisualLogEntry* InEntry)
 {
-	if (InActor && InActor->IsPendingKill() == false)
+	TSharedPtr<FJsonObject> JsonEntryObject = MakeShareable(new FJsonObject);
+
+	JsonEntryObject->SetNumberField(VisualLogJson::TAG_TIMESTAMP, InEntry->TimeStamp);
+	JsonEntryObject->SetStringField(VisualLogJson::TAG_LOCATION, InEntry->Location.ToString());
+
+	const int32 StatusCategoryCount = InEntry->Status.Num();
+	const FVisualLogEntry::FStatusCategory* StatusCategory = InEntry->Status.GetData();
+	TArray< TSharedPtr<FJsonValue> > JsonStatus;
+	JsonStatus.AddZeroed(StatusCategoryCount);
+
+	for (int32 StatusIdx = 0; StatusIdx < InEntry->Status.Num(); StatusIdx++, StatusCategory++)
 	{
-		TimeStamp = InActor->GetWorld()->TimeSeconds;
-		Location = InActor->GetActorLocation();
-		InActor->GrabDebugSnapshot(this);
-		if (Children != NULL)
+		TSharedPtr<FJsonObject> JsonStatusCategoryObject = MakeShareable(new FJsonObject);
+		JsonStatusCategoryObject->SetStringField(VisualLogJson::TAG_CATEGORY, StatusCategory->Category);
+
+		TArray< TSharedPtr<FJsonValue> > JsonStatusLines;
+		JsonStatusLines.AddZeroed(StatusCategory->Data.Num());
+		for (int32 LineIdx = 0; LineIdx < StatusCategory->Data.Num(); LineIdx++)
 		{
-			TWeakObjectPtr<UObject>* WeakActorPtr = Children->GetData();
-			for (int32 Index = 0; Index < Children->Num(); ++Index, ++WeakActorPtr)
-			{
-				if (WeakActorPtr->IsValid())
-				{
-					const AActor* ChildActor = Cast<AActor>(WeakActorPtr->Get());
-					if (ChildActor)
-					{
-						ChildActor->GrabDebugSnapshot(this);
-					}
-				}
-			}
+			JsonStatusLines[LineIdx] = MakeShareable(new FJsonValueString(StatusCategory->Data[LineIdx]));
 		}
+
+		JsonStatusCategoryObject->SetArrayField(VisualLogJson::TAG_STATUSLINES, JsonStatusLines);
+		JsonStatus[StatusIdx] = MakeShareable(new FJsonValueObject(JsonStatusCategoryObject));
 	}
-}
+	JsonEntryObject->SetArrayField(VisualLogJson::TAG_STATUS, JsonStatus);
 
-FVisLogEntry::FVisLogEntry(float InTimeStamp, FVector InLocation, const UObject* Object, TArray<TWeakObjectPtr<UObject> >* Children)
-{
-	TimeStamp = InTimeStamp;
-	Location = InLocation;
-	if (Children != NULL)
+	const int32 LogLineCount = InEntry->LogLines.Num();
+	const FVisualLogEntry::FLogLine* LogLine = InEntry->LogLines.GetData();
+	TArray< TSharedPtr<FJsonValue> > JsonLines;
+	JsonLines.AddZeroed(LogLineCount);
+
+	for (int32 LogLineIndex = 0; LogLineIndex < LogLineCount; ++LogLineIndex, ++LogLine)
+	{
+		TSharedPtr<FJsonObject> JsonLogLineObject = MakeShareable(new FJsonObject);
+		JsonLogLineObject->SetStringField(VisualLogJson::TAG_CATEGORY, LogLine->Category.ToString());
+		JsonLogLineObject->SetNumberField(VisualLogJson::TAG_VERBOSITY, LogLine->Verbosity);
+		JsonLogLineObject->SetStringField(VisualLogJson::TAG_LINE, LogLine->Line);
+		JsonLogLineObject->SetStringField(VisualLogJson::TAG_TAGNAME, LogLine->TagName.ToString());
+		JsonLogLineObject->SetNumberField(VisualLogJson::TAG_USERDATA, LogLine->UserData);
+		JsonLines[LogLineIndex] = MakeShareable(new FJsonValueObject(JsonLogLineObject));
+	}
+
+	JsonEntryObject->SetArrayField(VisualLogJson::TAG_LOGLINES, JsonLines);
+
+	const int32 ElementsToDrawCount = InEntry->ElementsToDraw.Num();
+	const FVisualLogEntry::FElementToDraw* Element = InEntry->ElementsToDraw.GetData();
+	TArray< TSharedPtr<FJsonValue> > JsonElementsToDraw;
+	JsonElementsToDraw.AddZeroed(ElementsToDrawCount);
+
+	for (int32 ElementIndex = 0; ElementIndex < ElementsToDrawCount; ++ElementIndex, ++Element)
 	{
-		TWeakObjectPtr<UObject>* WeakActorPtr = Children->GetData();
-		for (int32 Index = 0; Index < Children->Num(); ++Index, ++WeakActorPtr)
+		TSharedPtr<FJsonObject> JsonElementToDrawObject = MakeShareable(new FJsonObject);
+
+		JsonElementToDrawObject->SetStringField(VisualLogJson::TAG_DESCRIPTION, Element->Description);
+		JsonElementToDrawObject->SetStringField(VisualLogJson::TAG_CATEGORY, Element->Category.ToString());
+		JsonElementToDrawObject->SetNumberField(VisualLogJson::TAG_VERBOSITY, Element->Verbosity);
+
+		const int32 EncodedTypeColorSize = Element->Type << 24 | Element->Color << 16 | Element->Thicknes;
+		JsonElementToDrawObject->SetStringField(VisualLogJson::TAG_TYPECOLORSIZE, FString::Printf(TEXT("%d"), EncodedTypeColorSize));
+
+		const int32 PointsCount = Element->Points.Num();
+		TArray< TSharedPtr<FJsonValue> > JsonStringPoints;
+		JsonStringPoints.AddZeroed(PointsCount);
+		const FVector* PointToDraw = Element->Points.GetData();
+		for (int32 PointIndex = 0; PointIndex < PointsCount; ++PointIndex, ++PointToDraw)
 		{
-			if (WeakActorPtr->IsValid())
-			{
-				const AActor* ChildActor = Cast<AActor>(WeakActorPtr->Get());
-				if (ChildActor)
-				{
-					ChildActor->GrabDebugSnapshot(this);
-				}
-			}
+			JsonStringPoints[PointIndex] = MakeShareable(new FJsonValueString(PointToDraw->ToString()));
 		}
+		JsonElementToDrawObject->SetArrayField(VisualLogJson::TAG_POINTS, JsonStringPoints);
+
+		JsonElementsToDraw[ElementIndex] = MakeShareable(new FJsonValueObject(JsonElementToDrawObject));
+	}
+
+	JsonEntryObject->SetArrayField(VisualLogJson::TAG_ELEMENTSTODRAW, JsonElementsToDraw);
+
+	const int32 HistogramSamplesCount = InEntry->HistogramSamples.Num();
+	const FVisualLogEntry::FHistogramSample* Sample = InEntry->HistogramSamples.GetData();
+	TArray<TSharedPtr<FJsonValue> > JsonHistogramSamples;
+	JsonHistogramSamples.AddZeroed(HistogramSamplesCount);
+
+	for (int32 SampleIndex = 0; SampleIndex < HistogramSamplesCount; ++SampleIndex, ++Sample)
+	{
+		TSharedPtr<FJsonObject> JsonSampleObject = MakeShareable(new FJsonObject);
+
+		JsonSampleObject->SetStringField(VisualLogJson::TAG_CATEGORY, Sample->Category.ToString());
+		JsonSampleObject->SetNumberField(VisualLogJson::TAG_VERBOSITY, Sample->Verbosity);
+		JsonSampleObject->SetStringField(VisualLogJson::TAG_HISTOGRAMSAMPLE, Sample->SampleValue.ToString());
+		JsonSampleObject->SetStringField(VisualLogJson::TAG_HISTOGRAMGRAPHNAME, Sample->GraphName.ToString());
+		JsonSampleObject->SetStringField(VisualLogJson::TAG_HISTOGRAMDATANAME, Sample->DataName.ToString());
+
+		JsonHistogramSamples[SampleIndex] = MakeShareable(new FJsonValueObject(JsonSampleObject));
+	}
+
+	JsonEntryObject->SetArrayField(VisualLogJson::TAG_HISTOGRAMSAMPLES, JsonHistogramSamples);
+
+	int32 DataBlockIndex = 0;
+	TArray<TSharedPtr<FJsonValue> > JasonDataBlocks;
+	JasonDataBlocks.AddZeroed(InEntry->DataBlocks.Num());
+	for (const auto CurrentData : InEntry->DataBlocks)
+	{
+		TSharedPtr<FJsonObject> JsonSampleObject = MakeShareable(new FJsonObject);
+
+		TArray<uint8> CompressedData;
+		const int32 UncompressedSize = CurrentData.Data.Num();
+		const int32 HeaderSize = sizeof(int32);
+		CompressedData.Init(0, HeaderSize + FMath::TruncToInt(1.1f * UncompressedSize));
+
+		int32 CompressedSize = CompressedData.Num() - HeaderSize;
+		uint8* DestBuffer = CompressedData.GetData();
+		FMemory::Memcpy(DestBuffer, &UncompressedSize, HeaderSize);
+		DestBuffer += HeaderSize;
+
+		FCompression::CompressMemory((ECompressionFlags)(COMPRESS_ZLIB | COMPRESS_BiasMemory), (void*)DestBuffer, CompressedSize, (void*)CurrentData.Data.GetData(), UncompressedSize);
+		CompressedData.SetNum(CompressedSize + HeaderSize, true);
+		const FString CurrentDataAsString = FBase64::Encode(CompressedData);
+
+		JsonSampleObject->SetStringField(VisualLogJson::TAG_CATEGORY, CurrentData.Category.ToString());
+		JsonSampleObject->SetStringField(VisualLogJson::TAG_TAGNAME, CurrentData.TagName.ToString());
+		JsonSampleObject->SetStringField(VisualLogJson::TAG_DATABLOCK_DATA, CurrentDataAsString);
+		JsonSampleObject->SetNumberField(VisualLogJson::TAG_VERBOSITY, CurrentData.Verbosity);
+
+		JasonDataBlocks[DataBlockIndex++] = MakeShareable(new FJsonValueObject(JsonSampleObject));
 	}
+	JsonEntryObject->SetArrayField(VisualLogJson::TAG_DATABLOCK, JasonDataBlocks);
+
+	return MakeShareable(new FJsonValueObject(JsonEntryObject));
 }
 
-FVisLogEntry::FVisLogEntry(TSharedPtr<FJsonValue> FromJson)
+TSharedPtr<FVisualLogEntry> JsonToVisualLogEntry(TSharedPtr<FJsonValue> FromJson)
 {
 	TSharedPtr<FJsonObject> JsonEntryObject = FromJson->AsObject();
 	if (JsonEntryObject.IsValid() == false)
 	{
-		return;
+		return NULL;
 	}
 
-	TimeStamp = JsonEntryObject->GetNumberField(VisualLogJson::TAG_TIMESTAMP);
-	Location.InitFromString(JsonEntryObject->GetStringField(VisualLogJson::TAG_LOCATION));
+	TSharedPtr<FVisualLogEntry> VisLogEntry = MakeShareable(new FVisualLogEntry());
+	VisLogEntry->TimeStamp = JsonEntryObject->GetNumberField(VisualLogJson::TAG_TIMESTAMP);
+	VisLogEntry->Location.InitFromString(JsonEntryObject->GetStringField(VisualLogJson::TAG_LOCATION));
 
 	TArray< TSharedPtr<FJsonValue> > JsonStatus = JsonEntryObject->GetArrayField(VisualLogJson::TAG_STATUS);
 	if (JsonStatus.Num() > 0)
 	{
-		Status.AddZeroed(JsonStatus.Num());
+		VisLogEntry->Status.AddZeroed(JsonStatus.Num());
 		for (int32 StatusIndex = 0; StatusIndex < JsonStatus.Num(); ++StatusIndex)
 		{
 			TSharedPtr<FJsonObject> JsonStatusCategory = JsonStatus[StatusIndex]->AsObject();
 			if (JsonStatusCategory.IsValid())
 			{
-				Status[StatusIndex].Category = JsonStatusCategory->GetStringField(VisualLogJson::TAG_CATEGORY);
-				
+				VisLogEntry->Status[StatusIndex].Category = JsonStatusCategory->GetStringField(VisualLogJson::TAG_CATEGORY);
+
 				TArray< TSharedPtr<FJsonValue> > JsonStatusLines = JsonStatusCategory->GetArrayField(VisualLogJson::TAG_STATUSLINES);
 				if (JsonStatusLines.Num() > 0)
 				{
-					Status[StatusIndex].Data.AddZeroed(JsonStatusLines.Num());					
+					VisLogEntry->Status[StatusIndex].Data.AddZeroed(JsonStatusLines.Num());
 					for (int32 LineIdx = 0; LineIdx < JsonStatusLines.Num(); LineIdx++)
 					{
-						Status[StatusIndex].Data.Add(JsonStatusLines[LineIdx]->AsString());
+						VisLogEntry->Status[StatusIndex].Data.Add(JsonStatusLines[LineIdx]->AsString());
 					}
 				}
 			}
@@ -98,17 +178,17 @@ FVisLogEntry::FVisLogEntry(TSharedPtr<FJsonValue> FromJson)
 	TArray< TSharedPtr<FJsonValue> > JsonLines = JsonEntryObject->GetArrayField(VisualLogJson::TAG_LOGLINES);
 	if (JsonLines.Num() > 0)
 	{
-		LogLines.AddZeroed(JsonLines.Num());
+		VisLogEntry->LogLines.AddZeroed(JsonLines.Num());
 		for (int32 LogLineIndex = 0; LogLineIndex < JsonLines.Num(); ++LogLineIndex)
 		{
 			TSharedPtr<FJsonObject> JsonLogLine = JsonLines[LogLineIndex]->AsObject();
 			if (JsonLogLine.IsValid())
 			{
-				LogLines[LogLineIndex].Category = FName(*(JsonLogLine->GetStringField(VisualLogJson::TAG_CATEGORY)));
-				LogLines[LogLineIndex].Verbosity = TEnumAsByte<ELogVerbosity::Type>((uint8)FMath::TruncToInt(JsonLogLine->GetNumberField(VisualLogJson::TAG_VERBOSITY)));
-				LogLines[LogLineIndex].Line = JsonLogLine->GetStringField(VisualLogJson::TAG_LINE);
-				LogLines[LogLineIndex].TagName = FName(*(JsonLogLine->GetStringField(VisualLogJson::TAG_TAGNAME)));
-				LogLines[LogLineIndex].UserData = (int64)(JsonLogLine->GetNumberField(VisualLogJson::TAG_USERDATA));
+				VisLogEntry->LogLines[LogLineIndex].Category = FName(*(JsonLogLine->GetStringField(VisualLogJson::TAG_CATEGORY)));
+				VisLogEntry->LogLines[LogLineIndex].Verbosity = TEnumAsByte<ELogVerbosity::Type>((uint8)FMath::TruncToInt(JsonLogLine->GetNumberField(VisualLogJson::TAG_VERBOSITY)));
+				VisLogEntry->LogLines[LogLineIndex].Line = JsonLogLine->GetStringField(VisualLogJson::TAG_LINE);
+				VisLogEntry->LogLines[LogLineIndex].TagName = FName(*(JsonLogLine->GetStringField(VisualLogJson::TAG_TAGNAME)));
+				VisLogEntry->LogLines[LogLineIndex].UserData = (int64)(JsonLogLine->GetNumberField(VisualLogJson::TAG_USERDATA));
 
 			}
 		}
@@ -117,13 +197,13 @@ FVisLogEntry::FVisLogEntry(TSharedPtr<FJsonValue> FromJson)
 	TArray< TSharedPtr<FJsonValue> > JsonElementsToDraw = JsonEntryObject->GetArrayField(VisualLogJson::TAG_ELEMENTSTODRAW);
 	if (JsonElementsToDraw.Num() > 0)
 	{
-		ElementsToDraw.Reserve(JsonElementsToDraw.Num());
+		VisLogEntry->ElementsToDraw.Reserve(JsonElementsToDraw.Num());
 		for (int32 ElementIndex = 0; ElementIndex < JsonElementsToDraw.Num(); ++ElementIndex)
 		{
 			TSharedPtr<FJsonObject> JsonElementObject = JsonElementsToDraw[ElementIndex]->AsObject();
 			if (JsonElementObject.IsValid())
 			{
-				FElementToDraw& Element = ElementsToDraw[ElementsToDraw.Add(FElementToDraw())];
+				FVisualLogEntry::FElementToDraw& Element = VisLogEntry->ElementsToDraw[VisLogEntry->ElementsToDraw.Add(FVisualLogEntry::FElementToDraw())];
 
 				Element.Description = JsonElementObject->GetStringField(VisualLogJson::TAG_DESCRIPTION);
 				Element.Category = FName(*(JsonElementObject->GetStringField(VisualLogJson::TAG_CATEGORY)));
@@ -131,7 +211,7 @@ FVisLogEntry::FVisLogEntry(TSharedPtr<FJsonValue> FromJson)
 
 				// Element->Type << 24 | Element->Color << 16 | Element->Thicknes;
 				int32 EncodedTypeColorSize = 0;
-				TTypeFromString<int32>::FromString(EncodedTypeColorSize, *(JsonElementObject->GetStringField(VisualLogJson::TAG_TYPECOLORSIZE)));				
+				TTypeFromString<int32>::FromString(EncodedTypeColorSize, *(JsonElementObject->GetStringField(VisualLogJson::TAG_TYPECOLORSIZE)));
 				Element.Type = EncodedTypeColorSize >> 24;
 				Element.Color = (EncodedTypeColorSize << 8) >> 24;
 				Element.Thicknes = (EncodedTypeColorSize << 16) >> 16;
@@ -142,7 +222,7 @@ FVisLogEntry::FVisLogEntry(TSharedPtr<FJsonValue> FromJson)
 					Element.Points.AddUninitialized(JsonPoints.Num());
 					for (int32 PointIndex = 0; PointIndex < JsonPoints.Num(); ++PointIndex)
 					{
-						Element.Points[PointIndex].InitFromString(JsonPoints[PointIndex]->AsString());							
+						Element.Points[PointIndex].InitFromString(JsonPoints[PointIndex]->AsString());
 					}
 				}
 			}
@@ -152,13 +232,13 @@ FVisLogEntry::FVisLogEntry(TSharedPtr<FJsonValue> FromJson)
 	TArray< TSharedPtr<FJsonValue> > JsonHistogramSamples = JsonEntryObject->GetArrayField(VisualLogJson::TAG_HISTOGRAMSAMPLES);
 	if (JsonHistogramSamples.Num() > 0)
 	{
-		HistogramSamples.Reserve(JsonHistogramSamples.Num());
+		VisLogEntry->HistogramSamples.Reserve(JsonHistogramSamples.Num());
 		for (int32 SampleIndex = 0; SampleIndex < JsonHistogramSamples.Num(); ++SampleIndex)
 		{
 			TSharedPtr<FJsonObject> JsonSampleObject = JsonHistogramSamples[SampleIndex]->AsObject();
 			if (JsonSampleObject.IsValid())
 			{
-				FHistogramSample& Sample = HistogramSamples[HistogramSamples.Add(FHistogramSample())];
+				FVisualLogEntry::FHistogramSample& Sample = VisLogEntry->HistogramSamples[VisLogEntry->HistogramSamples.Add(FVisualLogEntry::FHistogramSample())];
 
 				Sample.Category = FName(*(JsonSampleObject->GetStringField(VisualLogJson::TAG_CATEGORY)));
 				Sample.Verbosity = TEnumAsByte<ELogVerbosity::Type>((uint8)FMath::TruncToInt(JsonSampleObject->GetNumberField(VisualLogJson::TAG_VERBOSITY)));
@@ -172,13 +252,13 @@ FVisLogEntry::FVisLogEntry(TSharedPtr<FJsonValue> FromJson)
 	TArray< TSharedPtr<FJsonValue> > JasonDataBlocks = JsonEntryObject->GetArrayField(VisualLogJson::TAG_DATABLOCK);
 	if (JasonDataBlocks.Num() > 0)
 	{
-		DataBlocks.Reserve(JasonDataBlocks.Num());
+		VisLogEntry->DataBlocks.Reserve(JasonDataBlocks.Num());
 		for (int32 SampleIndex = 0; SampleIndex < JasonDataBlocks.Num(); ++SampleIndex)
 		{
 			TSharedPtr<FJsonObject> JsonSampleObject = JasonDataBlocks[SampleIndex]->AsObject();
 			if (JsonSampleObject.IsValid())
 			{
-				FDataBlock& Sample = DataBlocks[DataBlocks.Add(FDataBlock())];
+				FVisualLogEntry::FDataBlock& Sample = VisLogEntry->DataBlocks[VisLogEntry->DataBlocks.Add(FVisualLogEntry::FDataBlock())];
 				Sample.TagName = FName(*(JsonSampleObject->GetStringField(VisualLogJson::TAG_TAGNAME)));
 				Sample.Category = FName(*(JsonSampleObject->GetStringField(VisualLogJson::TAG_CATEGORY)));
 				Sample.Verbosity = TEnumAsByte<ELogVerbosity::Type>((uint8)FMath::TruncToInt(JsonSampleObject->GetNumberField(VisualLogJson::TAG_VERBOSITY)));
@@ -202,229 +282,8 @@ FVisLogEntry::FVisLogEntry(TSharedPtr<FJsonValue> FromJson)
 		}
 	}
 
+	return VisLogEntry;
 }
-
-TSharedPtr<FJsonValue> FVisLogEntry::ToJson() const
-{
-	TSharedPtr<FJsonObject> JsonEntryObject = MakeShareable(new FJsonObject);
-
-	JsonEntryObject->SetNumberField(VisualLogJson::TAG_TIMESTAMP, TimeStamp);
-	JsonEntryObject->SetStringField(VisualLogJson::TAG_LOCATION, Location.ToString());
-
-	const int32 StatusCategoryCount = Status.Num();
-	const FVisLogEntry::FStatusCategory* StatusCategory = Status.GetData();
-	TArray< TSharedPtr<FJsonValue> > JsonStatus;
-	JsonStatus.AddZeroed(StatusCategoryCount);
-
-	for (int32 StatusIdx = 0; StatusIdx < Status.Num(); StatusIdx++, StatusCategory++)
-	{
-		TSharedPtr<FJsonObject> JsonStatusCategoryObject = MakeShareable(new FJsonObject);
-		JsonStatusCategoryObject->SetStringField(VisualLogJson::TAG_CATEGORY, StatusCategory->Category);
-
-		TArray< TSharedPtr<FJsonValue> > JsonStatusLines;
-		JsonStatusLines.AddZeroed(StatusCategory->Data.Num());
-		for (int32 LineIdx = 0; LineIdx < StatusCategory->Data.Num(); LineIdx++)
-		{
-			JsonStatusLines[LineIdx] = MakeShareable(new FJsonValueString(StatusCategory->Data[LineIdx]));
-		}
-
-		JsonStatusCategoryObject->SetArrayField(VisualLogJson::TAG_STATUSLINES, JsonStatusLines);
-		JsonStatus[StatusIdx] = MakeShareable(new FJsonValueObject(JsonStatusCategoryObject));
-	}
-	JsonEntryObject->SetArrayField(VisualLogJson::TAG_STATUS, JsonStatus);
-
-	const int32 LogLineCount = LogLines.Num();
-	const FVisLogEntry::FLogLine* LogLine = LogLines.GetData();
-	TArray< TSharedPtr<FJsonValue> > JsonLines;
-	JsonLines.AddZeroed(LogLineCount);
-
-	for (int32 LogLineIndex = 0; LogLineIndex < LogLineCount; ++LogLineIndex, ++LogLine)
-	{
-		TSharedPtr<FJsonObject> JsonLogLineObject = MakeShareable(new FJsonObject);
-		JsonLogLineObject->SetStringField(VisualLogJson::TAG_CATEGORY, LogLine->Category.ToString());
-		JsonLogLineObject->SetNumberField(VisualLogJson::TAG_VERBOSITY, LogLine->Verbosity);
-		JsonLogLineObject->SetStringField(VisualLogJson::TAG_LINE, LogLine->Line);
-		JsonLogLineObject->SetStringField(VisualLogJson::TAG_TAGNAME, LogLine->TagName.ToString());
-		JsonLogLineObject->SetNumberField(VisualLogJson::TAG_USERDATA, LogLine->UserData);
-		JsonLines[LogLineIndex] = MakeShareable(new FJsonValueObject(JsonLogLineObject));
-	}
-
-	JsonEntryObject->SetArrayField(VisualLogJson::TAG_LOGLINES, JsonLines);
-
-	const int32 ElementsToDrawCount = ElementsToDraw.Num();
-	const FVisLogEntry::FElementToDraw* Element = ElementsToDraw.GetData();
-	TArray< TSharedPtr<FJsonValue> > JsonElementsToDraw;
-	JsonElementsToDraw.AddZeroed(ElementsToDrawCount);
-
-	for (int32 ElementIndex = 0; ElementIndex < ElementsToDrawCount; ++ElementIndex, ++Element)
-	{
-		TSharedPtr<FJsonObject> JsonElementToDrawObject = MakeShareable(new FJsonObject);
-
-		JsonElementToDrawObject->SetStringField(VisualLogJson::TAG_DESCRIPTION, Element->Description);
-		JsonElementToDrawObject->SetStringField(VisualLogJson::TAG_CATEGORY, Element->Category.ToString());
-		JsonElementToDrawObject->SetNumberField(VisualLogJson::TAG_VERBOSITY, Element->Verbosity);
-
-		const int32 EncodedTypeColorSize = Element->Type << 24 | Element->Color << 16 | Element->Thicknes;
-		JsonElementToDrawObject->SetStringField(VisualLogJson::TAG_TYPECOLORSIZE, FString::Printf(TEXT("%d"), EncodedTypeColorSize));
-
-		const int32 PointsCount = Element->Points.Num();
-		TArray< TSharedPtr<FJsonValue> > JsonStringPoints;
-		JsonStringPoints.AddZeroed(PointsCount);
-		const FVector* PointToDraw = Element->Points.GetData();
-		for (int32 PointIndex = 0; PointIndex < PointsCount; ++PointIndex, ++PointToDraw)
-		{
-			JsonStringPoints[PointIndex] = MakeShareable(new FJsonValueString(PointToDraw->ToString()));
-		}
-		JsonElementToDrawObject->SetArrayField(VisualLogJson::TAG_POINTS, JsonStringPoints);
-		
-		JsonElementsToDraw[ElementIndex] = MakeShareable(new FJsonValueObject(JsonElementToDrawObject));
-	}
-
-	JsonEntryObject->SetArrayField(VisualLogJson::TAG_ELEMENTSTODRAW, JsonElementsToDraw);
-	
-	const int32 HistogramSamplesCount = HistogramSamples.Num();
-	const FVisLogEntry::FHistogramSample* Sample = HistogramSamples.GetData();
-	TArray<TSharedPtr<FJsonValue> > JsonHistogramSamples;
-	JsonHistogramSamples.AddZeroed(HistogramSamplesCount);
-
-	for (int32 SampleIndex = 0; SampleIndex < HistogramSamplesCount; ++SampleIndex, ++Sample)
-	{
-		TSharedPtr<FJsonObject> JsonSampleObject = MakeShareable(new FJsonObject);
-
-		JsonSampleObject->SetStringField(VisualLogJson::TAG_CATEGORY, Sample->Category.ToString());
-		JsonSampleObject->SetNumberField(VisualLogJson::TAG_VERBOSITY, Sample->Verbosity);
-		JsonSampleObject->SetStringField(VisualLogJson::TAG_HISTOGRAMSAMPLE, Sample->SampleValue.ToString());
-		JsonSampleObject->SetStringField(VisualLogJson::TAG_HISTOGRAMGRAPHNAME, Sample->GraphName.ToString());
-		JsonSampleObject->SetStringField(VisualLogJson::TAG_HISTOGRAMDATANAME, Sample->DataName.ToString());
-
-		JsonHistogramSamples[SampleIndex] = MakeShareable(new FJsonValueObject(JsonSampleObject));
-	}
-
-	JsonEntryObject->SetArrayField(VisualLogJson::TAG_HISTOGRAMSAMPLES, JsonHistogramSamples);
-	
-	int32 DataBlockIndex = 0;
-	TArray<TSharedPtr<FJsonValue> > JasonDataBlocks;
-	JasonDataBlocks.AddZeroed(DataBlocks.Num());
-	for (const auto CurrentData : DataBlocks)
-	{
-		TSharedPtr<FJsonObject> JsonSampleObject = MakeShareable(new FJsonObject);
-
-		TArray<uint8> CompressedData;
-		const int32 UncompressedSize = CurrentData.Data.Num();
-		const int32 HeaderSize = sizeof(int32);
-		CompressedData.Init(0, HeaderSize + FMath::TruncToInt(1.1f * UncompressedSize));
-
-		int32 CompressedSize = CompressedData.Num() - HeaderSize;
-		uint8* DestBuffer = CompressedData.GetData();
-		FMemory::Memcpy(DestBuffer, &UncompressedSize, HeaderSize);
-		DestBuffer += HeaderSize;
-
-		FCompression::CompressMemory((ECompressionFlags)(COMPRESS_ZLIB | COMPRESS_BiasMemory), (void*)DestBuffer, CompressedSize, (void*)CurrentData.Data.GetData(), UncompressedSize);
-		CompressedData.SetNum(CompressedSize + HeaderSize, true);
-		const FString CurrentDataAsString = FBase64::Encode(CompressedData);
-
-		JsonSampleObject->SetStringField(VisualLogJson::TAG_CATEGORY, CurrentData.Category.ToString());
-		JsonSampleObject->SetStringField(VisualLogJson::TAG_TAGNAME, CurrentData.TagName.ToString());
-		JsonSampleObject->SetStringField(VisualLogJson::TAG_DATABLOCK_DATA, CurrentDataAsString);
-		JsonSampleObject->SetNumberField(VisualLogJson::TAG_VERBOSITY, CurrentData.Verbosity);
-
-		JasonDataBlocks[DataBlockIndex++] = MakeShareable(new FJsonValueObject(JsonSampleObject));
-	}
-	JsonEntryObject->SetArrayField(VisualLogJson::TAG_DATABLOCK, JasonDataBlocks);
-
-	return MakeShareable(new FJsonValueObject(JsonEntryObject));
-}
-
-void FVisLogEntry::AddElement(const TArray<FVector>& Points, const FName& CategoryName, const FColor& Color, const FString& Description, uint16 Thickness)
-{
-	FElementToDraw Element(Description, Color, Thickness, CategoryName);
-	Element.Points = Points;
-	Element.Type = FElementToDraw::Path;
-	ElementsToDraw.Add(Element);
-}
-
-void FVisLogEntry::AddElement(const FVector& Point, const FName& CategoryName, const FColor& Color, const FString& Description, uint16 Thickness)
-{
-	FElementToDraw Element(Description, Color, Thickness, CategoryName);
-	Element.Points.Add(Point);
-	Element.Type = FElementToDraw::SinglePoint;
-	ElementsToDraw.Add(Element);
-}
-
-void FVisLogEntry::AddElement(const FVector& Start, const FVector& End, const FName& CategoryName, const FColor& Color, const FString& Description, uint16 Thickness)
-{
-	FElementToDraw Element(Description, Color, Thickness, CategoryName);
-	Element.Points.Reserve(2);
-	Element.Points.Add(Start);
-	Element.Points.Add(End);
-	Element.Type = FElementToDraw::Segment;
-	ElementsToDraw.Add(Element);
-}
-
-void FVisLogEntry::AddElement(const FBox& Box, const FName& CategoryName, const FColor& Color, const FString& Description, uint16 Thickness)
-{
-	FElementToDraw Element(Description, Color, Thickness, CategoryName);
-	Element.Points.Reserve(2);
-	Element.Points.Add(Box.Min);
-	Element.Points.Add(Box.Max);
-	Element.Type = FElementToDraw::Box;
-	ElementsToDraw.Add(Element);
-}
-
-void FVisLogEntry::AddElement(const FVector& Orgin, const FVector& Direction, float Length, float AngleWidth, float AngleHeight, const FName& CategoryName, const FColor& Color, const FString& Description, uint16 Thickness)
-{
-	FElementToDraw Element(Description, Color, Thickness, CategoryName);
-	Element.Points.Reserve(3);
-	Element.Points.Add(Orgin);
-	Element.Points.Add(Direction);
-	Element.Points.Add(FVector(Length, AngleWidth, AngleHeight));
-	Element.Type = FElementToDraw::Cone;
-	ElementsToDraw.Add(Element);
-}
-
-void FVisLogEntry::AddElement(const FVector& Start, const FVector& End, float Radius, const FName& CategoryName, const FColor& Color, const FString& Description, uint16 Thickness)
-{
-	FElementToDraw Element(Description, Color, Thickness, CategoryName);
-	Element.Points.Reserve(3);
-	Element.Points.Add(Start);
-	Element.Points.Add(End);
-	Element.Points.Add(FVector(Radius, Thickness, 0));
-	Element.Type = FElementToDraw::Cylinder;
-	ElementsToDraw.Add(Element);
-}
-
-void FVisLogEntry::AddHistogramData(const FVector2D& DataSample, const FName& CategoryName, const FName& GraphName, const FName& DataName)
-{
-	FHistogramSample Sample;
-	Sample.Category = CategoryName;
-	Sample.GraphName = GraphName;
-	Sample.DataName = DataName;
-	Sample.SampleValue = DataSample;
-
-	HistogramSamples.Add(Sample);
-}
-
-void FVisLogEntry::AddDataBlock(const FString& TagName, const TArray<uint8>& BlobDataArray, const FName& CategoryName)
-{
-	FDataBlock DataBlock;
-	DataBlock.Category = CategoryName;
-	DataBlock.TagName = *TagName;
-	DataBlock.Data = BlobDataArray;
-
-	DataBlocks.Add(DataBlock);
-}
-
-void FVisLogEntry::AddCapsule(FVector const& Center, float HalfHeight, float Radius, const FQuat& Rotation, const FName& CategoryName, const FColor& Color, const FString& Description)
-{
-	FElementToDraw Element(Description, Color, 0, CategoryName);
-	Element.Points.Reserve(3);
-	Element.Points.Add(Center);
-	Element.Points.Add(FVector(HalfHeight, Radius, Rotation.X));
-	Element.Points.Add(FVector(Rotation.Y, Rotation.Z, Rotation.W));
-	Element.Type = FElementToDraw::Capsule;
-	ElementsToDraw.Add(Element);
-}
-
 //----------------------------------------------------------------------//
 // FActorsVisLog
 //----------------------------------------------------------------------//
@@ -432,7 +291,7 @@ FActorsVisLog::FActorsVisLog(const class UObject* Object, TArray<TWeakObjectPtr<
 	: Name(Object->GetFName())
 	, FullName(Object->GetFullName())
 {
-	const class AActor* AsActor = FVisualLog::Get().GetVisualLogRedirection(Object);
+	const class AActor* AsActor = Cast<AActor>(Object);
 
 	Entries.Reserve(VisLogInitialSize);
 	if (!AsActor)
@@ -440,11 +299,11 @@ FActorsVisLog::FActorsVisLog(const class UObject* Object, TArray<TWeakObjectPtr<
 		UWorld* World = GEngine->GetWorldFromContextObject(Object);
 		check(World);
 
-		Entries.Add(MakeShareable(new FVisLogEntry(World->GetTimeSeconds(), FVector(), Object, Children)));
+		Entries.Add(MakeShareable(new FVisualLogEntry(World->GetTimeSeconds(), FVector(), Object, Children)));
 	}
 	else
 	{
-		Entries.Add(MakeShareable(new FVisLogEntry(AsActor, Children)));
+		Entries.Add(MakeShareable(new FVisualLogEntry(AsActor, Children)));
 	}
 }
 
@@ -465,7 +324,7 @@ FActorsVisLog::FActorsVisLog(TSharedPtr<FJsonValue> FromJson)
 		Entries.AddZeroed(JsonEntries.Num());
 		for (int32 EntryIndex = 0; EntryIndex < JsonEntries.Num(); ++EntryIndex)
 		{
-			Entries[EntryIndex] = MakeShareable(new FVisLogEntry(JsonEntries[EntryIndex]));
+			Entries[EntryIndex] = JsonToVisualLogEntry(JsonEntries[EntryIndex]);
 		}
 	}
 }
@@ -477,46 +336,37 @@ TSharedPtr<FJsonValue> FActorsVisLog::ToJson() const
 	JsonLogObject->SetStringField(VisualLogJson::TAG_NAME, Name.ToString());
 	JsonLogObject->SetStringField(VisualLogJson::TAG_FULLNAME, FullName);
 
-	const TSharedPtr<FVisLogEntry>* Entry = Entries.GetData();
+	const TSharedPtr<FVisualLogEntry>* Entry = Entries.GetData();
 	const int32 EntryCount = Entries.Num();
 	TArray< TSharedPtr<FJsonValue> > JsonLogEntries;
 	JsonLogEntries.AddZeroed(EntryCount);
 
 	for (int32 EntryIndex = 0; EntryIndex < EntryCount; ++EntryIndex, ++Entry)
 	{
-		JsonLogEntries[EntryIndex] = (*Entry)->ToJson(); 
+		JsonLogEntries[EntryIndex] = GenerateJson((*Entry).Get());
 	}
 
 	JsonLogObject->SetArrayField(VisualLogJson::TAG_ENTRIES, JsonLogEntries);
 
 	return MakeShareable(new FJsonValueObject(JsonLogObject));
 }
+
 //----------------------------------------------------------------------//
 // FVisualLog
 //----------------------------------------------------------------------//
 FVisualLog::FVisualLog()
 	: FileAr(NULL)
 	, StartRecordingTime(0.0f)
-	, bIsRecording(GEngine ? GEngine->bEnableVisualLogRecordingOnStart : false)
 	, bIsRecordingOnServer(false)
 	, bIsRecordingToFile(false)
 	, bIsAllBlocked(false)
 {
 	Whitelist.Reserve(10);
-
-	if (FParse::Param(FCommandLine::Get(), TEXT("EnableAILogging")))
-	{
-		SetIsRecording(true, true);
-		if (GWorld)
-		{
-			GWorld->GetTimerManager().SetTimer(FTimerDelegate::CreateRaw(this, &FVisualLog::DumpRecordedLogs), 4, true);
-		}
-	}
 }
 
 FVisualLog::~FVisualLog()
 {
-	if (bIsRecording)
+	if (bIsRecordingToFile)
 	{
 		if (GWorld)
 		{
@@ -526,19 +376,38 @@ FVisualLog::~FVisualLog()
 	}
 }
 
-FString FVisualLog::GetLogFileFullName() const
+void FVisualLog::StartRecordingToFile(float TImeStamp)
 {
-	FString NameCore(TEXT(""));
-	if (LogFileNameGetter.IsBound())
+	SetIsRecording(IsRecording(), true);
+	if (GWorld)
 	{
-		NameCore = LogFileNameGetter.Execute();
+		GWorld->GetTimerManager().SetTimer(FTimerDelegate::CreateRaw(this, &FVisualLog::DumpRecordedLogs), 4, true);
 	}
-	if (NameCore.IsEmpty())
+}
+
+void FVisualLog::StopRecordingToFile(float TImeStamp)
+{
+	SetIsRecording(IsRecording(), false);
+	if (GWorld)
 	{
-		NameCore = TEXT("VisualLog");
+		GWorld->GetTimerManager().ClearTimer(FTimerDelegate::CreateRaw(this, &FVisualLog::DumpRecordedLogs));
 	}
+}
 
-	FString FileName = FString::Printf(TEXT("%s_%.0f-%.0f_%s.vlog"), *NameCore, StartRecordingTime, GWorld ? GWorld->TimeSeconds : 0, *FDateTime::Now().ToString());
+void FVisualLog::SetFileName(const FString& InFileName)
+{
+	BaseFileName = InFileName;
+}
+
+void FVisualLog::Serialize(const class UObject* LogOwner, const FVisualLogEntry& LogEntry)
+{
+	TSharedPtr<FActorsVisLog> Log = GetLog(LogOwner);
+	Log->Entries.Add(MakeShareable(new FVisualLogEntry(LogEntry)));
+}
+
+FString FVisualLog::GetLogFileFullName() const
+{
+	FString FileName = FString::Printf(TEXT("%s_%.0f-%.0f_%s.vlog"), *BaseFileName, StartRecordingTime, GWorld ? GWorld->TimeSeconds : 0, *FDateTime::Now().ToString());
 
 	const FString FullFilename = FString::Printf(TEXT("%slogs/%s"), *FPaths::GameSavedDir(), *FileName);
 
@@ -569,10 +438,10 @@ void FVisualLog::DumpRecordedLogs()
 			JsonLogEntries.AddZeroed(Log->Entries.Num());
 			for (int32 EntryIndex = 0; EntryIndex < Log->Entries.Num(); ++EntryIndex)
 			{
-				FVisLogEntry* Entry = Log->Entries[EntryIndex].Get();
+				FVisualLogEntry* Entry = Log->Entries[EntryIndex].Get();
 				if (Entry)
 				{
-					JsonLogEntries[EntryIndex] = Entry->ToJson();
+					JsonLogEntries[EntryIndex] = GenerateJson(Entry);
 				}
 			}
 			JsonLogObject->SetArrayField(VisualLogJson::TAG_ENTRIES, JsonLogEntries);
@@ -588,19 +457,9 @@ void FVisualLog::DumpRecordedLogs()
 	}
 }
 
-void FVisualLog::Serialize( const TCHAR* V, ELogVerbosity::Type Verbosity, const class FName& Category )
-{
-
-}
-
 void FVisualLog::SetIsRecording(bool NewRecording, bool bRecordToFile)
 {
-	if (GEngine && GEngine->bDisableAILogging)
-	{
-		return;
-	}
-
-	if (bIsRecording && bIsRecordingToFile && !NewRecording)
+	if (IsRecording() && bIsRecordingToFile && !bRecordToFile)
 	{
 		if (FileAr)
 		{
@@ -620,8 +479,7 @@ void FVisualLog::SetIsRecording(bool NewRecording, bool bRecordToFile)
 		bIsRecordingToFile = false;
 	}
 
-	bIsRecording = NewRecording;
-	if (bIsRecording)
+	if (IsRecording() && !bIsRecordingToFile && bRecordToFile)
 	{
 		bIsRecordingToFile = bRecordToFile;
 		StartRecordingTime = GWorld ? GWorld->TimeSeconds : 0.0f;
@@ -637,28 +495,28 @@ void FVisualLog::SetIsRecording(bool NewRecording, bool bRecordToFile)
 	}
 }
 
-FVisLogEntry*  FVisualLog::GetEntryToWrite(const class UObject* Object)
+FVisualLogEntry*  FVisualLog::GetEntryToWrite(const class UObject* Object, float Timestamp, const FVector& Location)
 {
 	const class AActor* RedirectionActor = GetVisualLogRedirection(Object);
 
 	UWorld* World = RedirectionActor && RedirectionActor->GetWorld() ? RedirectionActor->GetWorld() : GEngine->GetWorldFromContextObject(Object);
 	ensure(World);
 
-	const float TimeStamp = World->TimeSeconds;
+	const float TimeStamp = Timestamp < 0 ? World->TimeSeconds : Timestamp;
 	TSharedPtr<FActorsVisLog> Log = GetLog(RedirectionActor ? RedirectionActor : Object);
 	const int32 LastIndex = Log->Entries.Num() - 1;
-	FVisLogEntry* Entry = Log->Entries.Num() > 0 ? Log->Entries[LastIndex].Get() : NULL;
+	FVisualLogEntry* Entry = Log->Entries.Num() > 0 ? Log->Entries[LastIndex].Get() : NULL;
 
 	if (Entry == NULL || Entry->TimeStamp < TimeStamp)
 	{
 		// create new entry
 		if (RedirectionActor)
 		{
-			Entry = Log->Entries[Log->Entries.Add(MakeShareable(new FVisLogEntry(RedirectionActor, RedirectsMap.Find(RedirectionActor))))].Get();
+			Entry = Log->Entries[Log->Entries.Add(MakeShareable(new FVisualLogEntry(TimeStamp, Location, RedirectionActor, RedirectsMap.Find(RedirectionActor))))].Get();
 		}
 		else
 		{
-			Entry = Log->Entries[Log->Entries.Add(MakeShareable(new FVisLogEntry(TimeStamp, FVector(), Object, RedirectsMap.Find(RedirectionActor))))].Get();
+			Entry = Log->Entries[Log->Entries.Add(MakeShareable(new FVisualLogEntry(TimeStamp, Location, Object, RedirectsMap.Find(RedirectionActor))))].Get();
 		}
 	}
 
@@ -668,11 +526,11 @@ FVisLogEntry*  FVisualLog::GetEntryToWrite(const class UObject* Object)
 
 void FVisualLog::Cleanup(bool bReleaseMemory)
 {
-	if (bIsRecording && bIsRecordingToFile)
-	{
-		// dump collected data if needed
-		DumpRecordedLogs();
-	}
+	//if (bIsRecording && bIsRecordingToFile)
+	//{
+	//	// dump collected data if needed
+	//	DumpRecordedLogs();
+	//}
 
 	if (bReleaseMemory)
 	{
@@ -784,13 +642,13 @@ void FVisualLog::LogLine(const UObject* Object, const FName& CategoryName, ELogV
 		return;
 	}
 
-	FVisLogEntry* Entry = GetEntryToWrite(Object);
+	FVisualLogEntry* Entry = GetEntryToWrite(Object);
 
 	if (Entry != NULL)
 	{
 		// @todo will have to store CategoryName separately, and create a map of names 
 		// used in log to have saved logs independent from FNames index changes
-		int32 LineIndex = Entry->LogLines.Add(FVisLogEntry::FLogLine(CategoryName, Verbosity, Line, UserData));
+		int32 LineIndex = Entry->LogLines.Add(FVisualLogEntry::FLogLine(CategoryName, Verbosity, Line, UserData));
 		Entry->LogLines[LineIndex].TagName = TagName;
 	}
 }
@@ -812,19 +670,19 @@ public:
 			FString Command = FParse::Token(Cmd, 0);
 			if (Command == TEXT("record"))
 			{
-				FVisualLog::Get().SetIsRecording(true);
+				FVisualLogger::Get().SetIsRecording(true);
 				return true;
 			}
 			else if (Command == TEXT("stop"))
 			{
-				FVisualLog::Get().SetIsRecording(false);
+				FVisualLogger::Get().SetIsRecording(false);
 				return true;
 			}
 			else if (Command == TEXT("disableallbut"))
 			{
 				FString Category = FParse::Token(Cmd, 1);
-				FVisualLog::Get().BlockAllLogs(true);
-				FVisualLog::Get().AddCategortyToWhiteList(*Category);
+				FVisualLogger::Get().BlockAllCategories(true);
+				FVisualLogger::Get().GetWhiteList().AddUnique(*Category);
 				return true;
 			}
 #if WITH_EDITOR
@@ -848,4 +706,4 @@ public:
 } LogVisualizerExec;
 
 
-#endif // ENABLE_VISUAL_LOG
+#endif // !(UE_BUILD_SHIPPING || UE_BUILD_TEST)
diff --git a/Engine/Source/Runtime/Engine/Private/VisualLogger/VisualLogger.cpp b/Engine/Source/Runtime/Engine/Private/VisualLogger/VisualLogger.cpp
new file mode 100644
--- /dev/null
+++ b/Engine/Source/Runtime/Engine/Private/VisualLogger/VisualLogger.cpp
@@ -0,0 +1,154 @@
+// Copyright 1998-2014 Epic Games, Inc. All Rights Reserved.
+
+#include "EnginePrivate.h"
+#include "VisualLog.h"
+#include "VisualLogger/VisualLogger.h"
+
+#if ENABLE_VISUAL_LOG 
+
+DEFINE_LOG_CATEGORY(LogVisual);
+
+TMap<class UObject*, TArray<const class UObject*> > FVisualLogger::RedirectionMap;
+
+FVisualLogger::FVisualLogger()
+{
+	BlockAllCategories(false);
+	AddDevice(&FVisualLog::GetStatic());
+	SetIsRecording(GEngine ? !!GEngine->bEnableVisualLogRecordingOnStart : false);
+	SetIsRecordingOnServer(false);
+
+	if (FParse::Param(FCommandLine::Get(), TEXT("EnableAILogging")))
+	{
+		SetIsRecording(true);
+		SetIsRecordingToFile(true);
+	}
+}
+
+FVisualLogger::~FVisualLogger()
+{
+	SetIsRecording(false);
+	SetIsRecordingToFile(false);
+
+	RemoveDevice(&FVisualLog::GetStatic());
+}
+
+void FVisualLogger::Cleanup(bool bReleaseMemory)
+{
+	const bool WasRecordingToFile = IsRecordingToFile();
+	if (WasRecordingToFile)
+	{
+		SetIsRecordingToFile(false);
+	}
+
+	auto& OutputDevices = FVisualLogger::Get().OutputDevices;
+	for (auto* Device : OutputDevices)
+	{
+		Device->Cleanup(bReleaseMemory);
+	}
+
+	RedirectionMap.Reset();
+	LastUniqueIds.Reset();
+	CurrentEntryPerObject.Reset();
+
+	if (WasRecordingToFile)
+	{
+		SetIsRecordingToFile(true);
+	}
+}
+
+int32 FVisualLogger::GetUniqueId(float Timestamp)
+{
+	return LastUniqueIds.FindOrAdd(Timestamp)++;
+}
+
+void FVisualLogger::Redirect(class UObject* FromObject, class UObject* ToObject)
+{
+	UObject* OldRedirection = FindRedirection(FromObject);
+	UObject* NewRedierection = FindRedirection(ToObject);
+
+	UE_CVLOG(OldRedirection != NULL, OldRedirection, LogVisual, Log, TEXT("Redirected '%s' to '%s'"), *FromObject->GetName(), *NewRedierection->GetName());
+
+	TArray<const class UObject*>& OldArray = RedirectionMap.FindOrAdd(OldRedirection);
+	OldArray.RemoveSingleSwap(FromObject);
+
+	TArray<const class UObject*>& NewArray = RedirectionMap.FindOrAdd(NewRedierection);
+	NewArray.AddUnique(FromObject);
+
+	UE_CVLOG(NewRedierection != NULL, NewRedierection, LogVisual, Log, TEXT("Redirected '%s' to '%s'"), *FromObject->GetName(), *NewRedierection->GetName());
+}
+
+class UObject* FVisualLogger::FindRedirection(const class UObject* Object)
+{
+	for (auto& Redirection : RedirectionMap)
+	{
+		if (Redirection.Value.Find(Object) != INDEX_NONE)
+		{
+			return FindRedirection(Redirection.Key);
+		}
+	}
+
+	return const_cast<class UObject*>(Object);
+}
+
+void FVisualLogger::SetIsRecording(bool InIsRecording) 
+{ 
+	if (IsRecordingToFile())
+	{
+		SetIsRecordingToFile(false);
+	}
+	bIsRecording = InIsRecording; 
+};
+
+void FVisualLogger::SetIsRecordingToFile(bool InIsRecording)
+{
+	if (!bIsRecording && InIsRecording)
+	{
+		SetIsRecording(true);
+	}
+
+	FString BaseFileName = TEXT("VisualLog");
+	FString UserDefinedName;
+	FString MapName = GWorld ? GWorld->GetMapName() : TEXT("");
+	if (LogFileNameGetter.IsBound())
+	{
+		UserDefinedName = LogFileNameGetter.Execute();
+	}
+	if (UserDefinedName.Len() > 0)
+	{
+		if (MapName.Len() > 0)
+		{
+			BaseFileName = FString::Printf(TEXT("%s_%s"), *UserDefinedName, *MapName);
+		}
+		else
+		{
+			BaseFileName = UserDefinedName;
+		}
+	}
+
+	if (bIsRecordingToFile && !InIsRecording)
+	{
+		for (auto* Device : OutputDevices)
+		{
+			if (Device->HasFlags(VisualLogger::CanSaveToFile))
+			{
+				Device->StopRecordingToFile(GWorld ? GWorld->TimeSeconds : StartRecordingToFileTime);
+			}
+		}
+	}
+	else if (!bIsRecordingToFile && InIsRecording)
+	{
+		StartRecordingToFileTime = GWorld ? GWorld->TimeSeconds : 0;
+		for (auto* Device : OutputDevices)
+		{
+			if (Device->HasFlags(VisualLogger::CanSaveToFile))
+			{
+				Device->SetFileName(BaseFileName);
+				Device->StartRecordingToFile(StartRecordingToFileTime);
+			}
+		}
+	}
+
+	bIsRecordingToFile = InIsRecording;
+}
+
+#endif //ENABLE_VISUAL_LOG
\ No newline at end of file
diff --git a/Engine/Source/Runtime/Engine/Private/VisualLogger/VisualLoggerAutomationTests.cpp b/Engine/Source/Runtime/Engine/Private/VisualLogger/VisualLoggerAutomationTests.cpp
new file mode 100644
--- /dev/null
+++ b/Engine/Source/Runtime/Engine/Private/VisualLogger/VisualLoggerAutomationTests.cpp
@@ -0,0 +1,177 @@
+// Copyright 1998-2014 Epic Games, Inc. All Rights Reserved.
+
+#include "EnginePrivate.h"
+
+#include "VisualLog.h"
+#include "VisualLogger/VisualLogger.h"
+#include "VisualLogger/VisualLoggerAutomationTests.h"
+
+UVisualLoggerAutomationTests::UVisualLoggerAutomationTests(const class FPostConstructInitializeProperties& PCIP) : Super(PCIP)
+{
+
+}
+
+
+#if ENABLE_VISUAL_LOG
+
+class FVisualLoggerTestDevice : public FVisualLogDevice
+{
+public:
+	FVisualLoggerTestDevice();
+	virtual void Cleanup(bool bReleaseMemory = false) override;
+	virtual void StartRecordingToFile(float TImeStamp) override;
+	virtual void StopRecordingToFile(float TImeStamp) override;
+	virtual void SetFileName(const FString& InFileName) override;
+	virtual void Serialize(const class UObject* LogOwner, const FVisualLogEntry& LogEntry) override;
+
+	class UObject* LastObject;
+	FVisualLogEntry LastEntry;
+};
+
+FVisualLoggerTestDevice::FVisualLoggerTestDevice()
+{
+	Cleanup();
+}
+
+void FVisualLoggerTestDevice::Cleanup(bool bReleaseMemory)
+{
+	LastObject = NULL;
+	LastEntry.Reset();
+}
+
+void FVisualLoggerTestDevice::StartRecordingToFile(float TImeStamp)
+{
+
+}
+
+void FVisualLoggerTestDevice::StopRecordingToFile(float TImeStamp)
+{
+
+}
+
+void FVisualLoggerTestDevice::SetFileName(const FString& InFileName)
+{
+
+}
+
+void FVisualLoggerTestDevice::Serialize(const class UObject* LogOwner, const FVisualLogEntry& LogEntry)
+{
+	LastObject = const_cast<class UObject*>(LogOwner);
+	LastEntry = LogEntry;
+}
+
+#define CHECK_SUCCESS(__Test__) \
+if (!(__Test__)) \
+{ \
+	TestTrue(FString::Printf( TEXT("%s (%s:%d)"), TEXT(#__Test__), TEXT(__FILE__), __LINE__ ), __Test__); \
+	return false; \
+}
+#define CHECK_FAIL(__Test__) \
+if ((__Test__)) \
+{ \
+	TestFalse(FString::Printf( TEXT("%s (%s:%d)"), TEXT(#__Test__), TEXT(__FILE__), __LINE__ ), __Test__); \
+	return false; \
+}
+
+struct FTestDeviceContext
+{
+	FTestDeviceContext() 
+	{
+		Device.Cleanup();
+		FVisualLogger::Get().SetIsRecording(false);
+		FVisualLogger::Get().Cleanup();
+		FVisualLogger::Get().AddDevice(&Device);
+	}
+
+	~FTestDeviceContext() 
+	{
+		FVisualLogger::Get().SetIsRecording(false);
+		FVisualLogger::Get().RemoveDevice(&Device);
+		FVisualLogger::Get().Cleanup();
+		Device.Cleanup();
+	}
+
+	FVisualLoggerTestDevice Device;
+};
+
+IMPLEMENT_SIMPLE_AUTOMATION_TEST(FVisualLogTest, "Engine.VisualLogger.Logging simple text", EAutomationTestFlags::ATF_Game | EAutomationTestFlags::ATF_Editor | EAutomationTestFlags::ATF_Commandlet)
+/**
+*
+*
+* @param Parameters - Unused for this test
+* @return	TRUE if the test was successful, FALSE otherwise
+*/
+
+bool FVisualLogTest::RunTest(const FString& Parameters)
+{
+	FTestDeviceContext Context;
+
+	FVisualLogger::Get().SetIsRecording(false);
+	CHECK_FAIL(FVisualLogger::Get().IsRecording());
+
+	UE_VLOG(GWorld, LogVisual, Log, TEXT("Hello World"));
+	CHECK_SUCCESS(Context.Device.LastObject == NULL);
+	CHECK_SUCCESS(Context.Device.LastEntry.TimeStamp == -1);
+
+	FVisualLogger::Get().SetIsRecording(true);
+	CHECK_SUCCESS(FVisualLogger::Get().IsRecording());
+
+	{
+		const FString TextToLog = TEXT("Hello World");
+		UE_VLOG(GWorld, LogVisual, Log, *TextToLog);
+		CHECK_SUCCESS(Context.Device.LastObject != GWorld);
+		CHECK_SUCCESS(Context.Device.LastEntry.TimeStamp == -1);
+		FVisualLogEntry* CurrentEntry = FVisualLogger::Get().GetEntryToWrite(GWorld, GWorld->TimeSeconds, VisualLogger::DontCreate);
+		CHECK_SUCCESS(CurrentEntry != NULL);
+		CHECK_SUCCESS(CurrentEntry->TimeStamp == GWorld->TimeSeconds);
+		CHECK_SUCCESS(CurrentEntry->LogLines.Num() == 1);
+		CHECK_SUCCESS(CurrentEntry->LogLines[0].Category == LogVisual.GetCategoryName());
+		CHECK_SUCCESS(CurrentEntry->LogLines[0].Line == TextToLog);
+
+		FVisualLogger::Get().GetEntryToWrite(GWorld, GWorld->TimeSeconds+0.1); //generate new entry and serialize old one
+		CHECK_SUCCESS(Context.Device.LastEntry.TimeStamp == GWorld->TimeSeconds);
+		CHECK_SUCCESS(Context.Device.LastObject == GWorld);
+		CHECK_SUCCESS(Context.Device.LastEntry.LogLines.Num() == 1);
+		CHECK_SUCCESS(Context.Device.LastEntry.LogLines[0].Category == LogVisual.GetCategoryName());
+		CHECK_SUCCESS(Context.Device.LastEntry.LogLines[0].Line == TextToLog);
+	}
+
+	FVisualLogger::Get().SetIsRecording(false);
+	CHECK_FAIL(FVisualLogger::Get().IsRecording());
+
+	return true;
+}
+
+IMPLEMENT_SIMPLE_AUTOMATION_TEST(FVisualLogSegmentsTest, "Engine.VisualLogger.Logging segment shape", EAutomationTestFlags::ATF_Game | EAutomationTestFlags::ATF_Editor | EAutomationTestFlags::ATF_Commandlet)
+bool FVisualLogSegmentsTest::RunTest(const FString& Parameters)
+{
+	FTestDeviceContext Context;
+
+	FVisualLogger::Get().SetIsRecording(false);
+	CHECK_FAIL(FVisualLogger::Get().IsRecording());
+
+	UE_VLOG_SEGMENT(GWorld, LogVisual, Log, FVector(0, 0, 0), FVector(1, 0, 0), FColor::Red, TEXT("Simple segment log"));
+	CHECK_SUCCESS(Context.Device.LastObject == NULL);
+	CHECK_SUCCESS(Context.Device.LastEntry.TimeStamp == -1);
+
+	FVisualLogger::Get().SetIsRecording(true);
+	CHECK_SUCCESS(FVisualLogger::Get().IsRecording());
+	{
+		const FVector StartPoint(0, 0, 0);
+		const FVector EndPoint(1, 0, 0);
+		UE_VLOG_SEGMENT(GWorld, LogVisual, Log, StartPoint, EndPoint, FColor::Red, TEXT("Simple segment log"));
+		CHECK_SUCCESS(Context.Device.LastObject == NULL);
+		CHECK_SUCCESS(Context.Device.LastEntry.TimeStamp == -1);
+		FVisualLogEntry* CurrentEntry = FVisualLogger::Get().GetEntryToWrite(GWorld, GWorld->TimeSeconds, VisualLogger::DontCreate);
+		CHECK_SUCCESS(CurrentEntry != NULL);
+		CHECK_SUCCESS(CurrentEntry->TimeStamp == GWorld->TimeSeconds);
+		CHECK_SUCCESS(CurrentEntry->ElementsToDraw.Num() == 1);
+		CHECK_SUCCESS(CurrentEntry->ElementsToDraw[0].GetType() == FVisualLogEntry::FElementToDraw::Segment);
+		CHECK_SUCCESS(CurrentEntry->ElementsToDraw[0].Points.Num() == 2);
+		CHECK_SUCCESS(CurrentEntry->ElementsToDraw[0].Points[0] == StartPoint);
+		CHECK_SUCCESS(CurrentEntry->ElementsToDraw[0].Points[1] == EndPoint);
+	}
+	return true;
+}
+
+#endif //ENABLE_VISUAL_LOG
\ No newline at end of file
diff --git a/Engine/Source/Runtime/Engine/Private/VisualLogger/VisualLoggerKismetLibrary.cpp b/Engine/Source/Runtime/Engine/Private/VisualLogger/VisualLoggerKismetLibrary.cpp
new file mode 100644
--- /dev/null
+++ b/Engine/Source/Runtime/Engine/Private/VisualLogger/VisualLoggerKismetLibrary.cpp
@@ -0,0 +1,35 @@
+// Copyright 1998-2014 Epic Games, Inc. All Rights Reserved.
+
+#include "EnginePrivate.h"
+#include "VisualLog.h"
+#include "VisualLogger/VisualLogger.h"
+#include "VisualLogger/VisualLoggerKismetLibrary.h"
+
+UVisualLoggerKismetLibrary::UVisualLoggerKismetLibrary(const class FPostConstructInitializeProperties& PCIP)
+	: Super(PCIP)
+{
+}
+
+void UVisualLoggerKismetLibrary::LogText(UObject* WorldContextObject, FString Text, FName CategoryName)
+{
+#if ENABLE_VISUAL_LOG
+	const ELogVerbosity::Type DefaultVerbosity = ELogVerbosity::Log;
+	FVisualLogger::CategorizedLogf(WorldContextObject, FLogCategoryBase(*CategoryName.ToString(), DefaultVerbosity, DefaultVerbosity), DefaultVerbosity, INDEX_NONE, *Text);
+#endif
+}
+
+void UVisualLoggerKismetLibrary::LogLocation(UObject* WorldContextObject, FVector Location, FString Text, FLinearColor Color, float Radius, FName CategoryName)
+{
+#if ENABLE_VISUAL_LOG
+	const ELogVerbosity::Type DefaultVerbosity = ELogVerbosity::Log;
+	FVisualLogger::GeometryShapeLogf(WorldContextObject, FLogCategoryBase(*CategoryName.ToString(), DefaultVerbosity, DefaultVerbosity), DefaultVerbosity, INDEX_NONE, Location, Radius, Color, *Text);
+#endif
+}
+
+void UVisualLoggerKismetLibrary::LogBox(UObject* WorldContextObject, FBox Box, FString Text, FLinearColor ObjectColor, FName CategoryName)
+{
+#if ENABLE_VISUAL_LOG
+	const ELogVerbosity::Type DefaultVerbosity = ELogVerbosity::Log;
+	FVisualLogger::GeometryShapeLogf(WorldContextObject, FLogCategoryBase(*CategoryName.ToString(), DefaultVerbosity, DefaultVerbosity), DefaultVerbosity, INDEX_NONE, Box, ObjectColor, *Text);
+#endif
+}
diff --git a/Engine/Source/Runtime/Engine/Private/VisualLogger/VisualLoggerTypes.cpp b/Engine/Source/Runtime/Engine/Private/VisualLogger/VisualLoggerTypes.cpp
new file mode 100644
--- /dev/null
+++ b/Engine/Source/Runtime/Engine/Private/VisualLogger/VisualLoggerTypes.cpp
@@ -0,0 +1,169 @@
+// Copyright 1998-2014 Epic Games, Inc. All Rights Reserved.
+
+#include "EnginePrivate.h"
+#include "VisualLog.h"
+#include "VisualLogger/VisualLogger.h"
+
+#if ENABLE_VISUAL_LOG
+
+FVisualLogEntry::FVisualLogEntry(const FVisualLogEntry& Entry)
+{
+	TimeStamp = Entry.TimeStamp;
+	Location = Entry.Location;
+
+	LogLines = Entry.LogLines;
+	Status = Entry.Status;
+	ElementsToDraw = Entry.ElementsToDraw;
+	HistogramSamples = Entry.HistogramSamples;
+	DataBlocks = Entry.DataBlocks;
+}
+
+FVisualLogEntry::FVisualLogEntry(const class AActor* InActor, TArray<TWeakObjectPtr<UObject> >* Children)
+{
+	if (InActor && InActor->IsPendingKill() == false)
+	{
+		TimeStamp = InActor->GetWorld()->TimeSeconds;
+		Location = InActor->GetActorLocation();
+		InActor->GrabDebugSnapshot(this);
+		if (Children != NULL)
+		{
+			TWeakObjectPtr<UObject>* WeakActorPtr = Children->GetData();
+			for (int32 Index = 0; Index < Children->Num(); ++Index, ++WeakActorPtr)
+			{
+				if (WeakActorPtr->IsValid())
+				{
+					const AActor* ChildActor = Cast<AActor>(WeakActorPtr->Get());
+					if (ChildActor)
+					{
+						ChildActor->GrabDebugSnapshot(this);
+					}
+				}
+			}
+		}
+	}
+}
+
+FVisualLogEntry::FVisualLogEntry(float InTimeStamp, FVector InLocation, const UObject* Object, TArray<TWeakObjectPtr<UObject> >* Children)
+{
+	TimeStamp = InTimeStamp;
+	Location = InLocation;
+	const AActor* AsActor = Cast<AActor>(Object);
+	if (AsActor)
+	{
+		AsActor->GrabDebugSnapshot(this);
+	}
+	if (Children != NULL)
+	{
+		TWeakObjectPtr<UObject>* WeakActorPtr = Children->GetData();
+		for (int32 Index = 0; Index < Children->Num(); ++Index, ++WeakActorPtr)
+		{
+			if (WeakActorPtr->IsValid())
+			{
+				const AActor* ChildActor = Cast<AActor>(WeakActorPtr->Get());
+				if (ChildActor)
+				{
+					ChildActor->GrabDebugSnapshot(this);
+				}
+			}
+		}
+	}
+}
+
+void FVisualLogEntry::AddText(const FString& TextLine, const FName& CategoryName)
+{
+	LogLines.Add(FLogLine(CategoryName, ELogVerbosity::All, TextLine));
+}
+
+void FVisualLogEntry::AddElement(const TArray<FVector>& Points, const FName& CategoryName, const FColor& Color, const FString& Description, uint16 Thickness)
+{
+	FElementToDraw Element(Description, Color, Thickness, CategoryName);
+	Element.Points = Points;
+	Element.Type = FElementToDraw::Path;
+	ElementsToDraw.Add(Element);
+}
+
+void FVisualLogEntry::AddElement(const FVector& Point, const FName& CategoryName, const FColor& Color, const FString& Description, uint16 Thickness)
+{
+	FElementToDraw Element(Description, Color, Thickness, CategoryName);
+	Element.Points.Add(Point);
+	Element.Type = FElementToDraw::SinglePoint;
+	ElementsToDraw.Add(Element);
+}
+
+void FVisualLogEntry::AddElement(const FVector& Start, const FVector& End, const FName& CategoryName, const FColor& Color, const FString& Description, uint16 Thickness)
+{
+	FElementToDraw Element(Description, Color, Thickness, CategoryName);
+	Element.Points.Reserve(2);
+	Element.Points.Add(Start);
+	Element.Points.Add(End);
+	Element.Type = FElementToDraw::Segment;
+	ElementsToDraw.Add(Element);
+}
+
+void FVisualLogEntry::AddElement(const FBox& Box, const FName& CategoryName, const FColor& Color, const FString& Description, uint16 Thickness)
+{
+	FElementToDraw Element(Description, Color, Thickness, CategoryName);
+	Element.Points.Reserve(2);
+	Element.Points.Add(Box.Min);
+	Element.Points.Add(Box.Max);
+	Element.Type = FElementToDraw::Box;
+	ElementsToDraw.Add(Element);
+}
+
+void FVisualLogEntry::AddElement(const FVector& Orgin, const FVector& Direction, float Length, float AngleWidth, float AngleHeight, const FName& CategoryName, const FColor& Color, const FString& Description, uint16 Thickness)
+{
+	FElementToDraw Element(Description, Color, Thickness, CategoryName);
+	Element.Points.Reserve(3);
+	Element.Points.Add(Orgin);
+	Element.Points.Add(Direction);
+	Element.Points.Add(FVector(Length, AngleWidth, AngleHeight));
+	Element.Type = FElementToDraw::Cone;
+	ElementsToDraw.Add(Element);
+}
+
+void FVisualLogEntry::AddElement(const FVector& Start, const FVector& End, float Radius, const FName& CategoryName, const FColor& Color, const FString& Description, uint16 Thickness)
+{
+	FElementToDraw Element(Description, Color, Thickness, CategoryName);
+	Element.Points.Reserve(3);
+	Element.Points.Add(Start);
+	Element.Points.Add(End);
+	Element.Points.Add(FVector(Radius, Thickness, 0));
+	Element.Type = FElementToDraw::Cylinder;
+	ElementsToDraw.Add(Element);
+}
+
+void FVisualLogEntry::AddHistogramData(const FVector2D& DataSample, const FName& CategoryName, const FName& GraphName, const FName& DataName)
+{
+	FHistogramSample Sample;
+	Sample.Category = CategoryName;
+	Sample.GraphName = GraphName;
+	Sample.DataName = DataName;
+	Sample.SampleValue = DataSample;
+
+	HistogramSamples.Add(Sample);
+}
+
+FVisualLogEntry::FDataBlock& FVisualLogEntry::AddDataBlock(const FString& TagName, const TArray<uint8>& BlobDataArray, const FName& CategoryName)
+{
+	FDataBlock DataBlock;
+	DataBlock.Category = CategoryName;
+	DataBlock.TagName = *TagName;
+	DataBlock.Data = BlobDataArray;
+
+	const int32 Index = DataBlocks.Add(DataBlock);
+	return DataBlocks[Index];
+}
+
+void FVisualLogEntry::AddCapsule(FVector const& Center, float HalfHeight, float Radius, const FQuat & Rotation, const FName& CategoryName, const FColor& Color, const FString& Description)
+{
+	FElementToDraw Element(Description, Color, 0, CategoryName);
+	Element.Points.Reserve(3);
+	Element.Points.Add(Center);
+	Element.Points.Add(FVector(HalfHeight, Radius, Rotation.X));
+	Element.Points.Add(FVector(Rotation.Y, Rotation.Z, Rotation.W));
+	Element.Type = FElementToDraw::Capsule;
+	ElementsToDraw.Add(Element);
+}
+
+#endif //ENABLE_VISUAL_LOG
+
diff --git a/Engine/Source/Runtime/Engine/Private/World.cpp b/Engine/Source/Runtime/Engine/Private/World.cpp
--- a/Engine/Source/Runtime/Engine/Private/World.cpp
+++ b/Engine/Source/Runtime/Engine/Private/World.cpp
@@ -3104,7 +3104,7 @@ void UWorld::CleanupWorld(bool bSessionEnded, bool bCleanupResources, UWorld* Ne
 #endif //WITH_EDITOR
 
 #if ENABLE_VISUAL_LOG
-	FVisualLog::Get().Cleanup();
+	FVisualLogger::Get().Cleanup();
 #endif // ENABLE_VISUAL_LOG	
 
 	// Tell actors to remove their components from the scene.
diff --git a/Engine/Source/Runtime/Engine/Public/VisualLog.h b/Engine/Source/Runtime/Engine/Public/VisualLog.h
--- a/Engine/Source/Runtime/Engine/Public/VisualLog.h
+++ b/Engine/Source/Runtime/Engine/Public/VisualLog.h
@@ -1,6 +1,8 @@
 // Copyright 1998-2014 Epic Games, Inc. All Rights Reserved.
 #pragma once
 
+#include "VisualLogger/VisualLogger.h"
+
 #define LOG_TO_UELOG_AS_WELL 0
 
 // helpful text macros
@@ -12,8 +14,6 @@
 #define TEXT_SELF_IF_TRUE(Self) (Self ? TEXT(#Self) : TEXT_EMPTY)
 #define TEXT_NAME(Object) ((Object != NULL) ? *(Object->GetName()) : TEXT_NULL)
 
-DECLARE_LOG_CATEGORY_EXTERN(LogVisual, Warning, All);
-
 class FJsonValue;
 
 struct FVisLogSelfDrawingElement
@@ -21,200 +21,6 @@ struct FVisLogSelfDrawingElement
 	virtual void Draw(class UCanvas* Canvas) const = 0;
 };
 
-struct ENGINE_API FVisLogEntry
-{
-	float TimeStamp;
-#if ENABLE_VISUAL_LOG
-	FVector Location;
-	
-	struct FLogLine
-	{
-		FString Line;
-		FName Category;
-		TEnumAsByte<ELogVerbosity::Type> Verbosity;
-		int64 UserData;
-		FName TagName;
-
-		FLogLine(const FName& InCategory, ELogVerbosity::Type InVerbosity, const FString& InLine)
-			: Line(InLine), Category(InCategory), Verbosity(InVerbosity), UserData(0)
-		{}
-
-		FLogLine(const FName& InCategory, ELogVerbosity::Type InVerbosity, const FString& InLine, int64 InUserData)
-			: Line(InLine), Category(InCategory), Verbosity(InVerbosity), UserData(InUserData)
-		{}
-	};
-	TArray<FLogLine> LogLines;
-	
-	struct FStatusCategory
-	{
-		TArray<FString> Data;
-		FString Category;
-
-		FORCEINLINE void Add(const FString& Key, const FString& Value)
-		{
-			Data.Add(FString(Key).AppendChar(TEXT('|')) + Value);
-		}
-
-		FORCEINLINE bool GetDesc(int32 Index, FString& Key, FString& Value) const
-		{
-			if (Data.IsValidIndex(Index))
-			{
-				int32 SplitIdx = INDEX_NONE;
-				if (Data[Index].FindChar(TEXT('|'), SplitIdx))
-				{
-					Key = Data[Index].Left(SplitIdx);
-					Value = Data[Index].Mid(SplitIdx + 1);
-					return true;
-				}
-			}
-
-			return false;
-		}
-	};
-	TArray<FStatusCategory> Status;
-
-	struct FElementToDraw
-	{
-		enum EElementType {
-			Invalid = 0,
-			SinglePoint, // individual points. 
-			Segment, // pairs of points 
-			Path,	// sequence of point
-			Box,
-			Cone,
-			Cylinder,
-			Capsule,
-			// note that in order to remain backward compatibility in terms of log
-			// serialization new enum values need to be added at the end
-		};
-
-		FString Description;
-		FName Category;
-		TEnumAsByte<ELogVerbosity::Type> Verbosity;
-		TArray<FVector> Points;
-	protected:
-		friend struct FVisLogEntry;
-		uint8 Type;
-	public:
-		uint8 Color;
-		union
-		{
-			uint16 Thicknes;
-			uint16 Radius;
-		};
-
-		FElementToDraw(EElementType InType = Invalid) : Verbosity(ELogVerbosity::All), Type(InType), Color(0xff), Thicknes(0)
-		{}
-
-		FElementToDraw(const FString& InDescription, const FColor& InColor, uint16 InThickness, const FName& InCategory) : Category(InCategory), Verbosity(ELogVerbosity::All), Type(uint16(Invalid)), Thicknes(InThickness)
-		{
-			if (InDescription.IsEmpty() == false)
-			{
-				Description = InDescription;
-			}
-			SetColor(InColor);
-		}
-
-		FORCEINLINE void SetColor(const FColor& InColor) 
-		{ 
-			Color = ((InColor.DWColor() >> 30) << 6)
-				| (((InColor.DWColor() & 0x00ff0000) >> 22) << 4)
-				| (((InColor.DWColor() & 0x0000ff00) >> 14) << 2)
-				|  ((InColor.DWColor() & 0x000000ff) >> 6);
-		}
-
-		FORCEINLINE EElementType GetType() const
-		{
-			return EElementType(Type);
-		}
-
-		FORCEINLINE void SetType(EElementType InType)
-		{
-			Type = InType;
-		}
-
-		FORCEINLINE FColor GetFColor() const
-		{ 
-			return FColor(
-				((Color & 0xc0) << 24)
-				| ((Color & 0x30) << 18)
-				| ((Color & 0x0c) << 12)
-				| ((Color & 0x03) << 6)
-				);
-		}
-	};
-	TArray<FElementToDraw> ElementsToDraw;
-
-	struct FHistogramSample
-	{
-		FHistogramSample() : Verbosity(ELogVerbosity::All) {}
-
-		FName Category;
-		TEnumAsByte<ELogVerbosity::Type> Verbosity;
-		FName GraphName;
-		FName DataName;
-		FVector2D SampleValue;
-	};
-	TArray<FHistogramSample>	HistogramSamples;
-
-	struct FDataBlock
-	{
-		FDataBlock() : Verbosity(ELogVerbosity::All) {}
-
-		FName TagName;
-		FName Category;
-		TEnumAsByte<ELogVerbosity::Type> Verbosity;
-		TArray<uint8>	Data;
-	};
-	TArray<FDataBlock>	DataBlocks;
-
-
-	FVisLogEntry(const class AActor* Actor, TArray<TWeakObjectPtr<UObject> >* Children);
-	FVisLogEntry(float TimeStamp, FVector Location, const UObject* Object, TArray<TWeakObjectPtr<UObject> >* Children);
-	FVisLogEntry(TSharedPtr<FJsonValue> FromJson);
-	
-	TSharedPtr<FJsonValue> ToJson() const;
-
-	// path
-	void AddElement(const TArray<FVector>& Points, const FName& CategoryName, const FColor& Color = FColor::White, const FString& Description = TEXT(""), uint16 Thickness = 0);
-	// location
-	void AddElement(const FVector& Point, const FName& CategoryName, const FColor& Color = FColor::White, const FString& Description = TEXT(""), uint16 Thickness = 0);
-	// segment
-	void AddElement(const FVector& Start, const FVector& End, const FName& CategoryName, const FColor& Color = FColor::White, const FString& Description = TEXT(""), uint16 Thickness = 0);
-	// box
-	void AddElement(const FBox& Box, const FName& CategoryName, const FColor& Color = FColor::White, const FString& Description = TEXT(""), uint16 Thickness = 0);
-
-	// Cone
-	void AddElement(const FVector& Orgin, const FVector& Direction, float Length, float AngleWidth, float AngleHeight, const FName& CategoryName, const FColor& Color = FColor::White, const FString& Description = TEXT(""), uint16 Thickness = 0);
-
-	// Cylinder
-	void AddElement(const FVector& Start, const FVector& End, float Radius, const FName& CategoryName, const FColor& Color = FColor::White, const FString& Description = TEXT(""), uint16 Thickness = 0);
-
-	// histogram sample
-	void AddHistogramData(const FVector2D& DataSample, const FName& CategoryName, const FName& GraphName, const FName& DataName);
-
-	// Custom data block
-	void AddDataBlock(const FString& TagName, const TArray<uint8>& BlobDataArray, const FName& CategoryName);
-
-	// capsule
-	void AddCapsule(FVector const& Center, float HalfHeight, float Radius, const FQuat& Rotation, const FName& CategoryName, const FColor& Color = FColor::White, const FString& Description = TEXT(""));
-
-	// find index of status category
-	FORCEINLINE int32 FindStatusIndex(const FString& CategoryName)
-	{
-		for (int32 TestCategoryIndex = 0; TestCategoryIndex < Status.Num(); TestCategoryIndex++)
-		{
-			if (Status[TestCategoryIndex].Category == CategoryName)
-			{
-				return TestCategoryIndex;
-			}
-		}
-
-		return INDEX_NONE;
-	}
-#endif // ENABLE_VISUAL_LOG
-};
-
 #if ENABLE_VISUAL_LOG
 
 #if LOG_TO_UELOG_AS_WELL
@@ -223,120 +29,10 @@ struct ENGINE_API FVisLogEntry
 #define TO_TEXT_LOG(CategoryName, Verbosity, Format, ...) 
 #endif // LOG_TO_UELOG_AS_WELL
 
-#define REDIRECT_TO_VLOG(Dest) FVisualLog::Get().RedirectToVisualLog(this, Dest)
-#define REDIRECT_OBJECT_TO_VLOG(Src, Dest) FVisualLog::Get().RedirectToVisualLog(Src, Dest)
-
-#define CONNECT_WITH_VLOG(Dest)
-#define CONNECT_OBJECT_WITH_VLOG(Src, Dest)
-
-#define UE_VLOG(Object, CategoryName, Verbosity, Format, ...) \
-{ \
-	SCOPE_CYCLE_COUNTER(STAT_VisualLog); \
-	static_assert((ELogVerbosity::Verbosity & ELogVerbosity::VerbosityMask) < ELogVerbosity::NumVerbosity && ELogVerbosity::Verbosity > 0, "Verbosity must be constant and in range."); \
-	if (FVisualLog::Get().IsRecording() && (!FVisualLog::Get().IsAllBlocked() || FVisualLog::Get().InWhitelist(CategoryName.GetCategoryName())) && Object && !Object->HasAnyFlags(RF_ClassDefaultObject)) \
-	{ \
-		if (Object) \
-		{  \
-			FVisualLog::Get().LogLine(Object, CategoryName.GetCategoryName(), ELogVerbosity::Verbosity, FString::Printf(Format, ##__VA_ARGS__)); \
-		} \
-	} \
-	if (UE_LOG_CHECK_COMPILEDIN_VERBOSITY(CategoryName, Verbosity)) \
-	{ \
-		if (!CategoryName.IsSuppressed(ELogVerbosity::Verbosity) && (Object) != NULL) \
-		{ \
-			TO_TEXT_LOG(CategoryName, Verbosity, Format, ##__VA_ARGS__); \
-		} \
-	} \
-}
-
-#define UE_CVLOG(Condition, Object, CategoryName, Verbosity, Format, ...) if(Condition) \
-{ \
-	UE_VLOG(Object, CategoryName, Verbosity, Format, ##__VA_ARGS__); \
-}
-
-#define UE_VLOG_SEGMENT_THICK(Object, CategoryName, Verbosity, SegmentStart, SegmentEnd, Color, Thickness, DescriptionFormat, ...) \
-{ \
-	SCOPE_CYCLE_COUNTER(STAT_VisualLog); \
-	static_assert((ELogVerbosity::Verbosity & ELogVerbosity::VerbosityMask) < ELogVerbosity::NumVerbosity && ELogVerbosity::Verbosity > 0, "Verbosity must be constant and in range."); \
-	if (FVisualLog::Get().IsRecording() && (!FVisualLog::Get().IsAllBlocked() || FVisualLog::Get().InWhitelist(CategoryName.GetCategoryName())) && Object && !Object->HasAnyFlags(RF_ClassDefaultObject)) \
-	{ \
-		if (Object) \
-		{  \
-			FVisualLog::Get().GetEntryToWrite(Object)->AddElement(SegmentStart, SegmentEnd, CategoryName.GetCategoryName(), Color, FString::Printf(DescriptionFormat, ##__VA_ARGS__), Thickness); \
-		} \
-	} \
-}
-
-#define UE_VLOG_SEGMENT(Object, CategoryName, Verbosity, SegmentStart, SegmentEnd, Color, DescriptionFormat, ...) \
-	UE_VLOG_SEGMENT_THICK(Object, CategoryName, Verbosity, SegmentStart, SegmentEnd, Color, 0, DescriptionFormat, ##__VA_ARGS__)
-
-#define UE_VLOG_LOCATION(Object, CategoryName, Verbosity, Location, Radius, Color, DescriptionFormat, ...) \
-{ \
-	SCOPE_CYCLE_COUNTER(STAT_VisualLog); \
-	static_assert((ELogVerbosity::Verbosity & ELogVerbosity::VerbosityMask) < ELogVerbosity::NumVerbosity && ELogVerbosity::Verbosity > 0, "Verbosity must be constant and in range."); \
-	if (FVisualLog::Get().IsRecording() && (!FVisualLog::Get().IsAllBlocked() || FVisualLog::Get().InWhitelist(CategoryName.GetCategoryName())) && Object && !Object->HasAnyFlags(RF_ClassDefaultObject)) \
-	{ \
-		if (Object) \
-		{  \
-			FVisualLog::Get().GetEntryToWrite(Object)->AddElement(Location, CategoryName.GetCategoryName(), Color, FString::Printf(DescriptionFormat, ##__VA_ARGS__), Radius); \
-		} \
-	} \
-}
-
-#define UE_VLOG_BOX(Object, CategoryName, Verbosity, Box, Color, DescriptionFormat, ...) \
-{ \
-	SCOPE_CYCLE_COUNTER(STAT_VisualLog); \
-	static_assert((ELogVerbosity::Verbosity & ELogVerbosity::VerbosityMask) < ELogVerbosity::NumVerbosity && ELogVerbosity::Verbosity > 0, "Verbosity must be constant and in range."); \
-	if (FVisualLog::Get().IsRecording() && (!FVisualLog::Get().IsAllBlocked() || FVisualLog::Get().InWhitelist(CategoryName.GetCategoryName())) && Object && !Object->HasAnyFlags(RF_ClassDefaultObject)) \
-	{ \
-		if (Object) \
-		{  \
-			FVisualLog::Get().GetEntryToWrite(Object)->AddElement(Box, CategoryName.GetCategoryName(), Color, FString::Printf(DescriptionFormat, ##__VA_ARGS__)); \
-		} \
-	} \
-}
-
-#define UE_VLOG_CONE(Object, CategoryName, Verbosity, Orgin, Direction, Length, Angle, Color, DescriptionFormat, ...) \
-{ \
-	SCOPE_CYCLE_COUNTER(STAT_VisualLog); \
-	static_assert((ELogVerbosity::Verbosity & ELogVerbosity::VerbosityMask) < ELogVerbosity::NumVerbosity && ELogVerbosity::Verbosity > 0, "Verbosity must be constant and in range."); \
-	if (FVisualLog::Get().IsRecording() && (!FVisualLog::Get().IsAllBlocked() || FVisualLog::Get().InWhitelist(CategoryName.GetCategoryName())) && Object && !Object->HasAnyFlags(RF_ClassDefaultObject)) \
-			{ \
-		FVisualLog::Get().GetEntryToWrite(Object)->AddElement(Orgin, Direction, Length, Angle, Angle, CategoryName.GetCategoryName(), Color, FString::Printf(DescriptionFormat, ##__VA_ARGS__)); \
-			} \
-}
-
-#define UE_VLOG_CYLINDER(Object, CategoryName, Verbosity, Start, End, Radius, Color, DescriptionFormat, ...) \
-{ \
-	SCOPE_CYCLE_COUNTER(STAT_VisualLog); \
-	static_assert((ELogVerbosity::Verbosity & ELogVerbosity::VerbosityMask) < ELogVerbosity::NumVerbosity && ELogVerbosity::Verbosity > 0, "Verbosity must be constant and in range."); \
-	if (FVisualLog::Get().IsRecording() && (!FVisualLog::Get().IsAllBlocked() || FVisualLog::Get().InWhitelist(CategoryName.GetCategoryName())) && Object && !Object->HasAnyFlags(RF_ClassDefaultObject)) \
-	{ \
-		FVisualLog::Get().GetEntryToWrite(Object)->AddElement(Start, End, Radius, CategoryName.GetCategoryName(), Color, FString::Printf(DescriptionFormat, ##__VA_ARGS__)); \
-	} \
-}
-
-#define UE_VLOG_CAPSULE(Object, CategoryName, Verbosity, Center, HalfHeight, Radius, Rotation, Color, DescriptionFormat, ...) \
-{ \
-	SCOPE_CYCLE_COUNTER(STAT_VisualLog); \
-	static_assert((ELogVerbosity::Verbosity & ELogVerbosity::VerbosityMask) < ELogVerbosity::NumVerbosity && ELogVerbosity::Verbosity > 0, "Verbosity must be constant and in range."); \
-	if (FVisualLog::Get().IsRecording() && (!FVisualLog::Get().IsAllBlocked() || FVisualLog::Get().InWhitelist(CategoryName.GetCategoryName())) && Object && !Object->HasAnyFlags(RF_ClassDefaultObject)) \
-					{ \
-		FVisualLog::Get().GetEntryToWrite(Object)->AddCapsule(Center, HalfHeight, Radius, Rotation, CategoryName.GetCategoryName(), Color, FString::Printf(DescriptionFormat, ##__VA_ARGS__)); \
-					} \
-}
-
-#define UE_VLOG_HISTOGRAM(Object, CategoryName, Verbosity, GraphName, DataName, Data) \
-{ \
-	SCOPE_CYCLE_COUNTER(STAT_VisualLog); \
-	static_assert((ELogVerbosity::Verbosity & ELogVerbosity::VerbosityMask) < ELogVerbosity::NumVerbosity && ELogVerbosity::Verbosity > 0, "Verbosity must be constant and in range."); \
-	if (FVisualLog::Get().IsRecording() && (!FVisualLog::Get().IsAllBlocked() || FVisualLog::Get().InWhitelist(CategoryName.GetCategoryName())) && Object && !Object->HasAnyFlags(RF_ClassDefaultObject)) \
-	{ \
-		if (Object) \
-		{  \
-			FVisualLog::Get().GetEntryToWrite(Object)->AddHistogramData(Data, CategoryName.GetCategoryName(), GraphName, DataName); \
-		} \
-	} \
+namespace VisualLogger
+{
+	struct FDataBlock;
+	struct FVisualLogEntry;
 }
 
 namespace VisualLogJson
@@ -377,7 +73,7 @@ struct ENGINE_API FActorsVisLog
 
 	FName Name;
 	FString FullName;
-	TArray<TSharedPtr<FVisLogEntry> > Entries;
+	TArray<TSharedPtr<FVisualLogEntry> > Entries;
 
 	FActorsVisLog(const class UObject* Object, TArray<TWeakObjectPtr<UObject> >* Children);
 	FActorsVisLog(TSharedPtr<FJsonValue> FromJson);
@@ -387,17 +83,8 @@ struct ENGINE_API FActorsVisLog
 
 DECLARE_DELEGATE_RetVal(FString, FVisualLogFilenameGetterDelegate);
 
-struct FVisualLogExtensionInterface
-{
-	virtual void OnTimestampChange(float Timestamp, class UWorld* InWorld, class AActor* HelperActor) = 0;
-	virtual void DrawData(class UWorld* InWorld, class UCanvas* Canvas, class AActor* HelperActor, const FName& TagName, const FVisLogEntry::FDataBlock& DataBlock, float Timestamp) = 0;
-	virtual void DisableDrawingForData(class UWorld* InWorld, class UCanvas* Canvas, class AActor* HelperActor, const FName& TagName, const FVisLogEntry::FDataBlock& DataBlock, float Timestamp) = 0;
-	virtual void LogEntryLineSelectionChanged(TSharedPtr<struct FLogEntryItem> SelectedItem, int64 UserData, FName TagName) = 0;
-};
-
-
 /** This class is to capture all log output even if the Visual Logger is closed */
-class ENGINE_API FVisualLog : public FOutputDevice
+class ENGINE_API FVisualLog : public FVisualLogDevice
 {
 public:
 	typedef TMap<const class UObject*, TSharedPtr<FActorsVisLog> > FLogsMap;
@@ -408,13 +95,17 @@ public:
 	FVisualLog();
 	~FVisualLog();
 	
-	static FVisualLog& Get()
+	static FVisualLog& GetStatic()
 	{
-		static FVisualLog StaticLog;
-		return StaticLog;
+		static FVisualLog VisLog;
+		return VisLog;
 	}
 
-	void Cleanup(bool bReleaseMemory = false);
+	virtual void Cleanup(bool bReleaseMemory) override;
+	virtual void StartRecordingToFile(float TImeStamp) override;
+	virtual void StopRecordingToFile(float TImeStamp) override;
+	virtual void Serialize(const class UObject* LogOwner, const FVisualLogEntry& LogEntry) override;
+	virtual void SetFileName(const FString& InFileName) override;
 
 	void Redirect(class UObject* Source, const class AActor* NewRedirection);
 	void RedirectToVisualLog(const class UObject* Src, const class AActor* Dest);
@@ -428,7 +119,7 @@ public:
 	void ClearNewLogsObserver() { OnNewLogCreated = NULL; }
 
 	void SetIsRecording(bool NewRecording, bool bRecordToFile = false);
-	FORCEINLINE bool IsRecording() const { return !!bIsRecording; }
+	FORCEINLINE bool IsRecording() const { return true; }
 	void SetIsRecordingOnServer(bool NewRecording) { bIsRecordingOnServer = NewRecording; }
 	FORCEINLINE bool IsRecordingOnServer() const { return !!bIsRecordingOnServer; }
 	void DumpRecordedLogs();
@@ -441,22 +132,15 @@ public:
 
 	FArchive* FileAr;
 
-	FVisLogEntry* GetEntryToWrite(const class UObject* Object);
+	FVisualLogEntry* GetEntryToWrite(const class UObject* Object, float Timestamp = -1, const FVector& Location = FVector::ZeroVector);
 
 	/** highly encouradged to use FVisualLogFilenameGetterDelegate::CreateUObject with this */
 	void SetLogFileNameGetter(const FVisualLogFilenameGetterDelegate& InLogFileNameGetter) { LogFileNameGetter = InLogFileNameGetter; }
 
-	TMap<FName, FVisualLogExtensionInterface*> AllExtensions;
-	void RegisterExtension(FName TagName, FVisualLogExtensionInterface* ExtensionInterface) { AllExtensions.Add(TagName, ExtensionInterface); }
-	void UnregisterExtension(FName TagName, FVisualLogExtensionInterface* ExtensionInterface) { AllExtensions.Remove(TagName); }
-	FVisualLogExtensionInterface* GetExtensionForTag(const FName TagName) { return AllExtensions.Contains(TagName) ? AllExtensions[TagName] : NULL; }
-	TMap<FName, FVisualLogExtensionInterface*>& GetAllExtensions() { return AllExtensions; }
 protected:
 
 	FString GetLogFileFullName() const;
 
-	virtual void Serialize( const TCHAR* V, ELogVerbosity::Type Verbosity, const class FName& Category ) override;
-
 	void LogElementImpl(FVisLogSelfDrawingElement* Element);
 
 	FORCEINLINE_DEBUGGABLE TSharedPtr<FActorsVisLog> GetLog(const class UObject* Object)
@@ -484,27 +168,12 @@ private:
 
 	float StartRecordingTime;
 
-	int32 bIsRecording : 1;
+	FString BaseFileName;
+
 	int32 bIsRecordingOnServer : 1;
 	int32 bIsRecordingToFile : 1;
 	int32 bIsAllBlocked : 1;
 
 	FVisualLogFilenameGetterDelegate LogFileNameGetter;
 };
-
-#else
-	#define UE_VLOG(Actor, CategoryName, Verbosity, Format, ...) UE_LOG(CategoryName, Verbosity, Format, ##__VA_ARGS__)
-	#define UE_CVLOG(Condition, Actor, CategoryName, Verbosity, Format, ...) UE_CLOG(Condition, CategoryName, Verbosity, Format, ##__VA_ARGS__)
-	#define UE_VLOG_SEGMENT(Actor, CategoryName, Verbosity, SegmentStart, SegmentEnd, Color, DescriptionFormat, ...)
-	#define UE_VLOG_LOCATION(Actor, CategoryName, Verbosity, Location, Radius, Color, DescriptionFormat, ...)
-	#define UE_VLOG_BOX(Actor, CategoryName, Verbosity, Box, Color, DescriptionFormat, ...) 
-	#define UE_VLOG_HISTOGRAM(Actor, CategoryName, Verbosity, GraphName, DataName, Data)
-	#define UE_VLOG_CONE(Object, CategoryName, Verbosity, Orgin, Direction, Length, Angle, Color, DescriptionFormat, ...)
-	#define UE_VLOG_CYLINDER(Object, CategoryName, Verbosity, Start, End, Radius, Color, DescriptionFormat, ...)
-	#define UE_VLOG_CAPSULE(Object, CategoryName, Verbosity, Center, HalfHeight, Radius, Rotation, Color, DescriptionFormat, ...)
-
-	#define REDIRECT_TO_VLOG(Dest)
-	#define REDIRECT_OBJECT_TO_VLOG(Src, Destination) 
-	#define CONNECT_WITH_VLOG(Dest)
-	#define CONNECT_OBJECT_WITH_VLOG(Src, Dest)
 #endif //ENABLE_VISUAL_LOG
diff --git a/Engine/Source/Runtime/Engine/Public/VisualLogger/VisualLogger.h b/Engine/Source/Runtime/Engine/Public/VisualLogger/VisualLogger.h
new file mode 100644
--- /dev/null
+++ b/Engine/Source/Runtime/Engine/Public/VisualLogger/VisualLogger.h
@@ -0,0 +1,229 @@
+// Copyright 1998-2014 Epic Games, Inc. All Rights Reserved.
+#pragma once
+#include "VisualLoggerTypes.h"
+#include "EngineDefines.h"
+
+#if ENABLE_VISUAL_LOG
+
+#define REDIRECT_TO_VLOG(Dest) FVisualLogger::Get().Redirect(this, Dest)
+#define REDIRECT_OBJECT_TO_VLOG(Src, Dest) FVisualLogger::Get().Redirect(Src, Dest)
+
+#define CONNECT_WITH_VLOG(Dest)
+#define CONNECT_OBJECT_WITH_VLOG(Src, Dest)
+
+// Text, regular log
+#define UE_VLOG(LogOwner, CategoryName, Verbosity, Format, ...) FVisualLogger::CategorizedLogf(LogOwner, CategoryName, ELogVerbosity::Verbosity, INDEX_NONE, Format, ##__VA_ARGS__)
+#define UE_CVLOG(Condition, LogOwner, CategoryName, Verbosity, Format, ...)  if(Condition) {UE_VLOG(LogOwner, CategoryName, Verbosity, Format, ##__VA_ARGS__);} 
+// Segment shape
+#define UE_VLOG_SEGMENT(LogOwner, CategoryName, Verbosity, SegmentStart, SegmentEnd, Color, Format, ...) FVisualLogger::GeometryShapeLogf(LogOwner, CategoryName, ELogVerbosity::Verbosity, INDEX_NONE, SegmentStart, SegmentEnd, Color, Format, ##__VA_ARGS__)
+#define UE_CVLOG_SEGMENT(Condition, LogOwner, CategoryName, Verbosity, SegmentStart, SegmentEnd, Color, Format, ...) if(Condition) {UE_VLOG_SEGMENT(LogOwner, CategoryName, Verbosity, INDEX_NONE, SegmentStart, SegmentEnd, Color, Format, ##__VA_ARGS__);}
+// Segment shape
+#define UE_VLOG_SEGMENT_THICK(LogOwner, CategoryName, Verbosity, SegmentStart, SegmentEnd, Color, Thickness, Format, ...) FVisualLogger::GeometryShapeLogf(LogOwner, CategoryName, ELogVerbosity::Verbosity, INDEX_NONE, SegmentStart, SegmentEnd, Color, Format, ##__VA_ARGS__)
+#define UE_CVLOG_SEGMENT_THICK(Condition, LogOwner, CategoryName, Verbosity, SegmentStart, SegmentEnd, Color, Thickness, Format, ...) if(Condition) {UE_VLOG_SEGMENT_THICK(LogOwner, CategoryName, Verbosity, INDEX_NONE, SegmentStart, SegmentEnd, Color, Format, ##__VA_ARGS__);} 
+// Localization as sphere shape
+#define UE_VLOG_LOCATION(LogOwner, CategoryName, Verbosity, Location, Radius, Color, Format, ...) FVisualLogger::GeometryShapeLogf(LogOwner, CategoryName, ELogVerbosity::Verbosity, INDEX_NONE, Location, Radius, Color, Format, ##__VA_ARGS__)
+#define UE_CVLOG_LOCATION(Condition, LogOwner, CategoryName, Verbosity, Location, Radius, Color, Format, ...)  if(Condition) {UE_VLOG_LOCATION(LogOwner, CategoryName, Verbosity, Location, Radius, Color, Format, ##__VA_ARGS__);} 
+// Box shape
+#define UE_VLOG_BOX(LogOwner, CategoryName, Verbosity, Box, Color, Format, ...)  FVisualLogger::GeometryShapeLogf(LogOwner, CategoryName, ELogVerbosity::Verbosity, INDEX_NONE, Box, Color, Format, ##__VA_ARGS__)
+#define UE_CVLOG_BOX(Condition, LogOwner, CategoryName, Verbosity, Box, Color, Format, ...)  if(Condition) {UE_VLOG_BOX(Actor, CategoryName, Verbosity, Box, Color, Format, ##__VA_ARGS__);} 
+// Cone shape
+#define UE_VLOG_CONE(LogOwner, CategoryName, Verbosity, Orgin, Direction, Length, Angle, Color, Format, ...) FVisualLogger::GeometryShapeLogf(LogOwner, CategoryName, ELogVerbosity::Verbosity, INDEX_NONE, Orgin, Direction, Length, Angle, Color, Format, ##__VA_ARGS__)
+#define UE_CVLOG_CONE(Condition, LogOwner, CategoryName, Verbosity, Orgin, Direction, Length, Angle, Color, Format, ...)  if(Condition) {UE_VLOG_CONE(Actor, CategoryName, Verbosity, Orgin, Direction, Length, Angle, Color, Format, ##__VA_ARGS__);} 
+// Cylinder shape
+#define UE_VLOG_CYLINDER(LogOwner, CategoryName, Verbosity, Start, End, Radius, Color, Format, ...) FVisualLogger::GeometryShapeLogf(LogOwner, CategoryName, ELogVerbosity::Verbosity, INDEX_NONE, Start, End, Radius, Color, Format, ##__VA_ARGS__)
+#define UE_CVLOG_CYLINDER(LogOwner, CategoryName, Verbosity, Start, End, Radius, Color, Format, ...)  if(Condition) {UE_VLOG_CYLINDER(LogOwner, CategoryName, Verbosity, Start, End, Radius, Color, Format, ##__VA_ARGS__);} 
+// Capsule shape
+#define UE_VLOG_CAPSULE(LogOwner, CategoryName, Verbosity, Center, HalfHeight, Radius, Rotation, Color, Format, ...) FVisualLogger::GeometryShapeLogf(LogOwner, CategoryName, ELogVerbosity::Verbosity, INDEX_NONE, Center, HalfHeight, Radius, Rotation, Color, Format, ##__VA_ARGS__)
+#define UE_CVLOG_CAPSULE(LogOwner, CategoryName, Verbosity, Center, HalfHeight, Radius, Rotation, Color, Format, ...)  if(Condition) {UE_VLOG_CAPSULE(LogOwner, CategoryName, Verbosity, Center, HalfHeight, Radius, Rotation, Color, Format, ##__VA_ARGS__);} 
+// Histogram data for 2d graphs 
+#define UE_VLOG_HISTOGRAM(LogOwner, CategoryName, Verbosity, GraphName, DataName, Data) FVisualLogger::HistogramDataLogf(LogOwner, CategoryName, ELogVerbosity::Verbosity, INDEX_NONE, GraphName, DataName, Data, FColor::White, TEXT(""))
+#define UE_CVLOG_HISTOGRAM(LogOwner, CategoryName, Verbosity, GraphName, DataName, Data) if(Condition) {UE_VLOG_HISTOGRAM(LogOwner, CategoryName, Verbosity, GraphName, DataName, Data);} 
+
+#else
+#define REDIRECT_TO_VLOG(Dest)
+#define REDIRECT_OBJECT_TO_VLOG(Src, Dest)
+#define CONNECT_WITH_VLOG(Dest)
+#define CONNECT_OBJECT_WITH_VLOG(Src, Dest)
+
+#define UE_VLOG(Actor, CategoryName, Verbosity, Format, ...)
+#define UE_CVLOG(Condition, Actor, CategoryName, Verbosity, Format, ...)
+#define UE_VLOG_SEGMENT(Actor, CategoryName, Verbosity, SegmentStart, SegmentEnd, Color, DescriptionFormat, ...)
+#define UE_CVLOG_SEGMENT(Actor, CategoryName, Verbosity, SegmentStart, SegmentEnd, Color, DescriptionFormat, ...)
+#define UE_VLOG_SEGMENT_THICK(Actor, CategoryName, Verbosity, SegmentStart, SegmentEnd, Color, DescriptionFormat, ...)
+#define UE_CVLOG_SEGMENT_THICK(Actor, CategoryName, Verbosity, SegmentStart, SegmentEnd, Color, DescriptionFormat, ...)
+#define UE_VLOG_LOCATION(Actor, CategoryName, Verbosity, Location, Radius, Color, DescriptionFormat, ...)
+#define UE_CVLOG_LOCATION(Actor, CategoryName, Verbosity, Location, Radius, Color, DescriptionFormat, ...)
+#define UE_VLOG_BOX(Actor, CategoryName, Verbosity, Box, Color, DescriptionFormat, ...) 
+#define UE_CVLOG_BOX(Actor, CategoryName, Verbosity, Box, Color, DescriptionFormat, ...) 
+#define UE_VLOG_CONE(Object, CategoryName, Verbosity, Orgin, Direction, Length, Angle, Color, DescriptionFormat, ...)
+#define UE_CVLOG_CONE(Object, CategoryName, Verbosity, Orgin, Direction, Length, Angle, Color, DescriptionFormat, ...)
+#define UE_VLOG_CYLINDER(Object, CategoryName, Verbosity, Start, End, Radius, Color, DescriptionFormat, ...)
+#define UE_CVLOG_CYLINDER(Object, CategoryName, Verbosity, Start, End, Radius, Color, DescriptionFormat, ...)
+#define UE_VLOG_CAPSULE(Object, CategoryName, Verbosity, Center, HalfHeight, Radius, Rotation, Color, DescriptionFormat, ...)
+#define UE_CVLOG_CAPSULE(Object, CategoryName, Verbosity, Center, HalfHeight, Radius, Rotation, Color, DescriptionFormat, ...)
+#define UE_VLOG_HISTOGRAM(Actor, CategoryName, Verbosity, GraphName, DataName, Data)
+#define UE_CVLOG_HISTOGRAM(Actor, CategoryName, Verbosity, GraphName, DataName, Data)
+
+#endif //ENABLE_VISUAL_LOG
+
+// helper macros
+#define TEXT_EMPTY TEXT("")
+
+
+#if ENABLE_VISUAL_LOG
+
+DECLARE_LOG_CATEGORY_EXTERN(LogVisual, Warning, All);
+
+class FVisualLogDevice;
+class FVisualLogExtensionInterface;
+
+DECLARE_DELEGATE_RetVal(FString, FVisualLogFilenameGetterDelegate);
+
+class ENGINE_API FVisualLogger : public FOutputDevice
+{
+public:
+	// Regular text log
+	VARARG_DECL(static void, static void, return, CategorizedLogf, VARARG_NONE, const TCHAR*, VARARG_EXTRA(const class UObject* LogOwner) VARARG_EXTRA(const struct FLogCategoryBase& Category) VARARG_EXTRA(ELogVerbosity::Type Verbosity) VARARG_EXTRA(int32 UniqueLogId), VARARG_EXTRA(LogOwner) VARARG_EXTRA(Category) VARARG_EXTRA(Verbosity) VARARG_EXTRA(UniqueLogId));
+	// Segment log
+	VARARG_DECL(static void, static void, return, GeometryShapeLogf, VARARG_NONE, const TCHAR*, VARARG_EXTRA(const class UObject* LogOwner) VARARG_EXTRA(const struct FLogCategoryBase& Category) VARARG_EXTRA(ELogVerbosity::Type Verbosity) VARARG_EXTRA(int32 UniqueLogId) VARARG_EXTRA(const FVector& Start) VARARG_EXTRA(const FVector& End) VARARG_EXTRA(const FColor& Color), VARARG_EXTRA(LogOwner) VARARG_EXTRA(Category) VARARG_EXTRA(Verbosity) VARARG_EXTRA(UniqueLogId) VARARG_EXTRA(Start) VARARG_EXTRA(End) VARARG_EXTRA(Color));
+	// Location/Shpere log
+	VARARG_DECL(static void, static void, return, GeometryShapeLogf, VARARG_NONE, const TCHAR*, VARARG_EXTRA(const class UObject* LogOwner) VARARG_EXTRA(const struct FLogCategoryBase& Category) VARARG_EXTRA(ELogVerbosity::Type Verbosity) VARARG_EXTRA(int32 UniqueLogId) VARARG_EXTRA(const FVector& Location) VARARG_EXTRA(float Radius) VARARG_EXTRA(const FColor& Color), VARARG_EXTRA(LogOwner) VARARG_EXTRA(Category) VARARG_EXTRA(Verbosity) VARARG_EXTRA(UniqueLogId) VARARG_EXTRA(Location) VARARG_EXTRA(Radius) VARARG_EXTRA(Color));
+	// Box log
+	VARARG_DECL(static void, static void, return, GeometryShapeLogf, VARARG_NONE, const TCHAR*, VARARG_EXTRA(const class UObject* LogOwner) VARARG_EXTRA(const struct FLogCategoryBase& Category) VARARG_EXTRA(ELogVerbosity::Type Verbosity) VARARG_EXTRA(int32 UniqueLogId) VARARG_EXTRA(const FBox& Box) VARARG_EXTRA(const FColor& Color), VARARG_EXTRA(LogOwner) VARARG_EXTRA(Category) VARARG_EXTRA(Verbosity) VARARG_EXTRA(UniqueLogId) VARARG_EXTRA(Box) VARARG_EXTRA(Color));
+	// Cone log
+	VARARG_DECL(static void, static void, return, GeometryShapeLogf, VARARG_NONE, const TCHAR*, VARARG_EXTRA(const class UObject* LogOwner) VARARG_EXTRA(const struct FLogCategoryBase& Category) VARARG_EXTRA(ELogVerbosity::Type Verbosity) VARARG_EXTRA(int32 UniqueLogId) VARARG_EXTRA(const FVector& Orgin) VARARG_EXTRA(const FVector& Direction) VARARG_EXTRA(const float Length) VARARG_EXTRA(const float Angle) VARARG_EXTRA(const FColor& Color), VARARG_EXTRA(LogOwner) VARARG_EXTRA(Category) VARARG_EXTRA(Verbosity) VARARG_EXTRA(UniqueLogId) VARARG_EXTRA(Orgin) VARARG_EXTRA(Direction) VARARG_EXTRA(Length) VARARG_EXTRA(Angle) VARARG_EXTRA(Color));
+	// Cylinder log
+	VARARG_DECL(static void, static void, return, GeometryShapeLogf, VARARG_NONE, const TCHAR*, VARARG_EXTRA(const class UObject* LogOwner) VARARG_EXTRA(const struct FLogCategoryBase& Category) VARARG_EXTRA(ELogVerbosity::Type Verbosity) VARARG_EXTRA(int32 UniqueLogId) VARARG_EXTRA(const FVector& Start) VARARG_EXTRA(const FVector& End) VARARG_EXTRA(const float Radius) VARARG_EXTRA(const FColor& Color), VARARG_EXTRA(LogOwner) VARARG_EXTRA(Category) VARARG_EXTRA(Verbosity) VARARG_EXTRA(UniqueLogId) VARARG_EXTRA(Start) VARARG_EXTRA(End) VARARG_EXTRA(Radius) VARARG_EXTRA(Color));
+	// Capsule log
+	VARARG_DECL(static void, static void, return, GeometryShapeLogf, VARARG_NONE, const TCHAR*, VARARG_EXTRA(const class UObject* LogOwner) VARARG_EXTRA(const struct FLogCategoryBase& Category) VARARG_EXTRA(ELogVerbosity::Type Verbosity) VARARG_EXTRA(int32 UniqueLogId) VARARG_EXTRA(const FVector& Center) VARARG_EXTRA(float HalfHeight) VARARG_EXTRA(float Radius) VARARG_EXTRA(const FQuat & Rotation) VARARG_EXTRA(const FColor& Color), VARARG_EXTRA(LogOwner) VARARG_EXTRA(Category) VARARG_EXTRA(Verbosity) VARARG_EXTRA(UniqueLogId) VARARG_EXTRA(Center) VARARG_EXTRA(HalfHeight) VARARG_EXTRA(Radius) VARARG_EXTRA(Rotation) VARARG_EXTRA(Color));
+	//Histogram data
+	VARARG_DECL(static void, static void, return, HistogramDataLogf, VARARG_NONE, const TCHAR*, VARARG_EXTRA(const class UObject* LogOwner) VARARG_EXTRA(const struct FLogCategoryBase& Category) VARARG_EXTRA(ELogVerbosity::Type Verbosity) VARARG_EXTRA(int32 UniqueLogId) VARARG_EXTRA(FName GraphName) VARARG_EXTRA(FName DataName) VARARG_EXTRA(const FVector2D& Data) VARARG_EXTRA(const FColor& Color), VARARG_EXTRA(LogOwner) VARARG_EXTRA(Category) VARARG_EXTRA(Verbosity) VARARG_EXTRA(UniqueLogId) VARARG_EXTRA(GraphName) VARARG_EXTRA(DataName) VARARG_EXTRA(Data) VARARG_EXTRA(Color));
+
+	// static getter
+	static FVisualLogger& Get()
+	{
+		static FVisualLogger GVisLog;
+		return GVisLog;
+	}
+
+	FVisualLogger();
+	virtual ~FVisualLogger();
+
+	// Removes all logged data 
+	void Cleanup(bool bReleaseMemory = false);
+
+	/** Set log owner redirection from one object to another, to combine logs */
+	static void Redirect(class UObject* FromObject, class UObject* ToObject);
+
+	/** find and return redirection object for given object*/
+	static class UObject* FindRedirection(const class UObject* Object);
+
+	/** blocks all categories from logging. It can be bypassed with white list for categories */
+	FORCEINLINE void BlockAllCategories(bool InBlock) { bBlockedAllCategories = InBlock; }
+
+	/** checks if all categories are blocked */
+	FORCEINLINE bool IsBlockedForAllCategories() { return !!bBlockedAllCategories; }
+
+	/** Returns white list for modifications */
+	TArray<class FName>& GetWhiteList() { return CategoriesWhiteList; }
+
+	FORCEINLINE void AddCategortyToWhiteList(class FName Category) { CategoriesWhiteList.AddUnique(Category); }
+
+	FORCEINLINE void ClearWhiteList() { CategoriesWhiteList.Reset();  }
+
+	/** Generates and returns Id unique for given timestamp - used to connect different logs between (ex. text log with geometry shape) */
+	int32 GetUniqueId(float Timestamp);
+
+	/** Starts visual log collecting and recording */
+	void SetIsRecording(bool InIsRecording);
+	/** return information is vlog recording is enabled or not */
+	bool IsRecording() { return !!bIsRecording; }
+
+	/** Starts visual log collecting and recording */
+	void SetIsRecordingToFile(bool InIsRecording);
+	/** return information is vlog recording is enabled or not */
+	bool IsRecordingToFile() { return !!bIsRecordingToFile; }
+
+	void SetIsRecordingOnServer(bool IsRecording) { bIsRecordingOnServer = IsRecording; }
+	bool IsRecordingOnServer() { return !!bIsRecordingOnServer; }
+
+	/** Add visual logger output device */
+	void AddDevice(FVisualLogDevice* InDevice) { OutputDevices.AddUnique(InDevice); }
+	/** Remove visual logger output device */
+	void RemoveDevice(FVisualLogDevice* InDevice) { OutputDevices.RemoveSwap(InDevice); }
+	/** Remove visual logger output device */
+	TArray<FVisualLogDevice*>& GetDevices() { return OutputDevices; }
+	/** Returns  current entry for given TimeStap or creates another one  but first it serialize previous entry as completed to vislog devices. Use VisualLogger::DontCreate to get current entry without serialization*/
+	FVisualLogEntry* GetEntryToWrite(const class UObject* Object, float TimeStamp, VisualLogger::ECreateIfNeeded ShouldCreate = VisualLogger::Create);
+
+	/** FileName getter to set project specific file name for vlogs - highly encouradged to use FVisualLogFilenameGetterDelegate::CreateUObject with this */
+	void SetLogFileNameGetter(const FVisualLogFilenameGetterDelegate& InLogFileNameGetter) { LogFileNameGetter = InLogFileNameGetter; }
+
+	/** Register extension to use by LogVisualizer  */
+	void RegisterExtension(FName TagName, FVisualLogExtensionInterface* ExtensionInterface) { check(AllExtensions.Contains(TagName) == false); AllExtensions.Add(TagName, ExtensionInterface); }
+	/**  Removes previously registered extension */
+	void UnregisterExtension(FName TagName, FVisualLogExtensionInterface* ExtensionInterface) { AllExtensions.Remove(TagName); }
+	/** returns extension identified by given tag */
+	FVisualLogExtensionInterface* GetExtensionForTag(const FName TagName) { return AllExtensions.Contains(TagName) ? AllExtensions[TagName] : NULL; }
+	/** Returns reference to map with all registered extension */
+	TMap<FName, FVisualLogExtensionInterface*>& GetAllExtensions() { return AllExtensions; }
+
+private:
+	virtual void Serialize(const TCHAR* V, ELogVerbosity::Type Verbosity, const class FName& Category) override { ensureMsgf(0, TEXT("Regular serialize is forbiden for visual logs")); }
+
+protected:
+	/** Array of output devices to redirect to */
+	TArray<FVisualLogDevice*> OutputDevices;
+	// Map for inter-objects redirections
+	static TMap<class UObject*, TArray<const class UObject*> > RedirectionMap;
+	// white list of categories to bypass blocking
+	TArray<class FName>	CategoriesWhiteList;
+	// Visual Logger extensions map
+	TMap<FName, FVisualLogExtensionInterface*> AllExtensions;
+	// last generated unique id for given times tamp
+	TMap<float, int32> LastUniqueIds;
+	// Current entry with all data
+	TMap<const class UObject*, FVisualLogEntry>	CurrentEntryPerObject;
+	// if set all categories are blocked from logging
+	int32 bBlockedAllCategories : 1;
+	// if set we are recording and collecting all vlog data
+	int32 bIsRecording : 1;
+	// if set we are recording to file
+	int32 bIsRecordingToFile : 1;
+	// variable set (from cheat manager) when logging is active on server
+	int32 bIsRecordingOnServer : 1;
+	// start recording time
+	float StartRecordingToFileTime;
+	/** Delegate to set project specific file name for vlogs */
+	FVisualLogFilenameGetterDelegate LogFileNameGetter;
+};
+
+/**
+ * Interface for Visual Logger Device
+ */
+class FVisualLogDevice
+{
+public:
+	virtual void Cleanup(bool bReleaseMemory = false) = 0;
+	virtual void StartRecordingToFile(float TImeStamp) = 0;
+	virtual void StopRecordingToFile(float TImeStamp) = 0;
+	virtual void Serialize(const class UObject* LogOwner, const FVisualLogEntry& LogEntry) = 0;
+	virtual void SetFileName(const FString& InFileName) = 0;
+	virtual bool HasFlags(int32 InFlags) { return !!(InFlags & (VisualLogger::CanSaveToFile | VisualLogger::StoreLogsLocally)); }
+};
+
+/**
+ * Interface for extensions
+ */
+class FVisualLogExtensionInterface
+{
+public:
+	virtual void OnTimestampChange(float Timestamp, class UWorld* InWorld, class AActor* HelperActor) = 0;
+	virtual void DrawData(class UWorld* InWorld, class UCanvas* Canvas, class AActor* HelperActor, const FName& TagName, const FVisualLogEntry::FDataBlock& DataBlock, float Timestamp) = 0;
+	virtual void DisableDrawingForData(class UWorld* InWorld, class UCanvas* Canvas, class AActor* HelperActor, const FName& TagName, const FVisualLogEntry::FDataBlock& DataBlock, float Timestamp) = 0;
+	virtual void LogEntryLineSelectionChanged(TSharedPtr<struct FLogEntryItem> SelectedItem, int64 UserData, FName TagName) = 0;
+};
+
+#include "VisualLogger.inl"
+
+#endif //ENABLE_VISUAL_LOG
diff --git a/Engine/Source/Runtime/Engine/Public/VisualLogger/VisualLogger.inl b/Engine/Source/Runtime/Engine/Public/VisualLogger/VisualLogger.inl
new file mode 100644
--- /dev/null
+++ b/Engine/Source/Runtime/Engine/Public/VisualLogger/VisualLogger.inl
@@ -0,0 +1,174 @@
+// Copyright 1998-2014 Epic Games, Inc. All Rights Reserved.
+
+#include "Runtime/Core/Private/Misc/VarargsHelper.h"
+
+#define CHECK_VLOG_DATA() \
+	float CurrentTime = 0; \
+	const FName CategoryName = const_cast<struct FLogCategoryBase&>(Category).GetCategoryName(); \
+	if (GEngine && GEngine->bDisableAILogging || FVisualLogger::Get().bIsRecording == false || !Object || Object->HasAnyFlags(RF_ClassDefaultObject) || (FVisualLogger::Get().IsBlockedForAllCategories() && FVisualLogger::Get().CategoriesWhiteList.Find(CategoryName) == INDEX_NONE)) \
+	{ \
+		return; \
+	} \
+	UObject* LogOwner = FindRedirection(Object); \
+	UWorld* World = GEngine->GetWorldFromContextObject(Object); \
+	if (World == NULL) \
+	{ \
+		ensure(World); \
+		return; \
+	} \
+
+FORCEINLINE_DEBUGGABLE
+FVisualLogEntry* FVisualLogger::GetEntryToWrite(const class UObject* Object, float TimeStamp, VisualLogger::ECreateIfNeeded ShouldCreate)
+{
+	FVisualLogEntry* CurrentEntry = NULL;
+
+	bool InitializeNewEntry = false;
+	if (CurrentEntryPerObject.Contains(Object))
+	{
+		CurrentEntry = &CurrentEntryPerObject[Object];
+		if (TimeStamp > CurrentEntry->TimeStamp && ShouldCreate == VisualLogger::Create)
+		{
+			for (auto* Device : OutputDevices)
+			{
+				Device->Serialize(Object, *CurrentEntry);
+			}
+			InitializeNewEntry = true;
+		}
+	}
+	else
+	{
+		CurrentEntry = &CurrentEntryPerObject.Add(Object);
+		InitializeNewEntry = true;
+	}
+
+	if (InitializeNewEntry)
+	{
+		CurrentEntry->Reset();
+		CurrentEntry->TimeStamp = TimeStamp;
+		const class AActor* AsActor = Cast<class AActor>(Object);
+		if (AsActor)
+		{
+			CurrentEntry->Location = AsActor->GetActorLocation();
+			AsActor->GrabDebugSnapshot(CurrentEntry);
+		}
+		if (RedirectionMap.Contains(Object))
+		{
+			for (auto Child : RedirectionMap[Object])
+			{
+				const class AActor* AsActor = Cast<class AActor>(Child);
+				if (AsActor)
+				{
+					AsActor->GrabDebugSnapshot(CurrentEntry);
+				}
+			}
+		}
+	}
+
+	return CurrentEntry;
+}
+
+FORCEINLINE
+VARARG_BODY(void, FVisualLogger::CategorizedLogf, const TCHAR*, VARARG_EXTRA(const class UObject* Object) VARARG_EXTRA(const struct FLogCategoryBase& Category) VARARG_EXTRA(ELogVerbosity::Type Verbosity) VARARG_EXTRA(int32 UniqueLogId))
+{
+	CHECK_VLOG_DATA();
+
+	FVisualLogEntry* CurrentEntry = FVisualLogger::Get().GetEntryToWrite(LogOwner, World->TimeSeconds);
+	ensure(CurrentEntry);
+
+	GROWABLE_LOGF(
+		CurrentEntry->AddText(Buffer, CategoryName);
+	);
+}
+
+FORCEINLINE
+VARARG_BODY(void, FVisualLogger::GeometryShapeLogf, const TCHAR*, VARARG_EXTRA(const class UObject* Object) VARARG_EXTRA(const struct FLogCategoryBase& Category) VARARG_EXTRA(ELogVerbosity::Type Verbosity) VARARG_EXTRA(int32 UniqueLogId) VARARG_EXTRA(const FVector& Start) VARARG_EXTRA(const FVector& End) VARARG_EXTRA(const FColor& Color))
+{
+	CHECK_VLOG_DATA();
+
+	FVisualLogEntry* CurrentEntry = FVisualLogger::Get().GetEntryToWrite(LogOwner, World->TimeSeconds);
+	ensure(CurrentEntry);
+
+	GROWABLE_LOGF(
+		CurrentEntry->AddElement(Start, End, CategoryName, Color, Buffer);
+	);
+}
+
+FORCEINLINE
+VARARG_BODY(void, FVisualLogger::GeometryShapeLogf, const TCHAR*, VARARG_EXTRA(const class UObject* Object) VARARG_EXTRA(const struct FLogCategoryBase& Category) VARARG_EXTRA(ELogVerbosity::Type Verbosity) VARARG_EXTRA(int32 UniqueLogId) VARARG_EXTRA(const FVector& Location) VARARG_EXTRA(float Radius) VARARG_EXTRA(const FColor& Color))
+{
+	CHECK_VLOG_DATA();
+
+	FVisualLogEntry* CurrentEntry = FVisualLogger::Get().GetEntryToWrite(LogOwner, World->TimeSeconds);
+	ensure(CurrentEntry);
+
+	GROWABLE_LOGF(
+		CurrentEntry->AddElement(Location, CategoryName, Color, Buffer, Radius);
+	);
+}
+
+FORCEINLINE
+VARARG_BODY(void, FVisualLogger::GeometryShapeLogf, const TCHAR*, VARARG_EXTRA(const class UObject* Object) VARARG_EXTRA(const struct FLogCategoryBase& Category) VARARG_EXTRA(ELogVerbosity::Type Verbosity) VARARG_EXTRA(int32 UniqueLogId) VARARG_EXTRA(const FBox& Box) VARARG_EXTRA(const FColor& Color))
+{
+	CHECK_VLOG_DATA();
+
+	FVisualLogEntry* CurrentEntry = FVisualLogger::Get().GetEntryToWrite(LogOwner, World->TimeSeconds);
+	ensure(CurrentEntry);
+
+	GROWABLE_LOGF(
+		CurrentEntry->AddElement(Box, CategoryName, Color, Buffer);
+	);
+}
+
+FORCEINLINE
+VARARG_BODY(void, FVisualLogger::GeometryShapeLogf, const TCHAR*, VARARG_EXTRA(const class UObject* Object) VARARG_EXTRA(const struct FLogCategoryBase& Category) VARARG_EXTRA(ELogVerbosity::Type Verbosity) VARARG_EXTRA(int32 UniqueLogId) VARARG_EXTRA(const FVector& Orgin) VARARG_EXTRA(const FVector& Direction) VARARG_EXTRA(const float Length) VARARG_EXTRA(const float Angle)  VARARG_EXTRA(const FColor& Color))
+{
+	CHECK_VLOG_DATA();
+
+	FVisualLogEntry* CurrentEntry = FVisualLogger::Get().GetEntryToWrite(LogOwner, World->TimeSeconds);
+	ensure(CurrentEntry);
+
+	GROWABLE_LOGF(
+		CurrentEntry->AddElement(Orgin, Direction, Length, Angle, Angle, CategoryName, Color, Buffer);
+	);
+}
+
+FORCEINLINE
+VARARG_BODY(void, FVisualLogger::GeometryShapeLogf, const TCHAR*, VARARG_EXTRA(const class UObject* Object) VARARG_EXTRA(const struct FLogCategoryBase& Category) VARARG_EXTRA(ELogVerbosity::Type Verbosity) VARARG_EXTRA(int32 UniqueLogId) VARARG_EXTRA(const FVector& Start) VARARG_EXTRA(const FVector& End) VARARG_EXTRA(const float Radius) VARARG_EXTRA(const FColor& Color))
+{
+	CHECK_VLOG_DATA();
+
+	FVisualLogEntry* CurrentEntry = FVisualLogger::Get().GetEntryToWrite(LogOwner, World->TimeSeconds);
+	ensure(CurrentEntry);
+
+	GROWABLE_LOGF(
+		CurrentEntry->AddElement(Start, End, Radius, CategoryName, Color, Buffer);
+	);
+}
+
+FORCEINLINE
+VARARG_BODY(void, FVisualLogger::GeometryShapeLogf, const TCHAR*, VARARG_EXTRA(const class UObject* Object) VARARG_EXTRA(const struct FLogCategoryBase& Category) VARARG_EXTRA(ELogVerbosity::Type Verbosity) VARARG_EXTRA(int32 UniqueLogId) VARARG_EXTRA(const FVector& Center) VARARG_EXTRA(float HalfHeight) VARARG_EXTRA(float Radius) VARARG_EXTRA(const FQuat& Rotation) VARARG_EXTRA(const FColor& Color))
+{
+	CHECK_VLOG_DATA();
+
+	FVisualLogEntry* CurrentEntry = FVisualLogger::Get().GetEntryToWrite(LogOwner, World->TimeSeconds);
+	ensure(CurrentEntry);
+
+	GROWABLE_LOGF(
+		CurrentEntry->AddCapsule(Center, HalfHeight, Radius, Rotation, CategoryName, Color, Buffer);
+	);
+}
+
+FORCEINLINE
+VARARG_BODY(void, FVisualLogger::HistogramDataLogf, const TCHAR*, VARARG_EXTRA(const class UObject* Object) VARARG_EXTRA(const struct FLogCategoryBase& Category) VARARG_EXTRA(ELogVerbosity::Type Verbosity) VARARG_EXTRA(int32 UniqueLogId) VARARG_EXTRA(FName GraphName) VARARG_EXTRA(FName DataName) VARARG_EXTRA(const FVector2D& Data) VARARG_EXTRA(const FColor& Color))
+{
+	CHECK_VLOG_DATA();
+
+	FVisualLogEntry* CurrentEntry = FVisualLogger::Get().GetEntryToWrite(LogOwner, World->TimeSeconds);
+	ensure(CurrentEntry);
+
+	GROWABLE_LOGF(
+		CurrentEntry->AddHistogramData(Data, CategoryName, GraphName, DataName);
+	);
+}
+
+#undef  CHECK_VLOG_DATA
diff --git a/Engine/Source/Runtime/Engine/Public/VisualLogger/VisualLoggerTypes.h b/Engine/Source/Runtime/Engine/Public/VisualLogger/VisualLoggerTypes.h
new file mode 100644
--- /dev/null
+++ b/Engine/Source/Runtime/Engine/Public/VisualLogger/VisualLoggerTypes.h
@@ -0,0 +1,222 @@
+// Copyright 1998-2014 Epic Games, Inc. All Rights Reserved.
+#pragma once
+#include "EngineDefines.h"
+
+namespace VisualLogger
+{
+	enum ECreateIfNeeded
+	{
+		Invalid = -1,
+		DontCreate = 0,
+		Create = 1,
+	};
+
+	enum DeviceFlags
+	{
+		CanSaveToFile = 1,
+		StoreLogsLocally = 2,
+	};
+}
+
+struct ENGINE_API FVisualLogEntry
+{
+	struct FLogLine
+	{
+		FString Line;
+		FName Category;
+		TEnumAsByte<ELogVerbosity::Type> Verbosity;
+		int32 UniqueId;
+		int64 UserData;
+		FName TagName;
+
+		FLogLine(const FName& InCategory, ELogVerbosity::Type InVerbosity, const FString& InLine)
+			: Line(InLine), Category(InCategory), Verbosity(InVerbosity), UserData(0)
+		{}
+
+		FLogLine(const FName& InCategory, ELogVerbosity::Type InVerbosity, const FString& InLine, int64 InUserData)
+			: Line(InLine), Category(InCategory), Verbosity(InVerbosity), UserData(InUserData)
+		{}
+	};
+
+	struct FStatusCategory
+	{
+		TArray<FString> Data;
+		FString Category;
+		int32 UniqueId;
+
+		FORCEINLINE void Add(const FString& Key, const FString& Value)
+		{
+			Data.Add(FString(Key).AppendChar(TEXT('|')) + Value);
+		}
+
+		FORCEINLINE bool GetDesc(int32 Index, FString& Key, FString& Value) const
+		{
+			if (Data.IsValidIndex(Index))
+			{
+				int32 SplitIdx = INDEX_NONE;
+				if (Data[Index].FindChar(TEXT('|'), SplitIdx))
+				{
+					Key = Data[Index].Left(SplitIdx);
+					Value = Data[Index].Mid(SplitIdx + 1);
+					return true;
+				}
+			}
+
+			return false;
+		}
+	};
+
+	struct FElementToDraw
+	{
+		friend struct FVisualLogEntry;
+
+		enum EElementType {
+			Invalid = 0,
+			SinglePoint, // individual points. 
+			Segment, // pairs of points 
+			Path,	// sequence of point
+			Box,
+			Cone,
+			Cylinder,
+			Capsule,
+			// note that in order to remain backward compatibility in terms of log
+			// serialization new enum values need to be added at the end
+		};
+
+		FString Description;
+		FName Category;
+		TEnumAsByte<ELogVerbosity::Type> Verbosity;
+		TArray<FVector> Points;
+		int32 UniqueId;
+
+	public:
+		uint8 Type;
+		uint8 Color;
+		union
+		{
+			uint16 Thicknes;
+			uint16 Radius;
+		};
+
+		FElementToDraw(EElementType InType = Invalid) : Verbosity(ELogVerbosity::All), Type(InType), Color(0xff), Thicknes(0)
+		{}
+
+		FElementToDraw(const FString& InDescription, const FColor& InColor, uint16 InThickness, const FName& InCategory) : Category(InCategory), Verbosity(ELogVerbosity::All), Type(uint16(Invalid)), Thicknes(InThickness)
+		{
+			if (InDescription.IsEmpty() == false)
+			{
+				Description = InDescription;
+			}
+			SetColor(InColor);
+		}
+
+		FORCEINLINE void SetColor(const FColor& InColor)
+		{
+			Color = ((InColor.DWColor() >> 30) << 6)
+				| (((InColor.DWColor() & 0x00ff0000) >> 22) << 4)
+				| (((InColor.DWColor() & 0x0000ff00) >> 14) << 2)
+				| ((InColor.DWColor() & 0x000000ff) >> 6);
+		}
+
+		FORCEINLINE EElementType GetType() const
+		{
+			return EElementType(Type);
+		}
+
+		FORCEINLINE void SetType(EElementType InType)
+		{
+			Type = InType;
+		}
+
+		FORCEINLINE FColor GetFColor() const
+		{
+			return FColor(
+				((Color & 0xc0) << 24)
+				| ((Color & 0x30) << 18)
+				| ((Color & 0x0c) << 12)
+				| ((Color & 0x03) << 6)
+				);
+		}
+	};
+
+	struct FHistogramSample
+	{
+		FName Category;
+		TEnumAsByte<ELogVerbosity::Type> Verbosity;
+		FName GraphName;
+		FName DataName;
+		FVector2D SampleValue;
+		int32 UniqueId;
+	};
+
+	struct FDataBlock
+	{
+		FName TagName;
+		FName Category;
+		TEnumAsByte<ELogVerbosity::Type> Verbosity;
+		TArray<uint8>	Data;
+		int32 UniqueId;
+	};
+
+#if ENABLE_VISUAL_LOG
+		float TimeStamp;
+		FVector Location;
+
+		TArray<FLogLine> LogLines;
+		TArray<FStatusCategory> Status;
+		TArray<FElementToDraw> ElementsToDraw;
+		TArray<FHistogramSample>	HistogramSamples;
+		TArray<FDataBlock>	DataBlocks;
+
+		FVisualLogEntry() { Reset(); }
+		FVisualLogEntry(const FVisualLogEntry& Entry);
+		FVisualLogEntry(const class AActor* InActor, TArray<TWeakObjectPtr<UObject> >* Children);
+		FVisualLogEntry(float InTimeStamp, FVector InLocation, const UObject* Object, TArray<TWeakObjectPtr<UObject> >* Children);
+
+		void Reset()
+		{
+			TimeStamp = -1;
+			Location = FVector::ZeroVector;
+			LogLines.Reset();
+			Status.Reset();
+			ElementsToDraw.Reset();
+			HistogramSamples.Reset();
+			DataBlocks.Reset();
+		}
+
+		void AddText(const FString& TextLine, const FName& CategoryName);
+		// path
+		void AddElement(const TArray<FVector>& Points, const FName& CategoryName, const FColor& Color = FColor::White, const FString& Description = TEXT(""), uint16 Thickness = 0);
+		// location
+		void AddElement(const FVector& Point, const FName& CategoryName, const FColor& Color = FColor::White, const FString& Description = TEXT(""), uint16 Thickness = 0);
+		// segment
+		void AddElement(const FVector& Start, const FVector& End, const FName& CategoryName, const FColor& Color = FColor::White, const FString& Description = TEXT(""), uint16 Thickness = 0);
+		// box
+		void AddElement(const FBox& Box, const FName& CategoryName, const FColor& Color = FColor::White, const FString& Description = TEXT(""), uint16 Thickness = 0);
+		// Cone
+		void AddElement(const FVector& Orgin, const FVector& Direction, float Length, float AngleWidth, float AngleHeight, const FName& CategoryName, const FColor& Color = FColor::White, const FString& Description = TEXT(""), uint16 Thickness = 0);
+		// Cylinder
+		void AddElement(const FVector& Start, const FVector& End, float Radius, const FName& CategoryName, const FColor& Color = FColor::White, const FString& Description = TEXT(""), uint16 Thickness = 0);
+		// histogram sample
+		void AddHistogramData(const FVector2D& DataSample, const FName& CategoryName, const FName& GraphName, const FName& DataName);
+		// Custom data block
+		FVisualLogEntry::FDataBlock& AddDataBlock(const FString& TagName, const TArray<uint8>& BlobDataArray, const FName& CategoryName);
+		// capsule
+		void AddCapsule(FVector const& Center, float HalfHeight, float Radius, const FQuat & Rotation, const FName& CategoryName, const FColor& Color = FColor::White, const FString& Description = TEXT(""));
+
+		// find index of status category
+		FORCEINLINE int32 FindStatusIndex(const FString& CategoryName)
+		{
+			for (int32 TestCategoryIndex = 0; TestCategoryIndex < Status.Num(); TestCategoryIndex++)
+			{
+				if (Status[TestCategoryIndex].Category == CategoryName)
+				{
+					return TestCategoryIndex;
+				}
+			}
+
+			return INDEX_NONE;
+		}
+
+#endif // ENABLE_VISUAL_LOG
+	};
