Mobile Preview - Removing GRHIShaderPlatform

diff --git a/Engine/Plugins/Runtime/OculusRift/Source/OculusRift/Private/OculusRiftHMD.cpp b/Engine/Plugins/Runtime/OculusRift/Source/OculusRift/Private/OculusRiftHMD.cpp
--- a/Engine/Plugins/Runtime/OculusRift/Source/OculusRift/Private/OculusRiftHMD.cpp
+++ b/Engine/Plugins/Runtime/OculusRift/Source/OculusRift/Private/OculusRiftHMD.cpp
@@ -1634,13 +1634,13 @@ void FOculusRiftHMD::Startup()
 
 #ifdef OVR_DIRECT_RENDERING
 #if defined(OVR_D3D_VERSION) && (OVR_D3D_VERSION == 11)
-	if (IsPCPlatform(GRHIShaderPlatform) && !IsOpenGLPlatform(GRHIShaderPlatform))
+	if (IsPCPlatform(GMaxRHIShaderPlatform) && !IsOpenGLPlatform(GMaxRHIShaderPlatform))
 	{
 		pD3D11Bridge = new D3D11Bridge(this);
 	}
 #endif
 #if defined(OVR_GL)
-	if (IsOpenGLPlatform(GRHIShaderPlatform))
+	if (IsOpenGLPlatform(GMaxRHIShaderPlatform))
 	{
 		pOGLBridge = new OGLBridge(this);
 	}
@@ -1688,7 +1688,7 @@ void FOculusRiftHMD::Shutdown()
 
 void FOculusRiftHMD::UpdateDistortionCaps()
 {
-	if (IsOpenGLPlatform(GRHIShaderPlatform))
+	if (IsOpenGLPlatform(GMaxRHIShaderPlatform))
 	{
 		DistortionCaps &= ~ovrDistortionCap_SRGB;
 		DistortionCaps |= ovrDistortionCap_FlipInput;
diff --git a/Engine/Source/Editor/UnrealEd/Private/Lightmass/LightmassRender.cpp b/Engine/Source/Editor/UnrealEd/Private/Lightmass/LightmassRender.cpp
--- a/Engine/Source/Editor/UnrealEd/Private/Lightmass/LightmassRender.cpp
+++ b/Engine/Source/Editor/UnrealEd/Private/Lightmass/LightmassRender.cpp
@@ -208,12 +208,12 @@ public:
 			// This allows us to create a deterministic yet unique Id for the shader map that will be compiled for this FLightmassMaterialProxy
 			FMaterialShaderMapId ResourceId;
 			//@todo - always use highest quality level for static lighting
-			MaterialInstance->GetMaterialResourceId(GRHIShaderPlatform, EMaterialQualityLevel::Num, ResourceId);
+			MaterialInstance->GetMaterialResourceId(GMaxRHIShaderPlatform, EMaterialQualityLevel::Num, ResourceId);
 
 			// Override with a special usage so we won't re-use the shader map used by the MI for rendering
 			ResourceId.Usage = GetShaderMapUsage();
 
-			CacheShaders(ResourceId, GRHIShaderPlatform, true);
+			CacheShaders(ResourceId, GMaxRHIShaderPlatform, true);
 		}
 		else
 		{
@@ -222,7 +222,7 @@ public:
 			// Copy the material resource Id
 			// The FLightmassMaterialProxy's GetShaderMapUsage will set it apart from the MI's resource when it comes to finding a shader map
 
-			CacheShaders(GRHIShaderPlatform, true);
+			CacheShaders(GMaxRHIShaderPlatform, true);
 		}
 	}
 
diff --git a/Engine/Source/Runtime/EmptyRHI/Private/EmptyRHI.cpp b/Engine/Source/Runtime/EmptyRHI/Private/EmptyRHI.cpp
--- a/Engine/Source/Runtime/EmptyRHI/Private/EmptyRHI.cpp
+++ b/Engine/Source/Runtime/EmptyRHI/Private/EmptyRHI.cpp
@@ -28,8 +28,6 @@ FEmptyDynamicRHI::FEmptyDynamicRHI()
 
 	// Initialize the RHI capabilities.
 
-//	GRHIShaderPlatform = 
-
 // 	GPixelCenterOffset = 0.0f;
 // 	GSupportsVertexInstancing = false;
 // 	GSupportsEmulatedVertexInstancing = false;
diff --git a/Engine/Source/Runtime/Engine/Classes/Materials/Material.h b/Engine/Source/Runtime/Engine/Classes/Materials/Material.h
--- a/Engine/Source/Runtime/Engine/Classes/Materials/Material.h
+++ b/Engine/Source/Runtime/Engine/Classes/Materials/Material.h
@@ -992,7 +992,7 @@ public:
 	 * Backs up all material shaders to memory through serialization, organized by FMaterialShaderMap. 
 	 * This will also clear all FMaterialShaderMap references to FShaders.
 	 */
-	ENGINE_API static void BackupMaterialShadersToMemory(EShaderPlatform ShaderPlatform, TMap<FMaterialShaderMap*, TScopedPointer<TArray<uint8> > >& ShaderMapToSerializedShaderData);
+	ENGINE_API static void BackupMaterialShadersToMemory(TMap<FMaterialShaderMap*, TScopedPointer<TArray<uint8> > >& ShaderMapToSerializedShaderData);
 
 	/** 
 	 * Recreates FShaders for FMaterialShaderMap's from the serialized data.  Shader maps may not be complete after this due to changes in the shader keys.
diff --git a/Engine/Source/Runtime/Engine/Private/GameEngine.cpp b/Engine/Source/Runtime/Engine/Private/GameEngine.cpp
--- a/Engine/Source/Runtime/Engine/Private/GameEngine.cpp
+++ b/Engine/Source/Runtime/Engine/Private/GameEngine.cpp
@@ -276,7 +276,7 @@ TSharedRef<SWindow> UGameEngine::CreateGameWindow()
 	FFormatNamedArguments Args;
 	Args.Add( TEXT("GameName"), FText::FromString( FApp::GetGameName() ) );
 	Args.Add( TEXT("PlatformArchitecture"), PlatformBits );
-	Args.Add( TEXT("RHIName"), FText::FromName( LegacyShaderPlatformToShaderFormat( GRHIShaderPlatform ) ) );
+	Args.Add( TEXT("RHIName"), FText::FromName( LegacyShaderPlatformToShaderFormat( GMaxRHIShaderPlatform ) ) );
 
 	const FText AppName = FText::Format( NSLOCTEXT("UnrealEd", "GameWindowTitle", "{GameName} ({PlatformArchitecture}-bit, {RHIName})"), Args );
 
diff --git a/Engine/Source/Runtime/Engine/Private/GlobalShader.cpp b/Engine/Source/Runtime/Engine/Private/GlobalShader.cpp
--- a/Engine/Source/Runtime/Engine/Private/GlobalShader.cpp
+++ b/Engine/Source/Runtime/Engine/Private/GlobalShader.cpp
@@ -397,7 +397,7 @@ TShaderMap<FGlobalShaderType>* GetGlobalShaderMap(EShaderPlatform Platform, bool
 		VerifyGlobalShaders(Platform, bLoadedFromCacheFile);
 
 		extern int32 GCreateShadersOnLoad;
-		if (GCreateShadersOnLoad && Platform == GRHIShaderPlatform)
+		if (GCreateShadersOnLoad && Platform == GMaxRHIShaderPlatform)
 		{
 			for (TMap<FShaderType*, TRefCountPtr<FShader> >::TConstIterator ShaderIt(GGlobalShaderMap[Platform]->GetShaders()); ShaderIt; ++ShaderIt)
 			{
@@ -453,9 +453,11 @@ void RecompileGlobalShaders()
 		// Flush pending accesses to the existing global shaders.
 		FlushRenderingCommands();
 
-		// MOBILEPREVIEWTODO: Use global feature level bitfield to recompile for all active platforms
-		GetGlobalShaderMap(GRHIShaderPlatform)->Empty();
-		VerifyGlobalShaders(GRHIShaderPlatform, false);
+		UMaterialInterface::IterateOverActiveFeatureLevels([&](ERHIFeatureLevel::Type InFeatureLevel) {
+			auto ShaderPlatform = GShaderPlatformForFeatureLevel[InFeatureLevel];
+			GetGlobalShaderMap(ShaderPlatform)->Empty();
+			VerifyGlobalShaders(ShaderPlatform, false);
+		});
 
 		GShaderCompilingManager->ProcessAsyncResults(false, true);
 
diff --git a/Engine/Source/Runtime/Engine/Private/Materials/Material.cpp b/Engine/Source/Runtime/Engine/Private/Materials/Material.cpp
--- a/Engine/Source/Runtime/Engine/Private/Materials/Material.cpp
+++ b/Engine/Source/Runtime/Engine/Private/Materials/Material.cpp
@@ -447,7 +447,7 @@ void SerializeInlineShaderMaps(const TMap<const ITargetPlatform*,TArray<FMateria
 				FMaterialResource& LoadedResource = LoadedResources[ResourceIndex];
 				FMaterialShaderMap* LoadedShaderMap = LoadedResource.GetGameThreadShaderMap();
 
-				if (LoadedShaderMap && LoadedShaderMap->GetShaderPlatform() == GRHIShaderPlatform)
+				if (LoadedShaderMap && LoadedShaderMap->GetShaderPlatform() == GMaxRHIShaderPlatform)
 				{
 					EMaterialQualityLevel::Type LoadedQualityLevel = LoadedShaderMap->GetShaderMapId().QualityLevel;
 					ERHIFeatureLevel::Type LoadedFeatureLevel = LoadedShaderMap->GetShaderMapId().FeatureLevel;
@@ -2879,10 +2879,10 @@ void UMaterial::UpdateMaterialShaders(TArray<FShaderType*>& ShaderTypesToFlush,
 	// Update any FMaterials not belonging to a UMaterialInterface, for example FExpressionPreviews
 	// If we did not do this, the editor would crash the next time it tried to render one of those previews
 	// And didn't find a shader that had been flushed for the preview's shader map.
-	FMaterial::UpdateEditorLoadedMaterialResources();
+	FMaterial::UpdateEditorLoadedMaterialResources(ShaderPlatform);
 }
 
-void UMaterial::BackupMaterialShadersToMemory(EShaderPlatform ShaderPlatform, TMap<FMaterialShaderMap*, TScopedPointer<TArray<uint8> > >& ShaderMapToSerializedShaderData)
+void UMaterial::BackupMaterialShadersToMemory(TMap<FMaterialShaderMap*, TScopedPointer<TArray<uint8> > >& ShaderMapToSerializedShaderData)
 {
 	// Process FMaterialShaderMap's referenced by UObjects (UMaterial, UMaterialInstance)
 	for (TObjectIterator<UMaterialInterface> It; It; ++It)
diff --git a/Engine/Source/Runtime/Engine/Private/Materials/MaterialShared.cpp b/Engine/Source/Runtime/Engine/Private/Materials/MaterialShared.cpp
--- a/Engine/Source/Runtime/Engine/Private/Materials/MaterialShared.cpp
+++ b/Engine/Source/Runtime/Engine/Private/Materials/MaterialShared.cpp
@@ -411,7 +411,7 @@ bool FMaterial::NeedsGBuffer() const
 {
 	check(IsInParallelRenderingThread());
 
-	if (IsOpenGLPlatform(GRHIShaderPlatform)) // @todo: TTP #341211
+	if (IsOpenGLPlatform(GMaxRHIShaderPlatform)) // @todo: TTP #341211
 	{
 		return true;
 	}
@@ -1913,14 +1913,14 @@ bool FMaterial::GetMaterialExpressionSource( FString& OutSource, TMap<FMaterialE
 }
 
 /** Recompiles any materials in the EditorLoadedMaterialResources list if they are not complete. */
-void FMaterial::UpdateEditorLoadedMaterialResources()
+void FMaterial::UpdateEditorLoadedMaterialResources(EShaderPlatform InShaderPlatform)
 {
 	for (TSet<FMaterial*>::TIterator It(EditorLoadedMaterialResources); It; ++It)
 	{
 		FMaterial* CurrentMaterial = *It;
 		if (!CurrentMaterial->GetGameThreadShaderMap() || !CurrentMaterial->GetGameThreadShaderMap()->IsComplete(CurrentMaterial, true))
 		{
-			CurrentMaterial->CacheShaders(GRHIShaderPlatform, true);
+			CurrentMaterial->CacheShaders(InShaderPlatform, true);
 		}
 	}
 }
diff --git a/Engine/Source/Runtime/Engine/Private/NavMeshRenderingHelpers.h b/Engine/Source/Runtime/Engine/Private/NavMeshRenderingHelpers.h
--- a/Engine/Source/Runtime/Engine/Private/NavMeshRenderingHelpers.h
+++ b/Engine/Source/Runtime/Engine/Private/NavMeshRenderingHelpers.h
@@ -485,7 +485,7 @@ public:
 			}
 		}
 
-		const bool bNeedToSwitchVerticalAxis = RHINeedsToSwitchVerticalAxis(GRHIShaderPlatform);
+		const bool bNeedToSwitchVerticalAxis = RHINeedsToSwitchVerticalAxis(View->GetShaderPlatform());
 		const FTexture2DRHIRef DepthTexture;
 		const auto FeatureLevel = View->GetFeatureLevel();
 			
diff --git a/Engine/Source/Runtime/Engine/Private/PointLightComponent.cpp b/Engine/Source/Runtime/Engine/Private/PointLightComponent.cpp
--- a/Engine/Source/Runtime/Engine/Private/PointLightComponent.cpp
+++ b/Engine/Source/Runtime/Engine/Private/PointLightComponent.cpp
@@ -69,7 +69,7 @@ public:
 	 */
 	virtual bool GetWholeSceneProjectedShadowInitializer(const FSceneViewFamily& ViewFamily, TArray<FWholeSceneProjectedShadowInitializer, TInlineAllocator<6> >& OutInitializers) const
 	{
-		if (ViewFamily.Scene->GetFeatureLevel() >= ERHIFeatureLevel::SM4
+		if (ViewFamily.GetFeatureLevel() >= ERHIFeatureLevel::SM4
 			&& GAllowPointLightCubemapShadows != 0)
 		{
 			FWholeSceneProjectedShadowInitializer& OutInitializer = *new(OutInitializers) FWholeSceneProjectedShadowInitializer;
@@ -84,7 +84,7 @@ public:
 			OutInitializer.MinLightW = 0.1f;
 			OutInitializer.MaxDistanceToCastInLightW = Radius;
 			OutInitializer.SplitIndex = INDEX_NONE;
-			OutInitializer.bRayTracedDistanceFieldShadow = UseRayTracedDistanceFieldShadows() && DoesPlatformSupportDistanceFieldShadowing(GRHIShaderPlatform);
+			OutInitializer.bRayTracedDistanceFieldShadow = UseRayTracedDistanceFieldShadows() && DoesPlatformSupportDistanceFieldShadowing(ViewFamily.GetShaderPlatform());
 		
 			return true;
 		}
diff --git a/Engine/Source/Runtime/Engine/Private/RendererSupport.cpp b/Engine/Source/Runtime/Engine/Private/RendererSupport.cpp
--- a/Engine/Source/Runtime/Engine/Private/RendererSupport.cpp
+++ b/Engine/Source/Runtime/Engine/Private/RendererSupport.cpp
@@ -50,7 +50,7 @@ void ClearReferencesToRendererModuleClasses(
 
 	// Save off shaders by serializing them into memory, and remove all shader map references to FShaders
 	GlobalShaderData = TScopedPointer<TArray<uint8> >(BackupGlobalShaderMap(GRHIShaderPlatform));
-	UMaterial::BackupMaterialShadersToMemory(GRHIShaderPlatform, ShaderMapToSerializedShaderData);
+	UMaterial::BackupMaterialShadersToMemory(ShaderMapToSerializedShaderData);
 
 	// Verify no FShaders still in memory
 	for(TLinkedList<FShaderType*>::TIterator It(FShaderType::GetTypeList()); It; It.Next())
diff --git a/Engine/Source/Runtime/Engine/Private/ShaderCompiler/ShaderCompiler.cpp b/Engine/Source/Runtime/Engine/Private/ShaderCompiler/ShaderCompiler.cpp
--- a/Engine/Source/Runtime/Engine/Private/ShaderCompiler/ShaderCompiler.cpp
+++ b/Engine/Source/Runtime/Engine/Private/ShaderCompiler/ShaderCompiler.cpp
@@ -2241,7 +2241,7 @@ public:
 		}
 
 		Payload << MaterialsToLoad;
-		uint32 ShaderPlatform = ( uint32 )GRHIShaderPlatform;
+		uint32 ShaderPlatform = ( uint32 )GMaxRHIShaderPlatform;
 		Payload << ShaderPlatform;
 		// tell the other side the Ids we have so it doesn't send back duplicates
 		// (need to serialize this into a TArray since FShaderResourceId isn't known in the file server)
@@ -2266,7 +2266,7 @@ public:
 		FlushRenderingCommands();
 
 		// reload the global shaders
-		GetGlobalShaderMap(GRHIShaderPlatform, true);
+		GetGlobalShaderMap(GMaxRHIShaderPlatform, true);
 
 		//invalidate global bound shader states so they will be created with the new shaders the next time they are set (in SetGlobalBoundShaderState)
 		for(TLinkedList<FGlobalBoundShaderStateResource*>::TIterator It(FGlobalBoundShaderStateResource::GetGlobalBoundShaderStateList());It;It.Next())
@@ -2283,7 +2283,7 @@ public:
 			// parse the shaders
 			FMemoryReader MemoryReader(MeshMaterialMaps, true);
 			FNameAsStringProxyArchive Ar(MemoryReader);
-			FMaterialShaderMap::LoadForRemoteRecompile(Ar, GRHIShaderPlatform, MaterialsToLoad);
+			FMaterialShaderMap::LoadForRemoteRecompile(Ar, GMaxRHIShaderPlatform, MaterialsToLoad);
 
 			// gather the shader maps to reattach
 			for (TObjectIterator<UMaterial> It; It; ++It)
@@ -2417,9 +2417,12 @@ bool RecompileShaders(const TCHAR* Cmd, FOutputDevice& Ar)
 				
 				TArray<const FVertexFactoryType*> FactoryTypes;
 
-				BeginRecompileGlobalShaders(ShaderTypes, GRHIShaderPlatform);
-				UMaterial::UpdateMaterialShaders(ShaderTypes, FactoryTypes, GRHIShaderPlatform);
-				FinishRecompileGlobalShaders();
+				UMaterialInterface::IterateOverActiveFeatureLevels([&](ERHIFeatureLevel::Type InFeatureLevel) {
+					auto ShaderPlatform = GShaderPlatformForFeatureLevel[InFeatureLevel];
+					BeginRecompileGlobalShaders(ShaderTypes, ShaderPlatform);
+					UMaterial::UpdateMaterialShaders(ShaderTypes, FactoryTypes, ShaderPlatform);
+					FinishRecompileGlobalShaders();
+				});
 			}
 		}
 
diff --git a/Engine/Source/Runtime/Engine/Private/SpotLightComponent.cpp b/Engine/Source/Runtime/Engine/Private/SpotLightComponent.cpp
--- a/Engine/Source/Runtime/Engine/Private/SpotLightComponent.cpp
+++ b/Engine/Source/Runtime/Engine/Private/SpotLightComponent.cpp
@@ -137,7 +137,7 @@ public:
 		OutInitializer.MinLightW = 0.1f;
 		OutInitializer.MaxDistanceToCastInLightW = Radius;
 		OutInitializer.SplitIndex = INDEX_NONE;
-		OutInitializer.bRayTracedDistanceFieldShadow = UseRayTracedDistanceFieldShadows() && DoesPlatformSupportDistanceFieldShadowing(GRHIShaderPlatform);
+		OutInitializer.bRayTracedDistanceFieldShadow = UseRayTracedDistanceFieldShadows() && DoesPlatformSupportDistanceFieldShadowing(ViewFamily.GetShaderPlatform());
 		return true;
 	}
 
diff --git a/Engine/Source/Runtime/Engine/Private/StaticMeshRender.cpp b/Engine/Source/Runtime/Engine/Private/StaticMeshRender.cpp
--- a/Engine/Source/Runtime/Engine/Private/StaticMeshRender.cpp
+++ b/Engine/Source/Runtime/Engine/Private/StaticMeshRender.cpp
@@ -309,7 +309,7 @@ void FStaticMeshSceneProxy::SetIndexSource(int32 LODIndex, int32 SectionIndex, F
 	if (bWireframe)
 	{
 		if( LODModel.WireframeIndexBuffer.IsInitialized()
-			&& !(RHISupportsTessellation(GRHIShaderPlatform) && OutMeshElement.VertexFactory->GetType()->SupportsTessellationShaders())
+			&& !(RHISupportsTessellation(GetScene()->GetShaderPlatform()) && OutMeshElement.VertexFactory->GetType()->SupportsTessellationShaders())
 			)
 		{
 			OutMeshElement.Type = PT_LineList;
diff --git a/Engine/Source/Runtime/Engine/Public/MaterialShared.h b/Engine/Source/Runtime/Engine/Public/MaterialShared.h
--- a/Engine/Source/Runtime/Engine/Public/MaterialShared.h
+++ b/Engine/Source/Runtime/Engine/Public/MaterialShared.h
@@ -1047,7 +1047,7 @@ public:
 	}
 
 	/** Recompiles any materials in the EditorLoadedMaterialResources list if they are not complete. */
-	static void UpdateEditorLoadedMaterialResources();
+	static void UpdateEditorLoadedMaterialResources(EShaderPlatform InShaderPlatform);
 
 	/** Backs up any FShaders in editor loaded materials to memory through serialization and clears FShader references. */
 	static void BackupEditorLoadedMaterialShadersToMemory(TMap<FMaterialShaderMap*, TScopedPointer<TArray<uint8> > >& ShaderMapToSerializedShaderData);
diff --git a/Engine/Source/Runtime/Engine/Public/SceneInterface.h b/Engine/Source/Runtime/Engine/Public/SceneInterface.h
--- a/Engine/Source/Runtime/Engine/Public/SceneInterface.h
+++ b/Engine/Source/Runtime/Engine/Public/SceneInterface.h
@@ -305,6 +305,7 @@ public:
 	virtual bool IsEditorScene() const { return false; }
 
 	virtual ERHIFeatureLevel::Type GetFeatureLevel() const { return GMaxRHIFeatureLevel; }
+	EShaderPlatform GetShaderPlatform() const { return GShaderPlatformForFeatureLevel[GetFeatureLevel()]; }
 
 	bool ShouldUseDeferredRenderer() const
 	{
diff --git a/Engine/Source/Runtime/Engine/Public/SceneView.h b/Engine/Source/Runtime/Engine/Public/SceneView.h
--- a/Engine/Source/Runtime/Engine/Public/SceneView.h
+++ b/Engine/Source/Runtime/Engine/Public/SceneView.h
@@ -615,6 +615,8 @@ public:
 	ENGINE_API void ComputeFamilySize();
 
 	ENGINE_API ERHIFeatureLevel::Type GetFeatureLevel() const;
+
+	EShaderPlatform GetShaderPlatform() const { return GShaderPlatformForFeatureLevel[GetFeatureLevel()]; }
 };
 
 /**
diff --git a/Engine/Source/Runtime/Renderer/Private/CompositionLighting/CompositionLighting.cpp b/Engine/Source/Runtime/Renderer/Private/CompositionLighting/CompositionLighting.cpp
--- a/Engine/Source/Runtime/Renderer/Private/CompositionLighting/CompositionLighting.cpp
+++ b/Engine/Source/Runtime/Renderer/Private/CompositionLighting/CompositionLighting.cpp
@@ -357,7 +357,7 @@ void FCompositionLighting::ProcessLighting(FRHICommandListImmediate& RHICmdList,
 
 			if (View.bScreenSpaceSubsurfacePassNeeded && Radius > 0 && !bSimpleDynamicLighting && View.Family->EngineShowFlags.SubsurfaceScattering && 
 				//@todo-rco: Remove this when we fix the cross-compiler
-				!IsOpenGLPlatform(GRHIShaderPlatform))
+				!IsOpenGLPlatform(View.GetShaderPlatform()))
 			{
 				FRenderingCompositePass* PassSetup = Context.Graph.RegisterPass(new(FMemStack::Get()) FRCPassPostProcessSubsurfaceSetup(false, true));
 				PassSetup->SetInput(ePId_Input0, Context.FinalOutput);
diff --git a/Engine/Source/Runtime/Renderer/Private/DeferredShadingRenderer.cpp b/Engine/Source/Runtime/Renderer/Private/DeferredShadingRenderer.cpp
--- a/Engine/Source/Runtime/Renderer/Private/DeferredShadingRenderer.cpp
+++ b/Engine/Source/Runtime/Renderer/Private/DeferredShadingRenderer.cpp
@@ -1468,7 +1468,7 @@ FGlobalBoundShaderState DownsampleDepthBoundShaderState;
 /** Updates the downsized depth buffer with the current full resolution depth buffer. */
 void FDeferredShadingSceneRenderer::UpdateDownsampledDepthSurface(FRHICommandList& RHICmdList)
 {
-	if (GSceneRenderTargets.UseDownsizedOcclusionQueries() && IsFeatureLevelSupported(GRHIShaderPlatform, ERHIFeatureLevel::SM4))
+	if (GSceneRenderTargets.UseDownsizedOcclusionQueries() && (FeatureLevel >= ERHIFeatureLevel::SM4))
 	{
 		SetRenderTarget(RHICmdList, NULL, GSceneRenderTargets.GetSmallDepthSurface());
 
diff --git a/Engine/Source/Runtime/Renderer/Private/DistanceFieldShadowing.cpp b/Engine/Source/Runtime/Renderer/Private/DistanceFieldShadowing.cpp
--- a/Engine/Source/Runtime/Renderer/Private/DistanceFieldShadowing.cpp
+++ b/Engine/Source/Runtime/Renderer/Private/DistanceFieldShadowing.cpp
@@ -629,7 +629,7 @@ void FProjectedShadowInfo::RenderRayTracedDistanceFieldProjection(FRHICommandLis
 {
 	if (GDistanceFieldShadowing
 		&& View.GetFeatureLevel() >= ERHIFeatureLevel::SM5
-		&& DoesPlatformSupportDistanceFieldShadowing(GRHIShaderPlatform))
+		&& DoesPlatformSupportDistanceFieldShadowing(View.GetShaderPlatform()))
 	{
 		QUICK_SCOPE_CYCLE_COUNTER(STAT_RenderRayTracedDistanceFieldShadows);
 		SCOPED_DRAW_EVENT(RHICmdList, RayTracedDistanceFieldShadow, DEC_SCENE_ITEMS);
diff --git a/Engine/Source/Runtime/Renderer/Private/DistanceFieldSurfaceCacheLighting.cpp b/Engine/Source/Runtime/Renderer/Private/DistanceFieldSurfaceCacheLighting.cpp
--- a/Engine/Source/Runtime/Renderer/Private/DistanceFieldSurfaceCacheLighting.cpp
+++ b/Engine/Source/Runtime/Renderer/Private/DistanceFieldSurfaceCacheLighting.cpp
@@ -3316,7 +3316,7 @@ bool FDeferredShadingSceneRenderer::RenderDistanceFieldAOSurfaceCache(FRHIComman
 
 	if (GDistanceFieldAO 
 		&& FeatureLevel >= ERHIFeatureLevel::SM5
-		&& DoesPlatformSupportDistanceFieldAO(GRHIShaderPlatform)
+		&& DoesPlatformSupportDistanceFieldAO(View.GetShaderPlatform())
 		&& Views.Num() == 1
 		// ViewState is used to cache tile intersection resources which have to be sized based on the view
 		&& View.State)
@@ -3688,7 +3688,7 @@ void FDeferredShadingSceneRenderer::RenderMeshDistanceFieldVisualization(FRHICom
 
 	if (GDistanceFieldAO 
 		&& FeatureLevel >= ERHIFeatureLevel::SM5
-		&& DoesPlatformSupportDistanceFieldAO(GRHIShaderPlatform)
+		&& DoesPlatformSupportDistanceFieldAO(View.GetShaderPlatform())
 		&& Views.Num() == 1)
 	{
 		QUICK_SCOPE_CYCLE_COUNTER(STAT_RenderMeshDistanceFieldVis);
diff --git a/Engine/Source/Runtime/Renderer/Private/ForwardShadingRenderer.cpp b/Engine/Source/Runtime/Renderer/Private/ForwardShadingRenderer.cpp
--- a/Engine/Source/Runtime/Renderer/Private/ForwardShadingRenderer.cpp
+++ b/Engine/Source/Runtime/Renderer/Private/ForwardShadingRenderer.cpp
@@ -80,7 +80,7 @@ void FForwardShadingSceneRenderer::Render(FRHICommandListImmediate& RHICmdList)
 		return;
 	}
 
-	auto FeatureLevel = ViewFamily.Scene->GetFeatureLevel();
+	auto FeatureLevel = ViewFamily.GetFeatureLevel();
 
 	// Initialize global system textures (pass-through if already initialized).
 	GSystemTextures.InitializeTextures(RHICmdList, FeatureLevel);
@@ -161,7 +161,7 @@ void FForwardShadingSceneRenderer::Render(FRHICommandListImmediate& RHICmdList)
 			GSupportsShaderFramebufferFetch &&
 			ViewFamily.EngineShowFlags.PostProcessing &&
 			((View.bLightShaftUse) || (View.FinalPostProcessSettings.DepthOfFieldScale > 0.0) || 
-			((GRHIShaderPlatform == SP_METAL) && (CVarMobileMSAA ? CVarMobileMSAA->GetValueOnAnyThread() > 1 : false))
+			((ViewFamily.GetShaderPlatform() == SP_METAL) && (CVarMobileMSAA ? CVarMobileMSAA->GetValueOnAnyThread() > 1 : false))
 			);
 
 		// Convert alpha from depth to circle of confusion with sunshaft intensity.
diff --git a/Engine/Source/Runtime/Renderer/Private/ForwardTranslucentRendering.cpp b/Engine/Source/Runtime/Renderer/Private/ForwardTranslucentRendering.cpp
--- a/Engine/Source/Runtime/Renderer/Private/ForwardTranslucentRendering.cpp
+++ b/Engine/Source/Runtime/Renderer/Private/ForwardTranslucentRendering.cpp
@@ -343,7 +343,7 @@ void FForwardShadingSceneRenderer::RenderTranslucency(FRHICommandListImmediate&
 			#if PLATFORM_HTML5
 			// Copy the view so emulation of framebuffer fetch works for alpha=depth.
 			// Possible optimization: this copy shouldn't be needed unless something uses fetch of depth.
-			if(bLinearHDR64 && GSupportsRenderTargetFormat_PF_FloatRGBA && (GSupportsShaderFramebufferFetch == false) && (!IsPCPlatform(GRHIShaderPlatform)))
+			if(bLinearHDR64 && GSupportsRenderTargetFormat_PF_FloatRGBA && (GSupportsShaderFramebufferFetch == false) && (!IsPCPlatform(View.GetShaderPlatform())))
 			{
 				CopySceneAlpha(RHICmdList, View);
 			}
diff --git a/Engine/Source/Runtime/Renderer/Private/LightMapDensityRendering.h b/Engine/Source/Runtime/Renderer/Private/LightMapDensityRendering.h
--- a/Engine/Source/Runtime/Renderer/Private/LightMapDensityRendering.h
+++ b/Engine/Source/Runtime/Renderer/Private/LightMapDensityRendering.h
@@ -273,7 +273,7 @@ public:
 		DomainShader = NULL;
 
 		const EMaterialTessellationMode MaterialTessellationMode = MaterialResource->GetTessellationMode();
-		if(RHISupportsTessellation(GRHIShaderPlatform)
+		if(RHISupportsTessellation(View.GetShaderPlatform())
 			&& InVertexFactory->GetType()->SupportsTessellationShaders() 
 			&& MaterialTessellationMode != MTM_NoTessellation)
 		{
diff --git a/Engine/Source/Runtime/Renderer/Private/LightPropagationVolume.cpp b/Engine/Source/Runtime/Renderer/Private/LightPropagationVolume.cpp
--- a/Engine/Source/Runtime/Renderer/Private/LightPropagationVolume.cpp
+++ b/Engine/Source/Runtime/Renderer/Private/LightPropagationVolume.cpp
@@ -1281,7 +1281,7 @@ void FSceneViewState::CreateLightPropagationVolumeIfNeeded(ERHIFeatureLevel::Typ
 	bool bLPV = UseLightPropagationVolumeRT(InFeatureLevel);
 
 	//@todo-rco: Remove this when reenabling for OpenGL
-	if (bLPV && !IsOpenGLPlatform(GRHIShaderPlatform))
+	if (bLPV && !IsOpenGLPlatform(GShaderPlatformForFeatureLevel[InFeatureLevel]))
 	{
 		LightPropagationVolume = new FLightPropagationVolume();
 	}
diff --git a/Engine/Source/Runtime/Renderer/Private/PostProcess/PostProcessMobile.cpp b/Engine/Source/Runtime/Renderer/Private/PostProcess/PostProcessMobile.cpp
--- a/Engine/Source/Runtime/Renderer/Private/PostProcess/PostProcessMobile.cpp
+++ b/Engine/Source/Runtime/Renderer/Private/PostProcess/PostProcessMobile.cpp
@@ -926,7 +926,7 @@ void FRCPassPostProcessSunMaskES2::SetShader(const FRenderingCompositePassContex
 	uint32 UseFetchSunDof = (UseFetch << 2) + (UseSun << 1) + UseDof;
 
 	static const auto CVarMobileMSAA = IConsoleManager::Get().FindTConsoleVariableDataInt(TEXT("r.MobileMSAA"));
-	if((GRHIShaderPlatform == SP_METAL) && (CVarMobileMSAA ? CVarMobileMSAA->GetValueOnAnyThread() > 1 : false))
+	if((Context.GetShaderPlatform() == SP_METAL) && (CVarMobileMSAA ? CVarMobileMSAA->GetValueOnAnyThread() > 1 : false))
 	{
 		UseFetchSunDof = 8;
 	}
diff --git a/Engine/Source/Runtime/Renderer/Private/PostProcess/RenderingCompositionGraph.h b/Engine/Source/Runtime/Renderer/Private/PostProcess/RenderingCompositionGraph.h
--- a/Engine/Source/Runtime/Renderer/Private/PostProcess/RenderingCompositionGraph.h
+++ b/Engine/Source/Runtime/Renderer/Private/PostProcess/RenderingCompositionGraph.h
@@ -132,6 +132,7 @@ struct FRenderingCompositePassContext
 	}
 
 	ERHIFeatureLevel::Type GetFeatureLevel() const { return FeatureLevel; }
+	EShaderPlatform GetShaderPlatform() const { return GShaderPlatformForFeatureLevel[FeatureLevel]; }
 	TShaderMap<FGlobalShaderType>* GetShaderMap() const { check(ShaderMap); return ShaderMap; }
 
 	FRenderingCompositePass* Root;
diff --git a/Engine/Source/Runtime/ShaderCore/Private/Shader.cpp b/Engine/Source/Runtime/ShaderCore/Private/Shader.cpp
--- a/Engine/Source/Runtime/ShaderCore/Private/Shader.cpp
+++ b/Engine/Source/Runtime/ShaderCore/Private/Shader.cpp
@@ -435,7 +435,7 @@ void FShaderResource::InitRHI()
  		if (FPlatformProperties::RequiresCookedData())
  		{
  			UE_LOG(LogShaders, Fatal, TEXT("FShaderResource::InitRHI got platform %s but it is not compatible with %s"), 
- 				*LegacyShaderPlatformToShaderFormat((EShaderPlatform)Target.Platform).ToString(), *LegacyShaderPlatformToShaderFormat(GRHIShaderPlatform).ToString() );
+				*LegacyShaderPlatformToShaderFormat((EShaderPlatform)Target.Platform).ToString(), *LegacyShaderPlatformToShaderFormat(GMaxRHIShaderPlatform).ToString());
  		}
  		return;
  	}
