Changed VisualLog's automation tests to run only as game or editor tests #ue4

diff --git a/Engine/Source/Runtime/Engine/Private/VisualLogger/VisualLoggerAutomationTests.cpp b/Engine/Source/Runtime/Engine/Private/VisualLogger/VisualLoggerAutomationTests.cpp
--- a/Engine/Source/Runtime/Engine/Private/VisualLogger/VisualLoggerAutomationTests.cpp
+++ b/Engine/Source/Runtime/Engine/Private/VisualLogger/VisualLoggerAutomationTests.cpp
@@ -94,7 +94,7 @@ struct FTestDeviceContext
 	FVisualLoggerTestDevice Device;
 };
 
-IMPLEMENT_SIMPLE_AUTOMATION_TEST(FVisualLogTest, "Engine.VisualLogger.Logging simple text", EAutomationTestFlags::ATF_Game | EAutomationTestFlags::ATF_Editor | EAutomationTestFlags::ATF_Commandlet)
+IMPLEMENT_SIMPLE_AUTOMATION_TEST(FVisualLogTest, "Engine.VisualLogger.Logging simple text", EAutomationTestFlags::ATF_Game | EAutomationTestFlags::ATF_Editor)
 /**
 *
 *
@@ -142,7 +142,7 @@ bool FVisualLogTest::RunTest(const FString& Parameters)
 	return true;
 }
 
-IMPLEMENT_SIMPLE_AUTOMATION_TEST(FVisualLogSegmentsTest, "Engine.VisualLogger.Logging segment shape", EAutomationTestFlags::ATF_Game | EAutomationTestFlags::ATF_Editor | EAutomationTestFlags::ATF_Commandlet)
+IMPLEMENT_SIMPLE_AUTOMATION_TEST(FVisualLogSegmentsTest, "Engine.VisualLogger.Logging segment shape", EAutomationTestFlags::ATF_Game | EAutomationTestFlags::ATF_Editor)
 bool FVisualLogSegmentsTest::RunTest(const FString& Parameters)
 {
 	FTestDeviceContext Context;
