Make EWindowMode::Fullscreen the default on Mac OS X now that it should work.

diff --git a/Engine/Source/Runtime/Core/Private/Mac/CocoaWindow.cpp b/Engine/Source/Runtime/Core/Private/Mac/CocoaWindow.cpp
--- a/Engine/Source/Runtime/Core/Private/Mac/CocoaWindow.cpp
+++ b/Engine/Source/Runtime/Core/Private/Mac/CocoaWindow.cpp
@@ -357,8 +357,7 @@ NSString* NSPerformDragOperation = @"NSPerformDragOperation";
 	// Handle clicking on the titlebar fullscreen item
 	if(self.TargetWindowMode == EWindowMode::Windowed)
 	{
-		// @todo: Fix fullscreen mode mouse coordinate handling - for now default to windowed fullscreen
-		self.TargetWindowMode = EWindowMode::WindowedFullscreen;
+		self.TargetWindowMode = EWindowMode::Fullscreen;
 	}
 }
 
diff --git a/Engine/Source/Runtime/Engine/Private/GameEngine.cpp b/Engine/Source/Runtime/Engine/Private/GameEngine.cpp
--- a/Engine/Source/Runtime/Engine/Private/GameEngine.cpp
+++ b/Engine/Source/Runtime/Engine/Private/GameEngine.cpp
@@ -181,12 +181,7 @@ void UGameEngine::ConditionallyOverrideSettings(int32& ResolutionX, int32& Resol
 	else if (FParse::Param(FCommandLine::Get(),TEXT("FullScreen")))
 	{
 		// -FullScreen
-#if PLATFORM_MAC // @TODO: Fullscreen mode isn't working fully on OS X as it requires explicit CGL back-buffer size and mouse<->display coordinate conversions.
-		// For now use the windowed fullscreen mode which merely stretches to fill the display, which is what OS X does for us by default.
-		WindowMode = EWindowMode::WindowedFullscreen;
-#else
 		WindowMode = EWindowMode::Fullscreen;
-#endif
 	}
 
 	//fullscreen is always supported, but don't allow windowed mode on platforms that dont' support it.
diff --git a/Engine/Source/Runtime/Engine/Private/GameViewportClient.cpp b/Engine/Source/Runtime/Engine/Private/GameViewportClient.cpp
--- a/Engine/Source/Runtime/Engine/Private/GameViewportClient.cpp
+++ b/Engine/Source/Runtime/Engine/Private/GameViewportClient.cpp
@@ -2430,12 +2430,7 @@ bool UGameViewportClient::SetDisplayConfiguration(const FIntPoint* Dimensions, E
 
 bool UGameViewportClient::HandleToggleFullscreenCommand(const TCHAR* Cmd, FOutputDevice& Ar)
 {
-#if PLATFORM_MAC // @TODO: Fullscreen mode isn't working fully on OS X as it requires explicit CGL back-buffer size and mouse<->display coordinate conversions.
-				 // For now use the windowed fullscreen mode which merely stretches to fill the display, which is what OS X does for us by default.
-	return SetDisplayConfiguration(NULL, Viewport->IsFullscreen() ? EWindowMode::Windowed : EWindowMode::WindowedFullscreen);
-#else
 	return SetDisplayConfiguration(NULL, Viewport->IsFullscreen() ? EWindowMode::Windowed : EWindowMode::Fullscreen);
-#endif
 }
 
 bool UGameViewportClient::HandleSetResCommand( const TCHAR* Cmd, FOutputDevice& Ar )
