- Fixed issue when mesh index doesn't match with skeleton index
- Fixed crash when clicking viewport crashes when physics asset is missing

diff --git a/Engine/Source/Runtime/Engine/Private/Animation/AnimSequence.cpp b/Engine/Source/Runtime/Engine/Private/Animation/AnimSequence.cpp
--- a/Engine/Source/Runtime/Engine/Private/Animation/AnimSequence.cpp
+++ b/Engine/Source/Runtime/Engine/Private/Animation/AnimSequence.cpp
@@ -1577,6 +1577,20 @@ bool UAnimSequence::IsValidAdditive() const
 
 #if WITH_EDITOR
 
+int32 FindMeshBoneIndexFromBoneName(USkeleton * Skeleton, const FName &BoneName)
+{
+	USkeletalMesh * PreviewMesh = Skeleton->GetPreviewMesh();
+	const int32& SkeletonBoneIndex = Skeleton->GetReferenceSkeleton().FindBoneIndex(BoneName);
+
+	int32 BoneIndex = INDEX_NONE;
+
+	if(SkeletonBoneIndex != INDEX_NONE)
+	{
+		BoneIndex = Skeleton->GetMeshBoneIndexFromSkeletonBoneIndex(PreviewMesh, SkeletonBoneIndex);
+	}
+
+	return BoneIndex;
+}
 void FillUpTransformBasedOnRig(USkeleton* Skeleton, TArray<FTransform> & NodeSpaceBases, TArray<FTransform> &Rotations, TArray<FTransform> & Translations)
 {
 	TArray<FTransform> SpaceBases;
@@ -1605,14 +1619,7 @@ void FillUpTransformBasedOnRig(USkeleton* Skeleton, TArray<FTransform> & NodeSpa
 		{
 			const FName NodeName = Rig->GetNodeName(Index);
 			const FName& BoneName = Skeleton->GetRigBoneMapping(NodeName);
-			const int32& SkeletonBoneIndex = Skeleton->GetReferenceSkeleton().FindBoneIndex(BoneName);
-
-			int32 BoneIndex = INDEX_NONE;
-
-			if (SkeletonBoneIndex != INDEX_NONE)
-			{
-				BoneIndex = Skeleton->GetMeshBoneIndexFromSkeletonBoneIndex(PreviewMesh, SkeletonBoneIndex);
-			}
+			const int32& BoneIndex = FindMeshBoneIndexFromBoneName(Skeleton, BoneName);
 			
 			if (BoneIndex == INDEX_NONE)
 			{
@@ -1638,7 +1645,7 @@ void FillUpTransformBasedOnRig(USkeleton* Skeleton, TArray<FTransform> & NodeSpa
 					if (RotConstraint.TransformConstraints.Num() > 0)
 					{
 						const FName& ParentBoneName = Skeleton->GetRigBoneMapping(RotConstraint.TransformConstraints[0].ParentSpace);
-						const int32& ParentBoneIndex = Skeleton->GetReferenceSkeleton().FindBoneIndex(ParentBoneName);
+						const int32& ParentBoneIndex = FindMeshBoneIndexFromBoneName(Skeleton, ParentBoneName);
 
 						if (ParentBoneIndex != INDEX_NONE)
 						{
@@ -1652,7 +1659,7 @@ void FillUpTransformBasedOnRig(USkeleton* Skeleton, TArray<FTransform> & NodeSpa
 					if (TransConstraint.TransformConstraints.Num() > 0)
 					{
 						const FName& ParentBoneName = Skeleton->GetRigBoneMapping(TransConstraint.TransformConstraints[0].ParentSpace);
-						const int32& ParentBoneIndex = Skeleton->GetReferenceSkeleton().FindBoneIndex(ParentBoneName);
+						const int32& ParentBoneIndex = FindMeshBoneIndexFromBoneName(Skeleton, ParentBoneName);
 
 						if (ParentBoneIndex != INDEX_NONE)
 						{
diff --git a/Engine/Source/Runtime/Engine/Private/SkeletalMeshComponentPhysics.cpp b/Engine/Source/Runtime/Engine/Private/SkeletalMeshComponentPhysics.cpp
--- a/Engine/Source/Runtime/Engine/Private/SkeletalMeshComponentPhysics.cpp
+++ b/Engine/Source/Runtime/Engine/Private/SkeletalMeshComponentPhysics.cpp
@@ -1775,6 +1775,7 @@ extern float DebugLineLifetime;
 bool USkeletalMeshComponent::LineTraceComponent(struct FHitResult& OutHit, const FVector Start, const FVector End, const struct FCollisionQueryParams& Params)
 {
 	UPhysicsAsset* const PhysicsAsset = GetPhysicsAsset();
+	UWorld* const World = GetWorld();
 	bool bHaveHit = false;
 
 	float MinTime = MAX_FLT;
@@ -1793,7 +1794,7 @@ bool USkeletalMeshComponent::LineTraceComponent(struct FHitResult& OutHit, const
 	}
 
 #if !(UE_BUILD_SHIPPING || UE_BUILD_TEST)
-	if((GetWorld()->DebugDrawTraceTag != NAME_None) && (GetWorld()->DebugDrawTraceTag == Params.TraceTag))
+	if(World && (World->DebugDrawTraceTag != NAME_None) && (World->DebugDrawTraceTag == Params.TraceTag))
 	{
 		TArray<FHitResult> Hits;
 		if (bHaveHit)
