Changed the file extension in Mac AVI writer to .mov as that's what we're actually supporting there.

diff --git a/Engine/Source/Runtime/Engine/Private/AVIWriter.cpp b/Engine/Source/Runtime/Engine/Private/AVIWriter.cpp
--- a/Engine/Source/Runtime/Engine/Private/AVIWriter.cpp
+++ b/Engine/Source/Runtime/Engine/Private/AVIWriter.cpp
@@ -241,11 +241,11 @@ public:
 							GEngine->MatineeCaptureFPS, CaptureViewport->GetSizeXY().X, CaptureViewport->GetSizeXY().Y );
 						if (GEngine->bCompressMatineeCapture)
 						{
-							FCString::Sprintf( File, TEXT("%s_%d.avi"), *(OutputDir + TEXT("/") + GEngine->MatineePackageCaptureName + "_" + GEngine->MatineeCaptureName + FileNameInfo + "_compressed"), MovieCaptureIndex );
+							FCString::Sprintf( File, TEXT("%s_%d.mov"), *(OutputDir + TEXT("/") + GEngine->MatineePackageCaptureName + "_" + GEngine->MatineeCaptureName + FileNameInfo + "_compressed"), MovieCaptureIndex );
 						}
 						else
 						{
-							FCString::Sprintf( File, TEXT("%s_%d.avi"), *(OutputDir + TEXT("/") + GEngine->MatineePackageCaptureName + "_" + GEngine->MatineeCaptureName + FileNameInfo), MovieCaptureIndex );
+							FCString::Sprintf( File, TEXT("%s_%d.mov"), *(OutputDir + TEXT("/") + GEngine->MatineePackageCaptureName + "_" + GEngine->MatineeCaptureName + FileNameInfo), MovieCaptureIndex );
 						}
 						if (IFileManager::Get().FileSize(File) != -1)
 						{
@@ -262,7 +262,7 @@ public:
 					bool bFoundUnusedName = false;
 					while (!bFoundUnusedName)
 					{
-						FString FileName = FString::Printf( TEXT("Movie%i.avi"), MovieCaptureIndex );
+						FString FileName = FString::Printf( TEXT("Movie%i.mov"), MovieCaptureIndex );
 						FCString::Sprintf( File, TEXT("%s"), *(OutputDir / FileName) );
 						if (IFileManager::Get().FileSize(File) != -1)
 						{
@@ -503,11 +503,11 @@ public:
 																	 GEngine->MatineeCaptureFPS, CaptureViewport->GetSizeXY().X, CaptureViewport->GetSizeXY().Y );
 						if (GEngine->bCompressMatineeCapture)
 						{
-							FCString::Sprintf( File, TEXT("%s_%d.avi"), *(OutputDir + TEXT("/") + GEngine->MatineePackageCaptureName + "_" + GEngine->MatineeCaptureName + FileNameInfo + "_compressed"), MovieCaptureIndex );
+							FCString::Sprintf( File, TEXT("%s_%d.mov"), *(OutputDir + TEXT("/") + GEngine->MatineePackageCaptureName + "_" + GEngine->MatineeCaptureName + FileNameInfo + "_compressed"), MovieCaptureIndex );
 						}
 						else
 						{
-							FCString::Sprintf( File, TEXT("%s_%d.avi"), *(OutputDir + TEXT("/") + GEngine->MatineePackageCaptureName + "_" + GEngine->MatineeCaptureName + FileNameInfo), MovieCaptureIndex );
+							FCString::Sprintf( File, TEXT("%s_%d.mov"), *(OutputDir + TEXT("/") + GEngine->MatineePackageCaptureName + "_" + GEngine->MatineeCaptureName + FileNameInfo), MovieCaptureIndex );
 						}
 						if (IFileManager::Get().FileSize(File) != -1)
 						{
@@ -524,7 +524,7 @@ public:
 					bool bFoundUnusedName = false;
 					while (!bFoundUnusedName)
 					{
-						FString FileName = FString::Printf( TEXT("Movie%i.avi"), MovieCaptureIndex );
+						FString FileName = FString::Printf( TEXT("Movie%i.mov"), MovieCaptureIndex );
 						FCString::Sprintf( File, TEXT("%s"), *(OutputDir / FileName) );
 						if (IFileManager::Get().FileSize(File) != -1)
 						{
