#Orion

- Fixing crash
- Adding way to bypass CastChecked

diff --git a/Engine/Source/Runtime/Engine/Classes/Engine/LocalPlayer.h b/Engine/Source/Runtime/Engine/Classes/Engine/LocalPlayer.h
--- a/Engine/Source/Runtime/Engine/Classes/Engine/LocalPlayer.h
+++ b/Engine/Source/Runtime/Engine/Classes/Engine/LocalPlayer.h
@@ -32,9 +32,16 @@ struct ENGINE_API FLocalPlayerContext
 
 	/** Templated version of GetPlayerController() */
 	template<class T>
-	FORCEINLINE T* GetPlayerController() const
+	FORCEINLINE T* GetPlayerController(bool bCastChecked = true) const
 	{
-		return CastChecked<T>(GetPlayerController(), ECastCheckedType::NullAllowed);
+		if (bCastChecked)
+		{ 
+			return CastChecked<T>(GetPlayerController(), ECastCheckedType::NullAllowed);
+		}
+		else
+		{
+			return Cast<T>(GetPlayerController());
+		}
 	}
 
 	/** Getter for the Game State */
@@ -42,9 +49,16 @@ struct ENGINE_API FLocalPlayerContext
 
 	/** Templated Getter for the Game State */
 	template<class T>
-	FORCEINLINE T* GetGameState() const
+	FORCEINLINE T* GetGameState(bool bCastChecked = true) const
 	{
-		return CastChecked<T>(GetGameState(), ECastCheckedType::NullAllowed);
+		if (bCastChecked)
+		{
+			return CastChecked<T>(GetGameState(), ECastCheckedType::NullAllowed);
+		}
+		else
+		{
+			return Cast<T>(GetGameState());
+		}
 	}
 
 	/** Getter for the Player State */
@@ -52,9 +66,16 @@ struct ENGINE_API FLocalPlayerContext
 
 	/** Templated Getter for the Player State */
 	template<class T>
-	FORCEINLINE T* GetPlayerState() const
+	FORCEINLINE T* GetPlayerState(bool bCastChecked = true) const
 	{
-		return CastChecked<T>(GetPlayerState(), ECastCheckedType::NullAllowed);
+		if (bCastChecked)
+		{
+			return CastChecked<T>(GetPlayerState(), ECastCheckedType::NullAllowed);
+		}
+		else
+		{
+			return Cast<T>(GetPlayerState());
+		}
 	}
 
 	/** Getter for this player's HUD */
@@ -62,9 +83,16 @@ struct ENGINE_API FLocalPlayerContext
 
 	/** Templated Getter for the HUD */
 	template<class T>
-	FORCEINLINE T* GetHUD() const
+	FORCEINLINE T* GetHUD(bool bCastChecked = true) const
 	{
-		return CastChecked<T>(GetHUD(), ECastCheckedType::NullAllowed);
+		if (bCastChecked)
+		{
+			return CastChecked<T>(GetHUD(), ECastCheckedType::NullAllowed);
+		}
+		else
+		{
+			return Cast<T>(GetHUD());
+		}
 	}
 
 	/** Getter for the base pawn of this player */
@@ -72,9 +100,16 @@ struct ENGINE_API FLocalPlayerContext
 
 	/** Templated getter for the player's pawn */
 	template<class T>
-	FORCEINLINE T* GetPawn() const
+	FORCEINLINE T* GetPawn(bool bCastChecked = true) const
 	{
-		return CastChecked<T>(GetPawn(), ECastCheckedType::NullAllowed);
+		if (bCastChecked)
+		{
+			return CastChecked<T>(GetPawn(), ECastCheckedType::NullAllowed);
+		}
+		else
+		{
+			return Cast<T>(GetPawn());
+		}
 	}
 
 private:	
