Fix destruction performance regression where all destructibles registered for onContact even their chunks (which is super expensive)
Now we only register for onContact if destruction impact is enabled

diff --git a/Engine/Source/Runtime/Engine/Private/DestructibleComponent.cpp b/Engine/Source/Runtime/Engine/Private/DestructibleComponent.cpp
--- a/Engine/Source/Runtime/Engine/Private/DestructibleComponent.cpp
+++ b/Engine/Source/Runtime/Engine/Private/DestructibleComponent.cpp
@@ -244,7 +244,7 @@ void UDestructibleComponent::CreatePhysicsState()
 
 	// Passing AssetInstanceID = 0 so we'll have self-collision
 	AActor* Owner = GetOwner();
-	CreateShapeFilterData(MoveChannel, (Owner ? Owner->GetUniqueID() : 0), CollResponse, 0, 0, PQueryFilterData, PSimFilterData, BodyInstance.bUseCCD, true, false);
+	CreateShapeFilterData(MoveChannel, (Owner ? Owner->GetUniqueID() : 0), CollResponse, 0, 0, PQueryFilterData, PSimFilterData, BodyInstance.bUseCCD, TheDestructibleMesh->DefaultDestructibleParameters.DamageParameters.bEnableImpactDamage, false);
 
 	// Build filterData variations for complex and simple
 	PSimFilterData.word3 |= EPDF_SimpleCollision | EPDF_ComplexCollision;
@@ -1152,7 +1152,8 @@ void UDestructibleComponent::SetCollisionResponseForActor(PxRigidDynamic* Actor,
 		AActor* Owner = GetOwner();
 		bool bLargeChunk = IsChunkLarge(ChunkIdx);
 		const FCollisionResponse& ColResponse = bLargeChunk ? LargeChunkCollisionResponse : SmallChunkCollisionResponse;
-		CreateShapeFilterData(MoveChannel, (Owner ? Owner->GetUniqueID() : 0), ColResponse.GetResponseContainer(), 0, ChunkIdxToBoneIdx(ChunkIdx), PQueryFilterData, PSimFilterData, BodyInstance.bUseCCD, true, false);
+		//TODO: we currently assume chunks will not have impact damage as it's very expensive. Should look into exposing this a bit more
+		CreateShapeFilterData(MoveChannel, (Owner ? Owner->GetUniqueID() : 0), ColResponse.GetResponseContainer(), 0, ChunkIdxToBoneIdx(ChunkIdx), PQueryFilterData, PSimFilterData, BodyInstance.bUseCCD, false, false);
 		
 		PQueryFilterData.word3 |= EPDF_SimpleCollision | EPDF_ComplexCollision;
 
