When demo is done playing, don't destroy the demo driver, and pause any existing actors

diff --git a/Engine/Source/Runtime/Engine/Classes/Engine/DemoNetDriver.h b/Engine/Source/Runtime/Engine/Classes/Engine/DemoNetDriver.h
--- a/Engine/Source/Runtime/Engine/Classes/Engine/DemoNetDriver.h
+++ b/Engine/Source/Runtime/Engine/Classes/Engine/DemoNetDriver.h
@@ -33,6 +33,9 @@ class UDemoNetDriver : public UNetDriver
 	/** True if we're in the middle of recording a frame */
 	bool				bIsRecordingDemoFrame;
 
+	/** True if we are at the end of playing a demo */
+	bool				bDemoPlaybackDone;
+
 	/** This is our spectator controller that is used to view the demo world from */
 	APlayerController*	SpectatorController;
 
diff --git a/Engine/Source/Runtime/Engine/Private/DemoNetDriver.cpp b/Engine/Source/Runtime/Engine/Private/DemoNetDriver.cpp
--- a/Engine/Source/Runtime/Engine/Private/DemoNetDriver.cpp
+++ b/Engine/Source/Runtime/Engine/Private/DemoNetDriver.cpp
@@ -38,6 +38,7 @@ bool UDemoNetDriver::InitBase( bool bInitAsClient, FNetworkNotify* InNotify, con
 		Time					= 0;
 		DemoFrameNum			= 0;
 		bIsRecordingDemoFrame	= false;
+		bDemoPlaybackDone		= false;
 
 		return true;
 	}
@@ -324,6 +325,12 @@ void UDemoNetDriver::TickFlush( float DeltaSeconds )
 					SpectatorController->GetSpectatorPawn()->CustomTimeDilation = 1.0f / World->GetWorldSettings()->DemoPlayTimeDilation;
 				}
 			}
+
+			if ( bDemoPlaybackDone )
+			{
+				return;
+			}
+
 			TickDemoPlayback( DeltaSeconds );
 		}
 	}
@@ -373,24 +380,6 @@ void UDemoNetDriver::StopDemo()
 		// flush out any pending network traffic
 		ServerConnection->FlushNet();
 
-#if 0		
-		for (int32 i = ServerConnection->OpenChannels.Num() - 1; i >= 0; i--)
-		{
-			UChannel* OpenChannel = ServerConnection->OpenChannels[i];
-			if (OpenChannel != NULL)
-			{
-				UActorChannel* ActorChannel = Cast< UActorChannel >( OpenChannel );
-				if ( ActorChannel != NULL )
-				{
-					if ( !ActorChannel->GetActor()->IsNetStartupActor() )
-					{
-						OpenChannel->ConditionalCleanUp();
-					}
-				}
-			}
-		}
-#endif
-
 		ServerConnection->State = USOCK_Closed;
 		ServerConnection->Close();
 		ServerConnection->CleanUp(); // make sure DemoRecSpectator gets destroyed immediately
@@ -533,12 +522,37 @@ void UDemoNetDriver::TickDemoRecord( float DeltaSeconds )
 
 bool UDemoNetDriver::ReadDemoFrame()
 {
-	if ( FileAr->AtEnd() || FileAr->IsError() )
+	if ( FileAr->IsError() )
 	{
 		StopDemo();
 		return false;
 	}
 
+	if ( FileAr->AtEnd() )
+	{
+		bDemoPlaybackDone = true;
+
+		// Pause all non player controller actors
+		for ( int32 i = ServerConnection->OpenChannels.Num() - 1; i >= 0; i-- )
+		{
+			UChannel* OpenChannel = ServerConnection->OpenChannels[i];
+			if ( OpenChannel != NULL )
+			{
+				UActorChannel* ActorChannel = Cast< UActorChannel >( OpenChannel );
+				if ( ActorChannel != NULL && ActorChannel->GetActor() != NULL )
+				{
+					if ( Cast< APlayerController >( ActorChannel->GetActor() ) == NULL )
+					{
+						// Better way to pause each actor?
+						ActorChannel->GetActor()->CustomTimeDilation = 0.0f;
+					}
+				}
+			}
+		}
+
+		return false;
+	}
+
 	const int32 OldFilePos = FileAr->Tell();
 
 	float ServerDeltaTime;
