#ue4 - Make old projects using deprecated bUseControllerViewRotation on CameraComponent and SpringArmComponent still able to compile, by adding a deprecrated variable with the old variable's name.

- Efforts have been made to make the value of this deprecated variable reflect the correct value, but it may not behave correctly in all cases and uses of it should be replaced with bUsePawnControlRotation instead, in accordance with the upgrade notes.

diff --git a/Engine/Source/Runtime/Engine/Classes/Camera/CameraComponent.h b/Engine/Source/Runtime/Engine/Classes/Camera/CameraComponent.h
--- a/Engine/Source/Runtime/Engine/Classes/Camera/CameraComponent.h
+++ b/Engine/Source/Runtime/Engine/Classes/Camera/CameraComponent.h
@@ -96,4 +96,22 @@ public:
 	/** DEPRECATED: use bUsePawnControlRotation instead */
 	UPROPERTY()
 	uint32 bUseControllerViewRotation_DEPRECATED:1;
-};
\ No newline at end of file
+
+	/**
+	 * DEPRECATED variable: use "bUsePawnControlRotation" instead. Existing code using this value may not behave correctly.
+	 *
+	 * The correct way to be backwards-compatible for existing saved content using CameraComponents is:
+	 *  - add "_DEPRECATED" to any constructor initialization of bUseControllerViewRotation (which becomes bUseControllerViewRotation_DEPRECATED).
+	 *	- add initialization of the new "bUsePawnControlRotation" to your desired default value.
+	 *	- if there was no previous initialization of bUseControllerViewRotation in code, then you MUST add initialization of bUsePawnControlRotation
+	 *    to "true" if you wish to maintain the same behavior as before (since the default has changed).
+	 *
+	 * This is not a UPROPERTY, with good reason: we don't want to serialize in old values (bUseControllerViewRotation_DEPRECATED handles that).
+	 */
+	DEPRECATED(4.5, "This variable is deprecated; please see the upgrade notes. It only exists to allow compilation of old projects, and code should stop using it in favor of the new bUsePawnControlRotation.")
+	bool bUseControllerViewRotation;
+
+private:
+	// Workaround for compiler-generated function using deprecated variable. Nobody should copy-construct this class anyway.
+	UCameraComponent(const UCameraComponent&);
+};
diff --git a/Engine/Source/Runtime/Engine/Classes/GameFramework/SpringArmComponent.h b/Engine/Source/Runtime/Engine/Classes/GameFramework/SpringArmComponent.h
--- a/Engine/Source/Runtime/Engine/Classes/GameFramework/SpringArmComponent.h
+++ b/Engine/Source/Runtime/Engine/Classes/GameFramework/SpringArmComponent.h
@@ -84,6 +84,7 @@ class ENGINE_API USpringArmComponent : public USceneComponent
 	// UActorComponent interface
 	virtual void OnRegister() override;
 	virtual void TickComponent(float DeltaTime, enum ELevelTick TickType, FActorComponentTickFunction* ThisTickFunction) override;
+	virtual void PostLoad() override;
 	// End of UActorComponent interface
 
 	// USceneComponent interface
@@ -109,4 +110,17 @@ protected:
 	 * by default it returns bHitSomething ? TraceHitLocation : DesiredArmLocation
 	 */
 	virtual FVector BlendLocations(const FVector& DesiredArmLocation, const FVector& TraceHitLocation, bool bHitSomething, float DeltaTime);
+
+public:
+
+	/**
+	* DEPRECATED variable: use "bUsePawnControlRotation" instead. Existing code using this value may not behave correctly.
+	* This is not a UPROPERTY, with good reason: we don't want to serialize in old values.
+	*/
+	DEPRECATED(4.5, "This variable is deprecated; existing code using this value may not behave correctly. It only exists to allow compilation of old projects, and code should stop using it in favor of the new bUsePawnControlRotation.")
+	bool bUseControllerViewRotation;
+
+private:
+	// Workaround for compiler-generated function using deprecated variable. Nobody should copy-construct this class anyway.
+	USpringArmComponent(const USpringArmComponent&);
 };
diff --git a/Engine/Source/Runtime/Engine/Private/Camera/CameraComponent.cpp b/Engine/Source/Runtime/Engine/Private/Camera/CameraComponent.cpp
--- a/Engine/Source/Runtime/Engine/Private/Camera/CameraComponent.cpp
+++ b/Engine/Source/Runtime/Engine/Private/Camera/CameraComponent.cpp
@@ -7,6 +7,9 @@
 
 #define LOCTEXT_NAMESPACE "CameraComponent"
 
+static void SetDeprecatedControllerViewRotation(UCameraComponent& Component, bool bValue);
+
+
 //////////////////////////////////////////////////////////////////////////
 // UCameraComponent
 
@@ -29,6 +32,9 @@ UCameraComponent::UCameraComponent(const class FPostConstructInitializePropertie
 	bUseControllerViewRotation_DEPRECATED = true; // the previous default value before bUsePawnControlRotation replaced this var.
 	bUsePawnControlRotation = false;
 	bAutoActivate = true;
+
+	// Init deprecated var, for old code that may refer to it.
+	SetDeprecatedControllerViewRotation(*this, bUsePawnControlRotation);
 }
 
 void UCameraComponent::OnRegister()
@@ -61,6 +67,9 @@ void UCameraComponent::OnRegister()
 #endif
 
 	Super::OnRegister();
+
+	// Init deprecated var, for old code that may refer to it.
+	SetDeprecatedControllerViewRotation(*this, bUsePawnControlRotation);
 }
 
 void UCameraComponent::OnUnregister()
@@ -98,6 +107,9 @@ void UCameraComponent::PostLoad()
 	{
 		bUsePawnControlRotation = bUseControllerViewRotation_DEPRECATED;
 	}
+
+	// Init deprecated var, for old code that may refer to it.
+	SetDeprecatedControllerViewRotation(*this, bUsePawnControlRotation);
 }
 
 
@@ -187,4 +199,30 @@ void UCameraComponent::CheckForErrors()
 	}
 }
 #endif
+
+
+void SetDeprecatedControllerViewRotation(UCameraComponent& Component, bool bValue)
+{
+	// BEGIN_IGNORE_DEPRECATION_WARNINGS
+	#ifdef __clang__
+		#pragma clang diagnostic push
+		#pragma clang diagnostic ignored "-Wdeprecated-declarations"
+	#else
+		#pragma warning(push)
+		#pragma warning(disable:4995)
+		#pragma warning(disable:4996)
+	#endif
+
+	Component.bUseControllerViewRotation = bValue;
+
+	// END_IGNORE_DEPRECATION_WARNINGS
+	#ifdef __clang__
+		#pragma clang diagnostic pop
+	#else
+		#pragma warning(pop)
+	#endif
+}
+
+
 #undef LOCTEXT_NAMESPACE
+
diff --git a/Engine/Source/Runtime/Engine/Private/GameFramework/SpringArmComponent.cpp b/Engine/Source/Runtime/Engine/Private/GameFramework/SpringArmComponent.cpp
--- a/Engine/Source/Runtime/Engine/Private/GameFramework/SpringArmComponent.cpp
+++ b/Engine/Source/Runtime/Engine/Private/GameFramework/SpringArmComponent.cpp
@@ -5,6 +5,7 @@
 //////////////////////////////////////////////////////////////////////////
 // USpringArmComponent
 
+static void SetDeprecatedControllerViewRotation(USpringArmComponent& Component, bool bValue);
 const FName USpringArmComponent::SocketName(TEXT("SpringEndpoint"));
 
 extern float GAverageMS;
@@ -32,6 +33,9 @@ USpringArmComponent::USpringArmComponent(const class FPostConstructInitializePro
 	CameraRotationLagSpeed = 10.f;
 
 	RelativeSocketRotation = FQuat::Identity;
+
+	// Init deprecated var, for old code that may refer to it.
+	SetDeprecatedControllerViewRotation(*this, bUsePawnControlRotation);
 }
 
 void USpringArmComponent::UpdateDesiredArmLocation(bool bDoTrace, bool bDoLocationLag, bool bDoRotationLag, float DeltaTime)
@@ -117,6 +121,17 @@ void USpringArmComponent::OnRegister()
 {
 	Super::OnRegister();
 	UpdateDesiredArmLocation(false, false, false, 0.f);
+
+	// Init deprecated var, for old code that may refer to it.
+	SetDeprecatedControllerViewRotation(*this, bUsePawnControlRotation);
+}
+
+void USpringArmComponent::PostLoad()
+{
+	Super::PostLoad();
+
+	// Init deprecated var, for old code that may refer to it.
+	SetDeprecatedControllerViewRotation(*this, bUsePawnControlRotation);
 }
 
 void USpringArmComponent::TickComponent(float DeltaTime, enum ELevelTick TickType, FActorComponentTickFunction* ThisTickFunction)
@@ -175,3 +190,26 @@ void USpringArmComponent::QuerySupportedSockets(TArray<FComponentSocketDescripti
 {
 	new (OutSockets) FComponentSocketDescription(SocketName, EComponentSocketType::Socket);
 }
+
+
+void SetDeprecatedControllerViewRotation(USpringArmComponent& Component, bool bValue)
+{
+	// BEGIN_IGNORE_DEPRECATION_WARNINGS
+	#ifdef __clang__
+		#pragma clang diagnostic push
+		#pragma clang diagnostic ignored "-Wdeprecated-declarations"
+	#else
+		#pragma warning(push)
+		#pragma warning(disable:4995)
+		#pragma warning(disable:4996)
+	#endif
+
+	Component.bUseControllerViewRotation = bValue;
+
+	// END_IGNORE_DEPRECATION_WARNINGS
+	#ifdef __clang__
+		#pragma clang diagnostic pop
+	#else
+		#pragma warning(pop)
+	#endif
+}
