Improve destructible sync bodies code so that it updates all needed chunks per actor at once go. Saves 0.5ms when running with 600 active chunks

diff --git a/Engine/Source/Runtime/Engine/Classes/Components/DestructibleComponent.h b/Engine/Source/Runtime/Engine/Classes/Components/DestructibleComponent.h
--- a/Engine/Source/Runtime/Engine/Classes/Components/DestructibleComponent.h
+++ b/Engine/Source/Runtime/Engine/Classes/Components/DestructibleComponent.h
@@ -16,6 +16,7 @@ namespace physx
 	}
 #endif
 	class PxRigidDynamic;
+	class PxRigidActor;
 }
 
 /** Mapping info for destructible chunk user data. */
@@ -169,6 +170,16 @@ public:
 	 */
 	void SetChunkVisible( int32 ChunkIndex, bool bVisible );
 
+#if WITH_APEX
+	/** This method takes a collection of active actors and updates the chunks in one pass. Saves a lot of duplicate work instead of calling each individual chunk
+	 * 
+	 *  @param ActiveActors - The array of actors that need their transforms updated
+	 *
+     */
+	static void UpdateDestructibleChunkTM(TArray<const physx::PxRigidActor*> ActiveActors);
+#endif
+
+
 	/** This method sets a chunk's (fractured piece's) world rotation and translation.
 	 *
 	 * @param ChunkIndex - Which chunk to affect.  ChunkIndex must lie in the range: 0 <= ChunkIndex < ((DestructibleMesh*)USkeletalMesh)->ApexDestructibleAsset->chunkCount().
@@ -195,6 +206,18 @@ public:
 	FORCEINLINE static int32 ChunkIdxToBoneIdx(int32 ChunkIdx) { return ChunkIdx + 1; }
 	FORCEINLINE static int32 BoneIdxToChunkIdx(int32 BoneIdx) { return FMath::Max(BoneIdx - 1, 0); }
 private:
+
+	struct FUpdateChunksInfo
+	{
+		int32 ChunkIndex;
+		FTransform WorldTM;
+
+		FUpdateChunksInfo(int32 InChunkIndex, const FTransform& InWorldTM) : ChunkIndex(InChunkIndex), WorldTM(InWorldTM){}
+
+	};
+
+	void SetChunksWorldTM(const TArray<FUpdateChunksInfo>& UpdateInfos);
+
 	/** Collision response used for chunks */
 	FCollisionResponse LargeChunkCollisionResponse;
 	FCollisionResponse SmallChunkCollisionResponse;
diff --git a/Engine/Source/Runtime/Engine/Private/DestructibleComponent.cpp b/Engine/Source/Runtime/Engine/Private/DestructibleComponent.cpp
--- a/Engine/Source/Runtime/Engine/Private/DestructibleComponent.cpp
+++ b/Engine/Source/Runtime/Engine/Private/DestructibleComponent.cpp
@@ -848,6 +848,84 @@ void UDestructibleComponent::SetChunkVisible( int32 ChunkIndex, bool bVisible )
 	MarkRenderDynamicDataDirty();
 }
 
+#if WITH_APEX
+void UDestructibleComponent::UpdateDestructibleChunkTM(TArray<const PxRigidActor*> ActiveActors)
+{
+	//We want to consolidate the transforms so that we update each destructible component once by passing it an array of chunks to update.
+	//This helps avoid a lot of duplicated work like marking render dirty, computing inverse world component, etc...
+
+	TMap<UDestructibleComponent*, TArray<FUpdateChunksInfo> > ComponentUpdateMapping;
+	
+	//prepare map to update destructible components
+	TArray<PxShape*> Shapes;
+	for (const PxRigidActor* RigidActor : ActiveActors)
+	{
+		if (const FDestructibleChunkInfo* DestructibleChunkInfo = FPhysxUserData::Get<FDestructibleChunkInfo>(RigidActor->userData))
+		{
+			if (GApexModuleDestructible->owns(RigidActor) && DestructibleChunkInfo->OwningComponent.IsValid())
+			{
+				Shapes.AddUninitialized(RigidActor->getNbShapes());
+				int32 NumShapes = RigidActor->getShapes(Shapes.GetData(), Shapes.Num());
+				for (int32 ShapeIdx = 0; ShapeIdx < Shapes.Num(); ++ShapeIdx)
+				{
+					PxShape* Shape = Shapes[ShapeIdx];
+					int32 ChunkIndex;
+					if (NxDestructibleActor* DestructibleActor = GApexModuleDestructible->getDestructibleAndChunk(Shape, &ChunkIndex))
+					{
+						const physx::PxMat44 ChunkPoseRT = DestructibleActor->getChunkPose(ChunkIndex);
+						const physx::PxTransform Transform(ChunkPoseRT);
+						if (UDestructibleComponent* DestructibleComponent = Cast<UDestructibleComponent>(FPhysxUserData::Get<UPrimitiveComponent>(DestructibleActor->userData)))
+						{
+							if (DestructibleComponent->IsRegistered())
+							{
+								TArray<FUpdateChunksInfo>& UpdateInfos = ComponentUpdateMapping.FindOrAdd(DestructibleComponent);
+								FUpdateChunksInfo* UpdateInfo = new (UpdateInfos)FUpdateChunksInfo(ChunkIndex, P2UTransform(Transform));
+							}
+						}
+					}
+				}
+
+				Shapes.Empty(Shapes.Num());	//we want to keep largest capacity array to avoid reallocs
+			}
+		}
+	}
+	
+	//update each component
+	for (auto It = ComponentUpdateMapping.CreateIterator(); It; ++It)
+	{
+		UDestructibleComponent* DestructibleComponent = It.Key();
+		TArray<FUpdateChunksInfo>& UpdateInfos = It.Value();
+		DestructibleComponent->SetChunksWorldTM(UpdateInfos);
+		/*for (FUpdateChunksInfo& UpdateInfo : UpdateInfos)
+		{
+			DestructibleComponent->SetChunkWorldRT(UpdateInfo.ChunkIndex, UpdateInfo.WorldTM.GetRotation(), UpdateInfo.WorldTM.GetLocation());
+		}*/
+		
+	}
+
+}
+
+#endif
+
+void UDestructibleComponent::SetChunksWorldTM(const TArray<FUpdateChunksInfo>& UpdateInfos)
+{
+	FTransform InvWorldTM = ComponentToWorld.Inverse();
+	InvWorldTM.SetScale3D(FVector(1.f));
+
+	for (const FUpdateChunksInfo& UpdateInfo : UpdateInfos)
+	{
+		// Bone 0 is a dummy root bone
+		const int32 BoneIndex = ChunkIdxToBoneIdx(UpdateInfo.ChunkIndex);
+		SpaceBases[BoneIndex] = FTransform(UpdateInfo.WorldTM * InvWorldTM);
+	}
+
+	// Mark the transform as dirty, so the bounds are updated and sent to the render thread
+	MarkRenderTransformDirty();
+
+	// New bone positions need to be sent to render thread
+	MarkRenderDynamicDataDirty();
+}
+
 void UDestructibleComponent::SetChunkWorldRT( int32 ChunkIndex, const FQuat& WorldRotation, const FVector& WorldTranslation )
 {
 	// Bone 0 is a dummy root bone
diff --git a/Engine/Source/Runtime/Engine/Private/PhysicsEngine/PhysScene.cpp b/Engine/Source/Runtime/Engine/Private/PhysicsEngine/PhysScene.cpp
--- a/Engine/Source/Runtime/Engine/Private/PhysicsEngine/PhysScene.cpp
+++ b/Engine/Source/Runtime/Engine/Private/PhysicsEngine/PhysScene.cpp
@@ -733,12 +733,23 @@ void FPhysScene::UpdateActiveTransforms(uint32 SceneType)
 }
 #endif
 
+DEFINE_STAT(STAT_SyncComponentsToBodies);
+
 void FPhysScene::SyncComponentsToBodies(uint32 SceneType)
 {
+	SCOPE_CYCLE_COUNTER(STAT_TotalPhysicsTime);
+	SCOPE_CYCLE_COUNTER(STAT_SyncComponentsToBodies);
 #if WITH_PHYSX
+
+	SCOPED_SCENE_READ_LOCK(GetPhysXScene(SceneType));
+
 	PxU32 NumTransforms = ActiveTransforms[SceneType].Num();
 	const TArray<const PxActiveTransform*> & PActiveTransforms = ActiveTransforms[SceneType];
 
+#if WITH_APEX
+	TArray<const PxRigidActor*> ActiveDestructionActors;
+#endif
+
 	for (PxU32 TransformIdx = 0; TransformIdx<NumTransforms; TransformIdx++)
 	{
 		if (PActiveTransforms[TransformIdx] == NULL)	//it's possible to call TermBody on FBodyInstance after fetchResults, but before SyncComponentsToBodies - in this case a NULL is used to represent the stale data
@@ -755,7 +766,7 @@ void FPhysScene::SyncComponentsToBodies(uint32 SceneType)
 		//Special code for destructible chunk
 		if (const FDestructibleChunkInfo* DestructibleChunkInfo = FPhysxUserData::Get<FDestructibleChunkInfo>(RigidActor->userData))
 		{
-			if (GApexModuleDestructible->owns(RigidActor) && DestructibleChunkInfo->OwningComponent.IsValid())
+			/*if (GApexModuleDestructible->owns(RigidActor) && DestructibleChunkInfo->OwningComponent.IsValid())
 			{
 				//TODO: waiting on new API to avoid duplicate updates per shape.
 				TArray<PxShape*> Shapes;
@@ -780,7 +791,9 @@ void FPhysScene::SyncComponentsToBodies(uint32 SceneType)
 						}
 					}
 				}
-			}
+			}*/
+
+			ActiveDestructionActors.Add(RigidActor);
 		}
 #endif
 
@@ -814,6 +827,13 @@ void FPhysScene::SyncComponentsToBodies(uint32 SceneType)
 			}
 		}
 	}
+
+#if WITH_APEX
+	if (ActiveDestructionActors.Num())
+	{
+		UDestructibleComponent::UpdateDestructibleChunkTM(ActiveDestructionActors);
+	}
+#endif
 #endif
 }
 
diff --git a/Engine/Source/Runtime/Engine/Public/PhysicsPublic.h b/Engine/Source/Runtime/Engine/Public/PhysicsPublic.h
--- a/Engine/Source/Runtime/Engine/Public/PhysicsPublic.h
+++ b/Engine/Source/Runtime/Engine/Public/PhysicsPublic.h
@@ -19,6 +19,7 @@ DECLARE_CYCLE_STAT_EXTERN(TEXT("Phys Events Time"),STAT_PhysicsEventTime,STATGRO
 DECLARE_CYCLE_STAT_EXTERN(TEXT("Phys SetBodyTransform"),STAT_SetBodyTransform,STATGROUP_Physics, );
 DECLARE_CYCLE_STAT_EXTERN(TEXT("Phys SubstepStart"), STAT_SubstepSimulationStart,STATGROUP_Physics, );
 DECLARE_CYCLE_STAT_EXTERN(TEXT("Phys SubstepEnd"), STAT_SubstepSimulationEnd, STATGROUP_Physics, );
+DECLARE_CYCLE_STAT_EXTERN(TEXT("SyncComponentsToBodies"), STAT_SyncComponentsToBodies, STATGROUP_Physics, );
 
 DECLARE_DWORD_COUNTER_STAT_EXTERN(TEXT("Broadphase Adds"), STAT_NumBroadphaseAdds, STATGROUP_Physics, );
 DECLARE_DWORD_COUNTER_STAT_EXTERN(TEXT("Broadphase Removes"), STAT_NumBroadphaseRemoves, STATGROUP_Physics, );
