[347201] Added NearlyEquals to the Blueprint math library for vectors & rotators.

- Consistent with floating-point equality tests; note that Equals remains a test for true equality (no epsilon).
- Made tooltips slightly more verbose to highlight the differences between "exact" and "near" equality tests.

diff --git a/Engine/Source/Runtime/Engine/Classes/Kismet/KismetMathLibrary.h b/Engine/Source/Runtime/Engine/Classes/Kismet/KismetMathLibrary.h
--- a/Engine/Source/Runtime/Engine/Classes/Kismet/KismetMathLibrary.h
+++ b/Engine/Source/Runtime/Engine/Classes/Kismet/KismetMathLibrary.h
@@ -284,7 +284,7 @@ class UKismetMathLibrary : public UBlueprintFunctionLibrary
 	UFUNCTION(BlueprintPure, meta=(FriendlyName = "Equal (float)", CompactNodeTitle = "==", Keywords = "== equal"), Category="Math|Float")
 	static bool EqualEqual_FloatFloat(float A, float B);
 
-	/* Returns true if (|A - B| < ErrorTolerance) */
+	/* Returns true if A is nearly equal to B (|A - B| < ErrorTolerance) */
 	UFUNCTION(BlueprintPure, meta=(FriendlyName = "Nearly Equal (float)", Keywords = "== equal"), Category="Math|Float")
 	static bool NearlyEqual_FloatFloat(float A, float B, float ErrorTolerance = 1.e-6f);
 
@@ -521,11 +521,15 @@ class UKismetMathLibrary : public UBlueprintFunctionLibrary
 	UFUNCTION(BlueprintPure, meta=(FriendlyName = "RotateVectorAroundAxis"), Category="Math|Vector")
 	static FVector RotateAngleAxis(FVector InVect, float AngleDeg, FVector Axis);
 
-	/* Returns true if the values are equal (A == B) */
+	/* Returns true if two vectors are exactly equal (A == B) */
 	UFUNCTION(BlueprintPure, meta=(FriendlyName = "Equal (vector)", CompactNodeTitle = "==", Keywords = "== equal"), Category="Math|Vector")
 	static bool EqualEqual_VectorVector(FVector A, FVector B);
 
-	/* Returns true if the values are not equal (A != B) */
+	/* Returns true if two vectors are nearly equal (within a specified error tolerance) */
+	UFUNCTION(BlueprintPure, meta=(FriendlyName = "NearlyEqual (vector)", Keywords = "== equal"), Category="Math|Vector")
+	static bool NearlyEqual_VectorVector(FVector A, FVector B, float ErrorTolerance = 1.e-6f);
+
+	/* Returns true if two vectors are not equal (A != B) */
 	UFUNCTION(BlueprintPure, meta=(FriendlyName = "NotEqual (vector)", CompactNodeTitle = "!=", Keywords = "!= not equal"), Category="Math|Vector")
 	static bool NotEqual_VectorVector(FVector A, FVector B);
 
@@ -613,10 +617,14 @@ class UKismetMathLibrary : public UBlueprintFunctionLibrary
 	// Rotator functions.
 	//
 
-	//Returns true if rotator A is equal to rotator B (A == B)
+	//Returns true if rotator A is exactly equal to rotator B (A == B)
 	UFUNCTION(BlueprintPure, meta=(FriendlyName = "Equal (Rotator)", CompactNodeTitle = "==", Keywords = "== equal"), Category="Math|Rotator")
 	static bool EqualEqual_RotatorRotator(FRotator A, FRotator B);
 
+	//Returns true if rotator A is nearly equal to rotator B (within a specified error tolerance)
+	UFUNCTION(BlueprintPure, meta=(FriendlyName = "NearlyEqual (Rotator)", Keywords = "== equal"), Category="Math|Rotator")
+	static bool NearlyEqual_RotatorRotator(FRotator A, FRotator B, float ErrorTolerance = 1.e-6f);
+
 	//Returns true if rotator A does not equal rotator B (A != B)
 	UFUNCTION(BlueprintPure, meta=(FriendlyName = "NotEqual (Rotator)", CompactNodeTitle = "!=", Keywords = "!= not equal"), Category="Math|Rotator")
 	static bool NotEqual_RotatorRotator(FRotator A, FRotator B);
diff --git a/Engine/Source/Runtime/Engine/Private/KismetMathLibrary.cpp b/Engine/Source/Runtime/Engine/Private/KismetMathLibrary.cpp
--- a/Engine/Source/Runtime/Engine/Private/KismetMathLibrary.cpp
+++ b/Engine/Source/Runtime/Engine/Private/KismetMathLibrary.cpp
@@ -688,6 +688,11 @@ bool UKismetMathLibrary::EqualEqual_VectorVector(FVector A, FVector B)
 	return (A.X == B.X) && (A.Y == B.Y) && (A.Z == B.Z);
 }	
 
+bool UKismetMathLibrary::NearlyEqual_VectorVector(FVector A, FVector B, float ErrorTolerance)
+{
+	return A.Equals(B, ErrorTolerance);
+}
+
 bool UKismetMathLibrary::NotEqual_VectorVector(FVector A, FVector B)
 {
 	return (A.X != B.X) || (A.Y != B.Y) || (A.Z != B.Z);
@@ -839,6 +844,11 @@ bool UKismetMathLibrary::EqualEqual_RotatorRotator(FRotator A, FRotator B)
 	return (A.Pitch == B.Pitch) && (A.Yaw == B.Yaw) && (A.Roll == B.Roll);
 }	
 
+bool UKismetMathLibrary::NearlyEqual_RotatorRotator(FRotator A, FRotator B, float ErrorTolerance)
+{
+	return A.Equals(B, ErrorTolerance);
+}
+
 bool UKismetMathLibrary::NotEqual_RotatorRotator(FRotator A, FRotator B)
 {
 	return (A.Pitch != B.Pitch) || (A.Yaw != B.Yaw) || (A.Roll != B.Roll);
