Fix crash when demo spectator is not found

diff --git a/Engine/Source/Runtime/Engine/Private/DemoNetDriver.cpp b/Engine/Source/Runtime/Engine/Private/DemoNetDriver.cpp
--- a/Engine/Source/Runtime/Engine/Private/DemoNetDriver.cpp
+++ b/Engine/Source/Runtime/Engine/Private/DemoNetDriver.cpp
@@ -687,8 +687,21 @@ void UDemoNetDriver::SpawnDemoRecSpectator( UNetConnection* Connection )
 	check( Connection != NULL );
 
 	UClass* C = StaticLoadClass( AActor::StaticClass(), NULL, *DemoSpectatorClass, NULL, LOAD_None, NULL );
+
+	if ( C == NULL )
+	{
+		UE_LOG( LogDemo, Error, TEXT( "UDemoNetDriver::SpawnDemoRecSpectator: Failed to load demo spectator class." ) );
+		return;
+	}
+
 	APlayerController* Controller = World->SpawnActor<APlayerController>( C );
 
+	if ( Controller == NULL )
+	{
+		UE_LOG( LogDemo, Error, TEXT( "UDemoNetDriver::SpawnDemoRecSpectator: Failed to spawn demo spectator." ) );
+		return;
+	}
+
 	for ( FActorIterator It( World ); It; ++It)
 	{
 		if ( It->IsA( APlayerStart::StaticClass() ) )
