Fixed ensure when using the foliage mode select tool on an ISMC created in a blueprint.

diff --git a/Engine/Source/Runtime/Engine/Private/InstancedFoliage.cpp b/Engine/Source/Runtime/Engine/Private/InstancedFoliage.cpp
--- a/Engine/Source/Runtime/Engine/Private/InstancedFoliage.cpp
+++ b/Engine/Source/Runtime/Engine/Private/InstancedFoliage.cpp
@@ -1258,7 +1258,7 @@ void AInstancedFoliageActor::SelectInstance(UInstancedStaticMeshComponent* InCom
 		if (MeshInfo)
 		{
 			FFoliageInstanceCluster* Cluster = MeshInfo->InstanceClusters.FindByPredicate([InComponent](FFoliageInstanceCluster& Cluster){ return Cluster.ClusterComponent == InComponent; });
-			if (ensure(Cluster))
+			if (Cluster)
 			{
 				InComponent->Modify();
 				int32 InstanceIndex = Cluster->InstanceIndices[InComponentInstanceIndex];
