Compilation fixes

diff --git a/Engine/Source/Runtime/Engine/Classes/Animation/AnimationAsset.h b/Engine/Source/Runtime/Engine/Classes/Animation/AnimationAsset.h
--- a/Engine/Source/Runtime/Engine/Classes/Animation/AnimationAsset.h
+++ b/Engine/Source/Runtime/Engine/Classes/Animation/AnimationAsset.h
@@ -413,9 +413,12 @@ public:
 	 * @param ReplacementMap	Mapping of original asset to new asset
 	 **/
 	ENGINE_API virtual void ReplaceReferredAnimations(const TMap<UAnimSequence*, UAnimSequence*>& ReplacementMap);	
-#endif
 
-#if WITH_EDITOR
+	ENGINE_API void SetPreviewMesh(USkeletalMesh* PreviewMesh);
+	ENGINE_API USkeletalMesh* GetPreviewMesh();
+#endif //WITH_EDITOR
+
+#if WITH_EDITORONLY_DATA
 	/** Information for thumbnail rendering */
 	UPROPERTY(VisibleAnywhere, Instanced, Category = Thumbnail)
 	class UThumbnailInfo* ThumbnailInfo;
@@ -424,12 +427,7 @@ private:
 	/** The default skeletal mesh to use when previewing this asset - this only applies when you open Persona using this asset*/
 	UPROPERTY(duplicatetransient, AssetRegistrySearchable)
 	TAssetPtr<class USkeletalMesh> PreviewSkeletalMesh;
-
-public:
-	ENGINE_API void SetPreviewMesh(USkeletalMesh* PreviewMesh);
-	ENGINE_API USkeletalMesh* UAnimationAsset::GetPreviewMesh();
-
-#endif
+#endif //WITH_EDITORONLY_DATA
 
 public:
 	class USkeleton* GetSkeleton() const { return Skeleton; }
