High res screenshots now output as png rather than bmp.  Slight refactor.
#ttp 348615 - EDITOR: MAC: High resolution screenshot outputs to BMP which mac cant open
#branch UE4
#proj Editor.UnrealEd, Runtime.Engine, Runtime.Core

diff --git a/Engine/Source/Editor/UnrealEd/Private/EditorViewportClient.cpp b/Engine/Source/Editor/UnrealEd/Private/EditorViewportClient.cpp
--- a/Engine/Source/Editor/UnrealEd/Private/EditorViewportClient.cpp
+++ b/Engine/Source/Editor/UnrealEd/Private/EditorViewportClient.cpp
@@ -3778,12 +3778,22 @@ void FEditorViewportClient::TakeScreenshot(FViewport* InViewport, bool bInValida
 	{
 		check(Bitmap.Num() == InViewport->GetSizeXY().X * InViewport->GetSizeXY().Y);
 
+		// Initialize alpha channel of bitmap
+		for (auto& Pixel : Bitmap)
+		{
+			Pixel.A = 255;
+		}
+
 		// Create screenshot folder if not already present.
 		if ( IFileManager::Get().MakeDirectory( *FPaths::ScreenShotDir(), true ) )
 		{
 			// Save the contents of the array to a bitmap file.
-			FString ScreenshotSaveName = "";
-			if ( FFileHelper::CreateBitmap(*(FPaths::ScreenShotDir() / TEXT("ScreenShot")),InViewport->GetSizeXY().X,InViewport->GetSizeXY().Y,Bitmap.GetData(),NULL,&IFileManager::Get(),&ScreenshotSaveName) )
+			FHighResScreenshotConfig& HighResScreenshotConfig = GetHighResScreenshotConfig();
+			HighResScreenshotConfig.SetHDRCapture(false);
+
+			FString ScreenshotSaveName;
+			if (FFileHelper::GenerateNextBitmapFilename(FPaths::ScreenShotDir() / TEXT("ScreenShot"), TEXT("png"), ScreenshotSaveName) &&
+				HighResScreenshotConfig.SaveImage(ScreenshotSaveName, Bitmap, InViewport->GetSizeXY()))
 			{
 				// Setup the string with the path and name of the file
 				ScreenshotSaveResultText = NSLOCTEXT( "UnrealEd", "ScreenshotSavedAs", "Screenshot capture saved as" );					
@@ -3871,7 +3881,8 @@ void FEditorViewportClient::ProcessScreenShots(FViewport* InViewport)
 		// Default capture region is the entire viewport
 		FIntRect CaptureRect(0, 0, 0, 0);
 
-		bool bCaptureAreaValid = GetHighResScreenshotConfig().CaptureRegion.Area() > 0;
+		FHighResScreenshotConfig& HighResScreenshotConfig = GetHighResScreenshotConfig();
+		bool bCaptureAreaValid = HighResScreenshotConfig.CaptureRegion.Area() > 0;
 
 		// If capture region isn't valid, we need to determine which rectangle to capture from.
 		// We need to calculate a proper view rectangle so that we can take into account camera
@@ -3903,20 +3914,29 @@ void FEditorViewportClient::ProcessScreenShots(FViewport* InViewport)
 			if (GIsHighResScreenshot && bCaptureAreaValid)
 			{
 				// Highres screenshot capture region is valid, so use that
-				SourceRect = GetHighResScreenshotConfig().CaptureRegion;
+				SourceRect = HighResScreenshotConfig.CaptureRegion;
 			}
 
 			bool bWriteAlpha = false;
 
 			// If this is a high resolution screenshot and we are using the masking feature,
 			// Get the results of the mask rendering pass and insert into the alpha channel of the screenshot.
-			if (GIsHighResScreenshot && GetHighResScreenshotConfig().bMaskEnabled)
+			if (GIsHighResScreenshot && HighResScreenshotConfig.bMaskEnabled)
 			{
-				bWriteAlpha = GetHighResScreenshotConfig().MergeMaskIntoAlpha(Bitmap);
+				bWriteAlpha = HighResScreenshotConfig.MergeMaskIntoAlpha(Bitmap);
+			}
+
+			// Set full alpha on the bitmap
+			if (!bWriteAlpha)
+			{
+				for (auto& Pixel : Bitmap)
+				{
+					Pixel.A = 255;
+				}
 			}
 
 			// Save the bitmap to disc
-			FFileHelper::CreateBitmap(*ScreenShotName, BitmapSize.X, BitmapSize.Y, Bitmap.GetData(), &SourceRect, &IFileManager::Get(), NULL, bWriteAlpha);
+			HighResScreenshotConfig.SaveImage(ScreenShotName, Bitmap, BitmapSize);
 		}
 		
 		// Done with the request
diff --git a/Engine/Source/Editor/UnrealEd/Private/UnrealEdSrv.cpp b/Engine/Source/Editor/UnrealEd/Private/UnrealEdSrv.cpp
--- a/Engine/Source/Editor/UnrealEd/Private/UnrealEdSrv.cpp
+++ b/Engine/Source/Editor/UnrealEd/Private/UnrealEdSrv.cpp
@@ -1262,7 +1262,7 @@ bool UUnrealEdEngine::Exec( UWorld* InWorld, const TCHAR* Stream, FOutputDevice&
 
 				FString FileName;
 				const FString BaseFileName = FPaths::ScreenShotDir() / TEXT("EditorScreenshot");
-				FFileHelper::GenerateNextBitmapFilename(*BaseFileName, FileName);
+				FFileHelper::GenerateNextBitmapFilename(BaseFileName, TEXT("bmp"), FileName);
 				FFileHelper::CreateBitmap(*FileName, OutImageSize.X, OutImageSize.Y, OutImageData.GetData());
 			}
 		};
diff --git a/Engine/Source/Runtime/Core/Private/Misc/CoreMisc.cpp b/Engine/Source/Runtime/Core/Private/Misc/CoreMisc.cpp
--- a/Engine/Source/Runtime/Core/Private/Misc/CoreMisc.cpp
+++ b/Engine/Source/Runtime/Core/Private/Misc/CoreMisc.cpp
@@ -261,24 +261,25 @@ bool FFileHelper::SaveStringToFile( const FString& String, const TCHAR* Filename
 }
 
 /**
- * Generates the next unique bitmap filename
+ * Generates the next unique bitmap filename with a specified extension
  * 
- * @param Pattern filename with path, must not be 0, if with "bmp" extension (e.g. "out.bmp") the filename stays like this, if without (e.g. "out") automatic index numbers are appended (e.g. "out00002.bmp")
- * @param OutFilename reference to an FString where the newly generated filename will be placed
- * @param FileManager must not be 0
+ * @param Pattern		Filename with path, but without extension.
+ * @oaran Extension		File extension to be appended
+ * @param OutFilename	Reference to an FString where the newly generated filename will be placed
+ * @param FileManager	Reference to a IFileManager (or the global instance by default)
  *
  * @return true if success
  */
-bool FFileHelper::GenerateNextBitmapFilename( const FString& Pattern, FString& OutFilename, IFileManager* FileManager /*= &IFileManager::Get()*/ )
+bool FFileHelper::GenerateNextBitmapFilename( const FString& Pattern, const FString& Extension, FString& OutFilename, IFileManager* FileManager /*= &IFileManager::Get()*/ )
 {
-	TCHAR File[MAX_SPRINTF]=TEXT("");
+	FString File;
 	OutFilename = "";
 	bool bSuccess = false;
 
-	for( int32 TestBitmapIndex = GScreenshotBitmapIndex + 1; TestBitmapIndex < 65536; ++TestBitmapIndex )
+	for( int32 TestBitmapIndex = GScreenshotBitmapIndex + 1; TestBitmapIndex < 100000; ++TestBitmapIndex )
 	{
-		FCString::Sprintf( File, TEXT("%s%05i.bmp"), *Pattern, TestBitmapIndex );
-		if( FileManager->FileSize(File) < 0 )
+		File = FString::Printf(TEXT("%s%05i.%s"), *Pattern, TestBitmapIndex, *Extension);
+		if( FileManager->FileSize(*File) < 0 )
 		{
 			GScreenshotBitmapIndex = TestBitmapIndex;
 			OutFilename = File;
@@ -312,21 +313,19 @@ bool FFileHelper::CreateBitmap( const TCHAR* Pattern, int32 SourceWidth, int32 S
 		SubRectangle = &Src;
 	}
 
-	TCHAR File[MAX_SPRINTF]=TEXT("");
+	FString File;
 	// if the Pattern already has a .bmp extension, then use that the file to write to
 	if (FPaths::GetExtension(Pattern) == TEXT("bmp"))
 	{
-		FCString::Strcpy(File, Pattern);
+		File = Pattern;
 	}
 	else
 	{
-		FString Filename;
-		if (GenerateNextBitmapFilename(Pattern, Filename, FileManager))
+		if (GenerateNextBitmapFilename(Pattern, TEXT("bmp"), File, FileManager))
 		{
-			FCString::Strcpy(File, *Filename);
 			if ( OutFilename )
 			{
-				*OutFilename = Filename;
+				*OutFilename = File;
 			}
 		}
 		else
@@ -335,7 +334,7 @@ bool FFileHelper::CreateBitmap( const TCHAR* Pattern, int32 SourceWidth, int32 S
 		}
 	}
 
-	FArchive* Ar = FileManager->CreateDebugFileWriter( File );
+	FArchive* Ar = FileManager->CreateDebugFileWriter( *File );
 	if( Ar )
 	{
 		// Types.
diff --git a/Engine/Source/Runtime/Core/Public/Misc/CoreMisc.h b/Engine/Source/Runtime/Core/Public/Misc/CoreMisc.h
--- a/Engine/Source/Runtime/Core/Public/Misc/CoreMisc.h
+++ b/Engine/Source/Runtime/Core/Public/Misc/CoreMisc.h
@@ -203,15 +203,16 @@ struct CORE_API FFileHelper
 	static bool CreateBitmap( const TCHAR* Pattern, int32 DataWidth, int32 DataHeight, const class FColor* Data, struct FIntRect* SubRectangle = NULL, IFileManager* FileManager = &IFileManager::Get(), FString* OutFilename = NULL, bool bInWriteAlpha = false );
 
 	/**
-	 * Generates the next unique bitmap filename
-	 * 
-	 * @param Pattern filename with path, must not be 0, if with "bmp" extension (e.g. "out.bmp") the filename stays like this, if without (e.g. "out") automatic index numbers are addended (e.g. "out00002.bmp")
-	 * @param OutFilename reference to an fstring which will hold the generated filename
-	 * @param FileManager must not be 0
+	 * Generates the next unique bitmap filename with a specified extension
+	 *
+	 * @param Pattern		Filename with path, but without extension.
+	 * @oaran Extension		File extension to be appended
+	 * @param OutFilename	Reference to an FString where the newly generated filename will be placed
+	 * @param FileManager	Reference to a IFileManager (or the global instance by default)
 	 *
 	 * @return true if success
 	 */
-	static bool GenerateNextBitmapFilename( const FString& Pattern, FString& OutFilename, IFileManager* FileManager = &IFileManager::Get() );
+	static bool GenerateNextBitmapFilename(const FString& Pattern, const FString& Extension, FString& OutFilename, IFileManager* FileManager = &IFileManager::Get());
 	
 	/**
 	 *	Load the given ANSI text file to an array of strings - one FString per line of the file.
diff --git a/Engine/Source/Runtime/Engine/Private/GameViewportClient.cpp b/Engine/Source/Runtime/Engine/Private/GameViewportClient.cpp
--- a/Engine/Source/Runtime/Engine/Private/GameViewportClient.cpp
+++ b/Engine/Source/Runtime/Engine/Private/GameViewportClient.cpp
@@ -2503,7 +2503,7 @@ bool UGameViewportClient::HandleScreenshotCommand( const TCHAR* Cmd, FOutputDevi
 	if(Viewport)
 	{
 		const bool bShowUI = FParse::Command(&Cmd, TEXT("SHOWUI"));
-		FScreenshotRequest::RequestScreenshot( bShowUI );
+		FScreenshotRequest::RequestScreenshot( bShowUI, TEXT("png") );
 
 		GScreenMessagesRestoreState = GAreScreenMessagesEnabled;
 		GAreScreenMessagesEnabled = false;
diff --git a/Engine/Source/Runtime/Engine/Private/HighResScreenshot.cpp b/Engine/Source/Runtime/Engine/Private/HighResScreenshot.cpp
--- a/Engine/Source/Runtime/Engine/Private/HighResScreenshot.cpp
+++ b/Engine/Source/Runtime/Engine/Private/HighResScreenshot.cpp
@@ -95,26 +95,54 @@ void FHighResScreenshotConfig::SetHDRCapture(bool bCaptureHDRIN)
 	bCaptureHDR = bCaptureHDRIN;
 }
 
+template<typename> struct FPixelTypeTraits {};
+
+template<> struct FPixelTypeTraits<FColor>
+{
+	static const ERGBFormat::Type SourceChannelLayout = ERGBFormat::BGRA;
+
+	static FORCEINLINE bool IsWritingHDRImage(const bool)
+	{
+		return false;
+	}
+};
+
+template<> struct FPixelTypeTraits<FFloat16Color>
+{
+	static const ERGBFormat::Type SourceChannelLayout = ERGBFormat::RGBA;
+
+	static FORCEINLINE bool IsWritingHDRImage(const bool bCaptureHDR)
+	{
+		static const auto CVarDumpFramesAsHDR = IConsoleManager::Get().FindTConsoleVariableDataInt(TEXT("r.BufferVisualizationDumpFramesAsHDR"));
+		return bCaptureHDR || CVarDumpFramesAsHDR->GetValueOnRenderThread();
+	}
+};
+
 template<typename TPixelType>
-FString FHighResScreenshotConfig::SaveImage(const FString& File, const TArray<TPixelType>& Bitmap, const FIntPoint& BitmapSize, EPixelFormat SourceFormat) const
+bool FHighResScreenshotConfig::SaveImage(const FString& File, const TArray<TPixelType>& Bitmap, const FIntPoint& BitmapSize, FString* OutFilename) const
 {
+	typedef FPixelTypeTraits<TPixelType> Traits;
+
 	static_assert(ARE_TYPES_EQUAL(TPixelType, FFloat16Color) || ARE_TYPES_EQUAL(TPixelType, FColor), "Source format must be either FColor or FFloat16Color");
-	check(SourceFormat == PF_FloatRGBA || SourceFormat == PF_B8G8R8A8 || SourceFormat == PF_R8G8B8A8);
 	const int32 x = BitmapSize.X;
 	const int32 y = BitmapSize.Y;
 	check(Bitmap.Num() == x * y);
 
-	static const auto CVarDumpFramesAsHDR = IConsoleManager::Get().FindTConsoleVariableDataInt(TEXT("r.BufferVisualizationDumpFramesAsHDR"));
-	bool bIsWritingHDRImage = (bCaptureHDR || CVarDumpFramesAsHDR->GetValueOnRenderThread()) && SourceFormat == PF_FloatRGBA;
+	const bool bIsWritingHDRImage = Traits::IsWritingHDRImage(bCaptureHDR);
 
 	IFileManager* FileManager = &IFileManager::Get();
 	const size_t BitsPerPixel = (sizeof(TPixelType) / 4) * 8;
-	const ERGBFormat::Type SourceChannelLayout = SourceFormat == PF_B8G8R8A8 ? ERGBFormat::BGRA : ERGBFormat::RGBA;
 
 	TSharedPtr<class IImageWrapper> ImageCompressor = bIsWritingHDRImage ? ImageCompressorHDR : ImageCompressorLDR;
-	FString Filename = File + (bIsWritingHDRImage ? TEXT(".exr") : TEXT(".png"));
+	FString Extension = bIsWritingHDRImage ? TEXT(".exr") : TEXT(".png");
+	FString Filename = (FPaths::GetExtension(File, true) == Extension) ? File : FPaths::GetBaseFilename(File, false) + Extension;
 
-	if (ImageCompressor.IsValid() && ImageCompressor->SetRaw((void*)&Bitmap[0], sizeof(TPixelType)* x * y, x, y, SourceChannelLayout, BitsPerPixel))
+	if (OutFilename != nullptr)
+	{
+		*OutFilename = Filename;
+	}
+
+	if (ImageCompressor.IsValid() && ImageCompressor->SetRaw((void*)&Bitmap[0], sizeof(TPixelType)* x * y, x, y, Traits::SourceChannelLayout, BitsPerPixel))
 	{
 		FArchive* Ar = FileManager->CreateFileWriter(Filename.GetCharArray().GetData());
 		if (Ar != nullptr)
@@ -126,11 +154,12 @@ FString FHighResScreenshotConfig::SaveImage(const FString& File, const TArray<TP
 		}
 		else
 		{
-			Filename = FString("Failed to open for writing: ") + Filename;
+			return false;
 		}
 	}
-	return Filename;
+
+	return true;
 }
 
-template ENGINE_API FString FHighResScreenshotConfig::SaveImage<FColor>(const FString& File, const TArray<FColor>& Bitmap, const FIntPoint& BitmapSize, EPixelFormat SourceFormat) const;
-template ENGINE_API FString FHighResScreenshotConfig::SaveImage<FFloat16Color>(const FString& File, const TArray<FFloat16Color>& Bitmap, const FIntPoint& BitmapSize, EPixelFormat SourceFormat) const;
+template ENGINE_API bool FHighResScreenshotConfig::SaveImage<FColor>(const FString& File, const TArray<FColor>& Bitmap, const FIntPoint& BitmapSize, FString* OutFilename) const;
+template ENGINE_API bool FHighResScreenshotConfig::SaveImage<FFloat16Color>(const FString& File, const TArray<FFloat16Color>& Bitmap, const FIntPoint& BitmapSize, FString* OutFilename) const;
diff --git a/Engine/Source/Runtime/Engine/Private/UnrealClient.cpp b/Engine/Source/Runtime/Engine/Private/UnrealClient.cpp
--- a/Engine/Source/Runtime/Engine/Private/UnrealClient.cpp
+++ b/Engine/Source/Runtime/Engine/Private/UnrealClient.cpp
@@ -205,11 +205,11 @@ void FScreenshotRequest::RequestScreenshot( const FString& InFilename, bool bInS
 }
 
 
-void FScreenshotRequest::RequestScreenshot( bool bInShowUI )
+void FScreenshotRequest::RequestScreenshot( bool bInShowUI, const FString& InExtension )
 {
 	FString NewFilename;
 	CreateViewportScreenShotFilename( NewFilename );
-	FFileHelper::GenerateNextBitmapFilename(*NewFilename, Filename);
+	FFileHelper::GenerateNextBitmapFilename(NewFilename, InExtension, Filename);
 
 	bShowUI = bInShowUI;
 }
@@ -971,7 +971,7 @@ void FViewport::Draw( bool bShouldPresent /*= true */)
 			if( GIsHighResScreenshot || bTakeHighResScreenShot )
 			{
 				bool bShowUI = false;
-				FScreenshotRequest::RequestScreenshot( bShowUI );
+				FScreenshotRequest::RequestScreenshot( bShowUI, TEXT("png") );
 				GIsHighResScreenshot = true;
 				GScreenMessagesRestoreState = GAreScreenMessagesEnabled;
 				GAreScreenMessagesEnabled = false;
diff --git a/Engine/Source/Runtime/Engine/Public/HighResScreenshot.h b/Engine/Source/Runtime/Engine/Public/HighResScreenshot.h
--- a/Engine/Source/Runtime/Engine/Public/HighResScreenshot.h
+++ b/Engine/Source/Runtime/Engine/Public/HighResScreenshot.h
@@ -40,7 +40,7 @@ struct ENGINE_API FHighResScreenshotConfig
 
 	/** Save to image file **/
 	template<typename TPixelType>
-	ENGINE_API FString SaveImage(const FString& File, const TArray<TPixelType>& Bitmap, const FIntPoint& BitmapSize, EPixelFormat SourceFormat) const;
+	ENGINE_API bool SaveImage(const FString& File, const TArray<TPixelType>& Bitmap, const FIntPoint& BitmapSize, FString* OutFilename = nullptr) const;
 };
 
 ENGINE_API FHighResScreenshotConfig& GetHighResScreenshotConfig();
\ No newline at end of file
diff --git a/Engine/Source/Runtime/Engine/Public/UnrealClient.h b/Engine/Source/Runtime/Engine/Public/UnrealClient.h
--- a/Engine/Source/Runtime/Engine/Public/UnrealClient.h
+++ b/Engine/Source/Runtime/Engine/Public/UnrealClient.h
@@ -111,11 +111,13 @@ struct ENGINE_API FScreenshotRequest
 	static void RequestScreenshot( const FString& InFilename, bool bInShowUI );
 
 	/**
-	 * Requests a new screenshot.  This method will auto-generate a filename via the format "Screenshot{0} where {0} is an incremented coutner"
+	 * Requests a new screenshot.  This method will auto-generate a filename via the format "Screenshot{0}.{1} where {0} is an incremented coutner
+	 * and {1} is an optional extension (defaulting to .bmp)
 	 *
 	 * @param bInShowUI		Whether or not to show Slate UI
+	 * @param InExtension	Which file extension to use
 	 */
-	static void RequestScreenshot( bool bInShowUI );
+	static void RequestScreenshot( bool bInShowUI, const FString& InExtension = FString(TEXT("bmp") ) );
 
 	/**
 	 * Resets a screenshot request
diff --git a/Engine/Source/Runtime/Renderer/Private/PostProcess/RenderingCompositionGraph.cpp b/Engine/Source/Runtime/Renderer/Private/PostProcess/RenderingCompositionGraph.cpp
--- a/Engine/Source/Runtime/Renderer/Private/PostProcess/RenderingCompositionGraph.cpp
+++ b/Engine/Source/Runtime/Renderer/Private/PostProcess/RenderingCompositionGraph.cpp
@@ -238,6 +238,7 @@ void FRenderingCompositionGraph::RecursivelyGatherDependencies(const FRenderingC
 void FRenderingCompositionGraph::DumpOutputToFile(FRenderingCompositePassContext& Context, const FString& Filename, FRenderingCompositeOutput* Output) const
 {
 	FSceneRenderTargetItem& RenderTargetItem = Output->PooledRenderTarget->GetRenderTargetItem();
+	FHighResScreenshotConfig& HighResScreenshotConfig = GetHighResScreenshotConfig();
 	FTextureRHIRef Texture = RenderTargetItem.TargetableTexture ? RenderTargetItem.TargetableTexture : RenderTargetItem.ShaderResourceTexture;
 	check(Texture);
 	check(Texture->GetTexture2D());
@@ -246,7 +247,7 @@ void FRenderingCompositionGraph::DumpOutputToFile(FRenderingCompositePassContext
 	int32 MSAAXSamples = Texture->GetNumSamples();
 	if (GIsHighResScreenshot)
 	{
-		SourceRect = GetHighResScreenshotConfig().CaptureRegion;
+		SourceRect = HighResScreenshotConfig.CaptureRegion;
 		if (SourceRect.Area() == 0)
 		{
 			SourceRect = Context.View.ViewRect;
@@ -268,7 +269,7 @@ void FRenderingCompositionGraph::DumpOutputToFile(FRenderingCompositePassContext
 		{
 			TArray<FFloat16Color> Bitmap;
 			Context.RHICmdList.ReadSurfaceFloatData(Texture, SourceRect, Bitmap, (ECubeFace)0, 0, 0);
-			ResultPath = GetHighResScreenshotConfig().SaveImage(Filename, Bitmap, DestSize, PixelFormat);
+			HighResScreenshotConfig.SaveImage(Filename, Bitmap, DestSize, &ResultPath);
 		}
 		break;
 		case PF_R8G8B8A8:
@@ -283,7 +284,7 @@ void FRenderingCompositionGraph::DumpOutputToFile(FRenderingCompositePassContext
 			{
 				Pixel->A = 255;
 			}
-			ResultPath = GetHighResScreenshotConfig().SaveImage(Filename, Bitmap, DestSize, PixelFormat);
+			HighResScreenshotConfig.SaveImage(Filename, Bitmap, DestSize, &ResultPath);
 		}
 		break;
 	}
