Added _DEPRECATED to both GRHIFeatureLevel and GRHIShaderPlatform, to indicate that they shouldn't be used any more. Added documention to describe why this has been done

diff --git a/Engine/Source/Runtime/Engine/Private/GPUSkinCache.cpp b/Engine/Source/Runtime/Engine/Private/GPUSkinCache.cpp
--- a/Engine/Source/Runtime/Engine/Private/GPUSkinCache.cpp
+++ b/Engine/Source/Runtime/Engine/Private/GPUSkinCache.cpp
@@ -470,7 +470,7 @@ FGPUSkinCache::FElementCacheStatusInfo* FGPUSkinCache::FindEvictableCacheStatusI
 void FGPUSkinCache::DispatchSkinCacheProcess(FRHICommandListImmediate& RHICmdList, uint32 InputStreamFloatOffset, uint32 OutputBufferFloatOffset, FBoneBufferTypeRef BoneBuffer, FUniformBufferRHIRef UniformBuffer, const FVertexBufferInfo* VBInfo, uint32 VertexStride, uint32 VertexCount, const FVector& MeshOrigin, const FVector& MeshExtension, bool bUseExtraBoneInfluences)
 {
 	// MOBILEPREVIEWTODO: Proper value for this
-	const auto FeatureLevel = GRHIFeatureLevel;
+	const auto FeatureLevel = GRHIFeatureLevel_DEPRECATED;
 
 	if (FeatureLevel < ERHIFeatureLevel::SM5)
 	{
diff --git a/Engine/Source/Runtime/Engine/Private/GPUSkinVertexFactory.cpp b/Engine/Source/Runtime/Engine/Private/GPUSkinVertexFactory.cpp
--- a/Engine/Source/Runtime/Engine/Private/GPUSkinVertexFactory.cpp
+++ b/Engine/Source/Runtime/Engine/Private/GPUSkinVertexFactory.cpp
@@ -116,7 +116,7 @@ void FGPUBaseSkinVertexFactory::ShaderDataType::UpdateBoneData()
 {
 	uint32 NumBones = BoneMatrices.Num();
 	check(NumBones <= MaxGPUSkinBones);
-	if (GRHIFeatureLevel >= ERHIFeatureLevel::SM4)
+	if (GRHIFeatureLevel_DEPRECATED >= ERHIFeatureLevel::SM4)
 	{
 		static FSharedPoolPolicyData PoolPolicy;
 		uint32 NumVectors = NumBones*3;
@@ -690,7 +690,7 @@ void FGPUBaseSkinAPEXClothVertexFactory::ClothShaderType::UpdateClothSimulData(c
 {
 	uint32 NumSimulVerts = InSimulPositions.Num();
 
-	if (GRHIFeatureLevel >= ERHIFeatureLevel::SM4)
+	if (GRHIFeatureLevel_DEPRECATED >= ERHIFeatureLevel::SM4)
 	{
 		FMath::Min(NumSimulVerts, (uint32)MAX_APEXCLOTH_VERTICES_FOR_VB);
 
diff --git a/Engine/Source/Runtime/Engine/Private/InstancedStaticMesh.cpp b/Engine/Source/Runtime/Engine/Private/InstancedStaticMesh.cpp
--- a/Engine/Source/Runtime/Engine/Private/InstancedStaticMesh.cpp
+++ b/Engine/Source/Runtime/Engine/Private/InstancedStaticMesh.cpp
@@ -310,7 +310,7 @@ void FStaticMeshInstanceBuffer::AllocateData()
 	// Clear any old VertexData before allocating.
 	CleanUp();
 
-	const bool bInstanced = RHISupportsInstancing(GRHIShaderPlatform);
+	const bool bInstanced = RHISupportsInstancing(GRHIShaderPlatform_DEPRECATED);
 	const bool bNeedsCPUAccess = !bInstanced;
 	InstanceData = new FStaticMeshInstanceData(bNeedsCPUAccess);
 	// Calculate the vertex stride.
@@ -444,7 +444,7 @@ void FInstancedStaticMeshVertexFactory::Copy(const FInstancedStaticMeshVertexFac
 
 void FInstancedStaticMeshVertexFactory::InitRHI()
 {
-	const bool bInstanced = RHISupportsInstancing(GRHIShaderPlatform);
+	const bool bInstanced = RHISupportsInstancing(GRHIShaderPlatform_DEPRECATED);
 
 	// If the vertex buffer containing position is not the same vertex buffer containing the rest of the data,
 	// then initialize PositionStream and PositionDeclaration.
@@ -689,7 +689,7 @@ void FInstancedStaticMeshRenderData::InitStaticMeshVertexFactories(
 		FInstancedStaticMeshRenderData* InstancedRenderData,
 		UStaticMesh* Parent)
 {
-	const bool bInstanced = RHISupportsInstancing(GRHIShaderPlatform);
+	const bool bInstanced = RHISupportsInstancing(GRHIShaderPlatform_DEPRECATED);
 
 	for( int32 LODIndex=0;LODIndex<VertexFactories->Num(); LODIndex++ )
 	{
@@ -886,7 +886,7 @@ public:
 
 		check(InstancedRenderData.InstanceBuffer.GetStride() == sizeof(FInstancingUserData::FInstanceStream));
 
-		const bool bInstanced = RHISupportsInstancing(GRHIShaderPlatform);
+		const bool bInstanced = RHISupportsInstancing(GRHIShaderPlatform_DEPRECATED);
 
 		// Copy the parameters for LOD - all instances
 		UserData_AllInstances.StartCullDistance = InComponent->InstanceStartCullDistance;
@@ -1119,7 +1119,7 @@ void FInstancedStaticMeshSceneProxy::DrawDynamicElements(FPrimitiveDrawInterface
 
 int32 FInstancedStaticMeshSceneProxy::GetNumMeshBatches() const
 {
-	const bool bInstanced = RHISupportsInstancing(GRHIShaderPlatform);
+	const bool bInstanced = RHISupportsInstancing(GRHIShaderPlatform_DEPRECATED);
 
 	if (bInstanced)
 	{
@@ -1136,7 +1136,7 @@ int32 FInstancedStaticMeshSceneProxy::GetNumMeshBatches() const
 
 void FInstancedStaticMeshSceneProxy::SetupInstancedMeshBatch(int32 LODIndex, int32 BatchIndex, FMeshBatch& OutMeshBatch) const
 {
-	const bool bInstanced = RHISupportsInstancing(GRHIShaderPlatform);
+	const bool bInstanced = RHISupportsInstancing(GRHIShaderPlatform_DEPRECATED);
 	OutMeshBatch.VertexFactory = &InstancedRenderData.VertexFactories[LODIndex];
 	const uint32 NumInstances = InstancedRenderData.InstanceBuffer.GetNumInstances();
 	FMeshBatchElement& BatchElement0 = OutMeshBatch.Elements[0];
@@ -1886,7 +1886,7 @@ void FInstancedStaticMeshVertexFactoryShaderParameters::SetMesh( FRHICommandList
 		SetShaderValue(RHICmdList, VS, InstancingFadeOutParamsParameter, InstancingFadeOutParams );
 	}
 
-	const bool bInstanced = RHISupportsInstancing(GRHIShaderPlatform);
+	const bool bInstanced = RHISupportsInstancing(GRHIShaderPlatform_DEPRECATED);
 	if (!bInstanced)
 	{
 		if (CPUInstanceShadowMapBias.IsBound())
diff --git a/Engine/Source/Runtime/Engine/Private/Landscape/Landscape.cpp b/Engine/Source/Runtime/Engine/Private/Landscape/Landscape.cpp
--- a/Engine/Source/Runtime/Engine/Private/Landscape/Landscape.cpp
+++ b/Engine/Source/Runtime/Engine/Private/Landscape/Landscape.cpp
@@ -412,7 +412,7 @@ void ULandscapeComponent::PostLoad()
 		}
 
 		// If we're loading on a platform that doesn't require cooked data, attempt to load missing data from the DDC
-		if (!FPlatformProperties::RequiresCookedData() && (GIsEditor || GRHIFeatureLevel == ERHIFeatureLevel::ES2))
+		if (!FPlatformProperties::RequiresCookedData() && (GIsEditor || GMaxRHIFeatureLevel == ERHIFeatureLevel::ES2))
 		{
 			// Only check the DDC if we don't already have it loaded
 			if (!PlatformData.HasValidPlatformData())
diff --git a/Engine/Source/Runtime/Engine/Private/RendererSupport.cpp b/Engine/Source/Runtime/Engine/Private/RendererSupport.cpp
--- a/Engine/Source/Runtime/Engine/Private/RendererSupport.cpp
+++ b/Engine/Source/Runtime/Engine/Private/RendererSupport.cpp
@@ -49,7 +49,7 @@ void ClearReferencesToRendererModuleClasses(
 	}
 
 	// Save off shaders by serializing them into memory, and remove all shader map references to FShaders
-	GlobalShaderData = TScopedPointer<TArray<uint8> >(BackupGlobalShaderMap(GRHIShaderPlatform));
+	GlobalShaderData = TScopedPointer<TArray<uint8> >(BackupGlobalShaderMap(GRHIShaderPlatform_DEPRECATED));
 	UMaterial::BackupMaterialShadersToMemory(ShaderMapToSerializedShaderData);
 
 	// Verify no FShaders still in memory
@@ -150,9 +150,9 @@ void RestoreReferencesToRendererModuleClasses(
 
 	// Restore FShaders from the serialized memory blobs
 	// Shader maps may still not be complete after this due to code changes picked up in the recompile
-	RestoreGlobalShaderMap(GRHIShaderPlatform, *GlobalShaderData);
-	UMaterial::RestoreMaterialShadersFromMemory(GRHIShaderPlatform, ShaderMapToSerializedShaderData);
-	FMaterialShaderMap::FixupShaderTypes(GRHIShaderPlatform, ShaderTypeNames, VertexFactoryTypeNames);
+	RestoreGlobalShaderMap(GRHIShaderPlatform_DEPRECATED, *GlobalShaderData);
+	UMaterial::RestoreMaterialShadersFromMemory(GRHIShaderPlatform_DEPRECATED, ShaderMapToSerializedShaderData);
+	FMaterialShaderMap::FixupShaderTypes(GRHIShaderPlatform_DEPRECATED, ShaderTypeNames, VertexFactoryTypeNames);
 
 	TArray<FShaderType*> OutdatedShaderTypes;
 	TArray<const FVertexFactoryType*> OutdatedFactoryTypes;
diff --git a/Engine/Source/Runtime/Engine/Private/SkeletalMesh.cpp b/Engine/Source/Runtime/Engine/Private/SkeletalMesh.cpp
--- a/Engine/Source/Runtime/Engine/Private/SkeletalMesh.cpp
+++ b/Engine/Source/Runtime/Engine/Private/SkeletalMesh.cpp
@@ -1186,7 +1186,7 @@ void FStaticLODModel::Serialize( FArchive& Ar, UObject* Owner, int32 Idx )
 		if( Ar.IsLoading() )
 		{
 			// set cpu skinning flag on the vertex buffer so that the resource arrays know if they need to be CPU accessible
-			VertexBufferGPUSkin.SetNeedsCPUAccess(bKeepBuffersInCPUMemory || SkelMeshOwner->GetImportedResource()->RequiresCPUSkinning(GRHIFeatureLevel));
+			VertexBufferGPUSkin.SetNeedsCPUAccess(bKeepBuffersInCPUMemory || SkelMeshOwner->GetImportedResource()->RequiresCPUSkinning(GRHIFeatureLevel_DEPRECATED));
 		}
 		Ar << NumTexCoords;
 		Ar << VertexBufferGPUSkin;
@@ -4343,7 +4343,7 @@ void FSkeletalMeshSceneProxy::GetDynamicElementsSection(const TArray<const FScen
 			BatchElement.NumPrimitives = Section.NumTriangles;
 			if( GIsEditor && MeshObject->ProgressiveDrawingFraction != 1.f )
 			{
-				if (Mesh.MaterialRenderProxy->GetMaterial(GRHIFeatureLevel)->GetBlendMode() == BLEND_Translucent)
+				if (Mesh.MaterialRenderProxy->GetMaterial(GRHIFeatureLevel_DEPRECATED)->GetBlendMode() == BLEND_Translucent)
 				{
 					BatchElement.NumPrimitives = FMath::RoundToInt(((float)Section.NumTriangles)*FMath::Clamp<float>(MeshObject->ProgressiveDrawingFraction,0.f,1.f));
 					if( BatchElement.NumPrimitives == 0 )
diff --git a/Engine/Source/Runtime/Engine/Private/SkeletalMeshMerge.cpp b/Engine/Source/Runtime/Engine/Private/SkeletalMeshMerge.cpp
--- a/Engine/Source/Runtime/Engine/Private/SkeletalMeshMerge.cpp
+++ b/Engine/Source/Runtime/Engine/Private/SkeletalMeshMerge.cpp
@@ -272,7 +272,7 @@ static void BoneMapToNewRefSkel(const TArray<FBoneIndexType>& InBoneMap, const T
 */
 void FSkeletalMeshMerge::GenerateNewSectionArray( TArray<FNewSectionInfo>& NewSectionArray, int32 LODIdx )
 {
-	const int32 MaxGPUSkinBones = GetFeatureLevelMaxNumberOfBones(GRHIFeatureLevel);
+	const int32 MaxGPUSkinBones = GetFeatureLevelMaxNumberOfBones(GRHIFeatureLevel_DEPRECATED);
 
 	NewSectionArray.Empty();
 	for( int32 MeshIdx=0; MeshIdx < SrcMeshList.Num(); MeshIdx++ )
@@ -605,7 +605,7 @@ void FSkeletalMeshMerge::GenerateLODModel( int32 LODIdx )
 	// copy the new vertices and indices to the vertex buffer for the new model
 	MergeLODModel.VertexBufferGPUSkin.SetUseFullPrecisionUVs(MergeMesh->bUseFullPrecisionUVs);
 	// set CPU skinning on vertex buffer since it affects the type of TResourceArray needed
-	MergeLODModel.VertexBufferGPUSkin.SetNeedsCPUAccess(MergeResource->RequiresCPUSkinning(GRHIFeatureLevel));
+	MergeLODModel.VertexBufferGPUSkin.SetNeedsCPUAccess(MergeResource->RequiresCPUSkinning(GRHIFeatureLevel_DEPRECATED));
 	MergeLODModel.VertexBufferGPUSkin.SetHasExtraBoneInfluences(MergeResource->HasExtraBoneInfluences());
 	// Set the number of tex coords on this vertex buffer
 	MergeLODModel.VertexBufferGPUSkin.SetNumTexCoords(TotalNumUVs);
diff --git a/Engine/Source/Runtime/Engine/Private/SkeletalRenderGPUSkin.cpp b/Engine/Source/Runtime/Engine/Private/SkeletalRenderGPUSkin.cpp
--- a/Engine/Source/Runtime/Engine/Private/SkeletalRenderGPUSkin.cpp
+++ b/Engine/Source/Runtime/Engine/Private/SkeletalRenderGPUSkin.cpp
@@ -233,14 +233,14 @@ void FSkeletalMeshObjectGPUSkin::UpdateDynamicData_RenderThread(FRHICommandListI
 
 	if(DataPresent)
 	{
-		bool bGPUSkinCacheEnabled = GEnableGPUSkinCache && (GRHIFeatureLevel >= ERHIFeatureLevel::SM5);
+		bool bGPUSkinCacheEnabled = GEnableGPUSkinCache && (GRHIFeatureLevel_DEPRECATED >= ERHIFeatureLevel::SM5);
 		for( int32 ChunkIdx=0; ChunkIdx < Chunks.Num(); ChunkIdx++ )
 		{
 			const FSkelMeshChunk& Chunk = Chunks[ChunkIdx];
 
 			bool bClothFactory = (DynamicData->ClothSimulUpdateData.Num() > 0) && Chunk.HasApexClothData();
 
-			if (GRHIFeatureLevel < ERHIFeatureLevel::SM4)
+			if (GRHIFeatureLevel_DEPRECATED < ERHIFeatureLevel::SM4)
 			{
 				bClothFactory = false;
 			}
@@ -836,7 +836,7 @@ void FSkeletalMeshObjectGPUSkin::FVertexFactoryData::InitAPEXClothVertexFactorie
 	ClothVertexFactories.Empty(Chunks.Num());
 	for( int32 FactoryIdx=0; FactoryIdx < Chunks.Num(); FactoryIdx++ )
 	{
-		if (Chunks[FactoryIdx].HasApexClothData() && GRHIFeatureLevel >= ERHIFeatureLevel::SM4)
+		if (Chunks[FactoryIdx].HasApexClothData() && GRHIFeatureLevel_DEPRECATED >= ERHIFeatureLevel::SM4)
 		{
 			if (VertexBuffers.VertexBufferGPUSkin->HasExtraBoneInfluences())
 			{
diff --git a/Engine/Source/Runtime/IOS/MetalRHI/Private/MetalRHI.cpp b/Engine/Source/Runtime/IOS/MetalRHI/Private/MetalRHI.cpp
--- a/Engine/Source/Runtime/IOS/MetalRHI/Private/MetalRHI.cpp
+++ b/Engine/Source/Runtime/IOS/MetalRHI/Private/MetalRHI.cpp
@@ -42,8 +42,8 @@ FMetalDynamicRHI::FMetalDynamicRHI()
 	check( !GIsThreadedRendering );
 
 	// Initialize the RHI capabilities.
-	GMaxRHIFeatureLevelValue = ERHIFeatureLevel::ES2;
-	GMaxRHIShaderPlatformValue = SP_METAL;
+	GMaxRHIFeatureLevel = ERHIFeatureLevel::ES2;
+	GMaxRHIShaderPlatform = SP_METAL;
 
 	GPixelCenterOffset = 0.0f;
 	GSupportsVertexInstancing = false;
diff --git a/Engine/Source/Runtime/NullDrv/Private/NullRHI.cpp b/Engine/Source/Runtime/NullDrv/Private/NullRHI.cpp
--- a/Engine/Source/Runtime/NullDrv/Private/NullRHI.cpp
+++ b/Engine/Source/Runtime/NullDrv/Private/NullRHI.cpp
@@ -11,7 +11,7 @@
 
 FNullDynamicRHI::FNullDynamicRHI()
 {
-	GMaxRHIShaderPlatformValue = ShaderFormatToLegacyShaderPlatform(FName(FPlatformMisc::GetNullRHIShaderFormat()));
+	GMaxRHIShaderPlatform = ShaderFormatToLegacyShaderPlatform(FName(FPlatformMisc::GetNullRHIShaderFormat()));
 }
 
 
diff --git a/Engine/Source/Runtime/OpenGLDrv/Private/OpenGLCommands.cpp b/Engine/Source/Runtime/OpenGLDrv/Private/OpenGLCommands.cpp
--- a/Engine/Source/Runtime/OpenGLDrv/Private/OpenGLCommands.cpp
+++ b/Engine/Source/Runtime/OpenGLDrv/Private/OpenGLCommands.cpp
@@ -1566,7 +1566,7 @@ void FOpenGLDynamicRHI::RHISetRenderTargets(
 
 	FOpenGLTextureBase* NewDepthStencilRT = GetOpenGLTextureFromRHITexture(NewDepthStencilTargetRHI);
 
-	if (IsES2Platform(GRHIShaderPlatform) && !IsPCPlatform(GRHIShaderPlatform))
+	if (IsES2Platform(GRHIShaderPlatform_DEPRECATED) && !IsPCPlatform(GRHIShaderPlatform_DEPRECATED))
 	{
 		// @todo-mobile
 
diff --git a/Engine/Source/Runtime/OpenGLDrv/Private/OpenGLDevice.cpp b/Engine/Source/Runtime/OpenGLDrv/Private/OpenGLDevice.cpp
--- a/Engine/Source/Runtime/OpenGLDrv/Private/OpenGLDevice.cpp
+++ b/Engine/Source/Runtime/OpenGLDrv/Private/OpenGLDevice.cpp
@@ -621,8 +621,8 @@ static void InitRHICapabilitiesForGL()
 	const GLint MinorVersion = FOpenGL::GetMinorVersion();
 
 	// Shader platform & RHI feature level
-	GMaxRHIFeatureLevelValue = FOpenGL::GetFeatureLevel();
-	GMaxRHIShaderPlatformValue = FOpenGL::GetShaderPlatform();
+	GMaxRHIFeatureLevel = FOpenGL::GetFeatureLevel();
+	GMaxRHIShaderPlatform = FOpenGL::GetShaderPlatform();
 
 	// Emulate uniform buffers on ES2, unless we're on a desktop platform emulating ES2.
 	GUseEmulatedUniformBuffers = IsES2Platform(GMaxRHIShaderPlatform) && !IsPCPlatform(GMaxRHIShaderPlatform);
@@ -664,7 +664,7 @@ static void InitRHICapabilitiesForGL()
 	GSupportsVertexInstancing = true;
 	GHardwareHiddenSurfaceRemoval = FOpenGL::HasHardwareHiddenSurfaceRemoval();
 
-	GShaderPlatformForFeatureLevel[ERHIFeatureLevel::ES2] = (GMaxRHIFeatureLevel == ERHIFeatureLevel::ES2) ? GRHIShaderPlatform : SP_OPENGL_PCES2;
+	GShaderPlatformForFeatureLevel[ERHIFeatureLevel::ES2] = (GMaxRHIFeatureLevel == ERHIFeatureLevel::ES2) ? GRHIShaderPlatform_DEPRECATED : SP_OPENGL_PCES2;
 	GShaderPlatformForFeatureLevel[ERHIFeatureLevel::ES3_1] = SP_NumPlatforms;
 	GShaderPlatformForFeatureLevel[ERHIFeatureLevel::SM4] = PLATFORM_MAC ? SP_OPENGL_SM4_MAC : SP_OPENGL_SM4;
 	GShaderPlatformForFeatureLevel[ERHIFeatureLevel::SM5] = OPENGL_ES31 ? SP_OPENGL_ES31_EXT : SP_OPENGL_SM5;
@@ -934,7 +934,7 @@ static bool VerifyCompiledShader(GLuint Shader, const ANSICHAR* GlslCode, bool I
 static void CheckTextureCubeLodSupport()
 {
 #if PLATFORM_ANDROID
-	if (IsES2Platform(GRHIShaderPlatform))
+	if (IsES2Platform(GRHIShaderPlatform_DEPRECATED))
 	{
 		UE_LOG(LogRHI, Display, TEXT("Testing for shader compiler compatibility"));
 		// This code creates a sample program and finds out which hacks are required to compile it
diff --git a/Engine/Source/Runtime/OpenGLDrv/Private/OpenGLShaders.cpp b/Engine/Source/Runtime/OpenGLDrv/Private/OpenGLShaders.cpp
--- a/Engine/Source/Runtime/OpenGLDrv/Private/OpenGLShaders.cpp
+++ b/Engine/Source/Runtime/OpenGLDrv/Private/OpenGLShaders.cpp
@@ -336,7 +336,7 @@ ShaderType* CompileOpenGLShader(const TArray<uint8>& Code)
 		ANSICHAR* VersionString = new ANSICHAR[18];
 		memset((void*)VersionString, '\0', 18 * sizeof(ANSICHAR));
 
-		if (IsES2Platform(GRHIShaderPlatform))
+		if (IsES2Platform(GRHIShaderPlatform_DEPRECATED))
 		{
 			//#version 100 has to be the first line in the file, so it has to be added before anything else.
 			if (FOpenGL::UseES30ShadingLanguage())
@@ -365,7 +365,7 @@ ShaderType* CompileOpenGLShader(const TArray<uint8>& Code)
 #endif
 
 #if PLATFORM_ANDROID 
-		if (IsES2Platform(GRHIShaderPlatform))
+		if (IsES2Platform(GRHIShaderPlatform_DEPRECATED))
 		{
 			//Here is some code to add in a #define for textureCubeLodEXT
 			const ANSICHAR* ExtensionString = "";
@@ -1259,7 +1259,7 @@ static FOpenGLLinkedProgram* LinkProgram( const FOpenGLLinkedProgramConfiguratio
 	}
 	
 	// E.g. GLSL_430 uses layout(location=xx) instead of having to call glBindAttribLocation and glBindFragDataLocation
-	if (OpenGLShaderPlatformNeedsBindLocation(GRHIShaderPlatform))
+	if (OpenGLShaderPlatformNeedsBindLocation(GRHIShaderPlatform_DEPRECATED))
 	{
 		// Bind attribute indices.
 		if (Config.Shaders[CrossCompiler::SHADER_STAGE_VERTEX].Resource)
diff --git a/Engine/Source/Runtime/RHI/Private/RHI.cpp b/Engine/Source/Runtime/RHI/Private/RHI.cpp
--- a/Engine/Source/Runtime/RHI/Private/RHI.cpp
+++ b/Engine/Source/Runtime/RHI/Private/RHI.cpp
@@ -223,20 +223,10 @@ void RHIPrivateBeginFrame()
 // The current shader platform.
 //
 
-RHI_API EShaderPlatform GMaxRHIShaderPlatformValue = SP_PCD3D_SM5;
+RHI_API EShaderPlatform GMaxRHIShaderPlatform = SP_PCD3D_SM5;
 
 /** The maximum feature level supported on this machine */
-RHI_API ERHIFeatureLevel::Type GMaxRHIFeatureLevelValue = ERHIFeatureLevel::SM5;
-
-RHI_API ERHIFeatureLevel::Type GetMaxRHIFeatureLevel()
-{
-	return GMaxRHIFeatureLevelValue;
-}
-
-RHI_API EShaderPlatform GetMaxRHIShaderPlatform()
-{
-	return GMaxRHIShaderPlatformValue;
-}
+RHI_API ERHIFeatureLevel::Type GMaxRHIFeatureLevel = ERHIFeatureLevel::SM5;
 
 FName FeatureLevelNames[] = 
 {
diff --git a/Engine/Source/Runtime/RHI/Public/RHI.h b/Engine/Source/Runtime/RHI/Public/RHI.h
--- a/Engine/Source/Runtime/RHI/Public/RHI.h
+++ b/Engine/Source/Runtime/RHI/Public/RHI.h
@@ -237,26 +237,15 @@ extern RHI_API void GetFeatureLevelName(ERHIFeatureLevel::Type InFeatureLevel, F
 /** Creates an FName for the given feature level. */
 extern RHI_API void GetFeatureLevelName(ERHIFeatureLevel::Type InFeatureLevel, FName& OutName);
 
-extern RHI_API ERHIFeatureLevel::Type GMaxRHIFeatureLevelValue;
-extern RHI_API EShaderPlatform GMaxRHIShaderPlatformValue;
-
-/** Function for retrieving the current feature level. Only exists to provide a convient way of tracking access to this global during mobile preview work */
-extern RHI_API ERHIFeatureLevel::Type GetMaxRHIFeatureLevel();
-extern RHI_API EShaderPlatform GetMaxRHIShaderPlatform();
-
-/** treating GRHIFeatureLevel as a function allows for better usage tracking at small cost to performance. */
-#define RHI_FEATURE_LEVEL_AS_FUNCTION 0
-
-#define GMaxRHIFeatureLevel GMaxRHIFeatureLevelValue
-#define GMaxRHIShaderPlatform GMaxRHIShaderPlatformValue
-
-#if RHI_FEATURE_LEVEL_AS_FUNCTION
-	#define GRHIShaderPlatform GetMaxRHIShaderPlatform()
-	#define GRHIFeatureLevel GetMaxRHIFeatureLevel()
-#else
-	#define GRHIShaderPlatform GMaxRHIShaderPlatform
-	#define GRHIFeatureLevel GMaxRHIFeatureLevel
-#endif
+// The maximum feature level and shader platform available on this system
+extern RHI_API ERHIFeatureLevel::Type GMaxRHIFeatureLevel;
+extern RHI_API EShaderPlatform GMaxRHIShaderPlatform;
+
+// GRHIFeatureLevel and GRHIShaderPlatform have been deprecated. There is no longer a current featurelevel/shaderplatform that
+// should be used for all rendering, rather a specific set for each view. The deprecated names are being kept around temporarily 
+// while the last instances are removed.
+#define GRHIShaderPlatform_DEPRECATED GMaxRHIShaderPlatform
+#define GRHIFeatureLevel_DEPRECATED GMaxRHIFeatureLevel
 
 /** Table for finding out which shader platform corresponds to a given feature level for this RHI. */
 extern RHI_API EShaderPlatform GShaderPlatformForFeatureLevel[ERHIFeatureLevel::Num];
diff --git a/Engine/Source/Runtime/Windows/D3D11RHI/Private/D3D11Device.cpp b/Engine/Source/Runtime/Windows/D3D11RHI/Private/D3D11Device.cpp
--- a/Engine/Source/Runtime/Windows/D3D11RHI/Private/D3D11Device.cpp
+++ b/Engine/Source/Runtime/Windows/D3D11RHI/Private/D3D11Device.cpp
@@ -75,18 +75,18 @@ FD3D11DynamicRHI::FD3D11DynamicRHI(IDXGIFactory* InDXGIFactory,D3D_FEATURE_LEVEL
 	// ES2 feature level emulation in D3D11
 	if (FParse::Param(FCommandLine::Get(), TEXT("FeatureLevelES2")) && !GIsEditor)
 	{
-		GMaxRHIFeatureLevelValue = ERHIFeatureLevel::ES2;
-		GMaxRHIShaderPlatformValue = SP_PCD3D_ES2;
+		GMaxRHIFeatureLevel = ERHIFeatureLevel::ES2;
+		GMaxRHIShaderPlatform = SP_PCD3D_ES2;
 	}
 	else if(FeatureLevel == D3D_FEATURE_LEVEL_11_0)
 	{
-		GMaxRHIFeatureLevelValue = ERHIFeatureLevel::SM5;
-		GMaxRHIShaderPlatformValue = SP_PCD3D_SM5;
+		GMaxRHIFeatureLevel = ERHIFeatureLevel::SM5;
+		GMaxRHIShaderPlatform = SP_PCD3D_SM5;
 	}
 	else if(FeatureLevel == D3D_FEATURE_LEVEL_10_0)
 	{
-		GMaxRHIFeatureLevelValue = ERHIFeatureLevel::SM4;
-		GMaxRHIShaderPlatformValue = SP_PCD3D_SM4;
+		GMaxRHIFeatureLevel = ERHIFeatureLevel::SM4;
+		GMaxRHIShaderPlatform = SP_PCD3D_SM4;
 	}
 
 	GPixelCenterOffset = 0.0f;	// Note that in D3D11, there is no half-texel offset (ala DX9)	
diff --git a/Engine/Source/Runtime/Windows/D3D11RHI/Private/WinRT/WinRTD3D11Device.cpp b/Engine/Source/Runtime/Windows/D3D11RHI/Private/WinRT/WinRTD3D11Device.cpp
--- a/Engine/Source/Runtime/Windows/D3D11RHI/Private/WinRT/WinRTD3D11Device.cpp
+++ b/Engine/Source/Runtime/Windows/D3D11RHI/Private/WinRT/WinRTD3D11Device.cpp
@@ -14,7 +14,7 @@ FDynamicRHI* FD3D11DynamicRHIModule::CreateRHI()
 	if (DynamicRHI)
 	{
 		// Initialize the RHI capabilities.
-		GMaxRHIShaderPlatformValue = SP_PCD3D_SM5;
+		GMaxRHIShaderPlatform = SP_PCD3D_SM5;
 
 		// This will be set in the D3D11DynamicRHI::Init call!
 		//	GIsRHIInitialized = true;
