Set destructible components to use async scene automatically. By default async scene is turned off so this will still be in sync scene for new users but turning async scene on will automatically place destructibles in async

diff --git a/Engine/Source/Runtime/Engine/Private/DestructibleComponent.cpp b/Engine/Source/Runtime/Engine/Private/DestructibleComponent.cpp
--- a/Engine/Source/Runtime/Engine/Private/DestructibleComponent.cpp
+++ b/Engine/Source/Runtime/Engine/Private/DestructibleComponent.cpp
@@ -28,6 +28,7 @@ UDestructibleComponent::UDestructibleComponent(const class FPostConstructInitial
 
 	bHasCustomNavigableGeometry = EHasCustomNavigableGeometry::EvenIfNotCollidable;
 
+	BodyInstance.bUseAsyncScene = true;
 	BodyInstance.bEnableCollision_DEPRECATED = true;
 	static FName CollisionProfileName(TEXT("Destructible"));
 	SetCollisionProfileName(CollisionProfileName);
