Fixed potential quaternion normalization assert on long bone chains when using Mesh Space blend per weight.

diff --git a/Engine/Source/Runtime/Engine/Private/Animation/AnimationRuntime.cpp b/Engine/Source/Runtime/Engine/Private/Animation/AnimationRuntime.cpp
--- a/Engine/Source/Runtime/Engine/Private/Animation/AnimationRuntime.cpp
+++ b/Engine/Source/Runtime/Engine/Private/Animation/AnimationRuntime.cpp
@@ -911,7 +911,10 @@ void FAnimationRuntime::BlendMeshPosesPerBoneWeights(
 		OutPose.Bones[BoneIndex] = BlendAtom;
 		if (ParentIndex!=INDEX_NONE)
 		{
-			FQuat const LocalBlendQuat = BlendRotations[ParentIndex].Inverse() * BlendRotations[BoneIndex];
+			FQuat LocalBlendQuat = BlendRotations[ParentIndex].Inverse() * BlendRotations[BoneIndex];
+			
+			// local -> mesh -> local transformations can cause loss of precision for long bone chains, we have to normalize rotation there.
+			LocalBlendQuat.Normalize();
 			OutPose.Bones[BoneIndex].SetRotation(LocalBlendQuat);
 		}
 	}
