Improved comments

diff --git a/Engine/Source/Runtime/Engine/Private/DemoNetDriver.cpp b/Engine/Source/Runtime/Engine/Private/DemoNetDriver.cpp
--- a/Engine/Source/Runtime/Engine/Private/DemoNetDriver.cpp
+++ b/Engine/Source/Runtime/Engine/Private/DemoNetDriver.cpp
@@ -825,8 +825,9 @@ void UDemoNetConnection::HandleClientPlayer( APlayerController* PC, UNetConnecti
 
 bool UDemoNetConnection::ClientHasInitializedLevelFor(const UObject* TestObject)
 {
-	// there is no way to know when demo playback will load levels,
-	// so we assume it always is anytime after initial replication and rely on the client to load it in time to make the reference work
+	// We save all currently streamed levels into the demo stream so we can force the demo playback client
+	// to stay in sync with the recording server
+	// This may need to be tweaked or re-evaluated when we start recording demos on the client
 	return ( GetDriver()->DemoFrameNum > 2 || Super::ClientHasInitializedLevelFor( TestObject ) );
 }
 
