Copying over vert coloring data when constructing a BP from a mesh (and allowing for you to paint over that on a per instance basis).
TTP #344682

diff --git a/Engine/Source/Editor/UnrealEd/Private/Kismet2/Kismet2.cpp b/Engine/Source/Editor/UnrealEd/Private/Kismet2/Kismet2.cpp
--- a/Engine/Source/Editor/UnrealEd/Private/Kismet2/Kismet2.cpp
+++ b/Engine/Source/Editor/UnrealEd/Private/Kismet2/Kismet2.cpp
@@ -853,6 +853,14 @@ protected:
 				Property->CopyCompleteValue_InContainer(Destination, Source);
 			}
 		}
+
+		if (UStaticMeshComponent* TemplateMeshComponent = Cast<UStaticMeshComponent>(Destination))
+		{
+			UStaticMeshComponent* SourceMeshComponent = CastChecked<UStaticMeshComponent>(Source);
+			// vert coloring data lives in the mesh component's LODData member,
+			// which is transient, so it wouldn't be copied in the above loop
+			TemplateMeshComponent->CopyInstanceVertexColorsIfCompatible(SourceMeshComponent);
+		}
 	}
 
 	/** Creates a USCS Node for the passed in component */
diff --git a/Engine/Source/Runtime/Engine/Private/StaticMeshComponent.cpp b/Engine/Source/Runtime/Engine/Private/StaticMeshComponent.cpp
--- a/Engine/Source/Runtime/Engine/Private/StaticMeshComponent.cpp
+++ b/Engine/Source/Runtime/Engine/Private/StaticMeshComponent.cpp
@@ -58,8 +58,14 @@ public:
 				if (StaticMeshComponent->LODData.IsValidIndex(LODIndex))
 				{
 					FStaticMeshComponentLODInfo& LODInfo = StaticMeshComponent->LODData[LODIndex];
-
-					// Should not already have overriden vertex colors
+					// this component could have been constructed from a template
+					// that had its own vert color overrides; so before we apply
+					// the instance's color data, we need to clear the old
+					// vert colors (so we can properly call InitFromColorArray())
+					StaticMeshComponent->RemoveInstanceVertexColorsFromLOD(LODIndex);
+					// may not be null at the start (could have been initialized 
+					// from a  component template with vert coloring), but should
+					// be null at this point, after RemoveInstanceVertexColorsFromLOD()
 					if (LODInfo.OverrideVertexColors == NULL)
 					{
 						LODInfo.PaintedVertices = VertexColorLODData.PaintedVertices;
@@ -1013,11 +1019,20 @@ void UStaticMeshComponent::ImportCustomProperties(const TCHAR* SourceText, FFeed
 		LODInfo.ImportText(&SourceText);
 
 		// Populate the OverrideVertexColors buffer
-		check(!LODInfo.OverrideVertexColors);
-		const TCHAR* Result = FCString::Stristr(SourceText, TEXT("ColorVertexData"));
-		if (Result)
+		if (const TCHAR* VertColorStr = FCString::Stristr(SourceText, TEXT("ColorVertexData")))
 		{
-			SourceText = Result;
+			SourceText = VertColorStr;
+
+			// this component could have been constructed from a template that
+			// had its own vert color overrides; so before we apply the
+			// custom color data, we need to clear the old vert colors (so
+			// we can properly call ImportText())
+			RemoveInstanceVertexColorsFromLOD(LODIndex);
+
+			// may not be null at the start (could have been initialized 
+			// from a blueprint component template with vert coloring), but 
+			// should be null by this point, after RemoveInstanceVertexColorsFromLOD()
+			check(LODInfo.OverrideVertexColors == nullptr);
 
 			LODInfo.OverrideVertexColors = new FColorVertexBuffer;
 			LODInfo.OverrideVertexColors->ImportText(SourceText);
