Unshelved from pending changelist '2322927':

RIBBON BLOCKER FIX

Just blindly undoing the most suspect part of the cause CL. Not worked out why this fixes it yet.

diff --git a/Engine/Source/Runtime/Engine/Private/Particles/ParticleSystemRender.cpp b/Engine/Source/Runtime/Engine/Private/Particles/ParticleSystemRender.cpp
--- a/Engine/Source/Runtime/Engine/Private/Particles/ParticleSystemRender.cpp
+++ b/Engine/Source/Runtime/Engine/Private/Particles/ParticleSystemRender.cpp
@@ -671,43 +671,120 @@ void FDynamicSpriteEmitterDataBase::RenderDebug(const FParticleSystemSceneProxy*
 	}
 }
 
-void FDynamicSpriteEmitterDataBase::BuildViewFillData(
-	const FParticleSystemSceneProxy* Proxy, 
-	const FSceneView* InView, 
-	int32 InVertexCount, 
-	int32 InVertexSize, 
-	int32 InDynamicParameterVertexStride, 
-	FGlobalDynamicVertexBuffer::FAllocation& DynamicVertexAllocation,
-	FGlobalDynamicIndexBuffer::FAllocation& DynamicIndexAllocation,
-	FGlobalDynamicVertexBuffer::FAllocation* DynamicParameterAllocation,
-	FAsyncBufferFillData& Data) const
+/**
+ *	Set up an buffer for async filling
+ *
+ *	@param	Proxy					The primitive scene proxy for the emitter.
+ *	@param	InBufferIndex			Index of this buffer
+ *	@param	InView					View for this buffer
+ *	@param	InVertexCount			Count of verts for this buffer
+ *	@param	InVertexSize			Stride of these verts, only used for verification
+ *	@param	InDynamicParameterVertexStride	Stride of the dynamic parameter
+ */
+void FDynamicSpriteEmitterDataBase::BuildViewFillData(FParticleSystemSceneProxy* Proxy, int32 InBufferIndex,const FSceneView *InView,int32 InVertexCount,int32 InVertexSize,int32 InDynamicParameterVertexStride)
 {
+	bool bSuccess = true;
+
+	FAsyncBufferFillData Data;
+
 	Data.LocalToWorld = Proxy->GetLocalToWorld();
 	Data.WorldToLocal = Proxy->GetWorldToLocal();
 	Data.View = InView;
 	check(Data.VertexSize == 0 || Data.VertexSize == InVertexSize);
 
-	DynamicVertexAllocation = FGlobalDynamicVertexBuffer::Get().Allocate( InVertexCount * InVertexSize );
-	Data.VertexData = DynamicVertexAllocation.Buffer;
+	check(InBufferIndex == InstanceDataAllocations.Num())
+	FGlobalDynamicVertexBuffer::FAllocation* DynamicVertexAllocation = new(InstanceDataAllocations) FGlobalDynamicVertexBuffer::FAllocation();
+	*DynamicVertexAllocation = FGlobalDynamicVertexBuffer::Get().Allocate( InVertexCount * InVertexSize );
+	Data.VertexData = DynamicVertexAllocation->Buffer;
 	Data.VertexCount = InVertexCount;
 	Data.VertexSize = InVertexSize;
 
+	bSuccess &= DynamicVertexAllocation->IsValid();
+
 	int32 NumIndices, IndexStride;
 	GetIndexAllocInfo(NumIndices, IndexStride);
 	check((uint32)NumIndices <= 65535);
 	check(IndexStride > 0);
 
-	DynamicIndexAllocation = FGlobalDynamicIndexBuffer::Get().Allocate( NumIndices, IndexStride );
-	Data.IndexData = DynamicIndexAllocation.Buffer;
+	FGlobalDynamicIndexBuffer::FAllocation* DynamicIndexAllocation = new(IndexDataAllocations) FGlobalDynamicIndexBuffer::FAllocation();
+	*DynamicIndexAllocation = FGlobalDynamicIndexBuffer::Get().Allocate( NumIndices, IndexStride );
+	Data.IndexData = DynamicIndexAllocation->Buffer;
 	Data.IndexCount = NumIndices;
 
-	Data.DynamicParameterData = NULL;
+	bSuccess &= DynamicIndexAllocation->IsValid();
 
+	Data.DynamicParameterData = NULL;
 	if( bUsesDynamicParameter )
 	{
 		check( InDynamicParameterVertexStride > 0 );
+		check(InBufferIndex == DynamicParameterDataAllocations.Num())
+		FGlobalDynamicVertexBuffer::FAllocation* DynamicParameterAllocation = new(DynamicParameterDataAllocations) FGlobalDynamicVertexBuffer::FAllocation();
 		*DynamicParameterAllocation = FGlobalDynamicVertexBuffer::Get().Allocate( InVertexCount * InDynamicParameterVertexStride );
 		Data.DynamicParameterData = DynamicParameterAllocation->Buffer;
+
+		bSuccess &= DynamicParameterAllocation->IsValid();
+	}
+
+	//Create the data but only fill it if all the allocations succeeded.
+	if (InBufferIndex >= AsyncBufferFillTasks.Num())
+	{
+		new (AsyncBufferFillTasks) FAsyncBufferFillData();
+	}
+	check(InBufferIndex < AsyncBufferFillTasks.Num()); // please add the views in order
+
+	if( bSuccess )
+	{
+		AsyncBufferFillTasks[InBufferIndex] = Data;
+	}
+}
+
+/**
+ *	Set up all buffers for async filling
+ *
+ *	@param	Proxy							The primitive scene proxy for the emitter.
+ *	@param	ViewFamily						View family to process
+ *	@param	VisibilityMap					Visibility map for the sub-views
+ *	@param	bOnlyOneView					If true, then we don't need per-view buffers
+ *	@param	InVertexCount					Count of verts for this buffer
+ *	@param	InVertexSize					Stride of these verts, only used for verification
+ *  @param  InDynamicParameterVertexSize	Stride of the dynamic parameter
+ */
+void FDynamicSpriteEmitterDataBase::BuildViewFillDataAndSubmit(FParticleSystemSceneProxy* Proxy, const FSceneViewFamily* ViewFamily,const uint32 VisibilityMap,bool bOnlyOneView,int32 InVertexCount,int32 InVertexSize,int32 InDynamicParameterVertexSize)
+{
+	EnsureAsyncTaskComplete();
+	int32 NumUsedViews = 0;
+	InstanceDataAllocations.Empty();
+	IndexDataAllocations.Empty();
+	DynamicParameterDataAllocations.Empty();
+	for (int32 ViewIndex = 0; ViewIndex < ViewFamily->Views.Num(); ViewIndex++)
+	{
+		if (VisibilityMap & (1<<ViewIndex))
+		{
+			const FSceneView* View = ViewFamily->Views[ViewIndex];
+			BuildViewFillData(Proxy,NumUsedViews++,View,InVertexCount,InVertexSize, InDynamicParameterVertexSize);
+			if (bOnlyOneView)
+			{
+				break;
+			}
+		}
+	}
+	if (AsyncBufferFillTasks.Num() > NumUsedViews)
+	{
+		AsyncBufferFillTasks.RemoveAt(NumUsedViews,AsyncBufferFillTasks.Num() - NumUsedViews);
+	}
+	if (NumUsedViews)
+	{
+#if 0 // async particle fill is broken in UE4
+		if (!GIsEditor)
+		{
+			check(!AsyncTask.GetReference()); // this should not be still outstanding
+			AsyncTask = TGraphTask<FAsyncParticleFill>::CreateTask(NULL, ENamedThreads::RenderThread).ConstructAndDispatchWhenReady(this);
+		}
+		else
+#endif
+		{
+			DoBufferFill();
+		}
 	}
 }
 
@@ -2963,8 +3040,8 @@ void FDynamicBeam2EmitterData::GetDynamicMeshElementsEmitter(const FParticleSyst
 	FGlobalDynamicVertexBuffer::FAllocation DynamicParameterAllocation;
 	FAsyncBufferFillData Data;
 
-	BuildViewFillData(Proxy, View, Source.VertexCount, sizeof(FParticleBeamTrailVertex), 0, DynamicVertexAllocation, DynamicIndexAllocation, &DynamicParameterAllocation, Data);
-	DoBufferFill(Data);
+	//BuildViewFillData(Proxy, View, Source.VertexCount, sizeof(FParticleBeamTrailVertex), 0, DynamicVertexAllocation, DynamicIndexAllocation, &DynamicParameterAllocation, Data);
+	//DoBufferFill(Data);
 
 	if (Source.bUseLocalSpace == false)
 	{
@@ -3188,25 +3265,9 @@ void FDynamicBeam2EmitterData::PreRenderView(FParticleSystemSceneProxy* Proxy, c
 	// Only need to do this once per-view
 	if (LastFramePreRendered < FrameNumber)
 	{
-		InstanceDataAllocations.Empty();
-		IndexDataAllocations.Empty();
-		DynamicParameterDataAllocations.Empty();
-
-		for (int32 ViewIndex = 0; ViewIndex < ViewFamily->Views.Num(); ViewIndex++)
-		{
-			if (VisibilityMap & (1 << ViewIndex))
-			{
-				const FSceneView* View = ViewFamily->Views[ViewIndex];
+		bool bOnlyOneView = !GIsEditor && ((GEngine && GEngine->GameViewport && (GEngine->GameViewport->GetCurrentSplitscreenConfiguration() == ESplitScreenType::None) && !GEngine->IsStereoscopic3D()) ? true : false);
 
-				FGlobalDynamicVertexBuffer::FAllocation* DynamicVertexAllocation = new(InstanceDataAllocations) FGlobalDynamicVertexBuffer::FAllocation();
-				FGlobalDynamicIndexBuffer::FAllocation* DynamicIndexAllocation = new(IndexDataAllocations) FGlobalDynamicIndexBuffer::FAllocation();
-				FGlobalDynamicVertexBuffer::FAllocation* DynamicParameterAllocation = bUsesDynamicParameter ? new(DynamicParameterDataAllocations) FGlobalDynamicVertexBuffer::FAllocation() : NULL;
-				FAsyncBufferFillData* Data = new (AsyncBufferFillTasks) FAsyncBufferFillData();
-
-				BuildViewFillData(Proxy, View, Source.VertexCount, sizeof(FParticleBeamTrailVertex), 0, *DynamicVertexAllocation, *DynamicIndexAllocation, DynamicParameterAllocation, *Data);
-				DoBufferFill(*Data);
-			}
-		}
+		BuildViewFillDataAndSubmit(Proxy,ViewFamily,VisibilityMap,bOnlyOneView,Source.VertexCount,sizeof(FParticleBeamTrailVertex),0);
 
 		// Set the frame tracker
 		LastFramePreRendered = FrameNumber;
@@ -5846,8 +5907,8 @@ void FDynamicTrailsEmitterData::GetDynamicMeshElementsEmitter(const FParticleSys
 	const int32 VertexStride = GetDynamicVertexStride(ViewFamily.GetFeatureLevel());
 	const int32 DynamicParameterVertexStride = bUsesDynamicParameter ? GetDynamicParameterVertexStride() : 0;
 
-	BuildViewFillData(Proxy, View, SourcePointer->VertexCount, VertexStride, DynamicParameterVertexStride, DynamicVertexAllocation, DynamicIndexAllocation, &DynamicParameterAllocation, Data);
-	DoBufferFill(Data);
+	//BuildViewFillData(Proxy, View, SourcePointer->VertexCount, VertexStride, DynamicParameterVertexStride, DynamicVertexAllocation, DynamicIndexAllocation, &DynamicParameterAllocation, Data);
+	//DoBufferFill(Data);
 
 	if (SourcePointer->bUseLocalSpace == false)
 	{
@@ -6099,24 +6160,16 @@ void FDynamicTrailsEmitterData::PreRenderView(FParticleSystemSceneProxy* Proxy,
 			DynamicParameterVertexStride = GetDynamicParameterVertexStride();
 		}
 
-		bool bOnlyOneView = false;
+		bool bOnlyOneView = ShouldUsePrerenderView() || ((GEngine && GEngine->GameViewport && (GEngine->GameViewport->GetCurrentSplitscreenConfiguration() == ESplitScreenType::None)) ? true : false);
 
-		for (int32 ViewIndex = 0; ViewIndex < ViewFamily->Views.Num(); ViewIndex++)
+		// Render both views when in stereo
+		if (GEngine->HMDDevice.IsValid() && GEngine->IsStereoscopic3D())
 		{
-			if (VisibilityMap & (1 << ViewIndex))
-			{
-				const FSceneView* View = ViewFamily->Views[ViewIndex];
-
-				FGlobalDynamicVertexBuffer::FAllocation* DynamicVertexAllocation = new(InstanceDataAllocations) FGlobalDynamicVertexBuffer::FAllocation();
-				FGlobalDynamicIndexBuffer::FAllocation* DynamicIndexAllocation = new(IndexDataAllocations) FGlobalDynamicIndexBuffer::FAllocation();
-				FGlobalDynamicVertexBuffer::FAllocation* DynamicParameterAllocation = bUsesDynamicParameter ? new(DynamicParameterDataAllocations) FGlobalDynamicVertexBuffer::FAllocation() : NULL;
-				FAsyncBufferFillData* Data = new (AsyncBufferFillTasks) FAsyncBufferFillData();
-
-				BuildViewFillData(Proxy, View, SourcePointer->VertexCount, VertexStride, DynamicParameterVertexStride, *DynamicVertexAllocation, *DynamicIndexAllocation, DynamicParameterAllocation, *Data);
-				DoBufferFill(*Data);
-			}
+			bOnlyOneView = false;
 		}
 
+		BuildViewFillDataAndSubmit(Proxy,ViewFamily,VisibilityMap,bOnlyOneView,SourcePointer->VertexCount,VertexStride, DynamicParameterVertexStride);
+
 		// Set the frame tracker
 		LastFramePreRendered = FrameNumber;
 	}
diff --git a/Engine/Source/Runtime/Engine/Public/ParticleHelper.h b/Engine/Source/Runtime/Engine/Public/ParticleHelper.h
--- a/Engine/Source/Runtime/Engine/Public/ParticleHelper.h
+++ b/Engine/Source/Runtime/Engine/Public/ParticleHelper.h
@@ -1077,6 +1077,42 @@ struct FAsyncBufferFillData
 	}
 };
 
+/*-----------------------------------------------------------------------------
+	Async Fill Task, simple wrapper to forward the request to a FDynamicSpriteEmitterDataBase
+-----------------------------------------------------------------------------*/
+
+struct FAsyncParticleFill
+{
+	/** Emitter to forward to   */
+	struct FDynamicSpriteEmitterDataBase* Parent;
+
+	/** Constructor, just sets up the parent pointer  
+	  * @param InParent emitter to forward the eventual async call to
+	*/
+	FAsyncParticleFill(struct FDynamicSpriteEmitterDataBase* InParent)
+		: Parent(InParent)
+	{
+	}
+
+	/** Work function, just forwards the request to the parent  */
+	void DoTask(ENamedThreads::Type CurrentThread, const FGraphEventRef& MyCompletionGraphEvent);
+
+	FORCEINLINE TStatId GetStatId() const
+	{
+		return GET_STATID( STAT_ParticleAsyncTime );
+	}
+
+	static ENamedThreads::Type GetDesiredThread()
+	{
+		return ENamedThreads::AnyThread;
+	}
+
+	static ESubsequentsMode::Type GetSubsequentsMode() 
+	{ 
+		return ESubsequentsMode::TrackSubsequents; 
+	}
+};
+
 // TAsyncBufferFillTasks - handy typedef for an inline array of buffer fill tasks
 typedef TArray<FAsyncBufferFillData, TInlineAllocator<2> > TAsyncBufferFillTasks;
 
@@ -1448,6 +1484,7 @@ struct FDynamicSpriteEmitterDataBase : public FDynamicEmitterDataBase
 {
 	FDynamicSpriteEmitterDataBase(const UParticleModuleRequired* RequiredModule) : 
 		FDynamicEmitterDataBase(RequiredModule),
+		AsyncTask(NULL),
 		bUsesDynamicParameter( false )
 	{
 		MaterialResource[0] = NULL;
@@ -1456,6 +1493,7 @@ struct FDynamicSpriteEmitterDataBase : public FDynamicEmitterDataBase
 
 	virtual ~FDynamicSpriteEmitterDataBase()
 	{
+		EnsureAsyncTaskComplete();
 	}
 
 	/**
@@ -1532,31 +1570,62 @@ struct FDynamicSpriteEmitterDataBase : public FDynamicEmitterDataBase
 	 */
 	virtual void RenderDebug(const FParticleSystemSceneProxy* Proxy, FPrimitiveDrawInterface* PDI, const FSceneView* View, bool bCrosses) const;
 
-	virtual void DoBufferFill(FAsyncBufferFillData& Me) const
+	/**
+	 *	Fill index and vertex buffers. Often called from a different thread
+	 *
+	 */
+	void DoBufferFill()
 	{
-		// Must be overridden if called
-		check(0);
+		for (int32 TaskIndex = 0; TaskIndex < AsyncBufferFillTasks.Num(); TaskIndex++) 
+		{
+			DoBufferFill(AsyncBufferFillTasks[TaskIndex]);
+		}
+	}
+	/**
+	 *	Fill index and vertex buffers. Often called from a different thread
+	 *
+	 *	@param	Me			buffer pair to compute
+	 */
+	virtual void DoBufferFill(FAsyncBufferFillData& Me)const
+	{
+		// this must be overridden, but in some cases a destructor call will leave this a no-op
+		// because the vtable has been reset to the base class 
+		// checkf(0, TEXT("DoBufferFill MUST be overridden"));
 	}
-
 	/**
 	 *	Set up an buffer for async filling
 	 *
 	 *	@param	Proxy					The primitive scene proxy for the emitter.
+	 *	@param	InBufferIndex			Index of this buffer
 	 *	@param	InView					View for this buffer
 	 *	@param	InVertexCount			Count of verts for this buffer
 	 *	@param	InVertexSize			Stride of these verts, only used for verification
 	 *	@param	InDynamicParameterVertexStride	Stride of the dynamic parameter
 	 */
-	void BuildViewFillData(
-		const FParticleSystemSceneProxy* Proxy, 
-		const FSceneView *InView, 
-		int32 InVertexCount, 
-		int32 InVertexSize, 
-		int32 InDynamicParameterVertexSize, 
-		FGlobalDynamicVertexBuffer::FAllocation& DynamicVertexAllocation,
-		FGlobalDynamicIndexBuffer::FAllocation& DynamicIndexAllocation,
-		FGlobalDynamicVertexBuffer::FAllocation* DynamicParameterAllocation,
-		FAsyncBufferFillData& Data) const;
+	void BuildViewFillData(FParticleSystemSceneProxy* Proxy, int32 InBufferIndex,const FSceneView *InView,int32 InVertexCount,int32 InVertexSize,int32 InDynamicParameterVertexSize);
+
+	/**
+	 *	Set up all buffers for async filling
+	 *
+	 *	@param	Proxy							The primitive scene proxy for the emitter.
+	 *	@param	ViewFamily						View family to process
+	 *	@param	VisibilityMap					Visibility map for the sub-views
+	 *	@param	bOnlyOneView					If true, then we don't need per-view buffers
+	 *	@param	InVertexCount					Count of verts for this buffer
+	 *	@param	InVertexSize					Stride of these verts, only used for verification
+	 *	@param	InDynamicParameterVertexStride	Stride of the dynamic parameter
+	 */
+	void BuildViewFillDataAndSubmit(FParticleSystemSceneProxy* Proxy, const FSceneViewFamily* ViewFamily,const uint32 VisibilityMap,bool bOnlyOneView,int32 InVertexCount,int32 InVertexSize, int32 InDynamicParameterVertexSize);
+
+	void EnsureAsyncTaskComplete()
+	{
+		if (AsyncTask.GetReference())
+		{
+			SCOPE_CYCLE_COUNTER(STAT_ParticleAsyncWaitTime);
+			FTaskGraphInterface::Get().WaitUntilTaskCompletes(AsyncTask, ENamedThreads::RenderThread_Local);
+			AsyncTask = NULL;
+		}
+	}
 
 	/**
 	 *	Called to verify that a buffer is ready to use, blocks to wait and can sometimes execute the buffer fill on the current thread
@@ -1582,6 +1651,9 @@ struct FDynamicSpriteEmitterDataBase : public FDynamicEmitterDataBase
 	/** Async task is queued for execution */
 	bool									bAsyncTaskOutstanding;
 
+	/** Async task that is queued in the hi priority pool */
+	FGraphEventRef							AsyncTask;
+
 	/** Array of buffers for filling by async task */
 	TAsyncBufferFillTasks					AsyncBufferFillTasks;
 
