Fix destruction chunks not updating correctly when scaled

TTP#349008

diff --git a/Engine/Source/Runtime/Engine/Private/DestructibleComponent.cpp b/Engine/Source/Runtime/Engine/Private/DestructibleComponent.cpp
--- a/Engine/Source/Runtime/Engine/Private/DestructibleComponent.cpp
+++ b/Engine/Source/Runtime/Engine/Private/DestructibleComponent.cpp
@@ -896,12 +896,7 @@ void UDestructibleComponent::UpdateDestructibleChunkTM(TArray<const PxRigidActor
 	{
 		UDestructibleComponent* DestructibleComponent = It.Key();
 		TArray<FUpdateChunksInfo>& UpdateInfos = It.Value();
-		DestructibleComponent->SetChunksWorldTM(UpdateInfos);
-		/*for (FUpdateChunksInfo& UpdateInfo : UpdateInfos)
-		{
-			DestructibleComponent->SetChunkWorldRT(UpdateInfo.ChunkIndex, UpdateInfo.WorldTM.GetRotation(), UpdateInfo.WorldTM.GetLocation());
-		}*/
-		
+		DestructibleComponent->SetChunksWorldTM(UpdateInfos);	
 	}
 
 }
@@ -910,14 +905,19 @@ void UDestructibleComponent::UpdateDestructibleChunkTM(TArray<const PxRigidActor
 
 void UDestructibleComponent::SetChunksWorldTM(const TArray<FUpdateChunksInfo>& UpdateInfos)
 {
-	FTransform InvWorldTM = ComponentToWorld.Inverse();
-	InvWorldTM.SetScale3D(FVector(1.f));
+	const FQuat InvRotation = ComponentToWorld.GetRotation().Inverse();
 
 	for (const FUpdateChunksInfo& UpdateInfo : UpdateInfos)
 	{
 		// Bone 0 is a dummy root bone
 		const int32 BoneIndex = ChunkIdxToBoneIdx(UpdateInfo.ChunkIndex);
-		SpaceBases[BoneIndex] = FTransform(UpdateInfo.WorldTM * InvWorldTM);
+		const FVector WorldTranslation	= UpdateInfo.WorldTM.GetLocation();
+		const FQuat WorldRotation		= UpdateInfo.WorldTM.GetRotation();
+
+		const FQuat BoneRotation = InvRotation*WorldRotation;
+		const FVector BoneTranslation = InvRotation.RotateVector(WorldTranslation - ComponentToWorld.GetTranslation()) / ComponentToWorld.GetScale3D();
+
+		SpaceBases[BoneIndex] = FTransform(BoneRotation, BoneTranslation);
 	}
 
 	// Mark the transform as dirty, so the bounds are updated and sent to the render thread
