UE4 - fix shader recompilation

diff --git a/Engine/Source/Runtime/Engine/Private/BatchedElements.cpp b/Engine/Source/Runtime/Engine/Private/BatchedElements.cpp
--- a/Engine/Source/Runtime/Engine/Private/BatchedElements.cpp
+++ b/Engine/Source/Runtime/Engine/Private/BatchedElements.cpp
@@ -571,16 +571,16 @@ void FBatchedElements::PrepareShaders(
 					FSimpleElementGammaBasePS* BasePixelShader;
 					FGlobalBoundShaderState* BoundShaderState;
 
+					TShaderMapRef<FSimpleElementGammaPS_SRGB> PixelShader_SRGB(GetGlobalShaderMap(FeatureLevel));
+					TShaderMapRef<FSimpleElementGammaPS_Linear> PixelShader_Linear(GetGlobalShaderMap(FeatureLevel));
 					if (Texture->bSRGB)
 					{
-						static TShaderMapRef<FSimpleElementGammaPS_SRGB> PixelShader(GetGlobalShaderMap(FeatureLevel));
-						BasePixelShader = *PixelShader;
+						BasePixelShader = *PixelShader_SRGB;
 						BoundShaderState = &RegularSRGBBoundShaderState;
 					}
 					else
 					{
-						static TShaderMapRef<FSimpleElementGammaPS_Linear> PixelShader(GetGlobalShaderMap(FeatureLevel));
-						BasePixelShader = *PixelShader;
+						BasePixelShader = *PixelShader_Linear;
 						BoundShaderState = &RegularLinearBoundShaderState;
 					}
 
