Don't allow spawning during world teardown

diff --git a/Engine/Source/Editor/UnrealEd/Private/PlayLevel.cpp b/Engine/Source/Editor/UnrealEd/Private/PlayLevel.cpp
--- a/Engine/Source/Editor/UnrealEd/Private/PlayLevel.cpp
+++ b/Engine/Source/Editor/UnrealEd/Private/PlayLevel.cpp
@@ -293,6 +293,7 @@ void UEditorEngine::TeardownPlaySession(FWorldContext &PieWorldContext)
 {
 	check(PieWorldContext.WorldType == EWorldType::PIE);
 	PlayWorld = PieWorldContext.World();
+	PlayWorld->bIsTearingDown = true;
 
 	if (!PieWorldContext.RunAsDedicated)
 	{
diff --git a/Engine/Source/Runtime/Engine/Classes/Engine/World.h b/Engine/Source/Runtime/Engine/Classes/Engine/World.h
--- a/Engine/Source/Runtime/Engine/Classes/Engine/World.h
+++ b/Engine/Source/Runtime/Engine/Classes/Engine/World.h
@@ -913,6 +913,13 @@ public:
 	TMap< FGuid, class ULandscapeInfo* > LandscapeInfoMap;
 #endif // WITH_EDITORONLY_DATA
 
+	/**
+	 * This is a int on the level which is set when a light that needs to have lighting rebuilt
+	 * is moved.  This is then checked in CheckMap for errors to let you know that this level should
+	 * have lighting rebuilt.
+	 **/
+	uint32 NumLightingUnbuiltObjects;
+
 	/** frame rate is below DesiredFrameRate, so drop high detail actors */
 	uint32 bDropDetail:1;
 
@@ -943,12 +950,8 @@ public:
 	/** Starting gameplay. */
 	uint32 bStartup:1;
 
-	/**
-	 * This is a int on the level which is set when a light that needs to have lighting rebuilt
-	 * is moved.  This is then checked in CheckMap for errors to let you know that this level should
-	 * have lighting rebuilt.
-	 **/
-	uint32 NumLightingUnbuiltObjects;
+	/** Is the world being torn down */
+	uint32 bIsTearingDown:1;
 
 	/**
 	 * This is a bool that indicates that one or more blueprints in the level (blueprint instances, level script, etc)
diff --git a/Engine/Source/Runtime/Engine/Private/GameEngine.cpp b/Engine/Source/Runtime/Engine/Private/GameEngine.cpp
--- a/Engine/Source/Runtime/Engine/Private/GameEngine.cpp
+++ b/Engine/Source/Runtime/Engine/Private/GameEngine.cpp
@@ -620,6 +620,7 @@ bool UGameEngine::HandleExitCommand( const TCHAR* Cmd, FOutputDevice& Ar )
 	for (int32 WorldIndex = 0; WorldIndex < WorldList.Num(); ++WorldIndex)
 	{
 		UWorld* const World = WorldList[WorldIndex].World();
+		World->bIsTearingDown = true;
 
 		// Cancel any pending connection to a server
 		CancelPending(WorldList[WorldIndex]);
diff --git a/Engine/Source/Runtime/Engine/Private/LevelActor.cpp b/Engine/Source/Runtime/Engine/Private/LevelActor.cpp
--- a/Engine/Source/Runtime/Engine/Private/LevelActor.cpp
+++ b/Engine/Source/Runtime/Engine/Private/LevelActor.cpp
@@ -263,6 +263,11 @@ AActor* UWorld::SpawnActor( UClass* Class, FVector const* Location, FRotator con
 		UE_LOG(LogSpawn, Warning, TEXT("SpawnActor failed because we are running a ConstructionScript (%s)"), *Class->GetName() );
 		return NULL;
 	}
+	else if (bIsTearingDown)
+	{
+		UE_LOG(LogSpawn, Warning, TEXT("SpawnActor failed because we are in the process of tearing down the world"));
+		return NULL;
+	}
 
 	AActor* Template = SpawnParameters.Template;
 	// Use class's default actor as a template.
