Rendering/UMG - TC_EditorIcon now appears as TC_UserInterface2D.  Display name change only.

diff --git a/Engine/Source/Runtime/Engine/Classes/Engine/Texture.h b/Engine/Source/Runtime/Engine/Classes/Engine/Texture.h
--- a/Engine/Source/Runtime/Engine/Classes/Engine/Texture.h
+++ b/Engine/Source/Runtime/Engine/Classes/Engine/Texture.h
@@ -19,7 +19,7 @@ enum TextureCompressionSettings
 	TC_Displacementmap,
 	TC_VectorDisplacementmap,
 	TC_HDR,
-	TC_EditorIcon,
+	TC_EditorIcon UMETA(DisplayName="TC_UserInterface2D"),
 	TC_Alpha,
 	TC_DistanceFieldFont,
 	TC_MAX,
