#ue4 - Fixes to projectile movement uses of friction, and improvements to behavior when sliding along a surface, to remove variance under different timesteps and to apply consistent friction.

- Removed hacky bGravityEnabled flag from various velocity and acceleration functions.
- Deprecated old CalculateVelocity() method, renamed to ComputeVelocity() to align with the other functions in the class.
- Sliding movement now considers the acceleration (in any direction, not just gravity) and handles it in a new function: HandleSliding().
   - Friction is applied correctly when movement is parallel to and impacting a surface ("sliding").

diff --git a/Engine/Source/Runtime/Engine/Classes/GameFramework/ProjectileMovementComponent.h b/Engine/Source/Runtime/Engine/Classes/GameFramework/ProjectileMovementComponent.h
--- a/Engine/Source/Runtime/Engine/Classes/GameFramework/ProjectileMovementComponent.h
+++ b/Engine/Source/Runtime/Engine/Classes/GameFramework/ProjectileMovementComponent.h
@@ -74,14 +74,16 @@ class ENGINE_API UProjectileMovementComponent : public UMovementComponent
 	 * Percentage of velocity maintained after the bounce in the direction of the normal of impact (coefficient of restitution).
 	 * 1.0 = no velocity lost, 0.0 = no bounce. Ignored if bShouldBounce is false.
 	 */
-	UPROPERTY(EditAnywhere, BlueprintReadWrite, Category=ProjectileBounces)
+	UPROPERTY(EditAnywhere, BlueprintReadWrite, Category=ProjectileBounces, meta=(ClampMin="0", UIMin="0"))
 	float Bounciness;
 
 	/**
-	 * Percentage of velocity maintained after the bounce in the direction tangent to the normal of impact (coefficient of friction).
-	 * 0.0 = no velocity lost, 1.0 = no sliding. Ignored if bShouldBounce is false.
+	 * Coefficient of friction, affecting the resistance to sliding along a surface.
+	 * Normal range is [0,1] : 0.0 = no friction, 1.0+ = very high friction.
+	 * Also affects the percentage of velocity maintained after the bounce in the direction tangent to the normal of impact.
+	 * Ignored if bShouldBounce is false.
 	 */
-	UPROPERTY(EditAnywhere, BlueprintReadWrite, Category=ProjectileBounces)
+	UPROPERTY(EditAnywhere, BlueprintReadWrite, Category=ProjectileBounces, meta=(ClampMin="0", UIMin="0"))
 	float Friction;
 
 	/**
@@ -135,8 +137,19 @@ class ENGINE_API UProjectileMovementComponent : public UMovementComponent
 
 	bool ShouldApplyGravity() const { return ProjectileGravityScale != 0.f; }
 
-	/** @return the velocity after DeltaTime */
-	virtual FVector CalculateVelocity(FVector OldVelocity, float DeltaTime, bool bGravityEnabled) const;
+	/** DEPRECATED: renamed to ComputeVelocity(). */
+	DEPRECATED(4.6, "CalculateVelocity() is deprecated, use ComputeVelocity() instead.")
+	virtual FVector CalculateVelocity(FVector OldVelocity, float DeltaTime, bool bGravityEnabled_UNUSED) const;
+
+	/**
+	 * Given an initial velocity and a time step, compute a new velocity.
+	 * Default implementation applies the result of ComputeAcceleration() to velocity.
+	 * 
+	 * @param  InitialVelocity Initial velocity.
+	 * @param  DeltaTime Time step of the update.
+	 * @return Velocity after DeltaTime time step.
+	 */
+	virtual FVector ComputeVelocity(FVector InitialVelocity, float DeltaTime) const;
 
 	/** Clears the reference to UpdatedComponent, fires stop event, and stops ticking. */
 	UFUNCTION(BlueprintCallable, Category="Game|Components|ProjectileMovement")
@@ -192,6 +205,16 @@ protected:
 	/** Applies bounce logic (if enabled) to affect velocity upon impact, or stops the projectile if bounces are not enabled (or velocity is below threshold). Fires applicable events. */
 	virtual void HandleImpact(const FHitResult& Hit, float TimeSlice=0.f, const FVector& MoveDelta = FVector::ZeroVector) override;
 
+	/**
+	 * Handle case where projectile is sliding along a surface.
+	 * Velocity will be parallel to the impact surface upon entry to this method.
+	 * 
+	 * @param  InitialHit Hit result of impact causing slide. May be modified by this function to reflect any subsequent movement.
+	 * @param  SubTickTimeRemaining Time remaining in the tick. This function may update this time with any reduction to the simulation time requested.
+	 * @return True if simulation of the projectile should continue, false otherwise.
+	 */
+	virtual bool HandleSliding(FHitResult& Hit, float& SubTickTimeRemaining);
+
 	/** Computes result of a bounce and returns the new velocity. */
 	virtual FVector ComputeBounceResult(const FHitResult& Hit, float TimeSlice, const FVector& MoveDelta);
 
@@ -200,13 +223,13 @@ protected:
 	FVector LimitVelocity(FVector NewVelocity) const;
 
 	/** Compute the distance we should move in the given time, at a given a velocity. */
-	virtual FVector ComputeMoveDelta(const FVector& InVelocity, float DeltaTime, bool bGravityEnabled = true) const;
+	virtual FVector ComputeMoveDelta(const FVector& InVelocity, float DeltaTime) const;
 
 	/** Compute the acceleration that will be applied */
-	virtual FVector ComputeAcceleration(const FVector& InVelocity, float DeltaTime, bool bGravityEnabled) const;
+	virtual FVector ComputeAcceleration(const FVector& InVelocity, float DeltaTime) const;
 
 	/** Allow the projectile to track towards its homing target. */
-	virtual FVector ComputeHomingAcceleration(const FVector& InVelocity, float DeltaTime, bool bGravityEnabled) const;
+	virtual FVector ComputeHomingAcceleration(const FVector& InVelocity, float DeltaTime) const;
 
 	/** Compute gravity effect given current physics volume, projectile gravity scale, etc. */
 	float GetEffectiveGravityZ() const;
diff --git a/Engine/Source/Runtime/Engine/Private/ProjectileMovementComponent.cpp b/Engine/Source/Runtime/Engine/Private/ProjectileMovementComponent.cpp
--- a/Engine/Source/Runtime/Engine/Private/ProjectileMovementComponent.cpp
+++ b/Engine/Source/Runtime/Engine/Private/ProjectileMovementComponent.cpp
@@ -105,7 +105,7 @@ void UProjectileMovementComponent::TickComponent(float DeltaTime, enum ELevelTic
 	int32 Iterations = 0;
 	FHitResult Hit(1.f);
 	FVector OldHitNormal(0.f,0.f,1.f);
-	bool bSliding = false;
+	float PreviousHitTime = 1.f;
 
 	while( RemainingTime >= MIN_TICK_TIME && (Iterations < MaxSimulationIterations) && !ActorOwner->IsPendingKill() && UpdatedComponent )
 	{
@@ -117,7 +117,7 @@ void UProjectileMovementComponent::TickComponent(float DeltaTime, enum ELevelTic
 
 		Hit.Time = 1.f;
 		const FVector OldVelocity = Velocity;
-		FVector MoveDelta = ComputeMoveDelta(OldVelocity, TimeTick, !bSliding);
+		FVector MoveDelta = ComputeMoveDelta(OldVelocity, TimeTick);
 
 		const FRotator NewRotation = bRotationFollowsVelocity ? OldVelocity.Rotation() : ActorOwner->GetActorRotation();
 
@@ -141,14 +141,14 @@ void UProjectileMovementComponent::TickComponent(float DeltaTime, enum ELevelTic
 		}
 
 		// Handle hit result after movement
-		if (Hit.Time == 1.f)
+		if( !Hit.bBlockingHit )
 		{
-			bSliding = false;
+			PreviousHitTime = 1.f;
 
 			// Only calculate new velocity if events didn't change it during the movement update.
 			if (Velocity == OldVelocity)
 			{
-				Velocity = CalculateVelocity(Velocity, TimeTick, !bSliding);
+				Velocity = ComputeVelocity(Velocity, TimeTick);				
 			}
 		}
 		else
@@ -157,7 +157,7 @@ void UProjectileMovementComponent::TickComponent(float DeltaTime, enum ELevelTic
 			if (Velocity == OldVelocity)
 			{
 				// re-calculate end velocity for partial time
-				Velocity = CalculateVelocity(OldVelocity, TimeTick*Hit.Time, !bSliding);
+				Velocity = (Hit.Time > KINDA_SMALL_NUMBER) ? ComputeVelocity(OldVelocity, TimeTick * Hit.Time) : OldVelocity;
 			}
 
 			if (HandleHitWall(Hit, TimeTick, MoveDelta))
@@ -165,21 +165,25 @@ void UProjectileMovementComponent::TickComponent(float DeltaTime, enum ELevelTic
 				break;
 			}
 
-			if (NumBounces < 2)
-			{
-				RemainingTime += TimeTick * (1.f - Hit.Time);
-			}
 			NumBounces++;
-
-			// if velocity still into wall (after HitWall() had a chance to adjust), slide along wall
-			if ( (Velocity | Hit.Normal) <= 0.f || (Hit.Time == 0.f && NumBounces > 1) )
+			float SubTickTimeRemaining = TimeTick * (1.f - Hit.Time);
+
+			// Multiple hits within very short time period?
+			const bool bMultiHit = (PreviousHitTime <= KINDA_SMALL_NUMBER && Hit.Time <= KINDA_SMALL_NUMBER);
+			
+			// if velocity still into wall (after HandleHitWall() had a chance to adjust), slide along wall
+			const float DotTolerance = 0.01f;
+			const bool bSliding = (bMultiHit && FVector::Coincident(OldHitNormal, Hit.Normal, DotTolerance)) ||
+								  ((Velocity.SafeNormal() | Hit.Normal) <= DotTolerance);
+			
+			if (bSliding)
 			{
-				if ( (NumBounces > 1) && ((OldHitNormal | Hit.Normal) <= 0.f) )
+				if (bMultiHit && (OldHitNormal | Hit.Normal) <= 0.f)
 				{
 					//90 degree or less corner, so use cross product for direction
 					FVector NewDir = (Hit.Normal ^ OldHitNormal);
 					NewDir = NewDir.SafeNormal();
-					Velocity = (Velocity | NewDir) * NewDir;
+					Velocity = Velocity.ProjectOnToNormal(NewDir);
 					if ((OldVelocity | Velocity) < 0.f)
 					{
 						Velocity *= -1.f;
@@ -188,18 +192,34 @@ void UProjectileMovementComponent::TickComponent(float DeltaTime, enum ELevelTic
 				else 
 				{
 					//adjust to move along new wall
-					Velocity = Velocity - Hit.Normal * (Velocity | Hit.Normal);
-					if ( (Velocity | OldVelocity) <= 0.f )
+					Velocity = FVector::VectorPlaneProject(Velocity, Hit.Normal);
+				}
+
+				// Check min velocity.
+				if (Velocity.SizeSquared() < FMath::Square(BounceVelocityStopSimulatingThreshold))
+				{
+					StopSimulating(Hit);
+					break;
+				}
+				
+				// Velocity is now parallel to the impact surface.
+				if (SubTickTimeRemaining > KINDA_SMALL_NUMBER)
+				{
+					if (!HandleSliding(Hit, SubTickTimeRemaining))
 					{
-						if (Velocity.SizeSquared() < FMath::Square(BounceVelocityStopSimulatingThreshold))
-						{
-							StopSimulating(Hit);
-							break;
-						}
+						break;
 					}
 				}
-				OldHitNormal = Hit.Normal;
-				bSliding = true;
+			}
+
+			PreviousHitTime = Hit.Time;
+			OldHitNormal = Hit.Normal;
+			
+			// A few initial bounces should add more time and iterations to complete most of the simulation.
+			if (NumBounces <= 2 && SubTickTimeRemaining >= MIN_TICK_TIME)
+			{
+				RemainingTime += SubTickTimeRemaining;
+				Iterations--;
 			}
 		}
 	}
@@ -208,6 +228,57 @@ void UProjectileMovementComponent::TickComponent(float DeltaTime, enum ELevelTic
 }
 
 
+bool UProjectileMovementComponent::HandleSliding(FHitResult& Hit, float& SubTickTimeRemaining)
+{
+	FHitResult InitialHit(Hit);
+	const FVector OldHitNormal = Hit.Normal;
+
+	// Velocity is now parallel to the impact surface.
+	// Perform the move now, before adding gravity/accel again, so we don't just keep hitting the surface.
+	SafeMoveUpdatedComponent(Velocity * SubTickTimeRemaining, UpdatedComponent->GetComponentRotation(), true, Hit);
+
+	// A second hit can deflect the velocity (through the normal bounce code), for the next iteration.
+	if (Hit.bBlockingHit)
+	{
+		if (HandleHitWall(Hit, SubTickTimeRemaining, Velocity * SubTickTimeRemaining))
+		{
+			return false;
+		}
+	}
+	else
+	{
+		// Find velocity after elapsed time
+		const FVector PostTickVelocity = ComputeVelocity(Velocity, SubTickTimeRemaining);
+
+		// If pointing back into surface, apply friction and acceleration.
+		const FVector Force = (PostTickVelocity - Velocity);
+		const float ForceDotN = (Force | OldHitNormal);
+		if (ForceDotN < 0.f)
+		{
+			const FVector ProjectedForce = FVector::VectorPlaneProject(Force, OldHitNormal);
+			const FVector NewVelocity = Velocity + ProjectedForce;
+
+			const FVector FrictionForce = -NewVelocity.SafeNormal() * FMath::Min(-ForceDotN * Friction, NewVelocity.Size());
+			Velocity = NewVelocity + FrictionForce;
+		}
+		else
+		{
+			Velocity = PostTickVelocity;
+		}
+
+		// Check min velocity
+		if (Velocity.SizeSquared() < FMath::Square(BounceVelocityStopSimulatingThreshold))
+		{
+			StopSimulating(InitialHit);
+			return false;
+		}
+	}
+
+	SubTickTimeRemaining = SubTickTimeRemaining * (1.f - Hit.Time);
+	return true;
+}
+
+
 void UProjectileMovementComponent::SetVelocityInLocalSpace(FVector NewVelocity)
 {
 	if (UpdatedComponent)
@@ -217,11 +288,16 @@ void UProjectileMovementComponent::SetVelocityInLocalSpace(FVector NewVelocity)
 }
 
 
-FVector UProjectileMovementComponent::CalculateVelocity(FVector OldVelocity, float DeltaTime, bool bGravityEnabled) const
+FVector UProjectileMovementComponent::CalculateVelocity(FVector OldVelocity, float DeltaTime, bool bGravityEnabled_UNUSED) const
+{
+	return ComputeVelocity(OldVelocity, DeltaTime);
+}
+
+FVector UProjectileMovementComponent::ComputeVelocity(FVector InitialVelocity, float DeltaTime) const
 {
 	// v = v0 + a*t
-	const FVector Acceleration = ComputeAcceleration(OldVelocity, DeltaTime, bGravityEnabled);
-	FVector NewVelocity = OldVelocity + (Acceleration * DeltaTime);
+	const FVector Acceleration = ComputeAcceleration(InitialVelocity, DeltaTime);
+	FVector NewVelocity = InitialVelocity + (Acceleration * DeltaTime);
 
 	return LimitVelocity(NewVelocity);
 }
@@ -238,40 +314,37 @@ FVector UProjectileMovementComponent::LimitVelocity(FVector NewVelocity) const
 	return NewVelocity;
 }
 
-FVector UProjectileMovementComponent::ComputeMoveDelta(const FVector& InVelocity, float DeltaTime, bool bGravityEnabled) const
+FVector UProjectileMovementComponent::ComputeMoveDelta(const FVector& InVelocity, float DeltaTime) const
 {
 	// Velocity Verlet integration (http://en.wikipedia.org/wiki/Verlet_integration#Velocity_Verlet)
 	// The addition of p0 is done outside this method, we are just computing the delta.
 	// p = p0 + v0*t + 1/2*a*t^2
 
-	// We use CalculateVelocity() here to infer the acceleration, to make it easier to apply custom velocities.
+	// We use ComputeVelocity() here to infer the acceleration, to make it easier to apply custom velocities.
 	// p = p0 + v0*t + 1/2*((v1-v0)/t)*t^2
 	// p = p0 + v0*t + 1/2*((v1-v0))*t
 
-	const FVector NewVelocity = CalculateVelocity(InVelocity, DeltaTime, bGravityEnabled);
+	const FVector NewVelocity = ComputeVelocity(InVelocity, DeltaTime);
 	const FVector Delta = (InVelocity * DeltaTime) + (NewVelocity - InVelocity) * (0.5f * DeltaTime);
 	return Delta;
 }
 
-FVector UProjectileMovementComponent::ComputeAcceleration(const FVector& InVelocity, float DeltaTime, bool bGravityEnabled) const
+FVector UProjectileMovementComponent::ComputeAcceleration(const FVector& InVelocity, float DeltaTime) const
 {
 	FVector Acceleration(FVector::ZeroVector);
 
-	if (bGravityEnabled)
-	{
-		Acceleration.Z += GetEffectiveGravityZ();
-	}
+	Acceleration.Z += GetEffectiveGravityZ();
 
 	if (bIsHomingProjectile && HomingTargetComponent.IsValid())
 	{
-		Acceleration += ComputeHomingAcceleration(InVelocity, DeltaTime, bGravityEnabled);
+		Acceleration += ComputeHomingAcceleration(InVelocity, DeltaTime);
 	}
 
 	return Acceleration;
 }
 
 // Allow the projectile to track towards its homing target.
-FVector UProjectileMovementComponent::ComputeHomingAcceleration(const FVector& InVelocity, float DeltaTime, bool bGravityEnabled) const
+FVector UProjectileMovementComponent::ComputeHomingAcceleration(const FVector& InVelocity, float DeltaTime) const
 {
 	FVector HomingAcceleration = ((HomingTargetComponent->GetComponentLocation() - UpdatedComponent->GetComponentLocation()).SafeNormal() * HomingAccelerationMagnitude);
 	return HomingAcceleration;
@@ -313,21 +386,27 @@ bool UProjectileMovementComponent::HandleHitWall(const FHitResult& Hit, float Ti
 FVector UProjectileMovementComponent::ComputeBounceResult(const FHitResult& Hit, float TimeSlice, const FVector& MoveDelta)
 {
 	FVector TempVelocity = Velocity;
+	const float VDotNormal = (TempVelocity | Hit.Normal);
 
-	// Project velocity onto normal in reflected direction.
-	const FVector ProjectedNormal = Hit.Normal * -(TempVelocity | Hit.Normal);
+	// Only if velocity is opposed by normal
+	if (VDotNormal < 0.f)
+	{
+		// Project velocity onto normal in reflected direction.
+		const FVector ProjectedNormal = Hit.Normal * -VDotNormal;
 
-	// Point velocity in direction parallel to surface
-	TempVelocity += ProjectedNormal;
+		// Point velocity in direction parallel to surface
+		TempVelocity += ProjectedNormal;
 
-	// Only tangential velocity should be affected by friction.
-	TempVelocity *= FMath::Clamp(1.f - Friction, 0.f, 1.f);
+		// Only tangential velocity should be affected by friction.
+		const float ScaledFriction = FMath::Clamp(-VDotNormal / TempVelocity.Size(), 0.f, 1.f) * Friction;
+		TempVelocity *= FMath::Clamp(1.f - ScaledFriction, 0.f, 1.f);
 
-	// Coefficient of restitution only applies perpendicular to impact.
-	TempVelocity += (ProjectedNormal * FMath::Max(Bounciness, 0.f));
+		// Coefficient of restitution only applies perpendicular to impact.
+		TempVelocity += (ProjectedNormal * FMath::Max(Bounciness, 0.f));
 
-	// Bounciness or Friction > 1 could cause us to exceed velocity.
-	TempVelocity = LimitVelocity(TempVelocity);
+		// Bounciness could cause us to exceed max speed.
+		TempVelocity = LimitVelocity(TempVelocity);
+	}
 
 	return TempVelocity;
 }
@@ -398,7 +477,7 @@ bool UProjectileMovementComponent::CheckStillInWorld()
 			ActorOwner->OutsideWorldBounds();
 			// not safe to use physics or collision at this point
 			ActorOwner->SetActorEnableCollision(false);
-			FHitResult Hit;
+			FHitResult Hit(1.f);
 			StopSimulating(Hit);
 			return false;
 		}
