Handle case where PIE is being torn down and we aren't able to spawn new actors/duplicate player state.

diff --git a/Engine/Source/Runtime/Engine/Private/GameMode.cpp b/Engine/Source/Runtime/Engine/Private/GameMode.cpp
--- a/Engine/Source/Runtime/Engine/Private/GameMode.cpp
+++ b/Engine/Source/Runtime/Engine/Private/GameMode.cpp
@@ -1540,34 +1540,37 @@ void AGameMode::AddInactivePlayer(APlayerState* PlayerState, APlayerController*
 	if (!PlayerState->bFromPreviousLevel && !PlayerState->bOnlySpectator)
 	{
 		APlayerState* NewPlayerState = PlayerState->Duplicate();
-		GetWorld()->GameState->RemovePlayerState(NewPlayerState);
+		if (NewPlayerState)
+		{
+			GetWorld()->GameState->RemovePlayerState(NewPlayerState);
 
-		// make PlayerState inactive
-		NewPlayerState->SetReplicates(false);
+			// make PlayerState inactive
+			NewPlayerState->SetReplicates(false);
 
-		// delete after some time
-		NewPlayerState->SetLifeSpan(InactivePlayerStateLifeSpan);
+			// delete after some time
+			NewPlayerState->SetLifeSpan(InactivePlayerStateLifeSpan);
 
-		// On console, we have to check the unique net id as network address isn't valid
-		bool bIsConsole = GEngine->IsConsoleBuild();
+			// On console, we have to check the unique net id as network address isn't valid
+			bool bIsConsole = GEngine->IsConsoleBuild();
 
-		// make sure no duplicates
-		for (int32 i=0; i<InactivePlayerArray.Num(); i++)
-		{
-			APlayerState* CurrentPlayerState = InactivePlayerArray[i];
-			if ( (CurrentPlayerState == NULL) || CurrentPlayerState->IsPendingKill() ||
-				(!bIsConsole && (CurrentPlayerState->SavedNetworkAddress == NewPlayerState->SavedNetworkAddress)))
+			// make sure no duplicates
+			for (int32 i=0; i<InactivePlayerArray.Num(); i++)
 			{
-				InactivePlayerArray.RemoveAt(i,1);
-				i--;
+				APlayerState* CurrentPlayerState = InactivePlayerArray[i];
+				if ( (CurrentPlayerState == NULL) || CurrentPlayerState->IsPendingKill() ||
+					(!bIsConsole && (CurrentPlayerState->SavedNetworkAddress == NewPlayerState->SavedNetworkAddress)))
+				{
+					InactivePlayerArray.RemoveAt(i,1);
+					i--;
+				}
 			}
-		}
-		InactivePlayerArray.Add(NewPlayerState);
+			InactivePlayerArray.Add(NewPlayerState);
 
-		// cap at 16 saved PlayerStates
-		if ( InactivePlayerArray.Num() > 16 )
-		{
-			InactivePlayerArray.RemoveAt(0, InactivePlayerArray.Num() - 16);
+			// cap at 16 saved PlayerStates
+			if ( InactivePlayerArray.Num() > 16 )
+			{
+				InactivePlayerArray.RemoveAt(0, InactivePlayerArray.Num() - 16);
+			}
 		}
 	}
 
diff --git a/Engine/Source/Runtime/Engine/Private/PlayerState.cpp b/Engine/Source/Runtime/Engine/Private/PlayerState.cpp
--- a/Engine/Source/Runtime/Engine/Private/PlayerState.cpp
+++ b/Engine/Source/Runtime/Engine/Private/PlayerState.cpp
@@ -286,7 +286,11 @@ APlayerState* APlayerState::Duplicate()
 	SpawnInfo.Instigator = Instigator;
 	SpawnInfo.bNoCollisionFail = true;
 	APlayerState* NewPlayerState = GetWorld()->SpawnActor<APlayerState>(GetClass(), SpawnInfo );
-	CopyProperties(NewPlayerState);
+	// Can fail in case of multiplayer PIE teardown
+	if (NewPlayerState)
+	{
+		CopyProperties(NewPlayerState);
+	}
 	return NewPlayerState;
 }
 
