#ue4 - MovementComponent auto registration of tick function moved to OnRegister() rather than InitializeComponent().

diff --git a/Engine/Source/Runtime/Engine/Classes/GameFramework/MovementComponent.h b/Engine/Source/Runtime/Engine/Classes/GameFramework/MovementComponent.h
--- a/Engine/Source/Runtime/Engine/Classes/GameFramework/MovementComponent.h
+++ b/Engine/Source/Runtime/Engine/Classes/GameFramework/MovementComponent.h
@@ -94,7 +94,7 @@ public:
 	virtual void RegisterComponentTickFunctions(bool bRegister) override;
 
 	/** Overridden to auto-register the updated component if it starts NULL, and we can find a root component on our owner. */
-	virtual void InitializeComponent() override;
+	virtual void OnRegister() override;
 	// End ActorComponent interface
 
 	/** @return gravity that affects this component */
@@ -310,10 +310,6 @@ public:
 	/** Snap the updated component to the plane constraint, if enabled. */
 	UFUNCTION(BlueprintCallable, Category="Components|Movement|Planar")
 	virtual void SnapUpdatedComponentToPlane();
-
-private:
-	/** Transient flag indicating whether we are executing InitializeComponent. */
-	uint32 bInInitializeComponent:1;
 };
 
 
diff --git a/Engine/Source/Runtime/Engine/Private/MovementComponent.cpp b/Engine/Source/Runtime/Engine/Private/MovementComponent.cpp
--- a/Engine/Source/Runtime/Engine/Private/MovementComponent.cpp
+++ b/Engine/Source/Runtime/Engine/Private/MovementComponent.cpp
@@ -29,7 +29,6 @@ UMovementComponent::UMovementComponent(const class FPostConstructInitializePrope
 
 	bWantsInitializeComponent = true;
 	bAutoActivate = true;
-	bInInitializeComponent = false;
 }
 
 
@@ -66,10 +65,9 @@ void UMovementComponent::SetUpdatedComponent(UPrimitiveComponent* NewUpdatedComp
 }
 
 
-void UMovementComponent::InitializeComponent()
+void UMovementComponent::OnRegister()
 {
-	bInInitializeComponent = true;
-	Super::InitializeComponent();
+	Super::OnRegister();
 
 	UPrimitiveComponent* NewUpdatedComponent = NULL;
 	if (UpdatedComponent != NULL)
@@ -95,13 +93,15 @@ void UMovementComponent::InitializeComponent()
 
 	PlaneConstraintNormal = PlaneConstraintNormal.SafeNormal();
 	SetUpdatedComponent(NewUpdatedComponent);
-	bInInitializeComponent = false;
 }
 
 void UMovementComponent::RegisterComponentTickFunctions(bool bRegister)
 {
 	Super::RegisterComponentTickFunctions(bRegister);
 
+	// Super may start up the tick function when we don't want to.
+	UpdateTickRegistration();
+
 	// If the owner ticks, make sure we tick first
 	AActor* Owner = GetOwner();
 	if (bRegister && PrimaryComponentTick.bCanEverTick && Owner && Owner->CanEverTick())
@@ -112,14 +112,11 @@ void UMovementComponent::RegisterComponentTickFunctions(bool bRegister)
 
 void UMovementComponent::UpdateTickRegistration()
 {
-	if (!bAutoUpdateTickRegistration)
+	if (bAutoUpdateTickRegistration)
 	{
-		return;
+		const bool bHasUpdatedComponent = (UpdatedComponent != NULL);
+		SetComponentTickEnabled(bHasUpdatedComponent && bAutoActivate);
 	}
-
-	const bool bAllowedToStartTick = (!bInInitializeComponent || bAutoActivate);
-	const bool bHasUpdatedComponent = (UpdatedComponent != NULL);
-	SetComponentTickEnabled(bHasUpdatedComponent && bAllowedToStartTick);
 }
 
 
