#ue4 - ProjectileMovementComponent velocity initialization moved to OnRegister() from InitializeComponent().

diff --git a/Engine/Source/Runtime/Engine/Classes/GameFramework/ProjectileMovementComponent.h b/Engine/Source/Runtime/Engine/Classes/GameFramework/ProjectileMovementComponent.h
--- a/Engine/Source/Runtime/Engine/Classes/GameFramework/ProjectileMovementComponent.h
+++ b/Engine/Source/Runtime/Engine/Classes/GameFramework/ProjectileMovementComponent.h
@@ -123,7 +123,7 @@ class ENGINE_API UProjectileMovementComponent : public UMovementComponent
 
 	//Begin UMovementComponent Interface
 	virtual float GetMaxSpeed() const override { return MaxSpeed; }
-	virtual void InitializeComponent() override;
+	virtual void OnRegister() override;
 	//End UMovementComponent Interface
 
 	/**
diff --git a/Engine/Source/Runtime/Engine/Private/ProjectileMovementComponent.cpp b/Engine/Source/Runtime/Engine/Private/ProjectileMovementComponent.cpp
--- a/Engine/Source/Runtime/Engine/Private/ProjectileMovementComponent.cpp
+++ b/Engine/Source/Runtime/Engine/Private/ProjectileMovementComponent.cpp
@@ -45,9 +45,9 @@ void UProjectileMovementComponent::Serialize( FArchive& Ar)
 }
 
 
-void UProjectileMovementComponent::InitializeComponent()
+void UProjectileMovementComponent::OnRegister()
 {
-	Super::InitializeComponent();
+	Super::OnRegister();
 
 	if (Velocity.SizeSquared() > 0.f)
 	{
