Fixed crash that can happen if skeleton somehow removes the bones that were there and animations were saved before. Fixed so that it can fix up the indices.

diff --git a/Engine/Source/Runtime/Engine/Classes/Animation/AnimSequence.h b/Engine/Source/Runtime/Engine/Classes/Animation/AnimSequence.h
--- a/Engine/Source/Runtime/Engine/Classes/Animation/AnimSequence.h
+++ b/Engine/Source/Runtime/Engine/Classes/Animation/AnimSequence.h
@@ -737,6 +737,7 @@ public:
 	 * This increases framecount + time, so that it requires recompression
 	 */
 	ENGINE_API bool AddLoopingInterpolation();
+
 #endif
 
 	/** 
@@ -769,6 +770,9 @@ private:
 	bool ConvertAnimationDataToRiggingData(FAnimSequenceTrackContainer & RiggingAnimationData);
 	bool ConvertRiggingDataToAnimationData(FAnimSequenceTrackContainer & RiggingAnimationData);
 	int32 GetSpaceBasedAnimationData(TArray< TArray<FTransform> > & AnimationDataInComponentSpace, FAnimSequenceTrackContainer * RiggingAnimationData) const;
+
+	/** Verify Track Map is valid, if not, fix up */
+	void VerifyTrackMap();
 #endif
 
 	/**
diff --git a/Engine/Source/Runtime/Engine/Private/Animation/AnimSequence.cpp b/Engine/Source/Runtime/Engine/Private/Animation/AnimSequence.cpp
--- a/Engine/Source/Runtime/Engine/Private/Animation/AnimSequence.cpp
+++ b/Engine/Source/Runtime/Engine/Private/Animation/AnimSequence.cpp
@@ -397,22 +397,9 @@ void UAnimSequence::PostLoad()
 	// setup the Codec interfaces
 	AnimationFormat_SetInterfaceLinks(*this);
 
+#if WITH_EDITOR
 	// save track data for the case
-#if WITH_EDITORONLY_DATA
-	USkeleton* MySkeleton = GetSkeleton();
-	if( AnimationTrackNames.Num() != TrackToSkeletonMapTable.Num() && MySkeleton!=NULL )
-	{
-		UE_LOG(LogAnimation, Verbose, TEXT("Fixing track names (%s) from %s."), *GetName(), *MySkeleton->GetName());
-
-		const TArray<FBoneNode> & BoneTree = MySkeleton->GetBoneTree();
-		AnimationTrackNames.Empty();
-		AnimationTrackNames.AddUninitialized(TrackToSkeletonMapTable.Num());
-		for (int32 I=0; I<TrackToSkeletonMapTable.Num(); ++I)
-		{
-			const FTrackToSkeletonMap& TrackMap = TrackToSkeletonMapTable[I];
-			AnimationTrackNames[I] = MySkeleton->GetReferenceSkeleton().GetBoneName(TrackMap.BoneTreeIndex);
-		}
-	}
+	VerifyTrackMap();
 #endif
 
 	// save track data for the case
@@ -441,6 +428,57 @@ void UAnimSequence::PostLoad()
 	}
 }
 
+#if WITH_EDITOR
+void UAnimSequence::VerifyTrackMap()
+{
+	USkeleton* MySkeleton = GetSkeleton();
+	if(AnimationTrackNames.Num() != TrackToSkeletonMapTable.Num() && MySkeleton!=NULL)
+	{
+		UE_LOG(LogAnimation, Warning, TEXT("RESAVE ANIMATION NEEDED(%s): Fixing track names."), *GetName());
+
+		const TArray<FBoneNode> & BoneTree = MySkeleton->GetBoneTree();
+		AnimationTrackNames.Empty();
+		AnimationTrackNames.AddUninitialized(TrackToSkeletonMapTable.Num());
+		for(int32 I=0; I<TrackToSkeletonMapTable.Num(); ++I)
+		{
+			const FTrackToSkeletonMap& TrackMap = TrackToSkeletonMapTable[I];
+			AnimationTrackNames[I] = MySkeleton->GetReferenceSkeleton().GetBoneName(TrackMap.BoneTreeIndex);
+		}
+	}
+	else 
+	{
+		int32 NumTracks = AnimationTrackNames.Num();
+		int32 NumSkeletonBone = MySkeleton->GetReferenceSkeleton().GetNum();
+
+		bool bNeedsFixing = false;
+		// verify all tracks are still valid
+		for(int32 TrackIndex=0; TrackIndex<NumTracks; TrackIndex++)
+		{
+			int32 SkeletonBoneIndex = TrackToSkeletonMapTable[TrackIndex].BoneTreeIndex;
+			// invalid index found
+			if (NumSkeletonBone <= SkeletonBoneIndex)
+			{
+				// if one is invalid, fix up for all.
+				// you don't know what index got messed up
+				bNeedsFixing = true;
+				break;
+			}
+		}
+
+		if(bNeedsFixing)
+		{
+			UE_LOG(LogAnimation, Warning, TEXT("RESAVE ANIMATION NEEDED(%s): Fixing track index."), *GetName());
+
+			const TArray<FBoneNode> & BoneTree = MySkeleton->GetBoneTree();
+			for(int32 I=0; I<NumTracks; ++I)
+			{
+				TrackToSkeletonMapTable[I].BoneTreeIndex = MySkeleton->GetReferenceSkeleton().FindBoneIndex(AnimationTrackNames[I]);
+			}
+		}
+	}
+}
+
+#endif // WITH_EDITOR
 void UAnimSequence::FixAdditiveType()
 {
 	if( GetLinkerUE4Version() < VER_UE4_ADDITIVE_TYPE_CHANGE )
