#docs - Add or fix class descriptions for multiple framework classes.

diff --git a/Engine/Source/Runtime/Engine/Classes/Components/DestructibleComponent.h b/Engine/Source/Runtime/Engine/Classes/Components/DestructibleComponent.h
--- a/Engine/Source/Runtime/Engine/Classes/Components/DestructibleComponent.h
+++ b/Engine/Source/Runtime/Engine/Classes/Components/DestructibleComponent.h
@@ -37,7 +37,7 @@ struct FDestructibleChunkInfo
 DECLARE_DYNAMIC_MULTICAST_DELEGATE_TwoParams(FComponentFractureSignature, const FVector &, HitPoint, const FVector &, HitDirection);
 
 /**
- *	This component holds the physics data for a DestructibleActor
+ *	This component holds the physics data for a DestructibleActor.
  *
  *	The USkeletalMesh pointer in the base class (SkinnedMeshComponent) MUST be a DestructibleMesh
  */
diff --git a/Engine/Source/Runtime/Engine/Classes/Components/SkinnedMeshComponent.h b/Engine/Source/Runtime/Engine/Classes/Components/SkinnedMeshComponent.h
--- a/Engine/Source/Runtime/Engine/Classes/Components/SkinnedMeshComponent.h
+++ b/Engine/Source/Runtime/Engine/Classes/Components/SkinnedMeshComponent.h
@@ -1,11 +1,5 @@
 // Copyright 1998-2014 Epic Games, Inc. All Rights Reserved.
 
-/**
- * 
- * Skinned Mesh Component that supports Bone Skinned Mesh rendering
- * This class does not support animation
- *
- */
 
 #pragma once
 #include "Components/MeshComponent.h"
@@ -141,6 +135,15 @@ struct FSkelMeshComponentLODInfo
 	
 };
 
+
+/**
+ *
+ * Skinned mesh component that supports bone skinned mesh rendering.
+ * This class does not support animation.
+ *
+ * @see USkeletalMeshComponent
+*/
+
 UCLASS(hidecategories=Object, config=Engine, editinlinenew, abstract)
 class ENGINE_API USkinnedMeshComponent : public UMeshComponent
 {
diff --git a/Engine/Source/Runtime/Engine/Classes/Components/SkyLightComponent.h b/Engine/Source/Runtime/Engine/Classes/Components/SkyLightComponent.h
--- a/Engine/Source/Runtime/Engine/Classes/Components/SkyLightComponent.h
+++ b/Engine/Source/Runtime/Engine/Classes/Components/SkyLightComponent.h
@@ -6,10 +6,11 @@
 
 /** 
  * A cubemap texture resource that knows how to upload the capture data from a sky capture. 
- * @todo - support compression
  */
 class FSkyTextureCubeResource : public FTexture, private FDeferredCleanupInterface
 {
+	// @todo - support compression
+
 public:
 
 	FSkyTextureCubeResource() :
diff --git a/Engine/Source/Runtime/Engine/Classes/Components/StaticMeshComponent.h b/Engine/Source/Runtime/Engine/Classes/Components/StaticMeshComponent.h
--- a/Engine/Source/Runtime/Engine/Classes/Components/StaticMeshComponent.h
+++ b/Engine/Source/Runtime/Engine/Classes/Components/StaticMeshComponent.h
@@ -108,6 +108,7 @@ struct TStructOpsTypeTraits<FStaticMeshComponentLODInfo> : public TStructOpsType
 
 /**
  * StaticMeshComponent is used to create an instance of a UStaticMesh.
+ * A static mesh is a piece of geometry that consists of a static set of polygons.
  *
  * @see https://docs.unrealengine.com/latest/INT/Engine/Content/Types/StaticMeshes/
  * @see UStaticMesh
diff --git a/Engine/Source/Runtime/Engine/Classes/Components/TimelineComponent.h b/Engine/Source/Runtime/Engine/Classes/Components/TimelineComponent.h
--- a/Engine/Source/Runtime/Engine/Classes/Components/TimelineComponent.h
+++ b/Engine/Source/Runtime/Engine/Classes/Components/TimelineComponent.h
@@ -354,10 +354,9 @@ private:
 };
 
 /** 
- * A timeline component. 
- * Holds a series of Events, floats, vectors or colors with associated keyframes.
+ * TimelineComponent holds a series of events, floats, vectors or colors with associated keyframes.
  * Events can be triggered at keyframes along the timeline. 
- * floats, vectors an colors are interpolated between keyframes along the timeline.
+ * Floats, vectors, and colors are interpolated between keyframes along the timeline.
  */
 UCLASS(MinimalAPI)
 class UTimelineComponent : public UActorComponent
diff --git a/Engine/Source/Runtime/Engine/Classes/Engine/BlockingVolume.h b/Engine/Source/Runtime/Engine/Classes/Engine/BlockingVolume.h
--- a/Engine/Source/Runtime/Engine/Classes/Engine/BlockingVolume.h
+++ b/Engine/Source/Runtime/Engine/Classes/Engine/BlockingVolume.h
@@ -4,7 +4,7 @@
 #pragma once
 #include "BlockingVolume.generated.h"
 
-/** An invisible volume used to block other actors */
+/** An invisible volume used to block other actors. */
 UCLASS(MinimalAPI)
 class ABlockingVolume : public AVolume
 {
diff --git a/Engine/Source/Runtime/Engine/Classes/Engine/DebugCameraController.h b/Engine/Source/Runtime/Engine/Classes/Engine/DebugCameraController.h
--- a/Engine/Source/Runtime/Engine/Classes/Engine/DebugCameraController.h
+++ b/Engine/Source/Runtime/Engine/Classes/Engine/DebugCameraController.h
@@ -1,13 +1,12 @@
 // Copyright 1998-2014 Epic Games, Inc. All Rights Reserved.
 
-//-----------------------------------------------------------
-// Debug Camera Controller
-//
-// To turn it on, please press Alt+C or both (left and right) analogs on xbox pad
-//
-// Check the debug camera bindings in BaseInput.ini for the camera controls.
-//
-//-----------------------------------------------------------
+/**
+ * Camera controller that allows you to fly around a level mostly unrestricted by normal movement rules.
+ * 
+ * To turn it on, please press Alt+C or both (left and right) analogs on xbox pad, or use the "ToggleDebugCamera" console command.
+ * 
+ * Check the debug camera bindings in DefaultPawn.cpp for the camera controls.
+ */
 
 #pragma once
 #include "DebugCameraController.generated.h"
diff --git a/Engine/Source/Runtime/Engine/Classes/Engine/DebugCameraHUD.h b/Engine/Source/Runtime/Engine/Classes/Engine/DebugCameraHUD.h
--- a/Engine/Source/Runtime/Engine/Classes/Engine/DebugCameraHUD.h
+++ b/Engine/Source/Runtime/Engine/Classes/Engine/DebugCameraHUD.h
@@ -4,6 +4,9 @@
 #include "GameFramework/HUD.h"
 #include "DebugCameraHUD.generated.h"
 
+/**
+ * HUD that displays info for the DebugCameraController view.
+ */
 UCLASS(config=Game)
 class ENGINE_API ADebugCameraHUD : public AHUD
 {
diff --git a/Engine/Source/Runtime/Engine/Classes/Engine/GameInstance.h b/Engine/Source/Runtime/Engine/Classes/Engine/GameInstance.h
--- a/Engine/Source/Runtime/Engine/Classes/Engine/GameInstance.h
+++ b/Engine/Source/Runtime/Engine/Classes/Engine/GameInstance.h
@@ -1,12 +1,5 @@
 // Copyright 1998-2014 Epic Games, Inc. All Rights Reserved.
 
-//=============================================================================
-// GameInstance: High-level manager object for a instance of the running game.
-// Spawned at game creation and not destroyed until game instance is shut down.
-// Running as a standalone game, there will be one of these.  Running in PIE
-// will generate one of these per PIE instance.
-//=============================================================================
-
 #pragma once
 #include "GameInstance.generated.h"
 
@@ -28,6 +21,13 @@ namespace GameInstanceState
 class FOnlineSessionSearchResult;
 
 
+/**
+ * GameInstance: high-level manager object for an instance of the running game.
+ * Spawned at game creation and not destroyed until game instance is shut down.
+ * Running as a standalone game, there will be one of these.
+ * Running in PIE (play-in-editor) will generate one of these per PIE instance.
+ */
+
 UCLASS(config=Game, transient, BlueprintType, Blueprintable)
 class ENGINE_API UGameInstance : public UObject, public FExec
 {
diff --git a/Engine/Source/Runtime/Engine/Classes/Engine/LevelStreaming.h b/Engine/Source/Runtime/Engine/Classes/Engine/LevelStreaming.h
--- a/Engine/Source/Runtime/Engine/Classes/Engine/LevelStreaming.h
+++ b/Engine/Source/Runtime/Engine/Classes/Engine/LevelStreaming.h
@@ -62,8 +62,6 @@ DECLARE_DYNAMIC_MULTICAST_DELEGATE( FLevelStreamingLoadedStatus );
 DECLARE_DYNAMIC_MULTICAST_DELEGATE( FLevelStreamingVisibilityStatus );
 
 /**
- * LevelStreaming
- *
  * Abstract base class of container object encapsulating data required for streaming and providing 
  * interface for when a level should be streamed in and out of memory.
  *
diff --git a/Engine/Source/Runtime/Engine/Classes/Engine/LocalPlayer.h b/Engine/Source/Runtime/Engine/Classes/Engine/LocalPlayer.h
--- a/Engine/Source/Runtime/Engine/Classes/Engine/LocalPlayer.h
+++ b/Engine/Source/Runtime/Engine/Classes/Engine/LocalPlayer.h
@@ -125,8 +125,8 @@ private:
 
 /**
  *	Each player that is active on the current client has a LocalPlayer. It stays active across maps
- *	There may be several spawned in the case of splitscreen/coop
- *	There may be 0 spawned on servers
+ *	There may be several spawned in the case of splitscreen/coop.
+ *	There may be 0 spawned on servers.
  */
 UCLASS(Within=Engine, config=Engine, transient)
 class ENGINE_API ULocalPlayer : public UPlayer
diff --git a/Engine/Source/Runtime/Engine/Classes/Engine/Note.h b/Engine/Source/Runtime/Engine/Classes/Engine/Note.h
--- a/Engine/Source/Runtime/Engine/Classes/Engine/Note.h
+++ b/Engine/Source/Runtime/Engine/Classes/Engine/Note.h
@@ -1,9 +1,9 @@
 // Copyright 1998-2014 Epic Games, Inc. All Rights Reserved.
 
-//=============================================================================
-// A sticky note.  Level designers can place these in the level and then
-// view them as a batch in the error/warnings window.
-//=============================================================================
+/**
+ * A sticky note.  Level designers can place these in the level and then
+ * view them as a batch in the error/warnings window.
+ */
 
 #pragma once
 #include "Note.generated.h"
diff --git a/Engine/Source/Runtime/Engine/Classes/Engine/Player.h b/Engine/Source/Runtime/Engine/Classes/Engine/Player.h
--- a/Engine/Source/Runtime/Engine/Classes/Engine/Player.h
+++ b/Engine/Source/Runtime/Engine/Classes/Engine/Player.h
@@ -1,8 +1,8 @@
 // Copyright 1998-2014 Epic Games, Inc. All Rights Reserved.
 
-//=============================================================================
-// Player: Corresponds to a real player (a local camera or remote net player).
-//=============================================================================
+/**
+ * Player: Corresponds to a real player (a local camera or remote net player).
+ */
 
 #pragma once
 #include "Player.generated.h"
diff --git a/Engine/Source/Runtime/Engine/Classes/Engine/PostProcessVolume.h b/Engine/Source/Runtime/Engine/Classes/Engine/PostProcessVolume.h
--- a/Engine/Source/Runtime/Engine/Classes/Engine/PostProcessVolume.h
+++ b/Engine/Source/Runtime/Engine/Classes/Engine/PostProcessVolume.h
@@ -1,9 +1,9 @@
 // Copyright 1998-2014 Epic Games, Inc. All Rights Reserved.
 
-//=============================================================================
-// PostProcessVolume:  a post process settings volume
-// Used to affect post process settings in the game and editor.
-//=============================================================================
+/**
+ * PostProcessVolume:  a post process settings volume
+ * Used to affect post process settings in the game and editor.
+ */
 
 #pragma once
 
diff --git a/Engine/Source/Runtime/Engine/Classes/Engine/SphereReflectionCapture.h b/Engine/Source/Runtime/Engine/Classes/Engine/SphereReflectionCapture.h
--- a/Engine/Source/Runtime/Engine/Classes/Engine/SphereReflectionCapture.h
+++ b/Engine/Source/Runtime/Engine/Classes/Engine/SphereReflectionCapture.h
@@ -6,7 +6,7 @@
 class UDrawSphereComponent;
 
 /** 
- *	Actor used to capture the scene for reflection in a sphere shape 
+ *	Actor used to capture the scene for reflection in a sphere shape.
  *	@see https://docs.unrealengine.com/latest/INT/Resources/ContentExamples/Reflections/1_4
  */
 UCLASS(hidecategories = (Collision, Attachment, Actor), MinimalAPI)
diff --git a/Engine/Source/Runtime/Engine/Classes/GameFramework/CameraBlockingVolume.h b/Engine/Source/Runtime/Engine/Classes/GameFramework/CameraBlockingVolume.h
--- a/Engine/Source/Runtime/Engine/Classes/GameFramework/CameraBlockingVolume.h
+++ b/Engine/Source/Runtime/Engine/Classes/GameFramework/CameraBlockingVolume.h
@@ -3,6 +3,9 @@
 #pragma once
 #include "CameraBlockingVolume.generated.h"
 
+/**
+ * A volume which blocks the Camera channel by default.
+ */
 UCLASS()
 class ENGINE_API ACameraBlockingVolume : public AVolume
 {
diff --git a/Engine/Source/Runtime/Engine/Classes/GameFramework/PhysicsVolume.h b/Engine/Source/Runtime/Engine/Classes/GameFramework/PhysicsVolume.h
--- a/Engine/Source/Runtime/Engine/Classes/GameFramework/PhysicsVolume.h
+++ b/Engine/Source/Runtime/Engine/Classes/GameFramework/PhysicsVolume.h
@@ -5,8 +5,8 @@
 #include "PhysicsVolume.generated.h"
 
 /**
- * PhysicsVolume: A bounding volume which affects actor physics
- * Each AActor is affected at any time by one PhysicsVolume
+ * PhysicsVolume: A bounding volume which affects actor physics.
+ * Each AActor is affected at any time by one PhysicsVolume.
  */
 
 UCLASS()
diff --git a/Engine/Source/Runtime/Engine/Classes/GameFramework/PlayerInput.h b/Engine/Source/Runtime/Engine/Classes/GameFramework/PlayerInput.h
--- a/Engine/Source/Runtime/Engine/Classes/GameFramework/PlayerInput.h
+++ b/Engine/Source/Runtime/Engine/Classes/GameFramework/PlayerInput.h
@@ -2,8 +2,8 @@
 
 //=============================================================================
 // PlayerInput
-// Object within playercontroller that manages player input.
-// only spawned on client
+// Object within PlayerController that manages player input.
+// Only spawned on client.
 //=============================================================================
 
 #pragma once
@@ -240,6 +240,10 @@ struct FAxisKeyDetails
 	}
 };
 
+/**
+ * Object within PlayerController that manages player input.
+ * Only spawned on the client in network games.
+ */
 UCLASS(Within=PlayerController, config=Input, transient)
 class ENGINE_API UPlayerInput : public UObject
 {
diff --git a/Engine/Source/Runtime/Engine/Classes/GameFramework/PlayerState.h b/Engine/Source/Runtime/Engine/Classes/GameFramework/PlayerState.h
--- a/Engine/Source/Runtime/Engine/Classes/GameFramework/PlayerState.h
+++ b/Engine/Source/Runtime/Engine/Classes/GameFramework/PlayerState.h
@@ -24,12 +24,10 @@ struct PingAvgData
 	}
 };
 
-//=============================================================================
-// A PlayerState is created for every player on a server (or in a standalone game).
-// PlayerStates are replicated to all clients, and contain network game relevant information about the player,
-// such as playername, score, etc.
-//=============================================================================
-
+/**
+ * A PlayerState is created for every player on a server (or in a standalone game).
+ * PlayerStates are replicated to all clients, and contain network game relevant information about the player, such as playername, score, etc.
+ */
 UCLASS(BlueprintType, Blueprintable, notplaceable)
 class ENGINE_API APlayerState : public AInfo
 {
diff --git a/Engine/Source/Runtime/Engine/Classes/GameFramework/SpectatorPawn.h b/Engine/Source/Runtime/Engine/Classes/GameFramework/SpectatorPawn.h
--- a/Engine/Source/Runtime/Engine/Classes/GameFramework/SpectatorPawn.h
+++ b/Engine/Source/Runtime/Engine/Classes/GameFramework/SpectatorPawn.h
@@ -1,10 +1,9 @@
 // Copyright 1998-2014 Epic Games, Inc. All Rights Reserved.
 
-//=============================================================================
-// SpectatorPawns are simple pawns that can fly around the world, used by
-// PlayerControllers when in the spectator state.
-//
-//=============================================================================
+/**
+ * SpectatorPawns are simple pawns that can fly around the world, used by
+ * PlayerControllers when in the spectator state.
+ */
 
 #pragma once
 #include "GameFramework/DefaultPawn.h"
diff --git a/Engine/Source/Runtime/Engine/Classes/GameFramework/SpectatorPawnMovement.h b/Engine/Source/Runtime/Engine/Classes/GameFramework/SpectatorPawnMovement.h
--- a/Engine/Source/Runtime/Engine/Classes/GameFramework/SpectatorPawnMovement.h
+++ b/Engine/Source/Runtime/Engine/Classes/GameFramework/SpectatorPawnMovement.h
@@ -1,9 +1,9 @@
 // Copyright 1998-2014 Epic Games, Inc. All Rights Reserved.
 
 /**
- * Movement component updates position of associated PrimitiveComponent during its tick
+ * Movement component used by SpectatorPawn.
+ * Primarily exists to be able to ignore time dilation during tick.
  */
-
 #pragma once
 #include "GameFramework/FloatingPawnMovement.h"
 #include "SpectatorPawnMovement.generated.h"
diff --git a/Engine/Source/Runtime/Engine/Classes/GameFramework/SpringArmComponent.h b/Engine/Source/Runtime/Engine/Classes/GameFramework/SpringArmComponent.h
--- a/Engine/Source/Runtime/Engine/Classes/GameFramework/SpringArmComponent.h
+++ b/Engine/Source/Runtime/Engine/Classes/GameFramework/SpringArmComponent.h
@@ -5,7 +5,7 @@
 #include "SpringArmComponent.generated.h"
 
 /**
- * This component tried to maintain its children at a fixed distance from the parent,
+ * This component tries to maintain its children at a fixed distance from the parent,
  * but will retract the children if there is a collision, and spring back when there is no collision.
  *
  * Example: Use as a 'camera boom' to keep the follow camera for a player from colliding into the world.
