Merge github pull request 483. https://github.com/EpicGames/UnrealEngine/pull/483

diff --git a/Engine/Source/Runtime/Engine/Classes/Engine/NetSerialization.h b/Engine/Source/Runtime/Engine/Classes/Engine/NetSerialization.h
--- a/Engine/Source/Runtime/Engine/Classes/Engine/NetSerialization.h
+++ b/Engine/Source/Runtime/Engine/Classes/Engine/NetSerialization.h
@@ -965,7 +965,7 @@ struct FVector_NetQuantize10 : public FVector
 	{}
 
 	FORCEINLINE FVector_NetQuantize10( float InX, float InY, float InZ )
-	: FVector(InX, InY, InX)
+	: FVector(InX, InY, InZ)
 	{}
 
 	FORCEINLINE FVector_NetQuantize10(const FVector &InVec)
@@ -1010,7 +1010,7 @@ struct FVector_NetQuantize100 : public FVector
 	{}
 
 	FORCEINLINE FVector_NetQuantize100( float InX, float InY, float InZ )
-	: FVector(InX, InY, InX)
+	: FVector(InX, InY, InZ)
 	{}
 
 	FORCEINLINE FVector_NetQuantize100(const FVector &InVec)
@@ -1053,7 +1053,7 @@ struct FVector_NetQuantizeNormal : public FVector
 	{}
 
 	FORCEINLINE FVector_NetQuantizeNormal( float InX, float InY, float InZ )
-	: FVector(InX, InY, InX)
+	: FVector(InX, InY, InZ)
 	{}
 
 	FORCEINLINE FVector_NetQuantizeNormal(const FVector &InVec)
