Fixed non-safe threading calls to the SlowTask system

diff --git a/Engine/Source/Runtime/Engine/Private/TextureDerivedData.cpp b/Engine/Source/Runtime/Engine/Private/TextureDerivedData.cpp
--- a/Engine/Source/Runtime/Engine/Private/TextureDerivedData.cpp
+++ b/Engine/Source/Runtime/Engine/Private/TextureDerivedData.cpp
@@ -468,7 +468,7 @@ class FTextureStatusMessageContext : public FScopedSlowTask
 {
 public:
 	explicit FTextureStatusMessageContext(const FText& InMessage)
-		: FScopedSlowTask(0, InMessage)
+		: FScopedSlowTask(0, InMessage, IsInGameThread())
 	{
 		UE_LOG(LogTexture,Display,TEXT("%s"),*InMessage.ToString());
 		MakeDialog();
