UE4 - Added RHISupportsComputeShaders()

diff --git a/Engine/Source/Runtime/Engine/Private/GPUSort.cpp b/Engine/Source/Runtime/Engine/Private/GPUSort.cpp
--- a/Engine/Source/Runtime/Engine/Private/GPUSort.cpp
+++ b/Engine/Source/Runtime/Engine/Private/GPUSort.cpp
@@ -221,10 +221,9 @@ class FRadixSortClearOffsetsCS : public FGlobalShader
 
 public:
 
-	static bool ShouldCache( EShaderPlatform Platform )
+	static bool ShouldCache(EShaderPlatform Platform)
 	{
-		//@todo-rco: Fix the shader
-		return IsFeatureLevelSupported(Platform, ERHIFeatureLevel::SM5);
+		return RHISupportsComputeShaders(Platform);
 	}
 
 	static void ModifyCompilationEnvironment( EShaderPlatform Platform, FShaderCompilerEnvironment& OutEnvironment )
@@ -308,8 +307,7 @@ public:
 
 	static bool ShouldCache( EShaderPlatform Platform )
 	{
-		//@todo-rco: Fix the shader
-		return IsFeatureLevelSupported(Platform, ERHIFeatureLevel::SM5);
+		return RHISupportsComputeShaders(Platform);
 	}
 
 	static void ModifyCompilationEnvironment( EShaderPlatform Platform, FShaderCompilerEnvironment& OutEnvironment )
@@ -426,8 +424,7 @@ public:
 
 	static bool ShouldCache( EShaderPlatform Platform )
 	{
-		//@todo-rco: Fix the shader
-		return IsFeatureLevelSupported(Platform, ERHIFeatureLevel::SM5);
+		return RHISupportsComputeShaders(Platform);
 	}
 
 	static void ModifyCompilationEnvironment( EShaderPlatform Platform, FShaderCompilerEnvironment& OutEnvironment )
@@ -523,8 +520,7 @@ public:
 
 	static bool ShouldCache( EShaderPlatform Platform )
 	{
-		//@todo-rco: Fix the shader
-		return IsFeatureLevelSupported(Platform, ERHIFeatureLevel::SM5);
+		return RHISupportsComputeShaders(Platform);
 	}
 
 	static void ModifyCompilationEnvironment( EShaderPlatform Platform, FShaderCompilerEnvironment& OutEnvironment )
diff --git a/Engine/Source/Runtime/RHI/Public/RHIDefinitions.h b/Engine/Source/Runtime/RHI/Public/RHIDefinitions.h
--- a/Engine/Source/Runtime/RHI/Public/RHIDefinitions.h
+++ b/Engine/Source/Runtime/RHI/Public/RHIDefinitions.h
@@ -662,6 +662,12 @@ inline bool RHISupportsSeparateMSAAAndResolveTextures(const EShaderPlatform Plat
 	return Platform != SP_METAL && Platform != SP_METAL_MRT;
 }
 
+inline bool RHISupportsComputeShaders(const EShaderPlatform Platform)
+{
+	//@todo-rco: Add Metal support
+	return IsFeatureLevelSupported(Platform, ERHIFeatureLevel::SM5);// || (Platform == SP_METAL);
+}
+
 inline bool RHISupportsGeometryShaders(const EShaderPlatform Platform)
 {
 	return IsFeatureLevelSupported(Platform, ERHIFeatureLevel::SM4) && Platform != SP_METAL_MRT;
