fixed TTP 349118 CRITICAL: Regression: CRASH: Occurs when launching certain projects either in editor or game.

diff --git a/Engine/Source/Runtime/Engine/Classes/Materials/Material.h b/Engine/Source/Runtime/Engine/Classes/Materials/Material.h
--- a/Engine/Source/Runtime/Engine/Classes/Materials/Material.h
+++ b/Engine/Source/Runtime/Engine/Classes/Materials/Material.h
@@ -1116,10 +1116,10 @@ public:
 	/**
 	*	Get the expression input for the given property
 	*
-	*	@param	InProperty				The material property chain to inspect, such as MP_DiffuseColor.
+	*	@param	InProperty				The material property chain to inspect, such as MP_BaseColor.
 	*
 	*	@return	FExpressionInput*		A pointer to the expression input of the property specified, 
-	*									or NULL if an invalid property was requested (MP_DiffuseColor, MP_SpecularColor have been removed, they would return NULL).
+	*									or NULL if an invalid property was requested (some properties have been removed from UI, those return NULL).
 	*/
 	ENGINE_API FExpressionInput* GetExpressionInputForProperty(EMaterialProperty InProperty);
 
@@ -1136,7 +1136,7 @@ public:
 	/**
 	 *	Get the expression chain for the given property (ie fill in the given array with all expressions in the chain).
 	 *
-	 *	@param	InProperty				The material property chain to inspect, such as MP_DiffuseColor.
+	 *	@param	InProperty				The material property chain to inspect, such as MP_BaseColor.
 	 *	@param	OutExpressions			The array to fill in all of the expressions.
 	 *	@param	InStaticParameterSet	Optional static parameter set - if supplied only walk the StaticSwitch branches according to it.
 	 *
diff --git a/Engine/Source/Runtime/Engine/Private/Materials/Material.cpp b/Engine/Source/Runtime/Engine/Private/Materials/Material.cpp
--- a/Engine/Source/Runtime/Engine/Private/Materials/Material.cpp
+++ b/Engine/Source/Runtime/Engine/Private/Materials/Material.cpp
@@ -3528,7 +3528,7 @@ USubsurfaceProfile* UMaterial::GetSubsurfaceProfile_Internal() const
 	return SubsurfaceProfile; 
 }
 
-bool UMaterial::IsPropertyActive(EMaterialProperty InProperty)const 
+bool UMaterial::IsPropertyActive(EMaterialProperty InProperty) const 
 {
 	if(MaterialDomain == MD_PostProcess)
 	{
@@ -3541,7 +3541,8 @@ bool UMaterial::IsPropertyActive(EMaterialProperty InProperty)const
 	}
 	else if(MaterialDomain == MD_DeferredDecal)
 	{
-		if(InProperty >= MP_CustomizedUVs0 )
+		// todo: this should also check for <= MP_CustomizedUVs7 but then MaterialAttruibutes would not return true, should it? Wht we don't chekc for the checkbox in the material
+		if(InProperty >= MP_CustomizedUVs0)
 		{
 			return true;
 		}
@@ -3577,25 +3578,21 @@ bool UMaterial::IsPropertyActive(EMaterialProperty InProperty)const
 
 			case DBM_DBuffer_ColorNormalRoughness:
 				return InProperty == MP_Normal
-					|| InProperty == MP_DiffuseColor
 					|| InProperty == MP_BaseColor
 					|| InProperty == MP_Roughness
 					|| InProperty == MP_Opacity;
 
 			case DBM_DBuffer_Color:
-				return InProperty == MP_DiffuseColor
-					|| InProperty == MP_BaseColor
+				return InProperty == MP_BaseColor
 					|| InProperty == MP_Opacity;
 
 			case DBM_DBuffer_ColorNormal:
-				return InProperty == MP_DiffuseColor
-					|| InProperty == MP_BaseColor
+				return InProperty == MP_BaseColor
 					|| InProperty == MP_Normal
 					|| InProperty == MP_Opacity;
 
 			case DBM_DBuffer_ColorRoughness:
-				return InProperty == MP_DiffuseColor
-					|| InProperty == MP_BaseColor
+				return InProperty == MP_BaseColor
 					|| InProperty == MP_Roughness
 					|| InProperty == MP_Opacity;
 
