Update the feature level switched scene with precomputed light volume, PrecomputedVisibilityHandler and PrecomputedVolumeDistanceField.

diff --git a/Engine/Source/Runtime/Engine/Private/World.cpp b/Engine/Source/Runtime/Engine/Private/World.cpp
--- a/Engine/Source/Runtime/Engine/Private/World.cpp
+++ b/Engine/Source/Runtime/Engine/Private/World.cpp
@@ -5371,6 +5371,12 @@ void UWorld::ChangeFeatureLevel(ERHIFeatureLevel::Type InFeatureLevel)
 			GetRendererModule().RemoveScene(Scene);
 
 			GetRendererModule().AllocateScene(this, bRequiresHitProxies, FXSystem != nullptr, InFeatureLevel );
+
+			PersistentLevel->ReleaseRenderingResources();
+			PersistentLevel->InitializeRenderingResources();
+
+			PersistentLevel->PrecomputedVisibilityHandler.UpdateScene(Scene);
+			PersistentLevel->PrecomputedVolumeDistanceField.UpdateScene(Scene);
 		}
 
 		TriggerStreamingDataRebuild();
