CIS fix

diff --git a/Engine/Source/Runtime/Engine/Private/PlayerController.cpp b/Engine/Source/Runtime/Engine/Private/PlayerController.cpp
--- a/Engine/Source/Runtime/Engine/Private/PlayerController.cpp
+++ b/Engine/Source/Runtime/Engine/Private/PlayerController.cpp
@@ -3676,7 +3676,7 @@ void APlayerController::TickPlayerInput(const float DeltaSeconds, const bool bGa
 		if (SVirtualJoystick::ShouldDisplayTouchInterface())
 		{
 			CreateTouchInterface();
-			VirtualJoystick->SetVisibility(true, true);
+			VirtualJoystick->SetVisibility(EVisibility::Visible);
 		}
 		else
 		{
