added comment

diff --git a/Engine/Source/Runtime/Engine/Classes/Engine/Scene.h b/Engine/Source/Runtime/Engine/Classes/Engine/Scene.h
--- a/Engine/Source/Runtime/Engine/Classes/Engine/Scene.h
+++ b/Engine/Source/Runtime/Engine/Classes/Engine/Scene.h
@@ -767,7 +767,10 @@ struct FPostProcessSettings
 	// FPostProcessSettings constructor and possibly the SetBaseValues() method.
 	// -----------------------------------------------------------------------
 	
-	/** Allows custom post process materials to be defined,		using a MaterialInstance with the same Material as its parent to allow blending. Make sure you use the "PostProcess" domain type */
+	/**
+	 * Allows custom post process materials to be defined, using a MaterialInstance with the same Material as its parent to allow blending.
+	 * Make sure you use the "PostProcess" domain type. This can be used for any UObject object implementing the IBlendableInterface (e.g. could be used to fade weather settings).
+	 */
 	UPROPERTY(EditAnywhere, BlueprintReadWrite, Category=Misc)
 	TArray<UObject*> Blendables;
 
