Fixed material override not disabling SubsurfaceScattering passes

diff --git a/Engine/Source/Runtime/Engine/Private/Materials/MaterialInterface.cpp b/Engine/Source/Runtime/Engine/Private/Materials/MaterialInterface.cpp
--- a/Engine/Source/Runtime/Engine/Private/Materials/MaterialInterface.cpp
+++ b/Engine/Source/Runtime/Engine/Private/Materials/MaterialInterface.cpp
@@ -44,7 +44,11 @@ FMaterialRelevance UMaterialInterface::GetRelevance_Internal(const UMaterial* Ma
 	{
 		const FMaterialResource* MaterialResource = Material->GetMaterialResource(InFeatureLevel);
 		const bool bIsTranslucent = IsTranslucentBlendMode((EBlendMode)GetBlendMode());
-		const bool bIsLit = GetShadingModel() != MSM_Unlit;
+
+		// get the value from the game thread
+		EMaterialShadingModel ShadingModel = GetShadingModel_Internal();
+
+		const bool bIsLit =  ShadingModel != MSM_Unlit;
 		// Determine the material's view relevance.
 		FMaterialRelevance MaterialRelevance;
 		MaterialRelevance.bOpaque = !bIsTranslucent;
@@ -53,7 +57,8 @@ FMaterialRelevance UMaterialInterface::GetRelevance_Internal(const UMaterial* Ma
 		MaterialRelevance.bSeparateTranslucency = bIsTranslucent && Material->bEnableSeparateTranslucency;
 		MaterialRelevance.bNormalTranslucency = bIsTranslucent && !Material->bEnableSeparateTranslucency;
 		MaterialRelevance.bDisableDepthTest = bIsTranslucent && Material->bDisableDepthTest;
-		MaterialRelevance.bSubsurfaceProfile = (Material->MaterialDomain == MD_Surface) && !bIsTranslucent && (Material->GetShadingModel() == MSM_SubsurfaceProfile);
+		MaterialRelevance.bSubsurfaceProfile = (Material->MaterialDomain == MD_Surface) && !bIsTranslucent && (ShadingModel == MSM_SubsurfaceProfile);
+
 		return MaterialRelevance;
 	}
 	else
