Fixed crash with cooking Elemental map

diff --git a/Engine/Source/Runtime/Engine/Private/Animation/AnimSequence.cpp b/Engine/Source/Runtime/Engine/Private/Animation/AnimSequence.cpp
--- a/Engine/Source/Runtime/Engine/Private/Animation/AnimSequence.cpp
+++ b/Engine/Source/Runtime/Engine/Private/Animation/AnimSequence.cpp
@@ -445,7 +445,7 @@ void UAnimSequence::VerifyTrackMap()
 			AnimationTrackNames[I] = MySkeleton->GetReferenceSkeleton().GetBoneName(TrackMap.BoneTreeIndex);
 		}
 	}
-	else 
+	else if (MySkeleton != NULL)
 	{
 		int32 NumTracks = AnimationTrackNames.Num();
 		int32 NumSkeletonBone = MySkeleton->GetReferenceSkeleton().GetNum();
