#ue4 - Avoid updating the PhysicsVolume in MovementComponent::SetUpdatedComponent() if executing during registration, it will occur as part of UpdateOverlaps() in registration anyway.

diff --git a/Engine/Source/Runtime/Engine/Classes/GameFramework/MovementComponent.h b/Engine/Source/Runtime/Engine/Classes/GameFramework/MovementComponent.h
--- a/Engine/Source/Runtime/Engine/Classes/GameFramework/MovementComponent.h
+++ b/Engine/Source/Runtime/Engine/Classes/GameFramework/MovementComponent.h
@@ -310,6 +310,11 @@ public:
 	/** Snap the updated component to the plane constraint, if enabled. */
 	UFUNCTION(BlueprintCallable, Category="Components|Movement|Planar")
 	virtual void SnapUpdatedComponentToPlane();
+
+private:
+
+	/** Transient flag indicating whether we are executing OnRegister(). */
+	bool bInOnRegister;
 };
 
 
diff --git a/Engine/Source/Runtime/Engine/Private/MovementComponent.cpp b/Engine/Source/Runtime/Engine/Private/MovementComponent.cpp
--- a/Engine/Source/Runtime/Engine/Private/MovementComponent.cpp
+++ b/Engine/Source/Runtime/Engine/Private/MovementComponent.cpp
@@ -29,6 +29,7 @@ UMovementComponent::UMovementComponent(const class FPostConstructInitializePrope
 
 	bWantsInitializeComponent = true;
 	bAutoActivate = true;
+	bInOnRegister = false;
 }
 
 
@@ -49,8 +50,13 @@ void UMovementComponent::SetUpdatedComponent(UPrimitiveComponent* NewUpdatedComp
 	if ( UpdatedComponent != NULL )
 	{
 		UpdatedComponent->bShouldUpdatePhysicsVolume = true;
-		UpdatedComponent->UpdatePhysicsVolume(true);
 		UpdatedComponent->PhysicsVolumeChangedDelegate.AddUniqueDynamic(this, &UMovementComponent::PhysicsVolumeChanged);
+
+		if (!bInOnRegister)
+		{
+			// UpdateOverlaps() in component registration will take care of this.
+			UpdatedComponent->UpdatePhysicsVolume(true);
+		}
 		
 		// force ticks after movement component updates
 		UpdatedComponent->PrimaryComponentTick.AddPrerequisite(this, PrimaryComponentTick); 
@@ -67,6 +73,7 @@ void UMovementComponent::SetUpdatedComponent(UPrimitiveComponent* NewUpdatedComp
 
 void UMovementComponent::OnRegister()
 {
+	TGuardValue<bool> InOnRegisterGuard(bInOnRegister, true);
 	Super::OnRegister();
 
 	UPrimitiveComponent* NewUpdatedComponent = NULL;
