#ue4 - CharacterMovementComponent::SetUpdatedComponent() makes sure the OnComponentBeginOverlap delegate is added uniquely.

diff --git a/Engine/Source/Runtime/Engine/Private/CharacterMovementComponent.cpp b/Engine/Source/Runtime/Engine/Private/CharacterMovementComponent.cpp
--- a/Engine/Source/Runtime/Engine/Private/CharacterMovementComponent.cpp
+++ b/Engine/Source/Runtime/Engine/Private/CharacterMovementComponent.cpp
@@ -379,7 +379,7 @@ void UCharacterMovementComponent::SetUpdatedComponent(UPrimitiveComponent* NewUp
 
 	if (bEnablePhysicsInteraction)
 	{
-		UpdatedComponent->OnComponentBeginOverlap.AddDynamic(this, &UCharacterMovementComponent::CapsuleTouched);
+		UpdatedComponent->OnComponentBeginOverlap.AddUniqueDynamic(this, &UCharacterMovementComponent::CapsuleTouched);
 	}
 
 	if (bUseRVOAvoidance)
