Fix crash accessing textures on null materials

diff --git a/Engine/Source/Runtime/Engine/Private/Materials/MaterialInstance.cpp b/Engine/Source/Runtime/Engine/Private/Materials/MaterialInstance.cpp
--- a/Engine/Source/Runtime/Engine/Private/Materials/MaterialInstance.cpp
+++ b/Engine/Source/Runtime/Engine/Private/Materials/MaterialInstance.cpp
@@ -755,7 +755,10 @@ void UMaterialInstance::GetUsedTextures(TArray<UTexture*>& OutTextures, EMateria
 			if (Material)
 			{
 				const FMaterialResource* MaterialResource = Material->GetMaterialResource(FeatureLevel, QualityLevel);
-				GetTextureExpressionValues(MaterialResource, OutTextures);
+				if( MaterialResource )
+				{
+					GetTextureExpressionValues(MaterialResource, OutTextures);
+				}
 			}
 			else
 			{
