UE4 - Reflection captures fix for ES 3.1

diff --git a/Engine/Source/Runtime/Engine/Private/ReflectionCaptureComponent.cpp b/Engine/Source/Runtime/Engine/Private/ReflectionCaptureComponent.cpp
--- a/Engine/Source/Runtime/Engine/Private/ReflectionCaptureComponent.cpp
+++ b/Engine/Source/Runtime/Engine/Private/ReflectionCaptureComponent.cpp
@@ -967,7 +967,7 @@ void UReflectionCaptureComponent::PostLoad()
 	Super::PostLoad();
 
 	bool bRetainAllFeatureLevelData = GIsEditor && GMaxRHIFeatureLevel >= ERHIFeatureLevel::SM4;
-	bool bEncodedDataRequired = GMaxRHIFeatureLevel == ERHIFeatureLevel::ES2;
+	bool bEncodedDataRequired = (GMaxRHIFeatureLevel == ERHIFeatureLevel::ES2 || GMaxRHIFeatureLevel == ERHIFeatureLevel::ES3_1);
 	bool bFullDataRequired = GMaxRHIFeatureLevel >= ERHIFeatureLevel::SM4;
 
 	// If we're loading on a platform that doesn't require cooked data, attempt to load missing data from the DDC
@@ -1337,7 +1337,7 @@ void UReflectionCaptureComponent::UpdateReflectionCaptureContents(UWorld* WorldT
 #if WITH_EDITOR
 void UReflectionCaptureComponent::PreFeatureLevelChange(ERHIFeatureLevel::Type PendingFeatureLevel)
 {
-	if (PendingFeatureLevel == ERHIFeatureLevel::ES2)
+	if (PendingFeatureLevel == ERHIFeatureLevel::ES2 || PendingFeatureLevel == ERHIFeatureLevel::ES3_1)
 	{
 		// generate encoded hdr data for ES2 preview mode.
 		if (World != nullptr)
