#ue4 Exposing APlayerController::GetFocalPoint() to BPs

#github PR 485 from Rama

diff --git a/Engine/Source/Runtime/Engine/Classes/GameFramework/PlayerController.h b/Engine/Source/Runtime/Engine/Classes/GameFramework/PlayerController.h
--- a/Engine/Source/Runtime/Engine/Classes/GameFramework/PlayerController.h
+++ b/Engine/Source/Runtime/Engine/Classes/GameFramework/PlayerController.h
@@ -1411,6 +1411,7 @@ public:
 	 *  If there is a spectator Pawn, returns that Pawn's location.
 	 *  Otherwise, returns the PlayerController's spawn location (usually the last known Pawn location after it has died).
 	 */
+	UFUNCTION(BlueprintCallable, Category = "Pawn")
 	virtual FVector GetFocalLocation() const;
 
 protected:
