#ue4 Exposing PlayerId to blueprints as readonly

#github PR 484 by Rama

diff --git a/Engine/Source/Runtime/Engine/Classes/GameFramework/PlayerState.h b/Engine/Source/Runtime/Engine/Classes/GameFramework/PlayerState.h
--- a/Engine/Source/Runtime/Engine/Classes/GameFramework/PlayerState.h
+++ b/Engine/Source/Runtime/Engine/Classes/GameFramework/PlayerState.h
@@ -48,8 +48,8 @@ class ENGINE_API APlayerState : public AInfo
 	/** Previous player name.  Saved on client-side to detect player name changes. */
 	FString OldName;
 
-	/** Unique id number. */
-	UPROPERTY(replicated)
+	/** Unique net id number. Actual value varies based on current online subsystem, use it only as a guaranteed unique number per player. */
+	UPROPERTY(replicated, BlueprintReadOnly, Category=PlayerState)
 	int32 PlayerId;
 
 	/** Whether this player is currently a spectator */
