Fix case where overlap end events were being skipped when collision was being turned off.

diff --git a/Engine/Source/Runtime/Engine/Private/PrimitiveComponent.cpp b/Engine/Source/Runtime/Engine/Private/PrimitiveComponent.cpp
--- a/Engine/Source/Runtime/Engine/Private/PrimitiveComponent.cpp
+++ b/Engine/Source/Runtime/Engine/Private/PrimitiveComponent.cpp
@@ -2207,9 +2207,10 @@ void UPrimitiveComponent::UpdateOverlaps(TArray<FOverlapInfo> const* PendingOver
 		// bGenerateOverlapEvents is false or collision is disabled
 
 		// End all overlaps that exist, in case bGenerateOverlapEvents was true last tick (i.e. was just turned off)
-		for (auto CompIt = OverlappingComponents.CreateIterator(); CompIt; ++CompIt)
-		{
-			const FOverlapInfo& OtherOverlap = *CompIt;
+		// Iterate backwards since EndComponentOverlap will remove items from OverlappingComponents.
+		for (int32 OverlapIdx = OverlappingComponents.Num()-1; OverlapIdx >= 0; --OverlapIdx)
+		{		
+			const FOverlapInfo& OtherOverlap = OverlappingComponents[OverlapIdx];
 			if (OtherOverlap.OverlapInfo.Component.IsValid())
 			{
 				EndComponentOverlap(OtherOverlap, bDoNotifies, false);
