Fix windowed fullscreen behaviour, so that when specifying a resolution < desktop res, still create the fullscreen window

Make alt-tab work correctly. Maintain window resolution when entering fullscreen.

diff --git a/Engine/Source/Runtime/Engine/Private/GameViewportClient.cpp b/Engine/Source/Runtime/Engine/Private/GameViewportClient.cpp
--- a/Engine/Source/Runtime/Engine/Private/GameViewportClient.cpp
+++ b/Engine/Source/Runtime/Engine/Private/GameViewportClient.cpp
@@ -2429,7 +2429,9 @@ bool UGameViewportClient::SetDisplayConfiguration(const FIntPoint* Dimensions, E
 
 bool UGameViewportClient::HandleToggleFullscreenCommand(const TCHAR* Cmd, FOutputDevice& Ar)
 {
-	return SetDisplayConfiguration(NULL, Viewport->IsFullscreen() ? EWindowMode::Windowed : EWindowMode::Fullscreen);
+	auto FullScreenMode = (EWindowMode::Type)GetBoundFullScreenModeCVar() == EWindowMode::Fullscreen ? EWindowMode::Fullscreen : EWindowMode::WindowedFullscreen;
+	FSystemResolution::RequestResolutionChange(GSystemResolution.ResX, GSystemResolution.ResY, Viewport->IsFullscreen() ? EWindowMode::Windowed : FullScreenMode);
+	return true;
 }
 
 bool UGameViewportClient::HandleSetResCommand( const TCHAR* Cmd, FOutputDevice& Ar )
diff --git a/Engine/Source/Runtime/Engine/Private/Slate/SceneViewport.cpp b/Engine/Source/Runtime/Engine/Private/Slate/SceneViewport.cpp
--- a/Engine/Source/Runtime/Engine/Private/Slate/SceneViewport.cpp
+++ b/Engine/Source/Runtime/Engine/Private/Slate/SceneViewport.cpp
@@ -869,8 +869,20 @@ void FSceneViewport::ResizeFrame(uint32 NewSizeX, uint32 NewSizeY, EWindowMode::
 
 		if( WindowToResize.IsValid() )
 		{
-			int32 CVarValue = GetBoundFullScreenModeCVar();
 			EWindowMode::Type DesiredWindowMode = GetWindowModeType(NewWindowMode);
+
+			// If we're going into windowed fullscreen mode, we always want the window to fill the entire screen.
+			// When we calculate the scene view, we'll check the fullscreen mode and configure the screen percentage
+			// scaling so we actual render to the resolution we've been asked for.
+			if (DesiredWindowMode == EWindowMode::WindowedFullscreen)
+			{
+				FSlateRect Rect = WindowToResize->GetFullScreenInfo();
+				if (Rect.IsValid())
+				{
+					NewSizeX = Rect.GetSize().X;
+					NewSizeY = Rect.GetSize().Y;
+				}
+			}
 			
 			// Avoid resizing if nothing changes.
 			bool bNeedsResize = SizeX != NewSizeX || SizeY != NewSizeY || NewWindowMode != DesiredWindowMode || DesiredWindowMode != WindowToResize->GetWindowMode();
@@ -930,12 +942,9 @@ void FSceneViewport::ResizeFrame(uint32 NewSizeX, uint32 NewSizeY, EWindowMode::
 
 				LockMouseToViewport(!CurrentReplyState.ShouldReleaseMouseLock());
 
-				int32 NewWindowSizeX = NewSizeX;
-				int32 NewWindowSizeY = NewSizeY;
-
-				WindowToResize->Resize( FVector2D(NewWindowSizeX, NewWindowSizeY) );
+				WindowToResize->Resize(FVector2D(NewSizeX, NewSizeY));
 
-				ResizeViewport(NewWindowSizeX, NewWindowSizeY, NewWindowMode, InPosX, InPosY);
+				ResizeViewport(NewSizeX, NewSizeY, NewWindowMode, InPosX, InPosY);
 			}
 			UCanvas::UpdateAllCanvasSafeZoneData();
 		}		
diff --git a/Engine/Source/Runtime/Engine/Private/UnrealEngine.cpp b/Engine/Source/Runtime/Engine/Private/UnrealEngine.cpp
--- a/Engine/Source/Runtime/Engine/Private/UnrealEngine.cpp
+++ b/Engine/Source/Runtime/Engine/Private/UnrealEngine.cpp
@@ -318,11 +318,13 @@ void SystemResolutionSinkCallback()
 
 		if( GSystemResolution.ResX != ResX ||
 			GSystemResolution.ResY != ResY ||
-			GSystemResolution.WindowMode != WindowMode)
+			GSystemResolution.WindowMode != WindowMode ||
+			GSystemResolution.bForceRefresh)
 		{
 			GSystemResolution.ResX = ResX;
 			GSystemResolution.ResY = ResY;
 			GSystemResolution.WindowMode = WindowMode;
+			GSystemResolution.bForceRefresh = false;
 
 			if(GEngine && GEngine->GameViewport && GEngine->GameViewport->ViewportFrame)
 			{
diff --git a/Engine/Source/Runtime/Engine/Public/UnrealEngine.h b/Engine/Source/Runtime/Engine/Public/UnrealEngine.h
--- a/Engine/Source/Runtime/Engine/Public/UnrealEngine.h
+++ b/Engine/Source/Runtime/Engine/Public/UnrealEngine.h
@@ -415,11 +415,21 @@ struct ENGINE_API FSystemResolution
 	int32 ResX;
 	int32 ResY;
 	EWindowMode::Type WindowMode;
+	bool bForceRefresh;
 
 	// Helper function for changing system resolution via the r.setres console command
 	// This function will set r.setres, which will trigger a resolution change later on
 	// when the console variable sinks are called
 	static void RequestResolutionChange(int32 InResX, int32 InResY, EWindowMode::Type InWindowMode);
+
+	FSystemResolution()
+		: ResX(0)
+		, ResY(0)
+		, WindowMode(EWindowMode::Windowed)
+		, bForceRefresh(false)
+	{
+
+	}
 };
 
 ENGINE_API extern FSystemResolution GSystemResolution;
\ No newline at end of file
diff --git a/Engine/Source/Runtime/SlateRHIRenderer/Private/SlateRHIRenderer.cpp b/Engine/Source/Runtime/SlateRHIRenderer/Private/SlateRHIRenderer.cpp
--- a/Engine/Source/Runtime/SlateRHIRenderer/Private/SlateRHIRenderer.cpp
+++ b/Engine/Source/Runtime/SlateRHIRenderer/Private/SlateRHIRenderer.cpp
@@ -358,20 +358,10 @@ void FSlateRHIRenderer::RestoreSystemResolution(const TSharedRef<SWindow> InWind
 {
 	if (!GIsEditor)
 	{
-		uint32 ResX = GSystemResolution.ResX;
-		uint32 ResY = GSystemResolution.ResY;
-		
-		if( InWindow->GetWindowMode() == EWindowMode::WindowedFullscreen )
-		{
-			FViewportInfo* ViewInfo = WindowToViewportInfo.FindRef( &InWindow.Get() );
-			if( ViewInfo )
-			{
-				ResX = ViewInfo->Width;
-				ResY = ViewInfo->Height;
-			}
-		}
-		FSystemResolution::RequestResolutionChange(ResX, ResY, GSystemResolution.WindowMode);
-		GSystemResolution.ResX = GSystemResolution.ResY = 0;
+		// Force the window system to resize the active viewport, even though nothing might have appeared to change.
+		// On windows, DXGI might change the window resolution behind our backs when we alt-tab out. This will make
+		// sure that we are actually in the resolution we think we are.
+		GSystemResolution.bForceRefresh = true;
 	}
 }
 
