"Root Motion From Everything" - Root motion can now be taken from anim assets other than montages when the anim blueprint is set to "Root Motion From Everything" mode

diff --git a/Engine/Source/Editor/AnimGraph/Private/AnimPreviewInstance.cpp b/Engine/Source/Editor/AnimGraph/Private/AnimPreviewInstance.cpp
--- a/Engine/Source/Editor/AnimGraph/Private/AnimPreviewInstance.cpp
+++ b/Engine/Source/Editor/AnimGraph/Private/AnimPreviewInstance.cpp
@@ -11,7 +11,7 @@ UAnimPreviewInstance::UAnimPreviewInstance(const class FPostConstructInitializeP
 	, bForceRetargetBasePose(false)
 #endif
 {
-	
+	RootMotionMode = ERootMotionMode::RootMotionFromEverything;
 }
 
 void UAnimPreviewInstance::NativeInitializeAnimation()
diff --git a/Engine/Source/Editor/UnrealEd/Private/Animation/DebugSkelMeshComponent.cpp b/Engine/Source/Editor/UnrealEd/Private/Animation/DebugSkelMeshComponent.cpp
--- a/Engine/Source/Editor/UnrealEd/Private/Animation/DebugSkelMeshComponent.cpp
+++ b/Engine/Source/Editor/UnrealEd/Private/Animation/DebugSkelMeshComponent.cpp
@@ -190,23 +190,29 @@ float WrapInRange(float StartVal, float MinVal, float MaxVal)
 
 void UDebugSkelMeshComponent::ConsumeRootMotion(const FVector& FloorMin, const FVector& FloorMax)
 {
+	//Extract root motion regardless of where we use it so that we don't hit
+	//problems with it building up in the instance
+
+	FRootMotionMovementParams ExtractedRootMotion;
+
+	if (UAnimInstance* AnimInst = GetAnimInstance())
+	{
+		ExtractedRootMotion = AnimInst->ConsumeExtractedRootMotion();
+	}
+
 	if (bPreviewRootMotion)
 	{
-		if (UAnimInstance* AnimInst = GetAnimInstance())
+		if (ExtractedRootMotion.bHasRootMotion)
 		{
-			FRootMotionMovementParams ExtractedRootMotion = AnimInst->ConsumeExtractedRootMotion();
-			if (ExtractedRootMotion.bHasRootMotion)
-			{
-				AddLocalTransform(ExtractedRootMotion.RootMotionTransform);
-
-				//Handle moving component so that it stays within the editor floor
-				FTransform CurrentTransform = GetRelativeTransform();
-				FVector Trans = CurrentTransform.GetTranslation();
-				Trans.X = WrapInRange(Trans.X, FloorMin.X, FloorMax.X);
-				Trans.Y = WrapInRange(Trans.Y, FloorMin.Y, FloorMax.Y);
-				CurrentTransform.SetTranslation(Trans);
-				SetRelativeTransform(CurrentTransform);
-			}
+			AddLocalTransform(ExtractedRootMotion.RootMotionTransform);
+
+			//Handle moving component so that it stays within the editor floor
+			FTransform CurrentTransform = GetRelativeTransform();
+			FVector Trans = CurrentTransform.GetTranslation();
+			Trans.X = WrapInRange(Trans.X, FloorMin.X, FloorMax.X);
+			Trans.Y = WrapInRange(Trans.Y, FloorMin.Y, FloorMax.Y);
+			CurrentTransform.SetTranslation(Trans);
+			SetRelativeTransform(CurrentTransform);
 		}
 	}
 	else
diff --git a/Engine/Source/Runtime/Engine/Classes/Animation/AnimComposite.h b/Engine/Source/Runtime/Engine/Classes/Animation/AnimComposite.h
--- a/Engine/Source/Runtime/Engine/Classes/Animation/AnimComposite.h
+++ b/Engine/Source/Runtime/Engine/Classes/Animation/AnimComposite.h
@@ -19,9 +19,15 @@ public:
 	UPROPERTY()
 	struct FAnimTrack AnimationTrack;
 
+	// Begin UAnimSequenceBase interface
+	ENGINE_API virtual void OnAssetPlayerTickedInternal(FAnimAssetTickContext &Context, const float PreviousTime, const float MoveDelta, const FAnimTickRecord &Instance, class UAnimInstance* InstanceOwner) const override;
+	// End UAnimSequenceBase interface
+
+	// Begin UAnimSequence interface
 #if WITH_EDITOR
 	virtual bool GetAllAnimationSequencesReferred(TArray<UAnimSequence*>& AnimationSequences) override;
 	virtual void ReplaceReferredAnimations(const TMap<UAnimSequence*, UAnimSequence*>& ReplacementMap) override;
 #endif
+	// End UAnimSequence interface
 };
 
diff --git a/Engine/Source/Runtime/Engine/Classes/Animation/AnimCompositeBase.h b/Engine/Source/Runtime/Engine/Classes/Animation/AnimCompositeBase.h
--- a/Engine/Source/Runtime/Engine/Classes/Animation/AnimCompositeBase.h
+++ b/Engine/Source/Runtime/Engine/Classes/Animation/AnimCompositeBase.h
@@ -158,6 +158,9 @@ struct FAnimTrack
 
 	ENGINE_API int32 GetTrackAdditiveType() const;
 
+	/** Returns whether any of the animation sequences this track uses has root motion */
+	bool HasRootMotion() const;
+
 	/** 
 	 * Given a Track delta position [StartTrackPosition, EndTrackPosition]
 	 * See if any AnimSegment overlaps any of it, and if it does, break it up into RootMotionExtractionPieces.
@@ -189,5 +192,8 @@ class UAnimCompositeBase : public UAnimSequenceBase
 	/** Set Sequence Length */
 	ENGINE_API void SetSequenceLength(float InSequenceLength);
 #endif
+
+	// Extracts root motion from the supplied FAnimTrack between the Start End range specified
+	ENGINE_API void ExtractRootMotionFromTrack(const FAnimTrack &SlotAnimTrack, float StartTrackPosition, float EndTrackPosition, FRootMotionMovementParams &RootMotion) const;
 };
 
diff --git a/Engine/Source/Runtime/Engine/Classes/Animation/AnimInstance.h b/Engine/Source/Runtime/Engine/Classes/Animation/AnimInstance.h
--- a/Engine/Source/Runtime/Engine/Classes/Animation/AnimInstance.h
+++ b/Engine/Source/Runtime/Engine/Classes/Animation/AnimInstance.h
@@ -267,6 +267,10 @@ class ENGINE_API UAnimInstance : public UObject
 	UPROPERTY(transient)
 	TArray<struct FActiveVertexAnim> VertexAnims;
 
+	// Sets where this blueprint pulls Root Motion from
+	UPROPERTY(Category = RootMotion, EditDefaultsOnly)
+	TEnumAsByte<ERootMotionMode::Type> RootMotionMode;
+
 public:
 
 	// @todo document
@@ -305,6 +309,12 @@ public:
 	void ClearSlotNodeWeights();
 	bool IsActiveSlotNode(FName SlotNodeName) const;
 
+	// Allow slot nodes to store off their root motion weight during ticking
+	void UpdateSlotRootMotionWeight(FName SlotNodeName, float Weight);
+	// Get the root motion weight for the montage slot
+	float GetSlotRootMotionWeight(FName SlotNodeName) const;
+
+
 	// kismet event functions
 
 	UFUNCTION(BlueprintCallable, Category = "Animation")
@@ -577,6 +587,9 @@ public:
 
 	TMap<FName, float> ActiveSlotWeights;
 
+	// Mapping from slot name to weighting for that root motion
+	TMap<FName, float> ActiveSlotRootMotionWeights;
+
 #if WITH_EDITORONLY_DATA
 	// Maximum playback position ever reached (only used when debugging in Persona)
 	double LifeTimer;
diff --git a/Engine/Source/Runtime/Engine/Classes/Animation/AnimMontage.h b/Engine/Source/Runtime/Engine/Classes/Animation/AnimMontage.h
--- a/Engine/Source/Runtime/Engine/Classes/Animation/AnimMontage.h
+++ b/Engine/Source/Runtime/Engine/Classes/Animation/AnimMontage.h
@@ -243,7 +243,7 @@ public:
 	/** Simulate is same as Advance, but without calling any events or touching any of the instance data. So it performs a simulation of advancing the timeline. 
 	 * @fixme laurent can we make Advance use that, so we don't have 2 code paths which risk getting out of sync? */
 	bool SimulateAdvance(float DeltaTime, float& InOutPosition, struct FRootMotionMovementParams & OutRootMotionParams) const;
-	void Advance(float DeltaTime, struct FRootMotionMovementParams * OutRootMotionParams);
+	void Advance(float DeltaTime, struct FRootMotionMovementParams * OutRootMotionParams, bool bBlendRootMotion);
 
 	FName GetCurrentSection() const;
 	FName GetNextSection() const;
@@ -288,16 +288,16 @@ class UAnimMontage : public UAnimCompositeBase
 	UPROPERTY()
 	TArray<struct FBranchingPoint> BranchingPoints;
 
-	/** If this is on, it will allow extracting root motion translation **/
-	UPROPERTY(EditAnywhere, Category=RootMotion)
+	/** If this is on, it will allow extracting root motion translation. DEPRECATED in 4.5 root motion is controlled by anim sequences **/
+	UPROPERTY()
 	bool bEnableRootMotionTranslation;
 
-	/** If this is on, it will allow extracting root motion rotation **/
-	UPROPERTY(EditAnywhere, Category=RootMotion)
+	/** If this is on, it will allow extracting root motion rotation. DEPRECATED in 4.5 root motion is controlled by anim sequences **/
+	UPROPERTY()
 	bool bEnableRootMotionRotation;
 
-	/** Root Bone will be locked to that position when extracting root motion. **/
-	UPROPERTY(EditAnywhere, Category=RootMotion)
+	/** Root Bone will be locked to that position when extracting root motion. DEPRECATED in 4.5 root motion is controlled by anim sequences **/
+	UPROPERTY()
 	TEnumAsByte<ERootMotionRootLock::Type> RootMotionRootLock;
 
 #if WITH_EDITORONLY_DATA
@@ -375,6 +375,9 @@ public:
 	/** Prototype function to get animation data - this will need rework */
 	const FAnimTrack* GetAnimationData(FName SlotName) const;
 
+	/** Returns whether the anim sequences this montage have root motion enabled */
+	virtual bool HasRootMotion() const override;
+
 	/** Extract RootMotion Transform from a contiguous Track position range.
 	 * *CONTIGUOUS* means that if playing forward StartTractPosition < EndTrackPosition.
 	 * No wrapping over if looping. No jumping across different sections.
diff --git a/Engine/Source/Runtime/Engine/Classes/Animation/AnimSequence.h b/Engine/Source/Runtime/Engine/Classes/Animation/AnimSequence.h
--- a/Engine/Source/Runtime/Engine/Classes/Animation/AnimSequence.h
+++ b/Engine/Source/Runtime/Engine/Classes/Animation/AnimSequence.h
@@ -496,6 +496,18 @@ class UAnimSequence : public UAnimSequenceBase
 	UPROPERTY(EditAnywhere, AssetRegistrySearchable, Category=Animation)
 	FName RetargetSource;
 
+	/** If this is on, it will allow extracting of root motion **/
+	UPROPERTY(EditAnywhere, Category = RootMotion)
+	bool bEnableRootMotion;
+
+	/** Root Bone will be locked to that position when extracting root motion.**/
+	UPROPERTY(EditAnywhere, Category = RootMotion)
+	TEnumAsByte<ERootMotionRootLock::Type> RootMotionRootLock;
+
+	/** Have we copied root motion settings from an owning montage */
+	UPROPERTY()
+	bool bRootMotionSettingsCopiedFromMontage;
+
 #if WITH_EDITORONLY_DATA
 	/** Saved version number with CompressAnimations commandlet. To help with doing it in multiple passes. */
 	UPROPERTY()
@@ -549,10 +561,20 @@ class UAnimSequence : public UAnimSequenceBase
 	// End of UAnimationAsset interface
 
 	// Begin UAnimSequenceBase interface
-	ENGINE_API virtual void ExtractRootTrack(float Pos, FTransform& RootTransform, const FBoneContainer* RequiredBones) const override;
+	ENGINE_API virtual void OnAssetPlayerTickedInternal(FAnimAssetTickContext &Context, const float PreviousTime, const float MoveDelta, const FAnimTickRecord &Instance, class UAnimInstance* InstanceOwner) const override;
 	ENGINE_API virtual void UpgradeMorphTargetCurves() override;
+	ENGINE_API virtual bool HasRootMotion() const { return bEnableRootMotion; }
 	// End UAnimSequenceBase interface
 
+	// Extract Root Motion transform from the animation
+	FTransform ExtractRootMotion(const float & StartTime, const float & DeltaTime, const float bAllowLooping) const;
+
+	// Extract Root Motion transform from a contiguous position range (no looping)
+	FTransform ExtractRootMotionFromRange(float StartTrackPosition, float EndTrackPosition) const;
+
+	// Extract the transform from the root track for the given animation position
+	FTransform ExtractRootTrackTransform(float Pos, const FBoneContainer * RequiredBones) const;
+
 	// Begin Transform related functions 
 
 	/** 
@@ -583,7 +605,7 @@ private:
 	 * @param	RequiredBones		BoneContainer
 	 * @param	ExtractionContext	Extraction Context to access root motion extraction information.
 	 */
-	void ResetRootBoneForRootMotion(FTransformArrayA2 & BoneTransforms, const FBoneContainer& RequiredBones, const FAnimExtractContext& ExtractionContext) const;
+	void ResetRootBoneForRootMotion(FTransformArrayA2 & BoneTransforms, const FBoneContainer & RequiredBones, ERootMotionRootLock::Type RootMotionRootLock) const;
 	
 	/** 
 	 * Retarget a single bone transform, to apply right after extraction.
diff --git a/Engine/Source/Runtime/Engine/Classes/Animation/AnimSequenceBase.h b/Engine/Source/Runtime/Engine/Classes/Animation/AnimSequenceBase.h
--- a/Engine/Source/Runtime/Engine/Classes/Animation/AnimSequenceBase.h
+++ b/Engine/Source/Runtime/Engine/Classes/Animation/AnimSequenceBase.h
@@ -412,10 +412,10 @@ class UAnimSequenceBase : public UAnimationAsset
 	virtual float GetMaxCurrentTime() override { return SequenceLength; }
 	// End of UAnimationAsset interface
 
-protected:
-	virtual void ExtractRootTrack(float Pos, FTransform& RootTransform, const FBoneContainer* RequiredBones) const;
+	virtual void OnAssetPlayerTickedInternal(FAnimAssetTickContext &Context, const float PreviousTime, const float MoveDelta, const FAnimTickRecord &Instance, class UAnimInstance* InstanceOwner) const;
+
+	virtual bool HasRootMotion() const { return false; }
 
-public:
 	virtual void Serialize(FArchive& Ar) override;
 
 #if WITH_EDITOR
diff --git a/Engine/Source/Runtime/Engine/Classes/Animation/AnimationAsset.h b/Engine/Source/Runtime/Engine/Classes/Animation/AnimationAsset.h
--- a/Engine/Source/Runtime/Engine/Classes/Animation/AnimationAsset.h
+++ b/Engine/Source/Runtime/Engine/Classes/Animation/AnimationAsset.h
@@ -90,6 +90,24 @@ namespace ERootMotionRootLock
 	};
 }
 
+UENUM()
+namespace ERootMotionMode
+{
+	enum Type
+	{
+		// Leave root motion in animation.
+		NoRootMotionExtraction,
+
+		// Extract root motion but do not apply it.
+		IgnoreRootMotion,
+
+		// Root motion is taken from all animations contributing to the final pose, not suitable for network multiplayer setups
+		RootMotionFromEverything,
+
+		// Root motion is only taken from montages, suitable for network multiplayer setups
+		RootMotionFromMontagesOnly,
+	};
+}
 
 /** Animation Extraction Context */
 USTRUCT()
@@ -97,43 +115,29 @@ struct FAnimExtractContext
 {
 	GENERATED_USTRUCT_BODY()
 
-	/** Is Root Motion Translation being extracted? */
+	/** Is Root Motion being extracted? */
 	UPROPERTY()
-	bool bExtractRootMotionTranslation;
-
-	/** Is Root Motion Rotation being extracted? */
-	UPROPERTY()
-	bool bExtractRootMotionRotation;
+	bool bExtractRootMotion;
 
 	/** Position in animation to extract pose from */
 	UPROPERTY()
 	float CurrentTime;
 
-	/** Root Motion Root Bone Lock option. **/
-	UPROPERTY()
-	TEnumAsByte<ERootMotionRootLock::Type> RootMotionRootLock;
-
 	FAnimExtractContext()
-		: bExtractRootMotionTranslation(false)
-		, bExtractRootMotionRotation(false)
+		: bExtractRootMotion(false)
 		, CurrentTime(0.f)
-		, RootMotionRootLock(ERootMotionRootLock::RefPose)
 	{
 	}
 
 	FAnimExtractContext(float InCurrentTime)
-		: bExtractRootMotionTranslation(false)
-		, bExtractRootMotionRotation(false)
+		: bExtractRootMotion(false)
 		, CurrentTime(InCurrentTime)
-		, RootMotionRootLock(ERootMotionRootLock::RefPose)
 	{
 	}
 
-	FAnimExtractContext(float InCurrentTime, bool InbExtractRootMotionTranslation, bool InbExtractRootMotionRotation, ERootMotionRootLock::Type InRootMotionRootLock)
-		: bExtractRootMotionTranslation(InbExtractRootMotionTranslation)
-		, bExtractRootMotionRotation(InbExtractRootMotionRotation)
+	FAnimExtractContext(float InCurrentTime, bool InbExtractRootMotion)
+		: bExtractRootMotion(InbExtractRootMotion)
 		, CurrentTime(InCurrentTime)
-		, RootMotionRootLock(InRootMotionRootLock)
 	{
 	}
 };
@@ -238,15 +242,19 @@ struct FRootMotionMovementParams
 {
 	GENERATED_USTRUCT_BODY()
 
-		UPROPERTY()
-		bool bHasRootMotion;
+	UPROPERTY()
+	bool bHasRootMotion;
 
 	UPROPERTY()
-		FTransform RootMotionTransform;
+	float BlendWeight;
+
+	UPROPERTY()
+	FTransform RootMotionTransform;
 
 	FRootMotionMovementParams()
 		: bHasRootMotion(false)
 		, RootMotionTransform(FTransform::Identity)
+		, BlendWeight(0.f)
 	{
 	}
 
@@ -254,6 +262,7 @@ struct FRootMotionMovementParams
 	{
 		bHasRootMotion = true;
 		RootMotionTransform = InTransform;
+		BlendWeight = 1.f;
 	}
 
 	void Accumulate(const FTransform& InTransform)
@@ -276,9 +285,43 @@ struct FRootMotionMovementParams
 		}
 	}
 
+	void AccumulateWithBlend(const FTransform& InTransform, float InBlendWeight)
+	{
+		const ScalarRegister VBlendWeight(InBlendWeight);
+		if (bHasRootMotion)
+		{
+			RootMotionTransform.AccumulateWithShortestRotation(InTransform, VBlendWeight);
+			BlendWeight += InBlendWeight;
+		}
+		else
+		{
+			Set(InTransform * VBlendWeight);
+			BlendWeight = InBlendWeight;
+		}
+		RootMotionTransform.NormalizeRotation();
+	}
+
+	void AccumulateWithBlend(const FRootMotionMovementParams & MovementParams, float InBlendWeight)
+	{
+		if (MovementParams.bHasRootMotion)
+		{
+			AccumulateWithBlend(MovementParams.RootMotionTransform, InBlendWeight);
+		}
+	}
+
 	void Clear()
 	{
 		bHasRootMotion = false;
+		BlendWeight = 0.f;
+	}
+
+	void MakeUpToFullWeight()
+	{
+		float WeightLeft = FMath::Max(1.f - BlendWeight, 0.f);
+		if (WeightLeft > KINDA_SMALL_NUMBER)
+		{
+			AccumulateWithBlend(FTransform(), WeightLeft);
+		}
 	}
 };
 
@@ -286,8 +329,9 @@ struct FRootMotionMovementParams
 struct FAnimAssetTickContext
 {
 public:
-	FAnimAssetTickContext(float InDeltaTime)
+	FAnimAssetTickContext(float InDeltaTime, ERootMotionMode::Type InRootMotionMode)
 		: DeltaTime(InDeltaTime)
+		, RootMotionMode(InRootMotionMode)
 		, SyncPoint(0.0f)
 		, bIsLeader(true)
 	{
@@ -332,6 +376,12 @@ public:
 		return IsLeader();
 	}
 
+	//Root Motion accumulated from this tick context
+	FRootMotionMovementParams RootMotionMovementParams;
+
+	// The root motion mode of the owning AnimInstance
+	ERootMotionMode::Type RootMotionMode;
+
 private:
 	float DeltaTime;
 
diff --git a/Engine/Source/Runtime/Engine/Classes/GameFramework/Character.h b/Engine/Source/Runtime/Engine/Classes/GameFramework/Character.h
--- a/Engine/Source/Runtime/Engine/Classes/GameFramework/Character.h
+++ b/Engine/Source/Runtime/Engine/Classes/GameFramework/Character.h
@@ -671,6 +671,11 @@ public:
 	UFUNCTION(BlueprintCallable, Category=Animation)
 	bool IsPlayingRootMotion() const;
 
+	/** true if we are playing Root Motion right now, through a Montage with RootMotionMode == ERootMotionMode::RootMotionFromMontagesOnly.
+	 * This means code path for networked root motion is enabled. */
+	UFUNCTION(BlueprintCallable, Category = Animation)
+	bool IsPlayingNetworkedRootMotionMontage() const;
+
 	/** Called on the actor right before replication occurs */
 	virtual void PreReplication( IRepChangedPropertyTracker & ChangedPropertyTracker ) override;
 };
diff --git a/Engine/Source/Runtime/Engine/Private/Animation/AnimComposite.cpp b/Engine/Source/Runtime/Engine/Private/Animation/AnimComposite.cpp
--- a/Engine/Source/Runtime/Engine/Private/Animation/AnimComposite.cpp
+++ b/Engine/Source/Runtime/Engine/Private/Animation/AnimComposite.cpp
@@ -25,4 +25,11 @@ void UAnimComposite::ReplaceReferredAnimations(const TMap<UAnimSequence*, UAnimS
 {
 	AnimationTrack.ReplaceReferredAnimations(ReplacementMap);
 }
-#endif
\ No newline at end of file
+#endif
+
+void UAnimComposite::OnAssetPlayerTickedInternal(FAnimAssetTickContext &Context, const float PreviousTime, const float MoveDelta, const FAnimTickRecord &Instance, class UAnimInstance* InstanceOwner) const
+{
+	Super::OnAssetPlayerTickedInternal(Context, PreviousTime, MoveDelta, Instance, InstanceOwner);
+
+	ExtractRootMotionFromTrack(AnimationTrack, PreviousTime, PreviousTime + MoveDelta, Context.RootMotionMovementParams);
+}
\ No newline at end of file
diff --git a/Engine/Source/Runtime/Engine/Private/Animation/AnimCompositeBase.cpp b/Engine/Source/Runtime/Engine/Private/Animation/AnimCompositeBase.cpp
--- a/Engine/Source/Runtime/Engine/Private/Animation/AnimCompositeBase.cpp
+++ b/Engine/Source/Runtime/Engine/Private/Animation/AnimCompositeBase.cpp
@@ -187,6 +187,18 @@ void FAnimSegment::GetRootMotionExtractionStepsForTrackRange(TArray<FRootMotionE
 	}
 }
 
+bool FAnimTrack::HasRootMotion() const
+{
+	for (const FAnimSegment& AnimSegment : AnimSegments)
+	{
+		if (AnimSegment.AnimReference->HasRootMotion())
+		{
+			return true;
+		}
+	}
+	return false;
+}
+
 /** 
  * Given a Track delta position [StartTrackPosition, EndTrackPosition]
  * See if any AnimSegment overlaps any of it, and if any do, break them up into a sequence of FRootMotionExtractionStep.
@@ -421,15 +433,28 @@ void UAnimCompositeBase::SetSequenceLength(float InSequenceLength)
 }
 #endif
 
-/*
-void UAnimCompositeBase::AddAnimSegment( FAnimTrack & Track, FAnimSegment& NewSegment )
+void UAnimCompositeBase::ExtractRootMotionFromTrack(const FAnimTrack &SlotAnimTrack, float StartTrackPosition, float EndTrackPosition, FRootMotionMovementParams &RootMotion) const
 {
+	TArray<FRootMotionExtractionStep> RootMotionExtractionSteps;
+	SlotAnimTrack.GetRootMotionExtractionStepsForTrackRange(RootMotionExtractionSteps, StartTrackPosition, EndTrackPosition);
 
-}
-
-bool UAnimCompositeBase::DeleteAnimSegment( FAnimTrack & Track, FAnimSegment& SegmentToDelete )
-{
+	UE_LOG(LogRootMotion, Log, TEXT("\tUAnimMontage::ExtractRootMotionForTrackRange, NumSteps: %d, StartTrackPosition: %.3f, EndTrackPosition: %.3f"),
+		RootMotionExtractionSteps.Num(), StartTrackPosition, EndTrackPosition);
 
+	// Go through steps sequentially, extract root motion, and accumulate it.
+	// This has to be done in order so root motion translation & rotation is applied properly (as translation is relative to rotation)
+	for (int32 StepIndex = 0; StepIndex < RootMotionExtractionSteps.Num(); StepIndex++)
+	{
+		const FRootMotionExtractionStep & CurrentStep = RootMotionExtractionSteps[StepIndex];
+		if (CurrentStep.AnimSequence->bEnableRootMotion)
+		{
+			FTransform DeltaTransform = CurrentStep.AnimSequence->ExtractRootMotionFromRange(CurrentStep.StartPosition, CurrentStep.EndPosition);
+			RootMotion.Accumulate(DeltaTransform);
+		
+			UE_LOG(LogRootMotion, Log, TEXT("\t\tCurrentStep: %d, StartPos: %.3f, EndPos: %.3f, Anim: %s DeltaTransform Translation: %s, Rotation: %s"),
+				StepIndex, CurrentStep.StartPosition, CurrentStep.EndPosition, *CurrentStep.AnimSequence->GetName(),
+				*DeltaTransform.GetTranslation().ToCompactString(), *DeltaTransform.GetRotation().Rotator().ToCompactString());
+		}
+	}
 }
-*/
 
diff --git a/Engine/Source/Runtime/Engine/Private/Animation/AnimInstance.cpp b/Engine/Source/Runtime/Engine/Private/Animation/AnimInstance.cpp
--- a/Engine/Source/Runtime/Engine/Private/Animation/AnimInstance.cpp
+++ b/Engine/Source/Runtime/Engine/Private/Animation/AnimInstance.cpp
@@ -65,6 +65,7 @@ UAnimInstance::UAnimInstance(const class FPostConstructInitializeProperties& PCI
 	: Super(PCIP)
 {
 	RootNode = NULL;
+	RootMotionMode = ERootMotionMode::RootMotionFromMontagesOnly;
 }
 
 void UAnimInstance::MakeSequenceTickRecord(FAnimTickRecord& TickRecord, class UAnimSequenceBase* Sequence, bool bLooping, float PlayRate, float FinalBlendWeight, float& CurrentTime) const
@@ -94,7 +95,7 @@ void UAnimInstance::SequenceAdvanceImmediate(UAnimSequenceBase* Sequence, bool b
 	FAnimTickRecord TickRecord;
 	MakeSequenceTickRecord(TickRecord, Sequence, bLooping, PlayRate, /*FinalBlendWeight=*/ 1.0f, CurrentTime);
 
-	FAnimAssetTickContext TickContext(DeltaSeconds);
+	FAnimAssetTickContext TickContext(DeltaSeconds, RootMotionMode);
 	TickRecord.SourceAsset->TickAssetPlayerInstance(TickRecord, this, TickContext);
 }
 
@@ -103,7 +104,7 @@ void UAnimInstance::BlendSpaceAdvanceImmediate(class UBlendSpaceBase* BlendSpace
 	FAnimTickRecord TickRecord;
 	MakeBlendSpaceTickRecord(TickRecord, BlendSpace, BlendInput, BlendSampleDataCache, BlendFilter, bLooping, PlayRate, /*FinalBlendWeight=*/ 1.0f, CurrentTime);
 	
-	FAnimAssetTickContext TickContext(DeltaSeconds);
+	FAnimAssetTickContext TickContext(DeltaSeconds, RootMotionMode);
 	TickRecord.SourceAsset->TickAssetPlayerInstance(TickRecord, this, TickContext);
 }
 
@@ -133,13 +134,16 @@ void UAnimInstance::SequenceEvaluatePose(UAnimSequenceBase* Sequence, FA2Pose& P
 
 	USkeletalMeshComponent* Component = GetSkelMeshComponent();
 
+	FAnimExtractContext NewExtractContext(ExtractionContext);
+	NewExtractContext.bExtractRootMotion = RootMotionMode == ERootMotionMode::RootMotionFromEverything;
+
 	if(const UAnimSequence* AnimSequence = Cast<const UAnimSequence>(Sequence))
 	{
 		FAnimationRuntime::GetPoseFromSequence(
 			AnimSequence,
 			RequiredBones,
 			/*out*/ Pose.Bones, 
-			ExtractionContext);
+			NewExtractContext);
 	}
 	else if(const UAnimComposite* Composite = Cast<const UAnimComposite>(Sequence))
 	{
@@ -147,7 +151,7 @@ void UAnimInstance::SequenceEvaluatePose(UAnimSequenceBase* Sequence, FA2Pose& P
 			Composite->AnimationTrack, 
 			RequiredBones, 
 			/*out*/ Pose.Bones,
-			ExtractionContext);
+			NewExtractContext);
 	}
 	else
 	{
@@ -402,7 +406,7 @@ void UAnimInstance::UpdateAnimation(float DeltaSeconds)
 			const int32 GroupLeaderIndex = FMath::Max(SyncGroup.GroupLeaderIndex, 0);
 
 			// Tick the group leader
-			FAnimAssetTickContext TickContext(DeltaSeconds);
+			FAnimAssetTickContext TickContext(DeltaSeconds, RootMotionMode);
 			FAnimTickRecord& GroupLeader = SyncGroup.ActivePlayers[GroupLeaderIndex];
 			GroupLeader.SourceAsset->TickAssetPlayerInstance(GroupLeader, this, TickContext);
 
@@ -420,6 +424,11 @@ void UAnimInstance::UpdateAnimation(float DeltaSeconds)
 					}
 				}
 			}
+
+			if (RootMotionMode == ERootMotionMode::RootMotionFromEverything && TickContext.RootMotionMovementParams.bHasRootMotion)
+			{
+				ExtractedRootMotion.AccumulateWithBlend(TickContext.RootMotionMovementParams.RootMotionTransform, GroupLeader.EffectiveBlendWeight);
+			}
 		}
 	}
 
@@ -427,14 +436,22 @@ void UAnimInstance::UpdateAnimation(float DeltaSeconds)
 	for (int32 TickIndex = 0; TickIndex < UngroupedActivePlayers.Num(); ++TickIndex)
 	{
 		const FAnimTickRecord& AssetPlayerToTick = UngroupedActivePlayers[TickIndex];
-		FAnimAssetTickContext TickContext(DeltaSeconds);
+		FAnimAssetTickContext TickContext(DeltaSeconds, RootMotionMode);
 		AssetPlayerToTick.SourceAsset->TickAssetPlayerInstance(AssetPlayerToTick, this, TickContext);
+		if (RootMotionMode == ERootMotionMode::RootMotionFromEverything && TickContext.RootMotionMovementParams.bHasRootMotion)
+		{
+			ExtractedRootMotion.AccumulateWithBlend(TickContext.RootMotionMovementParams.RootMotionTransform, AssetPlayerToTick.EffectiveBlendWeight);
+		}
 	}
 
 	// update montage should run in game thread
 	// if we do multi threading, make sure this stays in game thread
 	Montage_Advance(DeltaSeconds);
 
+	// We may have just partially blended root motion, so make it up to 1 by
+	// blending in identity too
+	ExtractedRootMotion.MakeUpToFullWeight();
+
 	// now trigger Notifies
 	TriggerAnimNotifies(DeltaSeconds);
 
@@ -1264,7 +1281,7 @@ void UAnimInstance::SlotEvaluatePose(FName SlotNodeName, const FA2Pose & SourceP
 
 			// Extract pose from Track
 			UAnimMontage const * const MontageAsset = MontageInstance->Montage;
-			FAnimExtractContext ExtractionContext(MontageInstance->GetPosition(), MontageAsset->bEnableRootMotionTranslation, MontageAsset->bEnableRootMotionRotation, MontageAsset->RootMotionRootLock);
+			FAnimExtractContext ExtractionContext(MontageInstance->GetPosition(), MontageInstance->Montage->HasRootMotion() && RootMotionMode != ERootMotionMode::NoRootMotionExtraction);
 			FAnimationRuntime::GetPoseFromAnimTrack(*AnimTrack, RequiredBones, NewPose.Pose.Bones, ExtractionContext);
 
 			TotalWeight += MontageInstance->Weight;
@@ -1383,6 +1400,7 @@ void UAnimInstance::RegisterSlotNode(FName SlotNodeName)
 	}
 
 	ActiveSlotWeights.Add(SlotNodeName, 0.f);
+	ActiveSlotRootMotionWeights.Add(SlotNodeName, 0.f);
 }
 
 void UAnimInstance::UpdateSlotNodeWeight(FName SlotNodeName, float Weight)
@@ -1401,6 +1419,12 @@ void UAnimInstance::ClearSlotNodeWeights()
 		float& Weight = Iter.Value();
 		Weight = 0.f;
 	}
+
+	for (auto Iter = ActiveSlotRootMotionWeights.CreateIterator(); Iter; ++Iter)
+	{
+		float& Weight = Iter.Value();
+		Weight = 0.f;
+	}
 }
 
 bool UAnimInstance::IsActiveSlotNode(FName SlotNodeName) const
@@ -1409,6 +1433,21 @@ bool UAnimInstance::IsActiveSlotNode(FName SlotNodeName) const
 	return ( SlotNodeWeight && (*SlotNodeWeight > ZERO_ANIMWEIGHT_THRESH) );
 }
 
+void UAnimInstance::UpdateSlotRootMotionWeight(FName SlotNodeName, float Weight)
+{
+	float* CurrentWeight = ActiveSlotRootMotionWeights.Find(SlotNodeName);
+	if (CurrentWeight)
+	{
+		*CurrentWeight += Weight;
+	}
+}
+
+float UAnimInstance::GetSlotRootMotionWeight(FName SlotNodeName) const
+{
+	const float* SlotNodeWeight = ActiveSlotRootMotionWeights.Find(SlotNodeName);
+	return SlotNodeWeight ? *SlotNodeWeight : 0.f;
+}
+
 float UAnimInstance::GetCurveValue(FName CurveName)
 {
 	float* Value = EventCurves.Find(CurveName);
@@ -1517,8 +1556,6 @@ void UAnimInstance::Montage_UpdateWeight(float DeltaSeconds)
 
 void UAnimInstance::Montage_Advance(float DeltaSeconds)
 {
-	FRootMotionMovementParams LocalExtractedRootMotion;
-
 	// go through all montage instances, and update them
 	// and make sure their weight is updated properly
 	for (int32 InstanceIndex = 0; InstanceIndex<MontageInstances.Num(); InstanceIndex++)
@@ -1528,10 +1565,22 @@ void UAnimInstance::Montage_Advance(float DeltaSeconds)
 		ensure(MontageInstance);
 		if (MontageInstance)
 		{
-			// Only use root motion from active root motion montage.
-			// we don't support blending/weighting because of networking constraints.
-			bool const bExtractRootMotion = (MontageInstance == GetRootMotionMontageInstance());
-			MontageInstance->Advance(DeltaSeconds, bExtractRootMotion ? &LocalExtractedRootMotion : NULL);
+			bool const bUsingBlendedRootMotion = RootMotionMode == ERootMotionMode::RootMotionFromEverything;
+			bool const bNoRootMotionExtraction = RootMotionMode == ERootMotionMode::NoRootMotionExtraction;
+
+			// Extract root motion if we are using blend root motion (RootMotionFromEverything) or if we are set to extract root 
+			// motion AND we are the active root motion instance. This is so we can make root motion deterministic for networking when
+			// we are not using RootMotionFromEverything
+			bool const bExtractRootMotion = bUsingBlendedRootMotion || (!bNoRootMotionExtraction && MontageInstance == GetRootMotionMontageInstance());
+
+			FRootMotionMovementParams LocalExtractedRootMotion;
+			FRootMotionMovementParams* RootMotionParams = NULL;
+			if (bExtractRootMotion)
+			{
+				RootMotionParams = (RootMotionMode != ERootMotionMode::IgnoreRootMotion) ? &ExtractedRootMotion : &LocalExtractedRootMotion;
+			}
+
+			MontageInstance->Advance(DeltaSeconds, bExtractRootMotion ? &ExtractedRootMotion : NULL, bUsingBlendedRootMotion);
 
 			if (!MontageInstance->IsValid())
 			{
@@ -1551,12 +1600,6 @@ void UAnimInstance::Montage_Advance(float DeltaSeconds)
 #endif
 		}
 	}
-
-	// If Root Motion has been extracted, store it.
-	if ((GetRootMotionMontageInstance() != NULL) && LocalExtractedRootMotion.bHasRootMotion)
-	{
-		ExtractedRootMotion.Accumulate(LocalExtractedRootMotion);
-	}
 }
 
 float UAnimInstance::PlaySlotAnimation(UAnimSequenceBase* Asset, FName SlotNodeName, float BlendInTime, float BlendOutTime, float InPlayRate)
@@ -1666,7 +1709,7 @@ float UAnimInstance::Montage_Play(UAnimMontage * MontageToPlay, float InPlayRate
 			ActiveMontagesMap.Add(MontageToPlay, NewInstance);
 
 			// If we are playing root motion, set this instance as the one providing root motion.
-			if (MontageToPlay->bEnableRootMotionTranslation || MontageToPlay->bEnableRootMotionRotation)
+			if (MontageToPlay->HasRootMotion())
 			{
 				RootMotionMontageInstance = NewInstance;
 			}
diff --git a/Engine/Source/Runtime/Engine/Private/Animation/AnimMontage.cpp b/Engine/Source/Runtime/Engine/Private/Animation/AnimMontage.cpp
--- a/Engine/Source/Runtime/Engine/Private/Animation/AnimMontage.cpp
+++ b/Engine/Source/Runtime/Engine/Private/Animation/AnimMontage.cpp
@@ -27,7 +27,6 @@ UAnimMontage::UAnimMontage(const class FPostConstructInitializeProperties& PCIP)
 	bAnimBranchingPointNeedsSort = true;
 	BlendInTime = 0.25f;
 	BlendOutTime = 0.25f;
-	RootMotionRootLock = ERootMotionRootLock::RefPose;
 }
 
 bool UAnimMontage::IsValidSlot(FName InSlotName) const
@@ -374,6 +373,24 @@ void UAnimMontage::PostLoad()
 
 	SortAnimBranchingPointByTime();
 
+	bool bRootMotionEnabled = bEnableRootMotionTranslation || bEnableRootMotionRotation;
+
+	if (bRootMotionEnabled)
+	{
+		for (FSlotAnimationTrack& Slot : SlotAnimTracks)
+		{
+			for (FAnimSegment& Segment : Slot.AnimTrack.AnimSegments)
+			{
+				UAnimSequence * Sequence = Cast<UAnimSequence>(Segment.AnimReference);
+				if (Sequence && !Sequence->bRootMotionSettingsCopiedFromMontage)
+				{
+					Sequence->bEnableRootMotion = true;
+					Sequence->RootMotionRootLock = RootMotionRootLock;
+					Sequence->bRootMotionSettingsCopiedFromMontage = true;
+				}
+			}
+		}
+	}
 	// find preview base pose if it can
 #if WITH_EDITORONLY_DATA
 	if ( IsValidAdditive() && PreviewBasePose == NULL )
@@ -547,6 +564,18 @@ EAnimEventTriggerOffsets::Type UAnimMontage::CalculateOffsetForNotify(float Noti
 }
 #endif
 
+bool UAnimMontage::HasRootMotion() const
+{
+	for (const FSlotAnimationTrack& Track : SlotAnimTracks)
+	{
+		if (Track.AnimTrack.HasRootMotion())
+		{
+			return true;
+		}
+	}
+	return false;
+}
+
 /** Extract RootMotion Transform from a contiguous Track position range.
  * *CONTIGUOUS* means that if playing forward StartTractPosition < EndTrackPosition.
  * No wrapping over if looping. No jumping across different sections.
@@ -569,51 +598,8 @@ FTransform UAnimMontage::ExtractRootMotionFromTrackRange(float StartTrackPositio
 		// Get RootMotion pieces from this track.
 		// We can deal with looping animations, or multiple animations. So we break those up into sequential operations.
 		// (Animation, StartFrame, EndFrame) so we can then extract root motion sequentially.
-		TArray<FRootMotionExtractionStep> RootMotionExtractionSteps;
-		SlotAnimTrack.GetRootMotionExtractionStepsForTrackRange(RootMotionExtractionSteps, StartTrackPosition, EndTrackPosition);
-
-		UE_LOG(LogRootMotion, Log,  TEXT("\tUAnimMontage::ExtractRootMotionForTrackRange, NumSteps: %d, StartTrackPosition: %.3f, EndTrackPosition: %.3f"), 
-			RootMotionExtractionSteps.Num(), StartTrackPosition, EndTrackPosition);
+		ExtractRootMotionFromTrack(SlotAnimTrack, StartTrackPosition, EndTrackPosition, RootMotion);
 
-		// Process root motion steps if any
-		if( RootMotionExtractionSteps.Num() > 0 )
-		{
-			FTransform InitialTransform, StartTransform, EndTransform;
-
-			// Go through steps sequentially, extract root motion, and accumulate it.
-			// This has to be done in order so root motion translation & rotation is applied properly (as translation is relative to rotation)
-			for(int32 StepIndex=0; StepIndex<RootMotionExtractionSteps.Num(); StepIndex++)
-			{
-				const FRootMotionExtractionStep& CurrentStep = RootMotionExtractionSteps[StepIndex];
-				CurrentStep.AnimSequence->ExtractRootTrack(0.f, InitialTransform, NULL);
-				CurrentStep.AnimSequence->ExtractRootTrack(CurrentStep.StartPosition, StartTransform, NULL);
-				CurrentStep.AnimSequence->ExtractRootTrack(CurrentStep.EndPosition, EndTransform, NULL);
-				
-				// Transform to Component Space Rotation (inverse root transform from first frame)
-				const FTransform RootToComponentRot = FTransform(InitialTransform.GetRotation().Inverse());
-				StartTransform = RootToComponentRot * StartTransform;
-				EndTransform = RootToComponentRot * EndTransform;
-
-				FTransform DeltaTransform = EndTransform.GetRelativeTransform(StartTransform);
-
-				// Filter out rotation 
-				if( !bEnableRootMotionRotation )
-				{
-					DeltaTransform.SetRotation(FQuat::Identity);
-				}
-				// Filter out translation
-				if( !bEnableRootMotionTranslation )
-				{
-					DeltaTransform.SetTranslation(FVector::ZeroVector);
-				}
-
-				RootMotion.Accumulate(DeltaTransform);
-
-				UE_LOG(LogRootMotion, Log,  TEXT("\t\tCurrentStep: %d, StartPos: %.3f, EndPos: %.3f, Anim: %s DeltaTransform Translation: %s, Rotation: %s"),
-					StepIndex, CurrentStep.StartPosition, CurrentStep.EndPosition, *CurrentStep.AnimSequence->GetName(), 
-					*DeltaTransform.GetTranslation().ToCompactString(), *DeltaTransform.GetRotation().Rotator().ToCompactString() );
-			}
-		}
 	}
 
 	UE_LOG(LogRootMotion, Log,  TEXT("\tUAnimMontage::ExtractRootMotionForTrackRange RootMotionTransform: Translation: %s, Rotation: %s"),
@@ -926,7 +912,7 @@ bool FAnimMontageInstance::SimulateAdvance(float DeltaTime, float& InOutPosition
 	const float CombinedPlayRate = PlayRate * Montage->RateScale;
 	const bool bPlayingForward = (CombinedPlayRate > 0.f);
 
-	const bool bExtractRootMotion = (Montage->bEnableRootMotionRotation || Montage->bEnableRootMotionTranslation);
+	const bool bExtractRootMotion = Montage->HasRootMotion();
 
 	float DesiredDeltaMove = CombinedPlayRate * DeltaTime;
 	float OriginalMoveDelta = DesiredDeltaMove;
@@ -1006,7 +992,7 @@ bool FAnimMontageInstance::SimulateAdvance(float DeltaTime, float& InOutPosition
 	return true;
 }
 
-void FAnimMontageInstance::Advance(float DeltaTime, struct FRootMotionMovementParams * OutRootMotionParams)
+void FAnimMontageInstance::Advance(float DeltaTime, struct FRootMotionMovementParams * OutRootMotionParams, bool bBlendRootMotion)
 {
 	if( IsValid() )
 	{
@@ -1027,7 +1013,7 @@ void FAnimMontageInstance::Advance(float DeltaTime, struct FRootMotionMovementPa
 			const float CombinedPlayRate = PlayRate * Montage->RateScale;
 			const bool bPlayingForward = (CombinedPlayRate > 0.f);
 
-			const bool bExtractRootMotion = (OutRootMotionParams != NULL) && (Montage->bEnableRootMotionRotation || Montage->bEnableRootMotionTranslation);
+			const bool bExtractRootMotion = (OutRootMotionParams != NULL) && Montage->HasRootMotion();
 			
 			float DesiredDeltaMove = CombinedPlayRate * DeltaTime;
 			float OriginalMoveDelta = DesiredDeltaMove;
@@ -1078,7 +1064,17 @@ void FAnimMontageInstance::Advance(float DeltaTime, struct FRootMotionMovementPa
 						// Extract Root Motion for this time slice, and accumulate it.
 						if( bExtractRootMotion && AnimInstance.IsValid() )
 						{
-							OutRootMotionParams->Accumulate( Montage->ExtractRootMotionFromTrackRange(PrevPosition, Position) );
+							FTransform RootMotion = Montage->ExtractRootMotionFromTrackRange(PrevPosition, Position);
+							if (bBlendRootMotion)
+							{
+								const float RootMotionSlotWeight = AnimInstance.Get()->GetSlotRootMotionWeight(Montage->SlotAnimTracks[0].SlotName);
+								const float RootMotionInstanceWeight = Weight * RootMotionSlotWeight;
+								OutRootMotionParams->AccumulateWithBlend(RootMotion, RootMotionInstanceWeight);
+							}
+							else
+							{
+								OutRootMotionParams->Accumulate(RootMotion);
+							}
 						}
 
 						// If about to reach the end of the montage...
diff --git a/Engine/Source/Runtime/Engine/Private/Animation/AnimNode_Slot.cpp b/Engine/Source/Runtime/Engine/Private/Animation/AnimNode_Slot.cpp
--- a/Engine/Source/Runtime/Engine/Private/Animation/AnimNode_Slot.cpp
+++ b/Engine/Source/Runtime/Engine/Private/Animation/AnimNode_Slot.cpp
@@ -29,6 +29,8 @@ void FAnimNode_Slot::Update(const FAnimationUpdateContext& Context)
 
 	// Update cache in AnimInstance.
 	Context.AnimInstance->UpdateSlotNodeWeight(SlotName, SlotNodeWeight);
+	Context.AnimInstance->UpdateSlotRootMotionWeight(SlotName, Context.GetFinalBlendWeight());
+
 
 	if (SourceWeight > ZERO_ANIMWEIGHT_THRESH)
 	{
diff --git a/Engine/Source/Runtime/Engine/Private/Animation/AnimSequence.cpp b/Engine/Source/Runtime/Engine/Private/Animation/AnimSequence.cpp
--- a/Engine/Source/Runtime/Engine/Private/Animation/AnimSequence.cpp
+++ b/Engine/Source/Runtime/Engine/Private/Animation/AnimSequence.cpp
@@ -98,6 +98,9 @@ bool FCurveTrack::CompressCurveWeights()
 
 UAnimSequence::UAnimSequence(const class FPostConstructInitializeProperties& PCIP)
 	: Super(PCIP)
+	, bEnableRootMotion(false)
+	, RootMotionRootLock(ERootMotionRootLock::RefPose)
+	, bRootMotionSettingsCopiedFromMontage(false)
 {
 
 	RateScale = 1.0;
@@ -663,41 +666,105 @@ void UAnimSequence::ExtractBoneTransform(const TArray<struct FRawAnimSequenceTra
 	OutAtom.NormalizeRotation();
 }
 
-void UAnimSequence::ExtractRootTrack(float Pos, FTransform& RootTransform, const FBoneContainer* RequiredBones) const
+void UAnimSequence::OnAssetPlayerTickedInternal(FAnimAssetTickContext &Context, const float PreviousTime, const float MoveDelta, const FAnimTickRecord &Instance, class UAnimInstance* InstanceOwner) const
+{
+	Super::OnAssetPlayerTickedInternal(Context, PreviousTime, MoveDelta, Instance, InstanceOwner);
+
+	if (bEnableRootMotion)
+	{
+		Context.RootMotionMovementParams.Accumulate(ExtractRootMotion(PreviousTime, MoveDelta, Instance.bLooping));
+	}
+}
+
+FTransform UAnimSequence::ExtractRootTrackTransform(float Pos, const FBoneContainer * RequiredBones) const
 {
 	// we assume root is in first data if available = SkeletonIndex == 0 && BoneTreeIndex == 0)
 	if ((TrackToSkeletonMapTable.Num() > 0) && (TrackToSkeletonMapTable[0].BoneTreeIndex == 0) )
 	{
 		// if we do have root data, then return root data
+		FTransform RootTransform;
 		GetBoneTransform(RootTransform, 0, Pos, false );
-		return;
+		return RootTransform;
 	}
 
 	// Fallback to root bone from reference skeleton.
 	if( RequiredBones )
 	{
-		const FReferenceSkeleton& RefSkeleton = RequiredBones->GetReferenceSkeleton();
+		const FReferenceSkeleton & RefSkeleton = RequiredBones->GetReferenceSkeleton();
 		if( RefSkeleton.GetNum() > 0 )
 		{
-			RootTransform = RefSkeleton.GetRefBonePose()[0];
-			return;
+			return RefSkeleton.GetRefBonePose()[0];
 		}
 	}
 
-	USkeleton* MySkeleton = GetSkeleton();
+	USkeleton * MySkeleton = GetSkeleton();
 	// If we don't have a RequiredBones array, get root bone from default skeleton.
 	if( !RequiredBones &&  MySkeleton )
 	{
 		const FReferenceSkeleton RefSkeleton = MySkeleton->GetReferenceSkeleton();
 		if( RefSkeleton.GetNum() > 0 )
 		{
-			RootTransform = RefSkeleton.GetRefBonePose()[0];
-			return;
+			return RefSkeleton.GetRefBonePose()[0];
 		}
 	}
 
 	// Otherwise, use identity.
-	RootTransform = FTransform::Identity;
+	return FTransform::Identity;
+}
+
+FTransform UAnimSequence::ExtractRootMotion(const float & StartTime, const float & DeltaTime, const float bAllowLooping) const
+{
+	FRootMotionMovementParams RootMotionParams;
+
+	if (DeltaTime != 0.f)
+	{
+		bool const bPlayingBackwards = (DeltaTime < 0.f);
+
+		float PreviousPosition = StartTime;
+		float CurrentPosition = StartTime;
+		float DesiredDeltaMove = DeltaTime;
+
+		do
+		{
+			// Disable looping here. Advance to desired position, or beginning / end of animation 
+			const ETypeAdvanceAnim AdvanceType = FAnimationRuntime::AdvanceTime(false, DesiredDeltaMove, CurrentPosition, SequenceLength);
+
+			// Verify position assumptions
+			check(bPlayingBackwards ? (CurrentPosition <= PreviousPosition) : (CurrentPosition >= PreviousPosition));
+
+			RootMotionParams.Accumulate(ExtractRootMotionFromRange(PreviousPosition, CurrentPosition));
+
+			// If we've hit the end of the animation, and we're allowed to loop, keep going.
+			if ((AdvanceType == ETAA_Finished) && bAllowLooping)
+			{
+				const float ActualDeltaMove = (CurrentPosition - PreviousPosition);
+				DesiredDeltaMove -= ActualDeltaMove;
+
+				PreviousPosition = bPlayingBackwards ? SequenceLength : 0.f;
+				CurrentPosition = PreviousPosition;
+			}
+			else
+			{
+				break;
+			}
+		} while (true);
+	}
+
+	return RootMotionParams.RootMotionTransform;
+}
+
+FTransform UAnimSequence::ExtractRootMotionFromRange(float StartTrackPosition, float EndTrackPosition) const
+{
+	FTransform InitialTransform = ExtractRootTrackTransform(0.f, NULL);
+	FTransform StartTransform = ExtractRootTrackTransform(StartTrackPosition, NULL);
+	FTransform EndTransform = ExtractRootTrackTransform(EndTrackPosition, NULL);
+
+	// Transform to Component Space Rotation (inverse root transform from first frame)
+	const FTransform RootToComponentRot = FTransform(InitialTransform.GetRotation().Inverse());
+	StartTransform = RootToComponentRot * StartTransform;
+	EndTransform = RootToComponentRot * EndTransform;
+
+	return EndTransform.GetRelativeTransform(StartTransform);
 }
 
 void UAnimSequence::GetAnimationPose(FTransformArrayA2 & OutAtoms, const FBoneContainer& RequiredBones, const FAnimExtractContext& ExtractionContext) const
@@ -720,27 +787,14 @@ void UAnimSequence::GetAnimationPose(FTransformArrayA2 & OutAtoms, const FBoneCo
 	}
 }
 
-void UAnimSequence::ResetRootBoneForRootMotion(FTransformArrayA2 & BoneTransforms, const FBoneContainer& RequiredBones, const FAnimExtractContext& ExtractionContext) const
+void UAnimSequence::ResetRootBoneForRootMotion(FTransformArrayA2 & BoneTransforms, const FBoneContainer & RequiredBones, ERootMotionRootLock::Type RootMotionRootLock) const
 {
-	FTransform LockedRootBone;
-	switch (ExtractionContext.RootMotionRootLock)
+	switch (RootMotionRootLock)
 	{
-		case ERootMotionRootLock::AnimFirstFrame : ExtractRootTrack(0.f, LockedRootBone, &RequiredBones); break;
-		case ERootMotionRootLock::Zero : LockedRootBone = FTransform::Identity; break;
+		case ERootMotionRootLock::AnimFirstFrame: BoneTransforms[0] = ExtractRootTrackTransform(0.f, &RequiredBones); break;
+		case ERootMotionRootLock::Zero: BoneTransforms[0] = FTransform::Identity; break;
 		default:
-		case ERootMotionRootLock::RefPose : LockedRootBone = RequiredBones.GetRefPoseArray()[0]; break;
-	}
-
-	// If extracting root motion translation, reset translation part to first frame of animation.
-	if( ExtractionContext.bExtractRootMotionTranslation )
-	{
-		BoneTransforms[0].SetTranslation(LockedRootBone.GetTranslation());
-	}
-
-	// If extracting root motion rotation, reset rotation part to first frame of animation.
-	if( ExtractionContext.bExtractRootMotionRotation )
-	{
-		BoneTransforms[0].SetRotation(LockedRootBone.GetRotation());
+		case ERootMotionRootLock::RefPose: BoneTransforms[0] = RequiredBones.GetRefPoseArray()[0]; break;
 	}
 }
 
@@ -811,9 +865,9 @@ void UAnimSequence::GetBonePose(FTransformArrayA2 & OutAtoms, const FBoneContain
 			}
 		}
 
-		if( ExtractionContext.bExtractRootMotionTranslation || ExtractionContext.bExtractRootMotionRotation )
+		if( ExtractionContext.bExtractRootMotion && bEnableRootMotion )
 		{
-			ResetRootBoneForRootMotion(OutAtoms, RequiredBones, ExtractionContext);
+			ResetRootBoneForRootMotion(OutAtoms, RequiredBones, RootMotionRootLock);
 		}
 		return;
 	}
@@ -888,9 +942,9 @@ void UAnimSequence::GetBonePose(FTransformArrayA2 & OutAtoms, const FBoneContain
 		ExtractionContext.CurrentTime);
 
 	// Once pose has been extracted, snap root bone back to first frame if we are extracting root motion.
-	if( ExtractionContext.bExtractRootMotionTranslation || ExtractionContext.bExtractRootMotionRotation )
+	if( ExtractionContext.bExtractRootMotion && bEnableRootMotion)
 	{
-		ResetRootBoneForRootMotion(OutAtoms, RequiredBones, ExtractionContext);
+		ResetRootBoneForRootMotion(OutAtoms, RequiredBones, RootMotionRootLock);
 	}
 
 	// Anim Scale Retargeting
diff --git a/Engine/Source/Runtime/Engine/Private/Animation/AnimSequenceBase.cpp b/Engine/Source/Runtime/Engine/Private/Animation/AnimSequenceBase.cpp
--- a/Engine/Source/Runtime/Engine/Private/Animation/AnimSequenceBase.cpp
+++ b/Engine/Source/Runtime/Engine/Private/Animation/AnimSequenceBase.cpp
@@ -423,24 +423,19 @@ void UAnimSequenceBase::TickAssetPlayerInstance(const FAnimTickRecord& Instance,
 {
 	float& CurrentTime = *(Instance.TimeAccumulator);
 	const float PreviousTime = CurrentTime;
-
 	const float PlayRate = Instance.PlayRateMultiplier * this->RateScale;
+
+	float MoveDelta = 0.f;
+
 	if( Context.IsLeader() )
 	{
 		const float DeltaTime = Context.GetDeltaTime();
-		const float MoveDelta = PlayRate * DeltaTime;
+		MoveDelta = PlayRate * DeltaTime;
+
 		if( MoveDelta != 0.f )
 		{
 			// Advance time
 			FAnimationRuntime::AdvanceTime(Instance.bLooping, MoveDelta, CurrentTime, SequenceLength);
-
-			if( Context.ShouldGenerateNotifies() )
-			{
-				// Harvest and record notifies
-				TArray<const FAnimNotifyEvent*> AnimNotifies;
-				GetAnimNotifies(PreviousTime, MoveDelta, Instance.bLooping, AnimNotifies);
-				InstanceOwner->AddAnimNotifies(AnimNotifies, Instance.EffectiveBlendWeight);
-			}
 		}
 
 		Context.SetSyncPoint(CurrentTime / SequenceLength);
@@ -451,23 +446,20 @@ void UAnimSequenceBase::TickAssetPlayerInstance(const FAnimTickRecord& Instance,
 		CurrentTime = Context.GetSyncPoint() * SequenceLength;
 		//@TODO: NOTIFIES: Calculate AdvanceType based on what the new delta time is
 
-		if( Context.ShouldGenerateNotifies() && (CurrentTime != PreviousTime) )
+		if( CurrentTime != PreviousTime )
 		{
 			// Figure out delta time 
-			float MoveDelta = CurrentTime - PreviousTime;
+			MoveDelta = CurrentTime - PreviousTime;
 			// if we went against play rate, then loop around.
 			if( (MoveDelta * PlayRate) < 0.f )
 			{
 				MoveDelta += FMath::Sign<float>(PlayRate) * SequenceLength;
 			}
-
-			// Harvest and record notifies
-			TArray<const FAnimNotifyEvent*> AnimNotifies;
-			GetAnimNotifies(PreviousTime, MoveDelta, Instance.bLooping, AnimNotifies);
-			InstanceOwner->AddAnimNotifies(AnimNotifies, Instance.EffectiveBlendWeight);
 		}
 	}
 
+	OnAssetPlayerTickedInternal(Context, PreviousTime, MoveDelta, Instance, InstanceOwner);
+
 	// Evaluate Curve data now - even if time did not move, we still need to return curve if it exists
 	EvaluateCurveData(InstanceOwner, CurrentTime, Instance.EffectiveBlendWeight);
 }
@@ -628,35 +620,6 @@ void UAnimSequenceBase::EvaluateCurveData(class UAnimInstance* Instance, float C
 	RawCurveData.EvaluateCurveData(Instance, CurrentTime, BlendWeight);
 }
 
-void UAnimSequenceBase::ExtractRootTrack(float Pos, FTransform& RootTransform, const FBoneContainer* RequiredBones) const
-{
-	// Fallback to root bone from reference skeleton.
-	if (RequiredBones)
-	{
-		const FReferenceSkeleton& RefSkeleton = RequiredBones->GetReferenceSkeleton();
-		if (RefSkeleton.GetNum() > 0)
-		{
-			RootTransform = RefSkeleton.GetRefBonePose()[0];
-			return;
-		}
-	}
-
-	USkeleton* MySkeleton = GetSkeleton();
-	// If we don't have a RequiredBones array, get root bone from default skeleton.
-	if (!RequiredBones && MySkeleton)
-	{
-		const FReferenceSkeleton RefSkeleton = MySkeleton->GetReferenceSkeleton();
-		if (RefSkeleton.GetNum() > 0)
-		{
-			RootTransform = RefSkeleton.GetRefBonePose()[0];
-			return;
-		}
-	}
-
-	// Otherwise, use identity.
-	RootTransform = FTransform::Identity;
-};
-
 void UAnimSequenceBase::Serialize(FArchive& Ar)
 {
 	Super::Serialize(Ar);
@@ -672,3 +635,14 @@ void UAnimSequenceBase::Serialize(FArchive& Ar)
 	}
 	RawCurveData.Serialize(Ar);
 }
+
+void UAnimSequenceBase::OnAssetPlayerTickedInternal(FAnimAssetTickContext &Context, const float PreviousTime, const float MoveDelta, const FAnimTickRecord &Instance, class UAnimInstance* InstanceOwner) const
+{
+	if (Context.ShouldGenerateNotifies())
+	{
+		// Harvest and record notifies
+		TArray<const FAnimNotifyEvent*> AnimNotifies;
+		GetAnimNotifies(PreviousTime, MoveDelta, Instance.bLooping, AnimNotifies);
+		InstanceOwner->AddAnimNotifies(AnimNotifies, Instance.EffectiveBlendWeight);
+	}
+}
diff --git a/Engine/Source/Runtime/Engine/Private/Animation/AnimSingleNodeInstance.cpp b/Engine/Source/Runtime/Engine/Private/Animation/AnimSingleNodeInstance.cpp
--- a/Engine/Source/Runtime/Engine/Private/Animation/AnimSingleNodeInstance.cpp
+++ b/Engine/Source/Runtime/Engine/Private/Animation/AnimSingleNodeInstance.cpp
@@ -337,7 +337,7 @@ bool UAnimSingleNodeInstance::NativeEvaluateAnimation(FPoseContext& Output)
 			else
 			{
 				// if sekeltalmesh isn't there, we'll need to use skeleton
-				FAnimationRuntime::GetPoseFromSequence(Sequence, RequiredBones, Output.Pose.Bones, FAnimExtractContext(CurrentTime));
+				FAnimationRuntime::GetPoseFromSequence(Sequence, RequiredBones, Output.Pose.Bones, FAnimExtractContext(CurrentTime, Sequence->bEnableRootMotion));
 			}
 		}
 		else if (UAnimComposite* Composite = Cast<UAnimComposite>(CurrentAsset))
@@ -359,7 +359,7 @@ bool UAnimSingleNodeInstance::NativeEvaluateAnimation(FPoseContext& Output)
 				FAnimationRuntime::FillWithRefPose(BasePose.Bones, RequiredBones);
 				
 				//get the additive pose
-				FAnimationRuntime::GetPoseFromAnimTrack(AnimTrack, RequiredBones, AdditivePose.Bones, FAnimExtractContext(CurrentTime));
+				FAnimationRuntime::GetPoseFromAnimTrack(AnimTrack, RequiredBones, AdditivePose.Bones, FAnimExtractContext(CurrentTime, RootMotionMode == ERootMotionMode::RootMotionFromEverything));
 
 				// if additive, we should blend with source to make it fullbody
 				if (AdditiveAnimType == AAT_LocalSpaceBase)
@@ -374,7 +374,7 @@ bool UAnimSingleNodeInstance::NativeEvaluateAnimation(FPoseContext& Output)
 			else
 			{
 				//doesn't handle additive yet
-				FAnimationRuntime::GetPoseFromAnimTrack(AnimTrack, RequiredBones, Output.Pose.Bones, FAnimExtractContext(CurrentTime));
+				FAnimationRuntime::GetPoseFromAnimTrack(AnimTrack, RequiredBones, Output.Pose.Bones, FAnimExtractContext(CurrentTime, RootMotionMode == ERootMotionMode::RootMotionFromEverything));
 			}
 		}
 		else if (UAnimMontage* Montage = Cast<UAnimMontage>(CurrentAsset))
diff --git a/Engine/Source/Runtime/Engine/Private/Animation/AnimationRuntime.cpp b/Engine/Source/Runtime/Engine/Private/Animation/AnimationRuntime.cpp
--- a/Engine/Source/Runtime/Engine/Private/Animation/AnimationRuntime.cpp
+++ b/Engine/Source/Runtime/Engine/Private/Animation/AnimationRuntime.cpp
@@ -426,6 +426,7 @@ void FAnimationRuntime::GetPoseFromAnimTrack(
 				// Copy passed in Extraction Context, but override position and looping parameters.
 				FAnimExtractContext SequenceExtractionContext(ExtractionContext);
 				SequenceExtractionContext.CurrentTime = PositionInAnim;
+				SequenceExtractionContext.bExtractRootMotion &= Sequence->bEnableRootMotion;
 				FAnimationRuntime::GetPoseFromSequence(Sequence, RequiredBones, SourcePoses[NewIndex], SequenceExtractionContext);
 			}
 		}
@@ -585,7 +586,7 @@ void FAnimationRuntime::GetPoseFromBlendSpace(
 				const float Time = FMath::Clamp<float>(BlendSampleDataCache[I].Time, 0.f, Sample.Animation->SequenceLength);
 
 				// first one always fills up the source one
-				FAnimationRuntime::GetPoseFromSequence(Sample.Animation, RequiredBones, ChildrenTransform[I], FAnimExtractContext(Time));
+				FAnimationRuntime::GetPoseFromSequence(Sample.Animation, RequiredBones, ChildrenTransform[I], FAnimExtractContext(Time, true));
 			}
 			else
 			{
diff --git a/Engine/Source/Runtime/Engine/Private/Animation/BlendSpaceBase.cpp b/Engine/Source/Runtime/Engine/Private/Animation/BlendSpaceBase.cpp
--- a/Engine/Source/Runtime/Engine/Private/Animation/BlendSpaceBase.cpp
+++ b/Engine/Source/Runtime/Engine/Private/Animation/BlendSpaceBase.cpp
@@ -282,22 +282,27 @@ void UBlendSpaceBase::TickAssetPlayerInstance(const FAnimTickRecord& Instance, c
 							float& CurrentSampleDataTime = SampleEntry.Time;
 							CurrentSampleDataTime = SampleNormalizedCurrentTime * Sample.Animation->SequenceLength;
 
-							if( bGenerateNotifies && (!bTriggerNotifyHighestWeightedAnim || (I == HighestWeightIndex)))
-							{
-								// Figure out delta time 
-								float DeltaTimePosition = CurrentSampleDataTime - PrevSampleDataTime;
-								const float SampleMoveDelta = MoveDelta * Sample.Animation->RateScale;
+							// Figure out delta time 
+							float DeltaTimePosition = CurrentSampleDataTime - PrevSampleDataTime;
+							const float SampleMoveDelta = MoveDelta * Sample.Animation->RateScale;
 
-								// if we went against play rate, then loop around.
-								if ((SampleMoveDelta * DeltaTimePosition) < 0.f)
-								{
-									DeltaTimePosition += FMath::Sign<float>(SampleMoveDelta) * Sample.Animation->SequenceLength;
-								}
+							// if we went against play rate, then loop around.
+							if ((SampleMoveDelta * DeltaTimePosition) < 0.f)
+							{
+								DeltaTimePosition += FMath::Sign<float>(SampleMoveDelta) * Sample.Animation->SequenceLength;
+							}
 
+							if( bGenerateNotifies && (!bTriggerNotifyHighestWeightedAnim || (I == HighestWeightIndex)))
+							{
 								// Harvest and record notifies
 								Sample.Animation->GetAnimNotifies(PrevSampleDataTime, DeltaTimePosition, Instance.bLooping, Notifies);
 							}
 
+							if (Context.RootMotionMode == ERootMotionMode::RootMotionFromEverything && Sample.Animation->bEnableRootMotion)
+							{
+								Context.RootMotionMovementParams.AccumulateWithBlend(Sample.Animation->ExtractRootMotion(PrevSampleDataTime, DeltaTimePosition, Instance.bLooping), SampleEntry.GetWeight());
+							}
+
 							// handle curves
 							Sample.Animation->EvaluateCurveData(InstanceOwner, CurrentSampleDataTime, Instance.EffectiveBlendWeight * SampleEntry.GetWeight());
 
diff --git a/Engine/Source/Runtime/Engine/Private/Character.cpp b/Engine/Source/Runtime/Engine/Private/Character.cpp
--- a/Engine/Source/Runtime/Engine/Private/Character.cpp
+++ b/Engine/Source/Runtime/Engine/Private/Character.cpp
@@ -903,7 +903,7 @@ void ACharacter::OnRep_ReplicatedBasedMovement()
 	}
 
 	// Skip base updates while playing root motion, it is handled inside of OnRep_RootMotion
-	if (IsPlayingRootMotion())
+	if (IsPlayingNetworkedRootMotionMontage())
 	{
 		return;
 	}
@@ -958,7 +958,7 @@ void ACharacter::OnRep_ReplicatedBasedMovement()
 void ACharacter::OnRep_ReplicatedMovement()
 {
 	// Skip standard position correction if we are playing root motion, OnRep_RootMotion will handle it.
-	if( !IsPlayingRootMotion() )
+	if (!IsPlayingNetworkedRootMotionMontage())
 	{
 		Super::OnRep_ReplicatedMovement();
 	}
@@ -1247,7 +1247,22 @@ void ACharacter::GetLifetimeReplicatedProps( TArray< FLifetimeProperty > & OutLi
 
 bool ACharacter::IsPlayingRootMotion() const
 {
-	return (GetRootMotionAnimMontageInstance() != NULL);
+	if (Mesh && Mesh->AnimScriptInstance)
+	{
+		return	(Mesh->AnimScriptInstance->RootMotionMode == ERootMotionMode::RootMotionFromEverything) ||
+				(Mesh->AnimScriptInstance->GetRootMotionMontageInstance() != NULL);
+	}
+	return false;
+}
+
+bool ACharacter::IsPlayingNetworkedRootMotionMontage() const
+{
+	if (Mesh && Mesh->AnimScriptInstance)
+	{
+		return	(Mesh->AnimScriptInstance->RootMotionMode == ERootMotionMode::RootMotionFromMontagesOnly) &&
+			(Mesh->AnimScriptInstance->GetRootMotionMontageInstance() != NULL);
+	}
+	return false;
 }
 
 float ACharacter::PlayAnimMontage(class UAnimMontage* AnimMontage, float InPlayRate, FName StartSectionName)
diff --git a/Engine/Source/Runtime/Engine/Private/CharacterMovementComponent.cpp b/Engine/Source/Runtime/Engine/Private/CharacterMovementComponent.cpp
--- a/Engine/Source/Runtime/Engine/Private/CharacterMovementComponent.cpp
+++ b/Engine/Source/Runtime/Engine/Private/CharacterMovementComponent.cpp
@@ -921,7 +921,7 @@ void UCharacterMovementComponent::SimulatedTick(float DeltaSeconds)
 	SCOPE_CYCLE_COUNTER(STAT_CharacterMovementSimulated);
 
 	// If we are playing a RootMotion AnimMontage.
-	if( CharacterOwner && CharacterOwner->IsPlayingRootMotion() )
+	if (CharacterOwner && CharacterOwner->IsPlayingNetworkedRootMotionMontage())
 	{
 		bWasSimulatingRootMotion = true;
 		UE_LOG(LogRootMotion, Verbose, TEXT("UCharacterMovementComponent::SimulatedTick"));
@@ -981,7 +981,9 @@ void UCharacterMovementComponent::SimulatedTick(float DeltaSeconds)
 			}
 		}
 
-		if( UpdatedComponent->IsSimulatingPhysics() || CharacterOwner->IsMatineeControlled()  )
+		if( UpdatedComponent->IsSimulatingPhysics() 
+			|| (CharacterOwner && CharacterOwner->IsMatineeControlled()) 
+			|| (CharacterOwner && CharacterOwner->IsPlayingRootMotion()))
 		{
 			PerformMovement(DeltaSeconds);
 		}
@@ -1423,7 +1425,7 @@ void UCharacterMovementComponent::PerformMovement(float DeltaSeconds)
 			}
 
 			// For local human clients, save off root motion data so it can be used by movement networking code.
-			if( CharacterOwner->IsLocallyControlled() && (CharacterOwner->Role == ROLE_AutonomousProxy) )
+			if( CharacterOwner->IsLocallyControlled() && (CharacterOwner->Role == ROLE_AutonomousProxy) && CharacterOwner->IsPlayingNetworkedRootMotionMontage() )
 			{
 				CharacterOwner->ClientRootMotionParams = RootMotionParams;
 			}
@@ -5817,7 +5819,7 @@ void UCharacterMovementComponent::SendClientAdjustment()
 	}
 	else
 	{
-		if( CharacterOwner->IsPlayingRootMotion() )
+		if( CharacterOwner->IsPlayingNetworkedRootMotionMontage() )
 		{
 			FRotator Rotation = ServerData->PendingAdjustment.NewRot.GetNormalized();
 			FVector_NetQuantizeNormal CompressedRotation(Rotation.Pitch / 180.f, Rotation.Yaw / 180.f, Rotation.Roll / 180.f);
