Added some validity checks in hud draw for canvas and world

diff --git a/Engine/Source/Runtime/Engine/Private/HUD.cpp b/Engine/Source/Runtime/Engine/Private/HUD.cpp
--- a/Engine/Source/Runtime/Engine/Private/HUD.cpp
+++ b/Engine/Source/Runtime/Engine/Private/HUD.cpp
@@ -23,8 +23,8 @@ AHUD::AHUD(const class FPostConstructInitializeProperties& PCIP)
 	PrimaryActorTick.bCanEverTick = true;
 	bHidden = true;
 	bReplicates = false;
-	
-	WhiteColor = FColor(255, 255, 255, 255);	
+
+	WhiteColor = FColor(255, 255, 255, 255);
 	GreenColor = FColor(0, 255, 0, 255);
 	RedColor = FColor(255, 0, 0, 255);
 
@@ -43,16 +43,16 @@ void AHUD::SetCanvas(class UCanvas* InCanvas, class UCanvas* InDebugCanvas)
 
 void AHUD::Draw3DLine(FVector Start, FVector End, FColor LineColor)
 {
-	GetWorld()->LineBatcher->DrawLine(Start,End,LineColor,SDPG_World);
+	GetWorld()->LineBatcher->DrawLine(Start, End, LineColor, SDPG_World);
 }
 
-void AHUD::Draw2DLine(int32 X1,int32 Y1,int32 X2,int32 Y2,FColor LineColor)
+void AHUD::Draw2DLine(int32 X1, int32 Y1, int32 X2, int32 Y2, FColor LineColor)
 {
 	check(Canvas);
 
-	FCanvasLineItem LineItem( FVector2D(X1, Y1), FVector2D(X2, Y2) );
-	LineItem.SetColor( LineColor );
-	LineItem.Draw( Canvas->Canvas );
+	FCanvasLineItem LineItem(FVector2D(X1, Y1), FVector2D(X2, Y2));
+	LineItem.SetColor(LineColor);
+	LineItem.Draw(Canvas->Canvas);
 }
 
 void AHUD::PostInitializeComponents()
@@ -107,6 +107,11 @@ FVector2D AHUD::GetCoordinateOffset() const
 
 void AHUD::PostRender()
 {
+	// Theres nothing we can really do without a canvas or a world - so leave now in that case
+	if ( (GetWorld() == nullptr) || (Canvas == nullptr))
+	{
+		return;
+	}
 	// Set up delta time
 	RenderDelta = GetWorld()->TimeSeconds - LastHUDRenderTime;
 
@@ -303,18 +308,21 @@ bool AHUD::ShouldDisplayDebug(const FName& DebugType) const
 
 void AHUD::ShowDebugInfo(float& YL, float& YPos)
 {
-	if (!DebugDisplay.Contains(TEXT("Bones")))
+	if (DebugCanvas != nullptr )
 	{
-		FLinearColor BackgroundColor(0.f, 0.f, 0.f, 0.5f);
-		DebugCanvas->Canvas->DrawTile(0, 0, DebugCanvas->ClipX, DebugCanvas->ClipY, 0.f, 0.f, 0.f, 0.f, BackgroundColor);
-	}
+		if (!DebugDisplay.Contains(TEXT("Bones")))
+		{
+			FLinearColor BackgroundColor(0.f, 0.f, 0.f, 0.5f);
+			DebugCanvas->Canvas->DrawTile(0, 0, DebugCanvas->ClipX, DebugCanvas->ClipY, 0.f, 0.f, 0.f, 0.f, BackgroundColor);
+		}
 
-	FDebugDisplayInfo DisplayInfo(DebugDisplay,ToggledDebugCategories);
-	PlayerOwner->PlayerCameraManager->ViewTarget.Target->DisplayDebug(DebugCanvas, DisplayInfo, YL, YPos);
+		FDebugDisplayInfo DisplayInfo(DebugDisplay, ToggledDebugCategories);
+		PlayerOwner->PlayerCameraManager->ViewTarget.Target->DisplayDebug(DebugCanvas, DisplayInfo, YL, YPos);
 
-	if (ShouldDisplayDebug(NAME_Game))
-	{
-		GetWorld()->GetAuthGameMode()->DisplayDebug(DebugCanvas, DisplayInfo, YL, YPos);
+		if (ShouldDisplayDebug(NAME_Game))
+		{
+			GetWorld()->GetAuthGameMode()->DisplayDebug(DebugCanvas, DisplayInfo, YL, YPos);
+		}
 	}
 }
 
@@ -375,7 +383,7 @@ void AHUD::OnLostFocusPause(bool bEnable)
 
 void AHUD::DrawDebugTextList()
 {
-	if (DebugTextList.Num() > 0)
+	if( (DebugTextList.Num() > 0) && (DebugCanvas != nullptr ) )
 	{
 		FRotator CameraRot;
 		FVector CameraLoc;
