CIS fixes for clang

diff --git a/Engine/Source/Runtime/Engine/Classes/Animation/AnimationAsset.h b/Engine/Source/Runtime/Engine/Classes/Animation/AnimationAsset.h
--- a/Engine/Source/Runtime/Engine/Classes/Animation/AnimationAsset.h
+++ b/Engine/Source/Runtime/Engine/Classes/Animation/AnimationAsset.h
@@ -253,8 +253,8 @@ struct FRootMotionMovementParams
 
 	FRootMotionMovementParams()
 		: bHasRootMotion(false)
-		, RootMotionTransform(FTransform::Identity)
 		, BlendWeight(0.f)
+		, RootMotionTransform(FTransform::Identity)
 	{
 	}
 
@@ -330,8 +330,8 @@ struct FAnimAssetTickContext
 {
 public:
 	FAnimAssetTickContext(float InDeltaTime, ERootMotionMode::Type InRootMotionMode)
-		: DeltaTime(InDeltaTime)
-		, RootMotionMode(InRootMotionMode)
+		: RootMotionMode(InRootMotionMode)
+		, DeltaTime(InDeltaTime)
 		, SyncPoint(0.0f)
 		, bIsLeader(true)
 	{
