FPrimitiveSceneProxy::GetScene() now returns a reference to the FScene. Lots of code assumes the existance of a scene, so this cements this. Added a check() that the scene is valid when creating the proxy.

diff --git a/Engine/Plugins/2D/Paper2D/Source/Paper2D/Private/PaperFlipbookSceneProxy.cpp b/Engine/Plugins/2D/Paper2D/Source/Paper2D/Private/PaperFlipbookSceneProxy.cpp
--- a/Engine/Plugins/2D/Paper2D/Source/Paper2D/Private/PaperFlipbookSceneProxy.cpp
+++ b/Engine/Plugins/2D/Paper2D/Source/Paper2D/Private/PaperFlipbookSceneProxy.cpp
@@ -12,5 +12,5 @@ FPaperFlipbookSceneProxy::FPaperFlipbookSceneProxy(const UPaperFlipbookComponent
 	//@TODO: PAPER2D: WireframeColor = RenderComp->GetWireframeColor();
 
 	Material = InComponent->GetMaterial(0);
-	MaterialRelevance = InComponent->GetMaterialRelevance(GetScene()->GetFeatureLevel());
+	MaterialRelevance = InComponent->GetMaterialRelevance(GetScene().GetFeatureLevel());
 }
diff --git a/Engine/Plugins/2D/Paper2D/Source/Paper2D/Private/PaperSpriteSceneProxy.cpp b/Engine/Plugins/2D/Paper2D/Source/Paper2D/Private/PaperSpriteSceneProxy.cpp
--- a/Engine/Plugins/2D/Paper2D/Source/Paper2D/Private/PaperSpriteSceneProxy.cpp
+++ b/Engine/Plugins/2D/Paper2D/Source/Paper2D/Private/PaperSpriteSceneProxy.cpp
@@ -15,7 +15,7 @@ FPaperSpriteSceneProxy::FPaperSpriteSceneProxy(const UPaperSpriteComponent* InCo
 	Material = InComponent->GetMaterial(0);
 	AlternateMaterial = InComponent->GetMaterial(1);
 	MaterialSplitIndex = INDEX_NONE;
-	MaterialRelevance = InComponent->GetMaterialRelevance(GetScene()->GetFeatureLevel());
+	MaterialRelevance = InComponent->GetMaterialRelevance(GetScene().GetFeatureLevel());
 
 	SourceSprite = InComponent->SourceSprite; //@TODO: This is totally not threadsafe, and won't keep up to date if the actor's sprite changes, etc....
 }
diff --git a/Engine/Plugins/2D/Paper2D/Source/Paper2D/Private/PaperTileMapRenderSceneProxy.cpp b/Engine/Plugins/2D/Paper2D/Source/Paper2D/Private/PaperTileMapRenderSceneProxy.cpp
--- a/Engine/Plugins/2D/Paper2D/Source/Paper2D/Private/PaperTileMapRenderSceneProxy.cpp
+++ b/Engine/Plugins/2D/Paper2D/Source/Paper2D/Private/PaperTileMapRenderSceneProxy.cpp
@@ -19,7 +19,7 @@ FPaperTileMapRenderSceneProxy::FPaperTileMapRenderSceneProxy(const UPaperTileMap
 
 		if (Material)
 		{
-			MaterialRelevance = Material->GetRelevance(GetScene()->GetFeatureLevel());
+			MaterialRelevance = Material->GetRelevance(GetScene().GetFeatureLevel());
 		}
 	}
 }
diff --git a/Engine/Plugins/2D/Paper2D/Source/Paper2D/Private/Rendering/PaperBatchManager.cpp b/Engine/Plugins/2D/Paper2D/Source/Paper2D/Private/Rendering/PaperBatchManager.cpp
--- a/Engine/Plugins/2D/Paper2D/Source/Paper2D/Private/Rendering/PaperBatchManager.cpp
+++ b/Engine/Plugins/2D/Paper2D/Source/Paper2D/Private/Rendering/PaperBatchManager.cpp
@@ -64,9 +64,9 @@ UPaperBatchComponent* FPaperBatchManager::GetBatchComponent(UWorld* World)
 	return Batcher;
 }
 
-FPaperBatchSceneProxy* FPaperBatchManager::GetBatcher(FSceneInterface* Scene)
+FPaperBatchSceneProxy* FPaperBatchManager::GetBatcher(FSceneInterface& Scene)
 {
-	UPaperBatchComponent* BatchComponent = GetBatchComponent(Scene->GetWorld());
+	UPaperBatchComponent* BatchComponent = GetBatchComponent(Scene.GetWorld());
 	FPaperBatchSceneProxy* BatchProxy = static_cast<FPaperBatchSceneProxy*>(BatchComponent->SceneProxy);
 	check(BatchProxy);
 	return BatchProxy;
diff --git a/Engine/Plugins/2D/Paper2D/Source/Paper2D/Private/Rendering/PaperBatchManager.h b/Engine/Plugins/2D/Paper2D/Source/Paper2D/Private/Rendering/PaperBatchManager.h
--- a/Engine/Plugins/2D/Paper2D/Source/Paper2D/Private/Rendering/PaperBatchManager.h
+++ b/Engine/Plugins/2D/Paper2D/Source/Paper2D/Private/Rendering/PaperBatchManager.h
@@ -12,7 +12,7 @@ public:
 
 	// Gets the batcher for the specified scene
 	// Note: Can only be called on the render thread
-	static class FPaperBatchSceneProxy* GetBatcher(FSceneInterface* Scene);
+	static class FPaperBatchSceneProxy* GetBatcher(FSceneInterface& Scene);
 private:
 	static void OnWorldCreated(UWorld* World, const UWorld::InitializationValues IVS);
 	static void OnWorldDestroyed(UWorld* World);
diff --git a/Engine/Plugins/2D/Paper2D/Source/Paper2D/Private/Terrain/PaperTerrainComponent.cpp b/Engine/Plugins/2D/Paper2D/Source/Paper2D/Private/Terrain/PaperTerrainComponent.cpp
--- a/Engine/Plugins/2D/Paper2D/Source/Paper2D/Private/Terrain/PaperTerrainComponent.cpp
+++ b/Engine/Plugins/2D/Paper2D/Source/Paper2D/Private/Terrain/PaperTerrainComponent.cpp
@@ -50,7 +50,7 @@ FPaperTerrainSceneProxy::FPaperTerrainSceneProxy(const UPaperTerrainComponent* I
 	for (const FPaperTerrainMaterialPair& Batch : DrawingData)
 	{
 		const UMaterialInterface* MaterialInterface = (Batch.Material != nullptr) ? Batch.Material : UMaterial::GetDefaultMaterial(MD_Surface);
-		MaterialRelevance |= MaterialInterface->GetRelevance_Concurrent(GetScene()->GetFeatureLevel());
+		MaterialRelevance |= MaterialInterface->GetRelevance_Concurrent(GetScene().GetFeatureLevel());
 	}
 }
 
diff --git a/Engine/Plugins/Runtime/CableComponent/Source/CableComponent/Private/CableComponent.cpp b/Engine/Plugins/Runtime/CableComponent/Source/CableComponent/Private/CableComponent.cpp
--- a/Engine/Plugins/Runtime/CableComponent/Source/CableComponent/Private/CableComponent.cpp
+++ b/Engine/Plugins/Runtime/CableComponent/Source/CableComponent/Private/CableComponent.cpp
@@ -93,7 +93,7 @@ public:
 		: FPrimitiveSceneProxy(Component)
 		, Material(NULL)
 		, DynamicData(NULL)
-		, MaterialRelevance(Component->GetMaterialRelevance(GetScene()->GetFeatureLevel()))
+		, MaterialRelevance(Component->GetMaterialRelevance(GetScene().GetFeatureLevel()))
 		, NumSegments(Component->NumSegments)
 		, CableWidth(Component->CableWidth)
 		, NumSides(Component->NumSides)
diff --git a/Engine/Plugins/Runtime/CustomMeshComponent/Source/CustomMeshComponent/Private/CustomMeshComponent.cpp b/Engine/Plugins/Runtime/CustomMeshComponent/Source/CustomMeshComponent/Private/CustomMeshComponent.cpp
--- a/Engine/Plugins/Runtime/CustomMeshComponent/Source/CustomMeshComponent/Private/CustomMeshComponent.cpp
+++ b/Engine/Plugins/Runtime/CustomMeshComponent/Source/CustomMeshComponent/Private/CustomMeshComponent.cpp
@@ -81,7 +81,7 @@ public:
 
 	FCustomMeshSceneProxy(UCustomMeshComponent* Component)
 		: FPrimitiveSceneProxy(Component)
-		, MaterialRelevance(Component->GetMaterialRelevance(GetScene()->GetFeatureLevel()))
+		, MaterialRelevance(Component->GetMaterialRelevance(GetScene().GetFeatureLevel()))
 	{
 		const FColor VertexColor(255,255,255);
 
diff --git a/Engine/Source/Editor/UnrealEd/Private/Factories/ApexDestructibleAssetImport.cpp b/Engine/Source/Editor/UnrealEd/Private/Factories/ApexDestructibleAssetImport.cpp
--- a/Engine/Source/Editor/UnrealEd/Private/Factories/ApexDestructibleAssetImport.cpp
+++ b/Engine/Source/Editor/UnrealEd/Private/Factories/ApexDestructibleAssetImport.cpp
@@ -734,7 +734,7 @@ bool SetApexDestructibleAsset(UDestructibleMesh& DestructibleMesh, NxDestructibl
 	for(TObjectIterator<UDestructibleComponent> It; It; ++It)
 	{
 		UDestructibleComponent* DestructibleComponent = *It;
-		if(DestructibleComponent->SkeletalMesh == &DestructibleMesh && DestructibleComponent->GetScene())
+		if(DestructibleComponent->SkeletalMesh == &DestructibleMesh)
 		{
 			DestructibleComponent->DestroyPhysicsState();
 		}
@@ -903,7 +903,7 @@ bool SetApexDestructibleAsset(UDestructibleMesh& DestructibleMesh, NxDestructibl
 		for(TObjectIterator<UDestructibleComponent> It; It; ++It)
 		{
 			UDestructibleComponent* DestructibleComponent = *It;
-			if(DestructibleComponent->SkeletalMesh == &DestructibleMesh && DestructibleComponent->GetScene())
+			if(DestructibleComponent->SkeletalMesh == &DestructibleMesh)
 			{
 				FComponentReregisterContext ReregisterContext(DestructibleComponent);
 			}
diff --git a/Engine/Source/Editor/UnrealEd/Private/Fbx/FbxSkeletalMeshImport.cpp b/Engine/Source/Editor/UnrealEd/Private/Fbx/FbxSkeletalMeshImport.cpp
--- a/Engine/Source/Editor/UnrealEd/Private/Fbx/FbxSkeletalMeshImport.cpp
+++ b/Engine/Source/Editor/UnrealEd/Private/Fbx/FbxSkeletalMeshImport.cpp
@@ -1330,7 +1330,7 @@ USkeletalMesh* UnFbx::FFbxImporter::ImportSkeletalMesh(UObject* InParent, TArray
 		for(TObjectIterator<USkeletalMeshComponent> It; It; ++It)
 		{
 			USkeletalMeshComponent* SkelComp = *It;
-			if(SkelComp->SkeletalMesh == SkeletalMesh && SkelComp->GetScene())
+			if(SkelComp->SkeletalMesh == SkeletalMesh)
 			{
 				FComponentReregisterContext ReregisterContext(SkelComp);
 			}
diff --git a/Engine/Source/Runtime/Engine/Classes/Engine/NiagaraEffectRenderer.h b/Engine/Source/Runtime/Engine/Classes/Engine/NiagaraEffectRenderer.h
--- a/Engine/Source/Runtime/Engine/Classes/Engine/NiagaraEffectRenderer.h
+++ b/Engine/Source/Runtime/Engine/Classes/Engine/NiagaraEffectRenderer.h
@@ -111,7 +111,7 @@ protected:
 			Material = UMaterial::GetDefaultMaterial(MD_Surface);
 		}
 		check(Material);
-		MaterialRelevance = Material->GetRelevance(SceneProxy->GetScene()->GetFeatureLevel());
+		MaterialRelevance = Material->GetRelevance(SceneProxy->GetScene().GetFeatureLevel());
 	}
 
 	const FNiagaraSceneProxy *SceneProxy;
diff --git a/Engine/Source/Runtime/Engine/Private/Landscape/LandscapeRender.cpp b/Engine/Source/Runtime/Engine/Private/Landscape/LandscapeRender.cpp
--- a/Engine/Source/Runtime/Engine/Private/Landscape/LandscapeRender.cpp
+++ b/Engine/Source/Runtime/Engine/Private/Landscape/LandscapeRender.cpp
@@ -554,7 +554,7 @@ FLandscapeComponentSceneProxy::FLandscapeComponentSceneProxy(ULandscapeComponent
 {
 	LevelColor = FLinearColor(1.f, 1.f, 1.f);
 
-	const auto FeatureLevel = GetScene()->GetFeatureLevel();
+	const auto FeatureLevel = GetScene().GetFeatureLevel();
 	if (FeatureLevel <= ERHIFeatureLevel::ES3_1)
 	{
 		HeightmapTexture = NULL;
@@ -696,10 +696,12 @@ void FLandscapeComponentSceneProxy::CreateRenderThreadResources()
 {
 	check(HeightmapTexture != NULL);
 
+	auto FeatureLevel = GetScene().GetFeatureLevel();
+
 	SharedBuffers = FLandscapeComponentSceneProxy::SharedBuffersMap.FindRef(SharedBuffersKey);
 	if( SharedBuffers == NULL )
 	{
-		SharedBuffers = new FLandscapeSharedBuffers(SharedBuffersKey, SubsectionSizeQuads, NumSubsections, GetScene()->GetFeatureLevel(), bRequiresAdjacencyInformation);
+		SharedBuffers = new FLandscapeSharedBuffers(SharedBuffersKey, SubsectionSizeQuads, NumSubsections, FeatureLevel, bRequiresAdjacencyInformation);
 		FLandscapeComponentSceneProxy::SharedBuffersMap.Add(SharedBuffersKey, SharedBuffers);
 
 		if (!XYOffsetmapTexture)
@@ -730,11 +732,11 @@ void FLandscapeComponentSceneProxy::CreateRenderThreadResources()
 				// Recreate Index Buffers, this case happens only there are Landscape Components using different material (one uses tessellation, other don't use it) 
 				if (SharedBuffers->bUse32BitIndices && !((FRawStaticIndexBuffer16or32<uint32>*)SharedBuffers->IndexBuffers[0])->Num())
 				{
-					SharedBuffers->CreateIndexBuffers<uint32>(GetScene()->GetFeatureLevel(), bRequiresAdjacencyInformation);
+					SharedBuffers->CreateIndexBuffers<uint32>(FeatureLevel, bRequiresAdjacencyInformation);
 				}
 				else if (!((FRawStaticIndexBuffer16or32<uint16>*)SharedBuffers->IndexBuffers[0])->Num())
 				{
-					SharedBuffers->CreateIndexBuffers<uint16>(GetScene()->GetFeatureLevel(), bRequiresAdjacencyInformation);
+					SharedBuffers->CreateIndexBuffers<uint16>(FeatureLevel, bRequiresAdjacencyInformation);
 				}
 			}
 
@@ -1080,7 +1082,7 @@ void FLandscapeComponentSceneProxy::DrawStaticElements(FStaticPrimitiveDrawInter
 	FMaterialRenderProxy* RenderProxy = MaterialInterface->GetRenderProxy(false);
 
 	// Could be different from bRequiresAdjacencyInformation during shader compilation
-	bool bCurrentRequiresAdjacencyInformation = RequiresAdjacencyInformation( RenderProxy, GetScene()->GetFeatureLevel() );
+	bool bCurrentRequiresAdjacencyInformation = RequiresAdjacencyInformation( RenderProxy, GetScene().GetFeatureLevel() );
 
 	MeshBatch.VertexFactory = VertexFactory;
 	MeshBatch.MaterialRenderProxy = RenderProxy;
@@ -2634,7 +2636,7 @@ IMPLEMENT_VERTEX_FACTORY_TYPE(FLandscapeVertexFactory, "LandscapeVertexFactory",
 */
 void FLandscapeVertexFactory::Copy(const FLandscapeVertexFactory& Other)
 {
-	//SetSceneProxy(Other.GetSceneProxy());
+	//SetSceneProxy(Other.Proxy());
 	ENQUEUE_UNIQUE_RENDER_COMMAND_TWOPARAMETER(
 		FLandscapeVertexFactoryCopyData,
 		FLandscapeVertexFactory*,VertexFactory,this,
diff --git a/Engine/Source/Runtime/Engine/Private/Landscape/LandscapeRenderMobile.cpp b/Engine/Source/Runtime/Engine/Private/Landscape/LandscapeRenderMobile.cpp
--- a/Engine/Source/Runtime/Engine/Private/Landscape/LandscapeRenderMobile.cpp
+++ b/Engine/Source/Runtime/Engine/Private/Landscape/LandscapeRenderMobile.cpp
@@ -197,7 +197,7 @@ void FLandscapeComponentSceneProxyMobile::CreateRenderThreadResources()
 	SharedBuffers = FLandscapeComponentSceneProxy::SharedBuffersMap.FindRef(SharedBuffersKey);
 	if( SharedBuffers == NULL )
 	{
-		SharedBuffers = new FLandscapeSharedBuffers(SharedBuffersKey, SubsectionSizeQuads, NumSubsections, GetScene()->GetFeatureLevel(), false);
+		SharedBuffers = new FLandscapeSharedBuffers(SharedBuffersKey, SubsectionSizeQuads, NumSubsections, GetScene().GetFeatureLevel(), false);
 		FLandscapeComponentSceneProxy::SharedBuffersMap.Add(SharedBuffersKey, SharedBuffers);
 	}
 
diff --git a/Engine/Source/Runtime/Engine/Private/MaterialBillboardComponent.cpp b/Engine/Source/Runtime/Engine/Private/MaterialBillboardComponent.cpp
--- a/Engine/Source/Runtime/Engine/Private/MaterialBillboardComponent.cpp
+++ b/Engine/Source/Runtime/Engine/Private/MaterialBillboardComponent.cpp
@@ -108,7 +108,7 @@ public:
 			UMaterialInterface* Material = Elements[ElementIndex].Material;
 			if (Material)
 			{
-				MaterialRelevance |= Material->GetRelevance(GetScene()->GetFeatureLevel());
+				MaterialRelevance |= Material->GetRelevance(GetScene().GetFeatureLevel());
 			}
 		}
 
diff --git a/Engine/Source/Runtime/Engine/Private/ModelRender.cpp b/Engine/Source/Runtime/Engine/Private/ModelRender.cpp
--- a/Engine/Source/Runtime/Engine/Private/ModelRender.cpp
+++ b/Engine/Source/Runtime/Engine/Private/ModelRender.cpp
@@ -163,7 +163,7 @@ public:
 		{
 			const FModelElement& SourceElement = SourceElements[ElementIndex];
 			FElementInfo* Element = new(Elements) FElementInfo(SourceElement);
-			MaterialRelevance |= Element->GetMaterial()->GetRelevance(GetScene()->GetFeatureLevel());
+			MaterialRelevance |= Element->GetMaterial()->GetRelevance(GetScene().GetFeatureLevel());
 		}
 
 		// Try to find a color for level coloration.
diff --git a/Engine/Source/Runtime/Engine/Private/Particles/ParticleSystemRender.cpp b/Engine/Source/Runtime/Engine/Private/Particles/ParticleSystemRender.cpp
--- a/Engine/Source/Runtime/Engine/Private/Particles/ParticleSystemRender.cpp
+++ b/Engine/Source/Runtime/Engine/Private/Particles/ParticleSystemRender.cpp
@@ -1640,7 +1640,7 @@ int32 FDynamicSpriteEmitterData::RenderEmitter(FParticleSystemSceneProxy* Proxy,
 
 FParticleVertexFactoryBase* FDynamicSpriteEmitterData::BuildVertexFactory(const FParticleSystemSceneProxy* InOwnerProxy)
 {
-	return GParticleVertexFactoryPool.GetParticleVertexFactory(PVFT_Sprite, InOwnerProxy->GetScene()->GetFeatureLevel());
+	return GParticleVertexFactoryPool.GetParticleVertexFactory(PVFT_Sprite, InOwnerProxy->GetScene().GetFeatureLevel());
 }
 
 void FDynamicSpriteEmitterData::UpdateRenderThreadResourcesEmitter(const FParticleSystemSceneProxy* InOwnerProxy)
@@ -1853,7 +1853,7 @@ void FDynamicMeshEmitterData::Init( bool bInSelected,
 
 FParticleVertexFactoryBase* FDynamicMeshEmitterData::BuildVertexFactory(const FParticleSystemSceneProxy* InOwnerProxy)
 {
-	FParticleVertexFactoryBase* VertexFactory = GParticleVertexFactoryPool.GetParticleVertexFactory(PVFT_Mesh, InOwnerProxy->GetScene()->GetFeatureLevel());
+	FParticleVertexFactoryBase* VertexFactory = GParticleVertexFactoryPool.GetParticleVertexFactory(PVFT_Mesh, InOwnerProxy->GetScene().GetFeatureLevel());
 	SetupVertexFactory((FMeshParticleVertexFactory*)VertexFactory, StaticMesh->RenderData->LODResources[0]);
 	return VertexFactory;
 }
@@ -2954,7 +2954,7 @@ void FDynamicBeam2EmitterData::Init( bool bInSelected )
 
 FParticleVertexFactoryBase* FDynamicBeam2EmitterData::BuildVertexFactory(const FParticleSystemSceneProxy* InOwnerProxy)
 {
-	return (FParticleBeamTrailVertexFactory*)(GParticleVertexFactoryPool.GetParticleVertexFactory(PVFT_BeamTrail, InOwnerProxy->GetScene()->GetFeatureLevel()));
+	return (FParticleBeamTrailVertexFactory*)(GParticleVertexFactoryPool.GetParticleVertexFactory(PVFT_BeamTrail, InOwnerProxy->GetScene().GetFeatureLevel()));
 }
 
 /** Perform the actual work of filling the buffer, often called from another thread 
@@ -5881,7 +5881,7 @@ void FDynamicTrailsEmitterData::Init(bool bInSelected)
 
 FParticleVertexFactoryBase* FDynamicTrailsEmitterData::BuildVertexFactory(const FParticleSystemSceneProxy* InOwnerProxy)
 {
-	return (FParticleBeamTrailVertexFactory*)(GParticleVertexFactoryPool.GetParticleVertexFactory(PVFT_BeamTrail, InOwnerProxy->GetScene()->GetFeatureLevel()));
+	return (FParticleBeamTrailVertexFactory*)(GParticleVertexFactoryPool.GetParticleVertexFactory(PVFT_BeamTrail, InOwnerProxy->GetScene().GetFeatureLevel()));
 }
 
 void FDynamicTrailsEmitterData::GetDynamicMeshElementsEmitter(const FParticleSystemSceneProxy* Proxy, const FSceneView* View, const FSceneViewFamily& ViewFamily, int32 ViewIndex, FMeshElementCollector& Collector) const
diff --git a/Engine/Source/Runtime/Engine/Private/PrimitiveDrawingUtils.cpp b/Engine/Source/Runtime/Engine/Private/PrimitiveDrawingUtils.cpp
--- a/Engine/Source/Runtime/Engine/Private/PrimitiveDrawingUtils.cpp
+++ b/Engine/Source/Runtime/Engine/Private/PrimitiveDrawingUtils.cpp
@@ -15,7 +15,7 @@ void EmitMeshDrawEvents_Inner(FRHICommandList& RHICmdList, const FPrimitiveScene
 {
 	// Only show material name at the top level
 	// Note: this is the parent's material name, not the material instance
-	SCOPED_DRAW_EVENTF(RHICmdList, MaterialEvent, DEC_SCENE_ITEMS, *Mesh.MaterialRenderProxy->GetMaterial(PrimitiveSceneProxy->GetScene()->GetFeatureLevel())->GetFriendlyName());
+	SCOPED_DRAW_EVENTF(RHICmdList, MaterialEvent, DEC_SCENE_ITEMS, *Mesh.MaterialRenderProxy->GetMaterial(PrimitiveSceneProxy->GetScene().GetFeatureLevel())->GetFriendlyName());
 	if (PrimitiveSceneProxy)
 	{
 		// Show Actor, level and resource name inside the material name
diff --git a/Engine/Source/Runtime/Engine/Private/PrimitiveSceneProxy.cpp b/Engine/Source/Runtime/Engine/Private/PrimitiveSceneProxy.cpp
--- a/Engine/Source/Runtime/Engine/Private/PrimitiveSceneProxy.cpp
+++ b/Engine/Source/Runtime/Engine/Private/PrimitiveSceneProxy.cpp
@@ -73,6 +73,8 @@ FPrimitiveSceneProxy::FPrimitiveSceneProxy(const UPrimitiveComponent* InComponen
 ,	NumUncachedStaticLightingInteractions(0)
 #endif
 {
+	check(Scene);
+
 #if STATS
 	{
 		UObject const* StatObject = InComponent->AdditionalStatObject(); // prefer the additional object, this is usually the thing related to the component
diff --git a/Engine/Source/Runtime/Engine/Private/SkeletalMesh.cpp b/Engine/Source/Runtime/Engine/Private/SkeletalMesh.cpp
--- a/Engine/Source/Runtime/Engine/Private/SkeletalMesh.cpp
+++ b/Engine/Source/Runtime/Engine/Private/SkeletalMesh.cpp
@@ -3809,7 +3809,7 @@ FSkeletalMeshSceneProxy::FSkeletalMeshSceneProxy(const USkinnedMeshComponent* Co
 		,	PhysicsAssetForDebug(Component->GetPhysicsAsset())
 		,	bForceWireframe(Component->bForceWireframe)
 		,	bCanHighlightSelectedSections(Component->bCanHighlightSelectedSections)
-		,	MaterialRelevance(Component->GetMaterialRelevance(GetScene()->GetFeatureLevel()))
+		,	MaterialRelevance(Component->GetMaterialRelevance(GetScene().GetFeatureLevel()))
 		,	bMaterialsNeedMorphUsage_GameThread(false)
 {
 	check(MeshObject);
@@ -3824,7 +3824,7 @@ FSkeletalMeshSceneProxy::FSkeletalMeshSceneProxy(const USkinnedMeshComponent* Co
 		bAlwaysHasVelocity = true;
 	}
 
-	const auto FeatureLevel = GetScene()->GetFeatureLevel();
+	const auto FeatureLevel = GetScene().GetFeatureLevel();
 
 	// setup materials and performance classification for each LOD.
 	extern bool GForceDefaultMaterial;
@@ -3886,7 +3886,7 @@ FSkeletalMeshSceneProxy::FSkeletalMeshSceneProxy(const USkinnedMeshComponent* Co
 				MaterialRelevance |= Material->GetRelevance(FeatureLevel);
 			}
 
-			const bool bRequiresAdjacencyInformation = RequiresAdjacencyInformation( Material, &TGPUSkinVertexFactory<false>::StaticType, GetScene()->GetFeatureLevel() );
+			const bool bRequiresAdjacencyInformation = RequiresAdjacencyInformation( Material, &TGPUSkinVertexFactory<false>::StaticType, FeatureLevel );
 			if ( bRequiresAdjacencyInformation && LODModel.AdjacencyMultiSizeIndexContainer.IsIndexBufferValid() == false )
 			{
 				UE_LOG(LogSkeletalMesh, Warning, 
@@ -4625,7 +4625,7 @@ void FSkeletalMeshSceneProxy::UpdateMorphMaterialUsage_GameThread(bool bNeedsMor
 			UpdateSkelProxyLODSectionElementsCmd,
 				TSet<UMaterialInterface*>,MaterialsToSwap,MaterialsToSwap,
 				UMaterialInterface*,DefaultMaterial,UMaterial::GetDefaultMaterial(MD_Surface),
-				ERHIFeatureLevel::Type, FeatureLevel, GetScene()->GetFeatureLevel(),
+				ERHIFeatureLevel::Type, FeatureLevel, GetScene().GetFeatureLevel(),
 			FSkeletalMeshSceneProxy*,SkelMeshSceneProxy,this,
 			{
 					for( int32 LodIdx=0; LodIdx < SkelMeshSceneProxy->LODSections.Num(); LodIdx++ )
diff --git a/Engine/Source/Runtime/Engine/Private/StaticMeshRender.cpp b/Engine/Source/Runtime/Engine/Private/StaticMeshRender.cpp
--- a/Engine/Source/Runtime/Engine/Private/StaticMeshRender.cpp
+++ b/Engine/Source/Runtime/Engine/Private/StaticMeshRender.cpp
@@ -55,7 +55,7 @@ FStaticMeshSceneProxy::FStaticMeshSceneProxy(UStaticMeshComponent* InComponent):
 	ForcedLodModel(InComponent->ForcedLodModel),
 	bCastShadow(InComponent->CastShadow),
 	CollisionTraceFlag(ECollisionTraceFlag::CTF_UseDefault),
-	MaterialRelevance(InComponent->GetMaterialRelevance(GetScene()->GetFeatureLevel())),
+	MaterialRelevance(InComponent->GetMaterialRelevance(GetScene().GetFeatureLevel())),
 	CollisionResponse(InComponent->GetCollisionResponseToChannels())
 {
 	check(RenderData);
@@ -64,7 +64,7 @@ FStaticMeshSceneProxy::FStaticMeshSceneProxy(UStaticMeshComponent* InComponent):
 	LevelColor = FLinearColor(1,1,1);
 	PropertyColor = FLinearColor(1,1,1);
 
-	const auto FeatureLevel = GetScene()->GetFeatureLevel();
+	const auto FeatureLevel = GetScene().GetFeatureLevel();
 
 	// Copy the pointer to the volume data, async building of the data may modify the one on FStaticMeshLODResources while we are rendering
 	DistanceFieldData = RenderData->LODResources[0].DistanceFieldData;
@@ -220,7 +220,7 @@ bool FStaticMeshSceneProxy::GetMeshElement(int32 LODIndex, int32 BatchIndex, int
 	}
 
 	const bool bWireframe = false;
-	const bool bRequiresAdjacencyInformation = RequiresAdjacencyInformation( Material, OutMeshBatch.VertexFactory->GetType(), GetScene()->GetFeatureLevel() );
+	const bool bRequiresAdjacencyInformation = RequiresAdjacencyInformation( Material, OutMeshBatch.VertexFactory->GetType(), GetScene().GetFeatureLevel() );
 
 	SetIndexSource(LODIndex, SectionIndex, OutMeshBatch, bWireframe, bRequiresAdjacencyInformation );
 
@@ -309,7 +309,7 @@ void FStaticMeshSceneProxy::SetIndexSource(int32 LODIndex, int32 SectionIndex, F
 	if (bWireframe)
 	{
 		if( LODModel.WireframeIndexBuffer.IsInitialized()
-			&& !(RHISupportsTessellation(GetScene()->GetShaderPlatform()) && OutMeshElement.VertexFactory->GetType()->SupportsTessellationShaders())
+			&& !(RHISupportsTessellation(GetScene().GetShaderPlatform()) && OutMeshElement.VertexFactory->GetType()->SupportsTessellationShaders())
 			)
 		{
 			OutMeshElement.Type = PT_LineList;
@@ -425,7 +425,7 @@ void FStaticMeshSceneProxy::DrawStaticElements(FStaticPrimitiveDrawInterface* PD
 		//Never use the dynamic path in this path, because only unselected elements will use DrawStaticElements
 		bool bUseSelectedMaterial = false;
 		const bool bUseHoveredMaterial = false;
-		const auto FeatureLevel = GetScene()->GetFeatureLevel();
+		const auto FeatureLevel = GetScene().GetFeatureLevel();
 
 		//check if a LOD is being forced
 		if (ForcedLodModel > 0) 
diff --git a/Engine/Source/Runtime/Engine/Private/TextRenderComponent.cpp b/Engine/Source/Runtime/Engine/Private/TextRenderComponent.cpp
--- a/Engine/Source/Runtime/Engine/Private/TextRenderComponent.cpp
+++ b/Engine/Source/Runtime/Engine/Private/TextRenderComponent.cpp
@@ -420,7 +420,7 @@ FTextRenderSceneProxy::FTextRenderSceneProxy( UTextRenderComponent* Component) :
 	}
 
 	TextMaterial = EffectiveMaterial;
-	MaterialRelevance |= TextMaterial->GetMaterial()->GetRelevance(GetScene()->GetFeatureLevel());
+	MaterialRelevance |= TextMaterial->GetMaterial()->GetRelevance(GetScene().GetFeatureLevel());
 }
 
 FTextRenderSceneProxy::~FTextRenderSceneProxy()
diff --git a/Engine/Source/Runtime/Engine/Public/PrimitiveSceneProxy.h b/Engine/Source/Runtime/Engine/Public/PrimitiveSceneProxy.h
--- a/Engine/Source/Runtime/Engine/Public/PrimitiveSceneProxy.h
+++ b/Engine/Source/Runtime/Engine/Public/PrimitiveSceneProxy.h
@@ -318,7 +318,7 @@ public:
 	void SetCollisionEnabled_RenderThread(const bool bNewEnabled);
 
 	// Accessors.
-	inline FSceneInterface* GetScene() const { return Scene; }
+	inline FSceneInterface& GetScene() const { return *Scene; }
 	inline FPrimitiveComponentId GetPrimitiveComponentId() const { return PrimitiveComponentId; }
 	inline FPrimitiveSceneInfo* GetPrimitiveSceneInfo() const { return PrimitiveSceneInfo; }
 	inline const FMatrix& GetLocalToWorld() const { return LocalToWorld; }
diff --git a/Engine/Source/Runtime/Renderer/Private/RendererScene.cpp b/Engine/Source/Runtime/Renderer/Private/RendererScene.cpp
--- a/Engine/Source/Runtime/Renderer/Private/RendererScene.cpp
+++ b/Engine/Source/Runtime/Renderer/Private/RendererScene.cpp
@@ -381,7 +381,7 @@ void FScene::UpdatePrimitiveTransform_RenderThread(FRHICommandListImmediate& RHI
 	
 	// Update the primitive motion blur information.
 	// hack
-	FScene* Scene = (FScene*)PrimitiveSceneProxy->GetScene();
+	FScene* Scene = (FScene*)&PrimitiveSceneProxy->GetScene();
 
 	Scene->MotionBlurInfoData.UpdatePrimitiveMotionBlur(PrimitiveSceneProxy->GetPrimitiveSceneInfo());
 	
diff --git a/Engine/Source/Runtime/Renderer/Private/VelocityRendering.cpp b/Engine/Source/Runtime/Renderer/Private/VelocityRendering.cpp
--- a/Engine/Source/Runtime/Renderer/Private/VelocityRendering.cpp
+++ b/Engine/Source/Runtime/Renderer/Private/VelocityRendering.cpp
@@ -250,7 +250,7 @@ void FVelocityDrawingPolicy::SetMeshRenderState(
 	FViewInfo& View = (FViewInfo&)SceneView;	// NOTE: Assuming this cast is always safe!
 
 	// hack
-	FScene* Scene = (FScene*)PrimitiveSceneProxy->GetScene();
+	FScene* Scene = (FScene*)&PrimitiveSceneProxy->GetScene();
 
 	// previous transform can be stored in the scene for each primitive
 	if(Scene->MotionBlurInfoData.GetPrimitiveMotionBlurInfo(PrimitiveSceneProxy->GetPrimitiveSceneInfo(), PreviousLocalToWorld))
