Fix welding to non uniform not working. Also refactor AddShapesToRigidActor

diff --git a/Engine/Source/Runtime/Engine/Classes/PhysicsEngine/BodySetup.h b/Engine/Source/Runtime/Engine/Classes/PhysicsEngine/BodySetup.h
--- a/Engine/Source/Runtime/Engine/Classes/PhysicsEngine/BodySetup.h
+++ b/Engine/Source/Runtime/Engine/Classes/PhysicsEngine/BodySetup.h
@@ -587,7 +587,14 @@ public:
 	/** 
 	 *   Add the shapes defined by this body setup to the supplied PxRigidBody. 
 	 */
-	void AddShapesToRigidActor(physx::PxRigidActor* PDestActor, FVector& Scale3D, const FTransform* RelativeTM = NULL, TArray<physx::PxShape*> * NewShapes = NULL);
+	void AddShapesToRigidActor(physx::PxRigidActor* PDestActor, FVector& Scale3D, const FTransform& RelativeTM = FTransform::Identity, TArray<physx::PxShape*>* NewShapes = NULL);
+
+private:
+	void AddSpheresToRigidActor(physx::PxRigidActor* PDestActor, const FTransform& RelativeTM, float MinScale, float MinScaleAbs, TArray<physx::PxShape*>* NewShapes) const;
+	void AddBoxesToRigidActor(physx::PxRigidActor* PDestActor, const FTransform& RelativeTM, const FVector& Scale3D, const FVector& Scale3DAbs, TArray<physx::PxShape*>* NewShapes) const;
+	void AddSphylsToRigidActor(physx::PxRigidActor* PDestActor, const FTransform& RelativeTM, const FVector& Scale3D, const FVector& Scale3DAbs, TArray<physx::PxShape*>* NewShapes) const;
+	void AddConvexElemsToRigidActor(physx::PxRigidActor* PDestActor, const FTransform& RelativeTM, const FVector& Scale3D, const FVector& Scale3DAbs, TArray<physx::PxShape*>* NewShapes) const;
+	void AddTriMeshToRigidActor(physx::PxRigidActor* PDestActor, const FVector& Scale3D, const FVector& Scale3DAbs) const;
 #endif // WITH_PHYSX
 
 };
diff --git a/Engine/Source/Runtime/Engine/Private/PhysicsEngine/BodyInstance.cpp b/Engine/Source/Runtime/Engine/Private/PhysicsEngine/BodyInstance.cpp
--- a/Engine/Source/Runtime/Engine/Private/PhysicsEngine/BodyInstance.cpp
+++ b/Engine/Source/Runtime/Engine/Private/PhysicsEngine/BodyInstance.cpp
@@ -1589,8 +1589,8 @@ bool FBodyInstance::Weld(FBodyInstance* TheirBody, const FTransform& TheirTM)
 
 	FTransform RelativeTM = TheirTM.GetRelativeTransform(MyTM);
 
-	PxScene * PSyncScene = RigidActorSync ? RigidActorSync->getScene() : NULL;
-	PxScene * PAsyncScene = RigidActorAsync ? RigidActorAsync->getScene() : NULL;
+	PxScene* PSyncScene = RigidActorSync ? RigidActorSync->getScene() : NULL;
+	PxScene* PAsyncScene = RigidActorAsync ? RigidActorAsync->getScene() : NULL;
 
 	SCOPED_SCENE_WRITE_LOCK(PSyncScene);
 	SCOPED_SCENE_WRITE_LOCK(PAsyncScene);
@@ -1598,12 +1598,12 @@ bool FBodyInstance::Weld(FBodyInstance* TheirBody, const FTransform& TheirTM)
 	//child body gets placed into the same scenes as parent body
 	if (PxRigidActor* MyBody = RigidActorSync)
 	{
-		TheirBody->BodySetup->AddShapesToRigidActor(MyBody, Scale3D, &RelativeTM, &PNewShapes);
+		TheirBody->BodySetup->AddShapesToRigidActor(MyBody, Scale3D, RelativeTM, &PNewShapes);
 	}
 
 	if (PxRigidActor* MyBody = RigidActorAsync)
 	{
-		TheirBody->BodySetup->AddShapesToRigidActor(MyBody, Scale3D, &RelativeTM, &PNewShapes);
+		TheirBody->BodySetup->AddShapesToRigidActor(MyBody, Scale3D, RelativeTM, &PNewShapes);
 	}
 
 
diff --git a/Engine/Source/Runtime/Engine/Private/PhysicsEngine/BodySetup.cpp b/Engine/Source/Runtime/Engine/Private/PhysicsEngine/BodySetup.cpp
--- a/Engine/Source/Runtime/Engine/Private/PhysicsEngine/BodySetup.cpp
+++ b/Engine/Source/Runtime/Engine/Private/PhysicsEngine/BodySetup.cpp
@@ -290,235 +290,234 @@ void SetupNonUniformHelper(FVector& Scale3D, float& MinScale, float& MinScaleAbs
 	}
 }
 
-void UBodySetup::AddShapesToRigidActor(PxRigidActor* PDestActor, FVector& Scale3D, const FTransform* RelativeTM /* = NULL */, TArray<physx::PxShape*> * NewShapes /* = NULL */ )
+void GetContactOffsetParams(float& ContactOffsetFactor, float& MaxContactOffset)
 {
-#if WITH_RUNTIME_PHYSICS_COOKING || WITH_EDITOR
-	// in editor, there are a lot of things relying on body setup to create physics meshes
-	CreatePhysicsMeshes();
-#endif
+	// Get contact offset params
+	ContactOffsetFactor = CVarContactOffsetFactor.GetValueOnGameThread();
+	MaxContactOffset = CVarMaxContactOffset.GetValueOnGameThread();
+}
 
-	float MinScale;
-	float MinScaleAbs;
-	FVector Scale3DAbs;
-	SetupNonUniformHelper(Scale3D, MinScale, MinScaleAbs, Scale3DAbs);
+PxMaterial* GetDefaultPhysMaterial()
+{
+	check(GEngine->DefaultPhysMaterial != NULL);
+	return GEngine->DefaultPhysMaterial->GetPhysXMaterial();
+}
+
+void UBodySetup::AddSpheresToRigidActor(PxRigidActor* PDestActor, const FTransform& RelativeTM, float MinScale, float MinScaleAbs, TArray<PxShape*>* NewShapes) const
+{
+	float ContactOffsetFactor, MaxContactOffset;
+	GetContactOffsetParams(ContactOffsetFactor, MaxContactOffset);
+	PxMaterial* PDefaultMat = GetDefaultPhysMaterial();
 
-	if (RelativeTM)
+	for (int32 i = 0; i < AggGeom.SphereElems.Num(); i++)
 	{
-		float MinScaleRelative;
-		float MinScaleAbsRelative;
-		FVector Scale3DAbsRelative;
-		FVector Scale3DRelative = RelativeTM->GetScale3D();
+		const FKSphereElem& SphereElem = AggGeom.SphereElems[i];
 
-		SetupNonUniformHelper(Scale3DRelative, MinScaleRelative, MinScaleAbsRelative, Scale3DAbsRelative);
+		PxSphereGeometry PSphereGeom;
+		PSphereGeom.radius = (SphereElem.Radius * MinScaleAbs);
 
-		MinScaleAbs *= MinScaleAbsRelative;
-		Scale3DAbs.X *= Scale3DAbsRelative.X;
-		Scale3DAbs.Y *= Scale3DAbsRelative.Y;
-		Scale3DAbs.Z *= Scale3DAbsRelative.Z;
-	}
+		if (ensure(PSphereGeom.isValid()))
+		{
+			FVector LocalOrigin = RelativeTM.TransformPosition(SphereElem.Center);
+			PxTransform PLocalPose(U2PVector(LocalOrigin));
+			PLocalPose.p *= MinScale;
 
-	// Is the target a static actor
-	bool bDestStatic = (PDestActor->isRigidStatic() != NULL);
+			ensure(PLocalPose.isValid());
+			PxShape* NewShape = PDestActor->createShape(PSphereGeom, *PDefaultMat, PLocalPose);
 
-	// Get contact offset params
-	const float ContactOffsetFactor = CVarContactOffsetFactor.GetValueOnGameThread();
-	const float MaxContactOffset = CVarMaxContactOffset.GetValueOnGameThread();
+			if (NewShapes)
+			{
+				NewShapes->Add(NewShape);
+			}
 
-	check(GEngine->DefaultPhysMaterial != NULL);
-	PxMaterial* PDefaultMat = GEngine->DefaultPhysMaterial->GetPhysXMaterial();
+			const float ContactOffset = FMath::Min(MaxContactOffset, ContactOffsetFactor * PSphereGeom.radius);
+			NewShape->setContactOffset(ContactOffset);
+		}
+		else
+		{
+			UE_LOG(LogPhysics, Warning, TEXT("AddSpheresToRigidActor: [%s] SphereElem[%d] invalid"), *GetPathNameSafe(GetOuter()), i);
+		}
+	}
+}
 
-	// Create shapes for simple collision if we do not want to use the complex collision mesh 
-	// for simple queries as well
-	if (CollisionTraceFlag != ECollisionTraceFlag::CTF_UseComplexAsSimple)
+void UBodySetup::AddBoxesToRigidActor(PxRigidActor* PDestActor, const FTransform& RelativeTM, const FVector& Scale3D, const FVector& Scale3DAbs, TArray<PxShape*>* NewShapes) const
+{
+	float ContactOffsetFactor, MaxContactOffset;
+	GetContactOffsetParams(ContactOffsetFactor, MaxContactOffset);
+	PxMaterial* PDefaultMat = GetDefaultPhysMaterial();
+
+	for (int32 i = 0; i < AggGeom.BoxElems.Num(); i++)
 	{
-		// Sphere primitives
-		for (int32 i = 0; i < AggGeom.SphereElems.Num(); i++)
-		{
-			FKSphereElem* SphereElem = &AggGeom.SphereElems[i];
+		const FKBoxElem& BoxElem = AggGeom.BoxElems[i];
 
-			PxSphereGeometry PSphereGeom;
-			PSphereGeom.radius = (SphereElem->Radius * MinScaleAbs);
+		PxBoxGeometry PBoxGeom;
+		PBoxGeom.halfExtents.x = (0.5f * BoxElem.X * Scale3DAbs.X);
+		PBoxGeom.halfExtents.y = (0.5f * BoxElem.Y * Scale3DAbs.Y);
+		PBoxGeom.halfExtents.z = (0.5f * BoxElem.Z * Scale3DAbs.Z);
 
-			if(ensure(PSphereGeom.isValid()))
-			{
-				PxTransform PLocalPose( U2PVector( RelativeTM ? RelativeTM->TransformPosition(SphereElem->Center) : SphereElem->Center ) );
-				PLocalPose.p *= MinScale;
+		FTransform BoxTransform = BoxElem.GetTransform() * RelativeTM;
+		if (PBoxGeom.isValid() && BoxTransform.IsValid())
+		{
+			PxTransform PLocalPose(U2PTransform(BoxTransform));
+			PLocalPose.p.x *= Scale3D.X;
+			PLocalPose.p.y *= Scale3D.Y;
+			PLocalPose.p.z *= Scale3D.Z;
 
-				ensure(PLocalPose.isValid());
-				PxShape* NewShape = PDestActor->createShape(PSphereGeom, *PDefaultMat, PLocalPose);
-				
-				if (NewShapes)
-				{
-					NewShapes->Add(NewShape);
-				}
+			ensure(PLocalPose.isValid());
+			PxShape* NewShape = PDestActor->createShape(PBoxGeom, *PDefaultMat, PLocalPose);
 
-				const float ContactOffset = FMath::Min(MaxContactOffset, ContactOffsetFactor * PSphereGeom.radius);
-				NewShape->setContactOffset(ContactOffset);
-			}
-			else
+			if (NewShapes)
 			{
-				UE_LOG(LogPhysics, Warning, TEXT("AddShapesToRigidActor: [%s] SphereElem[%d] invalid"), *GetPathNameSafe(GetOuter()), i);
+				NewShapes->Add(NewShape);
 			}
-		}
 
-		// Box primitives
-		for (int32 i = 0; i < AggGeom.BoxElems.Num(); i++)
+			const float ContactOffset = FMath::Min(MaxContactOffset, ContactOffsetFactor * PBoxGeom.halfExtents.minElement());
+			NewShape->setContactOffset(ContactOffset);
+		}
+		else
 		{
-			FKBoxElem* BoxElem = &AggGeom.BoxElems[i];
+			UE_LOG(LogPhysics, Warning, TEXT("AddBoxesToRigidActor: [%s] BoxElems[%d] invalid or has invalid transform"), *GetPathNameSafe(GetOuter()), i);
+		}
+	}
+}
 
-			PxBoxGeometry PBoxGeom;
-			PBoxGeom.halfExtents.x = (0.5f * BoxElem->X * Scale3DAbs.X);
-			PBoxGeom.halfExtents.y = (0.5f * BoxElem->Y * Scale3DAbs.Y);
-			PBoxGeom.halfExtents.z = (0.5f * BoxElem->Z * Scale3DAbs.Z);
+void UBodySetup::AddSphylsToRigidActor(PxRigidActor* PDestActor, const FTransform& RelativeTM, const FVector& Scale3D, const FVector& Scale3DAbs, TArray<PxShape*>* NewShapes) const
+{
+	float ContactOffsetFactor, MaxContactOffset;
+	GetContactOffsetParams(ContactOffsetFactor, MaxContactOffset);
+	PxMaterial* PDefaultMat = GetDefaultPhysMaterial();
 
-			FTransform BoxTransform = RelativeTM ? BoxElem->GetTransform() * *RelativeTM : BoxElem->GetTransform();
-			if(PBoxGeom.isValid() && BoxTransform.IsValid())
-			{
-				PxTransform PLocalPose( U2PTransform( BoxTransform ));
-				PLocalPose.p.x *= Scale3D.X;
-				PLocalPose.p.y *= Scale3D.Y;
-				PLocalPose.p.z *= Scale3D.Z;
+	float ScaleRadius = FMath::Max(Scale3DAbs.X, Scale3DAbs.Y);
+	float ScaleLength = Scale3DAbs.Z;
 
-				ensure(PLocalPose.isValid());
-				PxShape* NewShape = PDestActor->createShape(PBoxGeom, *PDefaultMat, PLocalPose);
-	
-				if (NewShapes)
-				{
-					NewShapes->Add(NewShape);
-				}
+	for (int32 i = 0; i < AggGeom.SphylElems.Num(); i++)
+	{
+		const FKSphylElem& SphylElem = AggGeom.SphylElems[i];
+
+		// this is a bit confusing since radius and height is scaled
+		// first apply the scale first 
+		float Radius = FMath::Max(SphylElem.Radius * ScaleRadius, 0.1f);
+		float Length = SphylElem.Length + SphylElem.Radius * 2.f;
+		float HalfLength = Length * ScaleLength * 0.5f;
+		Radius = FMath::Clamp(Radius, 0.1f, HalfLength);	//radius is capped by half length
+		float HalfHeight = HalfLength - Radius;
+		HalfHeight = FMath::Max(0.1f, HalfHeight);
+
+		PxCapsuleGeometry PCapsuleGeom;
+		PCapsuleGeom.halfHeight = HalfHeight;
+		PCapsuleGeom.radius = Radius;
+
+		if (PCapsuleGeom.isValid())
+		{
+			// The stored sphyl transform assumes the sphyl axis is down Z. In PhysX, it points down X, so we twiddle the matrix a bit here (swap X and Z and negate Y).
+			PxTransform PLocalPose(U2PVector(RelativeTM.TransformPosition(SphylElem.Center)), U2PQuat(SphylElem.Orientation) * U2PSphylBasis);
+			PLocalPose.p.x *= Scale3D.X;
+			PLocalPose.p.y *= Scale3D.Y;
+			PLocalPose.p.z *= Scale3D.Z;
 
-				const float ContactOffset = FMath::Min(MaxContactOffset, ContactOffsetFactor * PBoxGeom.halfExtents.minElement());
-				NewShape->setContactOffset(ContactOffset);
-			}
-			else
+			ensure(PLocalPose.isValid());
+			PxShape* NewShape = PDestActor->createShape(PCapsuleGeom, *PDefaultMat, PLocalPose);
+
+			if (NewShapes)
 			{
-				UE_LOG(LogPhysics, Warning, TEXT("AddShapesToRigidActor: [%s] BoxElems[%d] invalid or has invalid transform"), *GetPathNameSafe(GetOuter()), i);
+				NewShapes->Add(NewShape);
 			}
-		}
 
-		// Sphyl (aka Capsule) primitives
+			const float ContactOffset = FMath::Min(MaxContactOffset, ContactOffsetFactor * PCapsuleGeom.radius);
+			NewShape->setContactOffset(ContactOffset);
+		}
+		else
 		{
-			float ScaleRadius = FMath::Max(Scale3DAbs.X, Scale3DAbs.Y);
-			float ScaleLength = Scale3DAbs.Z;
+			UE_LOG(LogPhysics, Warning, TEXT("AddSphylsToRigidActor: [%s] SphylElems[%d] invalid"), *GetPathNameSafe(GetOuter()), i);
+		}
+	}
+}
+
+void UBodySetup::AddConvexElemsToRigidActor(PxRigidActor* PDestActor, const FTransform& RelativeTM, const FVector& Scale3D, const FVector& Scale3DAbs, TArray<PxShape*>* NewShapes) const
+{
+	float ContactOffsetFactor, MaxContactOffset;
+	GetContactOffsetParams(ContactOffsetFactor, MaxContactOffset);
+	PxMaterial* PDefaultMat = GetDefaultPhysMaterial();
 
-			for (int32 i = 0; i < AggGeom.SphylElems.Num(); i++)
+	for (int32 i = 0; i < AggGeom.ConvexElems.Num(); i++)
+	{
+		const FKConvexElem& ConvexElem = AggGeom.ConvexElems[i];
+
+		if (ConvexElem.ConvexMesh)
+		{
+			PxTransform PLocalPose;
+			bool bUseNegX = CalcMeshNegScaleCompensation(Scale3D, PLocalPose);
+
+			PxConvexMeshGeometry PConvexGeom;
+			PConvexGeom.convexMesh = bUseNegX ? ConvexElem.ConvexMeshNegX : ConvexElem.ConvexMesh;
+			PConvexGeom.scale.scale = U2PVector(Scale3DAbs * ConvexElem.GetTransform().GetScale3D().GetAbs());
+			FTransform ConvexTransform = ConvexElem.GetTransform();
+			if (ConvexTransform.GetScale3D().X < 0 || ConvexTransform.GetScale3D().Y < 0 || ConvexTransform.GetScale3D().Z < 0)
+			{
+				UE_LOG(LogPhysics, Warning, TEXT("AddConvexElemsToRigidActor: [%s] ConvexElem[%d] has negative scale. Not currently supported"), *GetPathNameSafe(GetOuter()), i);
+			}
+			if (ConvexTransform.IsValid())
 			{
-				FKSphylElem* SphylElem = &AggGeom.SphylElems[i];
-
-				// this is a bit confusing since radius and height is scaled
-				// first apply the scale first 
-				float Radius = FMath::Max(SphylElem->Radius * ScaleRadius, 0.1f);
-				float Length = SphylElem->Length + SphylElem->Radius * 2.f;
-				float HalfLength = Length * ScaleLength * 0.5f;
-				Radius = FMath::Clamp(Radius, 0.1f, HalfLength);	//radius is capped by half length
-				float HalfHeight = HalfLength - Radius;
-				HalfHeight = FMath::Max(0.1f, HalfHeight);
-
-				PxCapsuleGeometry PCapsuleGeom;
-				PCapsuleGeom.halfHeight = HalfHeight;
-				PCapsuleGeom.radius = Radius;
-
-				if (PCapsuleGeom.isValid())
+				PxTransform PElementTransform = U2PTransform(ConvexTransform * RelativeTM);
+				PLocalPose.q *= PElementTransform.q;
+				PLocalPose.p = PElementTransform.p;
+				PLocalPose.p.x *= Scale3D.X;
+				PLocalPose.p.y *= Scale3D.Y;
+				PLocalPose.p.z *= Scale3D.Z;
+
+				if (PConvexGeom.isValid())
 				{
-					// The stored sphyl transform assumes the sphyl axis is down Z. In PhysX, it points down X, so we twiddle the matrix a bit here (swap X and Z and negate Y).
-					PxTransform PLocalPose(U2PVector(RelativeTM ? RelativeTM->TransformPosition(SphylElem->Center) : SphylElem->Center), U2PQuat(SphylElem->Orientation) * U2PSphylBasis);
-					PLocalPose.p.x *= Scale3D.X;
-					PLocalPose.p.y *= Scale3D.Y;
-					PLocalPose.p.z *= Scale3D.Z;
+					PxVec3 PBoundsExtents = PConvexGeom.convexMesh->getLocalBounds().getExtents();
 
 					ensure(PLocalPose.isValid());
-					PxShape* NewShape = PDestActor->createShape(PCapsuleGeom, *PDefaultMat, PLocalPose);
-	
+					PxShape* NewShape = PDestActor->createShape(PConvexGeom, *PDefaultMat, PLocalPose);
+
 					if (NewShapes)
 					{
 						NewShapes->Add(NewShape);
 					}
 
-					const float ContactOffset = FMath::Min(MaxContactOffset, ContactOffsetFactor * PCapsuleGeom.radius);
+					const float ContactOffset = FMath::Min(MaxContactOffset, ContactOffsetFactor * PBoundsExtents.minElement());
 					NewShape->setContactOffset(ContactOffset);
 				}
 				else
 				{
-					UE_LOG(LogPhysics, Warning, TEXT("AddShapesToRigidActor: [%s] SphylElems[%d] invalid"), *GetPathNameSafe(GetOuter()), i);
+					UE_LOG(LogPhysics, Warning, TEXT("AddConvexElemsToRigidActor: [%s] ConvexElem[%d] invalid"), *GetPathNameSafe(GetOuter()), i);
 				}
 			}
-		}
-
-		// Create convex shapes
-		for(int32 i=0; i<AggGeom.ConvexElems.Num(); i++)
-		{
-			FKConvexElem* ConvexElem = &(AggGeom.ConvexElems[i]);
-
-			if(ConvexElem->ConvexMesh != NULL)
-			{
-				PxTransform PLocalPose;
-				bool bUseNegX = CalcMeshNegScaleCompensation(Scale3D, PLocalPose);
-
-				PxConvexMeshGeometry PConvexGeom;
-				PConvexGeom.convexMesh = bUseNegX ? ConvexElem->ConvexMeshNegX : ConvexElem->ConvexMesh;
-				PConvexGeom.scale.scale = U2PVector(Scale3DAbs * ConvexElem->GetTransform().GetScale3D().GetAbs());
-				FTransform ConvexTransform = ConvexElem->GetTransform();
-				if (ConvexTransform.GetScale3D().X < 0 || ConvexTransform.GetScale3D().Y < 0 || ConvexTransform.GetScale3D().Z < 0)
-				{
-					UE_LOG(LogPhysics, Warning, TEXT("AddShapesToRigidActor: [%s] ConvexElem[%d] has negative scale. Not currently supported"), *GetPathNameSafe(GetOuter()), i);
-				}
-				if (ConvexTransform.IsValid())
-				{
-					PxTransform PElementTransform = U2PTransform(RelativeTM ? ConvexTransform * *RelativeTM : ConvexTransform);
-					PLocalPose.q *= PElementTransform.q;
-					PLocalPose.p += PElementTransform.p * MinScale;
-
-					if(PConvexGeom.isValid())
-					{
-						PxVec3 PBoundsExtents = PConvexGeom.convexMesh->getLocalBounds().getExtents();
-
-						ensure(PLocalPose.isValid());
-						PxShape* NewShape = PDestActor->createShape(PConvexGeom, *PDefaultMat, PLocalPose);
-	
-						if (NewShapes)
-						{
-							NewShapes->Add(NewShape);
-						}
-
-						const float ContactOffset = FMath::Min(MaxContactOffset, ContactOffsetFactor * PBoundsExtents.minElement());
-						NewShape->setContactOffset(ContactOffset);
-					}
-					else
-					{
-						UE_LOG(LogPhysics, Warning, TEXT("AddShapesToRigidActor: [%s] ConvexElem[%d] invalid"), *GetPathNameSafe(GetOuter()), i);
-					}
-				}
-				else
-				{
-					UE_LOG(LogPhysics, Warning, TEXT("AddShapesToRigidActor: [%s] ConvexElem[%d] has invalid transform"), *GetPathNameSafe(GetOuter()), i);
-				}				
-			}
 			else
 			{
-				UE_LOG(LogPhysics, Warning, TEXT("AddShapesToRigidActor: ConvexElem is missing ConvexMesh (%d: %s)"), i, *GetPathName());
+				UE_LOG(LogPhysics, Warning, TEXT("AddConvexElemsToRigidActor: [%s] ConvexElem[%d] has invalid transform"), *GetPathNameSafe(GetOuter()), i);
 			}
 		}
+		else
+		{
+			UE_LOG(LogPhysics, Warning, TEXT("AddConvexElemsToRigidActor: ConvexElem is missing ConvexMesh (%d: %s)"), i, *GetPathName());
+		}
 	}
+}
 
+void UBodySetup::AddTriMeshToRigidActor(PxRigidActor* PDestActor, const FVector& Scale3D, const FVector& Scale3DAbs) const
+{
+	float ContactOffsetFactor, MaxContactOffset;
+	GetContactOffsetParams(ContactOffsetFactor, MaxContactOffset);
+	PxMaterial* PDefaultMat = GetDefaultPhysMaterial();
 
-	// Create tri-mesh shape, when we are not using simple collision shapes for 
-	// complex queries as well
-	if( CollisionTraceFlag != ECollisionTraceFlag::CTF_UseSimpleAsComplex && 
-	    (TriMesh != NULL || TriMeshNegX != NULL) )
+	if(TriMesh || TriMeshNegX)
 	{
 		PxTransform PLocalPose;
 		bool bUseNegX = CalcMeshNegScaleCompensation(Scale3D, PLocalPose);
 
 		// Only case where TriMeshNegX should be null is BSP, which should not require negX version
-		if(bUseNegX && TriMeshNegX == NULL)
+		if (bUseNegX && TriMeshNegX == NULL)
 		{
-			UE_LOG(LogPhysics, Log,  TEXT("AddShapesToRigidActor: Want to use NegX but it doesn't exist! %s"), *GetPathName() );
+			UE_LOG(LogPhysics, Log, TEXT("AddTriMeshToRigidActor: Want to use NegX but it doesn't exist! %s"), *GetPathName());
 		}
 
 		PxTriangleMesh* UseTriMesh = bUseNegX ? TriMeshNegX : TriMesh;
-		if(UseTriMesh != NULL)
+		if (UseTriMesh != NULL)
 		{
-			
+
 
 			PxTriangleMeshGeometry PTriMeshGeom;
 			PTriMeshGeom.triangleMesh = bUseNegX ? TriMeshNegX : TriMesh;
@@ -527,15 +526,15 @@ void UBodySetup::AddShapesToRigidActor(PxRigidActor* PDestActor, FVector& Scale3
 			{
 				PTriMeshGeom.meshFlags |= PxMeshGeometryFlag::eDOUBLE_SIDED;
 			}
-			
 
-			if(PTriMeshGeom.isValid())
+
+			if (PTriMeshGeom.isValid())
 			{
 				ensure(PLocalPose.isValid());
 
 				// Create without 'sim shape' flag, problematic if it's kinematic, and it gets set later anyway.
 				PxShape* NewShape = PDestActor->createShape(PTriMeshGeom, *PDefaultMat, PxShapeFlag::eSCENE_QUERY_SHAPE | PxShapeFlag::eVISUALIZATION);
-				if(NewShape)
+				if (NewShape)
 				{
 					NewShape->setLocalPose(PLocalPose);
 					NewShape->setContactOffset(MaxContactOffset);
@@ -547,12 +546,57 @@ void UBodySetup::AddShapesToRigidActor(PxRigidActor* PDestActor, FVector& Scale3
 			}
 			else
 			{
-				UE_LOG(LogPhysics, Log, TEXT("AddShapesToRigidActor: TriMesh invalid"));
+				UE_LOG(LogPhysics, Log, TEXT("AddTriMeshToRigidActor: TriMesh invalid"));
 			}
 		}
 	}
 }
 
+void UBodySetup::AddShapesToRigidActor(PxRigidActor* PDestActor, FVector& Scale3D, const FTransform& RelativeTM /* = FTransform::Identity */, TArray<physx::PxShape*>* NewShapes /* = NULL */ )
+{
+#if WITH_RUNTIME_PHYSICS_COOKING || WITH_EDITOR
+	// in editor, there are a lot of things relying on body setup to create physics meshes
+	CreatePhysicsMeshes();
+#endif
+
+	float MinScale;
+	float MinScaleAbs;
+	FVector Scale3DAbs;
+	SetupNonUniformHelper(Scale3D, MinScale, MinScaleAbs, Scale3DAbs);
+
+	{
+		float MinScaleRelative;
+		float MinScaleAbsRelative;
+		FVector Scale3DAbsRelative;
+		FVector Scale3DRelative = RelativeTM.GetScale3D();
+
+		SetupNonUniformHelper(Scale3DRelative, MinScaleRelative, MinScaleAbsRelative, Scale3DAbsRelative);
+
+		MinScaleAbs *= MinScaleAbsRelative;
+		Scale3DAbs.X *= Scale3DAbsRelative.X;
+		Scale3DAbs.Y *= Scale3DAbsRelative.Y;
+		Scale3DAbs.Z *= Scale3DAbsRelative.Z;
+	}
+
+	// Create shapes for simple collision if we do not want to use the complex collision mesh 
+	// for simple queries as well
+	if (CollisionTraceFlag != ECollisionTraceFlag::CTF_UseComplexAsSimple)
+	{
+		AddSpheresToRigidActor(PDestActor, RelativeTM, MinScale, MinScaleAbs, NewShapes);
+		AddBoxesToRigidActor(PDestActor, RelativeTM, Scale3D, Scale3DAbs, NewShapes);
+		AddSphylsToRigidActor(PDestActor, RelativeTM, Scale3D, Scale3DAbs, NewShapes);
+		AddConvexElemsToRigidActor(PDestActor, RelativeTM, Scale3D, Scale3DAbs, NewShapes);
+	}
+
+
+	// Create tri-mesh shape, when we are not using simple collision shapes for 
+	// complex queries as well
+	if( CollisionTraceFlag != ECollisionTraceFlag::CTF_UseSimpleAsComplex )
+	{
+		AddTriMeshToRigidActor(PDestActor, Scale3D, Scale3DAbs);
+	}
+}
+
 #endif // WITH_PHYSX
 
 
