Engine: Fixed assertion when accessing non-existent virtual joystick

diff --git a/Engine/Source/Runtime/Engine/Private/PlayerController.cpp b/Engine/Source/Runtime/Engine/Private/PlayerController.cpp
--- a/Engine/Source/Runtime/Engine/Private/PlayerController.cpp
+++ b/Engine/Source/Runtime/Engine/Private/PlayerController.cpp
@@ -1022,7 +1022,7 @@ void APlayerController::CreateTouchInterface()
 	ULocalPlayer* LocalPlayer = Cast<ULocalPlayer>(Player);
 
 	// do we want to show virtual joysticks?
-	if (LocalPlayer && LocalPlayer->ViewportClient && !VirtualJoystick.IsValid() && SVirtualJoystick::ShouldDisplayTouchInterface())
+	if (LocalPlayer && LocalPlayer->ViewportClient && SVirtualJoystick::ShouldDisplayTouchInterface())
 	{
 		// in case we already had one, remove it
 		if (VirtualJoystick.IsValid())
@@ -3676,7 +3676,6 @@ void APlayerController::TickPlayerInput(const float DeltaSeconds, const bool bGa
 		if (SVirtualJoystick::ShouldDisplayTouchInterface())
 		{
 			CreateTouchInterface();
-			VirtualJoystick->SetVisibility(EVisibility::Visible);
 		}
 		else
 		{
