Engine: Removed enforcement of virtual joystick in TickPlayerInput

diff --git a/Engine/Source/Runtime/Engine/Private/PlayerController.cpp b/Engine/Source/Runtime/Engine/Private/PlayerController.cpp
--- a/Engine/Source/Runtime/Engine/Private/PlayerController.cpp
+++ b/Engine/Source/Runtime/Engine/Private/PlayerController.cpp
@@ -3671,16 +3671,6 @@ void APlayerController::TickPlayerInput(const float DeltaSeconds, const bool bGa
 				CurrentTouchablePrimitives[TouchIndexInt] = CurrentComponent;
 			}
 		}
-
-		// force touch UI on when it should be
-		if (SVirtualJoystick::ShouldDisplayTouchInterface())
-		{
-			CreateTouchInterface();
-		}
-		else
-		{
-			CleanupGameViewport();
-		}
 	}
 
 	if (PlayerInput != NULL)
