Fix non-unity build error
Remove incorrect casting of r.FullScreenMode value to EWindowMode

diff --git a/Engine/Source/Runtime/Engine/Private/GameViewportClient.cpp b/Engine/Source/Runtime/Engine/Private/GameViewportClient.cpp
--- a/Engine/Source/Runtime/Engine/Private/GameViewportClient.cpp
+++ b/Engine/Source/Runtime/Engine/Private/GameViewportClient.cpp
@@ -34,6 +34,8 @@ extern ENGINE_API UPrimitiveComponent* GDebugSelectedComponent;
 /** The lightmap used by the currently selected component, if it's a static mesh component. */
 extern ENGINE_API class FLightMap2D* GDebugSelectedLightmap;
 
+extern int32 GetBoundFullScreenModeCVar();
+
 /** Delegate called at the end of the frame when a screenshot is captured */
 FOnScreenshotCaptured UGameViewportClient::ScreenshotCapturedDelegate;
 
@@ -2429,7 +2431,7 @@ bool UGameViewportClient::SetDisplayConfiguration(const FIntPoint* Dimensions, E
 
 bool UGameViewportClient::HandleToggleFullscreenCommand(const TCHAR* Cmd, FOutputDevice& Ar)
 {
-	auto FullScreenMode = (EWindowMode::Type)GetBoundFullScreenModeCVar() == EWindowMode::Fullscreen ? EWindowMode::Fullscreen : EWindowMode::WindowedFullscreen;
+	auto FullScreenMode = GetBoundFullScreenModeCVar() == 0 ? EWindowMode::Fullscreen : EWindowMode::WindowedFullscreen;
 	FSystemResolution::RequestResolutionChange(GSystemResolution.ResX, GSystemResolution.ResY, Viewport->IsFullscreen() ? EWindowMode::Windowed : FullScreenMode);
 	return true;
 }
