fixed check()

diff --git a/Engine/Source/Runtime/Engine/Classes/Materials/MaterialInterface.h b/Engine/Source/Runtime/Engine/Classes/Materials/MaterialInterface.h
--- a/Engine/Source/Runtime/Engine/Classes/Materials/MaterialInterface.h
+++ b/Engine/Source/Runtime/Engine/Classes/Materials/MaterialInterface.h
@@ -340,6 +340,7 @@ public:
 	/** @return The material's relevance, from concurrent render thread updates. */
 	ENGINE_API FMaterialRelevance GetRelevance_Concurrent(ERHIFeatureLevel::Type InFeatureLevel) const;
 private:
+	// might get called from game or render thread
 	ENGINE_API FMaterialRelevance GetRelevance_Internal(const UMaterial* Material, ERHIFeatureLevel::Type InFeatureLevel) const;
 public:
 
diff --git a/Engine/Source/Runtime/Engine/Private/Materials/MaterialInterface.cpp b/Engine/Source/Runtime/Engine/Private/Materials/MaterialInterface.cpp
--- a/Engine/Source/Runtime/Engine/Private/Materials/MaterialInterface.cpp
+++ b/Engine/Source/Runtime/Engine/Private/Materials/MaterialInterface.cpp
@@ -45,8 +45,7 @@ FMaterialRelevance UMaterialInterface::GetRelevance_Internal(const UMaterial* Ma
 		const FMaterialResource* MaterialResource = Material->GetMaterialResource(InFeatureLevel);
 		const bool bIsTranslucent = IsTranslucentBlendMode((EBlendMode)GetBlendMode());
 
-		// get the value from the game thread
-		EMaterialShadingModel ShadingModel = GetShadingModel_Internal();
+		EMaterialShadingModel ShadingModel = GetShadingModel();
 
 		const bool bIsLit =  ShadingModel != MSM_Unlit;
 		// Determine the material's view relevance.
