Removed unused code

diff --git a/Engine/Source/Runtime/Engine/Classes/Engine/ActorChannel.h b/Engine/Source/Runtime/Engine/Classes/Engine/ActorChannel.h
--- a/Engine/Source/Runtime/Engine/Classes/Engine/ActorChannel.h
+++ b/Engine/Source/Runtime/Engine/Classes/Engine/ActorChannel.h
@@ -47,8 +47,6 @@ class ENGINE_API UActorChannel : public UChannel
 	// Variables.
 	UPROPERTY()
 	AActor* Actor;			// Actor this corresponds to.
-	UPROPERTY()
-	UClass* ActorClass;		// Class of the actor.
 
 	// Variables.
 	double	RelevantTime;			// Last time this actor was relevant to client.
diff --git a/Engine/Source/Runtime/Engine/Private/DataChannel.cpp b/Engine/Source/Runtime/Engine/Private/DataChannel.cpp
--- a/Engine/Source/Runtime/Engine/Private/DataChannel.cpp
+++ b/Engine/Source/Runtime/Engine/Private/DataChannel.cpp
@@ -1250,7 +1250,7 @@ void UActorChannel::SetClosingFlag()
 
 void UActorChannel::Close()
 {
-	UE_LOG(LogNetTraffic, Log, TEXT("UActorChannel::Close: ChIndex: %d, Actor: %s, ActorClass: %s, Name: %s"), ChIndex, Actor ? *Actor->GetPathName() : TEXT("NULL"), ActorClass ? *ActorClass->GetName() : TEXT("NULL"), *GetName() );
+	UE_LOG(LogNetTraffic, Log, TEXT("UActorChannel::Close: ChIndex: %d, Actor: %s"), ChIndex, Actor ? *Actor->GetFullName() : TEXT("NULL") );
 
 	UChannel::Close();
 
@@ -1302,7 +1302,6 @@ void UActorChannel::CleanupReplicators( const bool bKeepReplicators )
 	ReplicationMap.Empty();
 
 	ActorReplicator = NULL;
-	ActorClass		= NULL;
 }
 
 void UActorChannel::CleanUp()
@@ -1392,11 +1391,9 @@ void UActorChannel::SetChannelActor( AActor* InActor )
 	check(Actor==NULL);
 
 	// Set stuff.
-	Actor                      = InActor;
-	ActorClass                 = Actor->GetClass();
-	FClassNetCache* ClassCache = Connection->Driver->NetCache->GetClassNetCache( ActorClass );
+	Actor = InActor;
 
-	UE_LOG(LogNetTraffic, VeryVerbose, TEXT("SetChannelActor[%d]: Actor: %s ActorClass: %s. %s"), ChIndex, Actor ? *Actor->GetPathName() : TEXT("NULL"), ActorClass ? *ActorClass->GetName() : TEXT("NULL"), *GetName() );
+	UE_LOG(LogNetTraffic, VeryVerbose, TEXT("SetChannelActor[%d]: Actor: %s"), ChIndex, Actor ? *Actor->GetFullName() : TEXT("NULL") );
 
 	if ( Connection->PendingOutRec[ChIndex] > 0 )
 	{
@@ -1910,7 +1907,7 @@ FString UActorChannel::Describe()
 	if( Closing || !Actor )
 		return FString(TEXT("Actor=None ")) + UChannel::Describe();
 	else
-		return FString::Printf(TEXT("Actor=%s (Role=%i RemoteRole=%i) ActorClass: %s"), *Actor->GetFullName(), (int32)Actor->Role, (int32)Actor->GetRemoteRole(), ActorClass ? *ActorClass->GetName() : TEXT("NULL") ) + UChannel::Describe();
+		return FString::Printf(TEXT("Actor=%s (Role=%i RemoteRole=%i)"), *Actor->GetFullName(), (int32)Actor->Role, (int32)Actor->GetRemoteRole() ) + UChannel::Describe();
 }
 
 
