If the light volume is already attached to a scene then don't do it again.

diff --git a/Engine/Source/Runtime/Engine/Private/Level.cpp b/Engine/Source/Runtime/Engine/Private/Level.cpp
--- a/Engine/Source/Runtime/Engine/Private/Level.cpp
+++ b/Engine/Source/Runtime/Engine/Private/Level.cpp
@@ -1614,11 +1614,14 @@ void ULevel::InitializeRenderingResources()
 	// At the point at which Pre/PostEditChange is called on that transient ULevel, it is not part of any world and therefore should not have its rendering resources initialized
 	if (OwningWorld)
 	{
-		PrecomputedLightVolume->AddToScene(OwningWorld->Scene);
-
-		if (OwningWorld->Scene)
+		if( !PrecomputedLightVolume->IsAddedToScene() )
 		{
-			OwningWorld->Scene->OnLevelAddedToWorld(GetOutermost()->GetFName());
+			PrecomputedLightVolume->AddToScene(OwningWorld->Scene);
+
+			if (OwningWorld->Scene)
+			{
+				OwningWorld->Scene->OnLevelAddedToWorld(GetOutermost()->GetFName());
+			}
 		}
 	}
 }
