#ue4 - Deprecated CharacterMovementComponent::AdjustUpperHemisphereImpact() and UpperImpactNormalScale.

- Implementations of AdjustUpperHemisphereImpact() should override ComputeSlideVector() instead.

diff --git a/Engine/Source/Runtime/Core/Public/UObject/ObjectVersion.h b/Engine/Source/Runtime/Core/Public/UObject/ObjectVersion.h
--- a/Engine/Source/Runtime/Core/Public/UObject/ObjectVersion.h
+++ b/Engine/Source/Runtime/Core/Public/UObject/ObjectVersion.h
@@ -624,6 +624,8 @@ enum EUnrealEngineObjectUE4Version
 	VER_UE4_GLOBAL_EMITTER_SPAWN_RATE_SCALE,
 	// Cleanup destructible mesh settings
 	VER_UE4_CLEAN_DESTRUCTIBLE_SETTINGS,
+	// CharacterMovementComponent refactor of AdjustUpperHemisphereImpact and deprecation of some associated vars.
+	VER_UE4_CHARACTER_MOVEMENT_UPPER_IMPACT_BEHAVIOR,
 
 	// -----<new versions can be added before this line>-------------------------------------------------
 	// - this needs to be the last line (see note below)
diff --git a/Engine/Source/Runtime/Engine/Classes/GameFramework/CharacterMovementComponent.h b/Engine/Source/Runtime/Engine/Classes/GameFramework/CharacterMovementComponent.h
--- a/Engine/Source/Runtime/Engine/Classes/GameFramework/CharacterMovementComponent.h
+++ b/Engine/Source/Runtime/Engine/Classes/GameFramework/CharacterMovementComponent.h
@@ -439,6 +439,9 @@ public:
 	UPROPERTY()
 	float CrouchedSpeedMultiplier_DEPRECATED;
 
+	UPROPERTY()
+	float UpperImpactNormalScale_DEPRECATED;
+
 protected:
 
 	/**
@@ -514,14 +517,6 @@ public:
 	UPROPERTY(Category="Character Movement", EditAnywhere, BlueprintReadWrite, AdvancedDisplay)
 	float JumpOutOfWaterPitch;
 
-	/**
-	 * Impacts on the upper hemisphere will have the vertical velocity scaled by this factor of the normal's Z component (in AdjustUpperHemisphereImpact()).
-	 * This can dampen such impacts to feel less like a capsule. A value of 0 indicates that this behavior is disabled.
-	 * @see AdjustUpperHemisphereImpact()
-	 */
-	UPROPERTY(Category="Character Movement", EditAnywhere, BlueprintReadWrite, AdvancedDisplay)
-	float UpperImpactNormalScale;
-
 	/** Information about the floor the Character is standing on (updated only during walking movement). */
 	UPROPERTY(Category="Character Movement", VisibleInstanceOnly, BlueprintReadOnly)
 	FFindFloorResult CurrentFloor;
@@ -1286,8 +1281,8 @@ protected:
 
 	/**
 	 * Calculate slide vector along a surface.
-	 * Has special treatment when falling, to avoid boosting up slopes (calling HandleSlopeBoosting in this case).
-	 * Calls AdjustUpperHemisphereImpact() for upward falling movement that hits the top of the capsule, because commonly we don't want this to behave like a smooth capsule.
+	 * Has special treatment when falling, to avoid boosting up slopes (calling HandleSlopeBoosting() in this case).
+	 *
 	 * @param Delta:	Attempted move.
 	 * @param Time:		Amount of move to apply (between 0 and 1).
 	 * @param Normal:	Normal opposed to movement. Not necessarily equal to Hit.Normal (but usually is).
@@ -1297,10 +1292,9 @@ protected:
 	virtual FVector ComputeSlideVector(const FVector& Delta, const float Time, const FVector& Normal, const FHitResult& Hit) const override;
 
 	/**
-	 * Given an upward impact on the top of the capsule, allows calculation of a different movement delta. This can dampen such impacts to feel less like a capsule.
-	 * Default implementation scales the Delta Z value by: 1 - (Abs(Normal.Z) * UpperImpactNormalScale), clamped to [0..1]
-	 * @see UpperImpactNormalScale
+	 * (DEPRECATED) Given an upward impact on the top of the capsule, allows calculation of a different movement delta.
 	 */
+	DEPRECATED(4.6, "AdjustUpperHemisphereImpact() is deprecated, for custom behavior override ComputeSlideVector() instead.")
 	virtual FVector AdjustUpperHemisphereImpact(const FVector& Delta, const FHitResult& Hit) const;
 
 	/** 
diff --git a/Engine/Source/Runtime/Engine/Private/CharacterMovementComponent.cpp b/Engine/Source/Runtime/Engine/Private/CharacterMovementComponent.cpp
--- a/Engine/Source/Runtime/Engine/Private/CharacterMovementComponent.cpp
+++ b/Engine/Source/Runtime/Engine/Private/CharacterMovementComponent.cpp
@@ -192,7 +192,8 @@ UCharacterMovementComponent::UCharacterMovementComponent(const class FPostConstr
 	BrakingDecelerationSwimming = 0.f;
 	LedgeCheckThreshold = 4.0f;
 	JumpOutOfWaterPitch = 11.25f;
-	UpperImpactNormalScale = 0.5f;
+	UpperImpactNormalScale_DEPRECATED = 0.5f;
+	
 	Mass = 100.0f;
 	bJustTeleported = true;
 	CrouchedHalfHeight = 40.0f;
@@ -1949,30 +1950,12 @@ void UCharacterMovementComponent::TwoWallAdjust(FVector &Delta, const FHitResult
 }
 
 
-FVector UCharacterMovementComponent::ComputeSlideVector(const FVector& InDelta, const float Time, const FVector& Normal, const FHitResult& Hit) const
+FVector UCharacterMovementComponent::ComputeSlideVector(const FVector& Delta, const float Time, const FVector& Normal, const FHitResult& Hit) const
 {
-	const bool bFalling = IsFalling();
-	FVector Delta = InDelta;
-
-	// Don't make impacts on the upper hemisphere feel so much like a capsule
-	if (bFalling && Delta.Z > 0.f)
-	{
-		if (Hit.Normal.Z < KINDA_SMALL_NUMBER)
-		{
-			float PawnRadius, PawnHalfHeight;
-			CharacterOwner->CapsuleComponent->GetScaledCapsuleSize(PawnRadius, PawnHalfHeight);
-			const float UpperHemisphereZ = UpdatedComponent->GetComponentLocation().Z + PawnHalfHeight - PawnRadius;
-			if (Hit.ImpactPoint.Z > UpperHemisphereZ + KINDA_SMALL_NUMBER && IsWithinEdgeTolerance(Hit.Location, Hit.ImpactPoint, PawnRadius))
-			{
-				Delta = AdjustUpperHemisphereImpact(Delta, Hit);
-			}
-		}
-	}
-
 	FVector Result = Super::ComputeSlideVector(Delta, Time, Normal, Hit);
 
 	// prevent boosting up slopes
-	if ( bFalling && Result.Z > 0.f )
+	if (Result.Z > 0.f && IsFalling())
 	{
 		Result = HandleSlopeBoosting(Result, Delta, Time, Normal, Hit);
 	}
@@ -2016,9 +1999,10 @@ FVector UCharacterMovementComponent::HandleSlopeBoosting(const FVector& SlideRes
 }
 
 
+// TODO: deprecated, remove.
 FVector UCharacterMovementComponent::AdjustUpperHemisphereImpact(const FVector& Delta, const FHitResult& Hit) const
 {
-	const float ZScale = FMath::Clamp(1.f - (FMath::Abs(Hit.Normal.Z) * UpperImpactNormalScale), 0.f, 1.f);
+	const float ZScale = FMath::Clamp(1.f - (FMath::Abs(Hit.Normal.Z) * UpperImpactNormalScale_DEPRECATED), 0.f, 1.f);
 	return FVector(Delta.X, Delta.Y, Delta.Z * ZScale);
 }
 
