New spring component

diff --git a/Engine/Source/Editor/ComponentVisualizers/Private/ComponentVisualizers.cpp b/Engine/Source/Editor/ComponentVisualizers/Private/ComponentVisualizers.cpp
--- a/Engine/Source/Editor/ComponentVisualizers/Private/ComponentVisualizers.cpp
+++ b/Engine/Source/Editor/ComponentVisualizers/Private/ComponentVisualizers.cpp
@@ -5,6 +5,7 @@
 
 #include "SoundDefinitions.h"
 #include "Perception/PawnSensingComponent.h"
+#include "PhysicsEngine/PhysicsSpringComponent.h"
 
 #include "PointLightComponentVisualizer.h"
 #include "SpotLightComponentVisualizer.h"
@@ -14,6 +15,7 @@
 #include "SpringArmComponentVisualizer.h"
 #include "SplineComponentVisualizer.h"
 #include "SensingComponentVisualizer.h"
+#include "SpringComponentVisualizer.h"
 
 IMPLEMENT_MODULE( FComponentVisualizersModule, ComponentVisualizers );
 
@@ -27,6 +29,7 @@ void FComponentVisualizersModule::StartupModule()
 	RegisterComponentVisualizer(USpringArmComponent::StaticClass()->GetFName(), MakeShareable(new FSpringArmComponentVisualizer));
 	RegisterComponentVisualizer(USplineComponent::StaticClass()->GetFName(), MakeShareable(new FSplineComponentVisualizer));
 	RegisterComponentVisualizer(UPawnSensingComponent::StaticClass()->GetFName(), MakeShareable(new FSensingComponentVisualizer));
+	RegisterComponentVisualizer(UPhysicsSpringComponent::StaticClass()->GetFName(), MakeShareable(new FSpringComponentVisualizer));
 }
 
 void FComponentVisualizersModule::ShutdownModule()
diff --git a/Engine/Source/Editor/ComponentVisualizers/Private/SpringComponentVisualizer.cpp b/Engine/Source/Editor/ComponentVisualizers/Private/SpringComponentVisualizer.cpp
new file mode 100644
--- /dev/null
+++ b/Engine/Source/Editor/ComponentVisualizers/Private/SpringComponentVisualizer.cpp
@@ -0,0 +1,30 @@
+// Copyright 1998-2014 Epic Games, Inc. All Rights Reserved.
+
+#include "ComponentVisualizersPrivatePCH.h"
+
+#include "SpringComponentVisualizer.h"
+#include "PhysicsEngine/PhysicsSpringComponent.h"
+
+static const FColor	DisabledColor(128, 128, 128);
+static const FColor	CompressedColor(255, 0, 0);
+static const FColor RestColor(0, 255, 0);
+
+void FSpringComponentVisualizer::DrawVisualization( const UActorComponent* Component, const FSceneView* View, FPrimitiveDrawInterface* PDI )
+{
+	if (const UPhysicsSpringComponent* SpringComp = Cast<const UPhysicsSpringComponent>(Component))
+	{
+		//interpolate color by compression
+		const float CompressionNormalized = SpringComp->GetNormalizedCompressionScalar();
+		float R = FMath::Lerp(RestColor.R, CompressedColor.R, CompressionNormalized);
+		float G = FMath::Lerp(RestColor.G, CompressedColor.G, CompressionNormalized);
+		float B = FMath::Lerp(RestColor.B, CompressedColor.B, CompressionNormalized);
+		const FColor CurrentColor = SpringComp->bIsActive ? FColor(R, G, B) : DisabledColor;
+		
+		//draw capsule for spring sweep
+		const FTransform& WorldTM = SpringComp->GetComponentToWorld();
+		const FVector SpringStart = WorldTM.GetLocation();
+		const FVector SpringEnd = SpringComp->GetSpringCurrentEndPoint();
+		const float HalfHeight = (SpringEnd - SpringStart).Size() * 0.5f + SpringComp->SpringRadius;
+		DrawWireCapsule(PDI, FMath::Lerp(SpringStart, SpringEnd, 0.5f), WorldTM.GetUnitAxis(EAxis::Z), WorldTM.GetUnitAxis(EAxis::Y), WorldTM.GetUnitAxis(EAxis::X), CurrentColor, SpringComp->SpringRadius, HalfHeight, 25, SDPG_World);
+	}
+}
diff --git a/Engine/Source/Editor/ComponentVisualizers/Private/SpringComponentVisualizer.h b/Engine/Source/Editor/ComponentVisualizers/Private/SpringComponentVisualizer.h
new file mode 100644
--- /dev/null
+++ b/Engine/Source/Editor/ComponentVisualizers/Private/SpringComponentVisualizer.h
@@ -0,0 +1,13 @@
+// Copyright 1998-2014 Epic Games, Inc. All Rights Reserved.
+
+#pragma once
+
+#include "ComponentVisualizer.h"
+
+class FSpringComponentVisualizer : public FComponentVisualizer
+{
+public:
+	// Begin FComponentVisualizer interface
+	virtual void DrawVisualization(const UActorComponent* Component, const FSceneView* View, FPrimitiveDrawInterface* PDI) override;
+	// End FComponentVisualizer interface
+};
diff --git a/Engine/Source/Runtime/Engine/Classes/PhysicsEngine/PhysicsSpringComponent.h b/Engine/Source/Runtime/Engine/Classes/PhysicsEngine/PhysicsSpringComponent.h
new file mode 100644
--- /dev/null
+++ b/Engine/Source/Runtime/Engine/Classes/PhysicsEngine/PhysicsSpringComponent.h
@@ -0,0 +1,84 @@
+// Copyright 1998-2014 Epic Games, Inc. All Rights Reserved.
+
+
+#pragma once
+#include "PhysicsSpringComponent.generated.h"
+
+/** 
+ *	Used with objects that have physics to create a spring down the X direction
+ *	ie. point X in the direction you want generate spring.
+ */
+UCLASS(hidecategories=(Object, Mobility, LOD), ClassGroup=Physics, showcategories=Trigger, meta=(BlueprintSpawnableComponent))
+class ENGINE_API UPhysicsSpringComponent : public USceneComponent
+{
+	GENERATED_UCLASS_BODY()
+
+	/** Specifies how much strength the spring has. The higher the SpringStiffness the more force the spring can push on a body with. */
+	UPROPERTY(BlueprintReadWrite, EditAnywhere, Category = Physics)
+	float SpringStiffness;
+
+	/** Specifies how quickly the spring can absorb energy of a body. The higher the damping the less oscillation */
+	UPROPERTY(BlueprintReadWrite, EditAnywhere, Category = Physics)
+	float SpringDamping;
+
+	/** Determines how long the spring will be along the X-axis at rest. The spring will apply 0 force on a body when it's at rest. */
+	UPROPERTY(BlueprintReadWrite, EditAnywhere, Category = Physics, meta = (UIMin = 0.01, ClampMin = 0.01))
+	float SpringLengthAtRest;
+
+	/** Determines the radius of the spring. */
+	UPROPERTY(BlueprintReadWrite, EditAnywhere, Category = Physics, meta = (UIMin = 0.01, ClampMin = 0.01))
+	float SpringRadius;
+
+	/** Strength of thrust force applied to the base object. */
+	UPROPERTY(BlueprintReadWrite, Category = Physics)
+	TEnumAsByte<enum ECollisionChannel> SpringChannel;
+
+	/** If true, the spring will ignore all components in its own actor */
+	UPROPERTY(BlueprintReadWrite, EditAnywhere, Category = Physics, meta = (UIMin = 0.01, ClampMin = 0.01))
+	bool bIgnoreSelf;
+
+	/** The current compression of the spring. A spring at rest will have SpringCompression 0. */
+	UPROPERTY(BlueprintReadOnly, Category = Physics, transient)
+	float SpringCompression;
+
+	// Begin UActorComponent Interface
+	virtual void TickComponent(float DeltaTime, enum ELevelTick TickType, FActorComponentTickFunction *ThisTickFunction) override;
+	// End UActorComponent Interface
+
+	/** Returns the spring compression as a normalized scalar along spring direction.
+	 *  0 implies spring is at rest
+	 *  1 implies fully compressed */
+	UFUNCTION(BlueprintCallable, Category = Physics)
+	float GetNormalizedCompressionScalar() const;
+
+	/** Returns the spring resting point in world space.*/
+	UFUNCTION(BlueprintCallable, Category = Physics)
+	FVector GetSpringRestingPoint() const;
+
+	/** Returns the spring current end point in world space.*/
+	UFUNCTION(BlueprintCallable, Category = Physics)
+	FVector GetSpringCurrentEndPoint() const;
+
+	/** Returns the spring direction from start to resting point */
+	UFUNCTION(BlueprintCallable, Category = Physics)
+	FVector GetSpringDirection() const;
+
+private:
+
+	/** Sweeps along spring direction to see if spring needs to compress. Returned Collision time is independent of spring radius */
+	UPrimitiveComponent* GetSpringCollision(const FVector& Start, const FVector& End, float& CollisionTime) const;
+	
+	/** Computes new spring compression and force */
+	FVector ComputeNewSpringCompressionAndForce(const FVector& End, const float DeltaTime, float& NewSpringCompression) const;
+
+	/** We want to automatically set relative position of attached children */
+	void UpdateAttachedPosition() const;
+
+	/** Given a length, returns the point along the spring that is Length units away from the spring start */
+	FVector SpringPositionFromLength(float Length) const;
+
+	FVector CurrentEndPoint;
+};
+
+
+
diff --git a/Engine/Source/Runtime/Engine/Private/PhysicsEngine/PhysicsSpring.cpp b/Engine/Source/Runtime/Engine/Private/PhysicsEngine/PhysicsSpring.cpp
new file mode 100644
--- /dev/null
+++ b/Engine/Source/Runtime/Engine/Private/PhysicsEngine/PhysicsSpring.cpp
@@ -0,0 +1,123 @@
+// Copyright 1998-2014 Epic Games, Inc. All Rights Reserved.
+
+#include "EnginePrivate.h"
+#include "PhysicsEngine/PhysicsSpringComponent.h"
+#include "PhysXSupport.h"
+
+
+UPhysicsSpringComponent::UPhysicsSpringComponent(const class FPostConstructInitializeProperties& PCIP)
+	: Super(PCIP)
+{
+	PrimaryComponentTick.bCanEverTick = true;
+	PrimaryComponentTick.TickGroup = TG_PreCloth;
+	bAutoActivate = true;
+
+	SpringStiffness = 25.f;
+	SpringDamping = 3.f;
+	SpringLengthAtRest = 100.f;
+	SpringRadius = 10.f;
+	SpringChannel = ECC_WorldDynamic;
+	bIgnoreSelf = true;
+
+	SpringCompression = 0;
+}
+FVector UPhysicsSpringComponent::SpringPositionFromLength(float Length) const
+{
+	return ComponentToWorld.GetLocation() + Length * GetSpringDirection();
+}
+
+FVector UPhysicsSpringComponent::GetSpringDirection() const
+{
+	return ComponentToWorld.TransformVectorNoScale(FVector(1.f, 0.f, 0.f));
+}
+
+float UPhysicsSpringComponent::GetNormalizedCompressionScalar() const
+{
+	const float NormalizedComprssion = SpringCompression / SpringLengthAtRest;
+	return FMath::Clamp(NormalizedComprssion, 0.f, 1.f);
+}
+
+FVector UPhysicsSpringComponent::GetSpringCurrentEndPoint() const
+{
+	return SpringCompression == 0 ? GetSpringRestingPoint() : CurrentEndPoint;
+}
+
+FVector UPhysicsSpringComponent::GetSpringRestingPoint() const
+{
+	return SpringPositionFromLength(SpringLengthAtRest);
+}
+
+
+UPrimitiveComponent* UPhysicsSpringComponent::GetSpringCollision(const FVector& Start, const FVector& End, float& Time) const
+{
+	UWorld* World = GetWorld();
+	AActor* IgnoreActor = bIgnoreSelf ? GetOwner() : nullptr;
+
+	static FName NAME_Spring = FName(TEXT("SpringComponent"));
+	FCollisionQueryParams QueryParams(NAME_Spring, true, IgnoreActor);
+	FHitResult Hit;
+
+	UPrimitiveComponent* CollidedComponent = nullptr;
+	
+	const FVector Delta = End - Start;
+	const float DeltaSize = Delta.Size();
+	if (DeltaSize > SMALL_NUMBER)
+	{
+		if (bool bBlockingHit = World->SweepSingle(Hit, Start, End, FQuat::Identity, SpringChannel, FCollisionShape::MakeSphere(SpringRadius), QueryParams))
+		{
+			CollidedComponent = Hit.GetComponent();
+			Time = CollidedComponent ? Hit.Time : 1.f;
+		}
+	}
+
+	return CollidedComponent;
+}
+
+FVector UPhysicsSpringComponent::ComputeNewSpringCompressionAndForce(const FVector& End, const float DeltaTime, float& NewSpringCompression) const
+{
+	const FVector SpringDir = GetSpringDirection();
+	const FVector RestPoint = SpringPositionFromLength(SpringLengthAtRest);
+	NewSpringCompression = (End - RestPoint).Size();
+
+	const float SpringCompressionDelta = (NewSpringCompression - SpringCompression) / DeltaTime;
+	return SpringDir * ((-SpringStiffness*NewSpringCompression) + (-SpringDamping*SpringCompressionDelta));
+}
+
+void UPhysicsSpringComponent::UpdateAttachedPosition() const
+{
+	const FVector EndPosition = GetSpringCurrentEndPoint();
+	for (USceneComponent* ChildComponent : AttachChildren)
+	{
+		ChildComponent->SetWorldLocation(EndPosition);
+	}
+}
+
+void UPhysicsSpringComponent::TickComponent(float DeltaTime, enum ELevelTick TickType, FActorComponentTickFunction *ThisTickFunction)
+{
+	Super::TickComponent(DeltaTime, TickType, ThisTickFunction);
+
+	if (UPrimitiveComponent* BasePrimComp = Cast<UPrimitiveComponent>(AttachParent))
+	{
+		if (bIsActive && AttachParent)
+		{
+			const FVector SpringStart = ComponentToWorld.GetLocation();
+			const FVector SpringDesiredEnd = SpringPositionFromLength(SpringLengthAtRest);
+			float CollisionTime = 1.f;
+			
+			UPrimitiveComponent* CollideWith = GetSpringCollision(SpringStart, SpringDesiredEnd, CollisionTime);
+			CurrentEndPoint = FMath::Lerp(SpringStart, SpringDesiredEnd, CollisionTime);
+			
+			if(CollideWith)
+			{
+				const float Mass = BasePrimComp->GetMass();
+				float NewSpringCompression;
+				const FVector WorldForce = ComputeNewSpringCompressionAndForce(CurrentEndPoint, DeltaTime, NewSpringCompression);
+				
+				BasePrimComp->AddForceAtLocation(WorldForce*Mass, SpringStart, AttachSocketName);
+				SpringCompression = NewSpringCompression;
+			}
+
+			UpdateAttachedPosition();
+		}
+	}
+}
\ No newline at end of file
