Fixed TTP #349024 : CRASH: Closing editor with Select Skeleton window open

diff --git a/Engine/Source/Runtime/Engine/Private/SkeletalMesh.cpp b/Engine/Source/Runtime/Engine/Private/SkeletalMesh.cpp
--- a/Engine/Source/Runtime/Engine/Private/SkeletalMesh.cpp
+++ b/Engine/Source/Runtime/Engine/Private/SkeletalMesh.cpp
@@ -1950,6 +1950,12 @@ void USkeletalMesh::ReleaseResources()
 	ImportedResource->ReleaseResources();
 	// insert a fence to signal when these commands completed
 	ReleaseResourcesFence.BeginFence();
+
+	// release clothing assets
+	for (FClothingAssetData& Data : ClothingAssets)
+	{
+		Data.ApexClothingAsset.Reset();
+	}
 }
 
 template <bool bExtraBoneInfluencesT>
