missing #if WITH_APEX_CLOTHING

diff --git a/Engine/Source/Runtime/Engine/Private/SkeletalMesh.cpp b/Engine/Source/Runtime/Engine/Private/SkeletalMesh.cpp
--- a/Engine/Source/Runtime/Engine/Private/SkeletalMesh.cpp
+++ b/Engine/Source/Runtime/Engine/Private/SkeletalMesh.cpp
@@ -1951,11 +1951,13 @@ void USkeletalMesh::ReleaseResources()
 	// insert a fence to signal when these commands completed
 	ReleaseResourcesFence.BeginFence();
 
+#if WITH_APEX_CLOTHING
 	// release clothing assets
 	for (FClothingAssetData& Data : ClothingAssets)
 	{
 		Data.ApexClothingAsset.Reset();
 	}
+#endif // #if WITH_APEX_CLOTHING
 }
 
 template <bool bExtraBoneInfluencesT>
