Hide 'PhysicsVolume' category from BillboardComponent
Update MaterialBillboardComponent to match
Move CanCharacterStepUpOn to Collision category (silly to have its own category, and will also automatically hide it from things like billboards)

diff --git a/Engine/Source/Runtime/Engine/Classes/Components/BillboardComponent.h b/Engine/Source/Runtime/Engine/Classes/Components/BillboardComponent.h
--- a/Engine/Source/Runtime/Engine/Classes/Components/BillboardComponent.h
+++ b/Engine/Source/Runtime/Engine/Classes/Components/BillboardComponent.h
@@ -7,7 +7,7 @@
 /** 
  * A 2d texture that will be rendered always facing the camera.
  */
-UCLASS(ClassGroup=Rendering, collapsecategories, hidecategories=(Object,Activation,"Components|Activation",Physics,Collision,Lighting,LOD,Mesh), editinlinenew, meta=(BlueprintSpawnableComponent),MinimalAPI)
+UCLASS(ClassGroup=Rendering, collapsecategories, hidecategories=(Object,Activation,"Components|Activation",Physics,Collision,Lighting,LOD,Mesh,PhysicsVolume), editinlinenew, meta=(BlueprintSpawnableComponent),MinimalAPI)
 class UBillboardComponent : public UPrimitiveComponent
 {
 	GENERATED_UCLASS_BODY()
diff --git a/Engine/Source/Runtime/Engine/Classes/Components/MaterialBillboardComponent.h b/Engine/Source/Runtime/Engine/Classes/Components/MaterialBillboardComponent.h
--- a/Engine/Source/Runtime/Engine/Classes/Components/MaterialBillboardComponent.h
+++ b/Engine/Source/Runtime/Engine/Classes/Components/MaterialBillboardComponent.h
@@ -50,7 +50,7 @@ struct FMaterialSpriteElement
 /** 
  * A 2d material that will be rendered always facing the camera.
  */
-UCLASS(ClassGroup=Rendering, collapsecategories, hidecategories=Object, editinlinenew, meta=(BlueprintSpawnableComponent))
+UCLASS(ClassGroup=Rendering, collapsecategories, hidecategories=(Object,Activation,"Components|Activation",Physics,Collision,Lighting,LOD,Mesh,PhysicsVolume), editinlinenew, meta=(BlueprintSpawnableComponent))
 class ENGINE_API UMaterialBillboardComponent : public UPrimitiveComponent
 {
 	GENERATED_UCLASS_BODY()
diff --git a/Engine/Source/Runtime/Engine/Classes/Components/PrimitiveComponent.h b/Engine/Source/Runtime/Engine/Classes/Components/PrimitiveComponent.h
--- a/Engine/Source/Runtime/Engine/Classes/Components/PrimitiveComponent.h
+++ b/Engine/Source/Runtime/Engine/Classes/Components/PrimitiveComponent.h
@@ -444,7 +444,7 @@ public:
 	/**
 	 * Determine whether a Character can step up onto this component.
 	 */
-	UPROPERTY(EditDefaultsOnly, BlueprintReadOnly, Category=Base)
+	UPROPERTY(EditDefaultsOnly, BlueprintReadOnly, Category=Collision)
 	TEnumAsByte<enum ECanBeCharacterBase> CanCharacterStepUpOn;
 
 	/**
