Added code to allow material tiles to pass color to vertex color.
(Resolves TTP 348979)

diff --git a/Engine/Source/Runtime/Engine/Private/UserInterface/Canvas.cpp b/Engine/Source/Runtime/Engine/Private/UserInterface/Canvas.cpp
--- a/Engine/Source/Runtime/Engine/Private/UserInterface/Canvas.cpp
+++ b/Engine/Source/Runtime/Engine/Private/UserInterface/Canvas.cpp
@@ -318,16 +318,17 @@ bool FCanvasTileRendererItem::Render_RenderThread(FRHICommandListImmediate& RHIC
 	{
 		const FRenderData::FTileInst& Tile = Data->Tiles[TileIdx];
 		FTileRenderer::DrawTile(
-			RHICmdList,
+			RHICmdList, 
 			*View, 
 			Data->MaterialRenderProxy, 
 			bNeedsToSwitchVerticalAxis,
 			Tile.X, Tile.Y, Tile.SizeX, Tile.SizeY, 
 			Tile.U, Tile.V, Tile.SizeU, Tile.SizeV,
-			Canvas->IsHitTesting(), Tile.HitProxyId
+			Canvas->IsHitTesting(), Tile.HitProxyId,
+			Tile.InColor
 			);
 	}
-
+	
 	delete View->Family;
 	delete View;
 	if( Canvas->GetAllowedModes() & FCanvas::Allow_DeleteOnRender )
@@ -409,8 +410,8 @@ bool FCanvasTileRendererItem::Render_GameThread(const FCanvas* Canvas)
 				Parameters.bNeedsToSwitchVerticalAxis,
 				Tile.X, Tile.Y, Tile.SizeX, Tile.SizeY,
 				Tile.U, Tile.V, Tile.SizeU, Tile.SizeV,
-				Parameters.bIsHitTesting, Tile.HitProxyId
-				);
+				Parameters.bIsHitTesting, Tile.HitProxyId,
+				Tile.InColor);
 		}
 
 		delete Parameters.View->Family;
@@ -491,7 +492,7 @@ FBatchedElements* FCanvas::GetBatchedElements(EElementType InElementType, FBatch
 }
 
 
-void FCanvas::AddTileRenderItem(float X,float Y,float SizeX,float SizeY,float U,float V,float SizeU,float SizeV,const FMaterialRenderProxy* MaterialRenderProxy,FHitProxyId HitProxyId,bool bFreezeTime)
+void FCanvas::AddTileRenderItem(float X, float Y, float SizeX, float SizeY, float U, float V, float SizeU, float SizeV, const FMaterialRenderProxy* MaterialRenderProxy, FHitProxyId HitProxyId, bool bFreezeTime, FColor InColor)
 {
 	SCOPE_CYCLE_COUNTER(STAT_Canvas_AddTileRenderTime);
 
@@ -518,7 +519,7 @@ void FCanvas::AddTileRenderItem(float X,float Y,float SizeX,float SizeY,float U,
 		SortElement.RenderBatchArray.Add(RenderBatch);
 	}
 	// add the quad to the tile render batch
-	RenderBatch->AddTile( X,Y,SizeX,SizeY,U,V,SizeU,SizeV,HitProxyId);
+	RenderBatch->AddTile( X,Y,SizeX,SizeY,U,V,SizeU,SizeV,HitProxyId,InColor);
 }
 
 FCanvas::~FCanvas()
diff --git a/Engine/Source/Runtime/Engine/Private/UserInterface/CanvasItem.cpp b/Engine/Source/Runtime/Engine/Private/UserInterface/CanvasItem.cpp
--- a/Engine/Source/Runtime/Engine/Private/UserInterface/CanvasItem.cpp
+++ b/Engine/Source/Runtime/Engine/Private/UserInterface/CanvasItem.cpp
@@ -452,7 +452,7 @@ void FCanvasTileItem::RenderMaterialTile( class FCanvas* InCanvas, const FVector
 	}
 	FHitProxyId HitProxyId = InCanvas->GetHitProxyId();
 	// add the quad to the tile render batch
-	RenderBatch->AddTile( InPosition.X, InPosition.Y ,Size.X, Size.Y, UV0.X, UV0.Y, UV1.X-UV0.X, UV1.Y-UV0.Y, HitProxyId );
+	RenderBatch->AddTile( InPosition.X, InPosition.Y ,Size.X, Size.Y, UV0.X, UV0.Y, UV1.X-UV0.X, UV1.Y-UV0.Y, HitProxyId, Color);
 }
 
 
diff --git a/Engine/Source/Runtime/Engine/Public/CanvasTypes.h b/Engine/Source/Runtime/Engine/Public/CanvasTypes.h
--- a/Engine/Source/Runtime/Engine/Public/CanvasTypes.h
+++ b/Engine/Source/Runtime/Engine/Public/CanvasTypes.h
@@ -77,7 +77,7 @@ public:
 	/**
 	* Generates a new FCanvasTileRendererItem for the current sortkey and adds it to the sortelement list of itmes to render
 	*/
-	ENGINE_API void AddTileRenderItem(float X,float Y,float SizeX,float SizeY,float U,float V,float SizeU,float SizeV,const FMaterialRenderProxy* MaterialRenderProxy,FHitProxyId HitProxyId,bool bFreezeTime);
+	ENGINE_API void AddTileRenderItem(float X,float Y,float SizeX,float SizeY,float U,float V,float SizeU,float SizeV,const FMaterialRenderProxy* MaterialRenderProxy,FHitProxyId HitProxyId,bool bFreezeTime,FColor InColor);
 
 	/** 
 	* Sends a message to the rendering thread to draw the batched elements. 
@@ -791,11 +791,11 @@ public:
 	* @param SizeV - tile V size
 	* @param return number of tiles added
 	*/
-	FORCEINLINE int32 AddTile(float X,float Y,float SizeX,float SizeY,float U,float V,float SizeU,float SizeV,FHitProxyId HitProxyId)
+	FORCEINLINE int32 AddTile(float X,float Y,float SizeX,float SizeY,float U,float V,float SizeU,float SizeV,FHitProxyId HitProxyId,FColor InColor)
 	{
-		return Data->AddTile(X,Y,SizeX,SizeY,U,V,SizeU,SizeV,HitProxyId);
+		return Data->AddTile(X,Y,SizeX,SizeY,U,V,SizeU,SizeV,HitProxyId,InColor);
 	};
-
+	
 private:
 	class FRenderData
 	{
@@ -816,12 +816,13 @@ private:
 			float U,V;
 			float SizeU,SizeV;
 			FHitProxyId HitProxyId;
+			FColor InColor;
 		};
 		TArray<FTileInst> Tiles;
 
-		FORCEINLINE int32 AddTile(float X,float Y,float SizeX,float SizeY,float U,float V,float SizeU,float SizeV,FHitProxyId HitProxyId)
+		FORCEINLINE int32 AddTile(float X, float Y, float SizeX, float SizeY, float U, float V, float SizeU, float SizeV, FHitProxyId HitProxyId, FColor InColor)
 		{
-			FTileInst NewTile = {X,Y,SizeX,SizeY,U,V,SizeU,SizeV,HitProxyId};
+			FTileInst NewTile = {X,Y,SizeX,SizeY,U,V,SizeU,SizeV,HitProxyId,InColor};
 			return Tiles.Add(NewTile);
 		};
 	};
diff --git a/Engine/Source/Runtime/Engine/Public/TileRendering.h b/Engine/Source/Runtime/Engine/Public/TileRendering.h
--- a/Engine/Source/Runtime/Engine/Public/TileRendering.h
+++ b/Engine/Source/Runtime/Engine/Public/TileRendering.h
@@ -15,7 +15,7 @@ public:
 	 * Draw a tile at the given location and size, using the given UVs
 	 * (UV = [0..1]
 	 */
-	ENGINE_API static void DrawTile(FRHICommandListImmediate& RHICmdList, const class FSceneView& View, const FMaterialRenderProxy* MaterialRenderProxy, bool bNeedsToSwitchVerticalAxis, float X, float Y, float SizeX, float SizeY, float U, float V, float SizeU, float SizeV, bool bIsHitTesting = false, const FHitProxyId HitProxyId = FHitProxyId());
+	ENGINE_API static void DrawTile(FRHICommandListImmediate& RHICmdList, const class FSceneView& View, const FMaterialRenderProxy* MaterialRenderProxy, bool bNeedsToSwitchVerticalAxis, float X, float Y, float SizeX, float SizeY, float U, float V, float SizeU, float SizeV, bool bIsHitTesting = false, const FHitProxyId HitProxyId = FHitProxyId(), const FColor InVertexColor = FColor(255, 255, 255, 255));
 
 private:
 
