Missing file for passing material tile item to vertex colors

diff --git a/Engine/Source/Runtime/Engine/Private/TileRendering.cpp b/Engine/Source/Runtime/Engine/Private/TileRendering.cpp
--- a/Engine/Source/Runtime/Engine/Private/TileRendering.cpp
+++ b/Engine/Source/Runtime/Engine/Private/TileRendering.cpp
@@ -134,7 +134,7 @@ public:
 };
 TGlobalResource<FTileMesh> GTileMesh;
 
-void FTileRenderer::DrawTile(FRHICommandListImmediate& RHICmdList, const class FSceneView& View, const FMaterialRenderProxy* MaterialRenderProxy, bool bNeedsToSwitchVerticalAxis, float X, float Y, float SizeX, float SizeY, float U, float V, float SizeU, float SizeV, bool bIsHitTesting, const FHitProxyId HitProxyId)
+void FTileRenderer::DrawTile(FRHICommandListImmediate& RHICmdList, const class FSceneView& View, const FMaterialRenderProxy* MaterialRenderProxy, bool bNeedsToSwitchVerticalAxis, float X, float Y, float SizeX, float SizeY, float U, float V, float SizeU, float SizeV, bool bIsHitTesting, const FHitProxyId HitProxyId, const FColor InVertexColor)
 {
 	FMaterialTileVertex DestVertex[4];
 
@@ -144,7 +144,7 @@ void FTileRenderer::DrawTile(FRHICommandListImmediate& RHICmdList, const class F
 		DestVertex[0].Initialize(X + SizeX, View.ViewRect.Height() - (Y + SizeY), U + SizeU, V + SizeV);
 		DestVertex[1].Initialize(X, View.ViewRect.Height() - (Y + SizeY), U, V + SizeV);
 		DestVertex[2].Initialize(X + SizeX, View.ViewRect.Height() - Y, U + SizeU, V);
-		DestVertex[3].Initialize(X, View.ViewRect.Height() - Y, U, V);
+		DestVertex[3].Initialize(X, View.ViewRect.Height() - Y, U, V);		
 	}
 	else
 	{
@@ -154,6 +154,11 @@ void FTileRenderer::DrawTile(FRHICommandListImmediate& RHICmdList, const class F
 		DestVertex[3].Initialize(X, Y + SizeY, U, V + SizeV);
 	}
 
+	DestVertex[0].Color = InVertexColor.DWColor();
+	DestVertex[1].Color = InVertexColor.DWColor();
+	DestVertex[2].Color = InVertexColor.DWColor();
+	DestVertex[3].Color = InVertexColor.DWColor();
+
 	// update the FMeshBatch
 	FMeshBatch& Mesh = GTileMesh.MeshElement;
 	Mesh.UseDynamicData = true;
