Integrate particle optimization from PullRequest 284.

diff --git a/Engine/Source/Runtime/Engine/Private/Particles/ParticleGpuSimulation.cpp b/Engine/Source/Runtime/Engine/Private/Particles/ParticleGpuSimulation.cpp
--- a/Engine/Source/Runtime/Engine/Private/Particles/ParticleGpuSimulation.cpp
+++ b/Engine/Source/Runtime/Engine/Private/Particles/ParticleGpuSimulation.cpp
@@ -1848,21 +1848,14 @@ static void BuildParticleVertexBuffer( FVertexBufferRHIParamRef VertexBufferRHI,
 			{
 				const float IndexX = TileOffset.X + ((float)ParticleX / (float)GParticleSimulationTextureSizeX) + (0.5f / (float)GParticleSimulationTextureSizeX);
 				const float IndexY = TileOffset.Y + ((float)ParticleY / (float)GParticleSimulationTextureSizeY) + (0.5f / (float)GParticleSimulationTextureSizeY);
-
-				// on some platforms, union and/or bitfield writes to Locked memory are really slow, so use a forced int write instead
-				// and in fact one 32-bit write is faster than two uint16 writes (i.e. using .Encoded)
-				FParticleIndex Temp;
-				// We use the unsafe version of FP32 -> FP16 conversion because we know all values are in [0,1].
-				Temp.X.SetWithoutBoundsChecks(IndexX);
-				Temp.Y.SetWithoutBoundsChecks(IndexY);
-				*(uint32*)ParticleIndices = *(uint32*)&Temp;
+				ParticleIndices->X.SetWithoutBoundsChecks(IndexX);
+				ParticleIndices->Y.SetWithoutBoundsChecks(IndexY);					
 
 				// move to next particle
 				ParticleIndices += Stride;
 			}
 		}
 	}
-
 	RHIUnlockVertexBuffer( VertexBufferRHI );
 }
 
