#docs - Improve comments/tooltips for BreakHitResult in the editor. Add 'Time' value to the output.

diff --git a/Engine/Source/Runtime/Engine/Classes/Kismet/GameplayStatics.h b/Engine/Source/Runtime/Engine/Classes/Kismet/GameplayStatics.h
--- a/Engine/Source/Runtime/Engine/Classes/Kismet/GameplayStatics.h
+++ b/Engine/Source/Runtime/Engine/Classes/Kismet/GameplayStatics.h
@@ -387,20 +387,21 @@ class ENGINE_API UGameplayStatics : public UBlueprintFunctionLibrary
 	UFUNCTION(BlueprintCallable, Category="Rendering|Components|Decal", meta=(UnsafeDuringActorConstruction = "true"))
 	static UDecalComponent* SpawnDecalAttached(class UMaterialInterface* DecalMaterial, FVector DecalSize, class USceneComponent* AttachToComponent, FName AttachPointName = NAME_None, FVector Location = FVector(ForceInit), FRotator Rotation = FRotator::ZeroRotator, EAttachLocation::Type LocationType = EAttachLocation::KeepRelativeOffset, float LifeSpan = 0);
 
-	/** Extracts data from a FHitResult. 
-	 * @param Hit			The source FHitResult
-	 * @param Location		Location of the hit in world space.  If this was a Swept Shape test, this is the location where we can place the shape in the world where it will not penetrate.
-	 * @param Normal		Normal of the hit in world space - for capsule, it will be capsule normal(not contact geometry normal). Otherwise same as ImpactNormal.
+	/** Extracts data from a HitResult. 
+	 * @param Hit			The source HitResult.
+	 * @param Location		Location of the hit in world space. If this was a swept shape test, this is the location where we can place the shape in the world where it will not penetrate.
+	 * @param Normal		Normal of the hit in world space, for the object that was swept (e.g. for a sphere trace this points towards the sphere's center). Equal to ImpactNormal for line tests.
 	 * @param ImpactPoint	Location of the actual contact point of the trace shape with the surface of the hit object.
-	 * @param ImpactNormal	Normal of the hit from the actual contact of the trace shape in world space
-	 * @param PhysMat		
-	 * @param HitActor
-	 * @param HitComponent
-	 * @param HitBoneName
+	 * @param ImpactNormal	Normal of the hit in world space, for the object that was hit by the sweep.
+	 * @param Time			'Time' of impact along trace direction ranging from [0.0 to 1.0) if there is a hit, indicating time between start and end. Equals 1.0 if there is no hit.
+	 * @param PhysMat		Physical material that was hit. Must set bReturnPhysicalMaterial to true in the query params for this to be returned.
+	 * @param HitActor		Actor hit by the trace.
+	 * @param HitComponent	PrimitiveComponent hit by the trace.
+	 * @param HitBoneName	Name of the bone hit (valid only if we hit a skeletal mesh).
 	 * @param HitItem		Primitive-specific data recording which item in the primitive was hit
 	 */
 	UFUNCTION(BlueprintPure, Category = "Collision", meta=(NativeBreakFunc))
-	static void BreakHitResult(const struct FHitResult& Hit, FVector& Location, FVector& Normal, FVector& ImpactPoint, FVector& ImpactNormal, class UPhysicalMaterial*& PhysMat, class AActor*& HitActor, class UPrimitiveComponent*& HitComponent, FName& HitBoneName, int32& HitItem);
+	static void BreakHitResult(const struct FHitResult& Hit, FVector& Location, FVector& Normal, FVector& ImpactPoint, FVector& ImpactNormal, float& Time, class UPhysicalMaterial*& PhysMat, class AActor*& HitActor, class UPrimitiveComponent*& HitComponent, FName& HitBoneName, int32& HitItem);
 
 	/** Returns the EPhysicalSurface type of the given Hit. 
 	 * To edit surface type for your project, use ProjectSettings/Physics/PhysicalSurface section
diff --git a/Engine/Source/Runtime/Engine/Private/GameplayStatics.cpp b/Engine/Source/Runtime/Engine/Private/GameplayStatics.cpp
--- a/Engine/Source/Runtime/Engine/Private/GameplayStatics.cpp
+++ b/Engine/Source/Runtime/Engine/Private/GameplayStatics.cpp
@@ -555,12 +555,13 @@ UParticleSystemComponent* UGameplayStatics::SpawnEmitterAttached(UParticleSystem
 	return PSC;
 }
 
-void UGameplayStatics::BreakHitResult(const FHitResult& Hit, FVector& Location, FVector& Normal, FVector& ImpactPoint, FVector& ImpactNormal, UPhysicalMaterial*& PhysMat, AActor*& HitActor, UPrimitiveComponent*& HitComponent, FName& HitBoneName, int32& HitItem)
+void UGameplayStatics::BreakHitResult(const FHitResult& Hit, FVector& Location, FVector& Normal, FVector& ImpactPoint, FVector& ImpactNormal, float& Time, UPhysicalMaterial*& PhysMat, AActor*& HitActor, UPrimitiveComponent*& HitComponent, FName& HitBoneName, int32& HitItem)
 {
 	Location = Hit.Location;
 	Normal = Hit.Normal;
 	ImpactPoint = Hit.ImpactPoint;
 	ImpactNormal = Hit.ImpactNormal;
+	Time = Hit.Time;
 	PhysMat = Hit.PhysMaterial.Get();
 	HitActor = Hit.GetActor();
 	HitComponent = Hit.Component.Get();
