Stubbed out location for UI notification of delined ability confirmation.

diff --git a/Engine/Source/Runtime/Engine/Classes/GameFramework/PlayerController.h b/Engine/Source/Runtime/Engine/Classes/GameFramework/PlayerController.h
--- a/Engine/Source/Runtime/Engine/Classes/GameFramework/PlayerController.h
+++ b/Engine/Source/Runtime/Engine/Classes/GameFramework/PlayerController.h
@@ -338,11 +338,11 @@ public:
 	virtual void LocalTravel(const FString& URL);
 
 	/** Return the client to the main menu gracefully */
-	UFUNCTION(reliable, client)
+	UFUNCTION(Reliable, Client)
 	virtual void ClientReturnToMainMenu(const FString& ReturnReason);
 
 	/** Development RPC for testing object reference replication */
-	UFUNCTION(reliable, client)
+	UFUNCTION(Reliable, Client)
 	virtual void ClientRepObjRef(UObject* Object);
 
 	/**
@@ -507,7 +507,7 @@ public:
 	 * Tell the client to enable or disable voice chat (not muting)
 	 * @param bEnable enable or disable voice chat
 	 */
-	UFUNCTION(client, reliable)
+	UFUNCTION(Reliable, Client)
 	virtual void ClientEnableNetworkVoice(bool bEnable);
 
 	/** Enable voice chat transmission	 */
@@ -533,7 +533,7 @@ public:
 	 *
 	 * NOTE: This is done as an RPC instead of variable replication because ordering matters
 	 */
-	UFUNCTION(client, reliable)
+	UFUNCTION(Reliable, Client)
 	virtual void ClientVoiceHandshakeComplete();
 
 	/**
@@ -554,14 +554,14 @@ public:
 	 * Tell the client to mute a player for this controller
 	 * @param PlayerId player id to mute
 	 */
-	UFUNCTION(client, reliable)
+	UFUNCTION(Reliable, Client)
 	virtual void ClientMutePlayer(FUniqueNetIdRepl PlayerId);
 
 	/**
 	 * Tell the client to unmute a player for this controller
 	 * @param PlayerId player id to unmute
 	 */
-	UFUNCTION(client, reliable)
+	UFUNCTION(Reliable, Client)
 	virtual void ClientUnmutePlayer(FUniqueNetIdRepl PlayerId);
 
 	/**
@@ -603,15 +603,15 @@ public:
 	void ClientAddTextureStreamingLoc(FVector InLoc, float Duration, bool bOverrideLocation);
 
 	/** Tells client to cancel any pending map change. */
-	UFUNCTION(reliable, client)
+	UFUNCTION(Reliable, Client)
 	void ClientCancelPendingMapChange();
 
 	/** Set CurrentNetSpeed to the lower of its current value and Cap. */
-	UFUNCTION(reliable, client)
+	UFUNCTION(Reliable, Client)
 	void ClientCapBandwidth(int32 Cap);
 
 	/** Actually performs the level transition prepared by PrepareMapChange(). */
-	UFUNCTION(reliable, client)
+	UFUNCTION(Reliable, Client)	
 	void ClientCommitMapChange();
 
 	/** Tells the client to block until all pending level streaming actions are complete
@@ -622,7 +622,7 @@ public:
 	void ClientFlushLevelStreaming();
 
 	/** Forces GC at the end of the tick on the client */
-	UFUNCTION(reliable, client)
+	UFUNCTION(Reliable, Client)
 	void ClientForceGarbageCollection();
 
 	/** 
@@ -630,31 +630,35 @@ public:
 	 * @param	EndGameFocus - actor to view with camera
 	 * @param	bIsWinner - true if this controller is on winning team
 	 */
-	UFUNCTION(reliable, client)
+	UFUNCTION(Reliable, Client)
 	void ClientGameEnded(class AActor* EndGameFocus, bool bIsWinner);
 
 	/** 
 	 * Server uses this to force client into NewState .
 	 * @Note ALL STATE NAMES NEED TO BE DEFINED IN name table in UnrealNames.h to be correctly replicated (so they are mapped to the same thing on client and server).
 	 */
-	UFUNCTION(reliable, client)
+	UFUNCTION(Reliable, Client)
 	void ClientGotoState(FName NewState);
 
 	/** calls IgnoreLookInput on client */
-	UFUNCTION(client, reliable)
+	UFUNCTION(Reliable, Client)
 	void ClientIgnoreLookInput(bool bIgnore);
 
 	/** calls IgnoreMoveInput on client */
-	UFUNCTION(client, reliable)
+	UFUNCTION(Reliable, Client)
 	void ClientIgnoreMoveInput(bool bIgnore);
 
+	/** Informs the client of a rejected attempt to confirm an ability. */
+	UFUNCTION(Unreliable, Client)
+	void ClientNotifyRejectedAbilityConfirmation();
+
 	/**
 	 * Outputs a message to HUD
 	 * @param S - message to display
 	 * @param Type - @todo document
 	 * @param MsgLifeTime - Optional length of time to display 0 = default time
 	 */
-	UFUNCTION(reliable, client)
+	UFUNCTION(Reliable, Client)
 	void ClientMessage(const FString& S, FName Type = NAME_None, float MsgLifeTime = 0);
 
 	/** Play the indicated CameraAnim on this camera.
@@ -706,7 +710,7 @@ public:
 	 * @param bFirst - whether this is the first item in the list (so clear the list first)
 	 * @param bLast - whether this is the last item in the list (so start preparing the change after receiving it)
 	 */
-	UFUNCTION(reliable, client)
+	UFUNCTION(Reliable, Client)
 	void ClientPrepareMapChange(FName LevelName, bool bFirst, bool bLast);
 
 	/**
@@ -717,22 +721,22 @@ public:
 	 * @param bEnableStreaming	- Whether to start (true) or stop (false) streaming
 	 * @param CinematicTextureGroups	- Bitfield indicating which texture groups that use extra high-resolution mips
 	 */
-	UFUNCTION(reliable, client)
+	UFUNCTION(Reliable, Client)
 	void ClientPrestreamTextures(class AActor* ForcedActor, float ForceDuration, bool bEnableStreaming, int32 CinematicTextureGroups = 0);
 
 	/** Tell client to reset the PlayerController */
-	UFUNCTION(reliable, client)
+	UFUNCTION(Reliable, Client)
 	void ClientReset();
 
 	/** Tell client to restart the level */
-	UFUNCTION(reliable, client)
+	UFUNCTION(Reliable, Client)
 	void ClientRestart(class APawn* NewPawn);
 
 	/** 
 	 * Tells the client to block until all pending level streaming actions are complete.
 	 * Happens at the end of the tick primarily used to force update the client ASAP at join time.
 	 */
-	UFUNCTION(reliable, client)
+	UFUNCTION(Reliable, Client)
 	void ClientSetBlockOnAsyncLoading();
 
 	/** Tell client to fade camera
@@ -742,18 +746,18 @@ public:
 	 * @Param FadeTime - length of time for fade to occur over
 	 * @Param bFadeAudio - true to apply fading of audio alongside the video
 	 */
-	UFUNCTION(reliable, client)
+	UFUNCTION(Reliable, Client)
 	void ClientSetCameraFade(bool bEnableFading, FColor FadeColor = FColor(ForceInit), FVector2D FadeAlpha = FVector2D(ForceInit), float FadeTime = 0, bool bFadeAudio = false);
 
 	/**
 	 * Replicated function to set camera style on client
 	 * @param	NewCamMode, name defining the new camera mode
 	 */
-	UFUNCTION(reliable, client)
+	UFUNCTION(Reliable, Client)
 	void ClientSetCameraMode(FName NewCamMode);
 
 	/** Called by the server to synchronize cinematic transitions with the client */
-	UFUNCTION(reliable, client)
+	UFUNCTION(Reliable, Client)
 	void ClientSetCinematicMode(bool bInCinematicMode, bool bAffectsMovement, bool bAffectsTurning, bool bAffectsHUD);
 
 	/**
@@ -764,11 +768,11 @@ public:
 	 * @param ForceDuration	- Number of seconds to keep all mip-levels in memory, disregarding the normal priority logic.
 	 * @param CinematicTextureGroups	- Bitfield indicating which texture groups that use extra high-resolution mips
 	 */
-	UFUNCTION(reliable, client)
+	UFUNCTION(Reliable, Client)
 	void ClientSetForceMipLevelsToBeResident(class UMaterialInterface* Material, float ForceDuration, int32 CinematicTextureGroups = 0);
 
 	/** Set the client's class of HUD */
-	UFUNCTION(reliable, client)
+	UFUNCTION(Reliable, Client)
 	void ClientSetHUD(TSubclassOf<class AHUD> NewHUDClass);
 
 	/** Helper to get the size of the HUD canvas for this player controller.  Returns 0 if there is no HUD */
@@ -783,7 +787,7 @@ public:
 	 * @param A - new actor to set as view target
 	 * @param TransitionParams - parameters to use for controlling the transition
 	 */
-	UFUNCTION(reliable, client)
+	UFUNCTION(Reliable, Client)
 	void ClientSetViewTarget(class AActor* A, struct FViewTargetTransitionParams TransitionParams = FViewTargetTransitionParams());
 
 	/** Spawn a camera lens effect (e.g. blood). */
@@ -795,7 +799,7 @@ public:
 	virtual void ClientClearCameraLensEffects();
 
 	/** Stop camera animation on client. */
-	UFUNCTION(reliable, client)
+	UFUNCTION(Reliable, Client)
 	void ClientStopCameraAnim(class UCameraAnim* AnimToStop);
 
 	/** Stop camera shake on client.  */
@@ -858,7 +862,7 @@ public:
 	 * @param MapPackageGuid	The GUID of the map package to travel to - this is used to find the file when it has been autodownloaded,
 	 * 							so it is only needed for clients
 	 */
-	UFUNCTION(reliable, client)
+	UFUNCTION(Reliable, Client)
 	void ClientTravelInternal(const FString& URL, enum ETravelType TravelType, bool bSeamless = false, FGuid MapPackageGuid = FGuid());
 
 	/**
@@ -869,22 +873,22 @@ public:
 	 * @param bNewShouldBlockOnLoad - Whether we want to force a blocking load
 	 * @param LODIndex				- Current LOD index for a streaming level
 	 */
-	UFUNCTION(reliable, client)
+	UFUNCTION(Reliable, Client)
 	void ClientUpdateLevelStreamingStatus(FName PackageName, bool bNewShouldBeLoaded, bool bNewShouldBeVisible, bool bNewShouldBlockOnLoad, int32 LODIndex);
 
 	/** Notify client they were kicked from the server */
-	UFUNCTION(reliable, client)
+	UFUNCTION(Reliable, Client)
 	void ClientWasKicked(const FText& KickReason);
 
 	/** Assign Pawn to player, but avoid calling ClientRestart if we have already accepted this pawn */
-	UFUNCTION(reliable, client)
+	UFUNCTION(Reliable, Client)
 	void ClientRetryClientRestart(class APawn* NewPawn);
 
 	/** Call ClientRetryClientRestart, but only if the current pawn is not the currently acknowledged pawn (and throttled to avoid saturating the network). */
 	virtual void SafeRetryClientRestart();
 
 	/** send client localized message id */
-	UFUNCTION(reliable, client)
+	UFUNCTION(Reliable, Client)
 	void ClientReceiveLocalizedMessage(TSubclassOf<ULocalMessage> Message, int32 Switch = 0, class APlayerState* RelatedPlayerState_1 = NULL, class APlayerState* RelatedPlayerState_2 = NULL, class UObject* OptionalObject = NULL);
 
 	/** acknowledge possession of pawn */
@@ -963,7 +967,7 @@ public:
 	void ServerViewSelf(struct FViewTargetTransitionParams TransitionParams = FViewTargetTransitionParams());
 
 	/** @todo document */
-	UFUNCTION(reliable, client)
+	UFUNCTION(Reliable, Client)
 	void ClientTeamMessage(class APlayerState* SenderPlayerState, const FString& S, FName Type, float MsgLifeTime = 0);
 
 	/** Used by UGameplayDebuggingControllerComponent to replicate messages for AI debugging in network games. */
diff --git a/Engine/Source/Runtime/Engine/Private/PlayerController.cpp b/Engine/Source/Runtime/Engine/Private/PlayerController.cpp
--- a/Engine/Source/Runtime/Engine/Private/PlayerController.cpp
+++ b/Engine/Source/Runtime/Engine/Private/PlayerController.cpp
@@ -2741,6 +2741,11 @@ void APlayerController::ClientIgnoreLookInput_Implementation(bool bIgnore)
 	SetIgnoreLookInput(bIgnore);
 }
 
+void APlayerController::ClientNotifyRejectedAbilityConfirmation_Implementation()
+{
+	//TODO: Hook UI notification here. Might be good to include a cooldown on this.
+}
+
 
 void APlayerController::DisplayDebug(class UCanvas* Canvas, const FDebugDisplayInfo& DebugDisplay, float& YL, float& YPos)
 {
diff --git a/Engine/Source/Runtime/GameplayAbilities/Private/Abilities/GameplayAbilityTargetActor.cpp b/Engine/Source/Runtime/GameplayAbilities/Private/Abilities/GameplayAbilityTargetActor.cpp
--- a/Engine/Source/Runtime/GameplayAbilities/Private/Abilities/GameplayAbilityTargetActor.cpp
+++ b/Engine/Source/Runtime/GameplayAbilities/Private/Abilities/GameplayAbilityTargetActor.cpp
@@ -46,6 +46,16 @@ void AGameplayAbilityTargetActor::ConfirmTargetingAndContinue()
 	{
 		TargetDataReadyDelegate.Broadcast(FGameplayAbilityTargetDataHandle());
 	}
+	else
+	{
+		if (OwningAbility)
+		{
+			if (APlayerController* PC = OwningAbility->GetActorInfo().PlayerController.Get())
+			{
+				PC->ClientNotifyRejectedAbilityConfirmation();
+			}
+		}
+	}
 }
 
 void AGameplayAbilityTargetActor::ConfirmTargeting()
@@ -58,6 +68,16 @@ void AGameplayAbilityTargetActor::ConfirmTargeting()
 			Destroy();
 		}
 	}
+	else
+	{
+		if (OwningAbility)
+		{
+			if (APlayerController* PC = OwningAbility->GetActorInfo().PlayerController.Get())
+			{
+				PC->ClientNotifyRejectedAbilityConfirmation();
+			}
+		}
+	}
 }
 
 /** Outside code is saying 'stop everything and just forget about it' */
diff --git a/Engine/Source/Runtime/GameplayAbilities/Private/AbilitySystemComponent_Abilities.cpp b/Engine/Source/Runtime/GameplayAbilities/Private/AbilitySystemComponent_Abilities.cpp
--- a/Engine/Source/Runtime/GameplayAbilities/Private/AbilitySystemComponent_Abilities.cpp
+++ b/Engine/Source/Runtime/GameplayAbilities/Private/AbilitySystemComponent_Abilities.cpp
@@ -999,21 +999,39 @@ void UAbilitySystemComponent::ServerInputRelease_Implementation(FGameplayAbility
 void UAbilitySystemComponent::TargetConfirm()
 {
 	TArray<AGameplayAbilityTargetActor*> LeftoverTargetActors;
+	bool bAnyAbilityDeclined = false;
 	for (AGameplayAbilityTargetActor* TargetActor : SpawnedTargetActors)
 	{
 		if (TargetActor)
 		{
 			if (TargetActor->IsConfirmTargetingAllowed())
 			{
+				//TODO: There might not be any cases where this bool is false
+				if (!TargetActor->bDestroyOnConfirmation)
+				{
+					LeftoverTargetActors.Add(TargetActor);
+				}
 				TargetActor->ConfirmTargeting();
 			}
 			else
 			{
+				bAnyAbilityDeclined = true;
 				LeftoverTargetActors.Add(TargetActor);
 			}
 		}
 	}
-	SpawnedTargetActors = LeftoverTargetActors;		//These actors declined to confirm targeting, so keep contact with them.
+	SpawnedTargetActors = LeftoverTargetActors;		//These actors declined to confirm targeting, or are allowed to fire multiple times, so keep contact with them.
+	if (bAnyAbilityDeclined)
+	{
+		//Tell the client's player controller that we rejected at least one attempt to confirm. Mixed success/confirmation is a weird situation that may require further thought.
+		if (FGameplayAbilityActorInfo* ActorInfo = AbilityActorInfo.Get())
+		{
+			if (APlayerController* PC = ActorInfo->PlayerController.Get())
+			{
+				PC->ClientNotifyRejectedAbilityConfirmation();
+			}
+		}
+	}
 }
 
 void UAbilitySystemComponent::TargetCancel()
