Allow actor channels to continue to process bunches if needed when closed

diff --git a/Engine/Source/Runtime/Engine/Classes/Engine/ActorChannel.h b/Engine/Source/Runtime/Engine/Classes/Engine/ActorChannel.h
--- a/Engine/Source/Runtime/Engine/Classes/Engine/ActorChannel.h
+++ b/Engine/Source/Runtime/Engine/Classes/Engine/ActorChannel.h
@@ -91,7 +91,8 @@ class ENGINE_API UActorChannel : public UChannel
 
 	virtual void Tick() override;
 	void ProcessBunch( FInBunch & Bunch );
-	
+	void ProcessQueuedBunches();
+
 	virtual void ReceivedNak( int32 NakPacketId ) override;
 	
 	virtual void Close() override;
@@ -239,7 +240,7 @@ protected:
 	TSharedRef< FObjectReplicator > & FindOrCreateReplicator(UObject *Obj);
 	bool ObjectHasReplicator(UObject *Obj);	// returns whether we have already created a replicator for this object or not
 
-	virtual void CleanUp() override;
+	virtual bool CleanUp() override;
 
 	/** Closes the actor channel but with a 'dormant' flag set so it can be reopened */
 	virtual void BecomeDormant();
diff --git a/Engine/Source/Runtime/Engine/Classes/Engine/Channel.h b/Engine/Source/Runtime/Engine/Classes/Engine/Channel.h
--- a/Engine/Source/Runtime/Engine/Classes/Engine/Channel.h
+++ b/Engine/Source/Runtime/Engine/Classes/Engine/Channel.h
@@ -111,14 +111,7 @@ class ENGINE_API UChannel : public UObject
 	void AssertInSequenced();
 
 	/** cleans up channel if it hasn't already been */
-	void ConditionalCleanUp()
-	{
-		if (!IsPendingKill())
-		{
-			MarkPendingKill();
-			CleanUp();
-		}
-	}
+	void ConditionalCleanUp();
 
 	/** Returns true if channel is ready to go dormant (e.g., all outstanding property updates have been ACK'd) */
 	virtual bool ReadyForDormancy(bool suppressLogs=false) { return false; }
@@ -134,7 +127,7 @@ protected:
 	virtual void BecomeDormant() { }
 
 	/** cleans up channel structures and NULLs references to the channel */
-	virtual void CleanUp();
+	virtual bool CleanUp();
 
 private:
 
diff --git a/Engine/Source/Runtime/Engine/Classes/Engine/NetConnection.h b/Engine/Source/Runtime/Engine/Classes/Engine/NetConnection.h
--- a/Engine/Source/Runtime/Engine/Classes/Engine/NetConnection.h
+++ b/Engine/Source/Runtime/Engine/Classes/Engine/NetConnection.h
@@ -240,6 +240,10 @@ public:
 	/** @todo document */
 	TMap<TWeakObjectPtr<AActor>,class UActorChannel*> ActorChannels;
 
+	/** This holds a list of actor channels that want to fully shutdown, but need to continue processing bunches before doing so */
+	UPROPERTY()
+	TArray<class UActorChannel*> KeepProcessingActorChannelBunches;
+
 	/** Actors that have gone dormant on this connection	
 	 *  The only way to get on this list is when the actor channel closes. UActorChannel::Close.
      *  Once in this set, and only then, is an actor considered network dormant.
diff --git a/Engine/Source/Runtime/Engine/Classes/Engine/VoiceChannel.h b/Engine/Source/Runtime/Engine/Classes/Engine/VoiceChannel.h
--- a/Engine/Source/Runtime/Engine/Classes/Engine/VoiceChannel.h
+++ b/Engine/Source/Runtime/Engine/Classes/Engine/VoiceChannel.h
@@ -36,7 +36,7 @@ protected:
 	/** 
 	 * Cleans up any voice data remaining in the queue 
 	 */
-	virtual void CleanUp();
+	virtual bool CleanUp() override;
 
 	/**
 	 * Processes the in bound bunch to extract the voice data
diff --git a/Engine/Source/Runtime/Engine/Private/DataChannel.cpp b/Engine/Source/Runtime/Engine/Private/DataChannel.cpp
--- a/Engine/Source/Runtime/Engine/Private/DataChannel.cpp
+++ b/Engine/Source/Runtime/Engine/Private/DataChannel.cpp
@@ -88,8 +88,20 @@ void UChannel::Close()
 	}
 }
 
+void UChannel::ConditionalCleanUp()
+{
+	if ( !IsPendingKill() )
+	{
+		// CleanUp can return false to signify that we shouldn't mark pending kill quite yet
+		// We'll need to call cleanup again later on
+		if ( CleanUp() )
+		{
+			MarkPendingKill();
+		}
+	}
+}
 
-void UChannel::CleanUp()
+bool UChannel::CleanUp()
 {
 	checkSlow(Connection != NULL);
 	checkSlow(Connection->Channels[ChIndex] == this);
@@ -132,6 +144,8 @@ void UChannel::CleanUp()
 	verifySlow(Connection->OpenChannels.Remove(this) == 1);
 	Connection->Channels[ChIndex] = NULL;
 	Connection = NULL;
+
+	return true;
 }
 
 
@@ -1304,19 +1318,14 @@ void UActorChannel::CleanupReplicators( const bool bKeepReplicators )
 	ActorReplicator = NULL;
 }
 
-void UActorChannel::CleanUp()
+bool UActorChannel::CleanUp()
 {
-	const bool IsServer = (Connection->Driver->ServerConnection == NULL);
-
-	UE_LOG( LogNetTraffic, Log, TEXT( "UActorChannel::CleanUp: Channel: %i, IsServer: %s" ), ChIndex, IsServer ? TEXT( "YES" ) : TEXT( "NO" ) );
-
-	// Remove from hash and stuff.
-	SetClosingFlag();
+	const bool bIsServer = Connection->Driver->IsServer();
 
-	CleanupReplicators();
+	UE_LOG( LogNetTraffic, Log, TEXT( "UActorChannel::CleanUp: Channel: %i, IsServer: %s" ), ChIndex, bIsServer ? TEXT( "YES" ) : TEXT( "NO" ) );
 
 	// If we're the client, destroy this actor.
-	if (!IsServer)
+	if (!bIsServer)
 	{
 		check(Actor == NULL || Actor->IsValidLowLevel());
 		checkSlow(Connection != NULL);
@@ -1347,6 +1356,38 @@ void UActorChannel::CleanUp()
 		Connection->SentTemporaries.Remove(Actor);
 	}
 
+	if ( !bIsServer && Dormant && QueuedBunches.Num() > 0 )
+	{
+		UE_LOG( LogNet, Log, TEXT( "UActorChannel::CleanUp: Adding to KeepProcessingActorChannelBunches. Channel: %i" ), ChIndex );
+
+		// Remember the connection, since CleanUp below will NULL it
+		UNetConnection* OldConnection = Connection;
+
+		// This will unregister the channel, and make it free for opening again
+		// We need to do this, since the server will assume this channel is free once we ack this packet
+		Super::CleanUp();
+
+		// Restore connection property since we'll need it for processing bunches (the Super::CleanUp call above NULL'd it)
+		Connection = OldConnection;
+
+		// Add this channel to the KeepProcessingActorChannelBunches list
+		check( !Connection->KeepProcessingActorChannelBunches.Contains( this ) );
+		Connection->KeepProcessingActorChannelBunches.Add( this );
+
+		// We set ChIndex to -1 to signify that we've already been "closed" but we aren't done processing bunches
+		ChIndex = -1;
+
+		// Return false so we won't do pending kill yet
+		return false;
+	}
+
+	check( !Connection->KeepProcessingActorChannelBunches.Contains( this ) );
+
+	// Remove from hash and stuff.
+	SetClosingFlag();
+
+	CleanupReplicators();
+
 	// We don't care about any leftover pending guids at this point
 	PendingGuidResolves.Empty();
 
@@ -1358,7 +1399,13 @@ void UActorChannel::CleanUp()
 
 	QueuedBunches.Empty();
 
-	Super::CleanUp();
+	// We check for -1 here, which will be true if this channel has already been closed but still needed to process bunches before fully closing
+	if ( ChIndex >= 0 )	
+	{
+		return Super::CleanUp();
+	}
+
+	return true;
 }
 
 void UActorChannel::ReceivedNak( int32 NakPacketId )
@@ -1395,7 +1442,7 @@ void UActorChannel::SetChannelActor( AActor* InActor )
 
 	UE_LOG(LogNetTraffic, VeryVerbose, TEXT("SetChannelActor[%d]: Actor: %s"), ChIndex, Actor ? *Actor->GetFullName() : TEXT("NULL") );
 
-	if ( Connection->PendingOutRec[ChIndex] > 0 )
+	if ( ChIndex >= 0 && Connection->PendingOutRec[ChIndex] > 0 )
 	{
 		// send empty reliable bunches to synchronize both sides
 		// UE_LOG(LogNetTraffic, Log, TEXT("%i Actor %s WILL BE sending %i vs first %i"), ChIndex, *Actor->GetName(), Connection->PendingOutRec[ChIndex],Connection->OutReliable[ChIndex]);
@@ -1460,7 +1507,11 @@ void UActorChannel::SetChannelActorForDestroy( FActorDestructionInfo *DestructIn
 void UActorChannel::Tick()
 {
 	Super::Tick();
-	
+	ProcessQueuedBunches();
+}
+
+void UActorChannel::ProcessQueuedBunches()
+{
 	// Try to resolve any guids that are holding up the network stream on this channel
 	for ( auto It = PendingGuidResolves.CreateIterator(); It; ++It )
 	{
@@ -1474,7 +1525,7 @@ void UActorChannel::Tick()
 		if ( Connection->Driver->GuidCache->IsGUIDBroken( *It, true ) )
 		{
 			// This guid is broken, remove it, and warn
-			UE_LOG( LogNet, Warning, TEXT( "UActorChannel::Tick: Guid is broken. NetGUID: %s, ChIndex: %i, Actor: %s" ), *It->ToString(), ChIndex, Actor != NULL ? *Actor->GetPathName() : TEXT( "NULL" ) );
+			UE_LOG( LogNet, Warning, TEXT( "UActorChannel::ProcessQueuedBunches: Guid is broken. NetGUID: %s, ChIndex: %i, Actor: %s" ), *It->ToString(), ChIndex, Actor != NULL ? *Actor->GetPathName() : TEXT( "NULL" ) );
 			It.RemoveCurrent();
 			continue;
 		}
@@ -1489,7 +1540,7 @@ void UActorChannel::Tick()
 			delete QueuedBunches[i];
 		}
 
-		UE_LOG( LogNet, Log, TEXT( "UActorChannel::Tick: Flushing queued bunches. ChIndex: %i, Actor: %s, Queued: %i" ), ChIndex, Actor != NULL ? *Actor->GetPathName() : TEXT( "NULL" ), QueuedBunches.Num() );
+		UE_LOG( LogNet, Log, TEXT( "UActorChannel::ProcessQueuedBunches: Flushing queued bunches. ChIndex: %i, Actor: %s, Queued: %i" ), ChIndex, Actor != NULL ? *Actor->GetPathName() : TEXT( "NULL" ), QueuedBunches.Num() );
 
 		QueuedBunches.Empty();
 	}
@@ -1501,10 +1552,26 @@ void UActorChannel::Tick()
 
 		if ( FPlatformTime::Seconds() - QueuedBunchStartTime > QUEUED_BUNCH_TIMEOUT_IN_SECONDS )
 		{
-			UE_LOG( LogNet, Warning, TEXT( "UActorChannel::Tick: Queued bunches for longer than normal. ChIndex: %i, Actor: %s, Queued: %i" ), ChIndex, Actor != NULL ? *Actor->GetPathName() : TEXT( "NULL" ), QueuedBunches.Num() );
+			UE_LOG( LogNet, Warning, TEXT( "UActorChannel::ProcessQueuedBunches: Queued bunches for longer than normal. ChIndex: %i, Actor: %s, Queued: %i" ), ChIndex, Actor != NULL ? *Actor->GetPathName() : TEXT( "NULL" ), QueuedBunches.Num() );
 			QueuedBunchStartTime = FPlatformTime::Seconds();
 		}
 	}
+	else
+	{
+		if ( ChIndex == -1 )
+		{
+			UE_LOG( LogNet, Log, TEXT( "UActorChannel::ProcessQueuedBunches: Removing from KeepProcessingActorChannelBunches. ChIndex: %i, Actor: %s" ), ChIndex, Actor != NULL ? *Actor->GetPathName() : TEXT( "NULL" ) );
+
+			check( Connection->KeepProcessingActorChannelBunches.Contains( this ) );
+			Connection->KeepProcessingActorChannelBunches.Remove( this );
+			// Since we are done processing bunches, we can now actually clean this channel up
+			ConditionalCleanUp();
+		}
+		else
+		{
+			check( !Connection->KeepProcessingActorChannelBunches.Contains( this ) );
+		}
+	}
 }
 
 void UActorChannel::ReceivedBunch( FInBunch & Bunch )
diff --git a/Engine/Source/Runtime/Engine/Private/NetConnection.cpp b/Engine/Source/Runtime/Engine/Private/NetConnection.cpp
--- a/Engine/Source/Runtime/Engine/Private/NetConnection.cpp
+++ b/Engine/Source/Runtime/Engine/Private/NetConnection.cpp
@@ -1379,6 +1379,11 @@ void UNetConnection::Tick()
 			OpenChannels[i]->Tick();
 		}
 
+		for ( int32 i = KeepProcessingActorChannelBunches.Num() - 1; i >= 0; i-- )
+		{
+			KeepProcessingActorChannelBunches[i]->ProcessQueuedBunches();
+		}
+
 		// If channel 0 has closed, mark the connection as closed.
 		if( Channels[0]==NULL && (OutReliable[0]!=0 || InReliable[0]!=0) )
 		{
diff --git a/Engine/Source/Runtime/Engine/Private/VoiceChannel.cpp b/Engine/Source/Runtime/Engine/Private/VoiceChannel.cpp
--- a/Engine/Source/Runtime/Engine/Private/VoiceChannel.cpp
+++ b/Engine/Source/Runtime/Engine/Private/VoiceChannel.cpp
@@ -12,12 +12,12 @@
 #include "Engine/VoiceChannel.h"
 
 /** Cleans up any voice data remaining in the queue */
-void UVoiceChannel::CleanUp()
+bool UVoiceChannel::CleanUp()
 {
 	// Clear out refs to any voice packets so we don't leak
 	VoicePackets.Empty();
 	// Route to the parent class for their cleanup
-	Super::CleanUp();
+	return Super::CleanUp();
 }
 
 /**
