Fixed so that skeletalmesh component To World Transform updates with relative to LocalToAtom
 - otherwise it will generate extra movement because it didn't include LocalToAtom, root generates incorrect velocity, causing it to fly away

TTP 338026

diff --git a/Engine/Source/Runtime/Engine/Classes/Components/PrimitiveComponent.h b/Engine/Source/Runtime/Engine/Classes/Components/PrimitiveComponent.h
--- a/Engine/Source/Runtime/Engine/Classes/Components/PrimitiveComponent.h
+++ b/Engine/Source/Runtime/Engine/Classes/Components/PrimitiveComponent.h
@@ -1015,7 +1015,7 @@ public:
 
 	/** Move this component to match the physics rigid body pose. Note, a warning will be generated if you call this function on a component that is attached to something */
 	void SyncComponentToRBPhysics();
-
+	
 	/** 
 	 *	Returns the matrix that should be used to render this component. 
 	 *	Allows component class to perform graphical distortion to the component not supported by an FTransform 
@@ -1144,6 +1144,16 @@ protected:
 	virtual void CreatePhysicsState() override;
 	virtual void DestroyPhysicsState() override;
 	virtual void OnActorEnableCollisionChanged() override;
+	/**
+	 * Called to get the Component To World Transform from the Root BodyInstance
+	 * This needs to be virtual since SkeletalMeshComponent Root has to undo its own transform
+	 * Without this, the root LocalToAtom is overriden by physics simulation, causing kinematic velocity to 
+	 * accelerate simulation
+	 *
+	 * @param : UseBI - root body instsance
+	 * @return : New ComponentToWorld to use
+	 */
+	virtual FTransform GetComponentTransformFromBodyInstance(FBodyInstance* UseBI);
 public:
 	virtual void RegisterComponentTickFunctions(bool bRegister) override;
 #if WITH_EDITOR
diff --git a/Engine/Source/Runtime/Engine/Classes/Components/SkeletalMeshComponent.h b/Engine/Source/Runtime/Engine/Classes/Components/SkeletalMeshComponent.h
--- a/Engine/Source/Runtime/Engine/Classes/Components/SkeletalMeshComponent.h
+++ b/Engine/Source/Runtime/Engine/Classes/Components/SkeletalMeshComponent.h
@@ -833,8 +833,11 @@ public:
 	virtual void SetAllPhysicsLinearVelocity(FVector NewVel,bool bAddToCurrent = false) override;
 	virtual float GetMass() const override;
 	virtual float CalculateMass(FName BoneName = NAME_None) override;
+protected:
+	virtual FTransform GetComponentTransformFromBodyInstance(FBodyInstance* UseBI) override;
 	// End UPrimitiveComponent interface.
 
+public:
 	// Begin USkinnedMeshComponent interface
 	virtual bool UpdateLODStatus() override;
 	virtual void RefreshBoneTransforms( FActorComponentTickFunction* TickFunction = NULL ) override;
diff --git a/Engine/Source/Runtime/Engine/Private/PrimitiveComponentPhysics.cpp b/Engine/Source/Runtime/Engine/Private/PrimitiveComponentPhysics.cpp
--- a/Engine/Source/Runtime/Engine/Private/PrimitiveComponentPhysics.cpp
+++ b/Engine/Source/Runtime/Engine/Private/PrimitiveComponentPhysics.cpp
@@ -443,6 +443,11 @@ void UPrimitiveComponent::SetPhysMaterialOverride(UPhysicalMaterial* NewPhysMate
 	BodyInstance.SetPhysMaterialOverride(NewPhysMaterial);
 }
 
+FTransform UPrimitiveComponent::GetComponentTransformFromBodyInstance(FBodyInstance* UseBI)
+{
+	return UseBI->GetUnrealWorldTransform();
+}
+
 void UPrimitiveComponent::SyncComponentToRBPhysics()
 {
 	if(!IsRegistered())
@@ -474,7 +479,7 @@ void UPrimitiveComponent::SyncComponentToRBPhysics()
 	}
 
 	// See if the transform is actually different, and if so, move the component to match physics
-	const FTransform NewTransform = UseBI->GetUnrealWorldTransform();	
+	const FTransform NewTransform = GetComponentTransformFromBodyInstance(UseBI);	
 	if(!NewTransform.EqualsNoScale(ComponentToWorld))
 	{
 		const FVector MoveBy = NewTransform.GetLocation() - ComponentToWorld.GetLocation();
diff --git a/Engine/Source/Runtime/Engine/Private/SkeletalMeshComponentPhysics.cpp b/Engine/Source/Runtime/Engine/Private/SkeletalMeshComponentPhysics.cpp
--- a/Engine/Source/Runtime/Engine/Private/SkeletalMeshComponentPhysics.cpp
+++ b/Engine/Source/Runtime/Engine/Private/SkeletalMeshComponentPhysics.cpp
@@ -4946,5 +4946,11 @@ void USkeletalMeshComponent::ForceClothNextUpdateTeleportAndReset()
 #endif// #if WITH_APEX_CLOTHING
 }
 
-
+FTransform USkeletalMeshComponent::GetComponentTransformFromBodyInstance(FBodyInstance* UseBI)
+{
+	// undo root transform so that it only moves according to what actor itself suppose to move
+	FTransform BodyTransform = UseBI->GetUnrealWorldTransform();
+	FTransform RootTransform(LocalAtoms[0].GetRotation());
+	return RootTransform.GetRelativeTransformReverse(BodyTransform);
+}
 #undef LOCTEXT_NAMESPACE
