Navmesh tile generators are allocated on demand and destroyed after tile build is completed
Navmesh origin does not depend on navigation volumes and always at zero point.
Navmesh allocates much less tiles when there are more than one navigation volume

diff --git a/Engine/Source/Editor/UnrealEd/Private/EditorBuildUtils.cpp b/Engine/Source/Editor/UnrealEd/Private/EditorBuildUtils.cpp
--- a/Engine/Source/Editor/UnrealEd/Private/EditorBuildUtils.cpp
+++ b/Engine/Source/Editor/UnrealEd/Private/EditorBuildUtils.cpp
@@ -735,7 +735,6 @@ void FEditorBuildUtils::SubmitPackagesForAutomatedBuild( const TSet<UPackage*>&
 
 void FEditorBuildUtils::TriggerNavigationBuilder(UWorld* InWorld, EBuildOptions::Type Id)
 {
-#if WITH_NAVIGATION_GENERATOR
 	if( InWorld->GetWorldSettings()->bEnableNavigationSystem &&
 		InWorld->GetNavigationSystem() )
 	{
@@ -760,7 +759,6 @@ void FEditorBuildUtils::TriggerNavigationBuilder(UWorld* InWorld, EBuildOptions:
 		// Invoke navmesh generator
 		InWorld->GetNavigationSystem()->Build();
 	}
-#endif
 }
 
 #undef LOCTEXT_NAMESPACE
\ No newline at end of file
diff --git a/Engine/Source/Editor/UnrealEd/Private/NavigationBuildingNotification.cpp b/Engine/Source/Editor/UnrealEd/Private/NavigationBuildingNotification.cpp
--- a/Engine/Source/Editor/UnrealEd/Private/NavigationBuildingNotification.cpp
+++ b/Engine/Source/Editor/UnrealEd/Private/NavigationBuildingNotification.cpp
@@ -29,7 +29,7 @@ void FNavigationBuildingNotificationImpl::BuildStarted()
 		NavigationBuiltCompleteNotification.Pin()->ExpireAndFadeout();
 	}
 
-	FNotificationInfo Info( NSLOCTEXT("NavigationBuild", "NavigationBuildingInProgress", "Building Navigation") );
+	FNotificationInfo Info( GetNotificationText() );
 	Info.bFireAndForget = false;
 	Info.FadeOutDuration = 0.0f;
 	Info.ExpireDuration = 0.0f;
@@ -100,9 +100,27 @@ void FNavigationBuildingNotificationImpl::ClearCompleteNotification()
 	}
 }
 
+FText FNavigationBuildingNotificationImpl::GetNotificationText() const
+{
+	int32 RemainingTasks = 0;
+	UEditorEngine* const EEngine = Cast<UEditorEngine>(GEngine);
+	if (EEngine)
+	{
+		FWorldContext &EditorContext = EEngine->GetEditorWorldContext();
+		if (EditorContext.World() && EditorContext.World()->GetNavigationSystem())
+		{
+			RemainingTasks = EditorContext.World()->GetNavigationSystem()->GetNumRemainingBuildTasks();
+		}
+	}
+		
+	FFormatNamedArguments Args;
+	Args.Add(TEXT("RemainingTasks"), FText::AsNumber(RemainingTasks));
+	return FText::Format(NSLOCTEXT("NavigationBuild", "NavigationBuildingInProgress", "Building Navigation ({RemainingTasks})"), Args);
+}
+
+
 void FNavigationBuildingNotificationImpl::Tick(float DeltaTime)
 {
-#if WITH_NAVIGATION_GENERATOR
 	if (FPlayWorldCommandCallbacks::IsInPIE_AndRunning())
 	{
 		return;
@@ -138,10 +156,13 @@ void FNavigationBuildingNotificationImpl::Tick(float DeltaTime)
 		{
 			BuildFinished();
 		}
+		else if (bBuildInProgress && NavigationBuildNotificationPtr.IsValid())
+		{
+			NavigationBuildNotificationPtr.Pin()->SetText(GetNotificationText());
+		}
 
 		bPreviouslyDetectedBuild = bBuildInProgress;
 	}
-#endif
 }
 
 TStatId FNavigationBuildingNotificationImpl::GetStatId() const
diff --git a/Engine/Source/Editor/UnrealEd/Private/NavigationBuildingNotification.h b/Engine/Source/Editor/UnrealEd/Private/NavigationBuildingNotification.h
--- a/Engine/Source/Editor/UnrealEd/Private/NavigationBuildingNotification.h
+++ b/Engine/Source/Editor/UnrealEd/Private/NavigationBuildingNotification.h
@@ -31,6 +31,7 @@ protected:
 
 private:
 	void ClearCompleteNotification();
+	FText GetNotificationText() const;
 
 	bool bPreviouslyDetectedBuild;
 	double TimeOfStartedBuild;
diff --git a/Engine/Source/Editor/UnrealEd/Private/PlayLevel.cpp b/Engine/Source/Editor/UnrealEd/Private/PlayLevel.cpp
--- a/Engine/Source/Editor/UnrealEd/Private/PlayLevel.cpp
+++ b/Engine/Source/Editor/UnrealEd/Private/PlayLevel.cpp
@@ -184,12 +184,10 @@ void UEditorEngine::EndPlayMap()
 	}
 
 	// Lose the EditorWorld pointer (this is only maintained while PIEing)
-#if WITH_NAVIGATION_GENERATOR
 	if (EditorWorld->GetNavigationSystem())
 	{
 		EditorWorld->GetNavigationSystem()->OnPIEEnd();
 	}
-#endif // WITH_NAVIGATION_GENERATOR
 
 	EditorWorld->bAllowAudioPlayback = true;
 	EditorWorld = NULL;
@@ -1669,12 +1667,10 @@ void UEditorEngine::PlayInEditor( UWorld* InWorld, bool bInSimulateInEditor )
 	FEditorDelegates::BeginPIE.Broadcast(bInSimulateInEditor);
 
 	// let navigation know PIE starts so it can avoid any blueprint creation/deletion/instantiation affect editor map's navmesh changes
-#if WITH_NAVIGATION_GENERATOR
 	if (InWorld->GetNavigationSystem())
 	{
 		InWorld->GetNavigationSystem()->OnPIEStart();
 	}
-#endif // WITH_NAVIGATION_GENERATOR
 
 	ULevelEditorPlaySettings const* EditorPlayInSettings = GetDefault<ULevelEditorPlaySettings>();
 	check(EditorPlayInSettings);
@@ -2285,12 +2281,11 @@ UGameInstance* UEditorEngine::CreatePIEGameInstance(int32 PIEInstance, bool bInS
 	if (!bSuccess)
 	{
 		FEditorDelegates::EndPIE.Broadcast(bInSimulateInEditor);
-#if WITH_NAVIGATION_GENERATOR
+
 		if (EditorWorld->GetNavigationSystem())
 		{
 			EditorWorld->GetNavigationSystem()->OnPIEEnd();
 		}
-#endif // WITH_NAVIGATION_GENERATOR
 
 		return nullptr;
 	}
diff --git a/Engine/Source/Editor/UnrealEd/Private/Tests/EditorBuildPromotionTests.cpp b/Engine/Source/Editor/UnrealEd/Private/Tests/EditorBuildPromotionTests.cpp
--- a/Engine/Source/Editor/UnrealEd/Private/Tests/EditorBuildPromotionTests.cpp
+++ b/Engine/Source/Editor/UnrealEd/Private/Tests/EditorBuildPromotionTests.cpp
@@ -4176,7 +4176,6 @@ namespace BuildPromotionTestHelper
 			GUnrealEd->Exec(World, TEXT("MAP REBUILD ALLVISIBLE"));
 			UE_LOG(LogEditorBuildPromotionTests, Display, TEXT("Rebuilt the map"));
 
-#if WITH_NAVIGATION_GENERATOR
 			if (World->GetWorldSettings()->bEnableNavigationSystem &&
 				World->GetNavigationSystem())
 			{
@@ -4184,7 +4183,6 @@ namespace BuildPromotionTestHelper
 				World->GetNavigationSystem()->Build();
 				UE_LOG(LogEditorBuildPromotionTests, Display, TEXT("Built navigation"));
 			}
-#endif
 
 			//Force AutoApplyLighting on
 			ULevelEditorMiscSettings* LevelEdSettings = GetMutableDefault<ULevelEditorMiscSettings>();
diff --git a/Engine/Source/Runtime/Engine/Classes/AI/Navigation/NavMeshBoundsVolume.h b/Engine/Source/Runtime/Engine/Classes/AI/Navigation/NavMeshBoundsVolume.h
--- a/Engine/Source/Runtime/Engine/Classes/AI/Navigation/NavMeshBoundsVolume.h
+++ b/Engine/Source/Runtime/Engine/Classes/AI/Navigation/NavMeshBoundsVolume.h
@@ -15,8 +15,8 @@ class ANavMeshBoundsVolume : public AVolume
 	GENERATED_UCLASS_BODY()
 
 	// Begin AActor Interface
-	virtual void PostInitializeComponents() override;
-	virtual void PostLoad() override;
+	virtual void PostRegisterAllComponents() override;
+	virtual void PostUnregisterAllComponents() override;
 	// End AActor Interface
 #if WITH_EDITOR
 	// Begin UObject Interface
diff --git a/Engine/Source/Runtime/Engine/Classes/AI/Navigation/NavModifierVolume.h b/Engine/Source/Runtime/Engine/Classes/AI/Navigation/NavModifierVolume.h
--- a/Engine/Source/Runtime/Engine/Classes/AI/Navigation/NavModifierVolume.h
+++ b/Engine/Source/Runtime/Engine/Classes/AI/Navigation/NavModifierVolume.h
@@ -20,6 +20,7 @@ class ANavModifierVolume : public AVolume, public INavRelevantInterface
 	virtual FBox GetNavigationBounds() const;
 
 #if WITH_EDITOR
+	virtual void PostEditUndo() override;
 	virtual void PostEditChangeProperty(FPropertyChangedEvent& PropertyChangedEvent) override;
 #endif
 };
diff --git a/Engine/Source/Runtime/Engine/Classes/AI/Navigation/NavigationData.h b/Engine/Source/Runtime/Engine/Classes/AI/Navigation/NavigationData.h
--- a/Engine/Source/Runtime/Engine/Classes/AI/Navigation/NavigationData.h
+++ b/Engine/Source/Runtime/Engine/Classes/AI/Navigation/NavigationData.h
@@ -4,6 +4,7 @@
 #include "AI/Navigation/NavFilters/NavigationQueryFilter.h"
 #include "AI/Navigation/NavigationTypes.h"
 #include "GameFramework/Actor.h"
+#include "UniquePtr.h"
 #include "NavigationData.generated.h"
 
 class FNavDataGenerator; 
@@ -410,7 +411,7 @@ class ENGINE_API ANavigationData : public AActor
 
 	virtual bool NeedsRebuild() const { return false; }
 	virtual bool CanRebuild() const;
-
+	
 	//----------------------------------------------------------------------//
 	// Generation & data access                                                      
 	//----------------------------------------------------------------------//
@@ -427,20 +428,18 @@ protected:
 	virtual void FillConfig(FNavDataConfig& Dest) { Dest = NavDataConfig; }
 
 public:
-#if WITH_NAVIGATION_GENERATOR 
-	/** Retrieves navmesh's generator, creating one if not already present */
-	class FNavDataGenerator* GetGenerator(FNavigationSystem::ECreateIfEmpty CreateIfNone);
+	virtual void RebuildAll();
+	virtual void TickAsyncBuild(float DeltaSeconds);
 
-	const class FNavDataGenerator* GetGenerator() const { return NavDataGenerator.Get(); }
+	/** Retrieves navmesh's generator */
+	FNavDataGenerator* GetGenerator() { return NavDataGenerator.Get(); }
+	const FNavDataGenerator* GetGenerator() const { return NavDataGenerator.Get(); }
 
 	/** Request navigation data update after changes in nav octree */
 	virtual void RebuildDirtyAreas(const TArray<FNavigationDirtyArea>& DirtyAreas);
 
-#endif // WITH_NAVIGATION_GENERATOR
 	/** releases navigation generator if any has been created */
 protected:
-	virtual void DestroyGenerator();
-
 	/** register self with navigation system as new NavAreaDefinition(s) observer */
 	void RegisterAsNavAreaClassObserver();
 
@@ -682,10 +681,6 @@ protected:
 	/** get ID to assign for newly added area */
 	virtual int32 GetNewAreaID(const UClass* AreaClass) const;
 	
-#if WITH_NAVIGATION_GENERATOR
-	virtual FNavDataGenerator* ConstructGenerator(const FNavAgentProperties& AgentProps) PURE_VIRTUAL(ANavigationData::ConstructGenerator, return NULL; );
-#endif // WITH_NAVIGATION_GENERATOR
-
 protected:
 	/** Navigation data versioning. */
 	uint32 DataVersion;
@@ -701,18 +696,8 @@ protected:
 	typedef bool(*FNavRaycastPtr)(const ANavigationData* NavDataInstance, const FVector& RayStart, const FVector& RayEnd, FVector& HitLocation, TSharedPtr<const FNavigationQueryFilter> QueryFilter, const UObject* Querier);
 	FNavRaycastPtr RaycastImplementation; 
 
-public:
-	/** common shared pointer type for all generators */
-	typedef TSharedPtr<FNavDataGenerator, ESPMode::ThreadSafe> FNavDataGeneratorSharedPtr;
-
-	/** common weak pointer type for all generators */
-	typedef TWeakPtr<FNavDataGenerator, ESPMode::ThreadSafe> FNavDataGeneratorWeakPtr;
-
 protected:
-#if WITH_NAVIGATION_GENERATOR
-	FNavDataGeneratorSharedPtr NavDataGenerator;
-#endif // WITH_NAVIGATION_GENERATOR
-
+	TUniquePtr<FNavDataGenerator> NavDataGenerator;
 	/** 
 	 *	Container for all path objects generated with this Navigation Data instance. 
 	 *	Is meant to be added to only on GameThread, and in fact should user should never 
diff --git a/Engine/Source/Runtime/Engine/Classes/AI/Navigation/NavigationGraph.h b/Engine/Source/Runtime/Engine/Classes/AI/Navigation/NavigationGraph.h
--- a/Engine/Source/Runtime/Engine/Classes/AI/Navigation/NavigationGraph.h
+++ b/Engine/Source/Runtime/Engine/Classes/AI/Navigation/NavigationGraph.h
@@ -58,11 +58,6 @@ class ANavigationGraph : public ANavigationData
 {
 	GENERATED_UCLASS_BODY()
 
-protected:
-#if WITH_NAVIGATION_GENERATOR
-	virtual FNavDataGenerator* ConstructGenerator(const FNavAgentProperties& AgentProps) override;
-#endif // WITH_NAVIGATION_GENERATOR
-
 public:
 	static ANavigationData* CreateNavigationInstances(UNavigationSystem* NavSys);
 };
diff --git a/Engine/Source/Runtime/Engine/Classes/AI/Navigation/NavigationSystem.h b/Engine/Source/Runtime/Engine/Classes/AI/Navigation/NavigationSystem.h
--- a/Engine/Source/Runtime/Engine/Classes/AI/Navigation/NavigationSystem.h
+++ b/Engine/Source/Runtime/Engine/Classes/AI/Navigation/NavigationSystem.h
@@ -9,11 +9,6 @@
 
 #define NAVSYS_DEBUG (0 && UE_BUILD_DEBUG)
 
-// if we'll be rebuilding navigation at runtime
-#define WITH_RUNTIME_NAVIGATION_BUILDING (1 && WITH_NAVIGATION_GENERATOR)
-
-#define NAVOCTREE_CONTAINS_COLLISION_DATA (1 && WITH_RECAST)
-
 #define NAV_USE_MAIN_NAVIGATION_DATA NULL
 
 class UNavigationPath;
@@ -432,10 +427,8 @@ protected:
 	/** registers NavArea classes awaiting registration in PendingNavAreaRegistration */
 	void ProcessNavAreaPendingRegistration();
 
-#if WITH_NAVIGATION_GENERATOR
 	/** used to apply updates of nav volumes in navigation system's tick */
 	void PerformNavigationBoundsUpdate(ANavMeshBoundsVolume* NavVolume);
-#endif // WITH_NAVIGATION_GENERATOR
 
 	/** @return pointer to ANavigationData instance of given ID, or NULL if it was not found. Note it looks only through registered navigation data */
 	ANavigationData* GetNavDataWithID(const uint16 NavDataID) const;
@@ -459,10 +452,8 @@ public:
 	/** removes all navoctree entries for actor and its components */
 	static void ClearNavOctreeAll(AActor* Actor);
 
-#if WITH_NAVIGATION_GENERATOR
 	void AddDirtyArea(const FBox& NewArea, int32 Flags);
 	void AddDirtyAreas(const TArray<FBox>& NewAreas, int32 Flags);
-#endif // WITH_NAVIGATION_GENERATOR
 
 	const FNavigationOctree* GetNavOctree() const { return NavOctree; }
 
@@ -524,14 +515,10 @@ public:
 	//----------------------------------------------------------------------//
 	// building
 	//----------------------------------------------------------------------//
-#if WITH_NAVIGATION_GENERATOR
+	
 	/** Triggers navigation building on all eligible navigation data. */
 	virtual void Build();
 
-	// @todo document
-	void OnUpdateStreamingStarted();
-	// @todo document
-	void OnUpdateStreamingFinished();
 	// @todo document
 	void OnPIEStart();
 	// @todo document
@@ -546,7 +533,12 @@ public:
 
 	/** Used to display "navigation building in progress" notify */
 	bool IsNavigationBuildInProgress(bool bCheckDirtyToo = true);
-#endif // WITH_NAVIGATION_GENERATOR
+
+	/** Used to display "navigation building in progress" counter */
+	int32 GetNumRemainingBuildTasks() const;
+
+	/** Number of currently running tasks */
+	int32 GetNumRunningBuildTasks() const;
 
 	/** Sets up SuportedAgents and NavigationDataCreators. Override it to add additional setup, but make sure to call Super implementation */
 	virtual void DoInitialSetup();
@@ -647,11 +639,9 @@ protected:
 	/** List of actors relevant to generation of navigation data */
 	TArray<TWeakObjectPtr<AActor> > GenerationSeeds;
 
-#if WITH_NAVIGATION_GENERATOR
 	/** stores areas marked as dirty throughout the frame, processes them 
 	 *	once a frame in Tick function */
 	TArray<FNavigationDirtyArea> DirtyAreas;
-#endif // WITH_NAVIGATION_GENERATOR	
 
 	// async queries
 	FCriticalSection NavDataRegistrationSection;
@@ -659,6 +649,7 @@ protected:
 	uint32 bNavigationBuildingLocked:1;
 	uint32 bInitialBuildingLockActive:1;
 	uint32 bInitialSetupHasBeenPerformed:1;
+	uint32 bAsyncBuildPaused:1;
 
 	/** cached navigable world bounding box*/
 	mutable FBox NavigableWorldBounds;
@@ -720,8 +711,6 @@ protected:
 
 	void SetCrowdManager(UCrowdManager* NewCrowdManager); 
 
-#if WITH_NAVIGATION_GENERATOR
-
 	/** Add BSP collision data to navigation octree */
 	void AddLevelCollisionToOctree(ULevel* Level);
 	
@@ -751,7 +740,6 @@ private:
 	 
 	/** Handler for FWorldDelegates::LevelRemovedFromWorld event */
 	void OnLevelRemovedFromWorld(ULevel* InLevel, UWorld* InWorld);
-#endif // WITH_NAVIGATION_GENERATOR
 
 	/** Adds given request to requests queue. Note it's to be called only on game thread only */
 	void AddAsyncQuery(const FAsyncPathFindingQuery& Query);
diff --git a/Engine/Source/Runtime/Engine/Classes/AI/Navigation/NavigationTestingActor.h b/Engine/Source/Runtime/Engine/Classes/AI/Navigation/NavigationTestingActor.h
--- a/Engine/Source/Runtime/Engine/Classes/AI/Navigation/NavigationTestingActor.h
+++ b/Engine/Source/Runtime/Engine/Classes/AI/Navigation/NavigationTestingActor.h
@@ -129,6 +129,7 @@ class ENGINE_API ANavigationTestingActor : public AActor, public INavAgentInterf
 #endif
 
 	FNavPathSharedPtr LastPath;
+	FNavigationPath::FPathObserverDelegate::FDelegate PathObserver;
 
 	/** Dtor */
 	virtual ~ANavigationTestingActor();
@@ -159,6 +160,11 @@ class ENGINE_API ANavigationTestingActor : public AActor, public INavAgentInterf
 	void GatherDetailedData(class ANavigationTestingActor* Goal);
 	void SearchPathTo(class ANavigationTestingActor* Goal);
 
+	/*	Called when given path becomes invalid (via @see PathObserverDelegate)
+	 *	NOTE: InvalidatedPath doesn't have to be instance's current Path
+	 */
+	void OnPathEvent(FNavigationPath* InvalidatedPath, ENavPathEvent::Type Event);
+
 	// Virtual method to override if you want to customize the query being 
 	// constructed for the path find (e.g. change the filter or add 
 	// constraints/goal evaluators).
diff --git a/Engine/Source/Runtime/Engine/Classes/AI/Navigation/NavigationTypes.h b/Engine/Source/Runtime/Engine/Classes/AI/Navigation/NavigationTypes.h
--- a/Engine/Source/Runtime/Engine/Classes/AI/Navigation/NavigationTypes.h
+++ b/Engine/Source/Runtime/Engine/Classes/AI/Navigation/NavigationTypes.h
@@ -11,9 +11,6 @@
 #define DEFAULT_NAV_QUERY_EXTENT_HORIZONTAL 50.f
 #define DEFAULT_NAV_QUERY_EXTENT_VERTICAL 100.f
 
-/** Whether to compile in navigation data generation - should be compiled at least when WITH_EDITOR is true */
-#define WITH_NAVIGATION_GENERATOR (WITH_RECAST || WITH_EDITOR)
-
 /** uniform identifier type for navigation data elements may it be a polygon or graph node */
 typedef uint64 NavNodeRef;
 
diff --git a/Engine/Source/Runtime/Engine/Classes/AI/Navigation/RecastNavMesh.h b/Engine/Source/Runtime/Engine/Classes/AI/Navigation/RecastNavMesh.h
--- a/Engine/Source/Runtime/Engine/Classes/AI/Navigation/RecastNavMesh.h
+++ b/Engine/Source/Runtime/Engine/Classes/AI/Navigation/RecastNavMesh.h
@@ -326,22 +326,26 @@ struct FNavMeshTileData
 		FNavData(uint8* InNavData) : RawNavData(InNavData) {}
 		~FNavData();
 	};
-
+	
 	// layer index
-	int32 LayerIndex;
+	int32	LayerIndex;
+	FBox	LayerBBox;
 	// size of allocated data
-	int32 DataSize;
+	int32	DataSize;
 	// actual tile data
-	TSharedPtr<FNavData, ESPMode::ThreadSafe> NavData;
+	TSharedPtr<FNavData> NavData;
+	
 
 	FNavMeshTileData() : LayerIndex(0), DataSize(0) { }
 
-	explicit FNavMeshTileData(uint8* RawData, int32 RawDataSize, int32 LayerIdx = 0)
-		: LayerIndex(LayerIdx), DataSize(RawDataSize)
+	explicit FNavMeshTileData(uint8* RawData, int32 RawDataSize, int32 LayerIdx = 0, FBox LayerBounds = FBox(0))
+		: LayerIndex(LayerIdx)
+		, LayerBBox(LayerBounds)
+		, DataSize(RawDataSize)
 	{
 		NavData = MakeShareable(new FNavData(RawData));
 	}
-
+	
 	FORCEINLINE uint8* GetData()
 	{
 		check(NavData.IsValid());
@@ -358,51 +362,14 @@ struct FNavMeshTileData
 	{
 		return NavData.IsValid() ? NavData->RawNavData : NULL;
 	}
-	
+
 	FORCEINLINE bool operator==(const uint8* RawData) const
 	{
 		return GetData() == RawData;
 	}
 
-	FORCEINLINE bool IsValid() const { return NavData.IsValid() && GetData() != NULL && DataSize > 0; }
-	FORCEINLINE void Release() { NavData.Reset(); DataSize = 0; LayerIndex = 0; }
-};
-
-struct FRecastTickHelper : FTickableGameObject
-{
-	class ARecastNavMesh* Owner;
-
-	FRecastTickHelper() : Owner(NULL) {}
-	virtual void Tick(float DeltaTime);
-	virtual bool IsTickable() const { return (Owner != NULL); }
-	virtual bool IsTickableInEditor() const { return true; }
-	virtual TStatId GetStatId() const ;
-};
-
-struct FTileSetItem
-{
-	/** tile coords on grid */
-	int16 X, Y;
-	/** 2D location of tile center */
-	FVector2D TileCenter;
-	/** squared 2D distance to closest generation seed (e.g. player) */
-	float Dist2DSq;
-	/** sort order */
-	int32 SortOrder;
-	/** does it have compressed geometry stored in generator? */
-	uint8 bHasCompressedGeometry : 1;
-	/** does it have generated navmesh? */
- 	uint8 bHasNavmesh : 1;
-
-	FTileSetItem(int32 InitialX = 0, int32 InitialY = 0, int32 InitialOrder = 0)
-		: X(InitialX), Y(InitialY), SortOrder(InitialOrder), bHasCompressedGeometry(false), bHasNavmesh(false)
-	{
-	}
-
-	FTileSetItem(int16 TileX, int16 TileY, const FBox& TileBounds)
-		: X(TileX), Y(TileY), TileCenter(TileBounds.GetCenter()), SortOrder(0), bHasCompressedGeometry(false), bHasNavmesh(false)
-	{
-	}
+	FORCEINLINE bool IsValid() const { return NavData.IsValid() && GetData() != nullptr && DataSize > 0; }
+	FORCEINLINE void Release() { if (NavData.IsValid()) { NavData->RawNavData = nullptr; } DataSize = 0; LayerIndex = 0; }
 };
 
 DECLARE_MULTICAST_DELEGATE(FOnNavMeshUpdate);
@@ -552,14 +519,6 @@ class ENGINE_API ARecastNavMesh : public ANavigationData
 	UPROPERTY(EditAnywhere, Category=Generation, config, AdvancedDisplay)
 	int32 LayerChunkSplits;
 
-	/** Limit for tile grid size: width */
-	UPROPERTY(config)
-	int32 MaxTileGridWidth;
-
-	/** Limit for tile grid size: height */
-	UPROPERTY(config)
-	int32 MaxTileGridHeight;
-
 	/** Controls whether Navigation Areas will be sorted by cost before application 
 	 *	to navmesh during navmesh generation. This is relevant then there are
 	 *	areas overlapping and we want to have area cost express area relevancy
@@ -590,9 +549,6 @@ private:
 	UPROPERTY(config)
 	uint32 bUseVoxelCache : 1;
 
-	/** internal data to store time for next sort for our set with tiles */
-	float NextTimeToSortTiles;
-
 	/** indicates how often we will sort navigation tiles to mach players position */
 	UPROPERTY(config)
 	float TileSetUpdateInterval;
@@ -668,17 +624,11 @@ public:
 #endif // WITH_EDITOR
 	// End UObject Interface
 
-	/** Handles storing and releasing raw navigation tile data */
-	void UpdateNavTileData(uint8* OldTileData, FNavMeshTileData NewTileData);
-
 #if WITH_EDITOR
 	/** RecastNavMesh instances are dynamically spawned and should not be coppied */
 	virtual bool ShouldExport() override { return false; }
 #endif
-
-	/** called by TickHelper, works in game AND in editor */
-	void TickMe(float DeltaTime);
-
+	
 	virtual void CleanUp() override;
 
 	// Begin ANavigationData Interface
@@ -699,14 +649,9 @@ public:
 	virtual FBox GetBounds() const override { return GetNavMeshBounds(); }
 	/** Called on world origin changes **/
 	virtual void ApplyWorldOffset(const FVector& InOffset, bool bWorldShift) override;
-
-protected:
-#if WITH_NAVIGATION_GENERATOR
-	virtual FNavDataGenerator* ConstructGenerator(const FNavAgentProperties& AgentProps) override;
-	virtual void SortNavigationTiles();
-#endif // WITH_NAVIGATION_GENERATOR
 	// End ANavigationData Interface
 
+protected:
 	/** Serialization helper. */
 	void SerializeRecastNavMesh(FArchive& Ar, FPImplRecastNavMesh*& NavMesh);
 
@@ -718,10 +663,10 @@ public:
 	FBox GetNavMeshTileBounds(int32 TileIndex) const;
 
 	/** Retrieves XY coordinates of tile specified by index */
-	void GetNavMeshTileXY(int32 TileIndex, int32& OutX, int32& OutY, int32& Layer) const;
+	bool GetNavMeshTileXY(int32 TileIndex, int32& OutX, int32& OutY, int32& Layer) const;
 
 	/** Retrieves XY coordinates of tile specified by position */
-	void GetNavMeshTileXY(const FVector& Point, int32& OutX, int32& OutY) const;
+	bool GetNavMeshTileXY(const FVector& Point, int32& OutX, int32& OutY) const;
 
 	/** Retrieves all tile indices at matching XY coordinates */
 	void GetNavMeshTilesAt(int32 TileX, int32 TileY, TArray<int32>& Indices) const;
@@ -739,6 +684,9 @@ public:
 	/** Creates a task to be executed on GameThread calling UpdateDrawing */
 	void RequestDrawingUpdate();
 
+	/** Invalidates active paths that go through changed tiles  */
+	void InvalidateAffectedPaths(const TArray<uint32>& ChangedTiles);
+
 	virtual void SetConfig(const FNavDataConfig& Src) override;
 protected:
 	virtual void FillConfig(FNavDataConfig& Dest) override;
@@ -757,8 +705,6 @@ public:
 	//----------------------------------------------------------------------//
 	/** Debug rendering support. */
 	void GetDebugGeometry(FRecastDebugGeometry& OutGeometry, int32 TileIndex = INDEX_NONE) const;
-	// @todo docuement
-	void GetDebugTileBounds(FBox& OuterBox, int32& NumTilesX, int32& NumTilesY) const;
 
 	// @todo docuement
 	void DrawDebugPathCorridor(NavNodeRef const* PathPolys, int32 NumPathPolys, bool bPersistent=true) const;
@@ -890,27 +836,13 @@ public:
 	static const FRecastQueryFilter* GetNamedFilter(ERecastNamedFilter::Type FilterType);
 	FORCEINLINE static FNavPolyFlags GetNavLinkFlag() { return NavLinkFlag; }
 	
-	/** initialize tiles set with given size */
-	void ReserveTileSet(int32 NewGridWidth, int32 NewGridHeight);
-
-	/** get information about given tile index */
-	FTileSetItem* GetTileSetItemAt(int32 TileX, int32 TileY);
-	
-	const FTileSetItem* GetTileSet() const { return TileSet.GetData(); }
-	FTileSetItem* GetTileSet() { return TileSet.GetData(); }
-	
 	virtual bool NeedsRebuild() const override;
-
-	/** update offset for navmesh in recast library, to current one */
-	/** @param Offset - current offset */
-	void UpdateNavmeshOffset( const FBox& NavBounds );
+	virtual void RebuildAll() override;
 
 protected:
 	// @todo docuement
 	UPrimitiveComponent* ConstructRenderingComponentImpl();
 
-	virtual void DestroyGenerator() override;
-
 	/** Spawns an ARecastNavMesh instance, and configures it if AgentProps != NULL */
 	static ARecastNavMesh* SpawnInstance(UNavigationSystem* NavSys, const FNavDataConfig* AgentProps = NULL);
 
@@ -929,15 +861,9 @@ protected:
 	FPImplRecastNavMesh* GetRecastNavMeshImpl() { return RecastNavMeshImpl; }
 	const FPImplRecastNavMesh* GetRecastNavMeshImpl() const { return RecastNavMeshImpl; }
 
-	TArray<FTileSetItem> TileSet;
-
 private:
-	int32 GridWidth;
-	int32 GridHeight;
-
 	/** NavMesh versioning. */
 	uint32 NavMeshVersion;
-	FRecastTickHelper TickHelper;
 	
 	/** 
 	 * This is a pimpl-style arrangement used to tightly hide the Recast internals from the rest of the engine.
@@ -946,20 +872,8 @@ private:
 	 *	@TODO since it's no secret we're using recast there's no point in having separate implementation class. FPImplRecastNavMesh should be merged into ARecastNavMesh
 	 */
 	FPImplRecastNavMesh* RecastNavMeshImpl;
-
-	/** array containing tile references of all navmesh tiles recently rebuild. Handled and cleared in Tick function */
-	TArray<uint32> FreshTiles;
-
-	/** Holds shared pointers to navmesh tile data currently used
-	 *	@TODO will no longer be needed if we switch over to support only async pathfinding */	
-	TNavStatArray<FNavMeshTileData> TileNavDataContainer;
-		
-#if RECAST_ASYNC_REBUILDING
-	mutable FCriticalSection* NavDataReadWriteLock;
 	
-	/** this critical section will be used befor navmesh generator creation */
-	FCriticalSection NavDataReadWriteLockDummy;
-
+#if RECAST_ASYNC_REBUILDING
 	/** batch query counter */
 	mutable int32 BatchQueryCounter;
 
@@ -978,4 +892,20 @@ bool ARecastNavMesh::NavMeshRaycast(const ANavigationData* Self, const FVector&
 	FRaycastResult Result;
 	return NavMeshRaycast(Self, RayStart, RayEnd, HitLocation, QueryFilter, Querier, Result);
 }
+
+
+/** structure to cache owning RecastNavMesh data so that it doesn't have to be polled
+ *	directly from RecastNavMesh while asyncronously generating navmesh */
+struct FRecastNavMeshCachedData
+{
+	ARecastNavMesh::FNavPolyFlags FlagsPerArea[RECAST_MAX_AREAS];
+	ARecastNavMesh::FNavPolyFlags FlagsPerOffMeshLinkArea[RECAST_MAX_AREAS];
+	TMap<const UClass*, int32> AreaClassToIdMap;
+	const ARecastNavMesh* ActorOwner;
+	uint32 bUseSortFunction : 1;
+
+	static FRecastNavMeshCachedData Construct(const ARecastNavMesh* RecastNavMeshActor);
+	void OnAreaAdded(const UClass* AreaClass, int32 AreaID);
+};
+
 #endif // WITH_RECAST
diff --git a/Engine/Source/Runtime/Engine/Private/AI/Navigation/NavCollision.cpp b/Engine/Source/Runtime/Engine/Private/AI/Navigation/NavCollision.cpp
--- a/Engine/Source/Runtime/Engine/Private/AI/Navigation/NavCollision.cpp
+++ b/Engine/Source/Runtime/Engine/Private/AI/Navigation/NavCollision.cpp
@@ -189,10 +189,8 @@ bool UNavCollision::GatherCollision()
 	// get data from owner
 	if (StaticMeshOuter && StaticMeshOuter->BodySetup)
 	{
-#if WITH_NAVIGATION_GENERATOR
 		NavigationHelper::GatherCollision(StaticMeshOuter->BodySetup, this);
 		bCalculated = true;
-#endif // WITH_NAVIGATION_GENERATOR
 	}
 
 	return bCalculated;
diff --git a/Engine/Source/Runtime/Engine/Private/AI/Navigation/NavGraphGenerator.cpp b/Engine/Source/Runtime/Engine/Private/AI/Navigation/NavGraphGenerator.cpp
--- a/Engine/Source/Runtime/Engine/Private/AI/Navigation/NavGraphGenerator.cpp
+++ b/Engine/Source/Runtime/Engine/Private/AI/Navigation/NavGraphGenerator.cpp
@@ -19,41 +19,11 @@ FNavGraphGenerator::~FNavGraphGenerator()
 //----------------------------------------------------------------------//
 // FNavDataGenerator overrides
 //----------------------------------------------------------------------//
-bool FNavGraphGenerator::Generate() 
-{
-	return true;
-}
-
 bool FNavGraphGenerator::IsBuildInProgress(bool bCheckDirtyToo) const
 {
 	return false;
 }
 
-void FNavGraphGenerator::OnWorldInitDone(bool bAllowedToRebuild)
-{
-
-}
-
-void FNavGraphGenerator::OnNavigationBuildingLocked() 
-{
-
-}
-
-void FNavGraphGenerator::OnNavigationBuildingUnlocked(bool bForce) 
-{
-
-}
-
-void FNavGraphGenerator::OnNavigationBoundsUpdated(class AVolume* Volume) 
-{
-
-}
-
-void FNavGraphGenerator::OnNavigationDataDestroyed(class ANavigationData* NavData) 
-{
-
-}
-
 //----------------------------------------------------------------------//
 // 
 //----------------------------------------------------------------------//
@@ -62,11 +32,6 @@ void FNavGraphGenerator::Init()
 
 }
 
-bool FNavGraphGenerator::IsBuildingLocked() const
-{
-	return false;
-}
-
 void FNavGraphGenerator::CleanUpIntermediateData()
 {
 
@@ -77,7 +42,3 @@ void FNavGraphGenerator::UpdateBuilding()
 
 }
 
-void FNavGraphGenerator::TriggerGeneration()
-{
-
-}
diff --git a/Engine/Source/Runtime/Engine/Private/AI/Navigation/NavGraphGenerator.h b/Engine/Source/Runtime/Engine/Private/AI/Navigation/NavGraphGenerator.h
--- a/Engine/Source/Runtime/Engine/Private/AI/Navigation/NavGraphGenerator.h
+++ b/Engine/Source/Runtime/Engine/Private/AI/Navigation/NavGraphGenerator.h
@@ -22,33 +22,13 @@ public:
 	//----------------------------------------------------------------------//
 	// FNavDataGenerator overrides
 	//----------------------------------------------------------------------//
-	// Triggers navmesh building process
-	virtual bool Generate() override;
-
 	virtual bool IsBuildInProgress(bool bCheckDirtyToo = false) const override;
 
-	virtual void OnWorldInitDone(bool bAllowedToRebuild) override;
-		
-	//--- accessors --- //
-
 private:
 	// Performs initial setup of member variables so that generator is ready to do its thing from this point on
 	void Init();
-
-	bool IsBuildingLocked() const;
-
 	void CleanUpIntermediateData();
-
 	void UpdateBuilding();
-	
-	void TriggerGeneration();
-
-public:
-	virtual void OnNavigationBuildingLocked() override;
-	virtual void OnNavigationBuildingUnlocked(bool bForce) override;
-	virtual void OnNavigationBoundsUpdated(class AVolume* Volume) override;
-
-	virtual void OnNavigationDataDestroyed(class ANavigationData* NavData) override;
 
 private:
 	/** Bounding geometry definition. */
@@ -59,7 +39,4 @@ private:
 	class ANavigationGraph* DestNavGraph;
 
 	uint32 bInitialized:1;
-	uint32 bBuildFromScratchRequested:1;
-
-	mutable uint32 bBuildingLocked:1;
 };
diff --git a/Engine/Source/Runtime/Engine/Private/AI/Navigation/NavModifierVolume.cpp b/Engine/Source/Runtime/Engine/Private/AI/Navigation/NavModifierVolume.cpp
--- a/Engine/Source/Runtime/Engine/Private/AI/Navigation/NavModifierVolume.cpp
+++ b/Engine/Source/Runtime/Engine/Private/AI/Navigation/NavModifierVolume.cpp
@@ -28,6 +28,13 @@ FBox ANavModifierVolume::GetNavigationBounds() const
 
 #if WITH_EDITOR
 
+void ANavModifierVolume::PostEditUndo()
+{
+	Super::PostEditUndo();
+
+	UNavigationSystem::UpdateNavOctree(this);
+}
+
 void ANavModifierVolume::PostEditChangeProperty(FPropertyChangedEvent& PropertyChangedEvent)
 {
 	Super::PostEditChangeProperty(PropertyChangedEvent);
diff --git a/Engine/Source/Runtime/Engine/Private/AI/Navigation/NavigationData.cpp b/Engine/Source/Runtime/Engine/Private/AI/Navigation/NavigationData.cpp
--- a/Engine/Source/Runtime/Engine/Private/AI/Navigation/NavigationData.cpp
+++ b/Engine/Source/Runtime/Engine/Private/AI/Navigation/NavigationData.cpp
@@ -90,13 +90,6 @@ FSupportedAreaData::FSupportedAreaData(TSubclassOf<UNavArea> NavAreaClass, int32
 	}
 }
 
-//----------------------------------------------------------------------//
-// FNavDataGenerator
-//----------------------------------------------------------------------//
-void FNavDataGenerator::TriggerGeneration()
-{
-}
-
 //----------------------------------------------------------------------//
 // ANavigationData                                                                
 //----------------------------------------------------------------------//
@@ -160,18 +153,6 @@ void ANavigationData::PostInitializeComponents()
 		CleanUpAndMarkPendingKill();
 		return;
 	}
-
-	if (IsPendingKill() == true)
-	{
-		return;
-	}
-
-#if WITH_NAVIGATION_GENERATOR
-	if (bRebuildAtRuntime == true && NavDataGenerator.IsValid() == false)
-	{
-		GetGenerator(FNavigationSystem::Create);
-	}
-#endif
 }
 
 void ANavigationData::PostLoad() 
@@ -300,17 +281,6 @@ void ANavigationData::InstantiateAndRegisterRenderingComponent()
 #endif // !UE_BUILD_SHIPPING
 }
 
-void ANavigationData::DestroyGenerator()
-{
-#if WITH_NAVIGATION_GENERATOR
-	if (NavDataGenerator.IsValid() == true)
-	{
-		NavDataGenerator->OnNavigationDataDestroyed(this);
-		NavDataGenerator.Reset();
-	}
-#endif // WITH_NAVIGATION_GENERATOR
-};
-
 void ANavigationData::PurgeUnusedPaths()
 {
 	check(IsInGameThread());
@@ -330,9 +300,6 @@ void ANavigationData::PurgeUnusedPaths()
 void ANavigationData::PostEditUndo()
 {
 	Super::PostEditUndo();
-#if WITH_NAVIGATION_GENERATOR
-	GetGenerator(FNavigationSystem::Create);
-#endif // WITH_NAVIGATION_GENERATOR
 
 	UWorld* WorldOuter = GetWorld();
 	if (WorldOuter != NULL && WorldOuter->GetNavigationSystem() != NULL)
@@ -367,7 +334,6 @@ void ANavigationData::Destroyed()
 void ANavigationData::CleanUp()
 {
 	bRegistered = false;
-	DestroyGenerator();
 }
 
 void ANavigationData::CleanUpAndMarkPendingKill()
@@ -382,56 +348,30 @@ void ANavigationData::CleanUpAndMarkPendingKill()
 
 bool ANavigationData::CanRebuild() const
 {
-#if WITH_NAVIGATION_GENERATOR
-	return NavDataGenerator.IsValid();
-#endif
-
-	return false;
+	return NavDataGenerator.Get() != nullptr;
 }
 
-#if WITH_NAVIGATION_GENERATOR
-
-FNavDataGenerator* ANavigationData::GetGenerator(FNavigationSystem::ECreateIfEmpty CreateIfNone)
+void ANavigationData::RebuildAll()
 {
-	if ((NavDataGenerator.IsValid() == false || NavDataGenerator.GetSharedReferenceCount() <= 0)
-		&& CreateIfNone == FNavigationSystem::Create
-		&& (bRebuildAtRuntime == true || GetWorld()->IsGameWorld() == false))
-	{
-		// @todo this is not-that-clear design and is subject to change
-		if (NavDataConfig.IsValid() == false 
-			&& GetWorld()->GetNavigationSystem() != NULL
-			&& GetWorld()->GetNavigationSystem()->SupportedAgents.Num() < 2
-			)
-		{
-			// fill in AgentProps with whatever is the instance's setup
-			FillConfig(NavDataConfig);
-		}
 
-		FNavDataGenerator* Generator = ConstructGenerator(NavDataConfig);
+}
 
-		if (Generator != NULL)
-		{
-			// make sure no one has made this sharable yet
-			// it's a "check" not an "ensure" since it's really crucial
-			check(Generator->HasBeenAlreadyMadeSharable() == false);
-			NavDataGenerator = MakeShareable(Generator);
-		}
+void ANavigationData::TickAsyncBuild(float DeltaSeconds)
+{
+	if (NavDataGenerator)
+	{
+		NavDataGenerator->TickAsyncBuild(DeltaSeconds);
 	}
-
-	return NavDataGenerator.IsValid() && NavDataGenerator.GetSharedReferenceCount() > 0 ? NavDataGenerator.Get() : NULL;
 }
 
 void ANavigationData::RebuildDirtyAreas(const TArray<FNavigationDirtyArea>& DirtyAreas)
 {
-	FNavDataGenerator* Generator = GetGenerator(FNavigationSystem::DontCreate);
-	if (Generator)
+	if (NavDataGenerator)
 	{
-		Generator->RebuildDirtyAreas(DirtyAreas);
+		NavDataGenerator->RebuildDirtyAreas(DirtyAreas);
 	}
 }
 
-#endif // WITH_NAVIGATION_GENERATOR
-
 void ANavigationData::DrawDebugPath(FNavigationPath* Path, FColor PathColor, UCanvas* Canvas, bool bPersistent, const uint32 NextPathPointIndex) const
 {
 	Path->DebugDraw(this, PathColor, Canvas, bPersistent, NextPathPointIndex);
@@ -597,12 +537,10 @@ uint32 ANavigationData::LogMemUsed() const
 
 	UE_LOG(LogNavigation, Display, TEXT("%s: ANavigationData: %u\n    self: %d"), *GetName(), MemUsed, sizeof(ANavigationData));	
 
-#if WITH_NAVIGATION_GENERATOR
-	if (NavDataGenerator.IsValid())
+	if (NavDataGenerator)
 	{
 		NavDataGenerator->LogMemUsed();
 	}
-#endif // WITH_NAVIGATION_GENERATOR
 
 	return MemUsed;
 }
diff --git a/Engine/Source/Runtime/Engine/Private/AI/Navigation/NavigationGraph.cpp b/Engine/Source/Runtime/Engine/Private/AI/Navigation/NavigationGraph.cpp
--- a/Engine/Source/Runtime/Engine/Private/AI/Navigation/NavigationGraph.cpp
+++ b/Engine/Source/Runtime/Engine/Private/AI/Navigation/NavigationGraph.cpp
@@ -56,6 +56,10 @@ ANavigationGraphNode::ANavigationGraphNode(const class FPostConstructInitializeP
 ANavigationGraph::ANavigationGraph(const class FPostConstructInitializeProperties& PCIP)
 	: Super(PCIP)
 {
+	if (HasAnyFlags(RF_ClassDefaultObject) == false)
+	{
+		NavDataGenerator.Reset(new FNavGraphGenerator(this));
+	}
 }
 
 ANavigationData* ANavigationGraph::CreateNavigationInstances(UNavigationSystem* NavSys)
@@ -84,10 +88,3 @@ ANavigationData* ANavigationGraph::CreateNavigationInstances(UNavigationSystem*
 
 	return NULL;
 }
-
-#if WITH_NAVIGATION_GENERATOR
-FNavDataGenerator* ANavigationGraph::ConstructGenerator(const FNavAgentProperties& AgentProps)
-{
-	return new FNavGraphGenerator(this);
-}
-#endif
diff --git a/Engine/Source/Runtime/Engine/Private/AI/Navigation/NavigationOctree.cpp b/Engine/Source/Runtime/Engine/Private/AI/Navigation/NavigationOctree.cpp
--- a/Engine/Source/Runtime/Engine/Private/AI/Navigation/NavigationOctree.cpp
+++ b/Engine/Source/Runtime/Engine/Private/AI/Navigation/NavigationOctree.cpp
@@ -4,10 +4,7 @@
 #include "NavigationOctree.h"
 #include "RecastHelpers.h"
 #include "AI/Navigation/NavRelevantComponent.h"
-
-#if NAVOCTREE_CONTAINS_COLLISION_DATA
 #include "RecastNavMeshGenerator.h"
-#endif // NAVOCTREE_CONTAINS_COLLISION_DATA
 
 //----------------------------------------------------------------------//
 // FNavigationOctree
@@ -36,7 +33,6 @@ void FNavigationOctree::AddNode(UObject* ElementOb, INavRelevantInterface* NavEl
 	Data.Owner = ElementOb;
 	Data.Bounds = Bounds;
 
-#if NAVOCTREE_CONTAINS_COLLISION_DATA
 	UActorComponent* ActorComp = Cast<UActorComponent>(ElementOb);
 	if (bGatherGeometry && ActorComp)
 	{
@@ -56,7 +52,6 @@ void FNavigationOctree::AddNode(UObject* ElementOb, INavRelevantInterface* NavEl
 	const int32 ElementMemory = Data.GetAllocatedSize();
 	NodesMemory += ElementMemory;
 	INC_MEMORY_STAT_BY(STAT_Navigation_CollisionTreeMemory, ElementMemory);
-#endif
 
 	AddElement(Data);
 }
@@ -68,7 +63,6 @@ void FNavigationOctree::AppendToNode(const FOctreeElementId& Id, INavRelevantInt
 	Data = OrgData;
 	Data.Bounds = Bounds + OrgData.Bounds.GetBox();
 
-#if NAVOCTREE_CONTAINS_COLLISION_DATA
 	if (NavElement)
 	{
 		SCOPE_CYCLE_COUNTER(STAT_Navigation_GatheringNavigationModifiersSync);
@@ -85,7 +79,6 @@ void FNavigationOctree::AppendToNode(const FOctreeElementId& Id, INavRelevantInt
 
 	NodesMemory += MemoryDelta;
 	INC_MEMORY_STAT_BY(STAT_Navigation_CollisionTreeMemory, MemoryDelta);
-#endif
 
 	RemoveElement(Id);
 	AddElement(Data);
@@ -93,12 +86,10 @@ void FNavigationOctree::AppendToNode(const FOctreeElementId& Id, INavRelevantInt
 
 void FNavigationOctree::RemoveNode(const FOctreeElementId& Id)
 {
-#if NAVOCTREE_CONTAINS_COLLISION_DATA
 	FNavigationOctreeElement& Data = GetElementById(Id);
 	const int32 ElementMemory = Data.GetAllocatedSize();
 	NodesMemory -= ElementMemory;
 	DEC_MEMORY_STAT_BY(STAT_Navigation_CollisionTreeMemory, ElementMemory);
-#endif
 
 	RemoveElement(Id);
 }
diff --git a/Engine/Source/Runtime/Engine/Private/AI/Navigation/NavigationSystem.cpp b/Engine/Source/Runtime/Engine/Private/AI/Navigation/NavigationSystem.cpp
--- a/Engine/Source/Runtime/Engine/Private/AI/Navigation/NavigationSystem.cpp
+++ b/Engine/Source/Runtime/Engine/Private/AI/Navigation/NavigationSystem.cpp
@@ -39,7 +39,8 @@ DEFINE_LOG_CATEGORY_STATIC(LogNavOctree, Warning, All);
 
 DECLARE_CYCLE_STAT(TEXT("Rasterize triangles"), STAT_Navigation_RasterizeTriangles,STATGROUP_Navigation);
 DECLARE_CYCLE_STAT(TEXT("Nav Tick: area register"), STAT_Navigation_TickNavAreaRegister, STATGROUP_Navigation);
-DECLARE_CYCLE_STAT(TEXT("Nav Tick: rebuild"), STAT_Navigation_TickRebuild, STATGROUP_Navigation);
+DECLARE_CYCLE_STAT(TEXT("Nav Tick: mark dirty"), STAT_Navigation_TickMarkDirty, STATGROUP_Navigation);
+DECLARE_CYCLE_STAT(TEXT("Nav Tick: async build"), STAT_Navigation_TickAsyncBuild, STATGROUP_Navigation);
 DECLARE_CYCLE_STAT(TEXT("Nav Tick: async pathfinding"), STAT_Navigation_TickAsyncPathfinding, STATGROUP_Navigation);
 DECLARE_CYCLE_STAT(TEXT("Debug NavOctree Time"), STAT_DebugNavOctree, STATGROUP_Navigation);
 
@@ -205,13 +206,11 @@ UNavigationSystem::UNavigationSystem(const FPostConstructInitializeProperties& P
 		NavDataRegistrationQueue.Reserve( REGISTRATION_QUEUE_SIZE );
 	
 		NavOctree = new FNavigationOctree(FVector(0,0,0), 64000);
-#if WITH_NAVIGATION_GENERATOR
-	#if WITH_RECAST
+#if WITH_RECAST
 		NavOctree->ComponentExportDelegate = FNavigationOctree::FNavigableGeometryComponentExportDelegate::CreateStatic(&FRecastNavMeshGenerator::ExportComponentGeometry);
-	#endif // WITH_RECAST
+#endif // WITH_RECAST
 		FWorldDelegates::LevelAddedToWorld.AddUObject(this, &UNavigationSystem::OnLevelAddedToWorld);
 		FWorldDelegates::LevelRemovedFromWorld.AddUObject(this, &UNavigationSystem::OnLevelRemovedFromWorld);
-#endif // WITH_NAVIGATION_GENERATOR	
 	}
 	else
 	{
@@ -292,14 +291,13 @@ void UNavigationSystem::PostInitProperties()
 	
 		bInitialBuildingLockActive = bInitialBuildingLocked;
 
-#if WITH_NAVIGATION_GENERATOR
 		UWorld* World = GetWorld();
 		if (World)
 		{
 			if (World->PersistentLevel)
-		{
-			OnLevelAddedToWorld(World->PersistentLevel, World);
-		}
+			{
+				OnLevelAddedToWorld(World->PersistentLevel, World);
+			}
 
 			for (int32 Idx = 0; Idx < World->StreamingLevels.Num(); Idx++)
 			{
@@ -310,7 +308,6 @@ void UNavigationSystem::PostInitProperties()
 				}
 			}
 		}
-#endif
 
 		// register for any actor move change
 #if WITH_EDITOR
@@ -320,9 +317,7 @@ void UNavigationSystem::PostInitProperties()
 		}
 #endif
 		FCoreUObjectDelegates::PostLoadMap.AddUObject(this, &UNavigationSystem::OnPostLoadMap);
-#if WITH_NAVIGATION_GENERATOR
 		UNavigationSystem::NavigationDirtyEvent.AddUObject(this, &UNavigationSystem::OnNavigationDirtied);
-#endif // WITH_NAVIGATION_GENERATOR
 	}
 }
 
@@ -362,7 +357,7 @@ void UNavigationSystem::OnWorldInitDone(FNavigationSystem::EMode Mode)
 	else
 	{
 		PopulateNavOctree();
-
+		
 		// gather all navigation data instances and register all not-yet-registered
 		// (since it's quite possible navigation system was not ready by the time
 		// those instances were serialized-in or spawned)
@@ -387,7 +382,6 @@ void UNavigationSystem::OnWorldInitDone(FNavigationSystem::EMode Mode)
 			bInitialBuildingLockActive = false;
 		}
 		
-#if WITH_NAVIGATION_GENERATOR
 		if (bAutoCreateNavigationData == true)
 		{
 			SpawnMissingNavigationData();
@@ -399,10 +393,9 @@ void UNavigationSystem::OnWorldInitDone(FNavigationSystem::EMode Mode)
 				ANavigationData* NavData = NavDataSet[NavDataIndex];
 				if (NavData != NULL)
 				{
-					FNavDataGenerator* Generator = NavData->GetGenerator(FNavigationSystem::Create);
-					if (Generator != NULL)
+					if (bInitialBuildingLockActive == false && bNavigationAutoUpdateEnabled)
 					{
-						NavData->GetGenerator(FNavigationSystem::DontCreate)->OnWorldInitDone(bInitialBuildingLockActive == false && bNavigationAutoUpdateEnabled);
+						NavData->RebuildAll();
 					}
 				}
 			}
@@ -422,10 +415,12 @@ void UNavigationSystem::OnWorldInitDone(FNavigationSystem::EMode Mode)
 						if (Result == RegistrationSuccessful)
 						{
 #if WITH_RECAST
-							if (NavData->GetGenerator(FNavigationSystem::Create) != NULL
-								&& Cast<ARecastNavMesh>(NavData) != NULL)
+							if (Cast<ARecastNavMesh>(NavData) != NULL)
 							{
-								NavData->GetGenerator(FNavigationSystem::DontCreate)->OnWorldInitDone(bInitialBuildingLockActive == false && bNavigationAutoUpdateEnabled);
+								if (bInitialBuildingLockActive == false && bNavigationAutoUpdateEnabled)
+								{
+									NavData->RebuildAll();
+								}
 							}
 #endif // WITH_RECAST
 						}
@@ -439,7 +434,6 @@ void UNavigationSystem::OnWorldInitDone(FNavigationSystem::EMode Mode)
 				}
 			}
 		}
-#endif
 	}
 }
 
@@ -475,19 +469,14 @@ void UNavigationSystem::Tick(float DeltaSeconds)
 		ProcessNavAreaPendingRegistration();
 	}
 
-#if WITH_NAVIGATION_GENERATOR
 	if (PendingNavVolumeUpdates.Num() > 0)
 	{
 		for (auto Volume : PendingNavVolumeUpdates)
 		{
-			if (Volume != NULL && Volume->IsPendingKillPending() == false)
-			{
-				PerformNavigationBoundsUpdate(Volume);
-			}
+			PerformNavigationBoundsUpdate(Volume);
 		}
 		PendingNavVolumeUpdates.Reset();
 	}
-#endif // WITH_NAVIGATION_GENERATOR
 
 	if (PendingOctreeUpdates.Num() > 0)
 	{
@@ -506,9 +495,8 @@ void UNavigationSystem::Tick(float DeltaSeconds)
 		INC_FLOAT_STAT_BY(STAT_Navigation_CumulativeBuildTime,(float)ThisTime*1000);
 	}
 
-#if WITH_NAVIGATION_GENERATOR
 	{
-		SCOPE_CYCLE_COUNTER(STAT_Navigation_TickRebuild);
+		SCOPE_CYCLE_COUNTER(STAT_Navigation_TickMarkDirty);
 
 		DirtyAreasUpdateTime += DeltaSeconds;
 		const float DirtyAreasUpdateDeltaTime = 1.0f / DirtyAreasUpdateFreq;
@@ -529,7 +517,19 @@ void UNavigationSystem::Tick(float DeltaSeconds)
 			DirtyAreas.Reset();
 		}
 	}
-#endif // WITH_NAVIGATION_GENERATOR
+
+	// Tick navigation mesh async builders
+	if (!bAsyncBuildPaused)
+	{
+		SCOPE_CYCLE_COUNTER(STAT_Navigation_TickAsyncBuild);
+		for (ANavigationData* NavData : NavDataSet)
+		{
+			if (NavData)
+			{
+				NavData->TickAsyncBuild(DeltaSeconds);
+			}
+		}
+	}
 
 	if (AsyncPathFindingQueries.Num() > 0)
 	{
@@ -572,12 +572,11 @@ void UNavigationSystem::SetNavigationAutoUpdateEnabled(bool bNewEnable,UNavigati
 	if(bNewEnable != bNavigationAutoUpdateEnabled)
 	{
 		bNavigationAutoUpdateEnabled = bNewEnable; 
-#if WITH_NAVIGATION_GENERATOR
+
 		if (InNavigationsystem)
 		{
 			InNavigationsystem->EnableAllGenerators(bNewEnable, /*bForce=*/true);
 		}
-#endif // WITH_NAVIGATION_GENERATOR
 	}
 }
 #endif // WITH_EDITOR
@@ -1151,8 +1150,7 @@ bool UNavigationSystem::IsNavigationBuilt(const AWorldSettings* Settings) const
 		ANavigationData* NavData = NavDataSet[NavDataIndex];
 		if (NavData != NULL && NavData->GetWorldSettings() == Settings)
 		{
-#if WITH_NAVIGATION_GENERATOR
-			FNavDataGenerator* Generator = NavData->GetGenerator(FNavigationSystem::DontCreate);
+			FNavDataGenerator* Generator = NavData->GetGenerator();
 			if ((NavData->bRebuildAtRuntime == true 
 #if WITH_EDITOR
 				|| GEditor != NULL
@@ -1162,7 +1160,6 @@ bool UNavigationSystem::IsNavigationBuilt(const AWorldSettings* Settings) const
 				bIsBuilt = false;
 				break;
 			}
-#endif
 		}
 	}
 
@@ -1184,7 +1181,7 @@ bool UNavigationSystem::IsThereAnywhereToBuildNavigation() const
 	for (TActorIterator<ANavMeshBoundsVolume> It(GetWorld()); It; ++It)
 	{
 		ANavMeshBoundsVolume const* const V = (*It);
-		if (V != NULL)
+		if (V != NULL && !V->IsPendingKill())
 		{
 			bCreateNavigation = true;
 			break;
@@ -1800,7 +1797,6 @@ ANavigationData* UNavigationSystem::GetNavDataWithID(const uint16 NavDataID) con
 	return NULL;
 }
 
-#if WITH_NAVIGATION_GENERATOR
 void UNavigationSystem::AddDirtyArea(const FBox& NewArea, int32 Flags)
 {
 	if (Flags > 0)
@@ -1816,7 +1812,6 @@ void UNavigationSystem::AddDirtyAreas(const TArray<FBox>& NewAreas, int32 Flags)
 		AddDirtyArea(NewAreas[NewAreaIndex], Flags);
 	}
 }
-#endif // WITH_NAVIGATION_GENERATOR
 
 int32 GetDirtyFlagHelper(int32 UpdateFlags, int32 DefaultValue)
 {
@@ -1887,9 +1882,7 @@ void UNavigationSystem::AddElementToNavOctree(const FNavigationDirtyElement& Dir
 	{
 		if (DirtyElement.bHasPrevData)
 		{
-#if WITH_NAVIGATION_GENERATOR
 			AddDirtyArea(DirtyElement.PrevBounds, DirtyElement.PrevFlags);
-#endif // WITH_NAVIGATION_GENERATOR
 		}
 		
 		return;
@@ -1936,7 +1929,6 @@ void UNavigationSystem::AddElementToNavOctree(const FNavigationDirtyElement& Dir
 		NavOctree->AddNode(ElementOwner, DirtyElement.NavInterface, ElementBounds, GeneratedData);
 	}
 
-#if WITH_NAVIGATION_GENERATOR && NAVOCTREE_CONTAINS_COLLISION_DATA
 	const FBox BBox = GeneratedData.Bounds.GetBox();
 	const bool bValidBBox = BBox.IsValid && !BBox.GetSize().IsNearlyZero();
 	if (bValidBBox && !GeneratedData.IsEmpty())
@@ -1944,7 +1936,6 @@ void UNavigationSystem::AddElementToNavOctree(const FNavigationDirtyElement& Dir
 		const int32 DirtyFlag = DirtyElement.FlagsOverride ? DirtyElement.FlagsOverride : GeneratedData.Data.GetDirtyFlag();
 		AddDirtyArea(BBox, DirtyFlag);
 	}
-#endif // WITH_NAVIGATION_GENERATOR
 }
 
 bool UNavigationSystem::GetNavOctreeElementData(UObject* NodeOwner, int32& DirtyFlags, FBox& DirtyBounds)
@@ -1954,13 +1945,10 @@ bool UNavigationSystem::GetNavOctreeElementData(UObject* NodeOwner, int32& Dirty
 	{
 		if (NavOctree->IsValidElementId(*ElementId))
 		{
-#if WITH_NAVIGATION_GENERATOR && NAVOCTREE_CONTAINS_COLLISION_DATA
 			// mark area occupied by given actor as dirty
 			FNavigationOctreeElement& ElementData = NavOctree->GetElementById(*ElementId);
 			DirtyFlags = ElementData.Data.GetDirtyFlag();
 			DirtyBounds = ElementData.Bounds.GetBox();
-#endif // WITH_NAVIGATION_GENERATOR
-
 			return true;
 		}
 	}
@@ -2021,12 +2009,10 @@ void UNavigationSystem::RemoveNavOctreeElementId(const FOctreeElementId& Element
 {
 	if (NavOctree->IsValidElementId(ElementId))
 	{
-#if WITH_NAVIGATION_GENERATOR && NAVOCTREE_CONTAINS_COLLISION_DATA
 		// mark area occupied by given actor as dirty
 		FNavigationOctreeElement& ElementData = NavOctree->GetElementById(ElementId);
 		const int32 DirtyFlag = GetDirtyFlagHelper(UpdateFlags, ElementData.Data.GetDirtyFlag());
 		AddDirtyArea(ElementData.Bounds.GetBox(), DirtyFlag);
-#endif // WITH_NAVIGATION_GENERATOR
 		NavOctree->RemoveNode(ElementId);
 	}
 }
@@ -2292,7 +2278,6 @@ void UNavigationSystem::ReleaseInitialBuildingLock()
 		bInitialBuildingLockActive = false;
 		if (bNavigationBuildingLocked == false)
 		{
-#if WITH_NAVIGATION_GENERATOR
 			// apply pending changes
 			{
 				SCOPE_CYCLE_COUNTER(STAT_Navigation_AddingActorsToNavOctree);
@@ -2316,7 +2301,6 @@ void UNavigationSystem::ReleaseInitialBuildingLock()
 			// if navigation building is not blocked for other reasons then rebuild
 			// bForce == true to skip bNavigationBuildingLocked test
 			NavigationBuildingUnlock(/*bForce = */true);
-#endif // WITH_NAVIGATION_GENERATOR
 		}
 	}
 }
@@ -2355,7 +2339,6 @@ void UNavigationSystem::OnEditorModeChanged(FEdMode* Mode, bool IsEntering)
 }
 #endif
 
-#if WITH_NAVIGATION_GENERATOR
 void UNavigationSystem::OnNavigationBoundsUpdated(ANavMeshBoundsVolume* NavVolume)
 {
 	if (NavVolume == NULL)
@@ -2368,8 +2351,6 @@ void UNavigationSystem::OnNavigationBoundsUpdated(ANavMeshBoundsVolume* NavVolum
 
 void UNavigationSystem::PerformNavigationBoundsUpdate(ANavMeshBoundsVolume* NavVolume)
 {
-	const bool bIsInGame = GIsEditor == false || NavVolume->GetWorld()->IsPlayInEditor();
-
 	if (bNavDataRemovedDueToMissingNavBounds)
 	{
 		PopulateNavOctree();
@@ -2387,15 +2368,6 @@ void UNavigationSystem::PerformNavigationBoundsUpdate(ANavMeshBoundsVolume* NavV
 		{
 			SpawnMissingNavigationData();
 			ProcessRegistrationCandidates();
-
-			for (int32 NavDataIndex = 0; NavDataIndex < NavDataSet.Num(); ++NavDataIndex)
-			{
-				ANavigationData* NavData = NavDataSet[NavDataIndex];
-				if (NavData != NULL && (NavData->bRebuildAtRuntime || (GIsEditor && !bIsInGame)))
-				{
-					NavData->GetGenerator(FNavigationSystem::Create);
-				}
-			}
 		}
 	}
 
@@ -2404,22 +2376,7 @@ void UNavigationSystem::PerformNavigationBoundsUpdate(ANavMeshBoundsVolume* NavV
 	{
 		// rebuild all navmeshes
 		// NOTE: this will most probably be done in editor, if required to be performed at runtime may require some optimizations
-		for (int32 NavDataIndex = 0; NavDataIndex < NavDataSet.Num(); ++NavDataIndex)
-		{
-			ANavigationData* NavData = NavDataSet[NavDataIndex];
-			check(NavData);
-			if (NavData->GetGenerator() != NULL)
-			{
-				if ( Cast<ARecastNavMesh>(NavData) != NULL && (bIsInGame == false || ((ARecastNavMesh*)NavData)->bRebuildAtRuntime) )
-				{
-					// there are two possible cases here:
-					// 1. new volume is totally inside original generation bounds
-					// 2. it expands current bounds.
-					// @TODO currently handling every case as 1), 2) is an optimization
-					NavData->GetGenerator(FNavigationSystem::Create)->OnNavigationBoundsUpdated(NavVolume);
-				}
-			}
-		}
+		RebuildAll();
 	}
 #endif // WITH_RECAST
 }
@@ -2559,36 +2516,15 @@ ANavigationData* UNavigationSystem::CreateNavigationDataInstance(TSubclassOf<ANa
 	return Instance;
 }
 
-void UNavigationSystem::OnUpdateStreamingStarted()
-{
-	NavigationBuildingLock();
-}
-
-void UNavigationSystem::OnUpdateStreamingFinished()
-{
-	NavigationBuildingUnlock();
-
-	if (bWholeWorldNavigable == true)
-	{
-		const FBox NavBounds = NavigableWorldBounds;
-		const FBox NewNavBounds = GetWorldBounds();
-		if (FVector::DistSquared(NavBounds.GetCenter(), NewNavBounds.GetCenter()) > KINDA_SMALL_NUMBER
-			|| FVector::DistSquared(NavBounds.GetExtent(), NewNavBounds.GetExtent()) > KINDA_SMALL_NUMBER)
-		{
-			RebuildAll();
-		}
-	}
-}
-
 void UNavigationSystem::OnPIEStart()
 {
-	EnableAllGenerators(false);
+	// Do not tick async build for editor world while PIE is active
+	bAsyncBuildPaused = true;
 }
 
 void UNavigationSystem::OnPIEEnd()
 {
-	// it's set to true, but in editor it depends on user settings
-	EnableAllGenerators(bNavigationAutoUpdateEnabled);
+	bAsyncBuildPaused = false;
 }
 
 void UNavigationSystem::EnableAllGenerators(bool bEnable, bool bForce)
@@ -2613,16 +2549,6 @@ void UNavigationSystem::NavigationBuildingLock()
 	GetMainNavData(bBuildNavigationAtRuntime && IsThereAnywhereToBuildNavigation() ? FNavigationSystem::Create : FNavigationSystem::DontCreate);
 
 	bNavigationBuildingLocked = true;
-
-	// lock all generators
-	for (int32 NavDataIndex = 0; NavDataIndex < NavDataSet.Num(); ++NavDataIndex)
-	{
-		ANavigationData* NavData = NavDataSet[NavDataIndex];
-		if (NavData != NULL && NavData->GetGenerator() != NULL)
-		{
-			NavData->GetGenerator(FNavigationSystem::DontCreate)->OnNavigationBuildingLocked();
-		}
-	}
 }
 
 void UNavigationSystem::NavigationBuildingUnlock(bool bForce)
@@ -2631,16 +2557,8 @@ void UNavigationSystem::NavigationBuildingUnlock(bool bForce)
 	{
 		bNavigationBuildingLocked = false;
 		bInitialBuildingLockActive = false;
-
-		// unlock all generators
-		for (int32 NavDataIndex = 0; NavDataIndex < NavDataSet.Num(); ++NavDataIndex)
-		{
-			ANavigationData* NavData = NavDataSet[NavDataIndex];
-			if (NavData != NULL && NavData->GetGenerator() != NULL)
-			{
-				NavData->GetGenerator(FNavigationSystem::DontCreate)->OnNavigationBuildingUnlocked(bForce);
-			}
-		}
+		
+		RebuildAll();
 	}
 	else if (bInitialBuildingLockActive == true)
 	{
@@ -2653,12 +2571,15 @@ void UNavigationSystem::NavigationBuildingUnlock(bool bForce)
 
 void UNavigationSystem::RebuildAll()
 {
+	const bool bIsInGame = GetWorld()->IsGameWorld();
+	
 	for (int32 NavDataIndex = 0; NavDataIndex < NavDataSet.Num(); ++NavDataIndex)
 	{
 		ANavigationData* NavData = NavDataSet[NavDataIndex];
-		if (NavData != NULL && NavData->GetGenerator(FNavigationSystem::Create) != NULL)
+
+		if (NavData && NavData->bRebuildAtRuntime || (GIsEditor && !bIsInGame))
 		{
-			NavData->GetGenerator(FNavigationSystem::DontCreate)->RebuildAll();
+			NavData->RebuildAll();
 		}
 	}
 }
@@ -2687,6 +2608,36 @@ bool UNavigationSystem::IsNavigationBuildInProgress(bool bCheckDirtyToo)
 	return bRet;
 }
 
+int32 UNavigationSystem::GetNumRemainingBuildTasks() const
+{
+	int32 NumTasks = 0;
+	
+	for (ANavigationData* NavData : NavDataSet)
+	{
+		if (NavData && NavData->GetGenerator())
+		{
+			NumTasks+= NavData->GetGenerator()->GetNumRemaningBuildTasks();
+		}
+	}
+	
+	return NumTasks;
+}
+
+int32 UNavigationSystem::GetNumRunningBuildTasks() const 
+{
+	int32 NumTasks = 0;
+	
+	for (ANavigationData* NavData : NavDataSet)
+	{
+		if (NavData && NavData->GetGenerator())
+		{
+			NumTasks+= NavData->GetGenerator()->GetNumRunningBuildTasks();
+		}
+	}
+	
+	return NumTasks;
+}
+
 void UNavigationSystem::OnLevelAddedToWorld(ULevel* InLevel, UWorld* InWorld)
 {
 	if (InWorld == GetWorld())
@@ -2705,7 +2656,6 @@ void UNavigationSystem::OnLevelRemovedFromWorld(ULevel* InLevel, UWorld* InWorld
 
 void UNavigationSystem::AddLevelCollisionToOctree(ULevel* Level)
 {
-#if NAVOCTREE_CONTAINS_COLLISION_DATA
 	const TArray<FVector>* LevelGeom = Level ? Level->GetStaticNavigableGeometry() : NULL;
 	if (LevelGeom && NavOctree)
 	{
@@ -2719,12 +2669,10 @@ void UNavigationSystem::AddLevelCollisionToOctree(ULevel* Level)
 			AddDirtyArea(Bounds, ENavigationDirtyFlag::All);
 		}
 	}
-#endif // NAVOCTREE_CONTAINS_COLLISION_DATA
 }
 
 void UNavigationSystem::RemoveLevelCollisionFromOctree(ULevel* Level)
 {
-#if NAVOCTREE_CONTAINS_COLLISION_DATA
 	const FOctreeElementId* ElementId = GetObjectsNavOctreeId(Level);
 	if (ElementId != NULL)
 	{
@@ -2738,11 +2686,8 @@ void UNavigationSystem::RemoveLevelCollisionFromOctree(ULevel* Level)
 		NavOctree->RemoveNode(*ElementId);
 		RemoveObjectsNavOctreeId(Level);
 	}
-#endif // NAVOCTREE_CONTAINS_COLLISION_DATA
 }
 
-#endif // WITH_NAVIGATION_GENERATOR
-
 void UNavigationSystem::OnPostLoadMap()
 {
 	UE_LOG(LogNavigation, Log, TEXT("UNavigationSystem::OnPostLoadMap"));
@@ -2761,20 +2706,16 @@ void UNavigationSystem::OnPostLoadMap()
 #if WITH_EDITOR
 void UNavigationSystem::OnActorMoved(AActor* Actor)
 {
-#if WITH_NAVIGATION_GENERATOR
 	if (Cast<ANavMeshBoundsVolume>(Actor) != NULL)
 	{
 		OnNavigationBoundsUpdated((ANavMeshBoundsVolume*)Actor);
 	}
-#endif // WITH_NAVIGATION_GENERATOR
 }
 #endif // WITH_EDITOR
 
 void UNavigationSystem::OnNavigationDirtied(const FBox& Bounds)
 {
-#if WITH_NAVIGATION_GENERATOR
 	AddDirtyArea(Bounds, ENavigationDirtyFlag::All);
-#endif //WITH_NAVIGATION_GENERATOR
 }
 
 void UNavigationSystem::CleanUp(ECleanupMode Mode)
@@ -2789,11 +2730,9 @@ void UNavigationSystem::CleanUp(ECleanupMode Mode)
 #endif // WITH_EDITOR
 
 	FCoreUObjectDelegates::PostLoadMap.RemoveAll(this);
-#if WITH_NAVIGATION_GENERATOR
 	UNavigationSystem::NavigationDirtyEvent.RemoveAll(this);
 	FWorldDelegates::LevelAddedToWorld.RemoveAll(this);
 	FWorldDelegates::LevelRemovedFromWorld.RemoveAll(this);
-#endif // WITH_NAVIGATION_GENERATOR
 
 	if (NavOctree != NULL)
 	{
@@ -2906,11 +2845,7 @@ bool UNavigationSystem::IsNavigationBeingBuilt(UObject* WorldContextObject)
 	UWorld* World = GEngine->GetWorldFromContextObject(WorldContextObject); 
 	if (World != NULL && World->GetNavigationSystem() != NULL)
 	{
-#if WITH_NAVIGATION_GENERATOR
 		return World->GetNavigationSystem()->IsNavigationBuildInProgress();
-#else
-		return false;
-#endif
 	}
 
 	return false;
@@ -2921,19 +2856,18 @@ bool UNavigationSystem::IsNavigationBeingBuilt(UObject* WorldContextObject)
 //----------------------------------------------------------------------//
 bool UNavigationSystem::ShouldGeneratorRun(const FNavDataGenerator* Generator) const
 {
-#if WITH_NAVIGATION_GENERATOR
 	if (Generator != NULL)
 	{
 		for (int32 NavDataIndex = 0; NavDataIndex < NavDataSet.Num(); ++NavDataIndex)
 		{
 			ANavigationData* NavData = NavDataSet[NavDataIndex];
-			if (NavData != NULL && NavData->GetGenerator(FNavigationSystem::DontCreate) == Generator)
+			if (NavData != NULL && NavData->GetGenerator() == Generator)
 			{
 				return true;
 			}
 		}
 	}
-#endif
+
 	return false;
 }
 
diff --git a/Engine/Source/Runtime/Engine/Private/AI/Navigation/NavigationTestingActor.cpp b/Engine/Source/Runtime/Engine/Private/AI/Navigation/NavigationTestingActor.cpp
--- a/Engine/Source/Runtime/Engine/Private/AI/Navigation/NavigationTestingActor.cpp
+++ b/Engine/Source/Runtime/Engine/Private/AI/Navigation/NavigationTestingActor.cpp
@@ -56,6 +56,8 @@ ANavigationTestingActor::ANavigationTestingActor(const class FPostConstructIniti
 	CapsuleComponent->bCanEverAffectNavigation = false;
 
 	RootComponent = CapsuleComponent;
+
+	PathObserver = FNavigationPath::FPathObserverDelegate::FDelegate::CreateUObject(this, &ANavigationTestingActor::OnPathEvent);
 }
 
 ANavigationTestingActor::~ANavigationTestingActor()
@@ -212,11 +214,7 @@ void ANavigationTestingActor::PostLoad()
 void ANavigationTestingActor::TickMe()
 {
 	UNavigationSystem* NavSys = GetWorld() ? GetWorld()->GetNavigationSystem() : NULL;
-	if (NavSys
-#if WITH_NAVIGATION_GENERATOR
-		&& !NavSys->IsNavigationBuildInProgress()
-#endif
-		)
+	if (NavSys && !NavSys->IsNavigationBuildInProgress())
 	{
 #if WITH_RECAST && WITH_EDITORONLY_DATA
 		delete TickHelper;
@@ -347,6 +345,7 @@ void ANavigationTestingActor::SearchPathTo(class ANavigationTestingActor* Goal)
 	bPathSearchOutOfNodes = bPathExist ? Result.Path->DidSearchReachedLimit() : false;
 	LastPath = Result.Path;
 	PathCost = bPathExist ? Result.Path->GetCost() : 0.0f;
+	LastPath->AddObserver(PathObserver);
 
 	if (OffsetFromCornersDistance > 0.0f)
 	{
@@ -365,6 +364,30 @@ void ANavigationTestingActor::SearchPathTo(class ANavigationTestingActor* Goal)
 #endif
 }
 
+void ANavigationTestingActor::OnPathEvent(FNavigationPath* InvalidatedPath, ENavPathEvent::Type Event)
+{
+	if (InvalidatedPath == NULL || InvalidatedPath != LastPath.Get())
+    {
+		return;
+    }
+
+	switch (Event)
+	{
+		case ENavPathEvent::Invalidated:
+		{
+			UpdatePathfinding();
+			break;
+		}
+		case ENavPathEvent::RePathFailed:
+			break;
+		case ENavPathEvent::UpdatedDueToGoalMoved:
+		case ENavPathEvent::UpdatedDueToNavigationChanged:
+		{
+		}
+		break;
+	}
+}
+
 FPathFindingQuery ANavigationTestingActor::BuildPathFindingQuery(const ANavigationTestingActor* Goal) const
 {
 	check(Goal);
diff --git a/Engine/Source/Runtime/Engine/Private/AI/Navigation/PImplRecastNavMesh.cpp b/Engine/Source/Runtime/Engine/Private/AI/Navigation/PImplRecastNavMesh.cpp
--- a/Engine/Source/Runtime/Engine/Private/AI/Navigation/PImplRecastNavMesh.cpp
+++ b/Engine/Source/Runtime/Engine/Private/AI/Navigation/PImplRecastNavMesh.cpp
@@ -14,9 +14,7 @@
 #include "DetourNode.h"
 #include "DetourCommon.h"
 #include "RecastAlloc.h"
-#if WITH_NAVIGATION_GENERATOR
-	#include "RecastNavMeshGenerator.h"
-#endif // WITH_NAVIGATION_GENERATOR
+#include "RecastNavMeshGenerator.h"
 
 #include "AI/Navigation/NavLinkCustomInterface.h"
 #include "VisualLog.h"
@@ -245,6 +243,8 @@ FPImplRecastNavMesh::FPImplRecastNavMesh(ARecastNavMesh* Owner)
 {
 	check(Owner && "Owner must never be NULL");
 
+	DetourNavMesh = dtAllocNavMesh();
+
 	INC_DWORD_STAT_BY( STAT_NavigationMemory
 		, Owner->HasAnyFlags(RF_ClassDefaultObject) == false ? sizeof(*this) : 0 );
 };
@@ -252,11 +252,11 @@ FPImplRecastNavMesh::FPImplRecastNavMesh(ARecastNavMesh* Owner)
 FPImplRecastNavMesh::~FPImplRecastNavMesh()
 {
 	// release navmesh only if we own it
-	if (bOwnsNavMeshData && DetourNavMesh != NULL)
+	if (DetourNavMesh != nullptr)
 	{
 		dtFreeNavMesh(DetourNavMesh);
 	}
-	DetourNavMesh = NULL;
+	DetourNavMesh = nullptr;
 
 	DEC_DWORD_STAT_BY( STAT_NavigationMemory, sizeof(*this) );
 };
@@ -1750,32 +1750,6 @@ static FORCEINLINE float PointDistToSegment2DSquared(const float* PT, const floa
 	return dx*dx + dz*dz;
 }
 
-void FPImplRecastNavMesh::GetDebugTileBounds(FBox& OuterBox, int32& NumTilesX, int32& NumTilesY) const
-{
- 	if (DetourNavMesh)
- 	{
- 		dtNavMeshParams const* Params = DetourNavMesh->getParams();
-		NumTilesX = FMath::TruncToFloat( FMath::Sqrt((float)Params->maxTiles) );
-		NumTilesY = NumTilesX;
-
- 		FVector const Mn = Recast2UnrVector(Params->orig);
-
-		float RecastMax[3];
-		RecastMax[0] = Params->orig[0] + (Params->tileWidth * NumTilesX);
-		RecastMax[1] = Params->orig[1];
-		RecastMax[2] = Params->orig[2] + (Params->tileHeight * NumTilesY);
-		FVector const Mx = Recast2UnrVector(RecastMax);
-
-		OuterBox = FBox(Mx, Mn);
-	}
-	else
-	{
-		NumTilesX = 0;
-		NumTilesY = 0;
-		OuterBox = FBox();
-	}
-}
-
 /** 
  * Traverses given tile's edges and detects the ones that are either poly (i.e. not triangle, but whole navmesh polygon) 
  * or navmesh edge. Returns a pair of verts for each edge found.
@@ -1965,9 +1939,7 @@ void FPImplRecastNavMesh::GetDebugGeometry(FRecastDebugGeometry& OutGeometry, in
 		OutGeometry.BuiltMeshIndices.Reserve(NumIndicesToReserve);
 		OutGeometry.Clusters.AddZeroed(NumClusters);
 
-#if WITH_NAVIGATION_GENERATOR
 		const FRecastNavMeshGenerator* Generator = NavMeshOwner ? (const FRecastNavMeshGenerator*)(NavMeshOwner->GetGenerator()) : NULL;
-#endif
 
 		// spin through all polys in all tiles and draw them
 		// @see drawMeshTile() in recast code for reference
@@ -1982,12 +1954,8 @@ void FPImplRecastNavMesh::GetDebugGeometry(FRecastDebugGeometry& OutGeometry, in
 				continue;
 			}
 
-#if WITH_NAVIGATION_GENERATOR
 			const bool bIsBeingBuilt = Generator != NULL && !!NavMeshOwner->bDistinctlyDrawTilesBeingBuilt 
-				&& Generator->IsTileFresh(Header->x, Header->y);
-#else
-			const bool bIsBeingBuilt = false;
-#endif
+				&& Generator->IsTileChanged(TileIdx);
 			
 			// add all the poly verts
 			float* F = Tile->verts;
@@ -2199,7 +2167,7 @@ FBox FPImplRecastNavMesh::GetNavMeshTileBounds(int32 TileIndex) const
 }
 
 /** Retrieves XY coordinates of tile specified by index */
-void FPImplRecastNavMesh::GetNavMeshTileXY(int32 TileIndex, int32& OutX, int32& OutY, int32& OutLayer) const
+bool FPImplRecastNavMesh::GetNavMeshTileXY(int32 TileIndex, int32& OutX, int32& OutY, int32& OutLayer) const
 {
 	if (DetourNavMesh && TileIndex >= 0 && TileIndex < DetourNavMesh->getMaxTiles())
 	{
@@ -2215,12 +2183,15 @@ void FPImplRecastNavMesh::GetNavMeshTileXY(int32 TileIndex, int32& OutX, int32&
 				OutX = Header->x;
 				OutY = Header->y;
 				OutLayer = Header->layer;
+				return true;
 			}
 		}
 	}
+
+	return false;
 }
 
-void FPImplRecastNavMesh::GetNavMeshTileXY(const FVector& Point, int32& OutX, int32& OutY) const
+bool FPImplRecastNavMesh::GetNavMeshTileXY(const FVector& Point, int32& OutX, int32& OutY) const
 {
 	if (DetourNavMesh)
 	{
@@ -2234,7 +2205,10 @@ void FPImplRecastNavMesh::GetNavMeshTileXY(const FVector& Point, int32& OutX, in
 		ConstRecastNavMesh->calcTileLoc(&RecastPt.X, &TileX, &TileY);
 		OutX = TileX;
 		OutY = TileY;
+		return true;
 	}
+
+	return false;
 }
 
 void FPImplRecastNavMesh::GetNavMeshTilesAt(int32 TileX, int32 TileY, TArray<int32>& Indices) const
diff --git a/Engine/Source/Runtime/Engine/Private/AI/Navigation/PImplRecastNavMesh.h b/Engine/Source/Runtime/Engine/Private/AI/Navigation/PImplRecastNavMesh.h
--- a/Engine/Source/Runtime/Engine/Private/AI/Navigation/PImplRecastNavMesh.h
+++ b/Engine/Source/Runtime/Engine/Private/AI/Navigation/PImplRecastNavMesh.h
@@ -72,7 +72,6 @@ public:
 
 	/** Debug rendering. */
 	void GetDebugGeometry(FRecastDebugGeometry& OutGeometry, int32 TileIndex = INDEX_NONE) const;
-	void GetDebugTileBounds(FBox& OuterBox, int32& NumTilesX, int32& NumTilesY) const;
 
 	/** Returns bounding box for the whole navmesh. */
 	FBox GetNavMeshBounds() const;
@@ -81,10 +80,10 @@ public:
 	FBox GetNavMeshTileBounds(int32 TileIndex) const;
 
 	/** Retrieves XY and layer coordinates of tile specified by index */
-	void GetNavMeshTileXY(int32 TileIndex, int32& OutX, int32& OutY, int32& OutLayer) const;
+	bool GetNavMeshTileXY(int32 TileIndex, int32& OutX, int32& OutY, int32& OutLayer) const;
 
 	/** Retrieves XY coordinates of tile specified by position */
-	void GetNavMeshTileXY(const FVector& Point, int32& OutX, int32& OutY) const;
+	bool GetNavMeshTileXY(const FVector& Point, int32& OutX, int32& OutY) const;
 
 	/** Retrieves all tile indices at matching XY coordinates */
 	void GetNavMeshTilesAt(int32 TileX, int32 TileY, TArray<int32>& Indices) const;
diff --git a/Engine/Source/Runtime/Engine/Private/AI/Navigation/RecastNavMesh.cpp b/Engine/Source/Runtime/Engine/Private/AI/Navigation/RecastNavMesh.cpp
--- a/Engine/Source/Runtime/Engine/Private/AI/Navigation/RecastNavMesh.cpp
+++ b/Engine/Source/Runtime/Engine/Private/AI/Navigation/RecastNavMesh.cpp
@@ -21,18 +21,6 @@
 #include "UnrealEd.h"
 #endif
 
-#if RECAST_ASYNC_REBUILDING
-
-#define SECTION_LOCK_TILES			FConditionalScopeLock ReadWriteLock(NavDataReadWriteLock, BatchQueryCounter <= 0);
-#define SECTION_LOCK_TILES_FOR(x)	FConditionalScopeLock ReadWriteLock(x->NavDataReadWriteLock, x->BatchQueryCounter <= 0);
-
-#else
-
-#define SECTION_LOCK_TILES
-#define SECTION_LOCK_TILES_FOR(x)
-
-#endif // RECAST_ASYNC_REBUILDING
-
 #if WITH_RECAST
 /// Helper for accessing navigation query from different threads
 #define INITIALIZE_NAVQUERY(NavQueryVariable, NumNodes)	\
@@ -47,19 +35,6 @@
 
 #endif // WITH_RECAST
 
-namespace
-{
-	FORCEINLINE FVector2D TileIndexToXY(int32 GridWidth, int32 Index)
-	{
-		return FVector2D(Index % GridWidth, Index / GridWidth);
-	}
-
-	FORCEINLINE int32 TileXYToIndex(int32 GridWidth, int32 X, int32 Y)
-	{
-		return Y * GridWidth + X;
-	}
-}
-
 FNavMeshTileData::FNavData::~FNavData()
 {
 #if WITH_RECAST
@@ -69,41 +44,6 @@ FNavMeshTileData::FNavData::~FNavData()
 #endif
 }
 
-// Similar to FScopeLock, but locks only when condition if set
-class FConditionalScopeLock
-{
-public:
-	FConditionalScopeLock(FCriticalSection* InSyncObject, bool bShouldLock) 
-		: SyncObject(InSyncObject), bLock(bShouldLock)
-	{
-		if (bLock)
-		{
-			check(SyncObject);
-			SyncObject->Lock();
-		}
-	}
-
-	~FConditionalScopeLock()
-	{
-		if (bLock)
-		{
-			check(SyncObject);
-			SyncObject->Unlock();
-		}
-	}
-
-private:
-	FConditionalScopeLock();
-	FConditionalScopeLock(FConditionalScopeLock* InScopeLock);
-	FConditionalScopeLock& operator= (FConditionalScopeLock& InScopeLock) 
-	{ 
-		return *this; 
-	}
-
-	FCriticalSection* SyncObject;
-	bool bLock;
-};
-
 float ARecastNavMesh::DrawDistanceSq = 0.0f;
 #if !WITH_RECAST
 
@@ -171,24 +111,11 @@ uint32 FRecastDebugGeometry::GetAllocatedSize() const
 
 	return Size;
 }
+
 //----------------------------------------------------------------------//
 // ARecastNavMesh
 //----------------------------------------------------------------------//
 
-void FRecastTickHelper::Tick(float DeltaTime)
-{
-	if (Owner->IsPendingKill() == false)
-	{
-		Owner->TickMe(DeltaTime);
-	}
-}
-
-TStatId FRecastTickHelper::GetStatId() const 
-{
-	RETURN_QUICK_DECLARE_CYCLE_STAT(FRecastTickHelper, STATGROUP_Tickables);
-}
-
-
 namespace ERecastNamedFilter
 {
 	FRecastQueryFilter FilterOutNavLinksImpl;
@@ -243,10 +170,7 @@ ARecastNavMesh::ARecastNavMesh(const class FPostConstructInitializeProperties& P
 	, bPerformVoxelFiltering(true)	
 	, bMarkLowHeightAreas(false)
 	, bUseVirtualFilters(true)
-	, NextTimeToSortTiles(0.f)
 	, TileSetUpdateInterval(1.0f)
-	, GridWidth(-1)
-	, GridHeight(-1)
 	, NavMeshVersion(NAVMESHVER_LATEST)	
 	, RecastNavMeshImpl(NULL)
 {
@@ -274,9 +198,7 @@ ARecastNavMesh::ARecastNavMesh(const class FPostConstructInitializeProperties& P
 		RaycastImplementation = NavMeshRaycast;
 
 		RecastNavMeshImpl = new FPImplRecastNavMesh(this);
-		
-		TickHelper.Owner = this;
-
+	
 		// add predefined areas up front
 		SupportedAreas.Add(FSupportedAreaData(UNavArea_Null::StaticClass(), RECAST_NULL_AREA));
 		SupportedAreas.Add(FSupportedAreaData(UNavArea_LowHeight::StaticClass(), RECAST_LOW_AREA));
@@ -288,11 +210,6 @@ ARecastNavMesh::~ARecastNavMesh()
 {
 	if (HasAnyFlags(RF_ClassDefaultObject) == false)
 	{
-		for (int32 i = 0; i < TileNavDataContainer.Num(); i++)
-		{
-			DEC_MEMORY_STAT_BY(STAT_Navigation_NavDataMemory, TileNavDataContainer[i].DataSize);
-		}
-
 		DEC_DWORD_STAT_BY( STAT_NavigationMemory, sizeof(*this) );
 		DestroyRecastPImpl();
 	}
@@ -307,157 +224,6 @@ void ARecastNavMesh::DestroyRecastPImpl()
 	}
 }
 
-void ARecastNavMesh::TickMe(float DeltaTime)
-{
-#if WITH_NAVIGATION_GENERATOR
-	// sanity check
-	if (RecastNavMeshImpl == NULL)
-	{
-		return;
-	}
-
-	if (NavDataGenerator.IsValid() && ((FRecastNavMeshGenerator*)NavDataGenerator.Get())->HasResultsPending() == true)
-	{
-		TNavStatArray<FNavMeshGenerationResult> AsyncResults;
-		((FRecastNavMeshGenerator*)NavDataGenerator.Get())->GetAsyncResultsCopy(AsyncResults, /*bClearSource=*/true);
-
-		FNavMeshGenerationResult* Result = AsyncResults.GetData();
-		for (int32 ResultIndex = 0; ResultIndex < AsyncResults.Num(); ++ResultIndex, ++Result)
-		{
-			UpdateNavTileData(Result->OldRawNavData, Result->NewNavData);
-			FreshTiles.Add(Result->TileIndex);
-		}
-	}
-
-	const int32 FreshTilesCount = FreshTiles.Num();
-	if (FreshTilesCount > 0)
-	{
-		const TArray<uint32> FreshTilesCopy(FreshTiles);
-		FreshTiles.Reset();
-		const int32 PathsCount = ActivePaths.Num();
-
-		FNavPathWeakPtr* WeakPathPtr = (ActivePaths.GetData() + PathsCount - 1);
-		for (int32 PathIndex = PathsCount - 1; PathIndex >= 0; --PathIndex, --WeakPathPtr)
-		{
-			FNavPathSharedPtr SharedPath = WeakPathPtr->Pin();
-			if (WeakPathPtr->IsValid() == false)
-			{
-				ActivePaths.RemoveAtSwap(PathIndex, 1, /*bAllowShrinking=*/false);
-			}
-			else 
-			{
-				// iterate through all tile refs in FreshTilesCopy and 
-				const FNavMeshPath* Path = (const FNavMeshPath*)(SharedPath.Get());
-				if (Path->IsReady() == false)
-				{
-					// path not filled yet anyway
-					continue;
-				}
-
-				const int32 PathLenght = Path->PathCorridor.Num();
-				const NavNodeRef* PathPoly = Path->PathCorridor.GetData();
-				for (int32 NodeIndex = 0; NodeIndex < PathLenght; ++NodeIndex, ++PathPoly)
-				{
-					const uint32 NodeTileIdx = RecastNavMeshImpl->GetTileIndexFromPolyRef(*PathPoly);
-					if (FreshTilesCopy.Contains(NodeTileIdx))
-					{
-						SharedPath->Invalidate();
-						ActivePaths.RemoveAtSwap(PathIndex, 1, /*bAllowShrinking=*/false);
-						break;
-					}
-				}
-			}
-		}
-	}
-
-	NextTimeToSortTiles += DeltaTime;
-	if (NextTimeToSortTiles >= TileSetUpdateInterval)
-	{
-		SortNavigationTiles();
-		NextTimeToSortTiles = 0;
-	}
-#endif // WITH_NAVIGATION_GENERATOR
-}
-
-#if WITH_NAVIGATION_GENERATOR
-
-void ARecastNavMesh::SortNavigationTiles()
-{
-	struct FTileSetItemsSortPredicate
-	{
-		TArray<FTileSetItem>& TileSet;
-
-		FTileSetItemsSortPredicate(TArray<FTileSetItem>& InTileSet)
-			: TileSet(InTileSet)
-		{}
-
-		bool operator()(int32 A, int32 B) const 
-		{
-			// lesser distance = better
-			return TileSet[A].Dist2DSq < TileSet[B].Dist2DSq;
-		}
-	};
-
-	TArray<FVector2D> SeedLocations;
-
-	UWorld* CurWorld = GetWorld();
-	if (CurWorld == NULL)
-	{
-		return;
-	}
-
-	for (FConstPlayerControllerIterator PlayerIt = CurWorld->GetPlayerControllerIterator(); PlayerIt; ++PlayerIt)
-	{
-		if (*PlayerIt && (*PlayerIt)->GetPawn() != NULL)
-		{
-			const FVector2D SeedLoc((*PlayerIt)->GetPawn()->GetActorLocation());
-			SeedLocations.Add(SeedLoc);
-		}
-	}
-
-	const int32 NumSeeds = SeedLocations.Num();
-	if (NumSeeds <= 0)
-	{
-		return;
-	}
-	
-	// changes to TileSet must be thread safe, because tile workers are accessing it
-	{
-		SECTION_LOCK_TILES;
-
-		TArray<int32> SortedIndices;
-		SortedIndices.AddZeroed(TileSet.Num());
-
-		for (int32 TileIndex = 0; TileIndex < TileSet.Num(); TileIndex++)
-		{
-			FTileSetItem& TileItem = TileSet[TileIndex];
-			float BestDist2DSq = FLT_MAX;
-			for (int32 SeedIndex = 0; SeedIndex < SeedLocations.Num(); SeedIndex++)
-			{
-				const float Dist2DSq = FVector2D::DistSquared(TileItem.TileCenter, SeedLocations[SeedIndex]);
-				if (BestDist2DSq > Dist2DSq)
-				{
-					BestDist2DSq = Dist2DSq;
-				}
-			}
-
-			TileItem.Dist2DSq = BestDist2DSq;
-			SortedIndices[TileIndex] = TileIndex;
-		}
-
-		// always sort to keep array ordered (tile removals by workers)
-		SortedIndices.Sort(FTileSetItemsSortPredicate(TileSet));
-
-		for (int32 SortIndex = 0; SortIndex < TileSet.Num(); SortIndex++)
-		{
-			FTileSetItem& TileItem = TileSet[SortedIndices[SortIndex]];
-			TileItem.SortOrder = SortIndex;
-		}
-	}
-}
-
-#endif // WITH_NAVIGATION_GENERATOR
-
 ARecastNavMesh* ARecastNavMesh::SpawnInstance(UNavigationSystem* NavSys, const FNavDataConfig* AgentProps)
 {
 	FActorSpawnParameters SpawnInfo;
@@ -643,52 +409,10 @@ void ARecastNavMesh::UpdateNavMeshDrawing()
 #endif // UE_BUILD_SHIPPING
 }
 
-#if WITH_NAVIGATION_GENERATOR
-FNavDataGenerator* ARecastNavMesh::ConstructGenerator(const FNavAgentProperties& AgentProps)
-{
-	FRecastNavMeshGenerator* Generator = NULL;
-	// Need to know properties to generate navmesh
-	if (AgentProps.IsValid())
-	{
-		Generator = new FRecastNavMeshGenerator(this);
-		if (Generator != NULL)
-		{
-#if RECAST_ASYNC_REBUILDING
-			NavDataReadWriteLock = Generator->GetNavMeshModificationLock();
-#endif // RECAST_ASYNC_REBUILDING
-		}
-	}
-
-	return Generator;
-}
-#endif
-
-void ARecastNavMesh::DestroyGenerator()
-{
-#if RECAST_ASYNC_REBUILDING
-	// don't try to destroy generator with running async workers
-	FRecastNavMeshGenerator* MyGenerator = NavDataGenerator.IsValid() ? (FRecastNavMeshGenerator*)NavDataGenerator.Get() : NULL;
-	if (MyGenerator)
-	{
-		while (MyGenerator->IsAsyncBuildInProgress())
-		{
-			UE_LOG(LogNavigation, Log, TEXT("Can't destroy generator with active async worker, waiting..."));
-			FPlatformProcess::Sleep(0.010f);
-		}
-	}
-#endif
-
-	Super::DestroyGenerator();
-
-#if RECAST_ASYNC_REBUILDING
-	NavDataReadWriteLock = &NavDataReadWriteLockDummy;
-#endif // RECAST_ASYNC_REBUILDING
-}
-
 void ARecastNavMesh::CleanUp()
 {
 	Super::CleanUp();
-
+	NavDataGenerator.Reset();
 	DestroyRecastPImpl();
 }
 
@@ -719,10 +443,6 @@ void ARecastNavMesh::PostInitProperties()
 
 	if (HasAnyFlags(RF_ClassDefaultObject) == false)
 	{
-#if RECAST_ASYNC_REBUILDING
-		NavDataReadWriteLock = &NavDataReadWriteLockDummy;		
-#endif // RECAST_ASYNC_REBUILDING
-
 		DefaultQueryFilter->SetFilterType<FRecastQueryFilter>();
 		FRecastQueryFilter* DetourFilter = static_cast<FRecastQueryFilter*>(DefaultQueryFilter->GetImplementation());
 		DetourFilter->SetIsVirtual(bUseVirtualFilters);
@@ -791,14 +511,12 @@ void ARecastNavMesh::OnNavAreaAdded(const UClass* NavAreaClass, int32 AgentIndex
 		DefaultQueryFilter->SetFixedAreaEnteringCost(AreaID, DefArea->FixedAreaEnteringCost);
 	}
 
-#if WITH_NAVIGATION_GENERATOR
 	// update generator's cached data
-	FRecastNavMeshGenerator* MyGenerator = NavDataGenerator.IsValid() ? (FRecastNavMeshGenerator*)NavDataGenerator.Get() : NULL;
+	FRecastNavMeshGenerator* MyGenerator = static_cast<FRecastNavMeshGenerator*>(GetGenerator());
 	if (MyGenerator)
 	{
 		MyGenerator->OnAreaAdded(NavAreaClass, AreaID);
 	}
-#endif // WITH_NAVIGATION_GENERATOR
 }
 
 int32 ARecastNavMesh::GetNewAreaID(const UClass* AreaClass) const
@@ -873,20 +591,8 @@ void ARecastNavMesh::SortAreasForGenerator(TArray<FAreaNavModifier>& Areas) cons
 
 void ARecastNavMesh::SerializeRecastNavMesh(FArchive& Ar, FPImplRecastNavMesh*& NavMesh)
 {
-	if (!Ar.IsLoading() 
-		&& (NavMesh == NULL 
-#if WITH_NAVIGATION_GENERATOR
-			|| (GetGenerator() != NULL && GetGenerator()->IsBuildInProgress(/*bCheckDirtyToo=*/true) == true)
-#endif // WITH_NAVIGATION_GENERATOR
-			)
-		)
+	if (!Ar.IsLoading()	&& NavMesh == NULL)
 	{
-		//disable warning during PIE, we'll have information about wrong navmesh on screen if dynamic building is disabled
-		if (NavMesh != NULL && Ar.IsSaving())
-		{
-			// nothing to write or while building, which in turn could result in corrupted save
-			UE_LOG(LogNavigation, Warning, TEXT("ARecastNavMesh: Unable to serialize navmesh when build is in progress or navmesh is not up to date."));
-		}
 		return;
 	}
 
@@ -898,6 +604,14 @@ void ARecastNavMesh::SerializeRecastNavMesh(FArchive& Ar, FPImplRecastNavMesh*&
 			RecastNavMeshImpl = new FPImplRecastNavMesh(this);
 		}
 	}
+	
+	if (Ar.IsSaving())
+	{
+		if (NavDataGenerator)
+		{
+			NavDataGenerator->EnsureBuildCompletion();
+		}
+	}
 
 	if (RecastNavMeshImpl)
 	{
@@ -905,7 +619,6 @@ void ARecastNavMesh::SerializeRecastNavMesh(FArchive& Ar, FPImplRecastNavMesh*&
 	}	
 }
 
-
 void ARecastNavMesh::Serialize( FArchive& Ar )
 {
 	Super::Serialize(Ar);
@@ -989,7 +702,6 @@ void ARecastNavMesh::BeginBatchQuery() const
 	if (BatchQueryCounter <= 0)
 	{
 		BatchQueryCounter = 0;
-		NavDataReadWriteLock->Lock();
 	}
 
 	BatchQueryCounter++;
@@ -1000,9 +712,6 @@ void ARecastNavMesh::FinishBatchQuery() const
 {
 #if RECAST_ASYNC_REBUILDING
 	BatchQueryCounter--;
-
-	// release critical section when no other batch queries are active
-	NavDataReadWriteLock->Unlock();
 #endif // RECAST_ASYNC_REBUILDING
 }
 
@@ -1011,7 +720,6 @@ FBox ARecastNavMesh::GetNavMeshBounds() const
 	FBox Bounds;
 	if (RecastNavMeshImpl)
 	{
-		SECTION_LOCK_TILES;
 		Bounds = RecastNavMeshImpl->GetNavMeshBounds();
 	}
 
@@ -1023,36 +731,36 @@ FBox ARecastNavMesh::GetNavMeshTileBounds(int32 TileIndex) const
 	FBox Bounds;
 	if (RecastNavMeshImpl)
 	{
-		SECTION_LOCK_TILES;
 		Bounds = RecastNavMeshImpl->GetNavMeshTileBounds(TileIndex);
 	}
 
 	return Bounds;
 }
 
-void ARecastNavMesh::GetNavMeshTileXY(int32 TileIndex, int32& OutX, int32& OutY, int32& OutLayer) const
+bool ARecastNavMesh::GetNavMeshTileXY(int32 TileIndex, int32& OutX, int32& OutY, int32& OutLayer) const
 {
 	if (RecastNavMeshImpl)
 	{
-		SECTION_LOCK_TILES;
-		RecastNavMeshImpl->GetNavMeshTileXY(TileIndex, OutX, OutY, OutLayer);
+		return RecastNavMeshImpl->GetNavMeshTileXY(TileIndex, OutX, OutY, OutLayer);
 	}
+
+	return false;
 }
 
-void ARecastNavMesh::GetNavMeshTileXY(const FVector& Point, int32& OutX, int32& OutY) const
+bool ARecastNavMesh::GetNavMeshTileXY(const FVector& Point, int32& OutX, int32& OutY) const
 {
 	if (RecastNavMeshImpl)
 	{
-		SECTION_LOCK_TILES;
-		RecastNavMeshImpl->GetNavMeshTileXY(Point, OutX, OutY);
+		return RecastNavMeshImpl->GetNavMeshTileXY(Point, OutX, OutY);
 	}
+
+	return false;
 }
 
 void ARecastNavMesh::GetNavMeshTilesAt(int32 TileX, int32 TileY, TArray<int32>& Indices) const
 {
 	if (RecastNavMeshImpl)
 	{
-		SECTION_LOCK_TILES;
 		RecastNavMeshImpl->GetNavMeshTilesAt(TileX, TileY, Indices);
 	}
 }
@@ -1062,7 +770,6 @@ bool ARecastNavMesh::GetPolysInTile(int32 TileIndex, TArray<FNavPoly>& Polys) co
 	bool bSuccess = false;
 	if (RecastNavMeshImpl)
 	{
-		SECTION_LOCK_TILES;
 		bSuccess = RecastNavMeshImpl->GetPolysInTile(TileIndex, Polys);
 	}
 
@@ -1074,7 +781,6 @@ int32 ARecastNavMesh::GetNavMeshTilesCount() const
 	int32 NumTiles = 0;
 	if (RecastNavMeshImpl)
 	{
-		SECTION_LOCK_TILES;
 		NumTiles = RecastNavMeshImpl->GetNavMeshTilesCount();
 	}
 
@@ -1087,7 +793,6 @@ void ARecastNavMesh::GetEdgesForPathCorridor(const TArray<NavNodeRef>* PathCorri
 
 	if (RecastNavMeshImpl)
 	{
-		SECTION_LOCK_TILES;
 		RecastNavMeshImpl->GetEdgesForPathCorridor(PathCorridor, PathCorridorEdges);
 	}
 }
@@ -1097,7 +802,6 @@ FNavLocation ARecastNavMesh::GetRandomPoint(TSharedPtr<const FNavigationQueryFil
 	FNavLocation RandomPt;
 	if (RecastNavMeshImpl)
 	{
-		SECTION_LOCK_TILES;
 		RandomPt = RecastNavMeshImpl->GetRandomPoint(GetRightFilterRef(Filter), QueryOwner);
 	}
 
@@ -1109,7 +813,6 @@ bool ARecastNavMesh::GetRandomPointInRadius(const FVector& Origin, float Radius,
 	bool bSuccess = false;
 	if (RecastNavMeshImpl)
 	{
-		SECTION_LOCK_TILES;
 		bSuccess = RecastNavMeshImpl->GetRandomPointInRadius(Origin, Radius, OutResult, GetRightFilterRef(Filter), QueryOwner);
 	}
 
@@ -1121,7 +824,6 @@ bool ARecastNavMesh::GetRandomPointInCluster(NavNodeRef ClusterRef, FNavLocation
 	bool bSuccess = false;
 	if (RecastNavMeshImpl)
 	{
-		SECTION_LOCK_TILES;
 		bSuccess = RecastNavMeshImpl->GetRandomPointInCluster(ClusterRef, OutLocation);
 	}
 
@@ -1133,7 +835,6 @@ NavNodeRef ARecastNavMesh::GetClusterRef(NavNodeRef PolyRef) const
 	NavNodeRef ClusterRef = 0;
 	if (RecastNavMeshImpl)
 	{
-		SECTION_LOCK_TILES;
 		ClusterRef = RecastNavMeshImpl->GetClusterRefFromPolyRef(PolyRef);
 	}
 
@@ -1145,7 +846,6 @@ bool ARecastNavMesh::ProjectPoint(const FVector& Point, FNavLocation& OutLocatio
 	bool bSuccess = false;
 	if (RecastNavMeshImpl)
 	{
-		SECTION_LOCK_TILES;
 		bSuccess = RecastNavMeshImpl->ProjectPointToNavMesh(Point, OutLocation, Extent, GetRightFilterRef(Filter), QueryOwner);
 	}
 
@@ -1158,8 +858,6 @@ void ARecastNavMesh::BatchProjectPoints(TArray<FNavigationProjectionWork>& Workl
 	{
 		return;
 	}
-
-	SECTION_LOCK_TILES;
 	
 	const FNavigationQueryFilter& FilterToUse = GetRightFilterRef(Filter);
 
@@ -1199,7 +897,6 @@ bool ARecastNavMesh::ProjectPointMulti(const FVector& Point, TArray<FNavLocation
 	bool bSuccess = false;
 	if (RecastNavMeshImpl)
 	{
-		SECTION_LOCK_TILES;
 		bSuccess = RecastNavMeshImpl->ProjectPointMulti(Point, OutLocations, Extent, MinZ, MaxZ, GetRightFilterRef(Filter), QueryOwner);
 	}
 
@@ -1237,10 +934,7 @@ ENavigationQueryResult::Type ARecastNavMesh::CalcPathLengthAndCost(const FVector
 			Path->SetWantsStringPulling(false);
 			Path->SetWantsPathCorridor(true);
 			
-			{
-				SECTION_LOCK_TILES;
-				Result = RecastNavMeshImpl->FindPath(PathStart, PathEnd, Path.Get(), GetRightFilterRef(QueryFilter), QueryOwner);
-			}
+			Result = RecastNavMeshImpl->FindPath(PathStart, PathEnd, Path.Get(), GetRightFilterRef(QueryFilter), QueryOwner);
 
 			if (Result == ENavigationQueryResult::Success || (Result == ENavigationQueryResult::Fail && Path->IsPartial()))
 			{
@@ -1258,7 +952,6 @@ NavNodeRef ARecastNavMesh::FindNearestPoly(FVector const& Loc, FVector const& Ex
 	NavNodeRef PolyRef = 0;
 	if (RecastNavMeshImpl)
 	{
-		SECTION_LOCK_TILES;
 		PolyRef = RecastNavMeshImpl->FindNearestPoly(Loc, Extent, GetRightFilterRef(Filter), QueryOwner);
 	}
 
@@ -1277,7 +970,6 @@ void ARecastNavMesh::UpdateNavigationLinkArea(int32 UserId, TSubclassOf<class UN
 	{
 		UNavArea* DefArea = (UNavArea*)(AreaClass->GetDefaultObject());
 
-		SECTION_LOCK_TILES;
 		RecastNavMeshImpl->UpdateNavigationLinkArea(UserId, AreaId, DefArea->GetAreaFlags());
 	}
 }
@@ -1289,7 +981,6 @@ void ARecastNavMesh::UpdateSegmentLinkArea(int32 UserId, TSubclassOf<class UNavA
 	{
 		UNavArea* DefArea = (UNavArea*)(AreaClass->GetDefaultObject());
 
-		SECTION_LOCK_TILES;
 		RecastNavMeshImpl->UpdateSegmentLinkArea(UserId, AreaId, DefArea->GetAreaFlags());
 	}
 }
@@ -1299,7 +990,6 @@ bool ARecastNavMesh::GetPolyCenter(NavNodeRef PolyID, FVector& OutCenter) const
 	bool bSuccess = false;
 	if (RecastNavMeshImpl)
 	{
-		SECTION_LOCK_TILES;
 		bSuccess = RecastNavMeshImpl->GetPolyCenter(PolyID, OutCenter);
 	}
 
@@ -1311,7 +1001,6 @@ bool ARecastNavMesh::GetPolyVerts(NavNodeRef PolyID, TArray<FVector>& OutVerts)
 	bool bSuccess = false;
 	if (RecastNavMeshImpl)
 	{
-		SECTION_LOCK_TILES;
 		bSuccess = RecastNavMeshImpl->GetPolyVerts(PolyID, OutVerts);
 	}
 
@@ -1323,7 +1012,6 @@ uint32 ARecastNavMesh::GetPolyAreaID(NavNodeRef PolyID) const
 	uint32 AreaID = RECAST_DEFAULT_AREA;
 	if (RecastNavMeshImpl)
 	{
-		SECTION_LOCK_TILES;
 		AreaID = RecastNavMeshImpl->GetPolyAreaID(PolyID);
 	}
 
@@ -1335,7 +1023,6 @@ bool ARecastNavMesh::GetPolyFlags(NavNodeRef PolyID, uint16& PolyFlags, uint16&
 	bool bFound = false;
 	if (RecastNavMeshImpl)
 	{
-		SECTION_LOCK_TILES;
 		uint8 AreaType = RECAST_DEFAULT_AREA;
 		bFound = RecastNavMeshImpl->GetPolyData(PolyID, PolyFlags, AreaType);
 		if (bFound)
@@ -1354,7 +1041,6 @@ bool ARecastNavMesh::GetClosestPointOnPoly(NavNodeRef PolyID, const FVector& Tes
 	bool bFound = false;
 	if (RecastNavMeshImpl)
 	{
-		SECTION_LOCK_TILES;
 		bFound = RecastNavMeshImpl->GetClosestPointOnPoly(PolyID, TestPt, PointOnPoly);
 	}
 
@@ -1366,7 +1052,6 @@ bool ARecastNavMesh::GetPolyTileIndex(NavNodeRef PolyID, uint32& PolyIndex, uint
 	bool bFound = false;
 	if (RecastNavMeshImpl)
 	{
-		SECTION_LOCK_TILES;
 		bFound = RecastNavMeshImpl->GetPolyTileIndex(PolyID, PolyIndex, TileIndex);
 	}
 
@@ -1378,7 +1063,6 @@ bool ARecastNavMesh::GetLinkEndPoints(NavNodeRef LinkPolyID, FVector& PointA, FV
 	bool bSuccess = false;
 	if (RecastNavMeshImpl)
 	{
-		SECTION_LOCK_TILES;
 		bSuccess = RecastNavMeshImpl->GetLinkEndPoints(LinkPolyID, PointA, PointB);
 	}
 
@@ -1390,7 +1074,6 @@ bool ARecastNavMesh::GetClusterBounds(NavNodeRef ClusterRef, FBox& OutBounds) co
 	bool bSuccess = false;
 	if (RecastNavMeshImpl)
 	{
-		SECTION_LOCK_TILES;
 		bSuccess = RecastNavMeshImpl->GetClusterBounds(ClusterRef, OutBounds);
 	}
 
@@ -1402,7 +1085,6 @@ bool ARecastNavMesh::GetPolysWithinPathingDistance(FVector const& StartLoc, cons
 	bool bSuccess = false;
 	if (RecastNavMeshImpl)
 	{
-		SECTION_LOCK_TILES;
 		bSuccess = RecastNavMeshImpl->GetPolysWithinPathingDistance(StartLoc, PathingDistance, GetRightFilterRef(Filter), QueryOwner, FoundPolys);
 	}
 
@@ -1413,20 +1095,10 @@ void ARecastNavMesh::GetDebugGeometry(FRecastDebugGeometry& OutGeometry, int32 T
 {
 	if (RecastNavMeshImpl)
 	{
-		SECTION_LOCK_TILES;
 		RecastNavMeshImpl->GetDebugGeometry(OutGeometry, TileIndex);
 	}
 }
 
-void ARecastNavMesh::GetDebugTileBounds(FBox& OuterBox, int32& NumTilesX, int32& NumTilesY) const
-{
-	if (RecastNavMeshImpl)
-	{
-		SECTION_LOCK_TILES;
-		RecastNavMeshImpl->GetDebugTileBounds(OuterBox, NumTilesX, NumTilesY);
-	}
-}
-
 void ARecastNavMesh::RequestDrawingUpdate()
 {
 #if !UE_BUILD_SHIPPING
@@ -1486,11 +1158,55 @@ void ARecastNavMesh::DrawDebugPathCorridor(NavNodeRef const* PathPolys, int32 Nu
 	}
 }
 
+void ARecastNavMesh::InvalidateAffectedPaths(const TArray<uint32>& ChangedTiles)
+{
+	const int32 PathsCount = ActivePaths.Num();
+	const int32 ChangedTilesCount = ChangedTiles.Num();
+	
+	if (ChangedTilesCount == 0 || PathsCount == 0)
+	{
+		return;
+	}
+	
+	FNavPathWeakPtr* WeakPathPtr = (ActivePaths.GetData() + PathsCount - 1);
+	for (int32 PathIndex = PathsCount - 1; PathIndex >= 0; --PathIndex, --WeakPathPtr)
+	{
+		FNavPathSharedPtr SharedPath = WeakPathPtr->Pin();
+		if (WeakPathPtr->IsValid() == false)
+		{
+			ActivePaths.RemoveAtSwap(PathIndex, 1, /*bAllowShrinking=*/false);
+		}
+		else 
+		{
+			// iterate through all tile refs in FreshTilesCopy and 
+			const FNavMeshPath* Path = (const FNavMeshPath*)(SharedPath.Get());
+			if (Path->IsReady() == false)
+			{
+				// path not filled yet anyway
+				continue;
+			}
+
+			const int32 PathLenght = Path->PathCorridor.Num();
+			const NavNodeRef* PathPoly = Path->PathCorridor.GetData();
+			for (int32 NodeIndex = 0; NodeIndex < PathLenght; ++NodeIndex, ++PathPoly)
+			{
+				const uint32 NodeTileIdx = RecastNavMeshImpl->GetTileIndexFromPolyRef(*PathPoly);
+				if (ChangedTiles.Contains(NodeTileIdx))
+				{
+					SharedPath->Invalidate();
+					ActivePaths.RemoveAtSwap(PathIndex, 1, /*bAllowShrinking=*/false);
+					break;
+				}
+			}
+		}
+	}
+}
+
 #if !UE_BUILD_SHIPPING
 uint32 ARecastNavMesh::LogMemUsed() const 
 {
 	const uint32 SuperMemUsed = Super::LogMemUsed();
-	uint32 MemUsed = FreshTiles.GetAllocatedSize();
+	uint32 MemUsed = 0;
 
 	UE_LOG(LogNavigation, Display, TEXT("%s: ARecastNavMesh: %u\n    self: %d"), *GetName(), MemUsed, sizeof(ARecastNavMesh));	
 
@@ -1519,7 +1235,6 @@ bool ARecastNavMesh::FilterPolys(TArray<NavNodeRef>& PolyRefs, const class FReca
 	bool bSuccess = false;
 	if (RecastNavMeshImpl)
 	{
-		SECTION_LOCK_TILES;
 		bSuccess = RecastNavMeshImpl->FilterPolys(PolyRefs, Filter, QueryOwner);
 	}
 
@@ -1530,7 +1245,6 @@ void ARecastNavMesh::ApplyWorldOffset(const FVector& InOffset, bool bWorldShift)
 {
 	if (RecastNavMeshImpl)
 	{
-		SECTION_LOCK_TILES;
 		RecastNavMeshImpl->ApplyWorldOffset(InOffset, bWorldShift);
 	}
 
@@ -1595,8 +1309,6 @@ FPathFindingResult ARecastNavMesh::FindPath(const FNavAgentProperties& AgentProp
 	}
 	else
 	{
-		SECTION_LOCK_TILES_FOR(RecastNavMesh);
-
 		if(Query.QueryFilter.IsValid())
 		{
 			Result.Result = RecastNavMesh->RecastNavMeshImpl->FindPath(Query.StartLocation, Query.EndLocation, *NavMeshPath,
@@ -1625,8 +1337,6 @@ bool ARecastNavMesh::TestPath(const FNavAgentProperties& AgentProperties, const
 	bool bPathExists = true;
 	if ((Query.StartLocation - Query.EndLocation).IsNearlyZero() == false)
 	{
-		SECTION_LOCK_TILES_FOR(RecastNavMesh);
-
 		ENavigationQueryResult::Type Result = RecastNavMesh->RecastNavMeshImpl->TestPath(Query.StartLocation, Query.EndLocation, *(Query.QueryFilter.Get()), Query.Owner.Get(), NumVisitedNodes);
 		bPathExists = (Result == ENavigationQueryResult::Success);
 	}
@@ -1650,8 +1360,6 @@ bool ARecastNavMesh::TestHierarchicalPath(const FNavAgentProperties& AgentProper
 
 	if ((Query.StartLocation - Query.EndLocation).IsNearlyZero() == false)
 	{
-		SECTION_LOCK_TILES_FOR(RecastNavMesh);
-
 		bool bUseFallbackSearch = false;
 		if (bCanUseHierachicalPath)
 		{
@@ -1690,10 +1398,7 @@ bool ARecastNavMesh::NavMeshRaycast(const ANavigationData* Self, const FVector&
 		return true;
 	}
 
-	{
-		SECTION_LOCK_TILES_FOR(RecastNavMesh);
-		RecastNavMesh->RecastNavMeshImpl->Raycast2D(RayStart, RayEnd, RecastNavMesh->GetRightFilterRef(QueryFilter), QueryOwner, Result);
-	}
+	RecastNavMesh->RecastNavMeshImpl->Raycast2D(RayStart, RayEnd, RecastNavMesh->GetRightFilterRef(QueryFilter), QueryOwner, Result);
 
 	HitLocation = Result.HasHit() ? (RayStart + (RayEnd - RayStart) * Result.HitTime) : RayEnd;
 
@@ -1712,10 +1417,7 @@ bool ARecastNavMesh::NavMeshRaycast(const ANavigationData* Self, NavNodeRef RayS
 	}
 
 	FRaycastResult Result;
-	{
-		SECTION_LOCK_TILES_FOR(RecastNavMesh);
-		RecastNavMesh->RecastNavMeshImpl->Raycast2D(RayStartNode, RayStart, RayEnd, RecastNavMesh->GetRightFilterRef(QueryFilter), QueryOwner, Result);
-	}
+	RecastNavMesh->RecastNavMeshImpl->Raycast2D(RayStartNode, RayStart, RayEnd, RecastNavMesh->GetRightFilterRef(QueryFilter), QueryOwner, Result);
 
 	HitLocation = Result.HasHit() ? (RayStart + (RayEnd - RayStart) * Result.HitTime) : RayEnd;
 
@@ -1729,8 +1431,6 @@ void ARecastNavMesh::BatchRaycast(TArray<FNavigationRaycastWork>& Workload, TSha
 		return;
 	}
 
-	SECTION_LOCK_TILES;
-
 	const FNavigationQueryFilter& FilterToUse = GetRightFilterRef(Filter);
 
 	FRecastSpeciaLinkFilter LinkFilter(UNavigationSystem::GetCurrent(GetWorld()), Querier);
@@ -1779,10 +1479,7 @@ bool ARecastNavMesh::IsSegmentOnNavmesh(const FVector& SegmentStart, const FVect
 	}
 	
 	FRaycastResult Result;
-	{
-		SECTION_LOCK_TILES;
-		RecastNavMeshImpl->Raycast2D(SegmentStart, SegmentEnd, GetRightFilterRef(Filter), QueryOwner, Result);
-	}
+	RecastNavMeshImpl->Raycast2D(SegmentStart, SegmentEnd, GetRightFilterRef(Filter), QueryOwner, Result);
 
 	return Result.HasHit() == false;
 }
@@ -1792,8 +1489,6 @@ bool ARecastNavMesh::FindStraightPath(const FVector& StartLoc, const FVector& En
 	bool bResult = false;
 	if (RecastNavMeshImpl)
 	{
-		SECTION_LOCK_TILES;
-
 		bResult = RecastNavMeshImpl->FindStraightPath(StartLoc, EndLoc, PathCorridor, PathPoints, CustomLinks);
 	}
 
@@ -1815,108 +1510,15 @@ int32 ARecastNavMesh::DebugPathfinding(const FPathFindingQuery& Query, TArray<FR
 
 	if ((Query.StartLocation - Query.EndLocation).IsNearlyZero() == false)
 	{
-		SECTION_LOCK_TILES_FOR(RecastNavMesh);
-
 		NumSteps = RecastNavMesh->RecastNavMeshImpl->DebugPathfinding(Query.StartLocation, Query.EndLocation, *(Query.QueryFilter.Get()), Query.Owner.Get(), Steps);
 	}
 
 	return NumSteps;
 }
 
-void ARecastNavMesh::UpdateNavTileData(uint8* OldTileData, FNavMeshTileData NewTileData)
-{
-	check(IsInGameThread());
-	static const float SecondsToHoldOldData = 1.f;
-
-	if (OldTileData != NULL)
-	{
-		int32 OldDataIndex = INDEX_NONE;
-		int32 OldDataSize = 0;
-
-		// find shared pointer holding OldTileData
-		const FNavMeshTileData* const RESTRICT DataEnd = TileNavDataContainer.GetData() + TileNavDataContainer.Num();
-		for(const FNavMeshTileData* RESTRICT Data = TileNavDataContainer.GetData(); Data < DataEnd; ++Data)
-		{
-			if( *Data == OldTileData )
-			{
-				OldDataIndex = (int32)(Data - TileNavDataContainer.GetData());
-				OldDataSize = Data->DataSize;
-				break;
-			}
-		}
-
-		ensure(OldDataIndex != INDEX_NONE);
-		if (OldDataIndex != INDEX_NONE)
-		{
-			DEC_MEMORY_STAT_BY(STAT_Navigation_NavDataMemory, OldDataSize);
-
-			//	remove from list of referenced memory.
-			// this should result in releasing OldTileData (unless it's being 
-			//	referenced by some other shared pointer)
-			//	bAllowShrinking=false since TileNavDataContainer will consistently retain 
-			// its size as long as navmesh has the same number of tiles
-			TileNavDataContainer.RemoveAtSwap(OldDataIndex, 1, /*bAllowShrinking=*/false);
-		}
-	}
-
-	if (NewTileData.IsValid()) 
-	{
-		INC_MEMORY_STAT_BY(STAT_Navigation_NavDataMemory, NewTileData.DataSize);
-		TileNavDataContainer.Add(NewTileData);
-	}
-}
-
-void ARecastNavMesh::ReserveTileSet(int32 NewGridWidth, int32 NewGridHeight)
-{
-	const int32 NewSize = NewGridWidth * NewGridHeight;	
-
-	if (NewGridWidth != GridWidth || NewGridHeight != GridHeight)
-	{
-		TileSet.Reset();		
-		if (TileSet.Max() < NewSize)
-		{
-			TileSet.Reserve(NewSize);
-		}
-
-		for (int32 Index = 0; Index < NewSize; ++Index)
-		{
-			const FVector2D Location = TileIndexToXY(NewGridWidth, Index);
-			TileSet.Add(FTileSetItem(Location.X, Location.Y, Index));
-		}
-		
-		GridWidth = NewGridWidth;
-		GridHeight = NewGridHeight;
-	}
-}
-
-FTileSetItem* ARecastNavMesh::GetTileSetItemAt(int32 TileX, int32 TileY)
-{
-	return &TileSet[TileXYToIndex(GridWidth, TileX, TileY)];
-}
-
-void ARecastNavMesh::UpdateNavmeshOffset( const FBox& NavBounds )
-{
-	if (RecastNavMeshImpl != NULL)
-	{
-		if (dtNavMesh* RecastNavMesh = RecastNavMeshImpl->GetRecastMesh())
-		{
-			const FBox RCNavBounds = Unreal2RecastBox(NavBounds);
-			const float moveBy[3] = { 
-				RCNavBounds.Min.X - RecastNavMeshImpl->GetRecastMesh()->getParams()->orig[0],
-				RCNavBounds.Min.Y - RecastNavMeshImpl->GetRecastMesh()->getParams()->orig[1],
-				RCNavBounds.Min.Z - RecastNavMeshImpl->GetRecastMesh()->getParams()->orig[2]
-			};
-			RecastNavMesh->applyWorldOffset(&moveBy[0]);
-		}
-	}
-}
-
-
 void ARecastNavMesh::UpdateNavVersion() 
 { 
-#if WITH_NAVIGATION_GENERATOR
 	NavMeshVersion = NAVMESHVER_LATEST; 
-#endif // WITH_NAVIGATION_GENERATOR
 }
 
 #if WITH_EDITOR
@@ -1952,13 +1554,7 @@ void ARecastNavMesh::PostEditChangeProperty( struct FPropertyChangedEvent& Prope
 				}
 			}
 
-#if WITH_NAVIGATION_GENERATOR
-			FNavDataGenerator* Generator = GetGenerator(FNavigationSystem::DontCreate);
-			if (Generator != NULL)
-			{
-				Generator->Generate();
-			}
-#endif // WITH_NAVIGATION_GENERATOR
+			RebuildAll();
 		}
 		else if (FObjectEditorUtils::GetCategoryFName(PropertyChangedEvent.Property) == NAME_Display)
 		{
@@ -1972,17 +1568,27 @@ void ARecastNavMesh::PostEditChangeProperty( struct FPropertyChangedEvent& Prope
 bool ARecastNavMesh::NeedsRebuild() const
 {
 	bool bLooksLikeNeeded = !RecastNavMeshImpl || RecastNavMeshImpl->GetRecastMesh() == 0;
-#if WITH_NAVIGATION_GENERATOR
-	FRecastNavMeshGenerator* MyGenerator = NavDataGenerator.IsValid() ? (FRecastNavMeshGenerator*)NavDataGenerator.Get() : NULL;
-	if (MyGenerator)
+	if (NavDataGenerator)
 	{
-		return bLooksLikeNeeded || MyGenerator->HasDirtyTiles() || MyGenerator->HasDirtyAreas();
+		return bLooksLikeNeeded || NavDataGenerator->GetNumRemaningBuildTasks() > 0;
 	}
-#endif // WITH_NAVIGATION_GENERATOR
 
 	return bLooksLikeNeeded;
 }
 
+void ARecastNavMesh::RebuildAll()
+{
+	Super::RebuildAll();
+
+	NavDataGenerator.Reset();
+	// Generator should be enabled in the editor and if navmesh supports runtime generation
+	if (bRebuildAtRuntime || (GetWorld() && !GetWorld()->IsGameWorld()))
+	{
+		NavDataGenerator.Reset(new FRecastNavMeshGenerator(this));
+		NavDataGenerator->RebuildAll();
+	}
+}
+
 bool ARecastNavMesh::IsVoxelCacheEnabled()
 {
 #if RECAST_ASYNC_REBUILDING
@@ -2021,3 +1627,59 @@ bool ARecastNavMesh::HasValidNavmesh() const
 
 	return false;
 }
+
+FRecastNavMeshCachedData FRecastNavMeshCachedData::Construct(const ARecastNavMesh* RecastNavMeshActor)
+{
+	check(RecastNavMeshActor);
+	
+	FRecastNavMeshCachedData CachedData;
+
+	CachedData.ActorOwner = RecastNavMeshActor;
+	// create copies from crucial ARecastNavMesh data
+	CachedData.bUseSortFunction = RecastNavMeshActor->bSortNavigationAreasByCost;
+
+	TArray<FSupportedAreaData> Areas;
+	RecastNavMeshActor->GetSupportedAreas(Areas);
+	FMemory::Memzero(CachedData.FlagsPerArea, sizeof(ARecastNavMesh::FNavPolyFlags) * RECAST_MAX_AREAS);
+
+	for (int32 i = 0; i < Areas.Num(); i++)
+	{
+		const UClass* AreaClass = Areas[i].AreaClass;
+		const UNavArea* DefArea = AreaClass ? ((UClass*)AreaClass)->GetDefaultObject<UNavArea>() : NULL;
+		if (DefArea)
+		{
+			CachedData.AreaClassToIdMap.Add(AreaClass, Areas[i].AreaID);
+			CachedData.FlagsPerArea[Areas[i].AreaID] = DefArea->GetAreaFlags();
+		}
+	}
+
+	FMemory::Memcpy(CachedData.FlagsPerOffMeshLinkArea, CachedData.FlagsPerArea, sizeof(CachedData.FlagsPerArea));
+	static const ARecastNavMesh::FNavPolyFlags NavLinkFlag = ARecastNavMesh::GetNavLinkFlag();
+	if (NavLinkFlag != 0)
+	{
+		ARecastNavMesh::FNavPolyFlags* AreaFlag = CachedData.FlagsPerOffMeshLinkArea;
+		for (int32 AreaIndex = 0; AreaIndex < RECAST_MAX_AREAS; ++AreaIndex, ++AreaFlag)
+		{
+			*AreaFlag |= NavLinkFlag;
+		}
+	}
+
+	return CachedData;
+}
+
+void FRecastNavMeshCachedData::OnAreaAdded(const UClass* AreaClass, int32 AreaID)
+{
+	const UNavArea* DefArea = AreaClass ? ((UClass*)AreaClass)->GetDefaultObject<UNavArea>() : NULL;
+	if (DefArea && AreaID >= 0)
+	{
+		AreaClassToIdMap.Add(AreaClass, AreaID);
+		FlagsPerArea[AreaID] = DefArea->GetAreaFlags();
+
+		static const ARecastNavMesh::FNavPolyFlags NavLinkFlag = ARecastNavMesh::GetNavLinkFlag();
+		if (NavLinkFlag != 0)
+		{
+			FlagsPerOffMeshLinkArea[AreaID] = FlagsPerArea[AreaID] | NavLinkFlag;
+		}
+	}		
+}
+
diff --git a/Engine/Source/Runtime/Engine/Private/AI/Navigation/RecastNavMeshGenerator.cpp b/Engine/Source/Runtime/Engine/Private/AI/Navigation/RecastNavMeshGenerator.cpp
--- a/Engine/Source/Runtime/Engine/Private/AI/Navigation/RecastNavMeshGenerator.cpp
+++ b/Engine/Source/Runtime/Engine/Private/AI/Navigation/RecastNavMeshGenerator.cpp
@@ -28,9 +28,6 @@
 #define GENERATE_SEGMENT_LINKS 1
 #define GENERATE_CLUSTER_LINKS 1
 
-/** if set, each tile generator will allocate copy of geometry data, to be used for async processing */
-#define CACHE_NAV_GENERATOR_DATA	RECAST_ASYNC_REBUILDING
-
 #define SHOW_NAV_EXPORT_PREVIEW 0
 
 #define TEXT_WEAKOBJ_NAME(obj) (obj.IsValid(false) ? *obj->GetName() : (obj.IsValid(false, true)) ? TEXT("MT-Unreachable") : TEXT("INVALID"))
@@ -42,137 +39,6 @@
 #define RECAST_STAT(...) 
 #endif
 
-#if !NAVOCTREE_CONTAINS_COLLISION_DATA
-#error FRecastNavMeshGenerator requires NavOctree to contain additional data
-#endif
-
-/** structure to cache owning RecastNavMesh data so that it doesn't have to be polled
- *	directly from RecastNavMesh while asyncronously generating navmesh */
-struct FRecastNavMeshCachedData
-{
-	ARecastNavMesh::FNavPolyFlags FlagsPerArea[RECAST_MAX_AREAS];
-	ARecastNavMesh::FNavPolyFlags FlagsPerOffMeshLinkArea[RECAST_MAX_AREAS];
-	TMap<const UClass*, int32> AreaClassToIdMap;
-	const ARecastNavMesh* ActorOwner;
-	uint32 bUseSortFunction : 1;
-
-	FRecastNavMeshCachedData(const ARecastNavMesh* RecastNavMeshActor) : ActorOwner(RecastNavMeshActor)
-	{
-		check(RecastNavMeshActor);
-
-		if (RecastNavMeshActor != NULL)
-		{
-			// create copies from crucial ARecastNavMesh data
-			bUseSortFunction = RecastNavMeshActor->bSortNavigationAreasByCost;
-
-			TArray<FSupportedAreaData> Areas;
-			RecastNavMeshActor->GetSupportedAreas(Areas);
-			FMemory::Memzero(FlagsPerArea, sizeof(ARecastNavMesh::FNavPolyFlags) * RECAST_MAX_AREAS);
-
-			for (int32 i = 0; i < Areas.Num(); i++)
-			{
-				const UClass* AreaClass = Areas[i].AreaClass;
-				const UNavArea* DefArea = AreaClass ? ((UClass*)AreaClass)->GetDefaultObject<UNavArea>() : NULL;
-				if (DefArea)
-				{
-					AreaClassToIdMap.Add(AreaClass, Areas[i].AreaID);
-					FlagsPerArea[Areas[i].AreaID] = DefArea->GetAreaFlags();
-				}
-			}
-
-			FMemory::Memcpy(FlagsPerOffMeshLinkArea, FlagsPerArea, sizeof(FlagsPerArea));
-			static const ARecastNavMesh::FNavPolyFlags NavLinkFlag = ARecastNavMesh::GetNavLinkFlag();
-			if (NavLinkFlag != 0)
-			{
-				ARecastNavMesh::FNavPolyFlags* AreaFlag = FlagsPerOffMeshLinkArea;
-				for (int32 AreaIndex = 0; AreaIndex < RECAST_MAX_AREAS; ++AreaIndex, ++AreaFlag)
-				{
-					*AreaFlag |= NavLinkFlag;
-				}
-			}
-		}
-	}
-
-	void OnAreaAdded(const UClass* AreaClass, int32 AreaID)
-	{
-		const UNavArea* DefArea = AreaClass ? ((UClass*)AreaClass)->GetDefaultObject<UNavArea>() : NULL;
-		if (DefArea && AreaID >= 0)
-		{
-			AreaClassToIdMap.Add(AreaClass, AreaID);
-			FlagsPerArea[AreaID] = DefArea->GetAreaFlags();
-
-			static const ARecastNavMesh::FNavPolyFlags NavLinkFlag = ARecastNavMesh::GetNavLinkFlag();
-			if (NavLinkFlag != 0)
-			{
-				FlagsPerOffMeshLinkArea[AreaID] = FlagsPerArea[AreaID] | NavLinkFlag;
-			}
-		}		
-	}
-};
-
-
-//
-class FAsyncNavTileBuildWorker : public FNonAbandonableTask 
-{
-public:
-	/**
-	 * Initializes the data and creates the async compression task.
-	 */
-	FAsyncNavTileBuildWorker(const ANavigationData::FNavDataGeneratorSharedPtr& InNavMeshGenerator, const int32 InTileId, const uint32 InVersion)
-		: NavMeshGenerator(InNavMeshGenerator), TileId(InTileId), Version(InVersion)
-	{
-	}
-
-	~FAsyncNavTileBuildWorker()
-	{
-	}
-
-	/**
-	 * Compresses the texture
-	 */
-	FORCEINLINE_DEBUGGABLE void DoWork()
-	{
-//		UE_LOG(LogNavigation, Warning, TEXT("FAsyncNavTileBuildWorker::DoWork %d"), TileId);
-		if (NavMeshGenerator.IsValid())
-		{
-			FRecastNavMeshGenerator* const RecastNavMeshGenerator = static_cast<FRecastNavMeshGenerator*>(NavMeshGenerator.Get());
-			SCOPE_CYCLE_COUNTER(STAT_Navigation_BuildTime)
-			STAT(double ThisTime = 0);
-			{
-				SCOPE_SECONDS_COUNTER(ThisTime);
-				RecastNavMeshGenerator->GenerateTile(TileId, Version);	
-			
-				DECLARE_CYCLE_STAT(TEXT("FSimpleDelegateGraphTask.Requesting to build next tile if necessary"),
-					STAT_FSimpleDelegateGraphTask_RequestingToBuildNextTileIfNecessary,
-					STATGROUP_TaskGraphTasks);
-
-				FSimpleDelegateGraphTask::CreateAndDispatchWhenReady(
-					FSimpleDelegateGraphTask::FDelegate::CreateThreadSafeSP(RecastNavMeshGenerator, &FRecastNavMeshGenerator::UpdateTileGenerationWorkers, TileId),
-					GET_STATID(STAT_FSimpleDelegateGraphTask_RequestingToBuildNextTileIfNecessary), NULL, ENamedThreads::GameThread);
-			}
-			INC_FLOAT_STAT_BY(STAT_Navigation_CumulativeBuildTime,(float)ThisTime*1000);
-		}
-	}
-
-	/** 
-	 * Give the name for external event viewers
-	 * @return	the name to display in external event viewers
-	 */
-	static const TCHAR* Name()
-	{
-		return TEXT("FAsyncNavTileBuildWorker");
-	}
-
-private:
-	ANavigationData::FNavDataGeneratorSharedPtr NavMeshGenerator;
-	const int32 TileId;
-	const uint32 Version;
-};
-
-// Consts
-static const int32 NAVMESH_TILE_GENERATOR_OWNS_DATA = 0; // @see DT_TILE_FREE_DATA
-const float FRecastNavMeshGenerator::DefaultFreshness = 0.5f;
-
 FORCEINLINE bool DoesBoxContainOrOverlapVector(const FBox& BigBox, const FVector& In)
 {
 	return (In.X >= BigBox.Min.X) && (In.X <= BigBox.Max.X) 
@@ -185,6 +51,24 @@ FORCEINLINE bool DoesBoxContainBox(const FBox& BigBox, const FBox& SmallBox)
 	return DoesBoxContainOrOverlapVector(BigBox, SmallBox.Min) && DoesBoxContainOrOverlapVector(BigBox, SmallBox.Max);
 }
 
+int32 GetTilesCountHelper(const class dtNavMesh* DetourMesh)
+{
+	int32 NumTiles = 0;
+	if (DetourMesh)
+	{
+		for (int32 i = 0; i < DetourMesh->getMaxTiles(); i++)
+		{
+			const dtMeshTile* TileData = DetourMesh->getTile(i);
+			if (TileData && TileData->header && TileData->dataSize > 0)
+			{
+				NumTiles++;
+			}
+		}
+	}
+
+	return NumTiles;
+}
+
 /**
  * Exports geometry to OBJ file. Can be used to verify NavMesh generation in RecastDemo app
  * @param FileName - full name of OBJ file with extension
@@ -1485,129 +1369,6 @@ struct FVoxelCacheRasterizeContext
 
 static FVoxelCacheRasterizeContext VoxelCacheContext;
 
-//----------------------------------------------------------------------//
-// FRecastTileDirtyState
-//----------------------------------------------------------------------//
-
-FRecastTileDirtyState::FRecastTileDirtyState(const class FRecastTileGenerator* DirtyGenerator)
-{
-	if (DirtyGenerator)
-	{
-		DirtyGenerator->GetDirtyState(*this);
-	}
-}
-
-void FRecastTileDirtyState::Append(const FRecastTileDirtyState& Other)
-{
-	bRebuildGeometry |= Other.bRebuildGeometry;
-	bRebuildLayers |= Other.bRebuildLayers;
-	bRebuildAllLayers |= Other.bRebuildAllLayers;
-
-	for (int32 i = 0; i < ARRAY_COUNT(DirtyLayers); i++)
-	{
-		DirtyLayers[i] |= Other.DirtyLayers[i];
-	}
-
-	if (FallbackDirtyLayers.Num() < Other.FallbackDirtyLayers.Num())
-	{
-		FallbackDirtyLayers.AddZeroed(Other.FallbackDirtyLayers.Num() - FallbackDirtyLayers.Num());
-	}
-
-	for (int32 i = 0; i < Other.FallbackDirtyLayers.Num(); i++)
-	{
-		FallbackDirtyLayers[i] |= Other.FallbackDirtyLayers[i];
-	}
-}
-
-void FRecastTileDirtyState::Clear()
-{
-	bRebuildGeometry = false;
-	bRebuildLayers = false;
-	bRebuildAllLayers = false;
-	FMemory::Memzero(DirtyLayers, sizeof(DirtyLayers));
-
-	const int32 NumLayers = FallbackDirtyLayers.Num();
-	if (NumLayers > 0)
-	{
-		FallbackDirtyLayers.Reset();
-		FallbackDirtyLayers.AddZeroed(NumLayers);
-	}
-}
-
-void FRecastTileDirtyState::MarkDirtyLayer(int32 LayerIdx)
-{
-	const int32 BitfieldSize = 8;
-	const int32 MaxBitfieldLayers = BitfieldSize * ARRAY_COUNT(DirtyLayers);
-
-	bRebuildLayers = true;
-	if (LayerIdx >= MaxBitfieldLayers)
-	{
-		// consider increasing size of DirtyLayers if this happens frequently
-		const int32 ArrayIdx = (LayerIdx - MaxBitfieldLayers) / BitfieldSize;
-		const int32 ShiftIdx = (LayerIdx - MaxBitfieldLayers) % BitfieldSize;
-
-		if (!FallbackDirtyLayers.IsValidIndex(ArrayIdx))
-		{
-			FallbackDirtyLayers.AddZeroed(ArrayIdx - FallbackDirtyLayers.Num() + 1);
-		}
-		
-		FallbackDirtyLayers[ArrayIdx] = FallbackDirtyLayers[ArrayIdx] | (1 << ShiftIdx);
-	}
-	else
-	{
-		const int32 ArrayIdx = LayerIdx / BitfieldSize;
-		const int32 ShiftIdx = LayerIdx % BitfieldSize;
-
-		DirtyLayers[ArrayIdx] = DirtyLayers[ArrayIdx] | (1 << ShiftIdx);
-	}
-}
-
-bool FRecastTileDirtyState::HasDirtyLayer(int32 LayerIdx) const
-{
-	const int32 BitfieldSize = 8;
-	const int32 MaxBitfieldLayers = BitfieldSize * ARRAY_COUNT(DirtyLayers);
-
-	bool bIsDirty = false;
-	if (bRebuildLayers)
-	{
-		if (bRebuildAllLayers)
-		{
-			bIsDirty = true;
-		}
-		else if (LayerIdx >= MaxBitfieldLayers)
-		{
-			const int32 ArrayIdx = (LayerIdx - MaxBitfieldLayers) / BitfieldSize;
-			const int32 ShiftIdx = (LayerIdx - MaxBitfieldLayers) % BitfieldSize;
-
-			if (FallbackDirtyLayers.IsValidIndex(ArrayIdx))
-			{
-				bIsDirty = (FallbackDirtyLayers[ArrayIdx] & (1 << ShiftIdx)) != 0;
-			}
-		}
-		else
-		{
-			const int32 ArrayIdx = LayerIdx / BitfieldSize;
-			const int32 ShiftIdx = LayerIdx % BitfieldSize;
-
-			bIsDirty = (DirtyLayers[ArrayIdx] & (1 << ShiftIdx)) != 0;
-		}
-	}
-
-	return bIsDirty;
-}
-
-
-//----------------------------------------------------------------------//
-// FRecastTileGenerator
-//----------------------------------------------------------------------//
-
-TNavStatArray<rcSpanCache> FRecastTileGenerator::StaticGeomSpans;
-TNavStatArray<float> FRecastTileGenerator::StaticGeomCoords;
-TNavStatArray<int32> FRecastTileGenerator::StaticGeomIndices;
-TArray<FAreaNavModifier> FRecastTileGenerator::StaticStaticAreas;
-TArray<FAreaNavModifier> FRecastTileGenerator::StaticDynamicAreas;
-TArray<FSimpleLinkNavModifier> FRecastTileGenerator::StaticOffmeshLinks;
-
 uint32 GetTileCacheSizeHelper(TArray<FNavMeshTileData>& CompressedTiles)
 {
 	uint32 TotalMemory = 0;
@@ -1619,57 +1380,51 @@ uint32 GetTileCacheSizeHelper(TArray<FNavMeshTileData>& CompressedTiles)
 	return TotalMemory;
 }
 
-FRecastTileGenerator::FRecastTileGenerator()
-	: bInitialized(false), bBeingRebuild(false), bRebuildPending(false), bAsyncBuildInProgress(false), TileX(-1), TileY(-1), TileId(-1), Version(uint32(-1))
-{
-}
-
-FRecastTileGenerator::~FRecastTileGenerator()
-{
-	DEC_MEMORY_STAT_BY(STAT_Navigation_TileCacheMemory, GetTileCacheSizeHelper(CompressedLayers));
-}
+//----------------------------------------------------------------------//
+// FRecastTileGenerator
+//----------------------------------------------------------------------//
 
-void FRecastTileGenerator::Init(class FRecastNavMeshGenerator* ParentGenerator, const int32 X, const int32 Y, const float TileBmin[3], const float TileBmax[3], const TNavStatArray<FBox>& BoundingBoxes)
+FRecastTileGenerator::FRecastTileGenerator(
+	const FRecastNavMeshGenerator* ParentGenerator,
+	const int32 X, 
+	const int32 Y,
+	TArray<FBox> DirtyAreas
+	)
 {
-	if (ParentGenerator == NULL)
-	{
-		UE_LOG(LogNavigation, Error, TEXT("FRecastTileGenerator: trying to initialize tile generator with empty navmesh generator"))
-		return;
-	}
-
-	AdditionalCachedData = ParentGenerator->GetAdditionalCachedData();
-
-	if (NavMeshGenerator.HasSameObject(ParentGenerator) == false)
-	{
-		NavMeshGenerator = ParentGenerator->AsShared();
-	}
-	TileId = ParentGenerator->GetTileIdAt(X,Y);
-	Version = ParentGenerator->GetVersion();
-
+	bSucceeded = false;
+	
 	TileX = X;
 	TileY = Y;
-	FMemory::Memcpy( BMin, TileBmin, sizeof( BMin ) );
-	FMemory::Memcpy( BMax, TileBmax, sizeof( BMin ) );
-	TileBB = Recast2UnrealBox(BMin, BMax);
-
+		
+	TileConfig = ParentGenerator->GetConfig();
+	Version = ParentGenerator->GetVersion();
+	AdditionalCachedData = ParentGenerator->GetAdditionalCachedData();
+	const FVector NavMeshOrigin = FVector::ZeroVector;
+	const FBox NavTotalBounds = ParentGenerator->GetTotalBounds();
+	const float TileCellSize = (TileConfig.tileSize * TileConfig.cs);
+	const FVector TileBMin = FVector((X+0) * TileCellSize, (Y+0) * TileCellSize, NavTotalBounds.Min.Z) + NavMeshOrigin;
+	const FVector TileBMax = FVector((X+1) * TileCellSize, (Y+1) * TileCellSize, NavTotalBounds.Max.Z) + NavMeshOrigin;
+	TileBB = FBox(TileBMin, TileBMax);
+
+	FBox RCBox = Unreal2RecastBox(TileBB);
+	rcVcopy(TileConfig.bmin, &RCBox.Min.X);
+	rcVcopy(TileConfig.bmax, &RCBox.Max.X);
+			
 	// from passed in boxes pick the ones overlapping with tile bounds
 	InclusionBounds.Reset();
-
 	bOutsideOfInclusionBounds = false;
-
-	if (BoundingBoxes.Num() > 0)
+	const auto& ParentBounds = ParentGenerator->GetInclusionBounds();
+	if (ParentBounds.Num() > 0)
 	{
 		bFullyEncapsulatedByInclusionBounds = false;
 
-		InclusionBounds.Reserve(BoundingBoxes.Num());
-
-		const FBox* Bounds = BoundingBoxes.GetData();	
-		for (int32 i = 0; i < BoundingBoxes.Num() && bFullyEncapsulatedByInclusionBounds == false; ++i, ++Bounds)
+		InclusionBounds.Reserve(ParentBounds.Num());
+		const FBox* Bounds = ParentBounds.GetData();	
+		for (int32 i = 0; i < ParentBounds.Num() && bFullyEncapsulatedByInclusionBounds == false; ++i, ++Bounds)
 		{	
-			if (Bounds->Intersect( TileBB ))
+			if (Bounds->Intersect(TileBB))
 			{
 				InclusionBounds.Add(*Bounds);
-
 				bFullyEncapsulatedByInclusionBounds = DoesBoxContainBox(*Bounds, TileBB);
 			}
 		}
@@ -1680,59 +1435,101 @@ void FRecastTileGenerator::Init(class FRecastNavMeshGenerator* ParentGenerator,
 	{
 		bFullyEncapsulatedByInclusionBounds = true;
 	}
+	
+	// In case we don't need to regenerate tile from scratch
+	// get compressed layers data from cache
+	if (DirtyAreas.Num() > 0)
+	{
+		CompressedLayers = ParentGenerator->GetCachedLayersData(FIntPoint(TileX, TileY));
+	}
+	
+	// We have to regenerate layers data in case it's missing
+	bRegenerateCompressedLayers = (CompressedLayers.Num() == 0);
 
-	FRecastBuildConfig TileConfig = ParentGenerator->GetConfig();
-	WalkableClimbVX = TileConfig.walkableClimb;
-	WalkableSlopeCos = FMath::Cos(FMath::DegreesToRadians(TileConfig.walkableSlopeAngle));
-	RasterizationPadding = TileConfig.borderSize * TileConfig.cs;
-	bInitialized = true;
-}
+	if (!bRegenerateCompressedLayers)
+	{
+		// Mark layers that needs to be updated
+		DirtyLayers.Init(false, CompressedLayers.Num());
+		for (const FNavMeshTileData& LayerData : CompressedLayers)
+		{
+			for (FBox DirtyBox : DirtyAreas)
+			{
+				if (DirtyBox.Intersect(LayerData.LayerBBox))
+				{
+					DirtyLayers[LayerData.LayerIndex] = true;
+				}
+			}
+		}
+	}
 
-void FRecastTileGenerator::InitiateRebuild()
-{
-	bBeingRebuild = true;
-	bRebuildPending = false;
-	GeneratingState = DirtyState;
-	DirtyState.Clear();
-}
+	GatherGeometry(ParentGenerator, bRegenerateCompressedLayers);
 
-void FRecastTileGenerator::AbortRebuild()
-{
-	bBeingRebuild = false;
-	bRebuildPending = false;
-	DirtyState.Append(GeneratingState);
-	GeneratingState.Clear();
-	ClearNavigationData();
+	//
+	UsedMemoryOnStartup = GetUsedMemCount() + sizeof(FRecastTileGenerator);
 }
 
-void FRecastTileGenerator::AbadonGeneration()
+void FRecastTileGenerator::DoWork()
 {
-	bBeingRebuild = false;
-	bRebuildPending = false;
-	DirtyState.Clear();
-	GeneratingState.Clear();
+	bSucceeded = GenerateTile();
 }
 
-void FRecastTileGenerator::FinishRebuild()
+const TCHAR* FRecastTileGenerator::Name()
 {
-	bBeingRebuild = false;
-	GeneratingState.Clear();
-	DirtyState.Append(GeneratingState);
+	return TEXT("FRecastTileGenerator");
 }
 
-void FRecastTileGenerator::StartAsyncBuild()
+void FRecastTileGenerator::GatherGeometry(const FRecastNavMeshGenerator* ParentGenerator, bool bGeometryChanged)
 {
-	bAsyncBuildInProgress = true;
-}
+	const UNavigationSystem*	NavSys = UNavigationSystem::GetCurrent(ParentGenerator->GetWorld());
+	const FNavigationOctree*	NavOctree = NavSys ? NavSys->GetNavOctree() : nullptr;
+	const FNavDataConfig*		NavDataConfig = &ParentGenerator->GetOwner()->NavDataConfig;
+	
+	for(FNavigationOctree::TConstElementBoxIterator<FNavigationOctree::DefaultStackAllocator> It(*NavOctree, ParentGenerator->GrowBoundingBox(TileBB, /*bIncludeAgentHeight*/ false));
+		It.HasPendingElements();
+		It.Advance())
+	{
+		const FNavigationOctreeElement& Element = It.GetCurrentElement();
+		const bool bShouldUse = Element.ShouldUseGeometry(NavDataConfig);
+		if (bShouldUse)
+		{
+			const bool bExportGeometry = bGeometryChanged && Element.Data.HasGeometry();
+			if (bExportGeometry)
+			{
+				if (ARecastNavMesh::IsVoxelCacheEnabled())
+				{
+					TNavStatArray<rcSpanCache> SpanData;
+					rcSpanCache* CachedVoxels = 0;
+					int32 NumCachedVoxels = 0;
 
-void FRecastTileGenerator::FinishAsyncBuild()
-{
-	bAsyncBuildInProgress = false;
-}
+					DECLARE_SCOPE_CYCLE_COUNTER(TEXT("Rasterization: prepare voxel cache"), Stat_RecastRasterCachePrep, STATGROUP_Navigation);
 
-void FRecastTileGenerator::MarkPendingRebuild()
-{
-	bRebuildPending = true;
+					if (!HasVoxelCache(Element.Data.VoxelData, CachedVoxels, NumCachedVoxels))
+					{
+						// rasterize
+						PrepareVoxelCache(Element.Data.CollisionData, SpanData);
+						CachedVoxels = SpanData.GetData();
+						NumCachedVoxels = SpanData.Num();
+
+						// encode
+						const int32 PrevElementMemory = Element.Data.GetAllocatedSize();
+						FNavigationRelevantData* ModData = (FNavigationRelevantData*)&Element.Data;
+						AddVoxelCache(ModData->VoxelData, CachedVoxels, NumCachedVoxels);
+
+						const int32 NewElementMemory = Element.Data.GetAllocatedSize();
+						const int32 ElementMemoryDelta = NewElementMemory - PrevElementMemory;
+						INC_MEMORY_STAT_BY(STAT_Navigation_CollisionTreeMemory, ElementMemoryDelta);
+					}
+				}
+				else
+				{
+					AppendGeometry(Element.Data.CollisionData);
+				}
+			}
+
+			const FCompositeNavModifier ModifierInstance = Element.GetModifierForAgent(NavDataConfig);
+			AppendModifier(ModifierInstance, bExportGeometry);
+		}
+	}
 }
 
 void FRecastTileGenerator::ApplyVoxelFilter(struct rcHeightfield* HF, float WalkableRadius)
@@ -1846,8 +1643,14 @@ void FRecastTileGenerator::ApplyVoxelFilter(struct rcHeightfield* HF, float Walk
 
 void FRecastTileGenerator::PrepareVoxelCache(const TNavStatArray<uint8>& RawCollisionCache, TNavStatArray<rcSpanCache>& SpanData)
 {
+	// tile's geometry: voxel cache (only for synchronous rebuilds)
+	const int32 WalkableClimbVX = TileConfig.walkableClimb;
+	const float WalkableSlopeCos = FMath::Cos(FMath::DegreesToRadians(TileConfig.walkableSlopeAngle));
+	const float RasterizationPadding = TileConfig.borderSize * TileConfig.cs;
+
 	FRecastGeometryCache CachedCollisions(RawCollisionCache.GetData());
-	VoxelCacheContext.SetupForTile(BMin, BMax, RasterizationPadding);
+
+	VoxelCacheContext.SetupForTile(TileConfig.bmin, TileConfig.bmax, RasterizationPadding);
 
 	float SlopeCosPerActor = WalkableSlopeCos;
 	CachedCollisions.Header.SlopeOverride.ModifyWalkableFloorZ(SlopeCosPerActor);
@@ -1885,7 +1688,7 @@ bool FRecastTileGenerator::HasVoxelCache(const TNavStatArray<uint8>& RawVoxelCac
 			return true;
 		}
 	}
-
+	
 	return false;
 }
 
@@ -1913,80 +1716,27 @@ void FRecastTileGenerator::AddVoxelCache(TNavStatArray<uint8>& RawVoxelCache, co
 	FMemory::Memcpy(RawVoxelCache.GetData() + NewCacheIdx + HeaderSize, CachedVoxels, VoxelsSize);
 }
 
-void FRecastTileGenerator::ClearGeometry()
-{
-#if CACHE_NAV_GENERATOR_DATA
-	GeomIndices.Empty();
-	GeomCoords.Empty();
-
-	// don't clear static areas here, they are used together with dynamic ones during nav data generation
-#else
-	StaticGeomIndices.Reset();
-	StaticGeomCoords.Reset();
-	StaticGeomSpans.Reset();
-#endif
-}
-
-void FRecastTileGenerator::ClearModifiers()
-{
-#if CACHE_NAV_GENERATOR_DATA
-	StaticAreas.Empty();
-	DynamicAreas.Empty();
-	OffmeshLinks.Empty();
-#else
-	StaticStaticAreas.Reset();
-	StaticDynamicAreas.Reset();
-	StaticOffmeshLinks.Reset();
-#endif
-}
-
-void FRecastTileGenerator::ClearStaticData()
-{
-	StaticGeomIndices.Empty();
-	StaticGeomCoords.Empty();
-	StaticGeomSpans.Empty();
-	StaticStaticAreas.Empty();
-	StaticDynamicAreas.Empty();
-	StaticOffmeshLinks.Empty();
-}
-
 void FRecastTileGenerator::AppendModifier(const FCompositeNavModifier& Modifier, bool bStatic)
 {
 	// append areas to correct array
 	if (bStatic)
 	{
-#if CACHE_NAV_GENERATOR_DATA
 		StaticAreas.Append(Modifier.GetAreas());
-#else
-		StaticStaticAreas.Append(Modifier.GetAreas());
-#endif
 	}
 	else
 	{
-#if CACHE_NAV_GENERATOR_DATA
 		DynamicAreas.Append(Modifier.GetAreas());
-#else
-		StaticDynamicAreas.Append(Modifier.GetAreas());
-#endif
 	}
 
 	// append all offmesh links (not included in compress layers)
-#if CACHE_NAV_GENERATOR_DATA
 	OffmeshLinks.Append(Modifier.GetSimpleLinks());
-#else
-	StaticOffmeshLinks.Append(Modifier.GetSimpleLinks());
-#endif
 
 	// evaluate custom links
 	const FCustomLinkNavModifier* LinkModifier = Modifier.GetCustomLinks().GetData();
 	for (int32 i = 0; i < Modifier.GetCustomLinks().Num(); i++, LinkModifier++)
 	{
 		FSimpleLinkNavModifier SimpleLinkCollection(UNavLinkDefinition::GetLinksDefinition(LinkModifier->GetNavLinkClass()), LinkModifier->LocalToWorld);
-#if CACHE_NAV_GENERATOR_DATA
 		OffmeshLinks.Add(SimpleLinkCollection);
-#else
-		StaticOffmeshLinks.Add(SimpleLinkCollection);
-#endif
 	}
 }
 
@@ -1996,32 +1746,26 @@ void FRecastTileGenerator::AppendGeometry(const TNavStatArray<uint8>& RawCollisi
 	{
 		return;
 	}
+	QUICK_SCOPE_CYCLE_COUNTER(STAT_RecastTileGenerator_AppendGeometry);
 
 	FRecastGeometryCache CollisionCache(RawCollisionCache.GetData());
-
-#if CACHE_NAV_GENERATOR_DATA
+	
 	const int32 FirstNewCoord = GeomCoords.Num();
 	const int32 FirstNewIndex = GeomIndices.Num();
 	const int32 VertBase = FirstNewCoord / 3;
+	const int32 NumCoords = CollisionCache.Header.NumVerts * 3;
+	const int32 NumIndices = CollisionCache.Header.NumFaces * 3;
+			
+	GeomCoords.Reserve(GeomCoords.Num() + NumCoords);
+	GeomCoords.AddUninitialized(NumCoords);
+	
+	GeomIndices.Reserve(GeomIndices.Num() + NumIndices);
+	GeomIndices.AddUninitialized(NumIndices);
 
-	GeomCoords.AddZeroed(CollisionCache.Header.NumVerts * 3);
-	GeomIndices.AddZeroed(CollisionCache.Header.NumFaces * 3);
-
-	FMemory::Memcpy(GeomCoords.GetData() + FirstNewCoord, CollisionCache.Verts, sizeof(float) * CollisionCache.Header.NumVerts * 3);
+	FMemory::Memcpy(GeomCoords.GetData() + FirstNewCoord, CollisionCache.Verts, sizeof(float) * NumCoords);
 	int32* DestIndex = GeomIndices.GetData() + FirstNewIndex;
-#else
-	const int32 FirstNewCoord = StaticGeomCoords.Num();
-	const int32 FirstNewIndex = StaticGeomIndices.Num();
-	const int32 VertBase = FirstNewCoord / 3;
-
-	StaticGeomCoords.AddZeroed(CollisionCache.Header.NumVerts * 3);
-	StaticGeomIndices.AddZeroed(CollisionCache.Header.NumFaces * 3);
 
-	FMemory::Memcpy(StaticGeomCoords.GetData() + FirstNewCoord, CollisionCache.Verts, sizeof(float) * CollisionCache.Header.NumVerts * 3);
-	int32* DestIndex = StaticGeomIndices.GetData() + FirstNewIndex;
-#endif
-
-	for (int32 i = 0; i < CollisionCache.Header.NumFaces * 3; i++, DestIndex++)
+	for (int32 i = 0; i < NumIndices; i++, DestIndex++)
 	{
 		*DestIndex = CollisionCache.Indices[i] + VertBase;
 	}
@@ -2034,7 +1778,6 @@ void FRecastTileGenerator::AppendGeometry(const TNavStatArray<FVector>& Verts, c
 		return;
 	}
 
-#if CACHE_NAV_GENERATOR_DATA
 	const int32 FirstNewCoord = GeomCoords.Num();
 	const int32 FirstNewIndex = GeomIndices.Num();
 	const int32 VertBase = FirstNewCoord / 3;
@@ -2044,18 +1787,7 @@ void FRecastTileGenerator::AppendGeometry(const TNavStatArray<FVector>& Verts, c
 
 	float* DestCoord = GeomCoords.GetData() + FirstNewCoord;
 	int32* DestIndex = GeomIndices.GetData() + FirstNewIndex;
-#else
-	const int32 FirstNewCoord = StaticGeomCoords.Num();
-	const int32 FirstNewIndex = StaticGeomIndices.Num();
-	const int32 VertBase = FirstNewCoord / 3;
-
-	StaticGeomCoords.AddZeroed(Verts.Num() * 3);
-	StaticGeomIndices.AddZeroed(Faces.Num());
-
-	float* DestCoord = StaticGeomCoords.GetData() + FirstNewCoord;
-	int32* DestIndex = StaticGeomIndices.GetData() + FirstNewIndex;
-#endif
-
+	
 	for (int32 i = 0; i < Verts.Num(); i++)
 	{
 		*DestCoord = Verts[i].X; DestCoord++;
@@ -2069,140 +1801,31 @@ void FRecastTileGenerator::AppendGeometry(const TNavStatArray<FVector>& Verts, c
 	}
 }
 
-void FRecastTileGenerator::AppendVoxels(rcSpanCache* SpanData, int32 NumSpans)
-{
-	if (NumSpans)
-	{
-#if !CACHE_NAV_GENERATOR_DATA
-		const int32 FirstSpanIdx = StaticGeomSpans.Num();
-		StaticGeomSpans.AddZeroed(NumSpans);
-
-		FMemory::Memcpy(StaticGeomSpans.GetData() + FirstSpanIdx, SpanData, NumSpans * sizeof(rcSpanCache));
-#endif
-	}
-}
-
-void FRecastTileGenerator::TriggerAsyncBuild()
-{
-#if RECAST_ASYNC_REBUILDING
-	bool bBuildStarted = false;
-
-	ANavigationData::FNavDataGeneratorSharedPtr NavMeshGeneratorSharedPtr = NavMeshGenerator.Pin();
-	if (NavMeshGeneratorSharedPtr.IsValid())
-	{
-		FAsyncNavTileBuildTask* AsyncBuildTask = new FAsyncNavTileBuildTask(NavMeshGeneratorSharedPtr, TileId, Version);
-
-		// Initiate
-		if (AsyncBuildTask != NULL)
-		{
-			InitiateRebuild();
-			bBuildStarted = true;
-
-			AsyncBuildTask->StartBackgroundTask();
-		}
-	}
-
-	if (!bBuildStarted)
-	{
-		AbortRebuild();
-	}
-#endif // RECAST_ASYNC_REBUILDING
-}
-
 bool FRecastTileGenerator::GenerateTile()
 {
-	struct FTileGenerationScopeLock 
-	{
-		FCriticalSection* SynchObject;
-		FRecastTileGenerator* TileGenerator;
-
-		FTileGenerationScopeLock(FCriticalSection* InSynchObject, FRecastTileGenerator* InTileGenerator) 
-			: SynchObject(InSynchObject)
-			, TileGenerator(InTileGenerator)
-		{
-			check(SynchObject);
-			check(InTileGenerator->bBeingRebuild && "FRecastTileGenerator::TriggerAsyncBuild is the only legit way of triggering navmesh tile generation!");			
-			SynchObject->Lock();
-		}
-
-		/**
-			* Destructor that performs a release on the synchronization object
-			*/
-		~FTileGenerationScopeLock(void)
-		{
-			TileGenerator->bBeingRebuild = false;
-			check(SynchObject);
-			SynchObject->Unlock();
-		}
-	};
-
-	bool bSuccess = false;
-	if (bBeingRebuild == false)
-	{
-		UE_LOG(LogNavigation, Error, TEXT("NavMeshGeneration: tile (%d,%d) was not marked for rebuilding, abort generation!"), TileX, TileY);
-		return bSuccess;
-	}
-
-	FTileGenerationScopeLock Lock(&GenerationLock, this);
-
-	if (bInitialized == false)
-	{
-		UE_LOG(LogNavigation, Error, TEXT("NavMeshGeneration: Trying to generate navmesh tile with uninitialized generator"));
-		return bSuccess;
-	}
-
-	// make sure it doesn't go away until end of this function
-	ANavigationData::FNavDataGeneratorSharedPtr NavMeshGeneratorSharedPtr = NavMeshGenerator.Pin();
-	if (NavMeshGeneratorSharedPtr.IsValid() == false)
-	{
-		return bSuccess;
-	}
+	bool bSuccess = true;
+	FNavMeshBuildContext BuildContext;
 
-	FRecastNavMeshGenerator* Generator = (FRecastNavMeshGenerator*)NavMeshGeneratorSharedPtr.Get();
-	FNavMeshBuildContext* BuildContext = Generator->GetBuildContext();
-
-	if (!GeneratingState.bRebuildGeometry && !GeneratingState.bRebuildLayers)
-	{
-		BuildContext->log(RC_LOG_WARNING, "NavMeshGeneration: trying to generate tile while it's not dirty, skipping");
-		return bSuccess;
-	}
-
-	RECAST_STAT(STAT_Navigation_Async_RecastBuilding);
-	const double BuildStartTime = FPlatformTime::Seconds();
-
-	if (GeneratingState.bRebuildGeometry)
+	if (bRegenerateCompressedLayers)
 	{
 		DEC_MEMORY_STAT_BY(STAT_Navigation_TileCacheMemory, GetTileCacheSizeHelper(CompressedLayers));
 		CompressedLayers.Reset();
-
-		const bool bLayersReady = GenerateCompressedLayers(BuildContext, Generator);
-
+		
+		bSuccess = GenerateCompressedLayers(BuildContext);
 		INC_MEMORY_STAT_BY(STAT_Navigation_TileCacheMemory, GetTileCacheSizeHelper(CompressedLayers));
-		ClearGeometry();
 
-		GeneratingState.bRebuildLayers = bLayersReady;
-		if (bLayersReady)
+		if (bSuccess)
 		{
-			GeneratingState.bRebuildAllLayers = true;
+			// Mark all layers as dirty
+			DirtyLayers.Init(true, CompressedLayers.Num());
 		}
 	}
 
-	if (GeneratingState.bRebuildLayers)
+	if (bSuccess)
 	{
-		bSuccess = GenerateNavigationData(BuildContext, Generator);
+		bSuccess = GenerateNavigationData(BuildContext);
 	}
 
-#if !(UE_BUILD_SHIPPING || UE_BUILD_TEST)
-	const double CurrentTime = FPlatformTime::Seconds();
-	const float TimeTaken = CurrentTime - BuildStartTime;
-	UE_LOG(LogNavigation, Display, TEXT("FRecastTileGenerator(%s) tile (%d,%d) took %.3fs"),
-		GeneratingState.bRebuildGeometry ? TEXT("full rebuild") : *FString::Printf(TEXT("layer update:%d"), NavigationData.Num()),
-		TileX, TileY, TimeTaken);
-
-	LastBuildTimeCost = TimeTaken;
-	LastBuildTimeStamp = CurrentTime;
-#endif // !(UE_BUILD_SHIPPING || UE_BUILD_TEST)
-
 	// it's possible to have valid generation with empty resulting tile (no navigable geometry in tile)
 	return bSuccess;
 }
@@ -2226,84 +1849,67 @@ struct FTileRasterizationContext
 	TArray<FNavMeshTileData> Layers;
 };
 
-bool FRecastTileGenerator::GenerateCompressedLayers(class FNavMeshBuildContext* BuildContext, FRecastNavMeshGenerator* Generator)
+bool FRecastTileGenerator::GenerateCompressedLayers(FNavMeshBuildContext& BuildContext)
 {
 	RECAST_STAT(STAT_Navigation_Async_Recast_Rasterize);
 	SCOPE_CYCLE_COUNTER(STAT_Navigation_RecastBuildCompressedLayers);
 	
-	FRecastBuildConfig TileConfig = Generator->GetConfig();
 	TileConfig.width = TileConfig.tileSize + TileConfig.borderSize*2;
 	TileConfig.height = TileConfig.tileSize + TileConfig.borderSize*2;
 
-	rcVcopy(TileConfig.bmin, BMin);
-	rcVcopy(TileConfig.bmax, BMax);
 	const float BBoxPadding = TileConfig.borderSize * TileConfig.cs;
 	TileConfig.bmin[0] -= BBoxPadding;
 	TileConfig.bmin[2] -= BBoxPadding;
 	TileConfig.bmax[0] += BBoxPadding;
 	TileConfig.bmax[2] += BBoxPadding;
 
-	BuildContext->log(RC_LOG_PROGRESS, "GenerateCompressedLayers:");
-	BuildContext->log(RC_LOG_PROGRESS, " - %d x %d cells", TileConfig.width, TileConfig.height);
-
-#if CACHE_NAV_GENERATOR_DATA
-	TNavStatArray<float>* PtrGeomCoords = &GeomCoords;
-	TNavStatArray<int32>* PtrGeomIndices = &GeomIndices;
-#else
-	TNavStatArray<float>* PtrGeomCoords = &StaticGeomCoords;
-	TNavStatArray<int32>* PtrGeomIndices = &StaticGeomIndices;
-#endif
-	TNavStatArray<rcSpanCache>* PtrGeomSpans = &StaticGeomSpans;
+	BuildContext.log(RC_LOG_PROGRESS, "GenerateCompressedLayers:");
+	BuildContext.log(RC_LOG_PROGRESS, " - %d x %d cells", TileConfig.width, TileConfig.height);
 
 	FTileRasterizationContext RasterContext;
 
 	// Allocate voxel heightfield where we rasterize our input data to.
-	if (PtrGeomIndices->Num() || PtrGeomSpans->Num())
+	if (GeomCoords.Num())
 	{
 		RasterContext.SolidHF = rcAllocHeightfield();
 		if (RasterContext.SolidHF == NULL)
 		{
-			BuildContext->log(RC_LOG_ERROR, "GenerateCompressedLayers: Out of memory 'SolidHF'.");
+			BuildContext.log(RC_LOG_ERROR, "GenerateCompressedLayers: Out of memory 'SolidHF'.");
 			return false;
 		}
-		if (!rcCreateHeightfield(BuildContext, *RasterContext.SolidHF, TileConfig.width, TileConfig.height, TileConfig.bmin, TileConfig.bmax, TileConfig.cs, TileConfig.ch))
+		if (!rcCreateHeightfield(&BuildContext, *RasterContext.SolidHF, TileConfig.width, TileConfig.height, TileConfig.bmin, TileConfig.bmax, TileConfig.cs, TileConfig.ch))
 		{
-			BuildContext->log(RC_LOG_ERROR, "GenerateCompressedLayers: Could not create solid heightfield.");
+			BuildContext.log(RC_LOG_ERROR, "GenerateCompressedLayers: Could not create solid heightfield.");
 			return false;
 		}
 	}
 
 	// Rasterize geometry
-	if (PtrGeomIndices->Num() && PtrGeomCoords->Num())
+	if (GeomCoords.Num() && GeomIndices.Num())
 	{
 		DECLARE_SCOPE_CYCLE_COUNTER(TEXT("Rasterization: without voxel cache"), Stat_RecastRasterNoCache, STATGROUP_Navigation);
 
-		const int32 NumFaces = PtrGeomIndices->Num() / 3;
-		const int32 NumVerts = PtrGeomCoords->Num() / 3;
+		const int32 NumFaces = GeomIndices.Num() / 3;
+		const int32 NumVerts = GeomCoords.Num() / 3;
 
 		RECAST_STAT(STAT_Navigation_RasterizeTriangles);
 		TNavStatArray<uint8> TriAreas;
+		TriAreas.Reserve(NumFaces);
 		TriAreas.AddZeroed(NumFaces);
 
-		rcMarkWalkableTriangles(BuildContext, TileConfig.walkableSlopeAngle,
-			PtrGeomCoords->GetData(), NumVerts, PtrGeomIndices->GetData(), NumFaces,
+		rcMarkWalkableTriangles(&BuildContext, TileConfig.walkableSlopeAngle,
+			GeomCoords.GetData(), NumVerts, GeomIndices.GetData(), NumFaces,
 			TriAreas.GetData());
 
-		rcRasterizeTriangles(BuildContext,
-			PtrGeomCoords->GetData(), NumVerts, 
-			PtrGeomIndices->GetData(), TriAreas.GetData(), NumFaces,
+		rcRasterizeTriangles(&BuildContext,
+			GeomCoords.GetData(), NumVerts, 
+			GeomIndices.GetData(), TriAreas.GetData(), NumFaces,
 			*RasterContext.SolidHF, TileConfig.walkableClimb);
 	}
-	else if (PtrGeomSpans->Num())
-	{
-		DECLARE_SCOPE_CYCLE_COUNTER(TEXT("Rasterization: apply voxel cache"), Stat_RecastRasterCacheApply, STATGROUP_Navigation);
-
-		rcAddSpans(BuildContext, *RasterContext.SolidHF, TileConfig.walkableClimb, PtrGeomSpans->GetData(), PtrGeomSpans->Num());
-	}
 
 	if (!RasterContext.SolidHF || RasterContext.SolidHF->pools == 0)
 	{
-		BuildContext->log(RC_LOG_WARNING, "GenerateCompressedLayers: empty tile - aborting");
+		BuildContext.log(RC_LOG_WARNING, "GenerateCompressedLayers: empty tile - aborting");
 		return true;
 	}
 
@@ -2318,11 +1924,11 @@ bool FRecastTileGenerator::GenerateCompressedLayers(class FNavMeshBuildContext*
 		// Once all geometry is rasterized, we do initial pass of filtering to
 		// remove unwanted overhangs caused by the conservative rasterization
 		// as well as filter spans where the character cannot possibly stand.
-		rcFilterLowHangingWalkableObstacles(BuildContext, TileConfig.walkableClimb, *RasterContext.SolidHF);
-		rcFilterLedgeSpans(BuildContext, TileConfig.walkableHeight, TileConfig.walkableClimb, *RasterContext.SolidHF);
+		rcFilterLowHangingWalkableObstacles(&BuildContext, TileConfig.walkableClimb, *RasterContext.SolidHF);
+		rcFilterLedgeSpans(&BuildContext, TileConfig.walkableHeight, TileConfig.walkableClimb, *RasterContext.SolidHF);
 		if (!TileConfig.bMarkLowHeightAreas)
 		{
-			rcFilterWalkableLowHeightSpans(BuildContext, TileConfig.walkableHeight, *RasterContext.SolidHF);
+			rcFilterWalkableLowHeightSpans(&BuildContext, TileConfig.walkableHeight, *RasterContext.SolidHF);
 		}
 	}
 
@@ -2334,12 +1940,12 @@ bool FRecastTileGenerator::GenerateCompressedLayers(class FNavMeshBuildContext*
 		RasterContext.CompactHF = rcAllocCompactHeightfield();
 		if (RasterContext.CompactHF == NULL)
 		{
-			BuildContext->log(RC_LOG_ERROR, "GenerateCompressedLayers: Out of memory 'CompactHF'.");
+			BuildContext.log(RC_LOG_ERROR, "GenerateCompressedLayers: Out of memory 'CompactHF'.");
 			return false;
 		}
-		if (!rcBuildCompactHeightfield(BuildContext, TileConfig.walkableHeight, TileConfig.walkableClimb, *RasterContext.SolidHF, *RasterContext.CompactHF))
+		if (!rcBuildCompactHeightfield(&BuildContext, TileConfig.walkableHeight, TileConfig.walkableClimb, *RasterContext.SolidHF, *RasterContext.CompactHF))
 		{
-			BuildContext->log(RC_LOG_ERROR, "GenerateCompressedLayers: Could not build compact data.");
+			BuildContext.log(RC_LOG_ERROR, "GenerateCompressedLayers: Could not build compact data.");
 			return false;
 		}
 	}
@@ -2351,23 +1957,23 @@ bool FRecastTileGenerator::GenerateCompressedLayers(class FNavMeshBuildContext*
 		if (TileConfig.walkableRadius > RECAST_VERY_SMALL_AGENT_RADIUS)
 		{
 			const bool bEroded = TileConfig.bMarkLowHeightAreas ?
-				rcErodeWalkableAndLowAreas(BuildContext, TileConfig.walkableRadius, HeightThreshold, RECAST_LOW_AREA, *RasterContext.CompactHF) :
-				rcErodeWalkableArea(BuildContext, TileConfig.walkableRadius, *RasterContext.CompactHF);
+				rcErodeWalkableAndLowAreas(&BuildContext, TileConfig.walkableRadius, HeightThreshold, RECAST_LOW_AREA, *RasterContext.CompactHF) :
+				rcErodeWalkableArea(&BuildContext, TileConfig.walkableRadius, *RasterContext.CompactHF);
 
 			if (!bEroded)
 			{
-				BuildContext->log(RC_LOG_ERROR, "GenerateCompressedLayers: Could not erode.");
+				BuildContext.log(RC_LOG_ERROR, "GenerateCompressedLayers: Could not erode.");
 				return false;
 			}
 		}
 		else if (TileConfig.bMarkLowHeightAreas)
 		{
-			rcMarkLowAreas(BuildContext, HeightThreshold, RECAST_LOW_AREA, *RasterContext.CompactHF);
+			rcMarkLowAreas(&BuildContext, HeightThreshold, RECAST_LOW_AREA, *RasterContext.CompactHF);
 		}
 	}
 
 	// (Optional) Mark areas.
-	MarkStaticAreas(BuildContext, *RasterContext.CompactHF, TileConfig);
+	MarkStaticAreas(BuildContext, *RasterContext.CompactHF);
 
 	// Build layers
 	{
@@ -2376,45 +1982,43 @@ bool FRecastTileGenerator::GenerateCompressedLayers(class FNavMeshBuildContext*
 		RasterContext.LayerSet = rcAllocHeightfieldLayerSet();
 		if (RasterContext.LayerSet == NULL)
 		{
-			BuildContext->log(RC_LOG_ERROR, "GenerateCompressedLayers: Out of memory 'LayerSet'.");
+			BuildContext.log(RC_LOG_ERROR, "GenerateCompressedLayers: Out of memory 'LayerSet'.");
 			return false;
 		}
 
 		if (TileConfig.regionPartitioning == RC_REGION_MONOTONE)
 		{
-			if (!rcBuildHeightfieldLayersMonotone(BuildContext, *RasterContext.CompactHF, TileConfig.borderSize, TileConfig.walkableHeight, *RasterContext.LayerSet))
+			if (!rcBuildHeightfieldLayersMonotone(&BuildContext, *RasterContext.CompactHF, TileConfig.borderSize, TileConfig.walkableHeight, *RasterContext.LayerSet))
 			{
-				BuildContext->log(RC_LOG_ERROR, "GenerateCompressedLayers: Could not build heightfield layers.");
+				BuildContext.log(RC_LOG_ERROR, "GenerateCompressedLayers: Could not build heightfield layers.");
 				return 0;
 			}
 		}
 		else if (TileConfig.regionPartitioning == RC_REGION_WATERSHED)
 		{
-			if (!rcBuildDistanceField(BuildContext, *RasterContext.CompactHF))
+			if (!rcBuildDistanceField(&BuildContext, *RasterContext.CompactHF))
 			{
-				BuildContext->log(RC_LOG_ERROR, "GenerateCompressedLayers: Could not build distance field.");
+				BuildContext.log(RC_LOG_ERROR, "GenerateCompressedLayers: Could not build distance field.");
 				return 0;
 			}
 
-			if (!rcBuildHeightfieldLayers(BuildContext, *RasterContext.CompactHF, TileConfig.borderSize, TileConfig.walkableHeight, *RasterContext.LayerSet))
+			if (!rcBuildHeightfieldLayers(&BuildContext, *RasterContext.CompactHF, TileConfig.borderSize, TileConfig.walkableHeight, *RasterContext.LayerSet))
 			{
-				BuildContext->log(RC_LOG_ERROR, "GenerateCompressedLayers: Could not build heightfield layers.");
+				BuildContext.log(RC_LOG_ERROR, "GenerateCompressedLayers: Could not build heightfield layers.");
 				return 0;
 			}
 		}
 		else
 		{
-			if (!rcBuildHeightfieldLayersChunky(BuildContext, *RasterContext.CompactHF, TileConfig.borderSize, TileConfig.walkableHeight, TileConfig.regionChunkSize, *RasterContext.LayerSet))
+			if (!rcBuildHeightfieldLayersChunky(&BuildContext, *RasterContext.CompactHF, TileConfig.borderSize, TileConfig.walkableHeight, TileConfig.regionChunkSize, *RasterContext.LayerSet))
 			{
-				BuildContext->log(RC_LOG_ERROR, "GenerateCompressedLayers: Could not build heightfield layers.");
+				BuildContext.log(RC_LOG_ERROR, "GenerateCompressedLayers: Could not build heightfield layers.");
 				return 0;
 			}
 		}
 
 		const int32 NumLayers = RasterContext.LayerSet->nlayers;
-		LayerBB.Reset();
-		LayerBB.AddZeroed(NumLayers);
-
+	
 		// use this to expand vertically layer's bounds
 		// this is needed to allow off-mesh connections that are not quite
 		// touching tile layer still connect with it.
@@ -2447,11 +2051,11 @@ bool FRecastTileGenerator::GenerateCompressedLayers(class FNavMeshBuildContext*
 			header.hmin = (unsigned short)layer->hmin;
 			header.hmax = (unsigned short)layer->hmax;
 
-			// Store bounds
-			LayerBB[i] = Recast2UnrealBox(header.bmin, header.bmax);
-			LayerBB[i].Min.Z -= StepHeights;
-			LayerBB[i].Max.Z += StepHeights;
-
+			// Layer bounds in unreal coords
+			FBox LayerBBox = Recast2UnrealBox(header.bmin, header.bmax);
+			LayerBBox.Min.Z -= StepHeights;
+			LayerBBox.Max.Z += StepHeights;
+			
 			// Compress tile layer
 			uint8* TileData = NULL;
 			int32 TileDataSize = 0;
@@ -2459,7 +2063,7 @@ bool FRecastTileGenerator::GenerateCompressedLayers(class FNavMeshBuildContext*
 			if (dtStatusFailed(status))
 			{
 				dtFree(TileData);
-				BuildContext->log(RC_LOG_ERROR, "GenerateCompressedLayers: failed to build layer.");
+				BuildContext.log(RC_LOG_ERROR, "GenerateCompressedLayers: failed to build layer.");
 				return false;
 			}
 
@@ -2470,55 +2074,48 @@ bool FRecastTileGenerator::GenerateCompressedLayers(class FNavMeshBuildContext*
 			if (CompressedData == NULL)
 			{
 				dtFree(TileData);
-				BuildContext->log(RC_LOG_ERROR, "GenerateCompressedLayers: Out of memory 'CompressedData'.");
+				BuildContext.log(RC_LOG_ERROR, "GenerateCompressedLayers: Out of memory 'CompressedData'.");
 				return false;
 			}
 
 			FMemory::Memcpy(CompressedData, TileData, TileDataSize);
-			RasterContext.Layers.Add(FNavMeshTileData(CompressedData, TileDataSize, i));
+			RasterContext.Layers.Add(FNavMeshTileData(CompressedData, TileDataSize, i, LayerBBox));
+
 			dtFree(TileData);
 
 			const int32 UncompressedSize = ((sizeof(dtTileCacheLayerHeader)+3) & ~3) + (3 * header.width * header.height);
 			const float Inv1kB = 1.0f / 1024.0f;
-			BuildContext->log(RC_LOG_PROGRESS, ">> Cache[%d,%d:%d] = %.2fkB (full:%.2fkB rate:%.2f%%)", TileX, TileY, i,
+			BuildContext.log(RC_LOG_PROGRESS, ">> Cache[%d,%d:%d] = %.2fkB (full:%.2fkB rate:%.2f%%)", TileX, TileY, i,
 				TileDataSize * Inv1kB, UncompressedSize * Inv1kB, 1.0f * TileDataSize / UncompressedSize);
 		}
 	}
 
 	// Transfer final data
 	CompressedLayers = RasterContext.Layers;
-	RasterContext.Layers.Reset();
 	return true;
 }
 
-void FRecastTileGenerator::MarkStaticAreas(class FNavMeshBuildContext* BuildContext, struct rcCompactHeightfield& CompactHF, const struct FRecastBuildConfig& TileConfig)
+void FRecastTileGenerator::MarkStaticAreas(FNavMeshBuildContext& BuildContext, rcCompactHeightfield& CompactHF)
 {
-#if CACHE_NAV_GENERATOR_DATA
-	TArray<FAreaNavModifier>* PtrStaticAreas = &StaticAreas;
-#else
-	TArray<FAreaNavModifier>* PtrStaticAreas = &StaticStaticAreas;
-#endif
-
-	const int32 NumAreas = PtrStaticAreas->Num();
+	const int32 NumAreas = StaticAreas.Num();
 	if (NumAreas == 0)
 	{
 		return;
 	}
 
-	FRecastNavMeshCachedData* AdditionalCachedDataPtr = AdditionalCachedData.Get();
-	if (AdditionalCachedDataPtr->bUseSortFunction && AdditionalCachedDataPtr->ActorOwner && NumAreas > 1)
+	if (AdditionalCachedData.bUseSortFunction && AdditionalCachedData.ActorOwner && NumAreas > 1)
 	{
-		AdditionalCachedDataPtr->ActorOwner->SortAreasForGenerator(*PtrStaticAreas);
+		AdditionalCachedData.ActorOwner->SortAreasForGenerator(StaticAreas);
 	}
 
 	RECAST_STAT(STAT_Navigation_Async_MarkAreas);
 	const float ExpandBy = TileConfig.AgentRadius * 1.5;
-
-	const FAreaNavModifier* Modifier = PtrStaticAreas->GetData();
+	const FAreaNavModifier* Modifier = StaticAreas.GetData();
+	
 	for (int32 ModifierIndex = 0; ModifierIndex < NumAreas; ++ModifierIndex, ++Modifier)
 	{
-		const int32* AreaID = AdditionalCachedDataPtr->AreaClassToIdMap.Find(Modifier->GetAreaClass());
-		const int32* ReplaceID = AdditionalCachedDataPtr->AreaClassToIdMap.Find(Modifier->GetAreaClassToReplace());
+		const int32* AreaID = AdditionalCachedData.AreaClassToIdMap.Find(Modifier->GetAreaClass());
+		const int32* ReplaceID = AdditionalCachedData.AreaClassToIdMap.Find(Modifier->GetAreaClassToReplace());
 		if (AreaID == NULL)
 		{
 			// happens when area is not supported by agent owning this navmesh
@@ -2540,11 +2137,11 @@ void FRecastTileGenerator::MarkStaticAreas(class FNavMeshBuildContext* BuildCont
 
 				if (ReplaceID)
 				{
-					rcReplaceCylinderArea(BuildContext, &(RecastPos.X), CylinderData.Radius, CylinderData.Height, *AreaID, *ReplaceID, CompactHF);
+					rcReplaceCylinderArea(&BuildContext, &(RecastPos.X), CylinderData.Radius, CylinderData.Height, *AreaID, *ReplaceID, CompactHF);
 				}
 				else
 				{
-					rcMarkCylinderArea(BuildContext, &(RecastPos.X), CylinderData.Radius, CylinderData.Height, *AreaID, CompactHF);
+					rcMarkCylinderArea(&BuildContext, &(RecastPos.X), CylinderData.Radius, CylinderData.Height, *AreaID, CompactHF);
 				}
 			}
 			break;
@@ -2561,11 +2158,11 @@ void FRecastTileGenerator::MarkStaticAreas(class FNavMeshBuildContext* BuildCont
 
 				if (ReplaceID)
 				{
-					rcReplaceBoxArea(BuildContext, &(RecastBox.Min.X), &(RecastBox.Max.X), *AreaID, *ReplaceID, CompactHF);
+					rcReplaceBoxArea(&BuildContext, &(RecastBox.Min.X), &(RecastBox.Max.X), *AreaID, *ReplaceID, CompactHF);
 				}
 				else
 				{
-					rcMarkBoxArea(BuildContext, &(RecastBox.Min.X), &(RecastBox.Max.X), *AreaID, CompactHF);
+					rcMarkBoxArea(&BuildContext, &(RecastBox.Min.X), &(RecastBox.Max.X), *AreaID, CompactHF);
 				}
 			}
 			break;
@@ -2596,12 +2193,12 @@ void FRecastTileGenerator::MarkStaticAreas(class FNavMeshBuildContext* BuildCont
 
 					if (ReplaceID)
 					{
-						rcReplaceConvexPolyArea(BuildContext, ConvexCoords.GetData(), ConvexVerts.Num(),
+						rcReplaceConvexPolyArea(&BuildContext, ConvexCoords.GetData(), ConvexVerts.Num(),
 							ConvexData.MinZ, ConvexData.MaxZ, *AreaID, *ReplaceID, CompactHF);
 					}
 					else
 					{
-						rcMarkConvexPolyArea(BuildContext, ConvexCoords.GetData(), ConvexVerts.Num(),
+						rcMarkConvexPolyArea(&BuildContext, ConvexCoords.GetData(), ConvexVerts.Num(),
 							ConvexData.MinZ, ConvexData.MaxZ, *AreaID, CompactHF);
 					}
 				}
@@ -2639,7 +2236,6 @@ struct FTileGenerationContext
 		PolyMesh = 0;
 		dtFreeTileCachePolyMeshDetail(Allocator, DetailMesh);
 		DetailMesh = 0;
-
 		// don't clear NavigationData here!
 	}
 
@@ -2653,50 +2249,41 @@ struct FTileGenerationContext
 	TArray<FNavMeshTileData> NavigationData;
 };
 
-bool FRecastTileGenerator::GenerateNavigationData(class FNavMeshBuildContext* BuildContext, FRecastNavMeshGenerator* Generator)
+bool FRecastTileGenerator::GenerateNavigationData(FNavMeshBuildContext& BuildContext)
 {
 	RECAST_STAT(STAT_Navigation_Async_Recast_Generate);
 	SCOPE_CYCLE_COUNTER(STAT_Navigation_RecastBuildNavigation);
-
-#if CACHE_NAV_GENERATOR_DATA
-	TArray<FSimpleLinkNavModifier>* PtrOffmeshLinks = &OffmeshLinks;
-#else
-	TArray<FSimpleLinkNavModifier>* PtrOffmeshLinks = &StaticOffmeshLinks;
-#endif
-
+		
 	FTileCacheAllocator MyAllocator;
 	FTileCacheCompressor TileCompressor;
 	
 	FTileGenerationContext GenerationContext(&MyAllocator);
-	GenerationContext.NavigationData.AddZeroed(CompressedLayers.Num());
+	GenerationContext.NavigationData.Reserve(CompressedLayers.Num());
 
-	FRecastNavMeshCachedData* AdditionalCachedDataPtr = AdditionalCachedData.Get();
-	const FRecastBuildConfig& TileConfig = Generator->GetConfig();
 	dtStatus status = DT_SUCCESS;
 
-	int32 NumLayers = 0;
 	for (int32 iLayer = 0; iLayer < CompressedLayers.Num(); iLayer++)
 	{
-		if (!GeneratingState.HasDirtyLayer(iLayer))
+		if (DirtyLayers[iLayer] == false)
 		{
 			// skip layers not marked for rebuild
 			continue;
 		}
-
+				
 		FNavMeshTileData& CompressedData = CompressedLayers[iLayer];
-		const dtTileCacheLayerHeader* TileHeader = (const dtTileCacheLayerHeader*)CompressedData.GetDataSafe();
+		const dtTileCacheLayerHeader* TileHeader = (const dtTileCacheLayerHeader*)CompressedData.GetData();
 		GenerationContext.ResetIntermediateData();
 
 		// Decompress tile layer data. 
-		status = dtDecompressTileCacheLayer(&MyAllocator, &TileCompressor, (unsigned char*)CompressedData.GetDataSafe(), CompressedData.DataSize, &GenerationContext.Layer);
+		status = dtDecompressTileCacheLayer(&MyAllocator, &TileCompressor, (unsigned char*)CompressedData.GetData(), CompressedData.DataSize, &GenerationContext.Layer);
 		if (dtStatusFailed(status))
 		{
-			BuildContext->log(RC_LOG_ERROR, "GenerateNavigationData: failed to decompress layer.");
+			BuildContext.log(RC_LOG_ERROR, "GenerateNavigationData: failed to decompress layer.");
 			return false;
 		}
 
 		// Rasterize obstacles.
-		MarkDynamicAreas(*GenerationContext.Layer, TileConfig);
+		MarkDynamicAreas(*GenerationContext.Layer);
 
 		{
 			RECAST_STAT(STAT_Navigation_Async_Recast_BuildRegions)
@@ -2710,14 +2297,14 @@ bool FRecastTileGenerator::GenerateNavigationData(class FNavMeshBuildContext* Bu
 				GenerationContext.DistanceField = dtAllocTileCacheDistanceField(&MyAllocator);
 				if (GenerationContext.DistanceField == NULL)
 				{
-					BuildContext->log(RC_LOG_ERROR, "GenerateNavigationData: Out of memory 'DistanceField'.");
+					BuildContext.log(RC_LOG_ERROR, "GenerateNavigationData: Out of memory 'DistanceField'.");
 					return false;
 				}
 
 				status = dtBuildTileCacheDistanceField(&MyAllocator, *GenerationContext.Layer, *GenerationContext.DistanceField);
 				if (dtStatusFailed(status))
 				{
-					BuildContext->log(RC_LOG_ERROR, "GenerateNavigationData: Failed to build distance field.");
+					BuildContext.log(RC_LOG_ERROR, "GenerateNavigationData: Failed to build distance field.");
 					return false;
 				}
 
@@ -2731,7 +2318,7 @@ bool FRecastTileGenerator::GenerateNavigationData(class FNavMeshBuildContext* Bu
 
 			if (dtStatusFailed(status))
 			{
-				BuildContext->log(RC_LOG_ERROR, "GenerateNavigationData: Failed to build regions.");
+				BuildContext.log(RC_LOG_ERROR, "GenerateNavigationData: Failed to build regions.");
 				return false;
 			}
 		}
@@ -2742,14 +2329,14 @@ bool FRecastTileGenerator::GenerateNavigationData(class FNavMeshBuildContext* Bu
 			GenerationContext.ContourSet = dtAllocTileCacheContourSet(&MyAllocator);
 			if (GenerationContext.ContourSet == NULL)
 			{
-				BuildContext->log(RC_LOG_ERROR, "GenerateNavigationData: Out of memory 'ContourSet'.");
+				BuildContext.log(RC_LOG_ERROR, "GenerateNavigationData: Out of memory 'ContourSet'.");
 				return false;
 			}
 
 			GenerationContext.ClusterSet = dtAllocTileCacheClusterSet(&MyAllocator);
 			if (GenerationContext.ClusterSet == NULL)
 			{
-				BuildContext->log(RC_LOG_ERROR, "GenerateNavigationData: Out of memory 'ClusterSet'.");
+				BuildContext.log(RC_LOG_ERROR, "GenerateNavigationData: Out of memory 'ClusterSet'.");
 				return false;
 			}
 
@@ -2758,7 +2345,7 @@ bool FRecastTileGenerator::GenerateNavigationData(class FNavMeshBuildContext* Bu
 				*GenerationContext.ContourSet, *GenerationContext.ClusterSet);
 			if (dtStatusFailed(status))
 			{
-				BuildContext->log(RC_LOG_ERROR, "GenerateNavigationData: Failed to generate contour set (0x%08X).", status);
+				BuildContext.log(RC_LOG_ERROR, "GenerateNavigationData: Failed to generate contour set (0x%08X).", status);
 				return false;
 			}
 		}
@@ -2769,21 +2356,21 @@ bool FRecastTileGenerator::GenerateNavigationData(class FNavMeshBuildContext* Bu
 			GenerationContext.PolyMesh = dtAllocTileCachePolyMesh(&MyAllocator);
 			if (GenerationContext.PolyMesh == NULL)
 			{
-				BuildContext->log(RC_LOG_ERROR, "GenerateNavigationData: Out of memory 'PolyMesh'.");
+				BuildContext.log(RC_LOG_ERROR, "GenerateNavigationData: Out of memory 'PolyMesh'.");
 				return false;
 			}
 
 			status = dtBuildTileCachePolyMesh(&MyAllocator, *GenerationContext.ContourSet, *GenerationContext.PolyMesh);
 			if (dtStatusFailed(status))
 			{
-				BuildContext->log(RC_LOG_ERROR, "GenerateNavigationData: Failed to generate poly mesh.");
+				BuildContext.log(RC_LOG_ERROR, "GenerateNavigationData: Failed to generate poly mesh.");
 				return false;
 			}
 
 			status = dtBuildTileCacheClusters(&MyAllocator, *GenerationContext.ClusterSet, *GenerationContext.PolyMesh);
 			if (dtStatusFailed(status))
 			{
-				BuildContext->log(RC_LOG_ERROR, "GenerateNavigationData: Failed to update cluster set.");
+				BuildContext.log(RC_LOG_ERROR, "GenerateNavigationData: Failed to update cluster set.");
 				return false;
 			}
 		}
@@ -2797,7 +2384,7 @@ bool FRecastTileGenerator::GenerateNavigationData(class FNavMeshBuildContext* Bu
 			GenerationContext.DetailMesh = dtAllocTileCachePolyMeshDetail(&MyAllocator);
 			if (GenerationContext.DetailMesh == NULL)
 			{
-				BuildContext->log(RC_LOG_ERROR, "GenerateNavigationData: Out of memory 'DetailMesh'.");
+				BuildContext.log(RC_LOG_ERROR, "GenerateNavigationData: Out of memory 'DetailMesh'.");
 				return false;
 			}
 
@@ -2805,7 +2392,7 @@ bool FRecastTileGenerator::GenerateNavigationData(class FNavMeshBuildContext* Bu
 				*GenerationContext.Layer, *GenerationContext.PolyMesh, *GenerationContext.DetailMesh);
 			if (dtStatusFailed(status))
 			{
-				BuildContext->log(RC_LOG_ERROR, "GenerateNavigationData: Failed to generate poly detail mesh.");
+				BuildContext.log(RC_LOG_ERROR, "GenerateNavigationData: Failed to generate poly detail mesh.");
 				return false;
 			}
 		}
@@ -2820,23 +2407,23 @@ bool FRecastTileGenerator::GenerateNavigationData(class FNavMeshBuildContext* Bu
 			if (GenerationContext.PolyMesh->nverts >= 0xffff)
 			{
 				// The vertex indices are ushorts, and cannot point to more than 0xffff vertices.
-				BuildContext->log(RC_LOG_ERROR, "Too many vertices per tile %d (max: %d).", GenerationContext.PolyMesh->nverts, 0xffff);
+				BuildContext.log(RC_LOG_ERROR, "Too many vertices per tile %d (max: %d).", GenerationContext.PolyMesh->nverts, 0xffff);
 				return false;
 			}
 
 			// if we didn't failed already then it's hight time we created data for off-mesh links
 			FOffMeshData OffMeshData;
-			if (PtrOffmeshLinks->Num() > 0)
+			if (OffmeshLinks.Num() > 0)
 			{
 				RECAST_STAT(STAT_Navigation_Async_GatherOffMeshData);
 
-				OffMeshData.Reserve(PtrOffmeshLinks->Num());
-				OffMeshData.AreaClassToIdMap = &AdditionalCachedDataPtr->AreaClassToIdMap;
-				OffMeshData.FlagsPerArea = AdditionalCachedDataPtr->FlagsPerOffMeshLinkArea;
+				OffMeshData.Reserve(OffmeshLinks.Num());
+				OffMeshData.AreaClassToIdMap = &AdditionalCachedData.AreaClassToIdMap;
+				OffMeshData.FlagsPerArea = AdditionalCachedData.FlagsPerOffMeshLinkArea;
 				const uint32 AgentMask = (1 << TileConfig.AgentIndex);
+				const FSimpleLinkNavModifier* LinkModifier = OffmeshLinks.GetData();
 
-				const FSimpleLinkNavModifier* LinkModifier = PtrOffmeshLinks->GetData();
-				for (int32 LinkModifierIndex = 0; LinkModifierIndex < PtrOffmeshLinks->Num(); ++LinkModifierIndex, ++LinkModifier)
+				for (int32 LinkModifierIndex = 0; LinkModifierIndex < OffmeshLinks.Num(); ++LinkModifierIndex, ++LinkModifier)
 				{
 					OffMeshData.AddLinks(LinkModifier->Links, LinkModifier->LocalToWorld, AgentMask);
 #if GENERATE_SEGMENT_LINKS
@@ -2849,7 +2436,7 @@ bool FRecastTileGenerator::GenerateNavigationData(class FNavMeshBuildContext* Bu
 			// Update poly flags from areas.
 			for (int32 i = 0; i < GenerationContext.PolyMesh->npolys; i++)
 			{
-				GenerationContext.PolyMesh->flags[i] = AdditionalCachedDataPtr->FlagsPerArea[GenerationContext.PolyMesh->areas[i]];
+				GenerationContext.PolyMesh->flags[i] = AdditionalCachedData.FlagsPerArea[GenerationContext.PolyMesh->areas[i]];
 			}
 
 			dtNavMeshCreateParams Params;
@@ -2891,55 +2478,39 @@ bool FRecastTileGenerator::GenerateNavigationData(class FNavMeshBuildContext* Bu
 
 			if (!dtCreateNavMeshData(&Params, &NavData, &NavDataSize))
 			{
-				BuildContext->log(RC_LOG_ERROR, "Could not build Detour navmesh.");
+				BuildContext.log(RC_LOG_ERROR, "Could not build Detour navmesh.");
 				return false;
 			}
 		}
 
-		GenerationContext.NavigationData[NumLayers] = FNavMeshTileData(NavData, NavDataSize, iLayer);
-		NumLayers++;
+		GenerationContext.NavigationData.Add(FNavMeshTileData(NavData, NavDataSize, iLayer, CompressedData.LayerBBox));
 
 		const float ModkB = 1.0f / 1024.0f;
-		BuildContext->log(RC_LOG_PROGRESS, ">> Layer[%d] = Verts(%d) Polys(%d) Memory(%.2fkB) Cache(%.2fkB)",
+		BuildContext.log(RC_LOG_PROGRESS, ">> Layer[%d] = Verts(%d) Polys(%d) Memory(%.2fkB) Cache(%.2fkB)",
 			iLayer, GenerationContext.PolyMesh->nverts, GenerationContext.PolyMesh->npolys,
-			GenerationContext.NavigationData[iLayer].DataSize * ModkB, CompressedLayers[iLayer].DataSize * ModkB);
+			GenerationContext.NavigationData.Last().DataSize * ModkB, CompressedLayers[iLayer].DataSize * ModkB);
 	}
 
 	// prepare navigation data of actually rebuild layers for transfer
-	NavigationData.Reset();
-	NavigationData.AddZeroed(NumLayers);
-	for (int32 i = 0; i < NumLayers; i++)
-	{
-		NavigationData[i] = GenerationContext.NavigationData[i];
-	}
-
+	NavigationData = GenerationContext.NavigationData;
 	return true;
 }
 
-void FRecastTileGenerator::MarkDynamicAreas(struct dtTileCacheLayer& Layer, const struct FRecastBuildConfig& TileConfig)
+void FRecastTileGenerator::MarkDynamicAreas(dtTileCacheLayer& Layer)
 {
-#if CACHE_NAV_GENERATOR_DATA
-	TArray<FAreaNavModifier>* PtrStaticAreas = &StaticAreas;
-	TArray<FAreaNavModifier>* PtrDynamicAreas = &DynamicAreas;
-#else
-	TArray<FAreaNavModifier>* PtrStaticAreas = &StaticStaticAreas;
-	TArray<FAreaNavModifier>* PtrDynamicAreas = &StaticDynamicAreas;
-#endif
-
-	if (PtrDynamicAreas->Num() == 0)
+	if (DynamicAreas.Num() == 0)
 	{
 		return;
 	}
 
 	// combine both area types, to avoid overriding high cost static with low cost dynamic one
 	TArray<FAreaNavModifier> CombinedAreas;
-	CombinedAreas.Append(*PtrStaticAreas);
-	CombinedAreas.Append(*PtrDynamicAreas);
+	CombinedAreas.Append(StaticAreas);
+	CombinedAreas.Append(DynamicAreas);
 
-	FRecastNavMeshCachedData* AdditionalCachedDataPtr = AdditionalCachedData.Get();
-	if (AdditionalCachedDataPtr->bUseSortFunction && AdditionalCachedDataPtr->ActorOwner && CombinedAreas.Num() > 1)
+	if (AdditionalCachedData.bUseSortFunction && AdditionalCachedData.ActorOwner && CombinedAreas.Num() > 1)
 	{
-		AdditionalCachedDataPtr->ActorOwner->SortAreasForGenerator(CombinedAreas);
+		AdditionalCachedData.ActorOwner->SortAreasForGenerator(CombinedAreas);
 	}
 
 	RECAST_STAT(STAT_Navigation_Async_MarkAreas);
@@ -2950,8 +2521,8 @@ void FRecastTileGenerator::MarkDynamicAreas(struct dtTileCacheLayer& Layer, cons
 	const FAreaNavModifier* Modifier = CombinedAreas.GetData();
 	for (int32 ModifierIndex = 0; ModifierIndex < CombinedAreas.Num(); ++ModifierIndex, ++Modifier)
 	{
-		const int32* AreaID = AdditionalCachedDataPtr->AreaClassToIdMap.Find(Modifier->GetAreaClass());
-		const int32* ReplaceID = AdditionalCachedDataPtr->AreaClassToIdMap.Find(Modifier->GetAreaClassToReplace());
+		const int32* AreaID = AdditionalCachedData.AreaClassToIdMap.Find(Modifier->GetAreaClass());
+		const int32* ReplaceID = AdditionalCachedData.AreaClassToIdMap.Find(Modifier->GetAreaClassToReplace());
 		if (AreaID == NULL)
 		{
 			// happens when area is not supported by agent owning this navmesh
@@ -3093,24 +2664,19 @@ uint32 FRecastTileGenerator::GetUsedMemCount() const
 	return TotalMemory;
 }
 
-void FRecastTileGenerator::SetDirty(const FNavigationDirtyArea& DirtyArea, const FBox& AreaBounds)
+static int32 CaclulateGridCellsCount(const TNavStatArray<FBox>& NavigableAreas, float TileSizeinWorldUnits)
 {
-	DirtyState.bRebuildGeometry |= DirtyArea.HasFlag(ENavigationDirtyFlag::Geometry);
-	if (DirtyState.bRebuildGeometry)
-	{
-		return;
-	}
-
-	if (DirtyArea.HasFlag(ENavigationDirtyFlag::DynamicModifier))
+	int32 CellsCount = 0;
+	for (FBox AreaBounds : NavigableAreas)
 	{
-		for (int32 i = 0; i < LayerBB.Num(); i++)
-		{
-			if (LayerBB[i].Intersect(AreaBounds))
-			{
-				DirtyState.MarkDirtyLayer(i);
-			}
-		}
+		// TODO: need more precise calculation, currently we don't take into account that volumes can be overlapped
+		FBox RCBox = Unreal2RecastBox(AreaBounds);
+		int32 XSize = FMath::CeilToInt(RCBox.GetSize().X/TileSizeinWorldUnits) + 1;
+		int32 YSize = FMath::CeilToInt(RCBox.GetSize().Z/TileSizeinWorldUnits) + 1;
+		CellsCount+= (XSize*YSize);
 	}
+	
+	return CellsCount;
 }
 
 //----------------------------------------------------------------------//
@@ -3118,214 +2684,139 @@ void FRecastTileGenerator::SetDirty(const FNavigationDirtyArea& DirtyArea, const
 //----------------------------------------------------------------------//
 
 FRecastNavMeshGenerator::FRecastNavMeshGenerator(class ARecastNavMesh* InDestNavMesh)
-	: DetourMesh(NULL), MaxActiveTiles(-1), NumActiveTiles(0), MaxActiveGenerators(64)
-	, TilesWidth(-1), TilesHeight(-1)
-	, GridWidth(-1), GridHeight(-1)
-	, TileSize(-1), RCNavBounds(0), UnrealNavBounds(0)
-	, DestNavMesh(InDestNavMesh), OutNavMesh(NULL)
-	, bInitialized(false), bBuildFromScratchRequested(false), bRebuildDirtyTilesRequested(false)
-	, bAbortAllTileGeneration(false), bOwnsDetourMesh(false), bBuildingLocked(false), Version(0)
+	: MaxActiveTiles(-1)
+	, NumActiveTiles(0)
+	, MaxTileGeneratorTasks(1)
+	, DestNavMesh(InDestNavMesh)
+	, bInitialized(false)
+	, Version(0)
 {
-#if WITH_EDITOR
-	bBuildingLocked = UNavigationSystem::GetIsNavigationAutoUpdateEnabled() == false;
-#endif // WITH_EDITOR
-	BuildContext = new FNavMeshBuildContext();
 	INC_DWORD_STAT_BY( STAT_NavigationMemory, sizeof(*this) );
+	SET_MEMORY_STAT(STAT_Navigation_TileCacheMemory, 0);
 
 	check(InDestNavMesh);
-	check(InDestNavMesh->GetWorld());
+	Init();
+	ConstructTiledNavMesh();
 }
 
 FRecastNavMeshGenerator::~FRecastNavMeshGenerator()
 {
-	// Free only when no ARecastNavMesh action assigned. It there's an actor then he 
-	// is responsible 
-	SetDetourMesh(NULL);
-
-	if (BuildContext)
-	{
-		delete BuildContext;
-		BuildContext = NULL;
-	}
-
+	DiscardCurrentBuildingTasks();
+	
 	DEC_DWORD_STAT_BY( STAT_NavigationMemory, sizeof(*this) );
+	SET_MEMORY_STAT(STAT_Navigation_TileCacheMemory, 0);
 }
 
 void FRecastNavMeshGenerator::Init()
 {
-	if (GetWorld() == NULL)
-	{
-		return;	
-	}
-
-	FScopeLock lock(&InitLock);
-
-	// use existing navmesh to respect per-mesh tweaked values, or default navmesh to get defaults
-	ARecastNavMesh const* const NavGenParams = DestNavMesh.IsValid() ? DestNavMesh.Get() : ARecastNavMesh::StaticClass()->GetDefaultObject<ARecastNavMesh>();
-	
 	// @todo those variables should be tweakable per navmesh actor
-	const float CellSize = NavGenParams ? NavGenParams->CellSize : 19.0;
-	const float CellHeight = NavGenParams ? NavGenParams->CellHeight : 10.0;
-	const float AgentHeight = NavGenParams ? NavGenParams->AgentHeight : (72.f * 2.0f);
-	const float MaxAgentHeight = NavGenParams ? NavGenParams->AgentMaxHeight : 160.f;
-	const float AgentMaxSlope = NavGenParams ? NavGenParams->AgentMaxSlope : 55.f;
-	const float AgentMaxClimb = NavGenParams ? NavGenParams->AgentMaxStepHeight : 35.f;
-	const float AgentRadius = NavGenParams ? NavGenParams->AgentRadius : RECAST_VERY_SMALL_AGENT_RADIUS;
-
-	SetUpGeneration(CellSize, CellHeight, AgentHeight, MaxAgentHeight, AgentMaxSlope, AgentMaxClimb, AgentRadius);
-
-	if (NavGenParams != NULL)
-	{
-		Config.minRegionArea = (int32)rcSqr(NavGenParams->MinRegionArea / CellSize);
-		Config.mergeRegionArea = (int32)rcSqr(NavGenParams->MergeRegionSize / CellSize);
-		Config.maxSimplificationError = NavGenParams->MaxSimplificationError;
-		Config.bPerformVoxelFiltering = NavGenParams->bPerformVoxelFiltering;
-		Config.bMarkLowHeightAreas = NavGenParams->bMarkLowHeightAreas;
-		if (NavGenParams->bMarkLowHeightAreas)
-		{
-			Config.walkableHeight = 1;
-		}
-
-		AdditionalCachedData = MakeShareable(new FRecastNavMeshCachedData(NavGenParams));
+	const float CellSize = DestNavMesh->CellSize;
+	const float CellHeight = DestNavMesh->CellHeight;
+	const float AgentHeight = DestNavMesh->AgentHeight;
+	const float MaxAgentHeight = DestNavMesh->AgentMaxHeight;
+	const float AgentMaxSlope = DestNavMesh->AgentMaxSlope;
+	const float AgentMaxClimb = DestNavMesh->AgentMaxStepHeight;
+	const float AgentRadius = DestNavMesh->AgentRadius;
 
-		const UNavigationSystem* NavSys = UNavigationSystem::GetCurrent(GetWorld());
-		Config.AgentIndex = NavSys ? NavSys->GetSupportedAgentIndex(NavGenParams) : 0;
+	Config.Reset();
+	
+	Config.cs = CellSize;
+	Config.ch = CellHeight;
+	Config.walkableSlopeAngle = AgentMaxSlope;
+	Config.walkableHeight = (int32)ceilf(AgentHeight / CellHeight);
+	Config.walkableClimb = (int32)ceilf(AgentMaxClimb / CellHeight);
+	const float WalkableRadius = FMath::CeilToFloat(AgentRadius / CellSize);
+	Config.walkableRadius = WalkableRadius;
+	
+	// store original sizes
+	Config.AgentHeight = AgentHeight;
+	Config.AgentMaxClimb = AgentMaxClimb;
+	Config.AgentRadius = AgentRadius;
 
-		Config.tileSize = FMath::TruncToInt(NavGenParams->TileSizeUU / CellSize);
+	Config.borderSize = WalkableRadius + 3;
+	Config.maxEdgeLen = (int32)(1200.0f / CellSize);
+	Config.maxSimplificationError = 1.3f;
+	// hardcoded, but can be overridden by RecastNavMesh params later
+	Config.minRegionArea = (int32)rcSqr(0);
+	Config.mergeRegionArea = (int32)rcSqr(20.f);
 
-		Config.regionChunkSize = Config.tileSize / NavGenParams->LayerChunkSplits;
-		Config.TileCacheChunkSize = Config.tileSize / NavGenParams->RegionChunkSplits;
-		Config.regionPartitioning = NavGenParams->LayerPartitioning;
-		Config.TileCachePartitionType = NavGenParams->RegionPartitioning;
-	}
-	else
+	Config.maxVertsPerPoly = (int32)MAX_VERTS_PER_POLY;
+	Config.detailSampleDist = 600.0f;
+	Config.detailSampleMaxError = 1.0f;
+	Config.PolyMaxHeight = (int32)ceilf(MaxAgentHeight / CellHeight);
+		
+	Config.minRegionArea = (int32)rcSqr(DestNavMesh->MinRegionArea / CellSize);
+	Config.mergeRegionArea = (int32)rcSqr(DestNavMesh->MergeRegionSize / CellSize);
+	Config.maxSimplificationError = DestNavMesh->MaxSimplificationError;
+	Config.bPerformVoxelFiltering = DestNavMesh->bPerformVoxelFiltering;
+	Config.bMarkLowHeightAreas = DestNavMesh->bMarkLowHeightAreas;
+	if (DestNavMesh->bMarkLowHeightAreas)
 	{
-		Config.tileSize = 64;
-		Config.regionPartitioning = RC_REGION_WATERSHED;
-		Config.TileCachePartitionType = RC_REGION_WATERSHED;
+		Config.walkableHeight = 1;
 	}
+	
+	AdditionalCachedData = FRecastNavMeshCachedData::Construct(DestNavMesh);
 
-	const float* BMin = Config.bmin;
-	const float* BMax = Config.bmax;
-	rcCalcGridSize(BMin, BMax, CellSize, &GridWidth, &GridHeight);
-	const float TileSizeInWorldUnits = Config.tileSize * Config.cs;
-	int32 NewTilesWidth = (GridWidth + Config.tileSize - 1) / Config.tileSize;
-	int32 NewTilesHeight = (GridHeight + Config.tileSize - 1) / Config.tileSize;
-
-	// limit max amount of tiles: config values
-	if (NavGenParams)
-	{
-		if (NewTilesWidth > NavGenParams->MaxTileGridWidth || NewTilesWidth < 0 ||
-			NewTilesHeight > NavGenParams->MaxTileGridHeight || NewTilesHeight < 0)
-		{
-			const int32 OrgTilesWidth = NewTilesWidth;
-			const int32 OrgTilesHeight = NewTilesHeight;
+	const UNavigationSystem* NavSys = UNavigationSystem::GetCurrent(GetWorld());
+	Config.AgentIndex = NavSys->GetSupportedAgentIndex(DestNavMesh);
 
-			NewTilesWidth = (NewTilesWidth < 0 || NewTilesWidth > NavGenParams->MaxTileGridWidth) ? NavGenParams->MaxTileGridWidth : NewTilesWidth;
-			NewTilesHeight = (NewTilesHeight < 0 || NewTilesHeight > NavGenParams->MaxTileGridHeight) ? NavGenParams->MaxTileGridHeight : NewTilesHeight;
+	Config.tileSize = FMath::TruncToInt(DestNavMesh->TileSizeUU / CellSize);
 
-			UE_LOG(LogNavigation, Error, TEXT("Navmesh bounds are too large! Limiting requested grid (%d x %d) to: (%d x %d)"),
-				OrgTilesWidth, OrgTilesHeight, NewTilesWidth, NewTilesHeight);
-		}
-	}
+	Config.regionChunkSize = Config.tileSize / DestNavMesh->LayerChunkSplits;
+	Config.TileCacheChunkSize = Config.tileSize / DestNavMesh->RegionChunkSplits;
+	Config.regionPartitioning = DestNavMesh->LayerPartitioning;
+	Config.TileCachePartitionType = DestNavMesh->RegionPartitioning;
 
-	// limit max amount of tiles: 64 bit poly address
-	const int32 MaxTileBits = 30;
-	const float AvgLayersPerTile = 8.0f;
-	const int32 MaxAllowedGridCells = FMath::TruncToInt((1 << MaxTileBits) / AvgLayersPerTile);
-	const int32 NumRequestedCells = NewTilesWidth * NewTilesHeight;
-	if (NumRequestedCells < 0 || NumRequestedCells > MaxAllowedGridCells)
+	TotalNavBounds = FBox(0);
+	// Collect bounding geometry
+	if (NavSys->ShouldGenerateNavigationEverywhere() == false)
 	{
-		const int32 LimitTiles = FMath::Sqrt(MaxAllowedGridCells);
-		const int32 OrgTilesWidth = NewTilesWidth;
-		const int32 OrgTilesHeight = NewTilesHeight;
+		InclusionBounds.Empty();
 
-		if (NewTilesHeight < 0 && NewTilesWidth < 0)
-		{
-			NewTilesWidth = LimitTiles;
-			NewTilesHeight = LimitTiles;
-		}
-		else if (NewTilesHeight > 0 && NewTilesHeight < LimitTiles)
-		{
-			NewTilesWidth = MaxAllowedGridCells / NewTilesHeight;
-		}
-		else if (NewTilesWidth > 0 && NewTilesWidth < LimitTiles)
-		{
-			NewTilesHeight = MaxAllowedGridCells / NewTilesWidth;
-		}
-		else
+		for (TActorIterator<ANavMeshBoundsVolume> It(NavSys->GetWorld()); It; ++It)
 		{
-			NewTilesWidth = LimitTiles;
-			NewTilesHeight = LimitTiles;
+			ANavMeshBoundsVolume const* const V = (*It);
+			if (V != NULL)
+			{
+				FBox Bounds = GrowBoundingBox(V->GetComponentsBoundingBox(true), true);
+				TotalNavBounds += Bounds;
+				InclusionBounds.Add(Bounds);
+			}
 		}
-
-		UE_LOG(LogNavigation, Error, TEXT("Navmesh bounds are too large! Limiting requested grid (%d x %d) to: (%d x %d)"),
-			OrgTilesWidth, OrgTilesHeight, NewTilesWidth, NewTilesHeight);
 	}
 
-	// do this once
-	if (bInitialized == false || NewTilesHeight != TilesHeight || NewTilesWidth != TilesWidth) 
+	if (!TotalNavBounds.IsValid)
 	{
-		FScopeLock Lock(&TileGenerationLock);
-
-		{
-			FScopeLock Lock(&NavMeshDirtyLock);
-			DirtyAreas.Reset();
-			DirtyGenerators.Empty(MaxActiveGenerators);
-		}
-		
-		float TileBmin[3];
-		float TileBmax[3];
-
-		TilesHeight = NewTilesHeight;
-		TilesWidth = NewTilesWidth;
-
-		++Version;
-		TileGenerators.Reset();
-		TileGenerators.AddUninitialized(TilesWidth * TilesHeight);
-
-		DestNavMesh->ReserveTileSet(TilesWidth, TilesHeight);
-		FTileSetItem* TileData = DestNavMesh->GetTileSet();
-		for (int32 TileIndex = 0; TileIndex < TilesWidth * TilesHeight; ++TileIndex, ++TileData)
-		{
-			const int32 X = TileData->X;
-			const int32 Y = TileData->Y;
-
-			TileBmin[0] = BMin[0] + X*TileSizeInWorldUnits;
-			TileBmin[1] = BMin[1];
-			TileBmin[2] = BMin[2] + Y*TileSizeInWorldUnits;
-
-			TileBmax[0] = BMin[0] + (X+1)*TileSizeInWorldUnits;
-			TileBmax[1] = BMax[1];
-			TileBmax[2] = BMin[2] + (Y+1)*TileSizeInWorldUnits;
-
-			// @todo check if this tile overlaps current navmesh generation boundaries 
-			// if not remove it
-			FRecastTileGenerator& TileGenerator = TileGenerators[ Y*TilesWidth + X ];
-			new(&TileGenerator) FRecastTileGenerator();
-			TileGenerator.Init(this, X, Y, TileBmin, TileBmax, InclusionBounds);
+		// if no bounds give Navigation System a chance to specify bounds
+		// will also happen if NavSystem->ShouldGenerateNavigationEverywhere() == true
+		TotalNavBounds = NavSys->GetWorldBounds();
+		InclusionBounds.Add(TotalNavBounds);
+	}
 
-			// reinitialize with all data
-			new(TileData) FTileSetItem(X, Y, TileGenerator.GetUnrealBB());
-		}
+	// limit max amount of tiles: 64 bit poly 
+	const int32 MaxTileBits = 30;
+	const int32 MaxTiles = (1 << 30);
+	const float AvgLayersPerGridCell = 8.0f; // TODO: should be tweakable per navmesh
+	MaxActiveTiles = FMath::CeilToInt(CaclulateGridCellsCount(InclusionBounds, Config.tileSize * Config.cs) * AvgLayersPerGridCell);
+	if (MaxActiveTiles < 0 || MaxActiveTiles > MaxTiles)
+	{
+		UE_LOG(LogNavigation, Error, TEXT("Navmesh bounds are too large! Limiting requested tiles count (%d) to: (%d)"), MaxActiveTiles, MaxTiles);
+		MaxActiveTiles = MaxTiles;
 	}
 
+	// do this once
+	++Version;
+
 	// Max tiles and max polys affect how the tile IDs are calculated.
 	// There are (sizeof(dtPolyRef)*8 - DT_MIN_SALT_BITS) bits available for 
 	// identifying a tile and a polygon.
-	int32 TileBits = MaxTileBits;
-	const float MaxLayers = TilesWidth * TilesHeight * AvgLayersPerTile;
-	TileBits = FMath::Min(FMath::TruncToInt(FMath::Log2(FMath::RoundUpToPowerOfTwo(MaxLayers))), MaxTileBits);
-	MaxActiveTiles = 1 << TileBits;
-
+	int32 TileBits = FMath::CeilLogTwo(MaxActiveTiles);
 	const int32 PolyBits = FMath::Min(30, (int)(sizeof(dtPolyRef)*8 - DT_MIN_SALT_BITS) - TileBits);
 	Config.MaxPolysPerTile = 1 << PolyBits;
 
 	/** setup maximum number of active tile generator*/
 	const int32 NumberOfWorkerThreads = FTaskGraphInterface::Get().GetNumWorkerThreads();
-	const int32 MaxGeneratorsThreads = NumberOfWorkerThreads > 1 ? NumberOfWorkerThreads - 1 : 1;
-	ActiveGenerators.Reset();
-	ActiveGenerators.AddZeroed(MaxGeneratorsThreads);
+	MaxTileGeneratorTasks = FMath::Max(NumberOfWorkerThreads*2, 1);
 	NumActiveTiles = 0;
 
 	// prepare voxel cache if needed
@@ -3340,1067 +2831,565 @@ void FRecastNavMeshGenerator::Init()
 bool FRecastNavMeshGenerator::ConstructTiledNavMesh() 
 {
 	bool bSuccess = false;
-	const bool bCreateNewNavMesh = DestNavMesh == NULL || DestNavMesh->GetRecastNavMeshImpl() == NULL
-		|| DestNavMesh->GetRecastNavMeshImpl()->GetRecastMesh() == NULL;
-
-	if (bCreateNewNavMesh)
+	
+	// create new Detour navmesh instance
+	DestNavMesh->GetRecastNavMeshImpl()->SetRecastMesh(dtAllocNavMesh());
+	
+	dtNavMesh* DetourMesh = DestNavMesh->GetRecastNavMeshImpl()->GetRecastMesh();
+	if (DetourMesh)
 	{
-		// create new Detour navmesh instance
-		SetDetourMesh(dtAllocNavMesh());
-
-		if ( DetourMesh != NULL )
-		{
-			dtNavMeshParams TiledMeshParameters;
-			FMemory::MemZero(TiledMeshParameters);	
-			rcVcopy(TiledMeshParameters.orig, Config.bmin);
-			TiledMeshParameters.tileWidth = Config.tileSize * Config.cs;
-			TiledMeshParameters.tileHeight = Config.tileSize * Config.cs;
-			TiledMeshParameters.maxTiles = MaxActiveTiles;
-			TiledMeshParameters.maxPolys = Config.MaxPolysPerTile;
-
-			const dtStatus status = DetourMesh->init(&TiledMeshParameters);
+		float Origin[3] = {0,0,0};
+		
+		dtNavMeshParams TiledMeshParameters;
+		FMemory::MemZero(TiledMeshParameters);	
+		rcVcopy(TiledMeshParameters.orig, Origin);
+		TiledMeshParameters.tileWidth = Config.tileSize * Config.cs;
+		TiledMeshParameters.tileHeight = Config.tileSize * Config.cs;
+		TiledMeshParameters.maxTiles = MaxActiveTiles;
+		TiledMeshParameters.maxPolys = Config.MaxPolysPerTile;
 
-			TransferGeneratedData();
+		const dtStatus status = DetourMesh->init(&TiledMeshParameters);
 
-			if ( dtStatusFailed(status) )
-			{
-				UE_LOG(LogNavigation, Warning, TEXT("ConstructTiledNavMesh: Could not init navmesh.") );
-				bSuccess = false;
-			}
-			else
-			{
-				bSuccess = true;
-			}
+		if ( dtStatusFailed(status) )
+		{
+			UE_LOG(LogNavigation, Warning, TEXT("ConstructTiledNavMesh: Could not init navmesh.") );
+			bSuccess = false;
 		}
 		else
 		{
-			UE_LOG(LogNavigation, Warning, TEXT("ConstructTiledNavMesh: Could not allocate navmesh.") );
-			bSuccess = false;
+			bSuccess = true;
+			NumActiveTiles = GetTilesCountHelper(DetourMesh);
 		}
 	}
 	else
 	{
-		SetDetourMesh(DestNavMesh->GetRecastNavMeshImpl()->GetRecastMesh(), DO_ForeignData);
-		bSuccess = true;
+		UE_LOG(LogNavigation, Warning, TEXT("ConstructTiledNavMesh: Could not allocate navmesh.") );
+		bSuccess = false;
 	}
 	
 	return bSuccess;
 }
 
-/** Inclusion geometry should be cached up before calling this. */
-void FRecastNavMeshGenerator::SetUpGeneration(float CellSize, float CellHeight, float AgentMinHeight, float AgentMaxHeight, float AgentMaxSlope, float AgentMaxClimb, float AgentRadius)
+bool FRecastNavMeshGenerator::RebuildAll()
 {
-	checkSlow(IsInGameThread() == true);
+	DestNavMesh->UpdateNavVersion();
+	
+	// if rebuilding all no point in keeping "old" invalidated areas
+	TArray<FNavigationDirtyArea> DirtyAreas;
+	for (FBox AreaBounds : InclusionBounds)
+	{
+		FNavigationDirtyArea DirtyArea(AreaBounds, ENavigationDirtyFlag::All);
+		DirtyAreas.Add(DirtyArea);
+	}
+	
+	MarkDirtyTiles(DirtyAreas);
+	
+	return true;
+}
 
-	const UNavigationSystem* NavSys = UNavigationSystem::GetCurrent(GetWorld());
-	if (NavSys == NULL || NavSys->GetWorld() == NULL)
+void FRecastNavMeshGenerator::EnsureBuildCompletion()
+{
+	while (PendingDirtyTiles.Num() || RunningDirtyTiles.Num())
 	{
-		return;
+		SubmitDirtyTiles(16);
+		
+		// Block until tasks are finished
+		for (FTileElement& Element : RunningDirtyTiles)
+		{
+			Element.AsyncTask->EnsureCompletion();
+		}
+		//
+		CollectCompletedTiles();
 	}
+}
 
-	Config.Reset();
-	Config.cs = CellSize;
-	Config.ch = CellHeight;
-	Config.walkableSlopeAngle = AgentMaxSlope;
-	Config.walkableHeight = (int32)ceilf(AgentMinHeight / CellHeight);
-	Config.walkableClimb = (int32)ceilf(AgentMaxClimb / CellHeight);
-	const float WalkableRadius = FMath::CeilToFloat(AgentRadius / CellSize);
-	Config.walkableRadius = WalkableRadius;
-	
-	// store original sizes
-	Config.AgentHeight = AgentMinHeight;
-	Config.AgentMaxClimb = AgentMaxClimb;
-	Config.AgentRadius = AgentRadius;
+void FRecastNavMeshGenerator::CancelBuild()
+{
+	DiscardCurrentBuildingTasks();
+	RunningDirtyTiles.Empty();
+	PendingDirtyTiles.Empty();
+	CachedLayerDataMap.Empty();
+	SET_MEMORY_STAT(STAT_Navigation_TileCacheMemory, 0);
 
-	Config.borderSize = WalkableRadius + 3;
-	Config.maxEdgeLen = (int32)(1200.0f / CellSize);
-	Config.maxSimplificationError = 1.3f;
-	// hardcoded, but can be overridden by RecastNavMesh params later
-	Config.minRegionArea = (int32)rcSqr(0);
-	Config.mergeRegionArea = (int32)rcSqr(20.f);
+#if	WITH_EDITOR	
+	RecentlyBuiltTiles.Empty();
+#endif//WITH_EDITOR
+}
 
-	Config.maxVertsPerPoly = (int32)MAX_VERTS_PER_POLY;
-	Config.detailSampleDist = 600.0f;
-	Config.detailSampleMaxError = 1.0f;
-	Config.PolyMaxHeight = (int32)ceilf(AgentMaxHeight / CellHeight);
+void FRecastNavMeshGenerator::TickAsyncBuild(float DeltaSeconds)
+{
+#if	WITH_EDITOR
+	// Remove expired tiles
+	{
+		const double Timestamp = FPlatformTime::Seconds();
+		RecentlyBuiltTiles.RemoveAllSwap([&](const FTileTimestamp& Tile) { return (Timestamp - Tile.Timestamp) > 0.5; });
+	}
+#endif//WITH_EDITOR
 		
-	FBox NavBounds(0);
-	
-	if (NavSys->ShouldGenerateNavigationEverywhere() == false)
+	// Collect and apply results from async tile build tasks in case we have any
+	TArray<uint32> CompletedTileIndices = CollectCompletedTiles();
+	if (CompletedTileIndices.Num() > 0)
 	{
-		// Collect bounding geometry
-		TArray<const ANavMeshBoundsVolume*> InclusionVolumes;
-		InclusionBounds.Empty();
+		// Invalidate active paths that go through regenerated tiles
+		DestNavMesh->InvalidateAffectedPaths(CompletedTileIndices);
 
-		for (TActorIterator<ANavMeshBoundsVolume> It(NavSys->GetWorld()); It; ++It)
+#if	WITH_EDITOR
+		// Store completed tiles with timestamps to have ability to distinguish during debug draw
 		{
-			ANavMeshBoundsVolume const* const V = (*It);
-			if (V != NULL)
+			const double Timestamp = FPlatformTime::Seconds();
+			RecentlyBuiltTiles.Reserve(RecentlyBuiltTiles.Num() + CompletedTileIndices.Num());
+			for (uint32 TiledIdx : CompletedTileIndices)
 			{
-				InclusionVolumes.Add(V);
+				FTileTimestamp TileTimestamp;
+				TileTimestamp.TileIdx = TiledIdx;
+				TileTimestamp.Timestamp = Timestamp;
+				RecentlyBuiltTiles.Add(TileTimestamp);
 			}
 		}
+#endif//WITH_EDITOR
 
-		// get aabb in ue3 coords
-		for (int32 VolumeIdx = 0; VolumeIdx < InclusionVolumes.Num(); ++VolumeIdx)
-		{
-			AVolume const* const Vol = InclusionVolumes[VolumeIdx];
-			if (Vol != NULL && Vol->BrushComponent.IsValid())
-			{
-				FBox Bounds = Vol->BrushComponent->Bounds.GetBox();
-				if (Bounds.GetSize().IsZero() == true)
-				{
-					if (Vol->BrushComponent->IsRegistered() == false)
-					{
-						Vol->BrushComponent->RegisterComponent();
-					}
-					Vol->BrushComponent->UpdateBounds();
-					Bounds = Vol->BrushComponent->Bounds.GetBox();
-				}
-
-				if (Bounds.GetSize().IsZero() == false)
-				{
-					NavBounds += Bounds;
-					InclusionBounds.Add(Bounds);
-				}
-			}
-		}
+		DestNavMesh->RequestDrawingUpdate();
 	}
 
-	if (!NavBounds.IsValid)
+	// Submit async tile build tasks in case we have dirty tiles and have room for them
+	const UNavigationSystem* NavSys = UNavigationSystem::GetCurrent(GetWorld());
+	check(NavSys);
+	const int32 NumRunningTasks = NavSys->GetNumRunningBuildTasks();
+	const int32 NumTasksToSubmit = MaxTileGeneratorTasks - NumRunningTasks;
+		
+	if (NumTasksToSubmit > 0)
 	{
-		// if no bounds give Navigation System a chance to specify bounds
-		// will also happen if NavSystem->ShouldGenerateNavigationEverywhere() == true
-		NavBounds = NavSys->GetWorldBounds();
-	}
-	
-	// expand bounds a bit to support later inclusion tests
-	NavBounds = NavBounds.ExpandBy(CellSize);
-
-	bool bAdjust = false;
-	bool bClampBounds = false;
-	const float ExtentLimit = float(MAX_int32);
-	FVector BoundsExtent = NavBounds.GetExtent();
-	if (BoundsExtent.X > ExtentLimit)
-	{
-		BoundsExtent.X = ExtentLimit;
-		bClampBounds = bAdjust = true;
-	}
-	else if (BoundsExtent.X < Config.cs)
-	{
-		// minor adjustment to have at least 1 voxel of size here
-		BoundsExtent.X = Config.cs;
-		bAdjust = true;
-	}
-	if (BoundsExtent.Y > ExtentLimit)
-	{
-		BoundsExtent.Y = ExtentLimit;
-		bClampBounds = bAdjust = true;
-	}
-	else if (BoundsExtent.Y < Config.cs)
-	{
-		// minor adjustment to have at least 1 voxel of size here
-		BoundsExtent.Y = Config.cs;
-		bAdjust = true;
-	}
-
-	if (bAdjust)
-	{
-		const FVector BoundsCenter = NavBounds.GetCenter();
-		NavBounds = FBox(BoundsCenter - BoundsExtent, BoundsCenter + BoundsExtent);
-
-		if (bClampBounds)
-		{
-			UE_LOG(LogNavigation, Warning, TEXT("Navigation bounds are too large. Cutting down every dimention down to %.f"), ExtentLimit);
-		}
-	}
-
-	UnrealNavBounds = NavBounds;
-	// now move the box to Recast coords
-	RCNavBounds = Unreal2RecastBox(NavBounds);
-	Config.bmin[0] = RCNavBounds.Min.X;
-	Config.bmin[1] = RCNavBounds.Min.Y;
-	Config.bmin[2] = RCNavBounds.Min.Z;
-	Config.bmax[0] = RCNavBounds.Max.X;
-	Config.bmax[1] = RCNavBounds.Max.Y;
-	Config.bmax[2] = RCNavBounds.Max.Z;
-
-	// update offset to current one
-	if (DestNavMesh.IsValid())
-	{
-		DestNavMesh->UpdateNavmeshOffset(NavBounds);
+		SubmitDirtyTiles(NumTasksToSubmit);
 	}
 }
 
-int32 GetTilesCountHelper(const class dtNavMesh* DetourMesh)
-{
-	int32 NumTiles = 0;
-	if (DetourMesh)
-	{
-		for (int32 i = 0; i < DetourMesh->getMaxTiles(); i++)
-		{
-			const dtMeshTile* TileData = DetourMesh->getTile(i);
-			if (TileData && TileData->header && TileData->dataSize > 0)
-			{
-				NumTiles++;
-			}
-		}
-	}
-
-	return NumTiles;
-}
-
-void FRecastNavMeshGenerator::SetDetourMesh(class dtNavMesh* NewDetourMesh, FRecastNavMeshGenerator::EDataOwnership OwnsData)
+void FRecastNavMeshGenerator::RebuildDirtyAreas(const TArray<FNavigationDirtyArea>& InDirtyAreas)
 {
-	if (NewDetourMesh != DetourMesh && !!bOwnsDetourMesh && DetourMesh != NULL)
-	{
-		dtFreeNavMesh(DetourMesh);
-	}
-	bOwnsDetourMesh = NewDetourMesh != NULL && OwnsData == FRecastNavMeshGenerator::DO_OwnsData;
-	DetourMesh = NewDetourMesh;
-
-	NumActiveTiles = GetTilesCountHelper(NewDetourMesh);
+	MarkDirtyTiles(InDirtyAreas);
 }
 
 void FRecastNavMeshGenerator::OnAreaAdded(const UClass* AreaClass, int32 AreaID)
 {
-	if (AdditionalCachedData.IsValid())
-	{
-		AdditionalCachedData->OnAreaAdded(AreaClass, AreaID);
-	}
-}
-
-void FRecastNavMeshGenerator::RequestGeneration()
-{
-	if (ShouldContinueBuilding() == false)
-	{
-		return;
-	}
-
-	const bool bIsInGameThread = IsInGameThread() == true;
-	if (AreAnyTilesBeingBuilt() == false && IsBuildingLocked() == false && bIsInGameThread)
-	{
-		Generate();
-	}
-	else if (bIsInGameThread == true)
-	{
-		bBuildFromScratchRequested = true;
-	}
-	else
-	{
-		DECLARE_CYCLE_STAT(TEXT("FSimpleDelegateGraphTask.Requesting navmesh rebuild from scratch from RequestGeneration"),
-			STAT_FSimpleDelegateGraphTask_RequestingNavmeshRebuildFromScratchFromRequestGeneration,
-			STATGROUP_TaskGraphTasks);
-
-		FSimpleDelegateGraphTask::CreateAndDispatchWhenReady(
-			FSimpleDelegateGraphTask::FDelegate::CreateThreadSafeSP(this, &FRecastNavMeshGenerator::RequestGeneration),
-			GET_STATID(STAT_FSimpleDelegateGraphTask_RequestingNavmeshRebuildFromScratchFromRequestGeneration),
-			NULL, ENamedThreads::GameThread
-		);
-	}
+	AdditionalCachedData.OnAreaAdded(AreaClass, AreaID);
 }
 
-void FRecastNavMeshGenerator::RequestDirtyTilesRebuild()
+bool FRecastNavMeshGenerator::IntercestsInclusionBounds(const FBox& Box) const
 {
-	if (bRebuildDirtyTilesRequested == true)
+	for (const FBox& NavBox : InclusionBounds)
 	{
-		return;
+		if (NavBox.Intersect(Box))
+		{
+			return true;
+		}
 	}
 
-	bRebuildDirtyTilesRequested = true;
-
-	DECLARE_CYCLE_STAT(TEXT("FSimpleDelegateGraphTask.NavMesh Tiles Regeneration"),
-		STAT_FSimpleDelegateGraphTask_NavMeshTilesRegeneration,
-		STATGROUP_TaskGraphTasks);
-
-		// kick off a task to rebuild influenced tiles			
-	FSimpleDelegateGraphTask::CreateAndDispatchWhenReady(
-		FSimpleDelegateGraphTask::FDelegate::CreateThreadSafeSP(this, &FRecastNavMeshGenerator::RegenerateDirtyTiles),
-		GET_STATID(STAT_FSimpleDelegateGraphTask_NavMeshTilesRegeneration), NULL, ENamedThreads::GameThread
-	);
+	return false;
 }
 
-bool FRecastNavMeshGenerator::Generate()
-{
-	bool const bGenerated = GenerateTiledNavMesh();
-	return bGenerated;
-}
 
-void FRecastNavMeshGenerator::RebuildDirtyAreas(const TArray<FNavigationDirtyArea>& InDirtyAreas)
+TArray<uint32> FRecastNavMeshGenerator::RemoveTileLayers(const int32 TileX, const int32 TileY)
 {
-	const FNavigationDirtyArea* DirtyArea = InDirtyAreas.GetData();
-	bool bOverlaps = false;
-	for (int32 DirtyIndex = 0; DirtyIndex < InDirtyAreas.Num(); ++DirtyIndex, ++DirtyArea)
+	TArray<uint32> ResultTileIndices;
+	dtNavMesh* DetourMesh = DestNavMesh->GetRecastNavMeshImpl()->GetRecastMesh();
+	const int32 NumLayers = DetourMesh ? DetourMesh->getTileCountAt(TileX, TileY) : 0;
+	
+	if (NumLayers > 0)
 	{
-		if (DirtyArea->Bounds.Intersect(UnrealNavBounds))
+		TArray<dtMeshTile*> Tiles;
+		Tiles.AddZeroed(NumLayers);
+		DetourMesh->getTilesAt(TileX, TileY, (const dtMeshTile**)Tiles.GetData(), NumLayers);
+	
+		for (int32 i = 0; i < NumLayers; i++)
 		{
-			bOverlaps = true;
-			break;
-		}
-	}
+			const int32 LayerIndex = Tiles[i]->header->layer;
+			const dtTileRef TileRef = DetourMesh->getTileRef(Tiles[i]);
 
-	if (bOverlaps)
-	{
-		FScopeLock Lock(&NavMeshDirtyLock);
-		DirtyAreas.Append(InDirtyAreas);
+			NumActiveTiles--;
+			UE_LOG(LogNavigation, Log, TEXT("%s> Tile (%d,%d:%d), removing TileRef: 0x%X (active:%d)"),
+				*DestNavMesh->GetName(), TileX, TileY, LayerIndex, TileRef, NumActiveTiles);
 
-		if (AreAnyTilesBeingBuilt() == false && IsBuildingLocked() == false)
-		{
-			RequestDirtyTilesRebuild();
+			uint8* RawNavData = NULL;
+			DetourMesh->removeTile(TileRef, &RawNavData, NULL);
+
+			ResultTileIndices.AddUnique(DetourMesh->decodePolyIdTile(TileRef));
 		}
 	}
+
+	return ResultTileIndices;
 }
 
-bool FRecastNavMeshGenerator::GenerateTiledNavMesh() 
+TArray<uint32> FRecastNavMeshGenerator::AddTile(const FRecastTileGenerator& TileGenerator)
 {
-	// it's possible for someone to request rebuild when there are tiles already rebuilding. 
-	// if number of tiles is to be changed then this rebuild needs to be postponed
-	if (AreAnyTilesBeingBuilt() == true || IsBuildingLocked() == true)
+	SCOPE_CYCLE_COUNTER(STAT_Navigation_AddTile);
+	
+	TArray<uint32> ResultTileIndices;
+	const int32 TileX = TileGenerator.GetTileX();
+	const int32 TileY = TileGenerator.GetTileY();
+	
+	// 
+	TArray<FNavMeshTileData> TileLayers = TileGenerator.GetNavigationData();
+	if (TileGenerator.IsFullyRegenerated())
 	{
-		bBuildFromScratchRequested = true;
-
-		return false;
+		// remove all layers
+		ResultTileIndices = RemoveTileLayers(TileX, TileY);
 	}
+	ResultTileIndices.Reserve(TileLayers.Num());
 
-	bBuildFromScratchRequested = false;
-	bRebuildDirtyTilesRequested = false;
-
-	Init();
-
-	// Gather path colliding geometry
-	const double BuildStartTime = FPlatformTime::Seconds();
-
-	if (DetourMesh != NULL)
+	dtNavMesh* DetourMesh = DestNavMesh->GetRecastNavMeshImpl()->GetRecastMesh();
+		
+	bool bHasNavmesh = true;
+	for (int32 i = 0; i < TileLayers.Num(); i++)
 	{
-		// @todo temporarily here to make every subsequent navmesh rebuilds work. Should go away 
-		// in next commit
-		//TileGenerators.Reset();
-		if (DestNavMesh.IsValid())
+		const int32 LayerIndex = TileLayers[i].LayerIndex;
+		if (!TileGenerator.IsLayerChanged(TileLayers[i].LayerIndex))
 		{
-			// Make owner destroy current navmesh. Again, this is a temporary bit
-			DestNavMesh->GetRecastNavMeshImpl()->SetRecastMesh(NULL);
+			continue;
 		}
-		else
-		{
-			dtFreeNavMesh(DetourMesh);
+				
+		const dtTileRef OldTileRef = DetourMesh->getTileRefAt(TileX, TileY, LayerIndex);
+
+		if (OldTileRef)
+		{
+			NumActiveTiles--;
+			UE_LOG(LogNavigation, Log, TEXT("%s> Tile (%d,%d:%d), removing TileRef: 0x%X (active:%d)"),
+				*DestNavMesh->GetName(), TileX, TileY, LayerIndex, OldTileRef, NumActiveTiles);
+
+			// this call will fill OldRawNavData with address of previously added data
+			// if OldRawNavData is NULL it means either this tile was empty, or the memory
+			// allocated has been owned by navmesh itself (i.e. has been serialized-in along
+			// with the level)
+			uint8* OldRawNavData = nullptr;
+			DetourMesh->removeTile(OldTileRef, &OldRawNavData, nullptr);
+			// data should always be owned by navmesh itself
+			check(OldRawNavData == nullptr);
+			
+			ResultTileIndices.AddUnique(DetourMesh->decodePolyIdTile(OldTileRef));
 		}
 
-		SetDetourMesh(NULL);
-	}
-
-	if ( (DetourMesh == NULL && ConstructTiledNavMesh() == false) )
-	{
-		return false;	
-	}
-
-	RebuildAll();
+		if (TileLayers[i].IsValid()) 
+		{
+			bool bRejectNavmesh = false;
+			dtTileRef ResultTileRef = 0;
 
-	UE_LOG(LogNavigation, Log, TEXT("RecastNavMeshGenerator: prepare tiles for generation took %.5fs")
-		, FPlatformTime::Seconds() - BuildStartTime);
+			// let navmesh know it's tile generator who owns the data
+			dtStatus status = DetourMesh->addTile(TileLayers[i].GetData(), TileLayers[i].DataSize, DT_TILE_FREE_DATA, 0, &ResultTileRef);
 
-	return true;
-}
+			if (dtStatusFailed(status))
+			{
+				if (dtStatusDetail(status, DT_OUT_OF_MEMORY))
+				{
+					UE_LOG(LogNavigation, Error, TEXT("%s> Tile (%d,%d:%d), tile limit reached!! (%d)"),
+						*DestNavMesh->GetName(), TileX, TileY, LayerIndex, DetourMesh->getMaxTiles());
+				}
 
-void FRecastNavMeshGenerator::GenerateTile(const int32 TileId, const uint32 TileGeneratorVersion)
-{
-	SCOPE_CYCLE_COUNTER(STAT_Navigation_TileBuildAsync);
+				bHasNavmesh = false;
+			}
+			else
+			{
+				ResultTileIndices.AddUnique(DetourMesh->decodePolyIdTile(ResultTileRef));
+				NumActiveTiles++;
 
-	if (bAbortAllTileGeneration  == true)
-	{
-		// the whole generator is shutting down, skip all building
-		UE_LOG(LogNavigation, Log, TEXT("FRecastNavMeshGenerator::GenerateTile abondoning tile (%d,%d) rebuild due to whole navmesh generator shutting down."));
-		return;
+				UE_LOG(LogNavigation, Log, TEXT("%s> Tile (%d,%d:%d), added TileRef: 0x%X (active:%d)"),
+					*DestNavMesh->GetName(), TileX, TileY, LayerIndex, ResultTileRef, NumActiveTiles);
+				
+				// NavMesh took the ownership of generated data, so we don't need to deallocate it
+				TileLayers[i].Release();
+			}
+		}
 	}
 
-	if (TileGeneratorVersion < Version)
-	{
-		UE_LOG(LogNavigation, Log, TEXT("FRecastNavMeshGenerator::GenerateTile abondoning tile (%d,%d) rebuild due to version mismatch."));
-		return;
-	}
-	// all needed data we have already cached locally, we just need to know the owner
-	// is not dead or about to be destroyed. No functions will be called on it.
-	if (DestNavMesh.IsValid(/*bEvenIfPendingKill = */false, /*bThreadsafeTest = */true) == false)
-	{
-		UE_LOG(LogNavigation, Warning, TEXT("FRecastNavMeshGenerator::GenerateTile failed to trigger actual tile building due to DestNavMesh Not Being Valid. Aborting rebuild for this tile."));
-		return;
-	}
+	return ResultTileIndices;
+}
 
-	if (TileId < TileGenerators.Num())
+void FRecastNavMeshGenerator::DiscardCurrentBuildingTasks()
+{
+	for (FTileElement& Element : RunningDirtyTiles)
 	{
-		FRecastTileGenerator& TileGenerator = TileGenerators[TileId];
-		TileGenerator.StartAsyncBuild();
-
-		UE_LOG(LogNavigation, Log, TEXT("%s> Generating Tile %d,%d"),
-			TEXT_WEAKOBJ_NAME(DestNavMesh), TileGenerator.GetTileX(), TileGenerator.GetTileY());
-
-		// tile generation can result in empty tile - we want to add it anyway 
-		// since it can mean there's no navigable geometry in regenerated tile
-		const bool bSuccess = TileGenerator.GenerateTile();
-
-		// cache recast navmesh actor instance if it's valid
-		ARecastNavMesh* CachedRecastNavMeshInstance = DestNavMesh.Get();
-		bool bAddTileSuccessful = false;
-
-		if (bSuccess && ShouldContinueBuilding() && UnrealNavBounds.IsValid)
-		{
-			// add newly generated set of layers
-			bAddTileSuccessful = AddTile(&TileGenerator, CachedRecastNavMeshInstance);
-
-			if (bAddTileSuccessful)
-			{
-				// store navmesh data in our UE4 object
-				bAddTileSuccessful = TransferGeneratedData();
-				TileGenerator.FinishRebuild();
-			}
-		} 
-		
-		if (bAddTileSuccessful == false)
+		if (Element.AsyncTask)
 		{
-			// abort current generation request
-			TileGenerator.AbortRebuild();
-
-			// mark dirty area
-			MarkAbortedGenerator(&TileGenerator);
-
-			// request rebuild only if owner navmesh _really_ exists
-			if (DestNavMesh.IsValid(/*bEvenIfPendingKill = */false, /*bThreadsafeTest = */true))
-			{
-				RequestDirtyTilesRebuild();
-			}		
+			Element.AsyncTask->EnsureCompletion();
+			delete Element.AsyncTask;
+			Element.AsyncTask = nullptr;
 		}
-
-		UE_LOG(LogNavigation, Log, TEXT("%s> Done Generating Tile %d,%d"),
-			TEXT_WEAKOBJ_NAME(DestNavMesh), TileGenerator.GetTileX(), TileGenerator.GetTileY());
-
-		TileGenerator.FinishAsyncBuild();
-		UpdateBuilding();
-		DestNavMesh->RequestDrawingUpdate();
-	}
-	else
-	{
-		UE_LOG(LogNavigation, Warning, TEXT("FRecastNavMeshGenerator::GenerateTile failed to trigger actual tile building due to Tile Index Being Out Off Array Bounds"));
 	}
+	
+	RunningDirtyTiles.Empty();
 }
 
-bool FRecastNavMeshGenerator::IsBuildingLocked() const 
-{ 
-	if (bBuildingLocked)
-	{
-		return true;
-	}
-
-#if WITH_NAVIGATION_GENERATOR
-	UWorld* World = GetWorld();
-	if (World && World->GetNavigationSystem() && World->GetNavigationSystem()->IsNavigationBuildingLocked())
-	{
-		return true;
-	}
-#endif // WITH_NAVIGATION_GENERATOR
-
-	return false;
-}
 
-void FRecastNavMeshGenerator::OnNavigationBuildingLocked()
+bool FRecastNavMeshGenerator::HasDirtyTiles() const
 {
-	bBuildingLocked = true;
+	return (PendingDirtyTiles.Num() > 0 || RunningDirtyTiles.Num() > 0);
 }
 
-void FRecastNavMeshGenerator::OnNavigationBuildingUnlocked(bool bForce)
+FBox FRecastNavMeshGenerator::GrowBoundingBox(const FBox& BBox, bool bIncludeAgentHeight) const
 {
-	if (bBuildingLocked == true || bForce == true)
-	{
-		bBuildingLocked = false;
-		bBuildFromScratchRequested |= bForce;
-		UpdateBuilding();
-	}
-}
+	const FVector BBoxGrowOffsetBoth = FVector(2.0f * Config.borderSize * Config.cs);
+	const FVector BBoxGrowOffsetMin = FVector(0, 0, bIncludeAgentHeight ? Config.AgentHeight : 0.0f);
 
-void FRecastNavMeshGenerator::TriggerGeneration()
-{
-	RequestGeneration();
+	return FBox(BBox.Min - BBoxGrowOffsetBoth - BBoxGrowOffsetMin, BBox.Max + BBoxGrowOffsetBoth);
 }
 
-void FRecastNavMeshGenerator::UpdateBuilding()
+TArray<FNavMeshTileData> FRecastNavMeshGenerator::GetCachedLayersData(FIntPoint GridCoord) const
 {
-	if (ShouldContinueBuilding() == false)
+	const TArray<FNavMeshTileData>* ExistingData = CachedLayerDataMap.Find(GridCoord);
+	if (ExistingData)
 	{
-		return;
+		return *ExistingData;
 	}
-	else if (bBuildFromScratchRequested == true)
-	{
-		bBuildFromScratchRequested = false;
-
-		{
-			FScopeLock Lock(&NavMeshDirtyLock);
-			DirtyAreas.Reset();
-			DirtyGenerators.Empty(MaxActiveGenerators);
-		}
-
-		// need to send it to main thread - uses some gamethread-only iterators
-		DECLARE_CYCLE_STAT(TEXT("FSimpleDelegateGraphTask.Requesting navmesh regen from UpdateBuilding"),
-			STAT_FSimpleDelegateGraphTask_RequestingNavmeshRegenFromUpdateBuilding,
-			STATGROUP_TaskGraphTasks);
-
-		FSimpleDelegateGraphTask::CreateAndDispatchWhenReady(
-			FSimpleDelegateGraphTask::FDelegate::CreateThreadSafeSP(this, &FRecastNavMeshGenerator::TriggerGeneration),
-			GET_STATID(STAT_FSimpleDelegateGraphTask_RequestingNavmeshRegenFromUpdateBuilding), NULL,
-			ENamedThreads::GameThread
-		);
-	}
-	else if(!UnrealNavBounds.IsValid)
-	{
-		return;
-	}
-	else if(bRebuildDirtyTilesRequested == false && AreAnyTilesBeingBuilt() == false)
+	else
 	{
-		RequestDirtyTilesRebuild();
+		return TArray<FNavMeshTileData>();
 	}
 }
 
-void FRecastNavMeshGenerator::RemoveTileLayers(const int32 TileX, const int32 TileY, TArray<FNavMeshGenerationResult>& AsyncResults)
+static FVector2D CalcTileCenter(int32 X, int32 Y, const FVector& NavMeshOrigin, float TileSizeInWorldUnits)
 {
-	SCOPE_CYCLE_COUNTER(STAT_Navigation_AddTile);
-	const int32 NumLayers = DetourMesh ? DetourMesh->getTileCountAt(TileX, TileY) : 0;
-	if (NumLayers <= 0)
-	{
-		return;
-	}
-
-	TArray<dtMeshTile*> Tiles;
-	Tiles.AddZeroed(NumLayers);
-	DetourMesh->getTilesAt(TileX, TileY, (const dtMeshTile**)Tiles.GetData(), NumLayers);
-	
-	for (int32 i = 0; i < NumLayers; i++)
-	{
-		const int32 LayerIndex = Tiles[i]->header->layer;
-		const dtTileRef TileRef = DetourMesh->getTileRef(Tiles[i]);
-
-		NumActiveTiles--;
-		UE_LOG(LogNavigation, Log, TEXT("%s> Tile (%d,%d:%d), removing TileRef: 0x%X (active:%d)"),
-			TEXT_WEAKOBJ_NAME(DestNavMesh), TileX, TileY, LayerIndex, TileRef, NumActiveTiles);
-
-		uint8* RawNavData = NULL;
-		DetourMesh->removeTile(TileRef, &RawNavData, NULL);
-
-		FNavMeshGenerationResult ResultInfo;
-		ResultInfo.OldRawNavData = RawNavData;
-		ResultInfo.OldTileRef = TileRef;
-		ResultInfo.TileIndex = DetourMesh->decodePolyIdTile(TileRef);
-		AsyncResults.Add(ResultInfo);
-	}
+	return (FVector2D(NavMeshOrigin) + FVector2D(X, Y)*TileSizeInWorldUnits + FVector2D(0.5f, 0.5f)*TileSizeInWorldUnits);
 }
 
-bool FRecastNavMeshGenerator::AddTile(FRecastTileGenerator* TileGenerator, ARecastNavMesh* CachedRecastNavMeshInstance)
+void FRecastNavMeshGenerator::MarkDirtyTiles(const TArray<FNavigationDirtyArea>& DirtyAreas)
 {
-	bool bOperationSuccessful = false;
-
-	SCOPE_CYCLE_COUNTER(STAT_Navigation_AddTile);
+	QUICK_SCOPE_CYCLE_COUNTER(STAT_RecastNavMeshGenerator_MarkDirtyTiles);
+	
+	check(bInitialized);
+	const float TileSizeInWorldUnits = Config.tileSize * Config.cs;
+	check(TileSizeInWorldUnits > 0);
+	const FVector NavMeshOrigin = FVector::ZeroVector;
+	int32 NumTilesMarked = 0;
 
-	SCOPE_CYCLE_COUNTER(STAT_Navigation_BuildTime)
-	STAT(double ThisTime = 0);
+	// find all tiles that need regeneration:
+	TSet<FTileElement> DirtyTiles;
+	for (const FNavigationDirtyArea& DirtyArea : DirtyAreas)
 	{
-		SCOPE_SECONDS_COUNTER(ThisTime);
-		TArray<FNavMeshTileData> TileLayers;
-		TileGenerator->TransferNavigationData(TileLayers);
-
-		TArray<FNavMeshGenerationResult> AsyncResults;
-		const int32 TileX = TileGenerator->GetTileX();
-		const int32 TileY = TileGenerator->GetTileY();
-
+		const FBox AdjustedAreaBounds = GrowBoundingBox(DirtyArea.Bounds, DirtyArea.HasFlag(ENavigationDirtyFlag::UseAgentHeight));
+		if (!IntercestsInclusionBounds(AdjustedAreaBounds))
 		{
-			FScopeLock Lock(&TileAddingLock);
-
-			// remove all layers if geometry has changed
-			if (TileGenerator->IsRebuildingGeometry())
-			{
-				RemoveTileLayers(TileX, TileY, AsyncResults);
-			}
+			// Skip whole area
+			continue;
+		}
+		
+		const int32 XMin = FMath::FloorToInt((AdjustedAreaBounds.Min.X - NavMeshOrigin.X) / TileSizeInWorldUnits);
+		const int32 XMax = FMath::FloorToInt((AdjustedAreaBounds.Max.X - NavMeshOrigin.X) / TileSizeInWorldUnits);
+		const int32 YMin = FMath::FloorToInt((AdjustedAreaBounds.Min.Y - NavMeshOrigin.Y) / TileSizeInWorldUnits);
+		const int32 YMax = FMath::FloorToInt((AdjustedAreaBounds.Max.Y - NavMeshOrigin.Y) / TileSizeInWorldUnits);
 
-			FTileSetItem* MyTileData = DestNavMesh.IsValid() ? DestNavMesh->GetTileSetItemAt(TileX, TileY) : NULL;
-			
-			if (MyTileData != NULL)
+		for (int32 y = YMin; y <= YMax; ++y)
+		{
+			for (int32 x = XMin; x <= XMax; ++x)
 			{
-				bOperationSuccessful = true;
-
-				MyTileData->bHasCompressedGeometry = true;
-
-				bool bHasNavmesh = true;
-				for (int32 i = 0; i < TileLayers.Num(); i++)
+				FBox TileBox(
+					FVector(x+0, y+0, 0)*TileSizeInWorldUnits, 
+					FVector(x+1, y+1, 0)*TileSizeInWorldUnits);
+				TileBox.Min.Z = TotalNavBounds.Min.Z;
+				TileBox.Max.Z = TotalNavBounds.Max.Z;
+				if (!IntercestsInclusionBounds(TileBox))
 				{
-					const int32 LayerIndex = TileLayers[i].LayerIndex;
-					const dtTileRef OldTileRef = DetourMesh->getTileRefAt(TileX, TileY, LayerIndex);
-
-					uint8* OldRawNavData = NULL;
-					if (OldTileRef)
+					// Skip this tile
+					continue;
+				}
+								
+				FTileElement Element;
+				Element.Coord = FIntPoint(x, y);
+				Element.bRebuildGeometry = (DirtyArea.Flags & ENavigationDirtyFlag::Geometry);
+				if (Element.bRebuildGeometry == false)
+				{
+					Element.DirtyAreas.Add(AdjustedAreaBounds);
+				}
+				
+				FTileElement* ExistingElement = DirtyTiles.Find(Element);
+				if (ExistingElement)
+				{
+					ExistingElement->bRebuildGeometry|= Element.bRebuildGeometry;
+					// Append area bounds to existing list 
+					if (ExistingElement->bRebuildGeometry == false)
 					{
-						NumActiveTiles--;
-						UE_LOG(LogNavigation, Log, TEXT("%s> Tile (%d,%d:%d), removing TileRef: 0x%X (active:%d)"),
-							TEXT_WEAKOBJ_NAME(DestNavMesh), TileX, TileY, LayerIndex, OldTileRef, NumActiveTiles);
-
-						// this call will fill OldRawNavData with address of previously added data
-						// if OldRawNavData is NULL it means either this tile was empty, or the memory
-						// allocated has been owned by navmesh itself (i.e. has been serialized-in along
-						// with the level)
-						DetourMesh->removeTile(OldTileRef, &OldRawNavData, NULL);
+						ExistingElement->DirtyAreas.Append(Element.DirtyAreas);
 					}
-
-					if (TileLayers[i].IsValid()) 
+					else
 					{
-						bool bRejectNavmesh = false;
-						dtTileRef ResultTileRef = 0;
-
-						// let navmesh know it's tile generator who owns the data
-						dtStatus status = DetourMesh->addTile(TileLayers[i].GetData(), TileLayers[i].DataSize, NAVMESH_TILE_GENERATOR_OWNS_DATA, 0, &ResultTileRef);
-
-						if (dtStatusFailed(status))
-						{
-							if (dtStatusDetail(status, DT_OUT_OF_MEMORY))
-							{
-								UE_LOG(LogNavigation, Error, TEXT("%s> Tile (%d,%d:%d), tile limit reached!! (%d)"),
-									TEXT_WEAKOBJ_NAME(DestNavMesh), TileX, TileY, LayerIndex, DetourMesh->getMaxTiles());
-							}
-
-							// release data here so that NavData is not stored in %
-							bRejectNavmesh = true;
-						}
-						else
-						{
-							NumActiveTiles++;
-							UE_LOG(LogNavigation, Log, TEXT("%s> Tile (%d,%d:%d), added TileRef: 0x%X (active:%d)"),
-								TEXT_WEAKOBJ_NAME(DestNavMesh), TileX, TileY, LayerIndex, ResultTileRef, NumActiveTiles);
-						}
-
-						if (bRejectNavmesh)
-						{
-							// release data here so that NavData is not stored in 
-							TileLayers[i].Release();
-							bHasNavmesh = false;
-						}
+						ExistingElement->DirtyAreas.Empty();
 					}
-
-					// store information that needs to be processed in game thread 
-					FNavMeshGenerationResult ResultInfo;
-					ResultInfo.OldRawNavData = OldRawNavData;
-					ResultInfo.NewNavData = TileLayers[i];
-					ResultInfo.OldTileRef = OldTileRef;
-					ResultInfo.TileIndex = DetourMesh->decodePolyIdTile(OldTileRef);
-					AsyncResults.Add(ResultInfo);
 				}
-
-				MyTileData->bHasNavmesh = bHasNavmesh;
+				else
+				{
+					DirtyTiles.Add(Element);
+				}
 			}
 		}
-
-		if (bOperationSuccessful)
-		{
-			// send off to game thread to not use critical section for modifying AsyncGenerationResultContainer
-			DECLARE_CYCLE_STAT(TEXT("FSimpleDelegateGraphTask.Storing async nav generetion result"),
-				STAT_FSimpleDelegateGraphTask_StoringAsyncNavGeneretionResult,
-				STATGROUP_TaskGraphTasks);
-
-			FSimpleDelegateGraphTask::CreateAndDispatchWhenReady(
-				FSimpleDelegateGraphTask::FDelegate::CreateThreadSafeSP(this, &FRecastNavMeshGenerator::StoreAsyncResults, AsyncResults),
-				GET_STATID(STAT_FSimpleDelegateGraphTask_StoringAsyncNavGeneretionResult), NULL, ENamedThreads::GameThread
-			);
-		}
-	}
-	INC_FLOAT_STAT_BY(STAT_Navigation_CumulativeBuildTime,(float)ThisTime*1000);
-
-	return bOperationSuccessful;
-}
-
-void FRecastNavMeshGenerator::StoreAsyncResults(TArray<FNavMeshGenerationResult> AsyncResults)
-{
-	// needs to be done on GameThread since we're not guarding access to AsyncGenerationResultContainer
-	// with any critical section
-	ensure(IsInGameThread());
-
-	AsyncGenerationResultContainer.Append(AsyncResults);
-}
-
-void FRecastNavMeshGenerator::GetAsyncResultsCopy(TNavStatArray<FNavMeshGenerationResult>& Dest, bool bClearSource)
-{
-	// needs to be done on GameThread since we're not guarding access to AsyncGenerationResultContainer
-	// with any critical section
-	ensure(IsInGameThread());
-
-	Dest = AsyncGenerationResultContainer;
-
-	if (bClearSource)
-	{
-		AsyncGenerationResultContainer.Reset();
 	}
-}
-
-bool FRecastNavMeshGenerator::HasDirtyTiles() const
-{
-	return (DirtyGenerators.Num() > 0);
-}
+	
+	NumTilesMarked = DirtyTiles.Num();
 
-void FRecastNavMeshGenerator::RebuildAll()
-{
-	if (bInitialized == false)
+	// Merge all pending tiles into one container
+	for (const FTileElement& Element : PendingDirtyTiles)
 	{
-		RequestGeneration();
+		FTileElement* ExistingElement = DirtyTiles.Find(Element);
+		if (ExistingElement)
+		{
+			ExistingElement->bRebuildGeometry|= Element.bRebuildGeometry;
+			// Append area bounds to existing list 
+			if (ExistingElement->bRebuildGeometry == false)
+			{
+				ExistingElement->DirtyAreas.Append(Element.DirtyAreas);
+			}
+			else
+			{
+				ExistingElement->DirtyAreas.Empty();
+			}
+		}
+		else
+		{
+			DirtyTiles.Add(Element);
+		}
 	}
-	else
+	
+	// Dump results into array
+	PendingDirtyTiles.Empty(DirtyTiles.Num());
+	for(const FTileElement& Element : DirtyTiles)
 	{
-		// if rebuilding all no point in keeping "old" invalidated areas
-		FNavigationDirtyArea BigArea(Recast2UnrealBox(Config.bmin, Config.bmax), ENavigationDirtyFlag::All);
-		FScopeLock Lock(&NavMeshDirtyLock);
-
-		DirtyAreas.Empty();
-		DirtyAreas.Add(BigArea);
-		RequestDirtyTilesRebuild();
+		PendingDirtyTiles.Add(Element);
 	}
 
-	if (DestNavMesh.IsValid())
+	// Sort tiles by proximity to players 
+	if (NumTilesMarked > 0)
 	{
-		DestNavMesh->UpdateNavVersion();
+		SortPendingBuildTiles();
 	}
 }
 
-void FRecastNavMeshGenerator::MarkAbortedGenerator(const FRecastTileGenerator* TileGenerator)
+void FRecastNavMeshGenerator::SortPendingBuildTiles()
 {
-	FNavigationDirtyArea DirtyArea(TileGenerator->GetUnrealBB().ExpandBy(-1.0f),
-		TileGenerator->HasDirtyGeometry() ? ENavigationDirtyFlag::All : ENavigationDirtyFlag::DynamicModifier);
-
-	FScopeLock Lock(&NavMeshDirtyLock);
-	DirtyAreas.Add(DirtyArea);
-}
-
-void FRecastNavMeshGenerator::MarkDirtyGenerators()
-{
-	if (DirtyAreas.Num() == 0)
+	TArray<FVector2D> SeedLocations;
+	UWorld* CurWorld = GetWorld();
+	if (CurWorld == nullptr)
 	{
 		return;
 	}
 
-	SCOPE_CYCLE_COUNTER(STAT_Navigation_BuildTime)
-
-	// create a copy of dirty areas array
-	TNavStatArray<FNavigationDirtyArea> DirtyAreasCopy;
+	// Collect players positions
+	for (FConstPlayerControllerIterator PlayerIt = CurWorld->GetPlayerControllerIterator(); PlayerIt; ++PlayerIt)
 	{
-		FScopeLock Lock(&NavMeshDirtyLock);
-		DirtyAreasCopy = DirtyAreas;
-		DirtyAreas.Reset();
-	}
-
-	FScopeLock Lock(&TileGenerationLock);
-	const float InvTileCellSize = 1.0f / (Config.tileSize * Config.cs);
-
-	TSet<int32> DirtyIndices;
-
-	// find all tiles that need regeneration:
-	const FNavigationDirtyArea* DirtyArea = DirtyAreasCopy.GetData();
-	for (int32 i = 0; i < DirtyAreasCopy.Num(); ++i, ++DirtyArea)
-	{
-		const FBox AdjustedBounds = GrowBoundingBox(DirtyArea->Bounds, DirtyArea->HasFlag(ENavigationDirtyFlag::UseAgentHeight));
-		const FBox RCBB = Unreal2RecastBox(AdjustedBounds);
-		const int32 XMin = FMath::Max(FMath::TruncToInt((RCBB.Min.X - RCNavBounds.Min.X) * InvTileCellSize), 0);
-		const int32 YMin = FMath::Max(FMath::TruncToInt((RCBB.Min.Z - RCNavBounds.Min.Z) * InvTileCellSize), 0);
-		const int32 XMax = FMath::Min(FMath::TruncToInt((RCBB.Max.X - RCNavBounds.Min.X) * InvTileCellSize), TilesWidth-1);
-		const int32 YMax = FMath::Min(FMath::TruncToInt((RCBB.Max.Z - RCNavBounds.Min.Z) * InvTileCellSize), TilesHeight-1);
-
-		for (int32 y = YMin; y <= YMax; ++y)
+		if (*PlayerIt && (*PlayerIt)->GetPawn() != NULL)
 		{
-			for (int32 x = XMin; x <= XMax; ++x)
-			{
-				// grab that generator and mark it as dirty
-				// @todo make thread safe
-
-				const int32 GenIdx = x + (y * TilesWidth);
-				TileGenerators[GenIdx].SetDirty(*DirtyArea, AdjustedBounds);
-				DirtyIndices.Add(GenIdx);
-			}
+			const FVector2D SeedLoc((*PlayerIt)->GetPawn()->GetActorLocation());
+			SeedLocations.Add(SeedLoc);
 		}
 	}
 
-	// store current dirty state associated with just added areas.
-	//
-	// Generators that already received octree data and are waiting for free
-	// worker thread will clear dirty state after finishing rebuild.
-	// This lost information will prevent them from running again, when
-	// building queue is finally emptied and StartDirtyGenerators is called.
-
-	FRecastTileDirtyState DirtyState;
-	for (TSet<int32>::TIterator It(DirtyIndices); It; ++It)
-	{
-		const int32 GenIdx = *It;
-		TileGenerators[GenIdx].GetDirtyState(DirtyState);
-		DirtyGenerators.Add(GenIdx, DirtyState);
-	}
-}
-
-struct FCompareNavMeshTiles
-{
-	const FTileSetItem* TileSet;
-
-	FCompareNavMeshTiles(const FTileSetItem* InTileSet) : TileSet(InTileSet) {}
-
-	FORCEINLINE bool operator()(const int32& A, const int32& B) const
-	{
-		const FTileSetItem* TileA = &TileSet[A];
-		const FTileSetItem* TileB = &TileSet[B];
-
-		// lesser goes first
-		return TileA->SortOrder < TileB->SortOrder;
-	}
-};
-
-void FRecastNavMeshGenerator::StartDirtyGenerators()
-{
-	if (DirtyGenerators.Num() == 0 || GeneratorsQueue.Num() > 0)
+	if (SeedLocations.Num() == 0)
 	{
-		return;
+		// Use navmesh origin for sorting
+		SeedLocations.Add(FVector2D::ZeroVector);
 	}
 
-	const UNavigationSystem* NavSys = UNavigationSystem::GetCurrent(GetWorld());
-	const FNavigationOctree* NavOctree = NavSys ? NavSys->GetNavOctree() : NULL;
-	if (NavOctree == NULL)
+	if (SeedLocations.Num() > 0)
 	{
-		UE_LOG(LogNavigation, Error, TEXT("Failed to rebuild dirty navmesh tiles due to %s being NULL"),
-			NavSys == NULL ? TEXT("NavigationSystem") : TEXT("NavOctree"));
-
-		return;
-	}
-
-	UE_LOG(LogNavigation, Log, TEXT("%s> StartDirtyGenerators"), TEXT_WEAKOBJ_NAME(DestNavMesh));
-
-	STAT(double ThisTime = 0);
-	{
-		SCOPE_SECONDS_COUNTER(ThisTime);
-		SCOPE_CYCLE_COUNTER(STAT_Navigation_BuildTime)
-
-		const float TileCellSize = (Config.tileSize * Config.cs);
-		const float* BMin = Config.bmin;
-		const float* BMax = Config.bmax;
-
-		FCompareNavMeshTiles SortFunc(DestNavMesh->GetTileSet());
-
-		FScopeLock Lock(&TileGenerationLock);
-		DirtyGenerators.KeySort(SortFunc);
-
-		FRecastTileDirtyState DirtyState;
-		const bool bRequestDrawingUpdate = DirtyGenerators.Num() > 0;
-		for (TMap<int32,FRecastTileDirtyState>::TIterator It(DirtyGenerators); It; ++It)
+		const float TileSizeInWorldUnits = Config.tileSize * Config.cs;
+		
+		// Calculate shortest distances between tiles and players
+		for (FTileElement& Element : PendingDirtyTiles)
 		{
-			const int32 GenIdx = It.Key();
-			DirtyState = It.Value();
-
-			FRecastTileGenerator* TileGenerator = &TileGenerators[GenIdx];
-
-			UE_LOG(LogNavigation, Log, TEXT("%s> Tile %d,%d dirty (%s)"), TEXT_WEAKOBJ_NAME(DestNavMesh),
-				TileGenerator->GetTileX(), TileGenerator->GetTileY(),
-				DirtyState.bRebuildGeometry ? TEXT("geometry") : TEXT("layers"));
-
-			// if tile's regenerating at the moment put it into delayed tile generators container
-			if (TileGenerator->IsBeingRebuild() == true)
+			FVector2D TileLocation = FVector2D(Element.Coord)*TileSizeInWorldUnits;
+			for (FVector2D SeedLocation : SeedLocations)
 			{
-				UE_LOG(LogNavigation, Log, TEXT("%s> Tile %d,%d is currently being built - postpone"),
-					TEXT_WEAKOBJ_NAME(DestNavMesh), TileGenerator->GetTileX(), TileGenerator->GetTileY());
-
-				continue;
-			}
-
-			// initialize full rebuild when geometry has changed
-			TileGenerator->ClearModifiers();
-			TileGenerator->SetDirtyState(DirtyState);
-
-			if (TileGenerator->HasDirtyGeometry())
-			{
-				const int32 X = TileGenerator->GetTileX();
-				const int32 Y = TileGenerator->GetTileY();
-
-				const float TileBMin[] = { BMin[0] + ((X + 0) * TileCellSize), BMin[1], BMin[2] + ((Y + 0) * TileCellSize) };
-				const float TileBMax[] = { BMin[0] + ((X + 1) * TileCellSize), BMax[1], BMin[2] + ((Y + 1) * TileCellSize) };
-
-				TileGenerator->ClearGeometry();
-				TileGenerator->Init(this, X, Y, TileBMin, TileBMax, InclusionBounds);
-
-				// Init call does some checks whether there's a point in starting generation at all
-				if (!TileGenerator->ShouldBeBuilt())
+				const float DistSq = FVector2D::DistSquared(TileLocation, SeedLocation);
+				if (DistSq < Element.SeedDistance)
 				{
-					TileGenerator->AbadonGeneration();
+					Element.SeedDistance = DistSq;
 				}
 			}
-
-			It.RemoveCurrent();
-
-			// skip if dirty flags were removed
-			if (!TileGenerator->HasDirtyGeometry() && !TileGenerator->HasDirtyLayers())
-			{
-				if (TileGenerator->IsPendingRebuild())
-				{
-					GeneratorsQueue.RemoveSingle(TileGenerator);
-				}
-
-				continue;
-			}
-
-			FillGeneratorData(TileGenerator, NavOctree);
-
-#if RECAST_ASYNC_REBUILDING
-			// async build: add to queue, up to MaxActiveGenerators at the same time
-			if (!TileGenerator->IsPendingRebuild())
-			{
-				TileGenerator->MarkPendingRebuild();
-				GeneratorsQueue.Add(TileGenerator);
-			}
-
-			if (GeneratorsQueue.Num() >= MaxActiveGenerators)
-			{
-				break;
-			}
-#else
-			// sync build: generate immediately
-			TileGenerator->InitiateRebuild();
-			GenerateTile(TileGenerator->GetId(), TileGenerator->GetVersion());
-#endif
 		}
 
-#if !RECAST_ASYNC_REBUILDING
-		// sync building: clean up static allocations
-		FRecastTileGenerator::ClearStaticData();
-#endif
-
-		// start processing generator queue
-		UpdateTileGenerationWorkers(INDEX_NONE);
-		if (bRequestDrawingUpdate)
-		{
-			DestNavMesh->RequestDrawingUpdate();
-		}
+		// nearest tiles should be at the end of the list
+		PendingDirtyTiles.Sort();
 	}
-
-	INC_FLOAT_STAT_BY(STAT_Navigation_CumulativeBuildTime,(float)ThisTime*1000);
-
-#if STATS
-	UE_LOG(LogNavigation, Log, TEXT("FRecastNavMeshGenerator::StartDirtyGenerators time: %.4fms"), ThisTime*1000);
-#endif
 }
 
-void FRecastNavMeshGenerator::FillGeneratorData(FRecastTileGenerator* TileGenerator, const FNavigationOctree* NavOctree)
+int32 FRecastNavMeshGenerator::SubmitDirtyTiles(const int32 NumTasksToSubmit)
 {
-	const bool bUseVoxelCache = ARecastNavMesh::IsVoxelCacheEnabled();
-	const FNavAgentProperties NavAgentProps = DestNavMesh->NavDataConfig;
-
-	for(FNavigationOctree::TConstElementBoxIterator<FNavigationOctree::DefaultStackAllocator> It(*NavOctree, GrowBoundingBox(TileGenerator->GetUnrealBB(), /*bIncludeAgentHeight=*/false));
-		It.HasPendingElements();
-		It.Advance())
+	QUICK_SCOPE_CYCLE_COUNTER(STAT_RecastNavMeshGenerator_SubmitDirtyTiles);
+	
+	if (PendingDirtyTiles.Num() == 0)
 	{
-		const FNavigationOctreeElement& Element = It.GetCurrentElement();
-		const bool bShouldUse = Element.ShouldUseGeometry(&DestNavMesh->NavDataConfig);
-		if (bShouldUse)
-		{
-			const bool bExportGeometry = TileGenerator->HasDirtyGeometry() && Element.Data.HasGeometry();
-			if (bExportGeometry)
-			{
-				if (bUseVoxelCache)
-				{
-					TNavStatArray<rcSpanCache> SpanData;
-					rcSpanCache* CachedVoxels = 0;
-					int32 NumCachedVoxels = 0;
-
-					DECLARE_SCOPE_CYCLE_COUNTER(TEXT("Rasterization: prepare voxel cache"), Stat_RecastRasterCachePrep, STATGROUP_Navigation);
-
-					if (!TileGenerator->HasVoxelCache(Element.Data.VoxelData, CachedVoxels, NumCachedVoxels))
-					{
-						// rasterize
-						TileGenerator->PrepareVoxelCache(Element.Data.CollisionData, SpanData);
-						CachedVoxels = SpanData.GetData();
-						NumCachedVoxels = SpanData.Num();
+		return 0;
+	}
+		
+	// 
+	int32 NumSubmittedTasks = 0;
+	for (int32 ElementIdx = PendingDirtyTiles.Num()-1; ElementIdx >= 0 && NumSubmittedTasks < NumTasksToSubmit; ElementIdx--)
+	{
+		FTileElement& Element = PendingDirtyTiles[ElementIdx];
+		// Make sure that we are not submitting generator for grid cell that is currently being regenerated
+		if (!RunningDirtyTiles.Contains(Element))
+		{
+			check(Element.AsyncTask == nullptr);
+
+			// Spawn async task
+			Element.AsyncTask = new FRecastTileGeneratorTask(
+				this,
+				Element.Coord.X,
+				Element.Coord.Y,
+				Element.DirtyAreas
+				);
+			Element.AsyncTask->StartBackgroundTask();
+			
+			RunningDirtyTiles.Add(Element);
+			NumSubmittedTasks++;
 
-						// encode
-						const int32 PrevElementMemory = Element.Data.GetAllocatedSize();
-						FNavigationRelevantData* ModData = (FNavigationRelevantData*)&Element.Data;
-						TileGenerator->AddVoxelCache(ModData->VoxelData, CachedVoxels, NumCachedVoxels);
+			// Remove submitted element from pending list
+			PendingDirtyTiles.RemoveAt(ElementIdx);
 
-						const int32 NewElementMemory = Element.Data.GetAllocatedSize();
-						const int32 ElementMemoryDelta = NewElementMemory - PrevElementMemory;
-						INC_MEMORY_STAT_BY(STAT_Navigation_CollisionTreeMemory, ElementMemoryDelta);
-					}
-
-					TileGenerator->AppendVoxels(CachedVoxels, NumCachedVoxels);
-				}
-				else
-				{
-					TileGenerator->AppendGeometry(Element.Data.CollisionData);
-				}
+			// Release memory, list could be quite big after map load
+			if (PendingDirtyTiles.Num() == 0)
+			{
+				PendingDirtyTiles.Empty(32);
 			}
-
-			const FCompositeNavModifier ModifierInstance = Element.GetModifierForAgent(&NavAgentProps);
-			TileGenerator->AppendModifier(ModifierInstance, bExportGeometry);
 		}
 	}
+	return NumSubmittedTasks;
 }
 
-void FRecastNavMeshGenerator::RegenerateDirtyTiles()
+TArray<uint32> FRecastNavMeshGenerator::CollectCompletedTiles()
 {
-	bRebuildDirtyTilesRequested = false;
-
-	if (ShouldContinueBuilding() == false || !UnrealNavBounds.IsValid || DestNavMesh.IsValid() == false)
-	{
-		return;
-	}
+	QUICK_SCOPE_CYCLE_COUNTER(STAT_RecastNavMeshGenerator_CollectCompletedTiles);
 
-	MarkDirtyGenerators();
-	StartDirtyGenerators();
-}
+	TArray<uint32> CompletedTiles;
+	CompletedTiles.Reserve(RunningDirtyTiles.Num());
 
-void FRecastNavMeshGenerator::UpdateTileGenerationWorkers(int32 TileId)
-{
-#if RECAST_ASYNC_REBUILDING
-	bool bRequestRenderingDirty = false;
-	FRecastTileGenerator** CurrentGenerator = ActiveGenerators.GetData();
-	int32 QueueIndex = 0;
-
-	for (int32 i = 0; i < ActiveGenerators.Num(); ++i, ++CurrentGenerator)
+	for (int32 Idx = RunningDirtyTiles.Num() - 1; Idx >=0; --Idx)
 	{
-		if (*CurrentGenerator == NULL || (*CurrentGenerator != NULL && (*CurrentGenerator)->GetId() == TileId))
-		{
-			*CurrentGenerator = NULL;
+		FTileElement& Element = RunningDirtyTiles[Idx];
+		check(Element.AsyncTask);
 
-			while (QueueIndex < GeneratorsQueue.Num())
+		if (Element.AsyncTask->IsDone())
+		{
+			// Add generated tiles to navmesh
+			const FRecastTileGenerator& TileGenerator = Element.AsyncTask->GetTask();
+			TArray<uint32> UpdatedTileIndices = AddTile(TileGenerator);
+			CompletedTiles.Append(UpdatedTileIndices);
+			
+			// Store layers data in the cache, so it can be reused later
+			// TODO: make this optional?
+			TArray<FNavMeshTileData> ComressedLayers = TileGenerator.GetCompressedLayers();
+			if (ComressedLayers.Num())
 			{
-				FRecastTileGenerator* GeneratorCandidate = GeneratorsQueue[QueueIndex++];
-				// generators in GeneratorsQueue should have FRecastTileGenerator::bRebuildPending
-				// @todo this is a safety feature, should not be needed once generators queuing is done less dodgy
-				if (GeneratorCandidate != NULL && GeneratorCandidate->IsPendingRebuild())
-				{
-					*CurrentGenerator = GeneratorCandidate;
-					(*CurrentGenerator)->TriggerAsyncBuild();
-					bRequestRenderingDirty = true;
-					break;
-				}
+				CachedLayerDataMap.Add(Element.Coord, ComressedLayers);
+			}
+			else
+			{
+				CachedLayerDataMap.Remove(Element.Coord);
 			}
-		}
-	}
-
-	if (QueueIndex > 0)
-	{
-		GeneratorsQueue.RemoveAt(0, QueueIndex, /*bAllowShrinking=*/false);
-	}
 
-	if (bRequestRenderingDirty)
-	{
-		// let navmesh drawing mark tiles being built
-		DestNavMesh->RequestDrawingUpdate();
+			// Destroy tile generator task
+			delete Element.AsyncTask;
+			Element.AsyncTask = nullptr;
+			// Remove completed tile element from a list of running tasks
+			RunningDirtyTiles.RemoveAt(Idx);
+		}
 	}
-#endif // RECAST_ASYNC_REBUILDING
 
-	// prepare next batch of dirty generators in next tick
-	if (GeneratorsQueue.Num() == 0 && DirtyGenerators.Num() > 0)
-	{
-		DECLARE_CYCLE_STAT(TEXT("FSimpleDelegateGraphTask.Prepare navmesh tile generators"),
-			STAT_FSimpleDelegateGraphTask_PrepareNavmeshTileGenerators,
-			STATGROUP_TaskGraphTasks);
-
-		FSimpleDelegateGraphTask::CreateAndDispatchWhenReady(
-			FSimpleDelegateGraphTask::FDelegate::CreateRaw(this, &FRecastNavMeshGenerator::StartDirtyGenerators), 
-			GET_STATID(STAT_FSimpleDelegateGraphTask_PrepareNavmeshTileGenerators), NULL,
-			ENamedThreads::GameThread
-		);
-	}
+	return CompletedTiles;
 }
 
 void FRecastNavMeshGenerator::ExportComponentGeometry(UActorComponent* Component, FNavigationRelevantData& Data)
@@ -4463,165 +3452,57 @@ void FRecastNavMeshGenerator::ExportRigidBodyGeometry(UBodySetup& BodySetup, TNa
 	}
 }
 
-bool FRecastNavMeshGenerator::TransferGeneratedData()
+bool FRecastNavMeshGenerator::IsBuildInProgress(bool bCheckDirtyToo) const
 {
-	if (DestNavMesh.IsValid() && DestNavMesh->GetRecastNavMeshImpl() != NULL)
+	if (RunningDirtyTiles.Num() || (bCheckDirtyToo && PendingDirtyTiles.Num()))
 	{
-		// make sure FPImplRecastNavMesh knows data is not his - generator own that data at the moment
-		DestNavMesh->GetRecastNavMeshImpl()->SetRecastMesh(DetourMesh, /*bOwnData=*/false);
-
-		// this should be done synchronously
-		DECLARE_CYCLE_STAT(TEXT("FSimpleDelegateGraphTask.Requesting navmesh redraw"),
-			STAT_FSimpleDelegateGraphTask_RequestingNavmeshRedraw,
-			STATGROUP_TaskGraphTasks);
-
-		FSimpleDelegateGraphTask::CreateAndDispatchWhenReady(
-			FSimpleDelegateGraphTask::FDelegate::CreateUObject(DestNavMesh.Get(), &ARecastNavMesh::UpdateNavMeshDrawing),
-			GET_STATID(STAT_FSimpleDelegateGraphTask_RequestingNavmeshRedraw), NULL, ENamedThreads::GameThread);
-
 		return true;
 	}
-	else
-	{
-		return false;
-	}
-}
-
-void FRecastNavMeshGenerator::RefreshParentReference()
-{
-#if WITH_EDITOR
-	FScopeLock Lock(&GatheringDataLock);
-
-	// we need to make sure there's any navmesh at all to rebuild 
-	if (DetourMesh == NULL)
-	{
-		// first try to get it from owner-ARecastNavMesh
-		FPImplRecastNavMesh* RecastMeshWrapper = DestNavMesh->GetRecastNavMeshImpl();
-		if (RecastMeshWrapper != NULL)
-		{
-			SetDetourMesh(RecastMeshWrapper->GetRecastMesh(), FRecastNavMeshGenerator::DO_ForeignData);
-		}
-
-		// if it's still empty then create one
-		if (DetourMesh == NULL && ConstructTiledNavMesh() == false)
-		{
-			UE_LOG(LogNavigation, Error, TEXT("Failed to find and generate Recast navmesh to rebuild"));
-		}
-	}
-#endif	//#if WITH_EDITOR
-}
 
-void FRecastNavMeshGenerator::OnNavigationBoundsUpdated(AVolume* Volume)
-{
-	if (Cast<ANavMeshBoundsVolume>(Volume) != NULL && Volume->BrushComponent.IsValid())
-	{
-		RequestGeneration();
-	}
+	return false;
 }
 
-void FRecastNavMeshGenerator::OnNavigationDataDestroyed(class ANavigationData* NavData)
+int32 FRecastNavMeshGenerator::GetNumRemaningBuildTasks() const
 {
-	if (NavData == DestNavMesh.Get())
-	{
-		SetDetourMesh(NULL);
-		DestNavMesh = NULL;
-		GeneratorsQueue.Reset();
-		ActiveGenerators.Reset();
-		bAbortAllTileGeneration = true;
-	}
+	return RunningDirtyTiles.Num() + PendingDirtyTiles.Num();
 }
 
-bool FRecastNavMeshGenerator::IsBuildInProgress(bool bCheckDirtyToo) const
+int32 FRecastNavMeshGenerator::GetNumRunningBuildTasks() const
 {
-	bool bRetValue = AreAnyTilesBeingBuilt(bCheckDirtyToo);
-	if (bCheckDirtyToo)
-	{
-		bRetValue = bRetValue || DirtyAreas.Num() > 0 || DirtyGenerators.Num() > 0;
-	}
-
-	return bRetValue;
+	return RunningDirtyTiles.Num();
 }
 
-bool FRecastNavMeshGenerator::AreAnyTilesBeingBuilt(bool bCheckDirtyToo) const
+bool FRecastNavMeshGenerator::IsTileChanged(int32 TileIdx) const
 {
-#if RECAST_ASYNC_REBUILDING	
-	bool bRet = false;
-	const FRecastTileGenerator* TileGenerator = TileGenerators.GetData();
-
-	for (int32 i = 0; i < TileGenerators.Num(); ++i, ++TileGenerator)
+#if WITH_EDITOR	
+	// Check recently built tiles
+	if (TileIdx > 0)
 	{
-		if ((bCheckDirtyToo && TileGenerator->IsDirty()) || TileGenerator->IsBeingRebuild())
+		FTileTimestamp TileTimestamp;
+		TileTimestamp.TileIdx = static_cast<uint32>(TileIdx);
+		if (RecentlyBuiltTiles.Contains(TileTimestamp))
 		{
-			bRet = true;
-			break;
-		}
-	}
-
-	return bRet;
-#else
-	// no need to check if anything is building since all tiles are build "instantly" 
-	// (blocking the one-and-only thread)
-	return false;
-#endif // RECAST_ASYNC_REBUILDING
-}
-
-bool FRecastNavMeshGenerator::IsAsyncBuildInProgress() const
-{
-#if RECAST_ASYNC_REBUILDING	
-	const FRecastTileGenerator* TileGenerator = TileGenerators.GetData();
-	for (int32 i = 0; i < TileGenerators.Num(); ++i, ++TileGenerator)
-	{
-		if (TileGenerator->IsAsyncBuildInProgress())
-		{
-			UE_LOG(LogNavigation, Log, TEXT("Waiting for async build of tile (%d,%d)"), TileGenerator->GetTileX(), TileGenerator->GetTileY());
 			return true;
 		}
 	}
-#endif
+#endif//WITH_EDITOR
 
 	return false;
 }
 
-bool FRecastNavMeshGenerator::IsTileFresh(int32 X, int32 Y, float FreshnessTime) const
-{
-	const int32 TileIndex = GetTileIdAt(X, Y);
-	if (TileGenerators.IsValidIndex(TileIndex) == false)
-	{
-		return false;
-	}
-	const FRecastTileGenerator* TileGenerator = &TileGenerators[TileIndex];
-
-	return TileGenerator != NULL && (TileGenerator->IsDirty() || TileGenerator->IsBeingRebuild() 
-#if !(UE_BUILD_SHIPPING || UE_BUILD_TEST)
-		|| (FPlatformTime::Seconds() - TileGenerator->GetLastBuildTimeStamp() < FreshnessTime)
-#endif // !(UE_BUILD_SHIPPING || UE_BUILD_TEST)
-		);
-}
-
-void FRecastNavMeshGenerator::OnWorldInitDone(bool bAllowedToRebuild)
-{
-	Init();
-
-	if ((DetourMesh != NULL || ConstructTiledNavMesh() == true) && bAllowedToRebuild)
-	{
-		RebuildAll();
-	}
-}
-
 uint32 FRecastNavMeshGenerator::LogMemUsed() const 
 {
 	UE_LOG(LogNavigation, Display, TEXT("    FRecastNavMeshGenerator: self %d"), sizeof(FRecastNavMeshGenerator));
 	
 	uint32 GeneratorsMem = 0;
-	for (int32 i = 0; i < TileGenerators.Num(); ++i)
+	for (FTileElement Element : RunningDirtyTiles)
 	{
-		GeneratorsMem += TileGenerators[i].GetUsedMemCount();
+		GeneratorsMem+= Element.AsyncTask->GetTask().UsedMemoryOnStartup; 
 	}
 
-	GeneratorsMem += TileGenerators.GetAllocatedSize();
-	UE_LOG(LogNavigation, Display, TEXT("    FRecastNavMeshGenerator: Total Generator\'s size %u, count %d"), GeneratorsMem, TileGenerators.Num());
+	UE_LOG(LogNavigation, Display, TEXT("    FRecastNavMeshGenerator: Total Generator\'s size %u, count %d"), GeneratorsMem, RunningDirtyTiles.Num());
 
-	return GeneratorsMem + sizeof(FRecastNavMeshGenerator);
+	return GeneratorsMem + sizeof(FRecastNavMeshGenerator) + PendingDirtyTiles.GetAllocatedSize() + PendingDirtyTiles.GetAllocatedSize();
 }
 
 #if !(UE_BUILD_SHIPPING || UE_BUILD_TEST)
@@ -4653,7 +3534,7 @@ void FRecastNavMeshGenerator::ExportNavigationData(const FString& FileName) cons
 			};
 			TArray<FAreaExportData> AreaExport;
 
-			for(FNavigationOctree::TConstElementBoxIterator<FNavigationOctree::DefaultStackAllocator> It(*NavOctree, UnrealNavBounds);
+			for(FNavigationOctree::TConstElementBoxIterator<FNavigationOctree::DefaultStackAllocator> It(*NavOctree, TotalNavBounds);
 				It.HasPendingElements();
 				It.Advance())
 			{
@@ -4755,7 +3636,6 @@ void FRecastNavMeshGenerator::ExportNavigationData(const FString& FileName) cons
 			}
 
 			AdditionalData += "# RecastDemo specific data\n";
-
 	#if 0
 			// use this bounds to have accurate navigation data bounds
 			const FVector Center = Unreal2RecastPoint(NavData->GetBounds().GetCenter());
@@ -4763,6 +3643,7 @@ void FRecastNavMeshGenerator::ExportNavigationData(const FString& FileName) cons
 			Extent = FVector(Extent.X, Extent.Z, Extent.Y);
 	#else
 			// this bounds match navigation bounds from level
+			FBox RCNavBounds = Unreal2RecastBox(TotalNavBounds);
 			const FVector Center = RCNavBounds.GetCenter();
 			const FVector Extent = RCNavBounds.GetExtent();
 	#endif
@@ -4773,8 +3654,6 @@ void FRecastNavMeshGenerator::ExportNavigationData(const FString& FileName) cons
 				Box.Max.X, Box.Max.Y, Box.Max.Z
 			);
 			
-#if WITH_NAVIGATION_GENERATOR
-			
 			const FRecastNavMeshGenerator* CurrentGen = static_cast<const FRecastNavMeshGenerator*>(NavData->GetGenerator());
 			check(CurrentGen);
 			AdditionalData += FString::Printf(TEXT("# AgentHeight\n"));
@@ -4813,7 +3692,6 @@ void FRecastNavMeshGenerator::ExportNavigationData(const FString& FileName) cons
 
 			AdditionalData += FString::Printf(TEXT("\n"));
 			
-#endif // WITH_NAVIGATION_GENERATOR
 			const FString FilePathName = FileName + FString::Printf(TEXT("_NavDataSet%d_%s.obj"), Index, *CurrentTimeStr) ;
 			ExportGeomToOBJFile(FilePathName, CoordBuffer, IndexBuffer, AdditionalData);
 		}
@@ -4840,7 +3718,6 @@ public:
 			{
 				if (const ANavigationData* NavData = InWorld->GetNavigationSystem()->GetMainNavData())
 				{
-#if WITH_NAVIGATION_GENERATOR
 					if (const FNavDataGenerator* Generator = NavData->GetGenerator())
 					{
 						const FString Name = NavData->GetName();
@@ -4848,7 +3725,6 @@ public:
 						return true;
 					}
 					else
-#endif // WITH_NAVIGATION_GENERATOR
 					{
 						UE_LOG(LogNavigation, Error, TEXT("Failed to export navigation data due to missing generator"));
 					}
diff --git a/Engine/Source/Runtime/Engine/Private/CheatManager.cpp b/Engine/Source/Runtime/Engine/Private/CheatManager.cpp
--- a/Engine/Source/Runtime/Engine/Private/CheatManager.cpp
+++ b/Engine/Source/Runtime/Engine/Private/CheatManager.cpp
@@ -890,13 +890,11 @@ void UCheatManager::WidgetReflector()
 
 void UCheatManager::RebuildNavigation()
 {
-#if WITH_NAVIGATION_GENERATOR
 	UNavigationSystem* NavSys = UNavigationSystem::GetCurrent(GetWorld());
 	if (NavSys)
 	{
 		NavSys->Build();
 	}
-#endif // WITH_NAVIGATION_GENERATOR
 }
 
 void UCheatManager::SetNavDrawDistance(float DrawDistance)
diff --git a/Engine/Source/Runtime/Engine/Private/NavMeshBoundsVolume.cpp b/Engine/Source/Runtime/Engine/Private/NavMeshBoundsVolume.cpp
--- a/Engine/Source/Runtime/Engine/Private/NavMeshBoundsVolume.cpp
+++ b/Engine/Source/Runtime/Engine/Private/NavMeshBoundsVolume.cpp
@@ -22,7 +22,6 @@ void ANavMeshBoundsVolume::PostEditChangeProperty( struct FPropertyChangedEvent&
 {
 	Super::PostEditChangeProperty(PropertyChangedEvent);
 
-#if WITH_NAVIGATION_GENERATOR
 	UNavigationSystem* NavSys = UNavigationSystem::GetCurrent(GetWorld());
 	if (GIsEditor && NavSys)
 	{
@@ -32,33 +31,28 @@ void ANavMeshBoundsVolume::PostEditChangeProperty( struct FPropertyChangedEvent&
 			NavSys->OnNavigationBoundsUpdated(this);
 		}
 	}
-#endif // WITH_NAVIGATION_GENERATOR
 }
 
 #endif // WITH_EDITOR
 
-void ANavMeshBoundsVolume::PostInitializeComponents() 
+void ANavMeshBoundsVolume::PostRegisterAllComponents() 
 {
-	Super::PostInitializeComponents();
-
-#if WITH_NAVIGATION_GENERATOR
+	Super::PostRegisterAllComponents();
+	
 	UNavigationSystem* NavSys = UNavigationSystem::GetCurrent(GetWorld());
 	if (NavSys && Role == ROLE_Authority)
 	{
 		NavSys->OnNavigationBoundsUpdated(this);
 	}
-#endif
 }
 
-void ANavMeshBoundsVolume::PostLoad()
+void ANavMeshBoundsVolume::PostUnregisterAllComponents() 
 {
-	Super::PostLoad();
-
-#if WITH_NAVIGATION_GENERATOR
+	Super::PostUnregisterAllComponents();
+	
 	UNavigationSystem* NavSys = UNavigationSystem::GetCurrent(GetWorld());
 	if (NavSys && Role == ROLE_Authority)
 	{
 		NavSys->OnNavigationBoundsUpdated(this);
 	}
-#endif
-}
\ No newline at end of file
+}
diff --git a/Engine/Source/Runtime/Engine/Private/NavMeshRenderingComponent.cpp b/Engine/Source/Runtime/Engine/Private/NavMeshRenderingComponent.cpp
--- a/Engine/Source/Runtime/Engine/Private/NavMeshRenderingComponent.cpp
+++ b/Engine/Source/Runtime/Engine/Private/NavMeshRenderingComponent.cpp
@@ -195,61 +195,6 @@ void UNavMeshRenderingComponent::GatherData(struct FNavMeshSceneProxyData* Curre
 			}
 		}
 
-		if (NavMesh->bDrawTileBounds)
-		{
-			FBox OuterBBox;
-			int32 NumTilesX, NumTilesY;
-			NavMesh->GetDebugTileBounds(OuterBBox, NumTilesX, NumTilesY);
-
-			float DrawZ = (OuterBBox.Min.Z + OuterBBox.Max.Z) * 0.5f;		// @hack average
-			FVector LL(OuterBBox.Min.X, OuterBBox.Min.Y, DrawZ);
-			FVector UR(OuterBBox.Max.X, OuterBBox.Max.Y, DrawZ);
-			FVector UL(LL.X, UR.Y, DrawZ);
-			FVector LR(UR.X, LL.Y, DrawZ);
-			if (bUseGetMeshElements)
-			{
-				CurrentData->ThickLineItems.Add(FNavMeshSceneProxyData::FDebugThickLine(LL, UL, NavMeshRenderColor_TileBounds, DefaultEdges_LineThickness));
-				CurrentData->ThickLineItems.Add(FNavMeshSceneProxyData::FDebugThickLine(UL, UR, NavMeshRenderColor_TileBounds, DefaultEdges_LineThickness));
-				CurrentData->ThickLineItems.Add(FNavMeshSceneProxyData::FDebugThickLine(UR, LR, NavMeshRenderColor_TileBounds, DefaultEdges_LineThickness));
-				CurrentData->ThickLineItems.Add(FNavMeshSceneProxyData::FDebugThickLine(LR, LL, NavMeshRenderColor_TileBounds, DefaultEdges_LineThickness));
-			}
-			else
-			{
-				CurrentData->BatchedElements.AddLine(LL, UL, NavMeshRenderColor_TileBounds, HitProxyId, DefaultEdges_LineThickness, 0, true);
-				CurrentData->BatchedElements.AddLine(UL, UR, NavMeshRenderColor_TileBounds, HitProxyId, DefaultEdges_LineThickness, 0, true);
-				CurrentData->BatchedElements.AddLine(UR, LR, NavMeshRenderColor_TileBounds, HitProxyId, DefaultEdges_LineThickness, 0, true);
-				CurrentData->BatchedElements.AddLine(LR, LL, NavMeshRenderColor_TileBounds, HitProxyId, DefaultEdges_LineThickness, 0, true);
-			}
-
-			float const TileSizeX = (LL.X - LR.X) / NumTilesX;
-			float const TileSizeY = (LL.Y - UL.Y) / NumTilesY;
-
-			for (int32 Idx = 1; Idx < NumTilesX; ++Idx)
-			{
-				float const DrawX = FMath::Lerp(LL.X, LR.X, (float)Idx / (float)NumTilesX);
-				if (bUseGetMeshElements)
-				{
-					CurrentData->ThickLineItems.Add(FNavMeshSceneProxyData::FDebugThickLine(FVector(DrawX, LL.Y, DrawZ), FVector(DrawX, UL.Y, DrawZ), NavMeshRenderColor_TileBounds, DefaultEdges_LineThickness));
-				}
-				else
-				{
-					CurrentData->BatchedElements.AddLine(FVector(DrawX, LL.Y, DrawZ), FVector(DrawX, UL.Y, DrawZ), NavMeshRenderColor_TileBounds, HitProxyId, DefaultEdges_LineThickness, 0, true);
-				}
-			}
-			for (int32 Idx=1; Idx<NumTilesY; ++Idx)
-			{
-				float const DrawY = FMath::Lerp(LL.Y, UL.Y, (float)Idx/(float)NumTilesY);
-				if (bUseGetMeshElements)
-				{
-					CurrentData->ThickLineItems.Add(FNavMeshSceneProxyData::FDebugThickLine(FVector(LL.X, DrawY, DrawZ), FVector(LR.X, DrawY, DrawZ), NavMeshRenderColor_TileBounds, DefaultEdges_LineThickness));
-				}
-				else
-				{
-					CurrentData->BatchedElements.AddLine(FVector(LL.X, DrawY, DrawZ), FVector(LR.X, DrawY, DrawZ), NavMeshRenderColor_TileBounds, HitProxyId, DefaultEdges_LineThickness, 0, true);
-				}
-			}
-		}
-
 		if (NavMesh->bDrawPathCollidingGeometry)
 		{
 			// draw all geometry gathered in navoctree
@@ -319,7 +264,7 @@ void UNavMeshRenderingComponent::GatherData(struct FNavMeshSceneProxyData* Curre
 			}
 		}
 
-		if (NavMesh->bDrawTileLabels || NavMesh->bDrawPolygonLabels || NavMesh->bDrawDefaultPolygonCost)
+		if (NavMesh->bDrawTileLabels || NavMesh->bDrawPolygonLabels || NavMesh->bDrawDefaultPolygonCost || NavMesh->bDrawTileBounds)
 		{
 			// calculate appropriate points for displaying debug labels
 			const int32 TilesCount = NavMesh->GetNavMeshTilesCount();
@@ -327,12 +272,8 @@ void UNavMeshRenderingComponent::GatherData(struct FNavMeshSceneProxyData* Curre
 
 			for (int32 TileIndex = 0; TileIndex < TilesCount; ++TileIndex)
 			{
-				int32 X = -1;
-				int32 Y = -1;
-				int32 Layer = -1;
-				NavMesh->GetNavMeshTileXY(TileIndex, X, Y, Layer);
-
-				if (X >= 0 && Y >= 0)
+				int32 X, Y, Layer;
+				if (NavMesh->GetNavMeshTileXY(TileIndex, X, Y, Layer))
 				{
 					const FBox TileBoundingBox = NavMesh->GetNavMeshTileBounds(TileIndex);
 					FVector TileLabelLocation = TileBoundingBox.GetCenter();
@@ -389,6 +330,31 @@ void UNavMeshRenderingComponent::GatherData(struct FNavMeshSceneProxyData* Curre
 							}
 						}
 					}							
+
+					if (NavMesh->bDrawTileBounds)
+					{
+						FBox TileBox = NavMesh->GetNavMeshTileBounds(TileIndex);
+
+						float DrawZ = (TileBox.Min.Z + TileBox.Max.Z) * 0.5f;		// @hack average
+						FVector LL(TileBox.Min.X, TileBox.Min.Y, DrawZ);
+						FVector UR(TileBox.Max.X, TileBox.Max.Y, DrawZ);
+						FVector UL(LL.X, UR.Y, DrawZ);
+						FVector LR(UR.X, LL.Y, DrawZ);
+						if (bUseGetMeshElements)
+						{
+							CurrentData->ThickLineItems.Add(FNavMeshSceneProxyData::FDebugThickLine(LL, UL, NavMeshRenderColor_TileBounds, DefaultEdges_LineThickness));
+							CurrentData->ThickLineItems.Add(FNavMeshSceneProxyData::FDebugThickLine(UL, UR, NavMeshRenderColor_TileBounds, DefaultEdges_LineThickness));
+							CurrentData->ThickLineItems.Add(FNavMeshSceneProxyData::FDebugThickLine(UR, LR, NavMeshRenderColor_TileBounds, DefaultEdges_LineThickness));
+							CurrentData->ThickLineItems.Add(FNavMeshSceneProxyData::FDebugThickLine(LR, LL, NavMeshRenderColor_TileBounds, DefaultEdges_LineThickness));
+						}
+						else
+						{
+							CurrentData->BatchedElements.AddLine(LL, UL, NavMeshRenderColor_TileBounds, HitProxyId, DefaultEdges_LineThickness, 0, true);
+							CurrentData->BatchedElements.AddLine(UL, UR, NavMeshRenderColor_TileBounds, HitProxyId, DefaultEdges_LineThickness, 0, true);
+							CurrentData->BatchedElements.AddLine(UR, LR, NavMeshRenderColor_TileBounds, HitProxyId, DefaultEdges_LineThickness, 0, true);
+							CurrentData->BatchedElements.AddLine(LR, LL, NavMeshRenderColor_TileBounds, HitProxyId, DefaultEdges_LineThickness, 0, true);
+						}
+					}
 				}
 			}
 		}
diff --git a/Engine/Source/Runtime/Engine/Private/UnrealEngine.cpp b/Engine/Source/Runtime/Engine/Private/UnrealEngine.cpp
--- a/Engine/Source/Runtime/Engine/Private/UnrealEngine.cpp
+++ b/Engine/Source/Runtime/Engine/Private/UnrealEngine.cpp
@@ -6715,6 +6715,17 @@ void DrawStatsHUD( UWorld* World, FViewport* Viewport, FCanvas* Canvas, UCanvas*
 				MessageY += 20;
 			}
 
+			if (NavSys && NavSys->bBuildNavigationAtRuntime)
+			{
+				const int32 NumTasks = NavSys->GetNumRemainingBuildTasks();
+				if (NumTasks > 0)
+				{
+					SmallTextItem.Text = FText::FromString(FString::Printf(TEXT("Building Navigation (%d)"), NumTasks));
+					Canvas->DrawItem(SmallTextItem, FVector2D(MessageX, MessageY));
+					MessageY += 20;
+				}
+			}
+
 			if( World->bKismetScriptError )
 			{
 				SmallTextItem.Text = LOCTEXT("BlueprintInLevelHadCompileErrorMessage", "BLUEPRINT COMPILE ERROR" );
diff --git a/Engine/Source/Runtime/Engine/Public/AI/NavDataGenerator.h b/Engine/Source/Runtime/Engine/Public/AI/NavDataGenerator.h
--- a/Engine/Source/Runtime/Engine/Public/AI/NavDataGenerator.h
+++ b/Engine/Source/Runtime/Engine/Public/AI/NavDataGenerator.h
@@ -4,46 +4,39 @@
 
 struct FNavigationDirtyArea;
 
-class ENGINE_API FNavDataGenerator : public TSharedFromThis<FNavDataGenerator, ESPMode::ThreadSafe>
+class ENGINE_API FNavDataGenerator
 {
 public:
 	virtual ~FNavDataGenerator() {}
 
-	/** 
-	 * Needs to be called on main thread (may need to reallocate main-mem containers and iterate actor list)
-	 */
-	virtual bool Generate() PURE_VIRTUAL(FNavDataGenerator::Generate, return false;);
+	/** Rebuilds all known navigation data */
+	virtual bool RebuildAll() { return false; };
+
+	/** Blocks until build is complete */
+	virtual void EnsureBuildCompletion() {};
+
+	/** Cancels build, may block until current running async tasks are finished */
+	virtual void CancelBuild() {};
+
+	/**  */
+	virtual void TickAsyncBuild(float DeltaSeconds) {};
 
 	/** Asks generator to update navigation affected by DirtyAreas */
 	virtual void RebuildDirtyAreas(const TArray<FNavigationDirtyArea>& DirtyAreas) {}
 
-	/** Rebuilds all known navigation data */
-	virtual void RebuildAll() PURE_VIRTUAL(FNavDataGenerator::RebuildAll, Generate(););
-
-	/** 
-	 * Schedules generation call in GameThread
-	 */
-	virtual void TriggerGeneration();
-	
 	/** determines whether this generator is performing navigation building actions at the moment*/
 	virtual bool IsBuildInProgress(bool bCheckDirtyToo = false) const { return false; }	
 
 	/** Called when world is loaded and all actors are initialized - to be used for post-load steps like data validating */
 	virtual void OnWorldInitDone(bool bAllowedToRebuild) {}	
 
-	virtual void OnNavigationBuildingLocked() {}
-
-	/** Lets Generator know that building navigation is now allowed. 
-	 *	@param bForce if true means it's requested from this Generator to 
-	 *	actually rebuild relevant data*/
-	virtual void OnNavigationBuildingUnlocked(bool bForce) {}
-
-	virtual void OnNavigationBoundsUpdated(class AVolume* Volume) {}
+	/** Returns number of remaining tasks till build is complete
+	 */
+	virtual int32 GetNumRemaningBuildTasks() const { return 0; };
 
-	/** notifies navigation generator that given ANavigationData is being destroyed
-	 *	and it's possible it's its parent 
+	/** Returns number of currently running tasks
 	 */
-	virtual void OnNavigationDataDestroyed(class ANavigationData* NavData) {}
+	virtual int32 GetNumRunningBuildTasks() const { return 0; };
 
 	//----------------------------------------------------------------------//
 	// debug
diff --git a/Engine/Source/Runtime/Engine/Public/AI/Navigation/RecastNavMeshGenerator.h b/Engine/Source/Runtime/Engine/Public/AI/Navigation/RecastNavMeshGenerator.h
--- a/Engine/Source/Runtime/Engine/Public/AI/Navigation/RecastNavMeshGenerator.h
+++ b/Engine/Source/Runtime/Engine/Public/AI/Navigation/RecastNavMeshGenerator.h
@@ -10,8 +10,11 @@
 
 #define MAX_VERTS_PER_POLY	6
 
-struct FNavigationRelevantData;
 class FNavigationOctree;
+class FNavMeshBuildContext;
+class FRecastNavMeshGenerator;
+struct FNavigationRelevantData;
+struct dtTileCacheLayer;
 
 struct FRecastBuildConfig : public rcConfig
 {
@@ -102,229 +105,161 @@ struct FRecastGeometryCache
 	FRecastGeometryCache(const uint8* Memory);
 };
 
-struct FRecastTileDirtyState
-{
-	/** dynamic array used when DirtyLayers is not enough */
-	TArray<uint8> FallbackDirtyLayers;
-	/** bitfields for marking layers to rebuild */
-	uint8 DirtyLayers[7];
-
-	/** set when any layer needs to be rebuild */
-	uint8 bRebuildLayers : 1;
-	/** set when all layers needs to be rebuild */
-	uint8 bRebuildAllLayers : 1;
-	/** set when geometry needs to be rebuild */
-	uint8 bRebuildGeometry : 1;
-
-	FRecastTileDirtyState() : bRebuildLayers(false), bRebuildAllLayers(false), bRebuildGeometry(false) { FMemory::Memzero(DirtyLayers, sizeof(DirtyLayers)); }
-	FRecastTileDirtyState(const class FRecastTileGenerator* DirtyGenerator);
-	
-	void Append(const FRecastTileDirtyState& Other);
-	void Clear();
-
-	void MarkDirtyLayer(int32 LayerIdx);
-	bool HasDirtyLayer(int32 LayerIdx) const;
-};
-
-struct FRecastDirtyGenerator
-{
-	class FRecastTileGenerator* Generator;
-	struct FRecastTileDirtyState State;
-
-	FRecastDirtyGenerator() : Generator(NULL) {}
-	FRecastDirtyGenerator(class FRecastTileGenerator* InGenerator) : Generator(InGenerator), State(InGenerator) {}
-};
-
 /**
  * Class handling generation of a single tile, caching data that can speed up subsequent tile generations
  */
-class ENGINE_API FRecastTileGenerator
+class ENGINE_API FRecastTileGenerator : public FNonAbandonableTask
 {
 public:
-	FRecastTileGenerator();
-	~FRecastTileGenerator();
+	FRecastTileGenerator(
+		const FRecastNavMeshGenerator* ParentGenerator,
+		const int32 X, 
+		const int32 Y,
+		TArray<FBox> DirtyAreas
+	);
+		
+	void DoWork();
+	static const TCHAR *Name();
 
-	// sets up tile generator instance with initial data
-	void Init(class FRecastNavMeshGenerator* ParentGenerator, const int32 X, const int32 Y, const float TileBmin[3], const float TileBmax[3], const TNavStatArray<FBox>& BoundingBoxes);
+	FORCEINLINE int32 GetTileX() const { return TileX; }
+	FORCEINLINE int32 GetTileY() const { return TileY; }
+	/** Whether specified layer was updated */
+	FORCEINLINE bool IsLayerChanged(int32 LayerIdx) const { return DirtyLayers[LayerIdx]; }
+	/** Whether tile data was fully regenerated */
+	FORCEINLINE bool IsFullyRegenerated() const { return bRegenerateCompressedLayers; }
 
-	// this is the only legit way of triggering navmesh tile generation. 
-	void TriggerAsyncBuild();
+	TArray<FNavMeshTileData> GetCompressedLayers() const { return CompressedLayers; }
+	TArray<FNavMeshTileData> GetNavigationData() const { return NavigationData; }
+	
+	uint32 GetUsedMemCount() const;
 
+	// Memory amount used to construct generator 
+	uint32 UsedMemoryOnStartup;
+	
+private:
 	/** Does the actual tile generations. 
 	 *	@note always trigger tile generation only via TriggerAsyncBuild. This is a worker function
 	 *	@return true if new tile navigation data has been generated and is ready to be added to navmesh instance, 
 	 *	false if failed or no need to generate (still valid)
 	 */
 	bool GenerateTile();
+	
+	void GatherGeometry(const FRecastNavMeshGenerator* ParentGenerator, bool bGeometryChanged);
 
-	// free allocation of navmesh data
-	void ClearNavigationData() { NavigationData.Reset(); }
-	// copy generated navmesh data for caller and free allocation
-	void TransferNavigationData(TArray<FNavMeshTileData>& TileLayers) { TileLayers = NavigationData; ClearNavigationData(); }
+	/** builds CompressedLayers array (geometry + modifiers) */
+	bool GenerateCompressedLayers(FNavMeshBuildContext& BuildContext);
 
-	FORCEINLINE int32 GetTileX() const { return TileX; }
-	FORCEINLINE int32 GetTileY() const { return TileY; }
-	FORCEINLINE int32 GetId() const { return TileId; }
-	FORCEINLINE int32 GetVersion() const { return Version; }
+	/** builds NavigationData array (layers + obstacles) */
+	bool GenerateNavigationData(FNavMeshBuildContext& BuildContext);
 
-	FORCEINLINE const FBox& GetUnrealBB() const { return TileBB; }
-	
-	void SetDirty(const FNavigationDirtyArea& DirtyArea, const FBox& AreaBounds);
-	FORCEINLINE void GetDirtyState(FRecastTileDirtyState& Data) const { Data = DirtyState; }
-	FORCEINLINE void SetDirtyState(const FRecastTileDirtyState& Data) { DirtyState.Append(Data); }
-	FORCEINLINE bool HasDirtyGeometry() const { return DirtyState.bRebuildGeometry; }
-	FORCEINLINE bool HasDirtyLayers() const { return DirtyState.bRebuildLayers; }
-	FORCEINLINE bool IsDirty() const { return DirtyState.bRebuildGeometry || DirtyState.bRebuildLayers; }
-	FORCEINLINE bool IsRebuildingGeometry() const { return GeneratingState.bRebuildGeometry; }
-	FORCEINLINE bool IsRebuildingLayers() const { return GeneratingState.bRebuildLayers; }
-	FORCEINLINE bool ShouldBeBuilt() const { return bOutsideOfInclusionBounds == 0; }
-	FORCEINLINE bool IsBeingRebuild() const { return bBeingRebuild; }
-	FORCEINLINE bool IsPendingRebuild() const { return bRebuildPending; }
-	FORCEINLINE bool IsAsyncBuildInProgress() const { return bAsyncBuildInProgress; }
-
-	/** Marks start of rebuilding process. */
-	void InitiateRebuild();
-	/** Tile remains dirty and will be built next time around */
-	void AbortRebuild();
-	/** called when tile generator prepared for generation signals there's actually
-	 *	no point in doing so, for example when ShouldBeBuilt() == false */
-	void AbadonGeneration();
-	/** finish rebuild process and update dirty state */
-	void FinishRebuild();
-	
-	/** mark start of async worker */
-	void StartAsyncBuild();
-	/** mark finish of async worker */
-	void FinishAsyncBuild();
+	void ApplyVoxelFilter(struct rcHeightfield* SolidHF, float WalkableRadius);
 
-	/** mark pending rebuild: generator has all data and waits in queue for free worker thread */
-	void MarkPendingRebuild();
+	/** apply areas from StaticAreas to heightfield */
+	void MarkStaticAreas(FNavMeshBuildContext& BuildContext, rcCompactHeightfield& CompactHF);
 
+	/** apply areas from DynamicAreas to layer */
+	void MarkDynamicAreas(dtTileCacheLayer& Layer);
+	
 	void AppendModifier(const FCompositeNavModifier& Modifier, bool bStatic);
 	/** Appends specified geometry to tile's geometry */
 	void AppendGeometry(const TNavStatArray<uint8>& RawCollisionCache);
 	void AppendGeometry(const TNavStatArray<FVector>& Verts, const TNavStatArray<int32>& Faces);
 	void AppendVoxels(rcSpanCache* SpanData, int32 NumSpans);
 	
-	uint32 GetUsedMemCount() const;
-	static int32 GetStaticVoxelSpanCount() { return StaticGeomSpans.Num(); }
-
-#if !(UE_BUILD_SHIPPING || UE_BUILD_TEST)
-	float GetLastBuildTimeStamp() const { return LastBuildTimeStamp; }
-#endif // !(UE_BUILD_SHIPPING || UE_BUILD_TEST)
-
-	void ClearGeometry();
-	void ClearModifiers();
-	static void ClearStaticData();
-
 	/** prepare voxel cache from collision data */
 	void PrepareVoxelCache(const TNavStatArray<uint8>& RawCollisionCache, TNavStatArray<rcSpanCache>& SpanData);
 	bool HasVoxelCache(const TNavStatArray<uint8>& RawVoxelCache, rcSpanCache*& CachedVoxels, int32& NumCachedVoxels) const;
 	void AddVoxelCache(TNavStatArray<uint8>& RawVoxelCache, const rcSpanCache* CachedVoxels, const int32 NumCachedVoxels) const;
 
 protected:
-
-	uint32 bInitialized : 1;
+	uint32 bSucceeded : 1;
+	uint32 bRegenerateCompressedLayers : 1;
 	uint32 bOutsideOfInclusionBounds : 1;
 	uint32 bFullyEncapsulatedByInclusionBounds : 1;
-	uint32 bBeingRebuild : 1;
-	uint32 bRebuildPending : 1;
-	uint32 bAsyncBuildInProgress : 1;
-
-	/** dirty state exposed for navigation system */
-	FRecastTileDirtyState DirtyState;
-	/** dirty state used in active processing */
-	FRecastTileDirtyState GeneratingState;
-
+	
 	int32 TileX;
 	int32 TileY;
-	int32 TileId;
 	uint32 Version;
-
-	/** bounding box minimum, Recast coords */
-	float BMin[3];
-	/** bounding box maximum, Recast coords */
-	float BMax[3];
-
 	/** Tile's bounding box, Unreal coords */
 	FBox TileBB;
-	/** Layer bounds and dirty flags, Unreal coords */
-	TArray<FBox> LayerBB;
+	
+	/** Layers dirty flags */
+	TBitArray<>  DirtyLayers;
+	
+	/** Parameters defining navmesh tiles */
+	FRecastBuildConfig TileConfig;
+
+	/** Bounding geometry definition. */
+	TNavStatArray<FBox> InclusionBounds;
 
-	/** NavMesh generator owning this tile generator */
-	ANavigationData::FNavDataGeneratorWeakPtr NavMeshGenerator;
 	/** Additional config */
-	TSharedPtr<struct FRecastNavMeshCachedData, ESPMode::ThreadSafe> AdditionalCachedData;
+	FRecastNavMeshCachedData AdditionalCachedData;
 
-	// generated tiles
+	// generated tile data
 	TArray<FNavMeshTileData> CompressedLayers;
 	TArray<FNavMeshTileData> NavigationData;
-
-	// tile's geometry: voxel cache (only for synchronous rebuilds)
-	float RasterizationPadding;
-	int32 WalkableClimbVX;
-	float WalkableSlopeCos;
-	static TNavStatArray<rcSpanCache> StaticGeomSpans;
-
+	
 	// tile's geometry: without voxel cache
-	TNavStatArray<float> GeomCoords;
-	TNavStatArray<int32> GeomIndices;
-	static TNavStatArray<float> StaticGeomCoords;
-	static TNavStatArray<int32> StaticGeomIndices;
-
-	TNavStatArray<FBox> InclusionBounds;
+	TArray<float> GeomCoords;
+	TArray<int32> GeomIndices;
 	// areas used for creating compressed layers: static zones
 	TArray<FAreaNavModifier> StaticAreas;
 	// areas used for creating navigation data: obstacles
 	TArray<FAreaNavModifier> DynamicAreas;
 	// navigation links
 	TArray<FSimpleLinkNavModifier> OffmeshLinks;
+};
 
-	static TArray<FAreaNavModifier> StaticStaticAreas;
-	static TArray<FAreaNavModifier> StaticDynamicAreas;
-	static TArray<FSimpleLinkNavModifier> StaticOffmeshLinks;
-
-	FCriticalSection GenerationLock;
-
-	typedef FAutoDeleteAsyncTask<class FAsyncNavTileBuildWorker> FAsyncNavTileBuildTask;
-
-#if !(UE_BUILD_SHIPPING || UE_BUILD_TEST)
-	// some statistics - not shippable
-	float LastBuildTimeCost;
-	float LastBuildTimeStamp;
-#endif // !(UE_BUILD_SHIPPING || UE_BUILD_TEST)
-
-	/** builds CompressedLayers array (geometry + modifiers) */
-	bool GenerateCompressedLayers(class FNavMeshBuildContext* BuildContext, FRecastNavMeshGenerator* Generator);
-
-	/** builds NavigationData array (layers + obstacles) */
-	bool GenerateNavigationData(class FNavMeshBuildContext* BuildContext, FRecastNavMeshGenerator* Generator);
-
-	void ApplyVoxelFilter(struct rcHeightfield* SolidHF, float WalkableRadius);
+typedef FAsyncTask<FRecastTileGenerator> FRecastTileGeneratorTask;
 
-	/** apply areas from StaticAreas to heightfield */
-	void MarkStaticAreas(class FNavMeshBuildContext* BuildContext, struct rcCompactHeightfield& CompactHF, const struct FRecastBuildConfig& TileConfig);
+struct FTileElement
+{
+	/** tile coords on grid */
+	FIntPoint	Coord;
+	/** distance to seed, used for sorting pending tiles */
+	float		SeedDistance; 
+	/** generator task associated with this tile */
+	FRecastTileGeneratorTask* AsyncTask;
+	/** Whether we need a full rebuild for this tile grid cell */
+	bool		bRebuildGeometry;
+	/** We need to store dirty area bounds to check which cached layers needs to be regenerated
+	 *  In case geometry is changed cached layers data will be fully regenerated without using dirty areas list
+	 */
+	TArray<FBox> DirtyAreas;
+	
+	FTileElement()
+		: Coord(FIntPoint::NoneValue)
+		, SeedDistance(MAX_flt)
+		, AsyncTask(nullptr)
+		, bRebuildGeometry(false)
+	{
+	}
 
-	/** apply areas from DynamicAreas to layer */
-	void MarkDynamicAreas(struct dtTileCacheLayer& Layer, const struct FRecastBuildConfig& TileConfig);
-};
+	bool operator == (const FTileElement& Other) const
+	{
+		return Coord == Other.Coord;
+	}
 
-struct FNavMeshDirtInfo
-{	
-	FBox Box;
-	const AActor* Actor;
+	bool operator < (const FTileElement& Other) const
+	{
+		return Other.SeedDistance < SeedDistance;
+	}
 
-	FNavMeshDirtInfo(const FBox& InBox, const AActor* InActor = 0) : Box(InBox), Actor(InActor) {}
+	friend uint32 GetTypeHash(const FTileElement& Element)
+	{
+		return GetTypeHash(Element.Coord);
+	}
 };
 
-struct FNavMeshGenerationResult
+struct FTileTimestamp
 {
-	dtTileRef OldTileRef;
-	uint8* OldRawNavData;
-	uint32 TileIndex;
-	FNavMeshTileData NewNavData;
-
-	FNavMeshGenerationResult() : OldTileRef(INVALID_NAVNODEREF), OldRawNavData(NULL) {}
+	uint32 TileIdx;
+	double Timestamp;
+	
+	bool operator == (const FTileTimestamp& Other) const
+	{
+		return TileIdx == Other.TileIdx;
+	}
 };
 
 /**
@@ -342,168 +277,88 @@ private:
 	FRecastNavMeshGenerator& operator=(FRecastNavMeshGenerator const& NoCopy) { check(0); return *this; }
 
 public:
-
-	// Triggers navmesh building process
-	virtual bool Generate() override;
+	virtual bool RebuildAll() override;
+	virtual void EnsureBuildCompletion() override;
+	virtual void CancelBuild() override;
+	virtual void TickAsyncBuild(float DeltaSeconds) override;
 
 	/** Asks generator to update navigation affected by DirtyAreas */
 	virtual void RebuildDirtyAreas(const TArray<FNavigationDirtyArea>& DirtyAreas) override;
-
-	virtual void TriggerGeneration() override;
-
 	virtual bool IsBuildInProgress(bool bCheckDirtyToo = false) const override;
+	virtual int32 GetNumRemaningBuildTasks() const override;
+	virtual int32 GetNumRunningBuildTasks() const override;
 
-	bool AreAnyTilesBeingBuilt(bool bCheckDirtyToo = false) const;
-	bool IsAsyncBuildInProgress() const;
-
-	/** Checks if a given tile is being build, is marked dirty, or has just finished building
+	/** Checks if a given tile is being build or has just finished building
 	 *	If TileIndex is out of bounds function returns false (i.e. not being built) */
-	bool IsTileFresh(int32 X, int32 Y, float FreshnessTime = FRecastNavMeshGenerator::DefaultFreshness) const;
-
-	virtual void OnWorldInitDone(bool bAllowedToRebuild) override;
-
-	/** Moves data from this generator to it's destination navmesh object. 
-	 *	@return fail if DestNavMesh is not valid in some way, true otherwise */
-	bool TransferGeneratedData();
-	
-	FORCEINLINE int32 GetTileIdAt(int32 X, int32 Y) const { return Y * TilesWidth + X; }
+	bool IsTileChanged(int32 TileIdx) const;
+		
 	FORCEINLINE uint32 GetVersion() const { return Version; }
 
-	/** Retrieves FCriticalSection used when adding tiles to navmesh instance, so that it can be locked
-	 *	while retrieving data for drawing for example. 
-	 *	@todo this is awkward way of doing this, but will be removed once FRecastNavMeshGenerator instance is owned by 
-	 *		FPImplRecastNavMesh, not ARecastNavMesh. Will be done soon. */
-	FCriticalSection* GetNavMeshModificationLock() const { return &TileAddingLock; }
-
-	const ARecastNavMesh* GetOwner() const { return DestNavMesh.Get(); }
-	
-	/** 
-	 *	@param Actor is a reference to make callee responsible for assuring it's valid
-	 */
-	static void ExportComponentGeometry(UActorComponent* Component, FNavigationRelevantData& Data);
-	static void ExportVertexSoupGeometry(const TArray<FVector>& Verts, FNavigationRelevantData& Data);
-
-	static void ExportRigidBodyGeometry(UBodySetup& BodySetup, TNavStatArray<FVector>& OutVertexBuffer, TNavStatArray<int32>& OutIndexBuffer, const FTransform& LocalToWorld = FTransform::Identity);
-	static void ExportRigidBodyGeometry(UBodySetup& BodySetup, TNavStatArray<FVector>& OutTriMeshVertexBuffer, TNavStatArray<int32>& OutTriMeshIndexBuffer, TNavStatArray<FVector>& OutConvexVertexBuffer, TNavStatArray<int32>& OutConvexIndexBuffer, TNavStatArray<int32>& OutShapeBuffer, const FTransform& LocalToWorld = FTransform::Identity);
-
-	/** update workers to (re)generate new tiles */
-	void UpdateTileGenerationWorkers(int32 TileId);
+	const ARecastNavMesh* GetOwner() const { return DestNavMesh; }
 
+	/** update area data */
+	void OnAreaAdded(const UClass* AreaClass, int32 AreaID);
+		
 	//--- accessors --- //
-	FORCEINLINE class UWorld* GetWorld() const { return DestNavMesh.IsValid() ? DestNavMesh->GetWorld() : NULL; }
+	FORCEINLINE class UWorld* GetWorld() const { return DestNavMesh->GetWorld(); }
 
-	class FNavMeshBuildContext* GetBuildContext()	{ return BuildContext; }
 	const FRecastBuildConfig& GetConfig() const { return Config; }
-	TSharedPtr<struct FRecastNavMeshCachedData, ESPMode::ThreadSafe> GetAdditionalCachedData() { return AdditionalCachedData; }
-private:
-	// Performs initial setup of member variables so that generator is ready to do its thing from this point on
-	void Init();
-
-	bool IsBuildingLocked() const;
 
-	/** @NOTE may produce false-negatives if outside of game thread (due to GC) */
-	FORCEINLINE bool ShouldContinueBuilding() const
-	{
-		return DestNavMesh.IsValid()
-			&& DestNavMesh != NULL
-			// using this construct rather than "DestNavMesh->GetWorld()" to 
-			// avoid unreachability checks. Note that this will be called
-			// outside of game thread and may stumble upon GC, when everything is
-			// unreachanble, while here we just care about "outer systems' existance"
-			&& Cast<ULevel>(DestNavMesh->GetOuter())
-			&& ((ULevel*)(DestNavMesh->GetOuter()))->OwningWorld != NULL
-			&& ((ULevel*)(DestNavMesh->GetOuter()))->OwningWorld->GetNavigationSystem() != NULL;
-	}
+	const TNavStatArray<FBox>& GetInclusionBounds() const { return InclusionBounds; }
+	
+	/** Whether specified box in inside or intersects with nav bounds  */
+	bool IntercestsInclusionBounds(const FBox& Box) const;
 
-	// Loads navmesh generation config
-	// Params:
-	//		cellSize - (in) size of voxel 2D grid
-	//		cellHeight - (in) height of voxel cell
-	//		agentMinHeight - (in) min walkable height
-	//		agentMaxHeight - (in) max walkable height
-	//		agentMaxSlope - (in) max angle of walkable slope (degrees)
-	//		agentMaxClimb - (in) max height of walkable step
-	void SetUpGeneration(float CellSize, float CellHeight, float AgentMinHeight, float AgentMaxHeight, float AgentMaxSlope, float AgentMaxClimb, float AgentRadius);
+	/** Total navigable area box, sum of all navigation volumes bounding boxes */
+	FBox GetTotalBounds() const { return TotalNavBounds; }
 
-	// Generates tiled navmesh. NavMesh tile size defined by FRecastNavMeshGenerator::TileSize
-	// @todo considering passing TileSize as a parameter to this function
-	bool GenerateTiledNavMesh();
+	const FRecastNavMeshCachedData& GetAdditionalCachedData() const { return AdditionalCachedData; }
 
-public:
-	void GenerateTile(const int32 TileId, const uint32 TileGeneratorVersion);
-	bool AddTile(FRecastTileGenerator* TileGenerator, ARecastNavMesh* CachedRecastNavMeshInstance);
-	void RegenerateDirtyTiles();
+	TArray<uint32> AddTile(const FRecastTileGenerator& TileGenerator);
 	bool HasDirtyTiles() const;
 
-	bool HasDirtyAreas() const { return DirtyAreas.Num() > 0; }
-	
-	FBox GrowBoundingBox(const FBox& BBox, bool bIncludeAgentHeight) const
-	{
-		const FVector BBoxGrowOffsetBoth = FVector(2.0f * Config.borderSize * Config.cs);
-		const FVector BBoxGrowOffsetMin = FVector(0, 0, bIncludeAgentHeight ? Config.AgentHeight : 0.0f);
-
-		return FBox(BBox.Min - BBoxGrowOffsetBoth - BBoxGrowOffsetMin, BBox.Max + BBoxGrowOffsetBoth);
-	}
-
-	FORCEINLINE bool HasResultsPending() const { return AsyncGenerationResultContainer.Num() > 0; }
-	void GetAsyncResultsCopy(TNavStatArray<FNavMeshGenerationResult>& Dest, bool bClearSource);
+	FBox GrowBoundingBox(const FBox& BBox, bool bIncludeAgentHeight) const;
 
-	void StoreAsyncResults(TArray<FNavMeshGenerationResult> AsyncResults);
+	/**  */
+	TArray<FNavMeshTileData> GetCachedLayersData(FIntPoint GridCoord) const;
 
 #if !(UE_BUILD_SHIPPING || UE_BUILD_TEST)
 	virtual void ExportNavigationData(const FString& FileName) const;
 #endif
-private:
-	/** Instantiates dtNavMesh and configures it for tiles generation. Returns false if failed */
-	bool ConstructTiledNavMesh();
 
-	void UpdateBuilding();
+	/** 
+	 *	@param Actor is a reference to make callee responsible for assuring it's valid
+	 */
+	static void ExportComponentGeometry(UActorComponent* Component, FNavigationRelevantData& Data);
+	static void ExportVertexSoupGeometry(const TArray<FVector>& Verts, FNavigationRelevantData& Data);
 
-	void RemoveTileLayers(const int32 TileX, const int32 TileY, TArray<FNavMeshGenerationResult>& AsyncResults);
-	
-private:
-	void RefreshParentReference();
-	/** calls MarkBoxDirty on non-main thread */
-	//void NotifyBoxDirty(FBox Bounds, const AActor* Actor);
-	//void MarkBoxDirty(FBox Bounds, const AActor* Actor);
+	static void ExportRigidBodyGeometry(UBodySetup& BodySetup, TNavStatArray<FVector>& OutVertexBuffer, TNavStatArray<int32>& OutIndexBuffer, const FTransform& LocalToWorld = FTransform::Identity);
+	static void ExportRigidBodyGeometry(UBodySetup& BodySetup, TNavStatArray<FVector>& OutTriMeshVertexBuffer, TNavStatArray<int32>& OutTriMeshIndexBuffer, TNavStatArray<FVector>& OutConvexVertexBuffer, TNavStatArray<int32>& OutConvexIndexBuffer, TNavStatArray<int32>& OutShapeBuffer, const FTransform& LocalToWorld = FTransform::Identity);
 
-	virtual void RebuildAll() override;
+private:
+	// Performs initial setup of member variables so that generator is ready to do its thing from this point on
+	void Init();
+		
+	// Sorts pending build tiles by proximity to player, so tiles closer to player will get generated first
+	void SortPendingBuildTiles();
 
-	void RequestDirtyTilesRebuild();
+	/** Instantiates dtNavMesh and configures it for tiles generation. Returns false if failed */
+	bool ConstructTiledNavMesh();
 	
-	/** mark dirty area for aborted generator */
-	void MarkAbortedGenerator(const FRecastTileGenerator* TileGenerator);
-
-	/** mark tile generators using dirty areas and clear dirty areas array */
-	void MarkDirtyGenerators();
-
-	/** start tile generators marked as dirty, up to MaxActiveGenerators running at the same time */
-	void StartDirtyGenerators();
-
-	/** fill generator from navigation octree */
-	void FillGeneratorData(FRecastTileGenerator* TileGenerator, const FNavigationOctree* NavOctree);
-
-	enum EDataOwnership
-	{
-		DO_ForeignData = 0,
-		DO_OwnsData,
-	};
+	/** Marks grid tiles affected by specified areas as dirty */
+	void MarkDirtyTiles(const TArray<FNavigationDirtyArea>& DirtyAreas);
 	
-	void SetDetourMesh(class dtNavMesh* NewDetourMesh, EDataOwnership OwnsData = FRecastNavMeshGenerator::DO_OwnsData);
-
-public:
-
-	/** update area data */
-	void OnAreaAdded(const UClass* AreaClass, int32 AreaID);
-
-	/** Called to trigger generation at earliest convenience */
-	void RequestGeneration();
-
-	virtual void OnNavigationBuildingLocked() override;
-	virtual void OnNavigationBuildingUnlocked(bool bForce) override;
-	virtual void OnNavigationBoundsUpdated(class AVolume* Volume) override;
+	/** Submits up to specified number of async generator tasks for grid tiles that was marked as dirty */
+	int32 SubmitDirtyTiles(const int32 NumTasksToSubmit);
+	
+	/** Collect completed async tasks and applies new tile data to a navmesh*/
+	TArray<uint32> CollectCompletedTiles();
 
-	virtual void OnNavigationDataDestroyed(class ANavigationData* NavData) override;
+	/** Removes all tiles at specified grid location */
+	TArray<uint32> RemoveTileLayers(const int32 TileX, const int32 TileY);
+	
+	/** Blocks until build for specified list of tiles is complete and discard results */
+	void DiscardCurrentBuildingTasks();
 
 	//----------------------------------------------------------------------//
 	// debug
@@ -514,75 +369,46 @@ private:
 	/** Parameters defining navmesh tiles */
 	struct dtNavMeshParams TiledMeshParams;
 	struct FRecastBuildConfig Config;
-	class dtNavMesh* DetourMesh;
-
+	
 	int32 MaxActiveTiles;
 	int32 NumActiveTiles;
-	int32 MaxActiveGenerators;
-
-	int32 TilesWidth;
-	int32 TilesHeight;
-	int32 GridWidth;
-	int32 GridHeight;
-
-	TSharedPtr<struct FRecastNavMeshCachedData, ESPMode::ThreadSafe> AdditionalCachedData;
-
-	// Additional configurable values
-	// @todo These settings need to be tweakable per navmesh at some point
-	int32 TileSize;
-
-	/** bounding box of the whole geometry data used to generate navmesh - in recast coords*/
-	FBox RCNavBounds;
-	/** bounding box of the whole geometry data used to generate navmesh - in unreal coords*/
-	FBox UnrealNavBounds;
+	int32 MaxTileGeneratorTasks;
+	static int32 ActiveTileGeneratorTasks;
 
-	class FNavMeshBuildContext* BuildContext;
+	/** Total bounding box that includes all volumes, in unreal units. */
+	FBox TotalNavBounds;
 
 	/** Bounding geometry definition. */
 	TNavStatArray<FBox> InclusionBounds;
 
-	TNavStatArray<FRecastTileGenerator> TileGenerators;
-
-	mutable FCriticalSection TileAddingLock;
-	mutable FCriticalSection TileGenerationLock;
-	FCriticalSection NavMeshDirtyLock;
-	FCriticalSection InitLock;
-	FCriticalSection GatheringDataLock;
+	/** Navigation mesh that owns this generator */
+	ARecastNavMesh*	DestNavMesh;
+	
+	/** List of dirty tiles that needs to be regenerated */
+	TNavStatArray<FTileElement> PendingDirtyTiles;			
+	
+	/** List of dirty tiles currently being regenerated */
+	TNavStatArray<FTileElement> RunningDirtyTiles;
 
-	TWeakObjectPtr<ARecastNavMesh> DestNavMesh;
+#if WITH_EDITOR
+	/** List of tiles that were recently regenerated */
+	TNavStatArray<FTileTimestamp> RecentlyBuiltTiles;
+#endif// WITH_EDITOR
 
-	/** The output navmesh we're generating. */
-	class FPImplRecastNavMesh* OutNavMesh;
+	/** */
+	FRecastNavMeshCachedData AdditionalCachedData;
 
-	TMapBase<const AActor*, FBox, false> ActorToAreaMap;
-	/** stores information about areas needing rebuilding */
-	TNavStatArray<FNavigationDirtyArea> DirtyAreas;
-	/** stores indices of tile generators that need rebuilding and their dirty state */
-	TMap<int32,FRecastTileDirtyState> DirtyGenerators;
+	/** Compressed layers data, can be reused for tiles generation */
+	TMap<FIntPoint, TArray<FNavMeshTileData>> CachedLayerDataMap;
 
 	/** */
-	TNavStatArray<FNavMeshGenerationResult> AsyncGenerationResultContainer;
+	TMapBase<const AActor*, FBox, false> ActorToAreaMap;
 
 	uint32 bInitialized:1;
-	uint32 bBuildFromScratchRequested:1;
-	uint32 bRebuildDirtyTilesRequested:1;
-	/** gets set to true when RecastNavMesh owned has been destroyed and there's 
-	 *	no more point in generating navigation data */
-	uint32 bAbortAllTileGeneration:1;
-	uint32 bOwnsDetourMesh:1;
-	mutable uint32 bBuildingLocked:1;
 
 	/** Runtime generator's version, increased every time all tile generators get invalidated
 	 *	like when navmesh size changes */
 	uint32 Version;
-
-	static const float DefaultFreshness;
-
-	/** queue of working generators - size of this array depends on the number of CPU cores */
-	TArray<FRecastTileGenerator*> ActiveGenerators;
-
-	/** sorted array of waiting generators to work asap*/
-	TArray<FRecastTileGenerator*> GeneratorsQueue;
 };
 
 #endif // WITH_RECAST
diff --git a/Engine/Source/Runtime/Engine/Public/AI/NavigationOctree.h b/Engine/Source/Runtime/Engine/Public/AI/NavigationOctree.h
--- a/Engine/Source/Runtime/Engine/Public/AI/NavigationOctree.h
+++ b/Engine/Source/Runtime/Engine/Public/AI/NavigationOctree.h
@@ -63,30 +63,19 @@ struct ENGINE_API FNavigationOctreeElement
 {
 	FBoxSphereBounds Bounds;
 	TWeakObjectPtr<UObject> Owner;
-
-#if NAVOCTREE_CONTAINS_COLLISION_DATA
 	FNavigationRelevantData Data;
-#endif // NAVOCTREE_CONTAINS_COLLISION_DATA
 
 	FORCEINLINE bool IsEmpty() const
 	{
 		const FBox BBox = Bounds.GetBox();
-#if NAVOCTREE_CONTAINS_COLLISION_DATA
 		return Data.IsEmpty() && (BBox.IsValid == 0 || BBox.GetSize().IsNearlyZero());
-#else
-		return BBox.IsValid == 0 || BBox.GetSize().IsNearlyZero();
-#endif // NAVOCTREE_CONTAINS_COLLISION_DATA
 	}
 
 	FORCEINLINE bool IsMatchingFilter(const FNavigationOctreeFilter& Filter) const
 	{
-#if NAVOCTREE_CONTAINS_COLLISION_DATA
 		return Data.IsMatchingFilter(Filter);
-#endif // NAVOCTREE_CONTAINS_COLLISION_DATA
-		return false;
 	}
 
-#if NAVOCTREE_CONTAINS_COLLISION_DATA
 	/** 
 	 *	retrieves Modifier, if it doesn't contain any "Meta Navigation Areas". 
 	 *	If it does then retrieves a copy with meta areas substituted with
@@ -111,7 +100,6 @@ struct ENGINE_API FNavigationOctreeElement
 	{
 		Data.Shrink();
 	}
-#endif // NAVOCTREE_CONTAINS_COLLISION_DATA
 };
 
 struct FNavigationOctreeSemantics
