CIS fixes

diff --git a/Engine/Source/Runtime/Engine/Private/AI/Navigation/NavigationSystem.cpp b/Engine/Source/Runtime/Engine/Private/AI/Navigation/NavigationSystem.cpp
--- a/Engine/Source/Runtime/Engine/Private/AI/Navigation/NavigationSystem.cpp
+++ b/Engine/Source/Runtime/Engine/Private/AI/Navigation/NavigationSystem.cpp
@@ -2656,6 +2656,7 @@ void UNavigationSystem::OnLevelRemovedFromWorld(ULevel* InLevel, UWorld* InWorld
 
 void UNavigationSystem::AddLevelCollisionToOctree(ULevel* Level)
 {
+#if WITH_RECAST
 	const TArray<FVector>* LevelGeom = Level ? Level->GetStaticNavigableGeometry() : NULL;
 	if (LevelGeom && NavOctree)
 	{
@@ -2669,6 +2670,7 @@ void UNavigationSystem::AddLevelCollisionToOctree(ULevel* Level)
 			AddDirtyArea(Bounds, ENavigationDirtyFlag::All);
 		}
 	}
+#endif// WITH_RECAST
 }
 
 void UNavigationSystem::RemoveLevelCollisionFromOctree(ULevel* Level)
diff --git a/Engine/Source/Runtime/Engine/Private/AI/Navigation/RecastNavMesh.cpp b/Engine/Source/Runtime/Engine/Private/AI/Navigation/RecastNavMesh.cpp
--- a/Engine/Source/Runtime/Engine/Private/AI/Navigation/RecastNavMesh.cpp
+++ b/Engine/Source/Runtime/Engine/Private/AI/Navigation/RecastNavMesh.cpp
@@ -1628,6 +1628,8 @@ bool ARecastNavMesh::HasValidNavmesh() const
 	return false;
 }
 
+#if WITH_RECAST
+
 FRecastNavMeshCachedData FRecastNavMeshCachedData::Construct(const ARecastNavMesh* RecastNavMeshActor)
 {
 	check(RecastNavMeshActor);
@@ -1683,3 +1685,4 @@ void FRecastNavMeshCachedData::OnAreaAdded(const UClass* AreaClass, int32 AreaID
 	}		
 }
 
+#endif// WITH_RECAST
\ No newline at end of file
