fixed navmesh tile bounds vs grid index pairing

diff --git a/Engine/Source/Runtime/Engine/Private/AI/Navigation/RecastNavMeshGenerator.cpp b/Engine/Source/Runtime/Engine/Private/AI/Navigation/RecastNavMeshGenerator.cpp
--- a/Engine/Source/Runtime/Engine/Private/AI/Navigation/RecastNavMeshGenerator.cpp
+++ b/Engine/Source/Runtime/Engine/Private/AI/Navigation/RecastNavMeshGenerator.cpp
@@ -1380,6 +1380,21 @@ uint32 GetTileCacheSizeHelper(TArray<FNavMeshTileData>& CompressedTiles)
 	return TotalMemory;
 }
 
+static FBox CalculateTileBounds(int32 X, int32 Y, const FVector& NavMeshOrigin, const FBox& TotalNavBounds, float TileSizeInWorldUnits)
+{
+	const FVector RcNavMeshOrigin = Unreal2RecastPoint(NavMeshOrigin);
+	FBox TileBox(
+		Recast2UnrealPoint(RcNavMeshOrigin + (FVector(X + 0, 0, Y + 0) * TileSizeInWorldUnits)),
+		Recast2UnrealPoint(RcNavMeshOrigin + (FVector(X + 1, 0, Y + 1) * TileSizeInWorldUnits))
+		);
+
+	TileBox.Min.Z = TotalNavBounds.Min.Z;
+	TileBox.Max.Z = TotalNavBounds.Max.Z;
+
+	// unreal coord space
+	return TileBox;
+}
+
 //----------------------------------------------------------------------//
 // FRecastTileGenerator
 //----------------------------------------------------------------------//
@@ -1402,11 +1417,9 @@ FRecastTileGenerator::FRecastTileGenerator(
 	const FVector NavMeshOrigin = FVector::ZeroVector;
 	const FBox NavTotalBounds = ParentGenerator->GetTotalBounds();
 	const float TileCellSize = (TileConfig.tileSize * TileConfig.cs);
-	const FVector TileBMin = FVector((X+0) * TileCellSize, (Y+0) * TileCellSize, NavTotalBounds.Min.Z) + NavMeshOrigin;
-	const FVector TileBMax = FVector((X+1) * TileCellSize, (Y+1) * TileCellSize, NavTotalBounds.Max.Z) + NavMeshOrigin;
-	TileBB = FBox(TileBMin, TileBMax);
 
-	FBox RCBox = Unreal2RecastBox(TileBB);
+	TileBB = CalculateTileBounds(X, Y, NavMeshOrigin, NavTotalBounds, TileCellSize);
+	const FBox RCBox = Unreal2RecastBox(TileBB);
 	rcVcopy(TileConfig.bmin, &RCBox.Min.X);
 	rcVcopy(TileConfig.bmax, &RCBox.Max.X);
 			
@@ -3143,11 +3156,6 @@ TArray<FNavMeshTileData> FRecastNavMeshGenerator::GetCachedLayersData(FIntPoint
 	}
 }
 
-static FVector2D CalcTileCenter(int32 X, int32 Y, const FVector& NavMeshOrigin, float TileSizeInWorldUnits)
-{
-	return (FVector2D(NavMeshOrigin) + FVector2D(X, Y)*TileSizeInWorldUnits + FVector2D(0.5f, 0.5f)*TileSizeInWorldUnits);
-}
-
 void FRecastNavMeshGenerator::MarkDirtyTiles(const TArray<FNavigationDirtyArea>& DirtyAreas)
 {
 	QUICK_SCOPE_CYCLE_COUNTER(STAT_RecastNavMeshGenerator_MarkDirtyTiles);
@@ -3156,6 +3164,7 @@ void FRecastNavMeshGenerator::MarkDirtyTiles(const TArray<FNavigationDirtyArea>&
 	const float TileSizeInWorldUnits = Config.tileSize * Config.cs;
 	check(TileSizeInWorldUnits > 0);
 	const FVector NavMeshOrigin = FVector::ZeroVector;
+	const FVector RcNavMeshOrigin = Unreal2RecastPoint(NavMeshOrigin);
 	int32 NumTilesMarked = 0;
 
 	// find all tiles that need regeneration:
@@ -3169,20 +3178,17 @@ void FRecastNavMeshGenerator::MarkDirtyTiles(const TArray<FNavigationDirtyArea>&
 			continue;
 		}
 		
-		const int32 XMin = FMath::FloorToInt((AdjustedAreaBounds.Min.X - NavMeshOrigin.X) / TileSizeInWorldUnits);
-		const int32 XMax = FMath::FloorToInt((AdjustedAreaBounds.Max.X - NavMeshOrigin.X) / TileSizeInWorldUnits);
-		const int32 YMin = FMath::FloorToInt((AdjustedAreaBounds.Min.Y - NavMeshOrigin.Y) / TileSizeInWorldUnits);
-		const int32 YMax = FMath::FloorToInt((AdjustedAreaBounds.Max.Y - NavMeshOrigin.Y) / TileSizeInWorldUnits);
+		const FBox RcAreaBounds = Unreal2RecastBox(AdjustedAreaBounds);
+		const int32 XMin = FMath::FloorToInt((RcAreaBounds.Min.X - RcNavMeshOrigin.X) / TileSizeInWorldUnits);
+		const int32 XMax = FMath::FloorToInt((RcAreaBounds.Max.X - RcNavMeshOrigin.X) / TileSizeInWorldUnits);
+		const int32 YMin = FMath::FloorToInt((RcAreaBounds.Min.Y - RcNavMeshOrigin.Y) / TileSizeInWorldUnits);
+		const int32 YMax = FMath::FloorToInt((RcAreaBounds.Max.Y - RcNavMeshOrigin.Y) / TileSizeInWorldUnits);
 
 		for (int32 y = YMin; y <= YMax; ++y)
 		{
 			for (int32 x = XMin; x <= XMax; ++x)
 			{
-				FBox TileBox(
-					FVector(x+0, y+0, 0)*TileSizeInWorldUnits, 
-					FVector(x+1, y+1, 0)*TileSizeInWorldUnits);
-				TileBox.Min.Z = TotalNavBounds.Min.Z;
-				TileBox.Max.Z = TotalNavBounds.Max.Z;
+				const FBox TileBox = CalculateTileBounds(x, y, NavMeshOrigin, TotalNavBounds, TileSizeInWorldUnits);
 				if (!IntercestsInclusionBounds(TileBox))
 				{
 					// Skip this tile
@@ -3290,10 +3296,11 @@ void FRecastNavMeshGenerator::SortPendingBuildTiles()
 		// Calculate shortest distances between tiles and players
 		for (FTileElement& Element : PendingDirtyTiles)
 		{
-			FVector2D TileLocation = FVector2D(Element.Coord)*TileSizeInWorldUnits;
+			const FBox TileBox = CalculateTileBounds(Element.Coord.X, Element.Coord.Y, FVector::ZeroVector, TotalNavBounds, TileSizeInWorldUnits);
+			FVector2D TileCenter2D = FVector2D(TileBox.GetCenter());
 			for (FVector2D SeedLocation : SeedLocations)
 			{
-				const float DistSq = FVector2D::DistSquared(TileLocation, SeedLocation);
+				const float DistSq = FVector2D::DistSquared(TileCenter2D, SeedLocation);
 				if (DistSq < Element.SeedDistance)
 				{
 					Element.SeedDistance = DistSq;
