Rename FPostConstructInitializeProperties to something simpler

diff --git a/Engine/Source/Runtime/Engine/Classes/Engine/NetDriver.h b/Engine/Source/Runtime/Engine/Classes/Engine/NetDriver.h
--- a/Engine/Source/Runtime/Engine/Classes/Engine/NetDriver.h
+++ b/Engine/Source/Runtime/Engine/Classes/Engine/NetDriver.h
@@ -312,7 +312,7 @@ public:
 #endif
 
 	// Constructors.
-	ENGINE_API UNetDriver(const class FPostConstructInitializeProperties& PCIP);
+	ENGINE_API UNetDriver(const FObjectInitializer& ObjectInitializer);
 
 
 	// Begin UObject interface.
diff --git a/Engine/Source/Runtime/Engine/Classes/Engine/PackageMapClient.h b/Engine/Source/Runtime/Engine/Classes/Engine/PackageMapClient.h
--- a/Engine/Source/Runtime/Engine/Classes/Engine/PackageMapClient.h
+++ b/Engine/Source/Runtime/Engine/Classes/Engine/PackageMapClient.h
@@ -92,8 +92,8 @@ class UPackageMapClient : public UPackageMap
 {
     GENERATED_UCLASS_BODY()
 
-	UPackageMapClient( const class FPostConstructInitializeProperties & PCIP, UNetConnection * InConnection, TSharedPtr< FNetGUIDCache > InNetGUIDCache ) : 
-		UPackageMap( PCIP )
+	UPackageMapClient( const FObjectInitializer & ObjectInitializer, UNetConnection * InConnection, TSharedPtr< FNetGUIDCache > InNetGUIDCache ) : 
+		UPackageMap( ObjectInitializer )
 	,	Connection(InConnection)
 	{
 		GuidCache				= InNetGUIDCache;
diff --git a/Engine/Source/Runtime/Engine/Classes/Engine/PendingNetGame.h b/Engine/Source/Runtime/Engine/Classes/Engine/PendingNetGame.h
--- a/Engine/Source/Runtime/Engine/Classes/Engine/PendingNetGame.h
+++ b/Engine/Source/Runtime/Engine/Classes/Engine/PendingNetGame.h
@@ -96,10 +96,10 @@ public:
 	FString					ConnectionError;
 
 	// Constructor.
-	UPendingNetGame( const class FPostConstructInitializeProperties& PCIP, const FURL& InURL );
+	UPendingNetGame( const FObjectInitializer& ObjectInitializer, const FURL& InURL );
 
 	// Constructor.
-	UPendingNetGame( const class FPostConstructInitializeProperties& PCIP );
+	UPendingNetGame( const FObjectInitializer& ObjectInitializer );
 
 	void	InitNetDriver();
 
diff --git a/Engine/Source/Runtime/Engine/Classes/Engine/Polys.h b/Engine/Source/Runtime/Engine/Classes/Engine/Polys.h
--- a/Engine/Source/Runtime/Engine/Classes/Engine/Polys.h
+++ b/Engine/Source/Runtime/Engine/Classes/Engine/Polys.h
@@ -576,8 +576,8 @@ class UPolys : public UObject
 	TTransArray<FPoly> Element;
 
 	// Constructors.
-	UPolys(const class FPostConstructInitializeProperties& PCIP)
-	:  UObject(PCIP)
+	UPolys(const FObjectInitializer& ObjectInitializer)
+	:  UObject(ObjectInitializer)
 	, Element( this )
 	{}
 
diff --git a/Engine/Source/Runtime/Engine/Classes/Engine/Selection.h b/Engine/Source/Runtime/Engine/Classes/Engine/Selection.h
--- a/Engine/Source/Runtime/Engine/Classes/Engine/Selection.h
+++ b/Engine/Source/Runtime/Engine/Classes/Engine/Selection.h
@@ -29,7 +29,7 @@ public:
 	/** Called to deselect everything */
 	static FSimpleMulticastDelegate SelectNoneEvent;
 
-	USelection(const class FPostConstructInitializeProperties& PCIP);
+	USelection(const FObjectInitializer& ObjectInitializer);
 
 	typedef ClassArray::TIterator TClassIterator;
 	typedef ClassArray::TConstIterator TClassConstIterator;
diff --git a/Engine/Source/Runtime/Engine/Classes/Engine/StaticMesh.h b/Engine/Source/Runtime/Engine/Classes/Engine/StaticMesh.h
--- a/Engine/Source/Runtime/Engine/Classes/Engine/StaticMesh.h
+++ b/Engine/Source/Runtime/Engine/Classes/Engine/StaticMesh.h
@@ -478,7 +478,7 @@ public:
 	/**
 	 * Default constructor
 	 */
-	ENGINE_API UStaticMesh(const class FPostConstructInitializeProperties& PCIP);
+	ENGINE_API UStaticMesh(const FObjectInitializer& ObjectInitializer);
 
 	// Begin UObject interface.
 #if WITH_EDITOR
diff --git a/Engine/Source/Runtime/Engine/Classes/Engine/VoiceChannel.h b/Engine/Source/Runtime/Engine/Classes/Engine/VoiceChannel.h
--- a/Engine/Source/Runtime/Engine/Classes/Engine/VoiceChannel.h
+++ b/Engine/Source/Runtime/Engine/Classes/Engine/VoiceChannel.h
@@ -22,8 +22,8 @@ class UVoiceChannel : public UChannel
 	/**
 	 * Default constructor
 	 */
-	UVoiceChannel(const class FPostConstructInitializeProperties& PCIP)
-		: UChannel(PCIP)
+	UVoiceChannel(const FObjectInitializer& ObjectInitializer)
+		: UChannel(ObjectInitializer)
 	{
 		// Register with the network channel system
 		ChannelClasses[CHTYPE_Voice] = GetClass();
diff --git a/Engine/Source/Runtime/Engine/Classes/Engine/World.h b/Engine/Source/Runtime/Engine/Classes/Engine/World.h
--- a/Engine/Source/Runtime/Engine/Classes/Engine/World.h
+++ b/Engine/Source/Runtime/Engine/Classes/Engine/World.h
@@ -982,7 +982,7 @@ public:
 	/**
 	 * UWorld default constructor
 	 */
-	UWorld( const class FPostConstructInitializeProperties& PCIP );
+	UWorld( const FObjectInitializer& ObjectInitializer );
 
 	/**
 	 * UWorld constructor called at game startup and when creating a new world in the Editor.
@@ -990,7 +990,7 @@ public:
 	 *
 	 * @param	InURL	URL associated with this world.
 	 */
-	UWorld( const class FPostConstructInitializeProperties& PCIP,const FURL& InURL );
+	UWorld( const FObjectInitializer& ObjectInitializer,const FURL& InURL );
 	
 	// LINE TRACE
 
diff --git a/Engine/Source/Runtime/Engine/Classes/GameFramework/Actor.h b/Engine/Source/Runtime/Engine/Classes/GameFramework/Actor.h
--- a/Engine/Source/Runtime/Engine/Classes/GameFramework/Actor.h
+++ b/Engine/Source/Runtime/Engine/Classes/GameFramework/Actor.h
@@ -72,7 +72,7 @@ public:
 	/**
 	 * Default UObject constructor.
 	 */
-	AActor(const class FPostConstructInitializeProperties& PCIP);
+	AActor(const FObjectInitializer& ObjectInitializer);
 
 	/**
 	 * Primary Actor tick function, which calls TickActor().
diff --git a/Engine/Source/Runtime/Engine/Classes/GameFramework/Character.h b/Engine/Source/Runtime/Engine/Classes/GameFramework/Character.h
--- a/Engine/Source/Runtime/Engine/Classes/GameFramework/Character.h
+++ b/Engine/Source/Runtime/Engine/Classes/GameFramework/Character.h
@@ -205,7 +205,7 @@ public:
 	/**
 	 * Default UObject constructor.
 	 */
-	ACharacter(const class FPostConstructInitializeProperties& PCIP);
+	ACharacter(const FObjectInitializer& ObjectInitializer);
 
 	/** The main skeletal mesh associated with this Character (optional sub-object). */
 	UPROPERTY(Category=Character, VisibleAnywhere, BlueprintReadOnly)
@@ -224,10 +224,10 @@ public:
 	UPROPERTY(Category=Character, VisibleAnywhere, BlueprintReadOnly)
 	TSubobjectPtr<class UCapsuleComponent> CapsuleComponent;
 
-	/** Name of the MeshComponent. Use this name if you want to prevent creation of the component (with PCIP.DoNotCreateDefaultSubobject). */
+	/** Name of the MeshComponent. Use this name if you want to prevent creation of the component (with ObjectInitializer.DoNotCreateDefaultSubobject). */
 	static FName MeshComponentName;
 
-	/** Name of the CharacterMovement component. Use this name if you want to use a different class (with PCIP.SetDefaultSubobjectClass). */
+	/** Name of the CharacterMovement component. Use this name if you want to use a different class (with ObjectInitializer.SetDefaultSubobjectClass). */
 	static FName CharacterMovementComponentName;
 
 	/** Name of the CapsuleComponent. */
diff --git a/Engine/Source/Runtime/Engine/Classes/GameFramework/CharacterMovementComponent.h b/Engine/Source/Runtime/Engine/Classes/GameFramework/CharacterMovementComponent.h
--- a/Engine/Source/Runtime/Engine/Classes/GameFramework/CharacterMovementComponent.h
+++ b/Engine/Source/Runtime/Engine/Classes/GameFramework/CharacterMovementComponent.h
@@ -130,7 +130,7 @@ public:
 	/**
 	 * Default UObject constructor.
 	 */
-	UCharacterMovementComponent(const class FPostConstructInitializeProperties& PCIP);
+	UCharacterMovementComponent(const FObjectInitializer& ObjectInitializer);
 
 protected:
 
diff --git a/Engine/Source/Runtime/Engine/Classes/GameFramework/DefaultPawn.h b/Engine/Source/Runtime/Engine/Classes/GameFramework/DefaultPawn.h
--- a/Engine/Source/Runtime/Engine/Classes/GameFramework/DefaultPawn.h
+++ b/Engine/Source/Runtime/Engine/Classes/GameFramework/DefaultPawn.h
@@ -84,7 +84,7 @@ private:
 	virtual void Turn(float Val);
 
 public:
-	/** Name of the MovementComponent.  Use this name if you want to use a different class (with PCIP.SetDefaultSubobjectClass). */
+	/** Name of the MovementComponent.  Use this name if you want to use a different class (with ObjectInitializer.SetDefaultSubobjectClass). */
 	static FName MovementComponentName;
 
 	/** DefaultPawn movement component */
@@ -98,7 +98,7 @@ public:
 	UPROPERTY(Category=Pawn, VisibleAnywhere, BlueprintReadOnly)
 	TSubobjectPtr<USphereComponent> CollisionComponent;
 
-	/** Name of the MeshComponent. Use this name if you want to prevent creation of the component (with PCIP.DoNotCreateDefaultSubobject). */
+	/** Name of the MeshComponent. Use this name if you want to prevent creation of the component (with ObjectInitializer.DoNotCreateDefaultSubobject). */
 	static FName MeshComponentName;
 
 	/** The mesh associated with this Pawn. */
diff --git a/Engine/Source/Runtime/Engine/Classes/GameFramework/Pawn.h b/Engine/Source/Runtime/Engine/Classes/GameFramework/Pawn.h
--- a/Engine/Source/Runtime/Engine/Classes/GameFramework/Pawn.h
+++ b/Engine/Source/Runtime/Engine/Classes/GameFramework/Pawn.h
@@ -35,7 +35,7 @@ public:
 	/**
 	 * Default UObject constructor.
 	 */
-	APawn(const class FPostConstructInitializeProperties& PCIP);
+	APawn(const FObjectInitializer& ObjectInitializer);
 
 	/** Return our PawnMovementComponent, if we have one. By default, returns the first PawnMovementComponent found. Native classes that create their own movement component should override this method for more efficiency. */
 	UFUNCTION(BlueprintCallable, meta=(Tooltip="Return our PawnMovementComponent, if we have one."), Category="Pawn")
diff --git a/Engine/Source/Runtime/Engine/Classes/GameFramework/WheeledVehicle.h b/Engine/Source/Runtime/Engine/Classes/GameFramework/WheeledVehicle.h
--- a/Engine/Source/Runtime/Engine/Classes/GameFramework/WheeledVehicle.h
+++ b/Engine/Source/Runtime/Engine/Classes/GameFramework/WheeledVehicle.h
@@ -19,10 +19,10 @@ class ENGINE_API AWheeledVehicle : public APawn
 	UPROPERTY(Category=Vehicle, VisibleDefaultsOnly, BlueprintReadOnly)
 	TSubobjectPtr<class UWheeledVehicleMovementComponent> VehicleMovement;
 
-	/** Name of the MeshComponent. Use this name if you want to prevent creation of the component (with PCIP.DoNotCreateDefaultSubobject). */
+	/** Name of the MeshComponent. Use this name if you want to prevent creation of the component (with ObjectInitializer.DoNotCreateDefaultSubobject). */
 	static FName VehicleMeshComponentName;
 
-	/** Name of the VehicleMovement. Use this name if you want to use a different class (with PCIP.SetDefaultSubobjectClass). */
+	/** Name of the VehicleMovement. Use this name if you want to use a different class (with ObjectInitializer.SetDefaultSubobjectClass). */
 	static FName VehicleMovementComponentName;
 
 	/** Util to get the wheeled vehicle movement component */
diff --git a/Engine/Source/Runtime/Engine/Private/AI/AISystemBase.cpp b/Engine/Source/Runtime/Engine/Private/AI/AISystemBase.cpp
--- a/Engine/Source/Runtime/Engine/Private/AI/AISystemBase.cpp
+++ b/Engine/Source/Runtime/Engine/Private/AI/AISystemBase.cpp
@@ -3,8 +3,8 @@
 #include "EnginePrivate.h"
 #include "AI/AISystemBase.h"
 
-UAISystemBase::UAISystemBase(const class FPostConstructInitializeProperties& PCIP)
-: Super(PCIP)
+UAISystemBase::UAISystemBase(const FObjectInitializer& ObjectInitializer)
+: Super(ObjectInitializer)
 {
 }
 
diff --git a/Engine/Source/Runtime/Engine/Private/AI/Navigation/AvoidanceManager.cpp b/Engine/Source/Runtime/Engine/Private/AI/Navigation/AvoidanceManager.cpp
--- a/Engine/Source/Runtime/Engine/Private/AI/Navigation/AvoidanceManager.cpp
+++ b/Engine/Source/Runtime/Engine/Private/AI/Navigation/AvoidanceManager.cpp
@@ -29,7 +29,7 @@ void FNavAvoidanceData::Init(class UAvoidanceManager* Avoidance, const FVector&
 bool UAvoidanceManager::bSystemActive = true;
 #endif
 
-UAvoidanceManager::UAvoidanceManager(const class FPostConstructInitializeProperties& PCIP) : Super(PCIP)
+UAvoidanceManager::UAvoidanceManager(const FObjectInitializer& ObjectInitializer) : Super(ObjectInitializer)
 {
 	DefaultTimeToLive = 1.5f;
 	LockTimeAfterAvoid = 0.2f;
diff --git a/Engine/Source/Runtime/Engine/Private/AI/Navigation/NavAreas/NavArea.cpp b/Engine/Source/Runtime/Engine/Private/AI/Navigation/NavAreas/NavArea.cpp
--- a/Engine/Source/Runtime/Engine/Private/AI/Navigation/NavAreas/NavArea.cpp
+++ b/Engine/Source/Runtime/Engine/Private/AI/Navigation/NavAreas/NavArea.cpp
@@ -3,7 +3,7 @@
 #include "EnginePrivate.h"
 #include "AI/Navigation/NavAreas/NavArea.h"
 
-UNavArea::UNavArea(const class FPostConstructInitializeProperties& PCIP) : Super(PCIP)
+UNavArea::UNavArea(const FObjectInitializer& ObjectInitializer) : Super(ObjectInitializer)
 {
 	DefaultCost = 1.f;
 	FixedAreaEnteringCost = 0.f;
diff --git a/Engine/Source/Runtime/Engine/Private/AI/Navigation/NavAreas/NavAreaMeta.cpp b/Engine/Source/Runtime/Engine/Private/AI/Navigation/NavAreas/NavAreaMeta.cpp
--- a/Engine/Source/Runtime/Engine/Private/AI/Navigation/NavAreas/NavAreaMeta.cpp
+++ b/Engine/Source/Runtime/Engine/Private/AI/Navigation/NavAreas/NavAreaMeta.cpp
@@ -3,7 +3,7 @@
 #include "EnginePrivate.h"
 #include "AI/Navigation/NavAreas/NavAreaMeta.h"
 
-UNavAreaMeta::UNavAreaMeta(const class FPostConstructInitializeProperties& PCIP) : Super(PCIP)
+UNavAreaMeta::UNavAreaMeta(const FObjectInitializer& ObjectInitializer) : Super(ObjectInitializer)
 {
 }
 
diff --git a/Engine/Source/Runtime/Engine/Private/AI/Navigation/NavAreas/NavAreaMeta_SwitchByAgent.cpp b/Engine/Source/Runtime/Engine/Private/AI/Navigation/NavAreas/NavAreaMeta_SwitchByAgent.cpp
--- a/Engine/Source/Runtime/Engine/Private/AI/Navigation/NavAreas/NavAreaMeta_SwitchByAgent.cpp
+++ b/Engine/Source/Runtime/Engine/Private/AI/Navigation/NavAreas/NavAreaMeta_SwitchByAgent.cpp
@@ -4,7 +4,7 @@
 #include "AI/Navigation/NavAreas/NavArea_Default.h"
 #include "AI/Navigation/NavAreas/NavAreaMeta_SwitchByAgent.h"
 
-UNavAreaMeta_SwitchByAgent::UNavAreaMeta_SwitchByAgent(const class FPostConstructInitializeProperties& PCIP) : Super(PCIP)
+UNavAreaMeta_SwitchByAgent::UNavAreaMeta_SwitchByAgent(const FObjectInitializer& ObjectInitializer) : Super(ObjectInitializer)
 {
 	Agent0Area = UNavArea_Default::StaticClass();
 	Agent1Area = UNavArea_Default::StaticClass();
diff --git a/Engine/Source/Runtime/Engine/Private/AI/Navigation/NavAreas/NavArea_Default.cpp b/Engine/Source/Runtime/Engine/Private/AI/Navigation/NavAreas/NavArea_Default.cpp
--- a/Engine/Source/Runtime/Engine/Private/AI/Navigation/NavAreas/NavArea_Default.cpp
+++ b/Engine/Source/Runtime/Engine/Private/AI/Navigation/NavAreas/NavArea_Default.cpp
@@ -3,6 +3,6 @@
 #include "EnginePrivate.h"
 #include "AI/Navigation/NavAreas/NavArea_Default.h"
 
-UNavArea_Default::UNavArea_Default(const class FPostConstructInitializeProperties& PCIP) : Super(PCIP)
+UNavArea_Default::UNavArea_Default(const FObjectInitializer& ObjectInitializer) : Super(ObjectInitializer)
 {
 }
diff --git a/Engine/Source/Runtime/Engine/Private/AI/Navigation/NavAreas/NavArea_LowHeight.cpp b/Engine/Source/Runtime/Engine/Private/AI/Navigation/NavAreas/NavArea_LowHeight.cpp
--- a/Engine/Source/Runtime/Engine/Private/AI/Navigation/NavAreas/NavArea_LowHeight.cpp
+++ b/Engine/Source/Runtime/Engine/Private/AI/Navigation/NavAreas/NavArea_LowHeight.cpp
@@ -3,7 +3,7 @@
 #include "EnginePrivate.h"
 #include "AI/Navigation/NavAreas/NavArea_LowHeight.h"
 
-UNavArea_LowHeight::UNavArea_LowHeight(const class FPostConstructInitializeProperties& PCIP) : Super(PCIP)
+UNavArea_LowHeight::UNavArea_LowHeight(const FObjectInitializer& ObjectInitializer) : Super(ObjectInitializer)
 {
 	DefaultCost = BIG_NUMBER;
 	DrawColor = FColor(0, 0, 128);
diff --git a/Engine/Source/Runtime/Engine/Private/AI/Navigation/NavAreas/NavArea_Null.cpp b/Engine/Source/Runtime/Engine/Private/AI/Navigation/NavAreas/NavArea_Null.cpp
--- a/Engine/Source/Runtime/Engine/Private/AI/Navigation/NavAreas/NavArea_Null.cpp
+++ b/Engine/Source/Runtime/Engine/Private/AI/Navigation/NavAreas/NavArea_Null.cpp
@@ -3,7 +3,7 @@
 #include "EnginePrivate.h"
 #include "AI/Navigation/NavAreas/NavArea_Null.h"
 
-UNavArea_Null::UNavArea_Null(const class FPostConstructInitializeProperties& PCIP) : Super(PCIP)
+UNavArea_Null::UNavArea_Null(const FObjectInitializer& ObjectInitializer) : Super(ObjectInitializer)
 {
 	DefaultCost = BIG_NUMBER;
 }
diff --git a/Engine/Source/Runtime/Engine/Private/AI/Navigation/NavCollision.cpp b/Engine/Source/Runtime/Engine/Private/AI/Navigation/NavCollision.cpp
--- a/Engine/Source/Runtime/Engine/Private/AI/Navigation/NavCollision.cpp
+++ b/Engine/Source/Runtime/Engine/Private/AI/Navigation/NavCollision.cpp
@@ -135,7 +135,7 @@ bool FDerivedDataNavCollisionCooker::Build( TArray<uint8>& OutData )
 //----------------------------------------------------------------------//
 // UNavCollision
 //----------------------------------------------------------------------//
-UNavCollision::UNavCollision(const class FPostConstructInitializeProperties& PCIP) : Super(PCIP)
+UNavCollision::UNavCollision(const FObjectInitializer& ObjectInitializer) : Super(ObjectInitializer)
 {	
 }
 
diff --git a/Engine/Source/Runtime/Engine/Private/AI/Navigation/NavFilters/NavigationQueryFilter.cpp b/Engine/Source/Runtime/Engine/Private/AI/Navigation/NavFilters/NavigationQueryFilter.cpp
--- a/Engine/Source/Runtime/Engine/Private/AI/Navigation/NavFilters/NavigationQueryFilter.cpp
+++ b/Engine/Source/Runtime/Engine/Private/AI/Navigation/NavFilters/NavigationQueryFilter.cpp
@@ -110,7 +110,7 @@ uint16 FNavigationQueryFilter::GetExcludeFlags() const
 //----------------------------------------------------------------------//
 // UNavigationQueryFilter
 //----------------------------------------------------------------------//
-UNavigationQueryFilter::UNavigationQueryFilter(const class FPostConstructInitializeProperties& PCIP) : Super(PCIP)
+UNavigationQueryFilter::UNavigationQueryFilter(const FObjectInitializer& ObjectInitializer) : Super(ObjectInitializer)
 {
 	IncludeFlags.Packed = 0xffff;
 	ExcludeFlags.Packed = 0;
diff --git a/Engine/Source/Runtime/Engine/Private/AI/Navigation/NavFilters/RecastFilter_UseDefaultArea.cpp b/Engine/Source/Runtime/Engine/Private/AI/Navigation/NavFilters/RecastFilter_UseDefaultArea.cpp
--- a/Engine/Source/Runtime/Engine/Private/AI/Navigation/NavFilters/RecastFilter_UseDefaultArea.cpp
+++ b/Engine/Source/Runtime/Engine/Private/AI/Navigation/NavFilters/RecastFilter_UseDefaultArea.cpp
@@ -3,7 +3,7 @@
 #include "EnginePrivate.h"
 #include "AI/Navigation/NavFilters/RecastFilter_UseDefaultArea.h"
 
-URecastFilter_UseDefaultArea::URecastFilter_UseDefaultArea(const class FPostConstructInitializeProperties& PCIP) : Super(PCIP)
+URecastFilter_UseDefaultArea::URecastFilter_UseDefaultArea(const FObjectInitializer& ObjectInitializer) : Super(ObjectInitializer)
 {
 }
 
diff --git a/Engine/Source/Runtime/Engine/Private/AI/Navigation/NavLinkCustomComponent.cpp b/Engine/Source/Runtime/Engine/Private/AI/Navigation/NavLinkCustomComponent.cpp
--- a/Engine/Source/Runtime/Engine/Private/AI/Navigation/NavLinkCustomComponent.cpp
+++ b/Engine/Source/Runtime/Engine/Private/AI/Navigation/NavLinkCustomComponent.cpp
@@ -11,7 +11,7 @@
 #include "AI/Navigation/NavAreas/NavArea_Default.h"
 #include "AI/Navigation/NavAreas/NavArea_Null.h"
 
-UNavLinkCustomComponent::UNavLinkCustomComponent(const class FPostConstructInitializeProperties& PCIP) : Super(PCIP)
+UNavLinkCustomComponent::UNavLinkCustomComponent(const FObjectInitializer& ObjectInitializer) : Super(ObjectInitializer)
 {
 	NavLinkUserId = 0;
 	LinkRelativeStart = FVector(70, 0, 0);
diff --git a/Engine/Source/Runtime/Engine/Private/AI/Navigation/NavLinkProxy.cpp b/Engine/Source/Runtime/Engine/Private/AI/Navigation/NavLinkProxy.cpp
--- a/Engine/Source/Runtime/Engine/Private/AI/Navigation/NavLinkProxy.cpp
+++ b/Engine/Source/Runtime/Engine/Private/AI/Navigation/NavLinkProxy.cpp
@@ -9,21 +9,21 @@
 #include "AI/Navigation/NavLinkCustomComponent.h"
 #include "AI/Navigation/NavLinkProxy.h"
 
-ANavLinkProxy::ANavLinkProxy(const class FPostConstructInitializeProperties& PCIP) : Super(PCIP)
+ANavLinkProxy::ANavLinkProxy(const FObjectInitializer& ObjectInitializer) : Super(ObjectInitializer)
 {
-	TSubobjectPtr<USceneComponent> SceneComponent = PCIP.CreateDefaultSubobject<USceneComponent>(this, TEXT("PositionComponent"));
+	TSubobjectPtr<USceneComponent> SceneComponent = ObjectInitializer.CreateDefaultSubobject<USceneComponent>(this, TEXT("PositionComponent"));
 	RootComponent = SceneComponent;
 
 	bHidden = true;
 
 #if WITH_EDITORONLY_DATA
-	EdRenderComp = PCIP.CreateDefaultSubobject<UNavLinkRenderingComponent>(this, TEXT("EdRenderComp"));
+	EdRenderComp = ObjectInitializer.CreateDefaultSubobject<UNavLinkRenderingComponent>(this, TEXT("EdRenderComp"));
 	EdRenderComp->PostPhysicsComponentTick.bCanEverTick = false;
 	EdRenderComp->AttachParent = RootComponent;
 #endif // WITH_EDITORONLY_DATA
 
 #if WITH_EDITOR
-	SpriteComponent = PCIP.CreateEditorOnlyDefaultSubobject<UBillboardComponent>(this, TEXT("Sprite"));
+	SpriteComponent = ObjectInitializer.CreateEditorOnlyDefaultSubobject<UBillboardComponent>(this, TEXT("Sprite"));
 	if (!IsRunningCommandlet() && (SpriteComponent != NULL))
 	{
 		static ConstructorHelpers::FObjectFinderOptional<UTexture2D> SpriteTexture(TEXT("/Engine/EditorResources/AI/S_NavLink"));
@@ -39,7 +39,7 @@ ANavLinkProxy::ANavLinkProxy(const class FPostConstructInitializeProperties& PCI
 	}
 #endif
 
-	SmartLinkComp = PCIP.CreateDefaultSubobject<UNavLinkCustomComponent>(this, TEXT("SmartLinkComp"));
+	SmartLinkComp = ObjectInitializer.CreateDefaultSubobject<UNavLinkCustomComponent>(this, TEXT("SmartLinkComp"));
 	SmartLinkComp->SetNavigationRelevancy(false);
 	SmartLinkComp->SetMoveReachedLink(this, &ANavLinkProxy::NotifySmartLinkReached);
 	bSmartLinkIsRelevant = false;
diff --git a/Engine/Source/Runtime/Engine/Private/AI/Navigation/NavLinkRenderingComponent.cpp b/Engine/Source/Runtime/Engine/Private/AI/Navigation/NavLinkRenderingComponent.cpp
--- a/Engine/Source/Runtime/Engine/Private/AI/Navigation/NavLinkRenderingComponent.cpp
+++ b/Engine/Source/Runtime/Engine/Private/AI/Navigation/NavLinkRenderingComponent.cpp
@@ -7,8 +7,8 @@
 //----------------------------------------------------------------------//
 // UNavLinkRenderingComponent
 //----------------------------------------------------------------------//
-UNavLinkRenderingComponent::UNavLinkRenderingComponent(const class FPostConstructInitializeProperties& PCIP)
-	: Super(PCIP)
+UNavLinkRenderingComponent::UNavLinkRenderingComponent(const FObjectInitializer& ObjectInitializer)
+	: Super(ObjectInitializer)
 {
 	// properties
 
diff --git a/Engine/Source/Runtime/Engine/Private/AI/Navigation/NavModifierComponent.cpp b/Engine/Source/Runtime/Engine/Private/AI/Navigation/NavModifierComponent.cpp
--- a/Engine/Source/Runtime/Engine/Private/AI/Navigation/NavModifierComponent.cpp
+++ b/Engine/Source/Runtime/Engine/Private/AI/Navigation/NavModifierComponent.cpp
@@ -5,7 +5,7 @@
 #include "AI/Navigation/NavModifierComponent.h"
 #include "AI/Navigation/NavAreas/NavArea_Null.h"
 
-UNavModifierComponent::UNavModifierComponent(const FPostConstructInitializeProperties& PCIP) : Super(PCIP)
+UNavModifierComponent::UNavModifierComponent(const FObjectInitializer& ObjectInitializer) : Super(ObjectInitializer)
 {
 	AreaClass = UNavArea_Null::StaticClass();
 }
diff --git a/Engine/Source/Runtime/Engine/Private/AI/Navigation/NavModifierVolume.cpp b/Engine/Source/Runtime/Engine/Private/AI/Navigation/NavModifierVolume.cpp
--- a/Engine/Source/Runtime/Engine/Private/AI/Navigation/NavModifierVolume.cpp
+++ b/Engine/Source/Runtime/Engine/Private/AI/Navigation/NavModifierVolume.cpp
@@ -7,8 +7,8 @@
 //----------------------------------------------------------------------//
 // ANavModifierVolume
 //----------------------------------------------------------------------//
-ANavModifierVolume::ANavModifierVolume(const class FPostConstructInitializeProperties& PCIP)
-	: Super(PCIP)
+ANavModifierVolume::ANavModifierVolume(const FObjectInitializer& ObjectInitializer)
+	: Super(ObjectInitializer)
 {
 }
 
diff --git a/Engine/Source/Runtime/Engine/Private/AI/Navigation/NavRelevantComponent.cpp b/Engine/Source/Runtime/Engine/Private/AI/Navigation/NavRelevantComponent.cpp
--- a/Engine/Source/Runtime/Engine/Private/AI/Navigation/NavRelevantComponent.cpp
+++ b/Engine/Source/Runtime/Engine/Private/AI/Navigation/NavRelevantComponent.cpp
@@ -3,7 +3,7 @@
 #include "EnginePrivate.h"
 #include "AI/Navigation/NavRelevantComponent.h"
 
-UNavRelevantComponent::UNavRelevantComponent(const class FPostConstructInitializeProperties& PCIP) : Super(PCIP)
+UNavRelevantComponent::UNavRelevantComponent(const FObjectInitializer& ObjectInitializer) : Super(ObjectInitializer)
 {
 	bNavigationRelevant = true;
 }
diff --git a/Engine/Source/Runtime/Engine/Private/AI/Navigation/NavTestRenderingComponent.cpp b/Engine/Source/Runtime/Engine/Private/AI/Navigation/NavTestRenderingComponent.cpp
--- a/Engine/Source/Runtime/Engine/Private/AI/Navigation/NavTestRenderingComponent.cpp
+++ b/Engine/Source/Runtime/Engine/Private/AI/Navigation/NavTestRenderingComponent.cpp
@@ -503,7 +503,7 @@ private:
 	uint32 bShowDiff : 1;
 };
 
-UNavTestRenderingComponent::UNavTestRenderingComponent(const class FPostConstructInitializeProperties& PCIP) : Super(PCIP)
+UNavTestRenderingComponent::UNavTestRenderingComponent(const FObjectInitializer& ObjectInitializer) : Super(ObjectInitializer)
 {
 }
 
diff --git a/Engine/Source/Runtime/Engine/Private/AI/Navigation/NavigationData.cpp b/Engine/Source/Runtime/Engine/Private/AI/Navigation/NavigationData.cpp
--- a/Engine/Source/Runtime/Engine/Private/AI/Navigation/NavigationData.cpp
+++ b/Engine/Source/Runtime/Engine/Private/AI/Navigation/NavigationData.cpp
@@ -93,8 +93,8 @@ FSupportedAreaData::FSupportedAreaData(TSubclassOf<UNavArea> NavAreaClass, int32
 //----------------------------------------------------------------------//
 // ANavigationData                                                                
 //----------------------------------------------------------------------//
-ANavigationData::ANavigationData(const class FPostConstructInitializeProperties& PCIP)
-	: Super(PCIP)
+ANavigationData::ANavigationData(const FObjectInitializer& ObjectInitializer)
+	: Super(ObjectInitializer)
 	, bEnableDrawing(false)
 	, bRebuildAtRuntime(false)
 	, DataVersion(NAVMESHVER_LATEST)
diff --git a/Engine/Source/Runtime/Engine/Private/AI/Navigation/NavigationGraph.cpp b/Engine/Source/Runtime/Engine/Private/AI/Navigation/NavigationGraph.cpp
--- a/Engine/Source/Runtime/Engine/Private/AI/Navigation/NavigationGraph.cpp
+++ b/Engine/Source/Runtime/Engine/Private/AI/Navigation/NavigationGraph.cpp
@@ -18,8 +18,8 @@ FNavGraphNode::FNavGraphNode()
 //----------------------------------------------------------------------//
 // UNavigationGraphNodeComponent
 //----------------------------------------------------------------------//
-UNavigationGraphNodeComponent::UNavigationGraphNodeComponent(const class FPostConstructInitializeProperties& PCIP)
-	: Super(PCIP)
+UNavigationGraphNodeComponent::UNavigationGraphNodeComponent(const FObjectInitializer& ObjectInitializer)
+	: Super(ObjectInitializer)
 {
 }
 
@@ -44,8 +44,8 @@ void UNavigationGraphNodeComponent::BeginDestroy()
 //----------------------------------------------------------------------//
 // ANavigationGraphNode
 //----------------------------------------------------------------------//
-ANavigationGraphNode::ANavigationGraphNode(const class FPostConstructInitializeProperties& PCIP)
-	: Super(PCIP)
+ANavigationGraphNode::ANavigationGraphNode(const FObjectInitializer& ObjectInitializer)
+	: Super(ObjectInitializer)
 {
 }
 
@@ -53,8 +53,8 @@ ANavigationGraphNode::ANavigationGraphNode(const class FPostConstructInitializeP
 // ANavigationGraph
 //----------------------------------------------------------------------//
 
-ANavigationGraph::ANavigationGraph(const class FPostConstructInitializeProperties& PCIP)
-	: Super(PCIP)
+ANavigationGraph::ANavigationGraph(const FObjectInitializer& ObjectInitializer)
+	: Super(ObjectInitializer)
 {
 	if (HasAnyFlags(RF_ClassDefaultObject) == false)
 	{
diff --git a/Engine/Source/Runtime/Engine/Private/AI/Navigation/NavigationInterfaces.cpp b/Engine/Source/Runtime/Engine/Private/AI/Navigation/NavigationInterfaces.cpp
--- a/Engine/Source/Runtime/Engine/Private/AI/Navigation/NavigationInterfaces.cpp
+++ b/Engine/Source/Runtime/Engine/Private/AI/Navigation/NavigationInterfaces.cpp
@@ -11,28 +11,28 @@
 
 uint32 INavLinkCustomInterface::NextUniqueId = 1;
 
-UNavAgentInterface::UNavAgentInterface(const class FPostConstructInitializeProperties& PCIP)
-	: Super(PCIP)
+UNavAgentInterface::UNavAgentInterface(const FObjectInitializer& ObjectInitializer)
+	: Super(ObjectInitializer)
 {
 }
 
-UNavigationPathGenerator::UNavigationPathGenerator(const class FPostConstructInitializeProperties& PCIP)
-	: Super(PCIP)
+UNavigationPathGenerator::UNavigationPathGenerator(const FObjectInitializer& ObjectInitializer)
+	: Super(ObjectInitializer)
 {
 }
 
-UNavNodeInterface::UNavNodeInterface(const class FPostConstructInitializeProperties& PCIP)
-	: Super(PCIP)
+UNavNodeInterface::UNavNodeInterface(const FObjectInitializer& ObjectInitializer)
+	: Super(ObjectInitializer)
 {
 }
 
-UNavLinkHostInterface::UNavLinkHostInterface(const class FPostConstructInitializeProperties& PCIP)
-	: Super(PCIP)
+UNavLinkHostInterface::UNavLinkHostInterface(const FObjectInitializer& ObjectInitializer)
+	: Super(ObjectInitializer)
 {
 }
 
-UNavLinkCustomInterface::UNavLinkCustomInterface(const class FPostConstructInitializeProperties& PCIP)
-	: Super(PCIP)
+UNavLinkCustomInterface::UNavLinkCustomInterface(const FObjectInitializer& ObjectInitializer)
+	: Super(ObjectInitializer)
 {
 }
 
@@ -59,12 +59,12 @@ FNavigationLink INavLinkCustomInterface::GetModifier(const INavLinkCustomInterfa
 	return LinkMod;
 }
 
-UNavPathObserverInterface::UNavPathObserverInterface(const class FPostConstructInitializeProperties& PCIP)
-	: Super(PCIP)
+UNavPathObserverInterface::UNavPathObserverInterface(const FObjectInitializer& ObjectInitializer)
+	: Super(ObjectInitializer)
 {
 }
 
-UNavRelevantInterface::UNavRelevantInterface(const class FPostConstructInitializeProperties& PCIP)
-	: Super(PCIP)
+UNavRelevantInterface::UNavRelevantInterface(const FObjectInitializer& ObjectInitializer)
+	: Super(ObjectInitializer)
 {
 }
diff --git a/Engine/Source/Runtime/Engine/Private/AI/Navigation/NavigationModifier.cpp b/Engine/Source/Runtime/Engine/Private/AI/Navigation/NavigationModifier.cpp
--- a/Engine/Source/Runtime/Engine/Private/AI/Navigation/NavigationModifier.cpp
+++ b/Engine/Source/Runtime/Engine/Private/AI/Navigation/NavigationModifier.cpp
@@ -30,8 +30,8 @@ FNavigationLinkBase::FNavigationLinkBase()
 //----------------------------------------------------------------------//
 // UNavLinkDefinition
 //----------------------------------------------------------------------//
-UNavLinkDefinition::UNavLinkDefinition(const class FPostConstructInitializeProperties& PCIP)
-	: Super(PCIP)
+UNavLinkDefinition::UNavLinkDefinition(const FObjectInitializer& ObjectInitializer)
+	: Super(ObjectInitializer)
 	, bHasDeterminedMetaAreaClass(false)
 	, bHasMetaAreaClass(false)
 	, bHasDeterminedAdjustableLinks(false)
@@ -585,7 +585,7 @@ uint32 FCompositeNavModifier::GetAllocatedSize() const
 //----------------------------------------------------------------------//
 // UNavLinkTrivial
 //----------------------------------------------------------------------//
-UNavLinkTrivial::UNavLinkTrivial(const class FPostConstructInitializeProperties& PCIP) : Super(PCIP)
+UNavLinkTrivial::UNavLinkTrivial(const FObjectInitializer& ObjectInitializer) : Super(ObjectInitializer)
 {
 	FNavigationLink& Link = Links[Links.Add(FNavigationLink(FVector(0,100,0), FVector(0,-100,0)))];
 }
diff --git a/Engine/Source/Runtime/Engine/Private/AI/Navigation/NavigationPath.cpp b/Engine/Source/Runtime/Engine/Private/AI/Navigation/NavigationPath.cpp
--- a/Engine/Source/Runtime/Engine/Private/AI/Navigation/NavigationPath.cpp
+++ b/Engine/Source/Runtime/Engine/Private/AI/Navigation/NavigationPath.cpp
@@ -974,8 +974,8 @@ FString FNavMeshPath::GetDescription() const
 // UNavigationPath
 //----------------------------------------------------------------------//
 
-UNavigationPath::UNavigationPath(const class FPostConstructInitializeProperties& PCIP)
-	: Super(PCIP)
+UNavigationPath::UNavigationPath(const FObjectInitializer& ObjectInitializer)
+	: Super(ObjectInitializer)
 	, bIsValid(false)
 	, bDebugDrawingEnabled(false)
 	, DebugDrawingColor(FColor::White)
diff --git a/Engine/Source/Runtime/Engine/Private/AI/Navigation/NavigationSystem.cpp b/Engine/Source/Runtime/Engine/Private/AI/Navigation/NavigationSystem.cpp
--- a/Engine/Source/Runtime/Engine/Private/AI/Navigation/NavigationSystem.cpp
+++ b/Engine/Source/Runtime/Engine/Private/AI/Navigation/NavigationSystem.cpp
@@ -181,8 +181,8 @@ bool UNavigationSystem::bUpdateNavOctreeOnComponentChange = true;
 // life cycle stuff                                                                
 //----------------------------------------------------------------------//
 
-UNavigationSystem::UNavigationSystem(const FPostConstructInitializeProperties& PCIP)
-	: Super(PCIP)
+UNavigationSystem::UNavigationSystem(const FObjectInitializer& ObjectInitializer)
+	: Super(ObjectInitializer)
 	, bWholeWorldNavigable(false)
 	, bAddPlayersToGenerationSeeds(true)
 	, bSkipAgentHeightCheckWhenPickingNavData(false)
diff --git a/Engine/Source/Runtime/Engine/Private/AI/Navigation/NavigationTestingActor.cpp b/Engine/Source/Runtime/Engine/Private/AI/Navigation/NavigationTestingActor.cpp
--- a/Engine/Source/Runtime/Engine/Private/AI/Navigation/NavigationTestingActor.cpp
+++ b/Engine/Source/Runtime/Engine/Private/AI/Navigation/NavigationTestingActor.cpp
@@ -21,10 +21,10 @@ TStatId FNavTestTickHelper::GetStatId() const
 	RETURN_QUICK_DECLARE_CYCLE_STAT(FNavTestTickHelper, STATGROUP_Tickables);
 }
 
-ANavigationTestingActor::ANavigationTestingActor(const class FPostConstructInitializeProperties& PCIP) : Super(PCIP)
+ANavigationTestingActor::ANavigationTestingActor(const FObjectInitializer& ObjectInitializer) : Super(ObjectInitializer)
 {
 #if WITH_EDITORONLY_DATA
-	EdRenderComp = PCIP.CreateDefaultSubobject<UNavTestRenderingComponent>(this, TEXT("EdRenderComp"));
+	EdRenderComp = ObjectInitializer.CreateDefaultSubobject<UNavTestRenderingComponent>(this, TEXT("EdRenderComp"));
 	EdRenderComp->PostPhysicsComponentTick.bCanEverTick = false;
 
 #if WITH_RECAST
@@ -48,7 +48,7 @@ ANavigationTestingActor::ANavigationTestingActor(const class FPostConstructIniti
 	// collision profile name set up - found in baseengine.ini
 	static FName CollisionProfileName(TEXT("Pawn"));
 
-	CapsuleComponent = PCIP.CreateDefaultSubobject<UCapsuleComponent>(this, TEXT("CollisionCylinder"));
+	CapsuleComponent = ObjectInitializer.CreateDefaultSubobject<UCapsuleComponent>(this, TEXT("CollisionCylinder"));
 	CapsuleComponent->InitCapsuleSize(NavAgentProps.AgentRadius, NavAgentProps.AgentHeight / 2);
 	CapsuleComponent->SetCollisionProfileName(CollisionProfileName);
 	CapsuleComponent->CanCharacterStepUpOn = ECB_No;
diff --git a/Engine/Source/Runtime/Engine/Private/AI/Navigation/RecastNavMesh.cpp b/Engine/Source/Runtime/Engine/Private/AI/Navigation/RecastNavMesh.cpp
--- a/Engine/Source/Runtime/Engine/Private/AI/Navigation/RecastNavMesh.cpp
+++ b/Engine/Source/Runtime/Engine/Private/AI/Navigation/RecastNavMesh.cpp
@@ -47,7 +47,7 @@ FNavMeshTileData::FNavData::~FNavData()
 float ARecastNavMesh::DrawDistanceSq = 0.0f;
 #if !WITH_RECAST
 
-ARecastNavMesh::ARecastNavMesh(const class FPostConstructInitializeProperties& PCIP) : Super(PCIP)
+ARecastNavMesh::ARecastNavMesh(const FObjectInitializer& ObjectInitializer) : Super(ObjectInitializer)
 {
 }
 
@@ -158,8 +158,8 @@ namespace FNavMeshConfig
 
 ARecastNavMesh::FNavPolyFlags ARecastNavMesh::NavLinkFlag = ARecastNavMesh::FNavPolyFlags(0);
 
-ARecastNavMesh::ARecastNavMesh(const class FPostConstructInitializeProperties& PCIP)
-	: Super(PCIP)
+ARecastNavMesh::ARecastNavMesh(const FObjectInitializer& ObjectInitializer)
+	: Super(ObjectInitializer)
 	, bDrawNavMeshEdges(true)
 	, bDrawNavLinks(true)
 	, bDistinctlyDrawTilesBeingBuilt(true)
diff --git a/Engine/Source/Runtime/Engine/Private/Actor.cpp b/Engine/Source/Runtime/Engine/Private/Actor.cpp
--- a/Engine/Source/Runtime/Engine/Private/Actor.cpp
+++ b/Engine/Source/Runtime/Engine/Private/Actor.cpp
@@ -37,8 +37,8 @@ FMakeNoiseDelegate AActor::MakeNoiseDelegate = FMakeNoiseDelegate::CreateStatic(
 FOnProcessEvent AActor::ProcessEventDelegate;
 #endif
 
-AActor::AActor(const class FPostConstructInitializeProperties& PCIP)
-	: Super(PCIP)
+AActor::AActor(const FObjectInitializer& ObjectInitializer)
+	: Super(ObjectInitializer)
 {
 	PrimaryActorTick.TickGroup = TG_PrePhysics;
 	// Default to no tick function, but if we set 'never ticks' to false (so there is a tick function) it is enabled by default
diff --git a/Engine/Source/Runtime/Engine/Private/ActorComponent.cpp b/Engine/Source/Runtime/Engine/Private/ActorComponent.cpp
--- a/Engine/Source/Runtime/Engine/Private/ActorComponent.cpp
+++ b/Engine/Source/Runtime/Engine/Private/ActorComponent.cpp
@@ -96,8 +96,8 @@ FGlobalComponentReregisterContext::~FGlobalComponentReregisterContext()
 }
 
 
-UActorComponent::UActorComponent(const class FPostConstructInitializeProperties& PCIP)
-	: Super(PCIP)
+UActorComponent::UActorComponent(const FObjectInitializer& ObjectInitializer)
+	: Super(ObjectInitializer)
 {
 	PrimaryComponentTick.TickGroup = TG_DuringPhysics;
 	PrimaryComponentTick.bStartWithTickEnabled = true;
diff --git a/Engine/Source/Runtime/Engine/Private/AmbientSound.cpp b/Engine/Source/Runtime/Engine/Private/AmbientSound.cpp
--- a/Engine/Source/Runtime/Engine/Private/AmbientSound.cpp
+++ b/Engine/Source/Runtime/Engine/Private/AmbientSound.cpp
@@ -27,10 +27,10 @@
 -----------------------------------------------------------------------------*/
 bool AAmbientSound::bUE4AudioRefactorMigrationUnderway = false;
 
-AAmbientSound::AAmbientSound(const class FPostConstructInitializeProperties& PCIP)
-	: Super(PCIP)
+AAmbientSound::AAmbientSound(const FObjectInitializer& ObjectInitializer)
+	: Super(ObjectInitializer)
 {
-	AudioComponent = PCIP.CreateDefaultSubobject<UAudioComponent>(this, TEXT("AudioComponent0"));
+	AudioComponent = ObjectInitializer.CreateDefaultSubobject<UAudioComponent>(this, TEXT("AudioComponent0"));
 
 	AudioComponent->bAutoActivate = true;
 	AudioComponent->bStopWhenOwnerDestroyed = true;
diff --git a/Engine/Source/Runtime/Engine/Private/AnimBlueprint.cpp b/Engine/Source/Runtime/Engine/Private/AnimBlueprint.cpp
--- a/Engine/Source/Runtime/Engine/Private/AnimBlueprint.cpp
+++ b/Engine/Source/Runtime/Engine/Private/AnimBlueprint.cpp
@@ -9,8 +9,8 @@
 //////////////////////////////////////////////////////////////////////////
 // UAnimBlueprint
 
-UAnimBlueprint::UAnimBlueprint(const class FPostConstructInitializeProperties& PCIP)
-	: Super(PCIP)
+UAnimBlueprint::UAnimBlueprint(const FObjectInitializer& ObjectInitializer)
+	: Super(ObjectInitializer)
 {
 }
 
diff --git a/Engine/Source/Runtime/Engine/Private/Animation/AimOffsetBlendSpace.cpp b/Engine/Source/Runtime/Engine/Private/Animation/AimOffsetBlendSpace.cpp
--- a/Engine/Source/Runtime/Engine/Private/Animation/AimOffsetBlendSpace.cpp
+++ b/Engine/Source/Runtime/Engine/Private/Animation/AimOffsetBlendSpace.cpp
@@ -9,8 +9,8 @@
 #include "Animation/AimOffsetBlendSpace.h"
 #include "Animation/AnimSequence.h"
 
-UAimOffsetBlendSpace::UAimOffsetBlendSpace(const class FPostConstructInitializeProperties& PCIP)
-	: Super(PCIP)
+UAimOffsetBlendSpace::UAimOffsetBlendSpace(const FObjectInitializer& ObjectInitializer)
+	: Super(ObjectInitializer)
 {
 }
 
diff --git a/Engine/Source/Runtime/Engine/Private/Animation/AimOffsetBlendSpace1D.cpp b/Engine/Source/Runtime/Engine/Private/Animation/AimOffsetBlendSpace1D.cpp
--- a/Engine/Source/Runtime/Engine/Private/Animation/AimOffsetBlendSpace1D.cpp
+++ b/Engine/Source/Runtime/Engine/Private/Animation/AimOffsetBlendSpace1D.cpp
@@ -9,8 +9,8 @@
 #include "Animation/AimOffsetBlendSpace1D.h"
 #include "Animation/AnimSequence.h"
 
-UAimOffsetBlendSpace1D::UAimOffsetBlendSpace1D(const class FPostConstructInitializeProperties& PCIP)
-	: Super(PCIP)
+UAimOffsetBlendSpace1D::UAimOffsetBlendSpace1D(const FObjectInitializer& ObjectInitializer)
+	: Super(ObjectInitializer)
 {
 }
 
diff --git a/Engine/Source/Runtime/Engine/Private/Animation/AnimBlueprintGeneratedClass.cpp b/Engine/Source/Runtime/Engine/Private/Animation/AnimBlueprintGeneratedClass.cpp
--- a/Engine/Source/Runtime/Engine/Private/Animation/AnimBlueprintGeneratedClass.cpp
+++ b/Engine/Source/Runtime/Engine/Private/Animation/AnimBlueprintGeneratedClass.cpp
@@ -177,8 +177,8 @@ void FAnimationFrameSnapshot::CopyToInstance(UAnimInstance* Instance)
 /////////////////////////////////////////////////////
 // UAnimBlueprintGeneratedClass
 
-UAnimBlueprintGeneratedClass::UAnimBlueprintGeneratedClass(const class FPostConstructInitializeProperties& PCIP)
-	: Super(PCIP)
+UAnimBlueprintGeneratedClass::UAnimBlueprintGeneratedClass(const FObjectInitializer& ObjectInitializer)
+	: Super(ObjectInitializer)
 {
 	RootAnimNodeIndex = INDEX_NONE;
 }
diff --git a/Engine/Source/Runtime/Engine/Private/Animation/AnimComposite.cpp b/Engine/Source/Runtime/Engine/Private/Animation/AnimComposite.cpp
--- a/Engine/Source/Runtime/Engine/Private/Animation/AnimComposite.cpp
+++ b/Engine/Source/Runtime/Engine/Private/Animation/AnimComposite.cpp
@@ -9,8 +9,8 @@
 #include "AnimationUtils.h"
 #include "AnimationRuntime.h"
 
-UAnimComposite::UAnimComposite(const class FPostConstructInitializeProperties& PCIP)
-	: Super(PCIP)
+UAnimComposite::UAnimComposite(const FObjectInitializer& ObjectInitializer)
+	: Super(ObjectInitializer)
 {
 
 }
diff --git a/Engine/Source/Runtime/Engine/Private/Animation/AnimCompositeBase.cpp b/Engine/Source/Runtime/Engine/Private/Animation/AnimCompositeBase.cpp
--- a/Engine/Source/Runtime/Engine/Private/Animation/AnimCompositeBase.cpp
+++ b/Engine/Source/Runtime/Engine/Private/Animation/AnimCompositeBase.cpp
@@ -421,8 +421,8 @@ void FAnimTrack::SortAnimSegments()
 // UAnimCompositeBase
 ///////////////////////////////////////////////////////
 
-UAnimCompositeBase::UAnimCompositeBase(const class FPostConstructInitializeProperties& PCIP)
-	: Super(PCIP)
+UAnimCompositeBase::UAnimCompositeBase(const FObjectInitializer& ObjectInitializer)
+	: Super(ObjectInitializer)
 {
 }
 
diff --git a/Engine/Source/Runtime/Engine/Private/Animation/AnimCompress.cpp b/Engine/Source/Runtime/Engine/Private/Animation/AnimCompress.cpp
--- a/Engine/Source/Runtime/Engine/Private/Animation/AnimCompress.cpp
+++ b/Engine/Source/Runtime/Engine/Private/Animation/AnimCompress.cpp
@@ -85,8 +85,8 @@ static void PackQuaternionToStream(
 	}
 }
 
-UAnimCompress::UAnimCompress(const class FPostConstructInitializeProperties& PCIP)
-	: Super(PCIP)
+UAnimCompress::UAnimCompress(const FObjectInitializer& ObjectInitializer)
+	: Super(ObjectInitializer)
 {
 	Description = TEXT("None");
 	TranslationCompressionFormat = ACF_None;
diff --git a/Engine/Source/Runtime/Engine/Private/Animation/AnimCompress_Automatic.cpp b/Engine/Source/Runtime/Engine/Private/Animation/AnimCompress_Automatic.cpp
--- a/Engine/Source/Runtime/Engine/Private/Animation/AnimCompress_Automatic.cpp
+++ b/Engine/Source/Runtime/Engine/Private/Animation/AnimCompress_Automatic.cpp
@@ -5,8 +5,8 @@
 #include "AnimationUtils.h"
 #include "Animation/AnimCompress_Automatic.h"
 
-UAnimCompress_Automatic::UAnimCompress_Automatic(const class FPostConstructInitializeProperties& PCIP)
-	: Super(PCIP)
+UAnimCompress_Automatic::UAnimCompress_Automatic(const FObjectInitializer& ObjectInitializer)
+	: Super(ObjectInitializer)
 {
 	Description = TEXT("Automatic");
 	MaxEndEffectorError = 1.0f;
diff --git a/Engine/Source/Runtime/Engine/Private/Animation/AnimCompress_BitwiseCompressOnly.cpp b/Engine/Source/Runtime/Engine/Private/Animation/AnimCompress_BitwiseCompressOnly.cpp
--- a/Engine/Source/Runtime/Engine/Private/Animation/AnimCompress_BitwiseCompressOnly.cpp
+++ b/Engine/Source/Runtime/Engine/Private/Animation/AnimCompress_BitwiseCompressOnly.cpp
@@ -10,8 +10,8 @@
 #include "AnimEncoding.h"
 #include "AnimationCompression.h"
 
-UAnimCompress_BitwiseCompressOnly::UAnimCompress_BitwiseCompressOnly(const class FPostConstructInitializeProperties& PCIP)
-	: Super(PCIP)
+UAnimCompress_BitwiseCompressOnly::UAnimCompress_BitwiseCompressOnly(const FObjectInitializer& ObjectInitializer)
+	: Super(ObjectInitializer)
 {
 	Description = TEXT("Bitwise Compress Only");
 }
diff --git a/Engine/Source/Runtime/Engine/Private/Animation/AnimCompress_LeastDestructive.cpp b/Engine/Source/Runtime/Engine/Private/Animation/AnimCompress_LeastDestructive.cpp
--- a/Engine/Source/Runtime/Engine/Private/Animation/AnimCompress_LeastDestructive.cpp
+++ b/Engine/Source/Runtime/Engine/Private/Animation/AnimCompress_LeastDestructive.cpp
@@ -10,8 +10,8 @@
 #include "AnimationUtils.h"
 #include "AnimEncoding.h"
 
-UAnimCompress_LeastDestructive::UAnimCompress_LeastDestructive(const class FPostConstructInitializeProperties& PCIP)
-	: Super(PCIP)
+UAnimCompress_LeastDestructive::UAnimCompress_LeastDestructive(const FObjectInitializer& ObjectInitializer)
+	: Super(ObjectInitializer)
 {
 	Description = TEXT("Least Destructive");
 	TranslationCompressionFormat = ACF_None;
diff --git a/Engine/Source/Runtime/Engine/Private/Animation/AnimCompress_PerTrackCompression.cpp b/Engine/Source/Runtime/Engine/Private/Animation/AnimCompress_PerTrackCompression.cpp
--- a/Engine/Source/Runtime/Engine/Private/Animation/AnimCompress_PerTrackCompression.cpp
+++ b/Engine/Source/Runtime/Engine/Private/Animation/AnimCompress_PerTrackCompression.cpp
@@ -928,8 +928,8 @@ public:
 	}
 };
 
-UAnimCompress_RemoveLinearKeys::UAnimCompress_RemoveLinearKeys(const class FPostConstructInitializeProperties& PCIP)
-	: Super(PCIP)
+UAnimCompress_RemoveLinearKeys::UAnimCompress_RemoveLinearKeys(const FObjectInitializer& ObjectInitializer)
+	: Super(ObjectInitializer)
 {
 	bNeedsSkeleton = true;
 	Description = TEXT("Remove Linear Keys");
@@ -945,8 +945,8 @@ UAnimCompress_RemoveLinearKeys::UAnimCompress_RemoveLinearKeys(const class FPost
 }
 
 
-UAnimCompress_PerTrackCompression::UAnimCompress_PerTrackCompression(const class FPostConstructInitializeProperties& PCIP)
-	: Super(PCIP)
+UAnimCompress_PerTrackCompression::UAnimCompress_PerTrackCompression(const FObjectInitializer& ObjectInitializer)
+	: Super(ObjectInitializer)
 {
 	Description = TEXT("Compress each track independently");
 	MaxPosDiffBitwise = 0.007f;
diff --git a/Engine/Source/Runtime/Engine/Private/Animation/AnimCompress_RemoveEverySecondKey.cpp b/Engine/Source/Runtime/Engine/Private/Animation/AnimCompress_RemoveEverySecondKey.cpp
--- a/Engine/Source/Runtime/Engine/Private/Animation/AnimCompress_RemoveEverySecondKey.cpp
+++ b/Engine/Source/Runtime/Engine/Private/Animation/AnimCompress_RemoveEverySecondKey.cpp
@@ -10,8 +10,8 @@
 #include "AnimEncoding.h"
 #include "AnimationCompression.h"
 
-UAnimCompress_RemoveEverySecondKey::UAnimCompress_RemoveEverySecondKey(const class FPostConstructInitializeProperties& PCIP)
-	: Super(PCIP)
+UAnimCompress_RemoveEverySecondKey::UAnimCompress_RemoveEverySecondKey(const FObjectInitializer& ObjectInitializer)
+	: Super(ObjectInitializer)
 {
 	Description = TEXT("Remove Every Second Key");
 	MinKeys = 10;
diff --git a/Engine/Source/Runtime/Engine/Private/Animation/AnimCompress_RemoveTrivialKeys.cpp b/Engine/Source/Runtime/Engine/Private/Animation/AnimCompress_RemoveTrivialKeys.cpp
--- a/Engine/Source/Runtime/Engine/Private/Animation/AnimCompress_RemoveTrivialKeys.cpp
+++ b/Engine/Source/Runtime/Engine/Private/Animation/AnimCompress_RemoveTrivialKeys.cpp
@@ -10,8 +10,8 @@
 #include "AnimEncoding.h"
 #include "AnimationCompression.h"
 
-UAnimCompress_RemoveTrivialKeys::UAnimCompress_RemoveTrivialKeys(const class FPostConstructInitializeProperties& PCIP)
-	: Super(PCIP)
+UAnimCompress_RemoveTrivialKeys::UAnimCompress_RemoveTrivialKeys(const FObjectInitializer& ObjectInitializer)
+	: Super(ObjectInitializer)
 {
 	Description = TEXT("Remove Trivial Keys");
 	MaxPosDiff = 0.0001f;
diff --git a/Engine/Source/Runtime/Engine/Private/Animation/AnimCompress_RevertToRaw.cpp b/Engine/Source/Runtime/Engine/Private/Animation/AnimCompress_RevertToRaw.cpp
--- a/Engine/Source/Runtime/Engine/Private/Animation/AnimCompress_RevertToRaw.cpp
+++ b/Engine/Source/Runtime/Engine/Private/Animation/AnimCompress_RevertToRaw.cpp
@@ -10,8 +10,8 @@
 #include "AnimationUtils.h"
 #include "AnimEncoding.h"
 
-UDEPRECATED_AnimCompress_RevertToRaw::UDEPRECATED_AnimCompress_RevertToRaw(const class FPostConstructInitializeProperties& PCIP)
-	: Super(PCIP)
+UDEPRECATED_AnimCompress_RevertToRaw::UDEPRECATED_AnimCompress_RevertToRaw(const FObjectInitializer& ObjectInitializer)
+	: Super(ObjectInitializer)
 {
 	Description = TEXT("Revert To Raw");
 	TranslationCompressionFormat = ACF_None;
diff --git a/Engine/Source/Runtime/Engine/Private/Animation/AnimInstance.cpp b/Engine/Source/Runtime/Engine/Private/Animation/AnimInstance.cpp
--- a/Engine/Source/Runtime/Engine/Private/Animation/AnimInstance.cpp
+++ b/Engine/Source/Runtime/Engine/Private/Animation/AnimInstance.cpp
@@ -61,8 +61,8 @@ DEFINE_LOG_CATEGORY(LogAnimNotify);
 // UAnimInstance
 /////////////////////////////////////////////////////
 
-UAnimInstance::UAnimInstance(const class FPostConstructInitializeProperties& PCIP)
-	: Super(PCIP)
+UAnimInstance::UAnimInstance(const FObjectInitializer& ObjectInitializer)
+	: Super(ObjectInitializer)
 {
 	RootNode = NULL;
 	RootMotionMode = ERootMotionMode::RootMotionFromMontagesOnly;
diff --git a/Engine/Source/Runtime/Engine/Private/Animation/AnimMontage.cpp b/Engine/Source/Runtime/Engine/Private/Animation/AnimMontage.cpp
--- a/Engine/Source/Runtime/Engine/Private/Animation/AnimMontage.cpp
+++ b/Engine/Source/Runtime/Engine/Private/Animation/AnimMontage.cpp
@@ -21,8 +21,8 @@ void FBranchingPoint::RefreshTriggerOffset(EAnimEventTriggerOffsets::Type Predic
 
 ///////////////////////////////////////////////////////////////////////////
 //
-UAnimMontage::UAnimMontage(const class FPostConstructInitializeProperties& PCIP)
-	: Super(PCIP)
+UAnimMontage::UAnimMontage(const FObjectInitializer& ObjectInitializer)
+	: Super(ObjectInitializer)
 {
 	bAnimBranchingPointNeedsSort = true;
 	BlendInTime = 0.25f;
diff --git a/Engine/Source/Runtime/Engine/Private/Animation/AnimNotify.cpp b/Engine/Source/Runtime/Engine/Private/Animation/AnimNotify.cpp
--- a/Engine/Source/Runtime/Engine/Private/Animation/AnimNotify.cpp
+++ b/Engine/Source/Runtime/Engine/Private/Animation/AnimNotify.cpp
@@ -6,8 +6,8 @@
 /////////////////////////////////////////////////////
 // UAnimNotify
 
-UAnimNotify::UAnimNotify(const class FPostConstructInitializeProperties& PCIP)
-	: Super(PCIP)
+UAnimNotify::UAnimNotify(const FObjectInitializer& ObjectInitializer)
+	: Super(ObjectInitializer)
 	, MeshContext(NULL)
 {
 
diff --git a/Engine/Source/Runtime/Engine/Private/Animation/AnimNotifyState.cpp b/Engine/Source/Runtime/Engine/Private/Animation/AnimNotifyState.cpp
--- a/Engine/Source/Runtime/Engine/Private/Animation/AnimNotifyState.cpp
+++ b/Engine/Source/Runtime/Engine/Private/Animation/AnimNotifyState.cpp
@@ -6,8 +6,8 @@
 /////////////////////////////////////////////////////
 // UAnimNotifyState
 
-UAnimNotifyState::UAnimNotifyState(const class FPostConstructInitializeProperties& PCIP)
-	: Super(PCIP)
+UAnimNotifyState::UAnimNotifyState(const FObjectInitializer& ObjectInitializer)
+	: Super(ObjectInitializer)
 {
 
 #if WITH_EDITORONLY_DATA
diff --git a/Engine/Source/Runtime/Engine/Private/Animation/AnimNotifyState_TimedParticleEffect.cpp b/Engine/Source/Runtime/Engine/Private/Animation/AnimNotifyState_TimedParticleEffect.cpp
--- a/Engine/Source/Runtime/Engine/Private/Animation/AnimNotifyState_TimedParticleEffect.cpp
+++ b/Engine/Source/Runtime/Engine/Private/Animation/AnimNotifyState_TimedParticleEffect.cpp
@@ -4,8 +4,8 @@
 #include "Animation/AnimNotifies/AnimNotifyState_TimedParticleEffect.h"
 #include "Particles/ParticleSystemComponent.h"
 
-UAnimNotifyState_TimedParticleEffect::UAnimNotifyState_TimedParticleEffect(const class FPostConstructInitializeProperties& PCIP)
-: Super(PCIP)
+UAnimNotifyState_TimedParticleEffect::UAnimNotifyState_TimedParticleEffect(const FObjectInitializer& ObjectInitializer)
+: Super(ObjectInitializer)
 {
 	PSTemplate = nullptr;
 	LocationOffset.Set(0.0f, 0.0f, 0.0f);
diff --git a/Engine/Source/Runtime/Engine/Private/Animation/AnimNotifyState_Trail.cpp b/Engine/Source/Runtime/Engine/Private/Animation/AnimNotifyState_Trail.cpp
--- a/Engine/Source/Runtime/Engine/Private/Animation/AnimNotifyState_Trail.cpp
+++ b/Engine/Source/Runtime/Engine/Private/Animation/AnimNotifyState_Trail.cpp
@@ -17,8 +17,8 @@ DEFINE_LOG_CATEGORY(LogAnimTrails);
 /////////////////////////////////////////////////////
 // UAnimNotifyState_Trail
 
-UAnimNotifyState_Trail::UAnimNotifyState_Trail(const class FPostConstructInitializeProperties& PCIP)
-	: Super(PCIP)
+UAnimNotifyState_Trail::UAnimNotifyState_Trail(const FObjectInitializer& ObjectInitializer)
+	: Super(ObjectInitializer)
 {
 	PSTemplate = NULL;
 	FirstSocketName = NAME_None;
diff --git a/Engine/Source/Runtime/Engine/Private/Animation/AnimSequence.cpp b/Engine/Source/Runtime/Engine/Private/Animation/AnimSequence.cpp
--- a/Engine/Source/Runtime/Engine/Private/Animation/AnimSequence.cpp
+++ b/Engine/Source/Runtime/Engine/Private/Animation/AnimSequence.cpp
@@ -96,8 +96,8 @@ bool FCurveTrack::CompressCurveWeights()
 /////////////////////////////////////////////////////
 // UAnimSequence
 
-UAnimSequence::UAnimSequence(const class FPostConstructInitializeProperties& PCIP)
-	: Super(PCIP)
+UAnimSequence::UAnimSequence(const FObjectInitializer& ObjectInitializer)
+	: Super(ObjectInitializer)
 	, bEnableRootMotion(false)
 	, RootMotionRootLock(ERootMotionRootLock::RefPose)
 	, bRootMotionSettingsCopiedFromMontage(false)
diff --git a/Engine/Source/Runtime/Engine/Private/Animation/AnimSequenceBase.cpp b/Engine/Source/Runtime/Engine/Private/Animation/AnimSequenceBase.cpp
--- a/Engine/Source/Runtime/Engine/Private/Animation/AnimSequenceBase.cpp
+++ b/Engine/Source/Runtime/Engine/Private/Animation/AnimSequenceBase.cpp
@@ -192,8 +192,8 @@ ENGINE_API bool FRawCurveTracks::DuplicateCurveData(USkeleton::AnimCurveUID ToCo
 
 /////////////////////////////////////////////////////
 
-UAnimSequenceBase::UAnimSequenceBase(const class FPostConstructInitializeProperties& PCIP)
-	: Super(PCIP)
+UAnimSequenceBase::UAnimSequenceBase(const FObjectInitializer& ObjectInitializer)
+	: Super(ObjectInitializer)
 	, RateScale(1.0f)
 {
 }
diff --git a/Engine/Source/Runtime/Engine/Private/Animation/AnimSet.cpp b/Engine/Source/Runtime/Engine/Private/Animation/AnimSet.cpp
--- a/Engine/Source/Runtime/Engine/Private/Animation/AnimSet.cpp
+++ b/Engine/Source/Runtime/Engine/Private/Animation/AnimSet.cpp
@@ -8,8 +8,8 @@
 /////////////////////////////////////////////////////
 // UAnimSet
 
-UAnimSet::UAnimSet(const class FPostConstructInitializeProperties& PCIP)
-	: Super(PCIP)
+UAnimSet::UAnimSet(const FObjectInitializer& ObjectInitializer)
+	: Super(ObjectInitializer)
 {
 
 }
diff --git a/Engine/Source/Runtime/Engine/Private/Animation/AnimSingleNodeInstance.cpp b/Engine/Source/Runtime/Engine/Private/Animation/AnimSingleNodeInstance.cpp
--- a/Engine/Source/Runtime/Engine/Private/Animation/AnimSingleNodeInstance.cpp
+++ b/Engine/Source/Runtime/Engine/Private/Animation/AnimSingleNodeInstance.cpp
@@ -21,8 +21,8 @@
 // UAnimSingleNodeInstance
 /////////////////////////////////////////////////////
 
-UAnimSingleNodeInstance::UAnimSingleNodeInstance(const class FPostConstructInitializeProperties& PCIP)
-	: Super(PCIP)
+UAnimSingleNodeInstance::UAnimSingleNodeInstance(const FObjectInitializer& ObjectInitializer)
+	: Super(ObjectInitializer)
 	, PlayRate(1.f)
 	, bLooping(true)
 	, bPlaying(true)
diff --git a/Engine/Source/Runtime/Engine/Private/Animation/AnimStateMachineTypes.cpp b/Engine/Source/Runtime/Engine/Private/Animation/AnimStateMachineTypes.cpp
--- a/Engine/Source/Runtime/Engine/Private/Animation/AnimStateMachineTypes.cpp
+++ b/Engine/Source/Runtime/Engine/Private/Animation/AnimStateMachineTypes.cpp
@@ -8,7 +8,7 @@
 /////////////////////////////////////////////////////
 // UAnimStateMachineTypes
 
-UAnimStateMachineTypes::UAnimStateMachineTypes(const FPostConstructInitializeProperties& PCIP)
-	: Super(PCIP)
+UAnimStateMachineTypes::UAnimStateMachineTypes(const FObjectInitializer& ObjectInitializer)
+	: Super(ObjectInitializer)
 {
 }
diff --git a/Engine/Source/Runtime/Engine/Private/Animation/BlendSpace.cpp b/Engine/Source/Runtime/Engine/Private/Animation/BlendSpace.cpp
--- a/Engine/Source/Runtime/Engine/Private/Animation/BlendSpace.cpp
+++ b/Engine/Source/Runtime/Engine/Private/Animation/BlendSpace.cpp
@@ -9,8 +9,8 @@
 #include "Animation/BlendSpace.h"
 #include "AnimationUtils.h"
 
-UBlendSpace::UBlendSpace(const class FPostConstructInitializeProperties& PCIP)
-	: Super(PCIP)
+UBlendSpace::UBlendSpace(const FObjectInitializer& ObjectInitializer)
+	: Super(ObjectInitializer)
 {
 	NumOfDimension = 2;
 }
diff --git a/Engine/Source/Runtime/Engine/Private/Animation/BlendSpace1D.cpp b/Engine/Source/Runtime/Engine/Private/Animation/BlendSpace1D.cpp
--- a/Engine/Source/Runtime/Engine/Private/Animation/BlendSpace1D.cpp
+++ b/Engine/Source/Runtime/Engine/Private/Animation/BlendSpace1D.cpp
@@ -9,8 +9,8 @@
 #include "AnimationUtils.h"
 #include "Animation/BlendSpace1D.h"
 
-UBlendSpace1D::UBlendSpace1D(const class FPostConstructInitializeProperties& PCIP)
-	: Super(PCIP)
+UBlendSpace1D::UBlendSpace1D(const FObjectInitializer& ObjectInitializer)
+	: Super(ObjectInitializer)
 {
 	NumOfDimension = 1;
 }
diff --git a/Engine/Source/Runtime/Engine/Private/Animation/BlendSpaceBase.cpp b/Engine/Source/Runtime/Engine/Private/Animation/BlendSpaceBase.cpp
--- a/Engine/Source/Runtime/Engine/Private/Animation/BlendSpaceBase.cpp
+++ b/Engine/Source/Runtime/Engine/Private/Animation/BlendSpaceBase.cpp
@@ -13,8 +13,8 @@
 // I can't make this member since I  need to change this in const function
 bool bNeedReinitializeFilter = false;
 
-UBlendSpaceBase::UBlendSpaceBase(const class FPostConstructInitializeProperties& PCIP)
-	: Super(PCIP)
+UBlendSpaceBase::UBlendSpaceBase(const FObjectInitializer& ObjectInitializer)
+	: Super(ObjectInitializer)
 {
 	
 }
diff --git a/Engine/Source/Runtime/Engine/Private/Animation/Rig.cpp b/Engine/Source/Runtime/Engine/Private/Animation/Rig.cpp
--- a/Engine/Source/Runtime/Engine/Private/Animation/Rig.cpp
+++ b/Engine/Source/Runtime/Engine/Private/Animation/Rig.cpp
@@ -13,8 +13,8 @@
 
 FName URig::WorldNodeName(TEXT("World"));
 
-URig::URig(const class FPostConstructInitializeProperties& PCIP)
-	: Super(PCIP)
+URig::URig(const FObjectInitializer& ObjectInitializer)
+	: Super(ObjectInitializer)
 {
 }
 
diff --git a/Engine/Source/Runtime/Engine/Private/Animation/SkeletalControl.cpp b/Engine/Source/Runtime/Engine/Private/Animation/SkeletalControl.cpp
--- a/Engine/Source/Runtime/Engine/Private/Animation/SkeletalControl.cpp
+++ b/Engine/Source/Runtime/Engine/Private/Animation/SkeletalControl.cpp
@@ -14,8 +14,8 @@ DEFINE_LOG_CATEGORY(LogSkeletalControl);
 /////////////////////////////////////////////////////
 // UBoneMaskFilter 
 
-UBoneMaskFilter::UBoneMaskFilter(const class FPostConstructInitializeProperties& PCIP)
-	: Super(PCIP)
+UBoneMaskFilter::UBoneMaskFilter(const FObjectInitializer& ObjectInitializer)
+	: Super(ObjectInitializer)
 {
 }
 
diff --git a/Engine/Source/Runtime/Engine/Private/Animation/Skeleton.cpp b/Engine/Source/Runtime/Engine/Private/Animation/Skeleton.cpp
--- a/Engine/Source/Runtime/Engine/Private/Animation/Skeleton.cpp
+++ b/Engine/Source/Runtime/Engine/Private/Animation/Skeleton.cpp
@@ -43,8 +43,8 @@ FArchive& operator<<(FArchive& Ar, FReferencePose & P)
 	return Ar;
 }
 
-USkeleton::USkeleton(const class FPostConstructInitializeProperties& PCIP)
-	: Super(PCIP)
+USkeleton::USkeleton(const FObjectInitializer& ObjectInitializer)
+	: Super(ObjectInitializer)
 {
 	// Make sure we have somewhere for curve names.
 	SmartNames.AddContainer(AnimCurveMappingName);
diff --git a/Engine/Source/Runtime/Engine/Private/Animation/VertexAnimation.cpp b/Engine/Source/Runtime/Engine/Private/Animation/VertexAnimation.cpp
--- a/Engine/Source/Runtime/Engine/Private/Animation/VertexAnimation.cpp
+++ b/Engine/Source/Runtime/Engine/Private/Animation/VertexAnimation.cpp
@@ -4,8 +4,8 @@
 #include "EnginePrivate.h"
 #include "Animation/VertexAnim/VertexAnimation.h"
 
-UVertexAnimation::UVertexAnimation(const class FPostConstructInitializeProperties& PCIP)
-	: Super(PCIP)
+UVertexAnimation::UVertexAnimation(const FObjectInitializer& ObjectInitializer)
+	: Super(ObjectInitializer)
 {
 }
 
diff --git a/Engine/Source/Runtime/Engine/Private/AnimationAsset.cpp b/Engine/Source/Runtime/Engine/Private/AnimationAsset.cpp
--- a/Engine/Source/Runtime/Engine/Private/AnimationAsset.cpp
+++ b/Engine/Source/Runtime/Engine/Private/AnimationAsset.cpp
@@ -3,8 +3,8 @@
 #include "EnginePrivate.h"
 #include "Animation/AnimSequence.h"
 
-UAnimationAsset::UAnimationAsset(const class FPostConstructInitializeProperties& PCIP)
-	: Super(PCIP)
+UAnimationAsset::UAnimationAsset(const FObjectInitializer& ObjectInitializer)
+	: Super(ObjectInitializer)
 {
 }
 
diff --git a/Engine/Source/Runtime/Engine/Private/ApplicationLifecycleComponent.cpp b/Engine/Source/Runtime/Engine/Private/ApplicationLifecycleComponent.cpp
--- a/Engine/Source/Runtime/Engine/Private/ApplicationLifecycleComponent.cpp
+++ b/Engine/Source/Runtime/Engine/Private/ApplicationLifecycleComponent.cpp
@@ -5,8 +5,8 @@
 #include "Components/ApplicationLifecycleComponent.h"
 #include "CallbackDevice.h"
 
-UApplicationLifecycleComponent::UApplicationLifecycleComponent(const class FPostConstructInitializeProperties& PCIP)
-	: Super(PCIP)
+UApplicationLifecycleComponent::UApplicationLifecycleComponent(const FObjectInitializer& ObjectInitializer)
+	: Super(ObjectInitializer)
 {
 }
 
diff --git a/Engine/Source/Runtime/Engine/Private/ArrowComponent.cpp b/Engine/Source/Runtime/Engine/Private/ArrowComponent.cpp
--- a/Engine/Source/Runtime/Engine/Private/ArrowComponent.cpp
+++ b/Engine/Source/Runtime/Engine/Private/ArrowComponent.cpp
@@ -345,8 +345,8 @@ private:
 #endif // #if WITH_EDITORONLY_DATA
 };
 
-UArrowComponent::UArrowComponent(const class FPostConstructInitializeProperties& PCIP)
-	: Super(PCIP)
+UArrowComponent::UArrowComponent(const FObjectInitializer& ObjectInitializer)
+	: Super(ObjectInitializer)
 {
 	// Structure to hold one-time initialization
 	struct FConstructorStatics
diff --git a/Engine/Source/Runtime/Engine/Private/AssetUserData.cpp b/Engine/Source/Runtime/Engine/Private/AssetUserData.cpp
--- a/Engine/Source/Runtime/Engine/Private/AssetUserData.cpp
+++ b/Engine/Source/Runtime/Engine/Private/AssetUserData.cpp
@@ -4,15 +4,15 @@
 
 
 
-UAssetUserData::UAssetUserData(const FPostConstructInitializeProperties& PCIP)
-	: UObject(PCIP)
+UAssetUserData::UAssetUserData(const FObjectInitializer& ObjectInitializer)
+	: UObject(ObjectInitializer)
 {
 
 }
 
 //////////////////////////////////////////////////////////////////////////
 
-UInterface_AssetUserData::UInterface_AssetUserData(const class FPostConstructInitializeProperties& PCIP)
-	: Super(PCIP)
+UInterface_AssetUserData::UInterface_AssetUserData(const class FObjectInitializer& ObjectInitializer)
+	: Super(ObjectInitializer)
 {
 }
\ No newline at end of file
diff --git a/Engine/Source/Runtime/Engine/Private/Atmosphere/Atmosphere.cpp b/Engine/Source/Runtime/Engine/Private/Atmosphere/Atmosphere.cpp
--- a/Engine/Source/Runtime/Engine/Private/Atmosphere/Atmosphere.cpp
+++ b/Engine/Source/Runtime/Engine/Private/Atmosphere/Atmosphere.cpp
@@ -10,14 +10,14 @@
 #include "ObjectEditorUtils.h"
 #endif
 
-AAtmosphericFog::AAtmosphericFog(const class FPostConstructInitializeProperties& PCIP)
-	: Super(PCIP)
+AAtmosphericFog::AAtmosphericFog(const FObjectInitializer& ObjectInitializer)
+	: Super(ObjectInitializer)
 {
-	AtmosphericFogComponent = PCIP.CreateDefaultSubobject<UAtmosphericFogComponent>(this, TEXT("AtmosphericFogComponent0"));
+	AtmosphericFogComponent = ObjectInitializer.CreateDefaultSubobject<UAtmosphericFogComponent>(this, TEXT("AtmosphericFogComponent0"));
 	RootComponent = AtmosphericFogComponent;
 
 #if WITH_EDITORONLY_DATA
-	ArrowComponent = PCIP.CreateEditorOnlyDefaultSubobject<UArrowComponent>(this, TEXT("ArrowComponent0"));
+	ArrowComponent = ObjectInitializer.CreateEditorOnlyDefaultSubobject<UArrowComponent>(this, TEXT("ArrowComponent0"));
 
 	if (!IsRunningCommandlet())
 	{
@@ -80,8 +80,8 @@ FAtmospherePrecomputeParameters::FAtmospherePrecomputeParameters():
 {
 }
 
-UAtmosphericFogComponent::UAtmosphericFogComponent(const class FPostConstructInitializeProperties& PCIP)
-	: Super(PCIP), TransmittanceResource(NULL), IrradianceResource(NULL), InscatterResource(NULL)
+UAtmosphericFogComponent::UAtmosphericFogComponent(const FObjectInitializer& ObjectInitializer)
+	: Super(ObjectInitializer), TransmittanceResource(NULL), IrradianceResource(NULL), InscatterResource(NULL)
 #if WITH_EDITORONLY_DATA
 	, PrecomputeDataHandler(NULL)
 #endif
diff --git a/Engine/Source/Runtime/Engine/Private/AudioComponent.cpp b/Engine/Source/Runtime/Engine/Private/AudioComponent.cpp
--- a/Engine/Source/Runtime/Engine/Private/AudioComponent.cpp
+++ b/Engine/Source/Runtime/Engine/Private/AudioComponent.cpp
@@ -9,8 +9,8 @@
 /*-----------------------------------------------------------------------------
 UAudioComponent implementation.
 -----------------------------------------------------------------------------*/
-UAudioComponent::UAudioComponent(const class FPostConstructInitializeProperties& PCIP)
-	: Super(PCIP)
+UAudioComponent::UAudioComponent(const FObjectInitializer& ObjectInitializer)
+	: Super(ObjectInitializer)
 {
 	bAutoDestroy = false;
 	bAutoPlay_DEPRECATED = false;
diff --git a/Engine/Source/Runtime/Engine/Private/AudioSettings.cpp b/Engine/Source/Runtime/Engine/Private/AudioSettings.cpp
--- a/Engine/Source/Runtime/Engine/Private/AudioSettings.cpp
+++ b/Engine/Source/Runtime/Engine/Private/AudioSettings.cpp
@@ -8,7 +8,7 @@
 #include "Sound/AudioSettings.h"
 
 
-UAudioSettings::UAudioSettings(const class FPostConstructInitializeProperties& PCIP)
-: Super(PCIP)
+UAudioSettings::UAudioSettings(const FObjectInitializer& ObjectInitializer)
+: Super(ObjectInitializer)
 {
 }
\ No newline at end of file
diff --git a/Engine/Source/Runtime/Engine/Private/AudioVolume.cpp b/Engine/Source/Runtime/Engine/Private/AudioVolume.cpp
--- a/Engine/Source/Runtime/Engine/Private/AudioVolume.cpp
+++ b/Engine/Source/Runtime/Engine/Private/AudioVolume.cpp
@@ -119,8 +119,8 @@ void FReverbSettings::PostSerialize(const FArchive& Ar)
 	}
 }
 
-AAudioVolume::AAudioVolume(const class FPostConstructInitializeProperties& PCIP)
-	: Super(PCIP)
+AAudioVolume::AAudioVolume(const FObjectInitializer& ObjectInitializer)
+	: Super(ObjectInitializer)
 {
 	BrushComponent->BodyInstance.bEnableCollision_DEPRECATED = false;
 
diff --git a/Engine/Source/Runtime/Engine/Private/BillboardComponent.cpp b/Engine/Source/Runtime/Engine/Private/BillboardComponent.cpp
--- a/Engine/Source/Runtime/Engine/Private/BillboardComponent.cpp
+++ b/Engine/Source/Runtime/Engine/Private/BillboardComponent.cpp
@@ -295,8 +295,8 @@ private:
 #endif // #if WITH_EDITORONLY_DATA
 };
 
-UBillboardComponent::UBillboardComponent(const class FPostConstructInitializeProperties& PCIP)
-	: Super(PCIP)
+UBillboardComponent::UBillboardComponent(const FObjectInitializer& ObjectInitializer)
+	: Super(ObjectInitializer)
 {
 	// Structure to hold one-time initialization
 	struct FConstructorStatics
diff --git a/Engine/Source/Runtime/Engine/Private/BlockingVolume.cpp b/Engine/Source/Runtime/Engine/Private/BlockingVolume.cpp
--- a/Engine/Source/Runtime/Engine/Private/BlockingVolume.cpp
+++ b/Engine/Source/Runtime/Engine/Private/BlockingVolume.cpp
@@ -5,8 +5,8 @@
 
 static FName InvisibleWall_NAME(TEXT("InvisibleWall"));
 
-ABlockingVolume::ABlockingVolume(const class FPostConstructInitializeProperties& PCIP)
-	: Super(PCIP)
+ABlockingVolume::ABlockingVolume(const FObjectInitializer& ObjectInitializer)
+	: Super(ObjectInitializer)
 {
 	BrushComponent->bCanEverAffectNavigation = true;
 	BrushComponent->BodyInstance.bEnableCollision_DEPRECATED = true;
diff --git a/Engine/Source/Runtime/Engine/Private/Blueprint.cpp b/Engine/Source/Runtime/Engine/Private/Blueprint.cpp
--- a/Engine/Source/Runtime/Engine/Private/Blueprint.cpp
+++ b/Engine/Source/Runtime/Engine/Private/Blueprint.cpp
@@ -174,8 +174,8 @@ void FBPVariableDescription::RemoveMetaData(const FName& Key)
 //////////////////////////////////////////////////////////////////////////
 // UBlueprintCore
 
-UBlueprintCore::UBlueprintCore(const class FPostConstructInitializeProperties& PCIP)
-	: Super(PCIP)
+UBlueprintCore::UBlueprintCore(const FObjectInitializer& ObjectInitializer)
+	: Super(ObjectInitializer)
 {
 	bLegacyGeneratedClassIsAuthoritative = false;
 	bLegacyNeedToPurgeSkelRefs = true;
@@ -266,8 +266,8 @@ void UBlueprintCore::GenerateDeterministicGuid()
 	BlueprintGuid.D = HashBuffer[ 4 ];
 }
 
-UBlueprint::UBlueprint(const class FPostConstructInitializeProperties& PCIP)
-	: Super(PCIP)
+UBlueprint::UBlueprint(const FObjectInitializer& ObjectInitializer)
+	: Super(ObjectInitializer)
 #if WITH_EDITOR
 	, bRunConstructionScriptOnDrag(true)
 	, bGenerateConstClass(false)
diff --git a/Engine/Source/Runtime/Engine/Private/BlueprintGeneratedClass.cpp b/Engine/Source/Runtime/Engine/Private/BlueprintGeneratedClass.cpp
--- a/Engine/Source/Runtime/Engine/Private/BlueprintGeneratedClass.cpp
+++ b/Engine/Source/Runtime/Engine/Private/BlueprintGeneratedClass.cpp
@@ -10,8 +10,8 @@
 #endif //WITH_EDITOR
 
 
-UBlueprintGeneratedClass::UBlueprintGeneratedClass(const class FPostConstructInitializeProperties& PCIP)
-	: Super(PCIP)
+UBlueprintGeneratedClass::UBlueprintGeneratedClass(const FObjectInitializer& ObjectInitializer)
+	: Super(ObjectInitializer)
 {
 	NumReplicatedProperties = 0;
 }
diff --git a/Engine/Source/Runtime/Engine/Private/BookMark.cpp b/Engine/Source/Runtime/Engine/Private/BookMark.cpp
--- a/Engine/Source/Runtime/Engine/Private/BookMark.cpp
+++ b/Engine/Source/Runtime/Engine/Private/BookMark.cpp
@@ -4,12 +4,12 @@
 #include "Engine/BookMark.h"
 #include "Engine/BookMark2D.h"
 
-UBookMark::UBookMark(const class FPostConstructInitializeProperties& PCIP)
-	: Super(PCIP)
+UBookMark::UBookMark(const FObjectInitializer& ObjectInitializer)
+	: Super(ObjectInitializer)
 {
 }
 
-UBookMark2D::UBookMark2D(const class FPostConstructInitializeProperties& PCIP)
-	: Super(PCIP)
+UBookMark2D::UBookMark2D(const FObjectInitializer& ObjectInitializer)
+	: Super(ObjectInitializer)
 {
 }
diff --git a/Engine/Source/Runtime/Engine/Private/BoxComponent.cpp b/Engine/Source/Runtime/Engine/Private/BoxComponent.cpp
--- a/Engine/Source/Runtime/Engine/Private/BoxComponent.cpp
+++ b/Engine/Source/Runtime/Engine/Private/BoxComponent.cpp
@@ -3,8 +3,8 @@
 
 #include "EnginePrivate.h"
 
-UBoxComponent::UBoxComponent(const class FPostConstructInitializeProperties& PCIP)
-	: Super(PCIP)
+UBoxComponent::UBoxComponent(const FObjectInitializer& ObjectInitializer)
+	: Super(ObjectInitializer)
 {
 	BoxExtent = FVector(32.0f, 32.0f, 32.0f);
 
diff --git a/Engine/Source/Runtime/Engine/Private/Breakpoint.cpp b/Engine/Source/Runtime/Engine/Private/Breakpoint.cpp
--- a/Engine/Source/Runtime/Engine/Private/Breakpoint.cpp
+++ b/Engine/Source/Runtime/Engine/Private/Breakpoint.cpp
@@ -5,8 +5,8 @@
 #include "Engine/Breakpoint.h"
 #include "BlueprintUtilities.h"
 
-UBreakpoint::UBreakpoint(const class FPostConstructInitializeProperties& PCIP)
-	: Super(PCIP)
+UBreakpoint::UBreakpoint(const FObjectInitializer& ObjectInitializer)
+	: Super(ObjectInitializer)
 {
 	bEnabled = false;
 	bStepOnce = false;
diff --git a/Engine/Source/Runtime/Engine/Private/Brush.cpp b/Engine/Source/Runtime/Engine/Private/Brush.cpp
--- a/Engine/Source/Runtime/Engine/Private/Brush.cpp
+++ b/Engine/Source/Runtime/Engine/Private/Brush.cpp
@@ -18,10 +18,10 @@ ABrush::FOnBrushRegistered ABrush::OnBrushRegistered;
 TArray< TWeakObjectPtr< ULevel > > ABrush::LevelsToRebuild;
 #endif
 
-ABrush::ABrush(const class FPostConstructInitializeProperties& PCIP)
-	: Super(PCIP)
+ABrush::ABrush(const FObjectInitializer& ObjectInitializer)
+	: Super(ObjectInitializer)
 {
-	BrushComponent = PCIP.CreateDefaultSubobject<UBrushComponent>(this, TEXT("BrushComponent0"));
+	BrushComponent = ObjectInitializer.CreateDefaultSubobject<UBrushComponent>(this, TEXT("BrushComponent0"));
 	BrushComponent->Mobility = EComponentMobility::Static;
 	BrushComponent->bGenerateOverlapEvents = false;
 	BrushComponent->bCanEverAffectNavigation = false;
diff --git a/Engine/Source/Runtime/Engine/Private/BrushBuilder.cpp b/Engine/Source/Runtime/Engine/Private/BrushBuilder.cpp
--- a/Engine/Source/Runtime/Engine/Private/BrushBuilder.cpp
+++ b/Engine/Source/Runtime/Engine/Private/BrushBuilder.cpp
@@ -10,8 +10,8 @@
 /*-----------------------------------------------------------------------------
 	UBrushBuilder.
 -----------------------------------------------------------------------------*/
-UBrushBuilder::UBrushBuilder(const class FPostConstructInitializeProperties& PCIP)
-	: Super(PCIP)
+UBrushBuilder::UBrushBuilder(const FObjectInitializer& ObjectInitializer)
+	: Super(ObjectInitializer)
 {
 	BitmapFilename = TEXT("BBGeneric");
 	ToolTip = TEXT("BrushBuilderName_Generic");
diff --git a/Engine/Source/Runtime/Engine/Private/BrushComponent.cpp b/Engine/Source/Runtime/Engine/Private/BrushComponent.cpp
--- a/Engine/Source/Runtime/Engine/Private/BrushComponent.cpp
+++ b/Engine/Source/Runtime/Engine/Private/BrushComponent.cpp
@@ -538,8 +538,8 @@ private:
 	FCollisionResponseContainer CollisionResponse;
 };
 
-UBrushComponent::UBrushComponent(const class FPostConstructInitializeProperties& PCIP)
-	: Super(PCIP)
+UBrushComponent::UBrushComponent(const FObjectInitializer& ObjectInitializer)
+	: Super(ObjectInitializer)
 {
 	bHiddenInGame = true;
 	AlwaysLoadOnClient = false;
diff --git a/Engine/Source/Runtime/Engine/Private/BrushShape.cpp b/Engine/Source/Runtime/Engine/Private/BrushShape.cpp
--- a/Engine/Source/Runtime/Engine/Private/BrushShape.cpp
+++ b/Engine/Source/Runtime/Engine/Private/BrushShape.cpp
@@ -3,8 +3,8 @@
 #include "EnginePrivate.h"
 #include "Engine/BrushShape.h"
 
-ABrushShape::ABrushShape(const class FPostConstructInitializeProperties& PCIP)
-	: Super(PCIP)
+ABrushShape::ABrushShape(const FObjectInitializer& ObjectInitializer)
+	: Super(ObjectInitializer)
 {
 
 	BrushComponent->BodyInstance.bEnableCollision_DEPRECATED = false;
diff --git a/Engine/Source/Runtime/Engine/Private/Camera/CameraActor.cpp b/Engine/Source/Runtime/Engine/Private/Camera/CameraActor.cpp
--- a/Engine/Source/Runtime/Engine/Private/Camera/CameraActor.cpp
+++ b/Engine/Source/Runtime/Engine/Private/Camera/CameraActor.cpp
@@ -11,11 +11,11 @@
 //////////////////////////////////////////////////////////////////////////
 // ACameraActor
 
-ACameraActor::ACameraActor(const class FPostConstructInitializeProperties& PCIP)
-	: Super(PCIP)
+ACameraActor::ACameraActor(const FObjectInitializer& ObjectInitializer)
+	: Super(ObjectInitializer)
 {
 	// Setup camera defaults
-	CameraComponent = PCIP.CreateDefaultSubobject<UCameraComponent>(this, TEXT("CameraComponent"));
+	CameraComponent = ObjectInitializer.CreateDefaultSubobject<UCameraComponent>(this, TEXT("CameraComponent"));
 	CameraComponent->FieldOfView = 90.0f;
 	CameraComponent->bConstrainAspectRatio = true;
 	CameraComponent->AspectRatio = 1.777778f;
diff --git a/Engine/Source/Runtime/Engine/Private/Camera/CameraAnim.cpp b/Engine/Source/Runtime/Engine/Private/Camera/CameraAnim.cpp
--- a/Engine/Source/Runtime/Engine/Private/Camera/CameraAnim.cpp
+++ b/Engine/Source/Runtime/Engine/Private/Camera/CameraAnim.cpp
@@ -16,8 +16,8 @@
 //////////////////////////////////////////////////////////////////////////
 // UCameraAnim
 
-UCameraAnim::UCameraAnim(const class FPostConstructInitializeProperties& PCIP)
-	: Super(PCIP)
+UCameraAnim::UCameraAnim(const FObjectInitializer& ObjectInitializer)
+	: Super(ObjectInitializer)
 {
 	AnimLength = 3.0f;
 	BaseFOV = 90.0f;
diff --git a/Engine/Source/Runtime/Engine/Private/Camera/CameraAnimInst.cpp b/Engine/Source/Runtime/Engine/Private/Camera/CameraAnimInst.cpp
--- a/Engine/Source/Runtime/Engine/Private/Camera/CameraAnimInst.cpp
+++ b/Engine/Source/Runtime/Engine/Private/Camera/CameraAnimInst.cpp
@@ -13,8 +13,8 @@
 //////////////////////////////////////////////////////////////////////////
 // UCameraAnimInst
 
-UCameraAnimInst::UCameraAnimInst(const class FPostConstructInitializeProperties& PCIP)
-	: Super(PCIP)
+UCameraAnimInst::UCameraAnimInst(const FObjectInitializer& ObjectInitializer)
+	: Super(ObjectInitializer)
 {
 	bFinished = true;
 	bAutoReleaseWhenFinished = true;
@@ -22,7 +22,7 @@ UCameraAnimInst::UCameraAnimInst(const class FPostConstructInitializeProperties&
 	TransientScaleModifier = 1.0f;
 	PlaySpace = ECameraAnimPlaySpace::CameraLocal;
 
-	InterpGroupInst = PCIP.CreateDefaultSubobject<UInterpGroupInst>(this, TEXT("InterpGroupInst0"));
+	InterpGroupInst = ObjectInitializer.CreateDefaultSubobject<UInterpGroupInst>(this, TEXT("InterpGroupInst0"));
 }
 
 void UCameraAnimInst::AdvanceAnim(float DeltaTime, bool bJump)
diff --git a/Engine/Source/Runtime/Engine/Private/Camera/CameraComponent.cpp b/Engine/Source/Runtime/Engine/Private/Camera/CameraComponent.cpp
--- a/Engine/Source/Runtime/Engine/Private/Camera/CameraComponent.cpp
+++ b/Engine/Source/Runtime/Engine/Private/Camera/CameraComponent.cpp
@@ -13,8 +13,8 @@ static void SetDeprecatedControllerViewRotation(UCameraComponent& Component, boo
 //////////////////////////////////////////////////////////////////////////
 // UCameraComponent
 
-UCameraComponent::UCameraComponent(const class FPostConstructInitializeProperties& PCIP)
-	: Super(PCIP)
+UCameraComponent::UCameraComponent(const FObjectInitializer& ObjectInitializer)
+	: Super(ObjectInitializer)
 {
 #if WITH_EDITORONLY_DATA
 	if (!IsRunningCommandlet())
diff --git a/Engine/Source/Runtime/Engine/Private/Camera/CameraModifier.cpp b/Engine/Source/Runtime/Engine/Private/Camera/CameraModifier.cpp
--- a/Engine/Source/Runtime/Engine/Private/Camera/CameraModifier.cpp
+++ b/Engine/Source/Runtime/Engine/Private/Camera/CameraModifier.cpp
@@ -12,8 +12,8 @@ DEFINE_LOG_CATEGORY_STATIC(LogCamera, Log, All);
 //////////////////////////////////////////////////////////////////////////
 // UCameraModifier
 
-UCameraModifier::UCameraModifier(const class FPostConstructInitializeProperties& PCIP)
-	: Super(PCIP)
+UCameraModifier::UCameraModifier(const FObjectInitializer& ObjectInitializer)
+	: Super(ObjectInitializer)
 {
 	Priority = 127;
 }
diff --git a/Engine/Source/Runtime/Engine/Private/Camera/CameraModifier_CameraShake.cpp b/Engine/Source/Runtime/Engine/Private/Camera/CameraModifier_CameraShake.cpp
--- a/Engine/Source/Runtime/Engine/Private/Camera/CameraModifier_CameraShake.cpp
+++ b/Engine/Source/Runtime/Engine/Private/Camera/CameraModifier_CameraShake.cpp
@@ -5,8 +5,8 @@
 //////////////////////////////////////////////////////////////////////////
 // UCameraModifier_CameraShake
 
-UCameraModifier_CameraShake::UCameraModifier_CameraShake(const class FPostConstructInitializeProperties& PCIP)
-	: Super(PCIP)
+UCameraModifier_CameraShake::UCameraModifier_CameraShake(const FObjectInitializer& ObjectInitializer)
+	: Super(ObjectInitializer)
 {
 	SplitScreenShakeScale = 0.5f;
 }
diff --git a/Engine/Source/Runtime/Engine/Private/Camera/CameraShake.cpp b/Engine/Source/Runtime/Engine/Private/Camera/CameraShake.cpp
--- a/Engine/Source/Runtime/Engine/Private/Camera/CameraShake.cpp
+++ b/Engine/Source/Runtime/Engine/Private/Camera/CameraShake.cpp
@@ -5,8 +5,8 @@
 //////////////////////////////////////////////////////////////////////////
 // UCameraShake
 
-UCameraShake::UCameraShake(const class FPostConstructInitializeProperties& PCIP)
-	: Super(PCIP)
+UCameraShake::UCameraShake(const FObjectInitializer& ObjectInitializer)
+	: Super(ObjectInitializer)
 {
 	AnimPlayRate = 1.0f;
 	AnimScale = 1.0f;
diff --git a/Engine/Source/Runtime/Engine/Private/CameraBlockingVolume.cpp b/Engine/Source/Runtime/Engine/Private/CameraBlockingVolume.cpp
--- a/Engine/Source/Runtime/Engine/Private/CameraBlockingVolume.cpp
+++ b/Engine/Source/Runtime/Engine/Private/CameraBlockingVolume.cpp
@@ -3,8 +3,8 @@
 #include "EnginePrivate.h"
 #include "GameFramework/CameraBlockingVolume.h"
 
-ACameraBlockingVolume::ACameraBlockingVolume(const class FPostConstructInitializeProperties& PCIP)
-	: Super(PCIP)
+ACameraBlockingVolume::ACameraBlockingVolume(const FObjectInitializer& ObjectInitializer)
+	: Super(ObjectInitializer)
 {
 	BrushComponent->SetCollisionResponseToChannel(ECC_Camera, ECR_Block);
 }
\ No newline at end of file
diff --git a/Engine/Source/Runtime/Engine/Private/CanvasRenderTarget2D.cpp b/Engine/Source/Runtime/Engine/Private/CanvasRenderTarget2D.cpp
--- a/Engine/Source/Runtime/Engine/Private/CanvasRenderTarget2D.cpp
+++ b/Engine/Source/Runtime/Engine/Private/CanvasRenderTarget2D.cpp
@@ -3,8 +3,8 @@
 #include "EnginePrivate.h"
 #include "Engine/CanvasRenderTarget2D.h"
 
-UCanvasRenderTarget2D::UCanvasRenderTarget2D( const class FPostConstructInitializeProperties& PCIP )
-	: Super(PCIP)
+UCanvasRenderTarget2D::UCanvasRenderTarget2D( const FObjectInitializer& ObjectInitializer )
+	: Super(ObjectInitializer)
 {
 	bNeedsTwoCopies = false;
 }
diff --git a/Engine/Source/Runtime/Engine/Private/CapsuleComponent.cpp b/Engine/Source/Runtime/Engine/Private/CapsuleComponent.cpp
--- a/Engine/Source/Runtime/Engine/Private/CapsuleComponent.cpp
+++ b/Engine/Source/Runtime/Engine/Private/CapsuleComponent.cpp
@@ -4,8 +4,8 @@
 #include "EnginePrivate.h"
 
 
-UCapsuleComponent::UCapsuleComponent(const class FPostConstructInitializeProperties& PCIP)
-	: Super(PCIP)
+UCapsuleComponent::UCapsuleComponent(const FObjectInitializer& ObjectInitializer)
+	: Super(ObjectInitializer)
 {
 	ShapeColor = FColor(223, 149, 157, 255);
 
diff --git a/Engine/Source/Runtime/Engine/Private/Character.cpp b/Engine/Source/Runtime/Engine/Private/Character.cpp
--- a/Engine/Source/Runtime/Engine/Private/Character.cpp
+++ b/Engine/Source/Runtime/Engine/Private/Character.cpp
@@ -22,8 +22,8 @@ FName ACharacter::MeshComponentName(TEXT("CharacterMesh0"));
 FName ACharacter::CharacterMovementComponentName(TEXT("CharMoveComp"));
 FName ACharacter::CapsuleComponentName(TEXT("CollisionCylinder"));
 
-ACharacter::ACharacter(const class FPostConstructInitializeProperties& PCIP)
-	: Super(PCIP)
+ACharacter::ACharacter(const FObjectInitializer& ObjectInitializer)
+	: Super(ObjectInitializer)
 {
 	// Structure to hold one-time initialization
 	struct FConstructorStatics
@@ -44,7 +44,7 @@ ACharacter::ACharacter(const class FPostConstructInitializeProperties& PCIP)
 	bUseControllerRotationRoll = false;
 	bUseControllerRotationYaw = true;
 
-	CapsuleComponent = PCIP.CreateDefaultSubobject<UCapsuleComponent>(this, ACharacter::CapsuleComponentName);
+	CapsuleComponent = ObjectInitializer.CreateDefaultSubobject<UCapsuleComponent>(this, ACharacter::CapsuleComponentName);
 	CapsuleComponent->InitCapsuleSize(34.0f, 88.0f);
 	CapsuleComponent->BodyInstance.bEnableCollision_DEPRECATED = true;
 
@@ -61,7 +61,7 @@ ACharacter::ACharacter(const class FPostConstructInitializeProperties& PCIP)
 	JumpMaxHoldTime = 0.0f;
 
 #if WITH_EDITORONLY_DATA
-	ArrowComponent = PCIP.CreateEditorOnlyDefaultSubobject<UArrowComponent>(this, TEXT("Arrow"));
+	ArrowComponent = ObjectInitializer.CreateEditorOnlyDefaultSubobject<UArrowComponent>(this, TEXT("Arrow"));
 	if (ArrowComponent)
 	{
 		ArrowComponent->ArrowColor = FColor(150, 200, 255);
@@ -73,14 +73,14 @@ ACharacter::ACharacter(const class FPostConstructInitializeProperties& PCIP)
 	}
 #endif // WITH_EDITORONLY_DATA
 
-	CharacterMovement = PCIP.CreateDefaultSubobject<UCharacterMovementComponent>(this, ACharacter::CharacterMovementComponentName);
+	CharacterMovement = ObjectInitializer.CreateDefaultSubobject<UCharacterMovementComponent>(this, ACharacter::CharacterMovementComponentName);
 	if (CharacterMovement)
 	{
 		CharacterMovement->UpdatedComponent = CapsuleComponent;
 		CrouchedEyeHeight = CharacterMovement->CrouchedHalfHeight * 0.80f;
 	}
 
-	Mesh = PCIP.CreateOptionalDefaultSubobject<USkeletalMeshComponent>(this, ACharacter::MeshComponentName);
+	Mesh = ObjectInitializer.CreateOptionalDefaultSubobject<USkeletalMeshComponent>(this, ACharacter::MeshComponentName);
 	if (Mesh)
 	{
 		Mesh->AlwaysLoadOnClient = true;
diff --git a/Engine/Source/Runtime/Engine/Private/CharacterMovementComponent.cpp b/Engine/Source/Runtime/Engine/Private/CharacterMovementComponent.cpp
--- a/Engine/Source/Runtime/Engine/Private/CharacterMovementComponent.cpp
+++ b/Engine/Source/Runtime/Engine/Private/CharacterMovementComponent.cpp
@@ -151,8 +151,8 @@ FString FCharacterMovementComponentPreClothTickFunction::DiagnosticMessage()
 	return Target->GetFullName() + TEXT("[UCharacterMovementComponent::PreClothTick]");
 }
 
-UCharacterMovementComponent::UCharacterMovementComponent(const class FPostConstructInitializeProperties& PCIP)
-	: Super(PCIP)
+UCharacterMovementComponent::UCharacterMovementComponent(const FObjectInitializer& ObjectInitializer)
+	: Super(ObjectInitializer)
 {
 	PreClothComponentTick.bCanEverTick = true;
 	PreClothComponentTick.bStartWithTickEnabled = false;
diff --git a/Engine/Source/Runtime/Engine/Private/CheatManager.cpp b/Engine/Source/Runtime/Engine/Private/CheatManager.cpp
--- a/Engine/Source/Runtime/Engine/Private/CheatManager.cpp
+++ b/Engine/Source/Runtime/Engine/Private/CheatManager.cpp
@@ -17,8 +17,8 @@ DEFINE_LOG_CATEGORY_STATIC(LogCheatManager, Log, All);
 
 #define LOCTEXT_NAMESPACE "CheatManager"	
 
-UCheatManager::UCheatManager(const class FPostConstructInitializeProperties& PCIP)
-	: Super(PCIP)
+UCheatManager::UCheatManager(const FObjectInitializer& ObjectInitializer)
+	: Super(ObjectInitializer)
 	, bToggleAILogging(false)
 {
 	DebugCameraControllerClass = ADebugCameraController::StaticClass();
diff --git a/Engine/Source/Runtime/Engine/Private/ChildActorComponent.cpp b/Engine/Source/Runtime/Engine/Private/ChildActorComponent.cpp
--- a/Engine/Source/Runtime/Engine/Private/ChildActorComponent.cpp
+++ b/Engine/Source/Runtime/Engine/Private/ChildActorComponent.cpp
@@ -5,8 +5,8 @@
 
 DEFINE_LOG_CATEGORY_STATIC(LogChildActorComponent, Warning, All);
 
-UChildActorComponent::UChildActorComponent(const class FPostConstructInitializeProperties& PCIP)
-	: Super(PCIP)
+UChildActorComponent::UChildActorComponent(const FObjectInitializer& ObjectInitializer)
+	: Super(ObjectInitializer)
 {
 }
 
diff --git a/Engine/Source/Runtime/Engine/Private/ClipPadEntry.cpp b/Engine/Source/Runtime/Engine/Private/ClipPadEntry.cpp
--- a/Engine/Source/Runtime/Engine/Private/ClipPadEntry.cpp
+++ b/Engine/Source/Runtime/Engine/Private/ClipPadEntry.cpp
@@ -3,8 +3,8 @@
 #include "EnginePrivate.h"
 #include "Engine/ClipPadEntry.h"
 
-UDEPRECATED_ClipPadEntry::UDEPRECATED_ClipPadEntry(const class FPostConstructInitializeProperties& PCIP)
-	: Super(PCIP)
+UDEPRECATED_ClipPadEntry::UDEPRECATED_ClipPadEntry(const FObjectInitializer& ObjectInitializer)
+	: Super(ObjectInitializer)
 {
 	Title = TEXT("Untitled");
 }
diff --git a/Engine/Source/Runtime/Engine/Private/Collision/CollisionProfile.cpp b/Engine/Source/Runtime/Engine/Private/Collision/CollisionProfile.cpp
--- a/Engine/Source/Runtime/Engine/Private/Collision/CollisionProfile.cpp
+++ b/Engine/Source/Runtime/Engine/Private/Collision/CollisionProfile.cpp
@@ -38,8 +38,8 @@ ENGINE_API const FName UCollisionProfile::BlockAllDynamic_ProfileName = FName(TE
 ENGINE_API const FName UCollisionProfile::Pawn_ProfileName = FName(TEXT("Pawn"));
 ENGINE_API const FName UCollisionProfile::Vehicle_ProfileName = FName(TEXT("Vehicle"));
 
-UCollisionProfile::UCollisionProfile(const class FPostConstructInitializeProperties& PCIP)
-	: Super(PCIP)
+UCollisionProfile::UCollisionProfile(const FObjectInitializer& ObjectInitializer)
+	: Super(ObjectInitializer)
 {
 }
 
diff --git a/Engine/Source/Runtime/Engine/Private/Commandlets/Commandlet.cpp b/Engine/Source/Runtime/Engine/Private/Commandlets/Commandlet.cpp
--- a/Engine/Source/Runtime/Engine/Private/Commandlets/Commandlet.cpp
+++ b/Engine/Source/Runtime/Engine/Private/Commandlets/Commandlet.cpp
@@ -3,8 +3,8 @@
 #include "EnginePrivate.h"
 #include "Commandlets/Commandlet.h"
 
-UCommandlet::UCommandlet(const class FPostConstructInitializeProperties& PCIP)
-	: Super(PCIP)
+UCommandlet::UCommandlet(const FObjectInitializer& ObjectInitializer)
+	: Super(ObjectInitializer)
 {
 	IsServer = true;
 	IsClient = true;
diff --git a/Engine/Source/Runtime/Engine/Private/Commandlets/SmokeTestCommandlet.cpp b/Engine/Source/Runtime/Engine/Private/Commandlets/SmokeTestCommandlet.cpp
--- a/Engine/Source/Runtime/Engine/Private/Commandlets/SmokeTestCommandlet.cpp
+++ b/Engine/Source/Runtime/Engine/Private/Commandlets/SmokeTestCommandlet.cpp
@@ -8,8 +8,8 @@
 #include "EnginePrivate.h"
 #include "Commandlets/SmokeTestCommandlet.h"
 
-USmokeTestCommandlet::USmokeTestCommandlet(const class FPostConstructInitializeProperties& PCIP)
-	: Super(PCIP)
+USmokeTestCommandlet::USmokeTestCommandlet(const FObjectInitializer& ObjectInitializer)
+	: Super(ObjectInitializer)
 {
 	IsClient = false;
 	IsEditor = false;
diff --git a/Engine/Source/Runtime/Engine/Private/ComponentDelegateBinding.cpp b/Engine/Source/Runtime/Engine/Private/ComponentDelegateBinding.cpp
--- a/Engine/Source/Runtime/Engine/Private/ComponentDelegateBinding.cpp
+++ b/Engine/Source/Runtime/Engine/Private/ComponentDelegateBinding.cpp
@@ -4,8 +4,8 @@
 #include "BlueprintUtilities.h"
 #include "Engine/ComponentDelegateBinding.h"
 
-UComponentDelegateBinding::UComponentDelegateBinding(const class FPostConstructInitializeProperties& PCIP)
-	: Super(PCIP)
+UComponentDelegateBinding::UComponentDelegateBinding(const FObjectInitializer& ObjectInitializer)
+	: Super(ObjectInitializer)
 {
 }
 
diff --git a/Engine/Source/Runtime/Engine/Private/Controller.cpp b/Engine/Source/Runtime/Engine/Private/Controller.cpp
--- a/Engine/Source/Runtime/Engine/Private/Controller.cpp
+++ b/Engine/Source/Runtime/Engine/Private/Controller.cpp
@@ -20,8 +20,8 @@
 DEFINE_LOG_CATEGORY(LogPath);
 
 
-AController::AController(const class FPostConstructInitializeProperties& PCIP)
-	: Super(PCIP)
+AController::AController(const FObjectInitializer& ObjectInitializer)
+	: Super(ObjectInitializer)
 {
 
 	PrimaryActorTick.bCanEverTick = true;
@@ -31,7 +31,7 @@ AController::AController(const class FPostConstructInitializeProperties& PCIP)
 #endif // WITH_EDITORONLY_DATA
 	bOnlyRelevantToOwner = true;
 
-	TransformComponent = PCIP.CreateDefaultSubobject<USceneComponent>(this, TEXT("TransformComponent0"));
+	TransformComponent = ObjectInitializer.CreateDefaultSubobject<USceneComponent>(this, TEXT("TransformComponent0"));
 	RootComponent = TransformComponent;
 
 	bCanBeDamaged = false;
diff --git a/Engine/Source/Runtime/Engine/Private/CullDistanceVolume.cpp b/Engine/Source/Runtime/Engine/Private/CullDistanceVolume.cpp
--- a/Engine/Source/Runtime/Engine/Private/CullDistanceVolume.cpp
+++ b/Engine/Source/Runtime/Engine/Private/CullDistanceVolume.cpp
@@ -3,8 +3,8 @@
 #include "EnginePrivate.h"
 #include "Engine/CullDistanceVolume.h"
 
-ACullDistanceVolume::ACullDistanceVolume(const class FPostConstructInitializeProperties& PCIP)
-	: Super(PCIP)
+ACullDistanceVolume::ACullDistanceVolume(const FObjectInitializer& ObjectInitializer)
+	: Super(ObjectInitializer)
 {
 	BrushComponent->BodyInstance.bEnableCollision_DEPRECATED = false;
 	BrushComponent->SetCollisionProfileName(UCollisionProfile::NoCollision_ProfileName);
diff --git a/Engine/Source/Runtime/Engine/Private/CurveBase.cpp b/Engine/Source/Runtime/Engine/Private/CurveBase.cpp
--- a/Engine/Source/Runtime/Engine/Private/CurveBase.cpp
+++ b/Engine/Source/Runtime/Engine/Private/CurveBase.cpp
@@ -796,8 +796,8 @@ bool FRichCurve::operator==( const FRichCurve& Curve ) const
 //////////////////////////////////////////////////////////////////////////
 // UCurveBase
 
-UCurveBase::UCurveBase(const class FPostConstructInitializeProperties& PCIP)
-	: Super(PCIP)
+UCurveBase::UCurveBase(const FObjectInitializer& ObjectInitializer)
+	: Super(ObjectInitializer)
 {
 }
 
diff --git a/Engine/Source/Runtime/Engine/Private/CurveEdPresetCurve.cpp b/Engine/Source/Runtime/Engine/Private/CurveEdPresetCurve.cpp
--- a/Engine/Source/Runtime/Engine/Private/CurveEdPresetCurve.cpp
+++ b/Engine/Source/Runtime/Engine/Private/CurveEdPresetCurve.cpp
@@ -10,7 +10,7 @@
  *	UCurveEdPresetCurve
  */
 
-UDEPRECATED_CurveEdPresetCurve::UDEPRECATED_CurveEdPresetCurve(const class FPostConstructInitializeProperties& PCIP)
-	: Super(PCIP)
+UDEPRECATED_CurveEdPresetCurve::UDEPRECATED_CurveEdPresetCurve(const FObjectInitializer& ObjectInitializer)
+	: Super(ObjectInitializer)
 {
 }
\ No newline at end of file
diff --git a/Engine/Source/Runtime/Engine/Private/CurveFloat.cpp b/Engine/Source/Runtime/Engine/Private/CurveFloat.cpp
--- a/Engine/Source/Runtime/Engine/Private/CurveFloat.cpp
+++ b/Engine/Source/Runtime/Engine/Private/CurveFloat.cpp
@@ -35,8 +35,8 @@ const FRichCurve* FRuntimeFloatCurve::GetRichCurveConst() const
 	}
 }
 
-UCurveFloat::UCurveFloat(const class FPostConstructInitializeProperties& PCIP)
-	: Super(PCIP)
+UCurveFloat::UCurveFloat(const FObjectInitializer& ObjectInitializer)
+	: Super(ObjectInitializer)
 {
 }
 
diff --git a/Engine/Source/Runtime/Engine/Private/CurveLinearColor.cpp b/Engine/Source/Runtime/Engine/Private/CurveLinearColor.cpp
--- a/Engine/Source/Runtime/Engine/Private/CurveLinearColor.cpp
+++ b/Engine/Source/Runtime/Engine/Private/CurveLinearColor.cpp
@@ -6,8 +6,8 @@
 
 #include "EnginePrivate.h"
 
-UCurveLinearColor::UCurveLinearColor(const class FPostConstructInitializeProperties& PCIP)
-	: Super(PCIP)
+UCurveLinearColor::UCurveLinearColor(const FObjectInitializer& ObjectInitializer)
+	: Super(ObjectInitializer)
 {
 }
 
diff --git a/Engine/Source/Runtime/Engine/Private/CurveTable.cpp b/Engine/Source/Runtime/Engine/Private/CurveTable.cpp
--- a/Engine/Source/Runtime/Engine/Private/CurveTable.cpp
+++ b/Engine/Source/Runtime/Engine/Private/CurveTable.cpp
@@ -12,8 +12,8 @@ DECLARE_CYCLE_STAT(TEXT("CurveTableRowHandle Eval"),STAT_CurveTableRowHandleEval
 
 
 //////////////////////////////////////////////////////////////////////////
-UCurveTable::UCurveTable(const class FPostConstructInitializeProperties& PCIP)
-	: Super(PCIP)
+UCurveTable::UCurveTable(const FObjectInitializer& ObjectInitializer)
+	: Super(ObjectInitializer)
 {
 }
 
diff --git a/Engine/Source/Runtime/Engine/Private/CurveVector.cpp b/Engine/Source/Runtime/Engine/Private/CurveVector.cpp
--- a/Engine/Source/Runtime/Engine/Private/CurveVector.cpp
+++ b/Engine/Source/Runtime/Engine/Private/CurveVector.cpp
@@ -7,8 +7,8 @@
 #include "EnginePrivate.h"
 
 
-UCurveVector::UCurveVector(const class FPostConstructInitializeProperties& PCIP)
-	: Super(PCIP)
+UCurveVector::UCurveVector(const FObjectInitializer& ObjectInitializer)
+	: Super(ObjectInitializer)
 {
 }
 
diff --git a/Engine/Source/Runtime/Engine/Private/DamageType.cpp b/Engine/Source/Runtime/Engine/Private/DamageType.cpp
--- a/Engine/Source/Runtime/Engine/Private/DamageType.cpp
+++ b/Engine/Source/Runtime/Engine/Private/DamageType.cpp
@@ -2,8 +2,8 @@
 
 #include "EnginePrivate.h"
 
-UDamageType::UDamageType(const class FPostConstructInitializeProperties& PCIP)
-	: Super(PCIP)
+UDamageType::UDamageType(const FObjectInitializer& ObjectInitializer)
+	: Super(ObjectInitializer)
 {
 	DamageImpulse = 800.0f;
 	DestructibleImpulse = 800.0f;
diff --git a/Engine/Source/Runtime/Engine/Private/DataAsset.cpp b/Engine/Source/Runtime/Engine/Private/DataAsset.cpp
--- a/Engine/Source/Runtime/Engine/Private/DataAsset.cpp
+++ b/Engine/Source/Runtime/Engine/Private/DataAsset.cpp
@@ -2,8 +2,8 @@
 
 #include "EnginePrivate.h"
 
-UDataAsset::UDataAsset(const class FPostConstructInitializeProperties& PCIP)
-	: Super(PCIP)
+UDataAsset::UDataAsset(const FObjectInitializer& ObjectInitializer)
+	: Super(ObjectInitializer)
 {
 }
 
diff --git a/Engine/Source/Runtime/Engine/Private/DataChannel.cpp b/Engine/Source/Runtime/Engine/Private/DataChannel.cpp
--- a/Engine/Source/Runtime/Engine/Private/DataChannel.cpp
+++ b/Engine/Source/Runtime/Engine/Private/DataChannel.cpp
@@ -33,8 +33,8 @@ static TAutoConsoleVariable<int32> CVarNetReliableDebug(
 	UChannel implementation.
 -----------------------------------------------------------------------------*/
 
-UChannel::UChannel(const class FPostConstructInitializeProperties& PCIP)
-	: Super(PCIP)
+UChannel::UChannel(const FObjectInitializer& ObjectInitializer)
+	: Super(ObjectInitializer)
 {
 }
 
diff --git a/Engine/Source/Runtime/Engine/Private/DataTable.cpp b/Engine/Source/Runtime/Engine/Private/DataTable.cpp
--- a/Engine/Source/Runtime/Engine/Private/DataTable.cpp
+++ b/Engine/Source/Runtime/Engine/Private/DataTable.cpp
@@ -9,8 +9,8 @@ DEFINE_LOG_CATEGORY(LogDataTable);
 ENGINE_API const FString FDataTableRowHandle::Unknown(TEXT("UNKNOWN"));
 ENGINE_API const FString FDataTableCategoryHandle::Unknown(TEXT("UNKNOWN"));
 
-UDataTable::UDataTable(const class FPostConstructInitializeProperties& PCIP)
-	: Super(PCIP)
+UDataTable::UDataTable(const FObjectInitializer& ObjectInitializer)
+	: Super(ObjectInitializer)
 {
 }
 
diff --git a/Engine/Source/Runtime/Engine/Private/DataTableFunctionLibrary.cpp b/Engine/Source/Runtime/Engine/Private/DataTableFunctionLibrary.cpp
--- a/Engine/Source/Runtime/Engine/Private/DataTableFunctionLibrary.cpp
+++ b/Engine/Source/Runtime/Engine/Private/DataTableFunctionLibrary.cpp
@@ -3,8 +3,8 @@
 #include "EnginePrivate.h"
 #include "../Classes/Kismet/DataTableFunctionLibrary.h"
 
-UDataTableFunctionLibrary::UDataTableFunctionLibrary(const FPostConstructInitializeProperties& PCIP)
-	: Super(PCIP)
+UDataTableFunctionLibrary::UDataTableFunctionLibrary(const FObjectInitializer& ObjectInitializer)
+	: Super(ObjectInitializer)
 {
 }
 
diff --git a/Engine/Source/Runtime/Engine/Private/Debug/DebugDrawService.cpp b/Engine/Source/Runtime/Engine/Private/Debug/DebugDrawService.cpp
--- a/Engine/Source/Runtime/Engine/Private/Debug/DebugDrawService.cpp
+++ b/Engine/Source/Runtime/Engine/Private/Debug/DebugDrawService.cpp
@@ -5,8 +5,8 @@
 TArray<TArray<FDebugDrawDelegate> > UDebugDrawService::Delegates;
 FEngineShowFlags UDebugDrawService::ObservedFlags(ESFIM_Editor);
 
-UDebugDrawService::UDebugDrawService(const class FPostConstructInitializeProperties& PCIP)
-	: Super(PCIP)
+UDebugDrawService::UDebugDrawService(const FObjectInitializer& ObjectInitializer)
+	: Super(ObjectInitializer)
 {
 	Delegates.Reserve(sizeof(FEngineShowFlags)*8);
 }
