Rename FPostConstructInitializeProperties to something simpler

diff --git a/Engine/Source/Runtime/Engine/Private/Debug/LogVisualizerCameraController.cpp b/Engine/Source/Runtime/Engine/Private/Debug/LogVisualizerCameraController.cpp
--- a/Engine/Source/Runtime/Engine/Private/Debug/LogVisualizerCameraController.cpp
+++ b/Engine/Source/Runtime/Engine/Private/Debug/LogVisualizerCameraController.cpp
@@ -11,8 +11,8 @@
 //----------------------------------------------------------------------//
 TWeakObjectPtr<ALogVisualizerCameraController> ALogVisualizerCameraController::Instance;
 
-ALogVisualizerCameraController::ALogVisualizerCameraController(const class FPostConstructInitializeProperties& PCIP)
-	: Super(PCIP)
+ALogVisualizerCameraController::ALogVisualizerCameraController(const FObjectInitializer& ObjectInitializer)
+	: Super(ObjectInitializer)
 {
 	SpeedScale = 1.5f;
 }
diff --git a/Engine/Source/Runtime/Engine/Private/Debug/LogVisualizerHUD.cpp b/Engine/Source/Runtime/Engine/Private/Debug/LogVisualizerHUD.cpp
--- a/Engine/Source/Runtime/Engine/Private/Debug/LogVisualizerHUD.cpp
+++ b/Engine/Source/Runtime/Engine/Private/Debug/LogVisualizerHUD.cpp
@@ -10,8 +10,8 @@
 //----------------------------------------------------------------------//
 // ALogVisualizerHUD
 //----------------------------------------------------------------------//
-ALogVisualizerHUD::ALogVisualizerHUD(const class FPostConstructInitializeProperties& PCIP)
-	: Super(PCIP)
+ALogVisualizerHUD::ALogVisualizerHUD(const FObjectInitializer& ObjectInitializer)
+	: Super(ObjectInitializer)
 {
 	bHidden = false;
 
diff --git a/Engine/Source/Runtime/Engine/Private/Debug/ReporterBase.cpp b/Engine/Source/Runtime/Engine/Private/Debug/ReporterBase.cpp
--- a/Engine/Source/Runtime/Engine/Private/Debug/ReporterBase.cpp
+++ b/Engine/Source/Runtime/Engine/Private/Debug/ReporterBase.cpp
@@ -5,8 +5,8 @@
 
 #define DASH_LINE_SIZE	5.0f
 
-UReporterBase::UReporterBase(const class FPostConstructInitializeProperties& PCIP) :
-	Super(PCIP)
+UReporterBase::UReporterBase(const FObjectInitializer& ObjectInitializer) :
+	Super(ObjectInitializer)
 {
 	bVisible = false;
 	
diff --git a/Engine/Source/Runtime/Engine/Private/Debug/ReporterGraph.cpp b/Engine/Source/Runtime/Engine/Private/Debug/ReporterGraph.cpp
--- a/Engine/Source/Runtime/Engine/Private/Debug/ReporterGraph.cpp
+++ b/Engine/Source/Runtime/Engine/Private/Debug/ReporterGraph.cpp
@@ -4,8 +4,8 @@
 #include "../../Classes/Debug/ReporterBase.h"
 #include "../../Classes/Debug/ReporterGraph.h"
 
-UReporterGraph::UReporterGraph(const class FPostConstructInitializeProperties& PCIP) :
-	Super(PCIP)
+UReporterGraph::UReporterGraph(const FObjectInitializer& ObjectInitializer) :
+	Super(ObjectInitializer)
 {
 	AxesColor = FLinearColor::Yellow;
 
diff --git a/Engine/Source/Runtime/Engine/Private/DebugCameraController.cpp b/Engine/Source/Runtime/Engine/Private/DebugCameraController.cpp
--- a/Engine/Source/Runtime/Engine/Private/DebugCameraController.cpp
+++ b/Engine/Source/Runtime/Engine/Private/DebugCameraController.cpp
@@ -26,8 +26,8 @@ extern bool TrackTexture( const FString& TextureName );
 static const float SPEED_SCALE_ADJUSTMENT = 0.5f;
 
 
-ADebugCameraController::ADebugCameraController(const class FPostConstructInitializeProperties& PCIP)
-	: Super(PCIP)
+ADebugCameraController::ADebugCameraController(const FObjectInitializer& ObjectInitializer)
+	: Super(ObjectInitializer)
 {
 	SelectedActor = NULL;
 	SelectedComponent = NULL;
diff --git a/Engine/Source/Runtime/Engine/Private/DebugCameraHUD.cpp b/Engine/Source/Runtime/Engine/Private/DebugCameraHUD.cpp
--- a/Engine/Source/Runtime/Engine/Private/DebugCameraHUD.cpp
+++ b/Engine/Source/Runtime/Engine/Private/DebugCameraHUD.cpp
@@ -21,8 +21,8 @@ static TAutoConsoleVariable<int32> CVarDebugCameraTraceComplex(
 // Externals
 // ------------------
 
-ADebugCameraHUD::ADebugCameraHUD(const class FPostConstructInitializeProperties& PCIP)
-	: Super(PCIP)
+ADebugCameraHUD::ADebugCameraHUD(const FObjectInitializer& ObjectInitializer)
+	: Super(ObjectInitializer)
 {
 	bHidden = false;
 }
diff --git a/Engine/Source/Runtime/Engine/Private/DecalActor.cpp b/Engine/Source/Runtime/Engine/Private/DecalActor.cpp
--- a/Engine/Source/Runtime/Engine/Private/DecalActor.cpp
+++ b/Engine/Source/Runtime/Engine/Private/DecalActor.cpp
@@ -15,10 +15,10 @@ namespace DecalEditorConstants
 }
 #endif
 
-ADecalActor::ADecalActor(const class FPostConstructInitializeProperties& PCIP)
-: Super(PCIP)
+ADecalActor::ADecalActor(const FObjectInitializer& ObjectInitializer)
+: Super(ObjectInitializer)
 {
-	Decal = PCIP.CreateDefaultSubobject<UDecalComponent>(this, TEXT("NewDecalComponent"));
+	Decal = ObjectInitializer.CreateDefaultSubobject<UDecalComponent>(this, TEXT("NewDecalComponent"));
 	Decal->RelativeScale3D = FVector(128.0f, 256.0f, 256.0f);
 
 	Decal->RelativeRotation = FRotator(-90, 0, 0);
@@ -26,7 +26,7 @@ ADecalActor::ADecalActor(const class FPostConstructInitializeProperties& PCIP)
 	RootComponent = Decal;
 
 #if WITH_EDITORONLY_DATA
-	BoxComponent = PCIP.CreateEditorOnlyDefaultSubobject<UBoxComponent>(this, TEXT("DrawBox0"));
+	BoxComponent = ObjectInitializer.CreateEditorOnlyDefaultSubobject<UBoxComponent>(this, TEXT("DrawBox0"));
 	if (BoxComponent != nullptr)
 	{
 		BoxComponent->BodyInstance.bEnableCollision_DEPRECATED = false;
@@ -39,8 +39,8 @@ ADecalActor::ADecalActor(const class FPostConstructInitializeProperties& PCIP)
 		BoxComponent->AttachParent = Decal;
 	}
 
-	ArrowComponent = PCIP.CreateEditorOnlyDefaultSubobject<UArrowComponent>(this, TEXT("ArrowComponent0"));
-	SpriteComponent = PCIP.CreateEditorOnlyDefaultSubobject<UBillboardComponent>(this, TEXT("Sprite"));
+	ArrowComponent = ObjectInitializer.CreateEditorOnlyDefaultSubobject<UArrowComponent>(this, TEXT("ArrowComponent0"));
+	SpriteComponent = ObjectInitializer.CreateEditorOnlyDefaultSubobject<UBillboardComponent>(this, TEXT("Sprite"));
 
 	if (!IsRunningCommandlet())
 	{
diff --git a/Engine/Source/Runtime/Engine/Private/DecalComponent.cpp b/Engine/Source/Runtime/Engine/Private/DecalComponent.cpp
--- a/Engine/Source/Runtime/Engine/Private/DecalComponent.cpp
+++ b/Engine/Source/Runtime/Engine/Private/DecalComponent.cpp
@@ -34,8 +34,8 @@ void FDeferredDecalProxy::SetTransform(const FTransform& InComponentToWorld)
 	ComponentTrans = InComponentToWorld;
 }
 
-UDecalComponent::UDecalComponent(const class FPostConstructInitializeProperties& PCIP)
-	: Super(PCIP)
+UDecalComponent::UDecalComponent(const FObjectInitializer& ObjectInitializer)
+	: Super(ObjectInitializer)
 {
 }
 
diff --git a/Engine/Source/Runtime/Engine/Private/DefaultPawn.cpp b/Engine/Source/Runtime/Engine/Private/DefaultPawn.cpp
--- a/Engine/Source/Runtime/Engine/Private/DefaultPawn.cpp
+++ b/Engine/Source/Runtime/Engine/Private/DefaultPawn.cpp
@@ -10,8 +10,8 @@ FName ADefaultPawn::CollisionComponentName(TEXT("CollisionComponent0"));
 FName ADefaultPawn::MeshComponentName(TEXT("MeshComponent0"));
 
 
-ADefaultPawn::ADefaultPawn(const class FPostConstructInitializeProperties& PCIP)
-	: Super(PCIP)
+ADefaultPawn::ADefaultPawn(const FObjectInitializer& ObjectInitializer)
+	: Super(ObjectInitializer)
 {
 	bCanBeDamaged = true;
 
@@ -22,7 +22,7 @@ ADefaultPawn::ADefaultPawn(const class FPostConstructInitializeProperties& PCIP)
 	BaseEyeHeight = 0.0f;
 	bCollideWhenPlacing = false;
 
-	CollisionComponent = PCIP.CreateDefaultSubobject<USphereComponent>(this, ADefaultPawn::CollisionComponentName);
+	CollisionComponent = ObjectInitializer.CreateDefaultSubobject<USphereComponent>(this, ADefaultPawn::CollisionComponentName);
 	CollisionComponent->InitSphereRadius(35.0f);
 	CollisionComponent->BodyInstance.bEnableCollision_DEPRECATED = true;
 
@@ -35,7 +35,7 @@ ADefaultPawn::ADefaultPawn(const class FPostConstructInitializeProperties& PCIP)
 
 	RootComponent = CollisionComponent;
 
-	MovementComponent = PCIP.CreateDefaultSubobject<UFloatingPawnMovement>(this, ADefaultPawn::MovementComponentName);
+	MovementComponent = ObjectInitializer.CreateDefaultSubobject<UFloatingPawnMovement>(this, ADefaultPawn::MovementComponentName);
 	MovementComponent->UpdatedComponent = CollisionComponent;
 
 	// Structure to hold one-time initialization
@@ -48,7 +48,7 @@ ADefaultPawn::ADefaultPawn(const class FPostConstructInitializeProperties& PCIP)
 
 	static FConstructorStatics ConstructorStatics;
 
-	MeshComponent = PCIP.CreateOptionalDefaultSubobject<UStaticMeshComponent>(this, ADefaultPawn::MeshComponentName);
+	MeshComponent = ObjectInitializer.CreateOptionalDefaultSubobject<UStaticMeshComponent>(this, ADefaultPawn::MeshComponentName);
 	if (MeshComponent)
 	{
 		MeshComponent->SetStaticMesh(ConstructorStatics.SphereMesh.Object);
diff --git a/Engine/Source/Runtime/Engine/Private/DefaultPhysicsVolume.cpp b/Engine/Source/Runtime/Engine/Private/DefaultPhysicsVolume.cpp
--- a/Engine/Source/Runtime/Engine/Private/DefaultPhysicsVolume.cpp
+++ b/Engine/Source/Runtime/Engine/Private/DefaultPhysicsVolume.cpp
@@ -5,8 +5,8 @@
 
 DEFINE_LOG_CATEGORY_STATIC(LogVolume, Log, All);
 
-ADefaultPhysicsVolume::ADefaultPhysicsVolume(const class FPostConstructInitializeProperties& PCIP)
-	: Super(PCIP)
+ADefaultPhysicsVolume::ADefaultPhysicsVolume(const FObjectInitializer& ObjectInitializer)
+	: Super(ObjectInitializer)
 {
 
 #if WITH_EDITORONLY_DATA
diff --git a/Engine/Source/Runtime/Engine/Private/DemoNetDriver.cpp b/Engine/Source/Runtime/Engine/Private/DemoNetDriver.cpp
--- a/Engine/Source/Runtime/Engine/Private/DemoNetDriver.cpp
+++ b/Engine/Source/Runtime/Engine/Private/DemoNetDriver.cpp
@@ -25,8 +25,8 @@ static const int32 MAX_DEMO_READ_WRITE_BUFFER = 1024 * 2;
 	UDemoNetDriver.
 -----------------------------------------------------------------------------*/
 
-UDemoNetDriver::UDemoNetDriver(const class FPostConstructInitializeProperties& PCIP)
-	: Super(PCIP)
+UDemoNetDriver::UDemoNetDriver(const FObjectInitializer& ObjectInitializer)
+	: Super(ObjectInitializer)
 {
 }
 
@@ -124,7 +124,7 @@ bool UDemoNetDriver::InitConnect( FNetworkNotify* InNotify, const FURL& ConnectU
 	World->DemoNetDriver = NULL;
 	SetWorld( NULL );
 
-	UDemoPendingNetGame * NewPendingNetGame = new UDemoPendingNetGame( FPostConstructInitializeProperties() );
+	UDemoPendingNetGame * NewPendingNetGame = new UDemoPendingNetGame( FObjectInitializer() );
 
 	NewPendingNetGame->DemoNetDriver = this;
 
@@ -721,7 +721,7 @@ void UDemoNetDriver::SpawnDemoRecSpectator( UNetConnection* Connection )
 	UDemoNetConnection.
 -----------------------------------------------------------------------------*/
 
-UDemoNetConnection::UDemoNetConnection( const class FPostConstructInitializeProperties& PCIP ) : Super( PCIP )
+UDemoNetConnection::UDemoNetConnection( const FObjectInitializer& ObjectInitializer ) : Super( ObjectInitializer )
 {
 	MaxPacket = 512;
 	InternalAck = true;
@@ -835,6 +835,6 @@ bool UDemoNetConnection::ClientHasInitializedLevelFor(const UObject* TestObject)
 	UDemoPendingNetGame.
 -----------------------------------------------------------------------------*/
 
-UDemoPendingNetGame::UDemoPendingNetGame( const class FPostConstructInitializeProperties& PCIP ) : Super( PCIP )
+UDemoPendingNetGame::UDemoPendingNetGame( const FObjectInitializer& ObjectInitializer ) : Super( ObjectInitializer )
 {
 }
diff --git a/Engine/Source/Runtime/Engine/Private/DestructibleActor.cpp b/Engine/Source/Runtime/Engine/Private/DestructibleActor.cpp
--- a/Engine/Source/Runtime/Engine/Private/DestructibleActor.cpp
+++ b/Engine/Source/Runtime/Engine/Private/DestructibleActor.cpp
@@ -9,10 +9,10 @@
 #include "PhysicsEngine/PhysXSupport.h"
 #include "PhysicsEngine/DestructibleActor.h"
 
-ADestructibleActor::ADestructibleActor(const class FPostConstructInitializeProperties& PCIP)
-	: Super(PCIP)
+ADestructibleActor::ADestructibleActor(const FObjectInitializer& ObjectInitializer)
+	: Super(ObjectInitializer)
 {
-	DestructibleComponent = PCIP.CreateDefaultSubobject<UDestructibleComponent>(this, TEXT("DestructibleComponent0"));
+	DestructibleComponent = ObjectInitializer.CreateDefaultSubobject<UDestructibleComponent>(this, TEXT("DestructibleComponent0"));
 	DestructibleComponent->bCanEverAffectNavigation = bAffectNavigation;
 	RootComponent = DestructibleComponent;
 }
diff --git a/Engine/Source/Runtime/Engine/Private/DestructibleComponent.cpp b/Engine/Source/Runtime/Engine/Private/DestructibleComponent.cpp
--- a/Engine/Source/Runtime/Engine/Private/DestructibleComponent.cpp
+++ b/Engine/Source/Runtime/Engine/Private/DestructibleComponent.cpp
@@ -16,8 +16,8 @@
 #include "PhysicsEngine/DestructibleActor.h"
 #include "Engine/DestructibleMesh.h"
 
-UDestructibleComponent::UDestructibleComponent(const class FPostConstructInitializeProperties& PCIP)
-	: Super(PCIP)
+UDestructibleComponent::UDestructibleComponent(const FObjectInitializer& ObjectInitializer)
+	: Super(ObjectInitializer)
 #if WITH_PHYSX
 	, PhysxUserData(this)
 #endif
diff --git a/Engine/Source/Runtime/Engine/Private/DestructibleFractureSettings.cpp b/Engine/Source/Runtime/Engine/Private/DestructibleFractureSettings.cpp
--- a/Engine/Source/Runtime/Engine/Private/DestructibleFractureSettings.cpp
+++ b/Engine/Source/Runtime/Engine/Private/DestructibleFractureSettings.cpp
@@ -227,8 +227,8 @@ void FFractureMaterial::FillNxFractureMaterialDesc(NxFractureMaterialDesc& PFrac
 // UDestructibleFractureSettings
 //////////////////////////////////////////////////////////////////////////
 
-UDestructibleFractureSettings::UDestructibleFractureSettings(const class FPostConstructInitializeProperties& PCIP)
-	: Super(PCIP)
+UDestructibleFractureSettings::UDestructibleFractureSettings(const FObjectInitializer& ObjectInitializer)
+	: Super(ObjectInitializer)
 {
 	CellSiteCount = 25;
 	OriginalSubmeshCount = 0;
diff --git a/Engine/Source/Runtime/Engine/Private/DestructibleMesh.cpp b/Engine/Source/Runtime/Engine/Private/DestructibleMesh.cpp
--- a/Engine/Source/Runtime/Engine/Private/DestructibleMesh.cpp
+++ b/Engine/Source/Runtime/Engine/Private/DestructibleMesh.cpp
@@ -17,8 +17,8 @@
 
 DEFINE_LOG_CATEGORY(LogDestructible)
 
-UDestructibleMesh::UDestructibleMesh(const class FPostConstructInitializeProperties& PCIP)
-	: Super(PCIP)
+UDestructibleMesh::UDestructibleMesh(const FObjectInitializer& ObjectInitializer)
+	: Super(ObjectInitializer)
 {
 }
 
diff --git a/Engine/Source/Runtime/Engine/Private/DeviceProfiles/DeviceProfile.cpp b/Engine/Source/Runtime/Engine/Private/DeviceProfiles/DeviceProfile.cpp
--- a/Engine/Source/Runtime/Engine/Private/DeviceProfiles/DeviceProfile.cpp
+++ b/Engine/Source/Runtime/Engine/Private/DeviceProfiles/DeviceProfile.cpp
@@ -3,8 +3,8 @@
 #include "EnginePrivate.h"
 #include "DeviceProfiles/DeviceProfile.h"
 
-UDeviceProfile::UDeviceProfile(const class FPostConstructInitializeProperties& PCIP)
-	: Super(PCIP)
+UDeviceProfile::UDeviceProfile(const FObjectInitializer& ObjectInitializer)
+	: Super(ObjectInitializer)
 {
 	BaseProfileName = TEXT("");
 	DeviceType = TEXT("");
diff --git a/Engine/Source/Runtime/Engine/Private/DeviceProfiles/DeviceProfileManager.cpp b/Engine/Source/Runtime/Engine/Private/DeviceProfiles/DeviceProfileManager.cpp
--- a/Engine/Source/Runtime/Engine/Private/DeviceProfiles/DeviceProfileManager.cpp
+++ b/Engine/Source/Runtime/Engine/Private/DeviceProfiles/DeviceProfileManager.cpp
@@ -123,8 +123,8 @@ void UDeviceProfileManager::InitializeCVarsForActiveDeviceProfile()
 }
 
 
-UDeviceProfileManager::UDeviceProfileManager( const class FPostConstructInitializeProperties& PCIP )
-	: Super( PCIP )
+UDeviceProfileManager::UDeviceProfileManager( const FObjectInitializer& ObjectInitializer )
+	: Super( ObjectInitializer )
 {
 	RenameIndex = 0;
 #if WITH_EDITOR
diff --git a/Engine/Source/Runtime/Engine/Private/DialogueTypes.cpp b/Engine/Source/Runtime/Engine/Private/DialogueTypes.cpp
--- a/Engine/Source/Runtime/Engine/Private/DialogueTypes.cpp
+++ b/Engine/Source/Runtime/Engine/Private/DialogueTypes.cpp
@@ -49,7 +49,7 @@ FDialogueWaveParameter::FDialogueWaveParameter() : DialogueWave(NULL)
 
 }
 
-UDialogueTypes::UDialogueTypes(const class FPostConstructInitializeProperties& PCIP)
-	: Super(PCIP)
+UDialogueTypes::UDialogueTypes(const FObjectInitializer& ObjectInitializer)
+	: Super(ObjectInitializer)
 {
 }
\ No newline at end of file
diff --git a/Engine/Source/Runtime/Engine/Private/DialogueVoice.cpp b/Engine/Source/Runtime/Engine/Private/DialogueVoice.cpp
--- a/Engine/Source/Runtime/Engine/Private/DialogueVoice.cpp
+++ b/Engine/Source/Runtime/Engine/Private/DialogueVoice.cpp
@@ -7,8 +7,8 @@
 #include "AudioDerivedData.h"
 #include "SubtitleManager.h"
 
-UDialogueVoice::UDialogueVoice(const class FPostConstructInitializeProperties& PCIP)
-	: Super(PCIP)
+UDialogueVoice::UDialogueVoice(const FObjectInitializer& ObjectInitializer)
+	: Super(ObjectInitializer)
 	, LocalizationGUID( FGuid::NewGuid() )
 {
 }
diff --git a/Engine/Source/Runtime/Engine/Private/DialogueWave.cpp b/Engine/Source/Runtime/Engine/Private/DialogueWave.cpp
--- a/Engine/Source/Runtime/Engine/Private/DialogueWave.cpp
+++ b/Engine/Source/Runtime/Engine/Private/DialogueWave.cpp
@@ -23,8 +23,8 @@ FDialogueContextMapping::FDialogueContextMapping() : SoundWave(NULL), Proxy(NULL
 
 }
 
-UDialogueSoundWaveProxy::UDialogueSoundWaveProxy(const class FPostConstructInitializeProperties& PCIP)
-	: Super(PCIP)
+UDialogueSoundWaveProxy::UDialogueSoundWaveProxy(const FObjectInitializer& ObjectInitializer)
+	: Super(ObjectInitializer)
 {
 }
 
@@ -91,8 +91,8 @@ void UDialogueSoundWaveProxy::Parse( class FAudioDevice* AudioDevice, const UPTR
 	}
 }
 
-UDialogueWave::UDialogueWave(const class FPostConstructInitializeProperties& PCIP)
-	: Super(PCIP)
+UDialogueWave::UDialogueWave(const FObjectInitializer& ObjectInitializer)
+	: Super(ObjectInitializer)
 	, LocalizationGUID( FGuid::NewGuid() )
 {
 	ContextMappings.Add( FDialogueContextMapping() );
diff --git a/Engine/Source/Runtime/Engine/Private/DirectionalLightComponent.cpp b/Engine/Source/Runtime/Engine/Private/DirectionalLightComponent.cpp
--- a/Engine/Source/Runtime/Engine/Private/DirectionalLightComponent.cpp
+++ b/Engine/Source/Runtime/Engine/Private/DirectionalLightComponent.cpp
@@ -610,8 +610,8 @@ private:
 	}
 };
 
-UDirectionalLightComponent::UDirectionalLightComponent(const class FPostConstructInitializeProperties& PCIP)
-	: Super(PCIP)
+UDirectionalLightComponent::UDirectionalLightComponent(const FObjectInitializer& ObjectInitializer)
+	: Super(ObjectInitializer)
 {
 #if WITH_EDITORONLY_DATA
 	if (!IsRunningCommandlet())
diff --git a/Engine/Source/Runtime/Engine/Private/DistributionFloatSoundParameter.cpp b/Engine/Source/Runtime/Engine/Private/DistributionFloatSoundParameter.cpp
--- a/Engine/Source/Runtime/Engine/Private/DistributionFloatSoundParameter.cpp
+++ b/Engine/Source/Runtime/Engine/Private/DistributionFloatSoundParameter.cpp
@@ -6,7 +6,7 @@
 /*-----------------------------------------------------------------------------
 	Parameter-based distributions
 -----------------------------------------------------------------------------*/
-	UDEPRECATED_DistributionFloatSoundParameter::UDEPRECATED_DistributionFloatSoundParameter(const class FPostConstructInitializeProperties& PCIP)
-	: Super(PCIP)
+	UDEPRECATED_DistributionFloatSoundParameter::UDEPRECATED_DistributionFloatSoundParameter(const FObjectInitializer& ObjectInitializer)
+	: Super(ObjectInitializer)
 {
 }
diff --git a/Engine/Source/Runtime/Engine/Private/Distributions.cpp b/Engine/Source/Runtime/Engine/Private/Distributions.cpp
--- a/Engine/Source/Runtime/Engine/Private/Distributions.cpp
+++ b/Engine/Source/Runtime/Engine/Private/Distributions.cpp
@@ -31,8 +31,8 @@ static_assert((LOOKUP_TABLE_MAX_SAMPLES & (LOOKUP_TABLE_MAX_SAMPLES - 1)) == 0,
 // Log category for distributions.
 DEFINE_LOG_CATEGORY_STATIC(LogDistributions,Warning,Warning);
 
-UDistribution::UDistribution(const class FPostConstructInitializeProperties& PCIP)
-	: Super(PCIP)
+UDistribution::UDistribution(const FObjectInitializer& ObjectInitializer)
+	: Super(ObjectInitializer)
 {
 }
 
@@ -1544,8 +1544,8 @@ void FComposableVectorDistribution::Resample( float MinIn, float MaxIn )
 	ResampleLookupTable( &LookupTable, OldTable, MinIn, MaxIn, LOOKUP_TABLE_MAX_SAMPLES );
 }
 
-UDistributionFloatConstant::UDistributionFloatConstant(const class FPostConstructInitializeProperties& PCIP)
-	: Super(PCIP)
+UDistributionFloatConstant::UDistributionFloatConstant(const FObjectInitializer& ObjectInitializer)
+	: Super(ObjectInitializer)
 {
 }
 
@@ -1678,8 +1678,8 @@ void UDistributionFloatConstant::SetTangents(int32 SubIndex, int32 KeyIndex, flo
 	check( KeyIndex == 0 );
 }
 
-UDistributionFloatConstantCurve::UDistributionFloatConstantCurve(const class FPostConstructInitializeProperties& PCIP)
-	: Super(PCIP)
+UDistributionFloatConstantCurve::UDistributionFloatConstantCurve(const FObjectInitializer& ObjectInitializer)
+	: Super(ObjectInitializer)
 {
 }
 
@@ -1837,8 +1837,8 @@ void UDistributionFloatConstantCurve::SetTangents(int32 SubIndex, int32 KeyIndex
 }
 
 
-UDistributionFloatUniform::UDistributionFloatUniform(const class FPostConstructInitializeProperties& PCIP)
-	: Super(PCIP)
+UDistributionFloatUniform::UDistributionFloatUniform(const FObjectInitializer& ObjectInitializer)
+	: Super(ObjectInitializer)
 {
 }
 
@@ -2041,8 +2041,8 @@ void UDistributionFloatUniform::SetTangents(int32 SubIndex, int32 KeyIndex, floa
 	check( KeyIndex == 0 );
 }
 
-UDistributionFloatUniformCurve::UDistributionFloatUniformCurve(const class FPostConstructInitializeProperties& PCIP)
-	: Super(PCIP)
+UDistributionFloatUniformCurve::UDistributionFloatUniformCurve(const FObjectInitializer& ObjectInitializer)
+	: Super(ObjectInitializer)
 {
 }
 
@@ -2308,8 +2308,8 @@ void UDistributionFloatUniformCurve::SetTangents(int32 SubIndex, int32 KeyIndex,
 	bIsDirty = true;
 }
 
-UDistributionVectorConstant::UDistributionVectorConstant(const class FPostConstructInitializeProperties& PCIP)
-	: Super(PCIP)
+UDistributionVectorConstant::UDistributionVectorConstant(const FObjectInitializer& ObjectInitializer)
+	: Super(ObjectInitializer)
 {
 }
 
@@ -2562,8 +2562,8 @@ void UDistributionVectorConstant::GetRange(FVector& OutMin, FVector& OutMax) con
 	OutMax	= Constant;
 }
 
-UDistributionVectorConstantCurve::UDistributionVectorConstantCurve(const class FPostConstructInitializeProperties& PCIP)
-	: Super(PCIP)
+UDistributionVectorConstantCurve::UDistributionVectorConstantCurve(const FObjectInitializer& ObjectInitializer)
+	: Super(ObjectInitializer)
 {
 }
 
@@ -2911,8 +2911,8 @@ void UDistributionVectorConstantCurve::GetRange(FVector& OutMin, FVector& OutMax
 	OutMax = MaxVec;
 }
 
-UDistributionVectorUniform::UDistributionVectorUniform(const class FPostConstructInitializeProperties& PCIP)
-	: Super(PCIP)
+UDistributionVectorUniform::UDistributionVectorUniform(const FObjectInitializer& ObjectInitializer)
+	: Super(ObjectInitializer)
 {
 	MirrorFlags[0] = EDVMF_Different;
 	MirrorFlags[1] = EDVMF_Different;
@@ -3497,8 +3497,8 @@ void UDistributionVectorUniform::GetRange(FVector& OutMin, FVector& OutMax) cons
 }
 
 
-UDistributionVectorUniformCurve::UDistributionVectorUniformCurve(const class FPostConstructInitializeProperties& PCIP)
-	: Super(PCIP)
+UDistributionVectorUniformCurve::UDistributionVectorUniformCurve(const FObjectInitializer& ObjectInitializer)
+	: Super(ObjectInitializer)
 {
 	bLockAxes1 = false;
 	bLockAxes2 = false;
@@ -4068,8 +4068,8 @@ void UDistributionVectorUniformCurve::GetRange(FVector& OutMin, FVector& OutMax)
 }
 
 
-UDistributionFloatParameterBase::UDistributionFloatParameterBase(const class FPostConstructInitializeProperties& PCIP)
-	: Super(PCIP)
+UDistributionFloatParameterBase::UDistributionFloatParameterBase(const FObjectInitializer& ObjectInitializer)
+	: Super(ObjectInitializer)
 {
 	MaxInput = 1.0f;
 	MaxOutput = 1.0f;
@@ -4106,8 +4106,8 @@ float UDistributionFloatParameterBase::GetValue( float F, UObject* Data, class F
 }
 
 
-UDistributionVectorParameterBase::UDistributionVectorParameterBase(const class FPostConstructInitializeProperties& PCIP)
-	: Super(PCIP)
+UDistributionVectorParameterBase::UDistributionVectorParameterBase(const FObjectInitializer& ObjectInitializer)
+	: Super(ObjectInitializer)
 {
 	MaxInput = FVector(1.0f,1.0f,1.0f);
 	MaxOutput = FVector(1.0f,1.0f,1.0f);
diff --git a/Engine/Source/Runtime/Engine/Private/DocumentationActor.cpp b/Engine/Source/Runtime/Engine/Private/DocumentationActor.cpp
--- a/Engine/Source/Runtime/Engine/Private/DocumentationActor.cpp
+++ b/Engine/Source/Runtime/Engine/Private/DocumentationActor.cpp
@@ -7,15 +7,15 @@
 #include "IDocumentation.h"
 #endif
 
-ADocumentationActor::ADocumentationActor(const class FPostConstructInitializeProperties& PCIP)
-	: Super(PCIP)
+ADocumentationActor::ADocumentationActor(const FObjectInitializer& ObjectInitializer)
+	: Super(ObjectInitializer)
 {
-	TSubobjectPtr<USceneComponent> SceneComponent = PCIP.CreateDefaultSubobject<USceneComponent>(this, TEXT("SceneComp"));
+	TSubobjectPtr<USceneComponent> SceneComponent = ObjectInitializer.CreateDefaultSubobject<USceneComponent>(this, TEXT("SceneComp"));
 	RootComponent = SceneComponent;	
 
 #if WITH_EDITORONLY_DATA
 	// Create a Material billboard to represent our actor
-	Billboard = PCIP.CreateDefaultSubobject<UMaterialBillboardComponent>(this, TEXT("BillboardComponent"));
+	Billboard = ObjectInitializer.CreateDefaultSubobject<UMaterialBillboardComponent>(this, TEXT("BillboardComponent"));
 	if (!IsRunningCommandlet() && (Billboard != NULL))
 	{
 		static ConstructorHelpers::FObjectFinder<UMaterial> MaterialAsset(TEXT("/Engine/EditorMaterials/HelpActorMaterial"));
diff --git a/Engine/Source/Runtime/Engine/Private/DrawFrustumComponent.cpp b/Engine/Source/Runtime/Engine/Private/DrawFrustumComponent.cpp
--- a/Engine/Source/Runtime/Engine/Private/DrawFrustumComponent.cpp
+++ b/Engine/Source/Runtime/Engine/Private/DrawFrustumComponent.cpp
@@ -170,8 +170,8 @@ private:
 	float FrustumEndDist;
 };
 
-UDrawFrustumComponent::UDrawFrustumComponent(const class FPostConstructInitializeProperties& PCIP)
-	: Super(PCIP)
+UDrawFrustumComponent::UDrawFrustumComponent(const FObjectInitializer& ObjectInitializer)
+	: Super(ObjectInitializer)
 {
 	FrustumColor = FColor(255, 0, 255, 255);
 
diff --git a/Engine/Source/Runtime/Engine/Private/DrawSphereComponent.cpp b/Engine/Source/Runtime/Engine/Private/DrawSphereComponent.cpp
--- a/Engine/Source/Runtime/Engine/Private/DrawSphereComponent.cpp
+++ b/Engine/Source/Runtime/Engine/Private/DrawSphereComponent.cpp
@@ -2,8 +2,8 @@
 
 #include "EnginePrivate.h"
 
-UDrawSphereComponent::UDrawSphereComponent(const class FPostConstructInitializeProperties& PCIP)
-	: Super(PCIP)
+UDrawSphereComponent::UDrawSphereComponent(const FObjectInitializer& ObjectInitializer)
+	: Super(ObjectInitializer)
 {
 	BodyInstance.bEnableCollision_DEPRECATED = false;
 	SetCollisionProfileName(UCollisionProfile::NoCollision_ProfileName);
diff --git a/Engine/Source/Runtime/Engine/Private/DynamicBlueprintBinding.cpp b/Engine/Source/Runtime/Engine/Private/DynamicBlueprintBinding.cpp
--- a/Engine/Source/Runtime/Engine/Private/DynamicBlueprintBinding.cpp
+++ b/Engine/Source/Runtime/Engine/Private/DynamicBlueprintBinding.cpp
@@ -5,8 +5,8 @@
 #include "Engine/InputDelegateBinding.h"
 
 
-UDynamicBlueprintBinding::UDynamicBlueprintBinding(const class FPostConstructInitializeProperties& PCIP)
-	: Super(PCIP)
+UDynamicBlueprintBinding::UDynamicBlueprintBinding(const FObjectInitializer& ObjectInitializer)
+	: Super(ObjectInitializer)
 {
 }
 
diff --git a/Engine/Source/Runtime/Engine/Private/EdGraph/EdGraph.cpp b/Engine/Source/Runtime/Engine/Private/EdGraph/EdGraph.cpp
--- a/Engine/Source/Runtime/Engine/Private/EdGraph/EdGraph.cpp
+++ b/Engine/Source/Runtime/Engine/Private/EdGraph/EdGraph.cpp
@@ -78,8 +78,8 @@ UEdGraph* FGraphReference::GetGraph() const
 /////////////////////////////////////////////////////
 // UEdGraph
 
-UEdGraph::UEdGraph(const class FPostConstructInitializeProperties& PCIP)
-	: Super(PCIP)
+UEdGraph::UEdGraph(const FObjectInitializer& ObjectInitializer)
+	: Super(ObjectInitializer)
 {
 	bEditable = true;
 	bAllowDeletion = true;
diff --git a/Engine/Source/Runtime/Engine/Private/EdGraph/EdGraphNode.cpp b/Engine/Source/Runtime/Engine/Private/EdGraph/EdGraphNode.cpp
--- a/Engine/Source/Runtime/Engine/Private/EdGraph/EdGraphNode.cpp
+++ b/Engine/Source/Runtime/Engine/Private/EdGraph/EdGraphNode.cpp
@@ -39,8 +39,8 @@ FGraphNodeContextMenuBuilder::FGraphNodeContextMenuBuilder(const UEdGraph* InGra
 /////////////////////////////////////////////////////
 // UEdGraphNode
 
-UEdGraphNode::UEdGraphNode(const class FPostConstructInitializeProperties& PCIP)
-	: Super(PCIP)
+UEdGraphNode::UEdGraphNode(const FObjectInitializer& ObjectInitializer)
+	: Super(ObjectInitializer)
 	, AdvancedPinDisplay(ENodeAdvancedPins::NoPins)
 {
 
diff --git a/Engine/Source/Runtime/Engine/Private/EdGraph/EdGraphNode_Documentation.cpp b/Engine/Source/Runtime/Engine/Private/EdGraph/EdGraphNode_Documentation.cpp
--- a/Engine/Source/Runtime/Engine/Private/EdGraph/EdGraphNode_Documentation.cpp
+++ b/Engine/Source/Runtime/Engine/Private/EdGraph/EdGraphNode_Documentation.cpp
@@ -13,8 +13,8 @@
 /////////////////////////////////////////////////////
 // UEdGraphNode_Documentation
 
-UEdGraphNode_Documentation::UEdGraphNode_Documentation( const class FPostConstructInitializeProperties& PCIP )
-	: Super( PCIP )
+UEdGraphNode_Documentation::UEdGraphNode_Documentation( const FObjectInitializer& ObjectInitializer )
+	: Super( ObjectInitializer )
 {
 #if WITH_EDITORONLY_DATA
 	bCanResizeNode = true;
diff --git a/Engine/Source/Runtime/Engine/Private/EdGraph/EdGraphPin.cpp b/Engine/Source/Runtime/Engine/Private/EdGraph/EdGraphPin.cpp
--- a/Engine/Source/Runtime/Engine/Private/EdGraph/EdGraphPin.cpp
+++ b/Engine/Source/Runtime/Engine/Private/EdGraph/EdGraphPin.cpp
@@ -69,8 +69,8 @@ bool FEdGraphPinType::Serialize(FArchive& Ar)
 /////////////////////////////////////////////////////
 // UEdGraphPin
 
-UEdGraphPin::UEdGraphPin(const class FPostConstructInitializeProperties& PCIP)
-	: Super(PCIP)
+UEdGraphPin::UEdGraphPin(const FObjectInitializer& ObjectInitializer)
+	: Super(ObjectInitializer)
 {
 #if WITH_EDITORONLY_DATA
 	bHidden = false;
diff --git a/Engine/Source/Runtime/Engine/Private/EdGraph/EdGraphSchema.cpp b/Engine/Source/Runtime/Engine/Private/EdGraph/EdGraphSchema.cpp
--- a/Engine/Source/Runtime/Engine/Private/EdGraph/EdGraphSchema.cpp
+++ b/Engine/Source/Runtime/Engine/Private/EdGraph/EdGraphSchema.cpp
@@ -236,8 +236,8 @@ void FEdGraphSchemaAction_NewNode::AddReferencedObjects( FReferenceCollector& Co
 /////////////////////////////////////////////////////
 // UEdGraphSchema
 
-UEdGraphSchema::UEdGraphSchema(const class FPostConstructInitializeProperties& PCIP)
-	: Super(PCIP)
+UEdGraphSchema::UEdGraphSchema(const FObjectInitializer& ObjectInitializer)
+	: Super(ObjectInitializer)
 {
 }
 
diff --git a/Engine/Source/Runtime/Engine/Private/EditorFramework/AssetImportData.cpp b/Engine/Source/Runtime/Engine/Private/EditorFramework/AssetImportData.cpp
--- a/Engine/Source/Runtime/Engine/Private/EditorFramework/AssetImportData.cpp
+++ b/Engine/Source/Runtime/Engine/Private/EditorFramework/AssetImportData.cpp
@@ -2,8 +2,8 @@
 
 #include "EnginePrivate.h"
 
-UAssetImportData::UAssetImportData(const class FPostConstructInitializeProperties& PCIP)
-	: Super(PCIP)
+UAssetImportData::UAssetImportData(const FObjectInitializer& ObjectInitializer)
+	: Super(ObjectInitializer)
 {
 	bDirty = false;
 }
diff --git a/Engine/Source/Runtime/Engine/Private/EditorFramework/ThumbnailInfo.cpp b/Engine/Source/Runtime/Engine/Private/EditorFramework/ThumbnailInfo.cpp
--- a/Engine/Source/Runtime/Engine/Private/EditorFramework/ThumbnailInfo.cpp
+++ b/Engine/Source/Runtime/Engine/Private/EditorFramework/ThumbnailInfo.cpp
@@ -2,7 +2,7 @@
 
 #include "EnginePrivate.h"
 
-UThumbnailInfo::UThumbnailInfo(const class FPostConstructInitializeProperties& PCIP)
-	: Super(PCIP)
+UThumbnailInfo::UThumbnailInfo(const FObjectInitializer& ObjectInitializer)
+	: Super(ObjectInitializer)
 {
 }
diff --git a/Engine/Source/Runtime/Engine/Private/EngineMessage.cpp b/Engine/Source/Runtime/Engine/Private/EngineMessage.cpp
--- a/Engine/Source/Runtime/Engine/Private/EngineMessage.cpp
+++ b/Engine/Source/Runtime/Engine/Private/EngineMessage.cpp
@@ -3,8 +3,8 @@
 #include "EnginePrivate.h"
 #include "Engine/Console.h"
 
-UEngineMessage::UEngineMessage(const class FPostConstructInitializeProperties& PCIP) :
-	Super(PCIP)
+UEngineMessage::UEngineMessage(const FObjectInitializer& ObjectInitializer) :
+	Super(ObjectInitializer)
 {
 }
 
diff --git a/Engine/Source/Runtime/Engine/Private/EngineTypes.cpp b/Engine/Source/Runtime/Engine/Private/EngineTypes.cpp
--- a/Engine/Source/Runtime/Engine/Private/EngineTypes.cpp
+++ b/Engine/Source/Runtime/Engine/Private/EngineTypes.cpp
@@ -2,13 +2,13 @@
 
 #include "EnginePrivate.h"
 
-UEngineBaseTypes::UEngineBaseTypes(const class FPostConstructInitializeProperties& PCIP)
-	: Super(PCIP)
+UEngineBaseTypes::UEngineBaseTypes(const FObjectInitializer& ObjectInitializer)
+	: Super(ObjectInitializer)
 {
 }
 
-UEngineTypes::UEngineTypes(const class FPostConstructInitializeProperties& PCIP)
-	: Super(PCIP)
+UEngineTypes::UEngineTypes(const FObjectInitializer& ObjectInitializer)
+	: Super(ObjectInitializer)
 {
 }
 
diff --git a/Engine/Source/Runtime/Engine/Private/FloatingPawnMovement.cpp b/Engine/Source/Runtime/Engine/Private/FloatingPawnMovement.cpp
--- a/Engine/Source/Runtime/Engine/Private/FloatingPawnMovement.cpp
+++ b/Engine/Source/Runtime/Engine/Private/FloatingPawnMovement.cpp
@@ -3,8 +3,8 @@
 #include "EnginePrivate.h"
 #include "GameFramework/FloatingPawnMovement.h"
 
-UFloatingPawnMovement::UFloatingPawnMovement(const class FPostConstructInitializeProperties& PCIP)
-	: Super(PCIP)
+UFloatingPawnMovement::UFloatingPawnMovement(const FObjectInitializer& ObjectInitializer)
+	: Super(ObjectInitializer)
 {
 	MaxSpeed = 1200.f;
 	Acceleration = 4000.f;
diff --git a/Engine/Source/Runtime/Engine/Private/FoliageComponent.cpp b/Engine/Source/Runtime/Engine/Private/FoliageComponent.cpp
--- a/Engine/Source/Runtime/Engine/Private/FoliageComponent.cpp
+++ b/Engine/Source/Runtime/Engine/Private/FoliageComponent.cpp
@@ -48,8 +48,8 @@ protected:
 	FVector4 FoliageNormalizedRotationAxisAndAngle;
 };
 
-UInteractiveFoliageComponent::UInteractiveFoliageComponent(const class FPostConstructInitializeProperties& PCIP)
-	: Super(PCIP)
+UInteractiveFoliageComponent::UInteractiveFoliageComponent(const FObjectInitializer& ObjectInitializer)
+	: Super(ObjectInitializer)
 {
 }
 
diff --git a/Engine/Source/Runtime/Engine/Private/Font.cpp b/Engine/Source/Runtime/Engine/Private/Font.cpp
--- a/Engine/Source/Runtime/Engine/Private/Font.cpp
+++ b/Engine/Source/Runtime/Engine/Private/Font.cpp
@@ -7,13 +7,13 @@
 #include "EnginePrivate.h"
 
 
-UFontImportOptions::UFontImportOptions(const class FPostConstructInitializeProperties& PCIP)
-	: Super(PCIP)
+UFontImportOptions::UFontImportOptions(const FObjectInitializer& ObjectInitializer)
+	: Super(ObjectInitializer)
 {
 }
 
-UFont::UFont(const class FPostConstructInitializeProperties& PCIP)
-	: Super(PCIP)
+UFont::UFont(const FObjectInitializer& ObjectInitializer)
+	: Super(ObjectInitializer)
 {
 	ScalingFactor = 1.0f;
 }
diff --git a/Engine/Source/Runtime/Engine/Private/GameEngine.cpp b/Engine/Source/Runtime/Engine/Private/GameEngine.cpp
--- a/Engine/Source/Runtime/Engine/Private/GameEngine.cpp
+++ b/Engine/Source/Runtime/Engine/Private/GameEngine.cpp
@@ -91,8 +91,8 @@ EWindowMode::Type GetWindowModeType(EWindowMode::Type WindowMode)
 	return EWindowMode::Fullscreen;
 }
 
-UGameEngine::UGameEngine(const class FPostConstructInitializeProperties& PCIP)
-: Super(PCIP)
+UGameEngine::UGameEngine(const FObjectInitializer& ObjectInitializer)
+: Super(ObjectInitializer)
 {
 }
 
@@ -360,8 +360,8 @@ void UGameEngine::RedrawViewports( bool bShouldPresent /*= true*/ )
 /*-----------------------------------------------------------------------------
 	Game init and exit.
 -----------------------------------------------------------------------------*/
-UEngine::UEngine(const class FPostConstructInitializeProperties& PCIP)
-	: Super(PCIP)
+UEngine::UEngine(const FObjectInitializer& ObjectInitializer)
+	: Super(ObjectInitializer)
 {
 	AsyncLoadingTimeLimit = 5.0f;
 	bAsyncLoadingUseFullTimeLimit = true;
diff --git a/Engine/Source/Runtime/Engine/Private/GameFramework/SaveGame.cpp b/Engine/Source/Runtime/Engine/Private/GameFramework/SaveGame.cpp
--- a/Engine/Source/Runtime/Engine/Private/GameFramework/SaveGame.cpp
+++ b/Engine/Source/Runtime/Engine/Private/GameFramework/SaveGame.cpp
@@ -2,8 +2,8 @@
 
 #include "EnginePrivate.h"
 
-USaveGame::USaveGame(const class FPostConstructInitializeProperties& PCIP)
-	: Super(PCIP)
+USaveGame::USaveGame(const FObjectInitializer& ObjectInitializer)
+	: Super(ObjectInitializer)
 {
 
 }
\ No newline at end of file
diff --git a/Engine/Source/Runtime/Engine/Private/GameFramework/SpringArmComponent.cpp b/Engine/Source/Runtime/Engine/Private/GameFramework/SpringArmComponent.cpp
--- a/Engine/Source/Runtime/Engine/Private/GameFramework/SpringArmComponent.cpp
+++ b/Engine/Source/Runtime/Engine/Private/GameFramework/SpringArmComponent.cpp
@@ -10,8 +10,8 @@ const FName USpringArmComponent::SocketName(TEXT("SpringEndpoint"));
 
 extern float GAverageMS;
 
-USpringArmComponent::USpringArmComponent(const class FPostConstructInitializeProperties& PCIP)
-	: Super(PCIP)
+USpringArmComponent::USpringArmComponent(const FObjectInitializer& ObjectInitializer)
+	: Super(ObjectInitializer)
 {
 	PrimaryComponentTick.bCanEverTick = true;
 	PrimaryComponentTick.TickGroup = TG_PostPhysics;
diff --git a/Engine/Source/Runtime/Engine/Private/GameFramework/TouchInterface.cpp b/Engine/Source/Runtime/Engine/Private/GameFramework/TouchInterface.cpp
--- a/Engine/Source/Runtime/Engine/Private/GameFramework/TouchInterface.cpp
+++ b/Engine/Source/Runtime/Engine/Private/GameFramework/TouchInterface.cpp
@@ -4,8 +4,8 @@
 #include "GameFramework/TouchInterface.h"
 #include "Slate.h"
 
-UTouchInterface::UTouchInterface(const class FPostConstructInitializeProperties& PCIP)
-	: Super(PCIP)
+UTouchInterface::UTouchInterface(const FObjectInitializer& ObjectInitializer)
+	: Super(ObjectInitializer)
 {
 	// defaults
 	ActiveOpacity = 1.0f;
diff --git a/Engine/Source/Runtime/Engine/Private/GameInstance.cpp b/Engine/Source/Runtime/Engine/Private/GameInstance.cpp
--- a/Engine/Source/Runtime/Engine/Private/GameInstance.cpp
+++ b/Engine/Source/Runtime/Engine/Private/GameInstance.cpp
@@ -14,8 +14,8 @@
 #endif
 
 
-UGameInstance::UGameInstance(const class FPostConstructInitializeProperties& PCIP)
-: Super(PCIP)
+UGameInstance::UGameInstance(const FObjectInitializer& ObjectInitializer)
+: Super(ObjectInitializer)
 {
 }
 
diff --git a/Engine/Source/Runtime/Engine/Private/GameMode.cpp b/Engine/Source/Runtime/Engine/Private/GameMode.cpp
--- a/Engine/Source/Runtime/Engine/Private/GameMode.cpp
+++ b/Engine/Source/Runtime/Engine/Private/GameMode.cpp
@@ -25,8 +25,8 @@ namespace MatchState
 	const FName Aborted = FName(TEXT("Aborted"));
 }
 
-AGameMode::AGameMode(const class FPostConstructInitializeProperties& PCIP)
-	: Super(PCIP
+AGameMode::AGameMode(const FObjectInitializer& ObjectInitializer)
+	: Super(ObjectInitializer
 		.DoNotCreateDefaultSubobject(TEXT("Sprite"))
 	)
 {
diff --git a/Engine/Source/Runtime/Engine/Private/GameNetworkManager.cpp b/Engine/Source/Runtime/Engine/Private/GameNetworkManager.cpp
--- a/Engine/Source/Runtime/Engine/Private/GameNetworkManager.cpp
+++ b/Engine/Source/Runtime/Engine/Private/GameNetworkManager.cpp
@@ -10,8 +10,8 @@
 DEFINE_LOG_CATEGORY_STATIC(LogGameNetworkManager, Log, All);
 
 
-AGameNetworkManager::AGameNetworkManager(const class FPostConstructInitializeProperties& PCIP)
-	: Super(PCIP)
+AGameNetworkManager::AGameNetworkManager(const FObjectInitializer& ObjectInitializer)
+	: Super(ObjectInitializer)
 {
 	MoveRepSize = 42.0f;
 	MAXPOSITIONERRORSQUARED = 3.0;
diff --git a/Engine/Source/Runtime/Engine/Private/GameSession.cpp b/Engine/Source/Runtime/Engine/Private/GameSession.cpp
--- a/Engine/Source/Runtime/Engine/Private/GameSession.cpp
+++ b/Engine/Source/Runtime/Engine/Private/GameSession.cpp
@@ -37,8 +37,8 @@ APlayerController* GetPlayerControllerFromNetId(UWorld* World, const FUniqueNetI
 	return NULL;
 }
 
-AGameSession::AGameSession(const class FPostConstructInitializeProperties& PCIP)
-	: Super(PCIP)
+AGameSession::AGameSession(const FObjectInitializer& ObjectInitializer)
+	: Super(ObjectInitializer)
 {
 }
 
diff --git a/Engine/Source/Runtime/Engine/Private/GameState.cpp b/Engine/Source/Runtime/Engine/Private/GameState.cpp
--- a/Engine/Source/Runtime/Engine/Private/GameState.cpp
+++ b/Engine/Source/Runtime/Engine/Private/GameState.cpp
@@ -9,8 +9,8 @@
 
 DEFINE_LOG_CATEGORY_STATIC(LogGameState, Log, All);
 
-AGameState::AGameState(const class FPostConstructInitializeProperties& PCIP)
-	: Super(PCIP
+AGameState::AGameState(const FObjectInitializer& ObjectInitializer)
+	: Super(ObjectInitializer
 		.DoNotCreateDefaultSubobject(TEXT("Sprite")) )
 {
 	SetRemoteRoleForBackwardsCompat(ROLE_SimulatedProxy);
diff --git a/Engine/Source/Runtime/Engine/Private/GameUserSettings.cpp b/Engine/Source/Runtime/Engine/Private/GameUserSettings.cpp
--- a/Engine/Source/Runtime/Engine/Private/GameUserSettings.cpp
+++ b/Engine/Source/Runtime/Engine/Private/GameUserSettings.cpp
@@ -14,8 +14,8 @@ enum EGameUserSettingsVersion
 };
 
 
-UGameUserSettings::UGameUserSettings(const class FPostConstructInitializeProperties& PCIP)
-:	Super(PCIP)
+UGameUserSettings::UGameUserSettings(const FObjectInitializer& ObjectInitializer)
+:	Super(ObjectInitializer)
 {
 	SetToDefaults();
 }
diff --git a/Engine/Source/Runtime/Engine/Private/GameViewportClient.cpp b/Engine/Source/Runtime/Engine/Private/GameViewportClient.cpp
--- a/Engine/Source/Runtime/Engine/Private/GameViewportClient.cpp
+++ b/Engine/Source/Runtime/Engine/Private/GameViewportClient.cpp
@@ -79,8 +79,8 @@ public:
 	}
 };
 
-UGameViewportClient::UGameViewportClient(const class FPostConstructInitializeProperties& PCIP)
-	: Super(PCIP)
+UGameViewportClient::UGameViewportClient(const FObjectInitializer& ObjectInitializer)
+	: Super(ObjectInitializer)
 	, EngineShowFlags(ESFIM_Game)
 	, CurrentBufferVisualizationMode(NAME_None)
 	, HighResScreenshotDialog(NULL)
diff --git a/Engine/Source/Runtime/Engine/Private/GameplayStatics.cpp b/Engine/Source/Runtime/Engine/Private/GameplayStatics.cpp
--- a/Engine/Source/Runtime/Engine/Private/GameplayStatics.cpp
+++ b/Engine/Source/Runtime/Engine/Private/GameplayStatics.cpp
@@ -16,8 +16,8 @@
 //////////////////////////////////////////////////////////////////////////
 // UGameplayStatics
 
-UGameplayStatics::UGameplayStatics(const class FPostConstructInitializeProperties& PCIP)
-	: Super(PCIP)
+UGameplayStatics::UGameplayStatics(const FObjectInitializer& ObjectInitializer)
+	: Super(ObjectInitializer)
 {
 }
 
diff --git a/Engine/Source/Runtime/Engine/Private/HUD.cpp b/Engine/Source/Runtime/Engine/Private/HUD.cpp
--- a/Engine/Source/Runtime/Engine/Private/HUD.cpp
+++ b/Engine/Source/Runtime/Engine/Private/HUD.cpp
@@ -16,8 +16,8 @@ DEFINE_LOG_CATEGORY_STATIC(LogHUD, Log, All);
 
 #define LOCTEXT_NAMESPACE "HUD"
 
-AHUD::AHUD(const class FPostConstructInitializeProperties& PCIP)
-	: Super(PCIP)
+AHUD::AHUD(const FObjectInitializer& ObjectInitializer)
+	: Super(ObjectInitializer)
 {
 	PrimaryActorTick.TickGroup = TG_DuringPhysics;
 	PrimaryActorTick.bCanEverTick = true;
diff --git a/Engine/Source/Runtime/Engine/Private/HeadMountedDisplayFunctionLibrary.cpp b/Engine/Source/Runtime/Engine/Private/HeadMountedDisplayFunctionLibrary.cpp
--- a/Engine/Source/Runtime/Engine/Private/HeadMountedDisplayFunctionLibrary.cpp
+++ b/Engine/Source/Runtime/Engine/Private/HeadMountedDisplayFunctionLibrary.cpp
@@ -3,8 +3,8 @@
 #include "EnginePrivate.h"
 #include "HeadMountedDisplay.h"
 
-UHeadMountedDisplayFunctionLibrary::UHeadMountedDisplayFunctionLibrary(const FPostConstructInitializeProperties& PCIP)
-	: Super(PCIP)
+UHeadMountedDisplayFunctionLibrary::UHeadMountedDisplayFunctionLibrary(const FObjectInitializer& ObjectInitializer)
+	: Super(ObjectInitializer)
 {
 }
 
diff --git a/Engine/Source/Runtime/Engine/Private/HeightFogComponent.cpp b/Engine/Source/Runtime/Engine/Private/HeightFogComponent.cpp
--- a/Engine/Source/Runtime/Engine/Private/HeightFogComponent.cpp
+++ b/Engine/Source/Runtime/Engine/Private/HeightFogComponent.cpp
@@ -8,8 +8,8 @@
 #include "Net/UnrealNetwork.h"
 #include "Engine/ExponentialHeightFog.h"
 
-UExponentialHeightFogComponent::UExponentialHeightFogComponent(const class FPostConstructInitializeProperties& PCIP)
-	: Super(PCIP)
+UExponentialHeightFogComponent::UExponentialHeightFogComponent(const FObjectInitializer& ObjectInitializer)
+	: Super(ObjectInitializer)
 {
 	bEnabled_DEPRECATED = true;
 
@@ -157,10 +157,10 @@ void UExponentialHeightFogComponent::SetStartDistance(float Value)
 //////////////////////////////////////////////////////////////////////////
 // AExponentialHeightFog
 
-AExponentialHeightFog::AExponentialHeightFog(const class FPostConstructInitializeProperties& PCIP)
-	: Super(PCIP)
+AExponentialHeightFog::AExponentialHeightFog(const FObjectInitializer& ObjectInitializer)
+	: Super(ObjectInitializer)
 {
-	Component = PCIP.CreateDefaultSubobject<UExponentialHeightFogComponent>(this, TEXT("HeightFogComponent0"));
+	Component = ObjectInitializer.CreateDefaultSubobject<UExponentialHeightFogComponent>(this, TEXT("HeightFogComponent0"));
 	RootComponent = Component;
 
 	bHidden = false;
diff --git a/Engine/Source/Runtime/Engine/Private/InGameAdvertising.cpp b/Engine/Source/Runtime/Engine/Private/InGameAdvertising.cpp
--- a/Engine/Source/Runtime/Engine/Private/InGameAdvertising.cpp
+++ b/Engine/Source/Runtime/Engine/Private/InGameAdvertising.cpp
@@ -6,8 +6,8 @@
 
 #include "EnginePrivate.h"
 
-UInGameAdManager::UInGameAdManager(const class FPostConstructInitializeProperties& PCIP)
-	: Super(PCIP)
+UInGameAdManager::UInGameAdManager(const FObjectInitializer& ObjectInitializer)
+	: Super(ObjectInitializer)
 {
 	bShouldPauseWhileAdOpen = true;
 }
diff --git a/Engine/Source/Runtime/Engine/Private/Info.cpp b/Engine/Source/Runtime/Engine/Private/Info.cpp
--- a/Engine/Source/Runtime/Engine/Private/Info.cpp
+++ b/Engine/Source/Runtime/Engine/Private/Info.cpp
@@ -2,11 +2,11 @@
 
 #include "EnginePrivate.h"
 
-AInfo::AInfo(const class FPostConstructInitializeProperties& PCIP)
-	: Super(PCIP)
+AInfo::AInfo(const FObjectInitializer& ObjectInitializer)
+	: Super(ObjectInitializer)
 {
 #if WITH_EDITORONLY_DATA
-	SpriteComponent = PCIP.CreateEditorOnlyDefaultSubobject<UBillboardComponent>(this, TEXT("Sprite"));
+	SpriteComponent = ObjectInitializer.CreateEditorOnlyDefaultSubobject<UBillboardComponent>(this, TEXT("Sprite"));
 	if (!IsRunningCommandlet() && (SpriteComponent != nullptr))
 	{
 		// Structure to hold one-time initialization
diff --git a/Engine/Source/Runtime/Engine/Private/InputActionDelegateBinding.cpp b/Engine/Source/Runtime/Engine/Private/InputActionDelegateBinding.cpp
--- a/Engine/Source/Runtime/Engine/Private/InputActionDelegateBinding.cpp
+++ b/Engine/Source/Runtime/Engine/Private/InputActionDelegateBinding.cpp
@@ -4,8 +4,8 @@
 #include "BlueprintUtilities.h"
 #include "Engine/InputActionDelegateBinding.h"
 
-UInputActionDelegateBinding::UInputActionDelegateBinding(const class FPostConstructInitializeProperties& PCIP)
-	: Super(PCIP)
+UInputActionDelegateBinding::UInputActionDelegateBinding(const FObjectInitializer& ObjectInitializer)
+	: Super(ObjectInitializer)
 {
 }
 
diff --git a/Engine/Source/Runtime/Engine/Private/InputAxisDelegateBinding.cpp b/Engine/Source/Runtime/Engine/Private/InputAxisDelegateBinding.cpp
--- a/Engine/Source/Runtime/Engine/Private/InputAxisDelegateBinding.cpp
+++ b/Engine/Source/Runtime/Engine/Private/InputAxisDelegateBinding.cpp
@@ -4,8 +4,8 @@
 #include "BlueprintUtilities.h"
 #include "Engine/InputAxisDelegateBinding.h"
 
-UInputAxisDelegateBinding::UInputAxisDelegateBinding(const class FPostConstructInitializeProperties& PCIP)
-	: Super(PCIP)
+UInputAxisDelegateBinding::UInputAxisDelegateBinding(const FObjectInitializer& ObjectInitializer)
+	: Super(ObjectInitializer)
 {
 }
 
diff --git a/Engine/Source/Runtime/Engine/Private/InputAxisKeyDelegateBinding.cpp b/Engine/Source/Runtime/Engine/Private/InputAxisKeyDelegateBinding.cpp
--- a/Engine/Source/Runtime/Engine/Private/InputAxisKeyDelegateBinding.cpp
+++ b/Engine/Source/Runtime/Engine/Private/InputAxisKeyDelegateBinding.cpp
@@ -5,8 +5,8 @@
 #include "Engine/InputAxisKeyDelegateBinding.h"
 
 
-UInputAxisKeyDelegateBinding::UInputAxisKeyDelegateBinding(const class FPostConstructInitializeProperties& PCIP)
-	: Super(PCIP)
+UInputAxisKeyDelegateBinding::UInputAxisKeyDelegateBinding(const FObjectInitializer& ObjectInitializer)
+	: Super(ObjectInitializer)
 {
 }
 
diff --git a/Engine/Source/Runtime/Engine/Private/InputComponent.cpp b/Engine/Source/Runtime/Engine/Private/InputComponent.cpp
--- a/Engine/Source/Runtime/Engine/Private/InputComponent.cpp
+++ b/Engine/Source/Runtime/Engine/Private/InputComponent.cpp
@@ -42,8 +42,8 @@ FInputChord::RelationshipType FInputChord::GetRelationship( const FInputChord& O
 /* UInputComponent interface
  *****************************************************************************/
 
-UInputComponent::UInputComponent( const class FPostConstructInitializeProperties& PCIP )
-	: Super(PCIP)
+UInputComponent::UInputComponent( const FObjectInitializer& ObjectInitializer )
+	: Super(ObjectInitializer)
 {
 	bBlockInput = false;
 }
diff --git a/Engine/Source/Runtime/Engine/Private/InputDelegateBinding.cpp b/Engine/Source/Runtime/Engine/Private/InputDelegateBinding.cpp
--- a/Engine/Source/Runtime/Engine/Private/InputDelegateBinding.cpp
+++ b/Engine/Source/Runtime/Engine/Private/InputDelegateBinding.cpp
@@ -10,8 +10,8 @@
 #include "Engine/InputAxisKeyDelegateBinding.h"
 #include "Engine/InputVectorAxisDelegateBinding.h"
 
-UInputDelegateBinding::UInputDelegateBinding(const class FPostConstructInitializeProperties& PCIP)
-	: Super(PCIP)
+UInputDelegateBinding::UInputDelegateBinding(const FObjectInitializer& ObjectInitializer)
+	: Super(ObjectInitializer)
 {
 }
 
diff --git a/Engine/Source/Runtime/Engine/Private/InputKeyDelegateBinding.cpp b/Engine/Source/Runtime/Engine/Private/InputKeyDelegateBinding.cpp
--- a/Engine/Source/Runtime/Engine/Private/InputKeyDelegateBinding.cpp
+++ b/Engine/Source/Runtime/Engine/Private/InputKeyDelegateBinding.cpp
@@ -4,8 +4,8 @@
 #include "BlueprintUtilities.h"
 #include "Engine/InputKeyDelegateBinding.h"
 
-UInputKeyDelegateBinding::UInputKeyDelegateBinding(const class FPostConstructInitializeProperties& PCIP)
-	: Super(PCIP)
+UInputKeyDelegateBinding::UInputKeyDelegateBinding(const FObjectInitializer& ObjectInitializer)
+	: Super(ObjectInitializer)
 {
 }
 
diff --git a/Engine/Source/Runtime/Engine/Private/InputTouchDelegateBinding.cpp b/Engine/Source/Runtime/Engine/Private/InputTouchDelegateBinding.cpp
--- a/Engine/Source/Runtime/Engine/Private/InputTouchDelegateBinding.cpp
+++ b/Engine/Source/Runtime/Engine/Private/InputTouchDelegateBinding.cpp
@@ -4,8 +4,8 @@
 #include "BlueprintUtilities.h"
 #include "Engine/InputTouchDelegateBinding.h"
 
-UInputTouchDelegateBinding::UInputTouchDelegateBinding(const class FPostConstructInitializeProperties& PCIP)
-	: Super(PCIP)
+UInputTouchDelegateBinding::UInputTouchDelegateBinding(const FObjectInitializer& ObjectInitializer)
+	: Super(ObjectInitializer)
 {
 }
 
diff --git a/Engine/Source/Runtime/Engine/Private/InputVectorAxisDelegateBinding.cpp b/Engine/Source/Runtime/Engine/Private/InputVectorAxisDelegateBinding.cpp
--- a/Engine/Source/Runtime/Engine/Private/InputVectorAxisDelegateBinding.cpp
+++ b/Engine/Source/Runtime/Engine/Private/InputVectorAxisDelegateBinding.cpp
@@ -4,8 +4,8 @@
 #include "BlueprintUtilities.h"
 #include "Engine/InputVectorAxisDelegateBinding.h"
 
-UInputVectorAxisDelegateBinding::UInputVectorAxisDelegateBinding(const class FPostConstructInitializeProperties& PCIP)
-	: Super(PCIP)
+UInputVectorAxisDelegateBinding::UInputVectorAxisDelegateBinding(const FObjectInitializer& ObjectInitializer)
+	: Super(ObjectInitializer)
 {
 }
 
diff --git a/Engine/Source/Runtime/Engine/Private/InstancedFoliage.cpp b/Engine/Source/Runtime/Engine/Private/InstancedFoliage.cpp
--- a/Engine/Source/Runtime/Engine/Private/InstancedFoliage.cpp
+++ b/Engine/Source/Runtime/Engine/Private/InstancedFoliage.cpp
@@ -89,8 +89,8 @@ FArchive& operator<<(FArchive& Ar, FFoliageComponentHashInfo& ComponentHashInfo)
 // UFoliageType
 //
 
-UFoliageType::UFoliageType(const FPostConstructInitializeProperties& PCIP)
-	: Super(PCIP)
+UFoliageType::UFoliageType(const FObjectInitializer& ObjectInitializer)
+	: Super(ObjectInitializer)
 {
 	Density = 100.0f;
 	Radius = 0.0f;
@@ -135,8 +135,8 @@ UFoliageType::UFoliageType(const FPostConstructInitializeProperties& PCIP)
 	BodyInstance.SetCollisionProfileName(UCollisionProfile::NoCollision_ProfileName);
 }
 
-UFoliageType_InstancedStaticMesh::UFoliageType_InstancedStaticMesh(const FPostConstructInitializeProperties& PCIP)
-	: Super(PCIP)
+UFoliageType_InstancedStaticMesh::UFoliageType_InstancedStaticMesh(const FObjectInitializer& ObjectInitializer)
+	: Super(ObjectInitializer)
 {
 	Mesh = nullptr;
 }
@@ -705,10 +705,10 @@ void FFoliageMeshInfo::SelectInstances(AInstancedFoliageActor* InIFA, bool bSele
 //
 // AInstancedFoliageActor
 //
-AInstancedFoliageActor::AInstancedFoliageActor(const FPostConstructInitializeProperties& PCIP)
-	: Super(PCIP)
+AInstancedFoliageActor::AInstancedFoliageActor(const FObjectInitializer& ObjectInitializer)
+	: Super(ObjectInitializer)
 {
-	TSubobjectPtr<USceneComponent> SceneComponent = PCIP.CreateDefaultSubobject<USceneComponent>(this, TEXT("RootComponent0"));
+	TSubobjectPtr<USceneComponent> SceneComponent = ObjectInitializer.CreateDefaultSubobject<USceneComponent>(this, TEXT("RootComponent0"));
 	RootComponent = SceneComponent;
 	RootComponent->Mobility = EComponentMobility::Movable;
 	
diff --git a/Engine/Source/Runtime/Engine/Private/InstancedStaticMesh.cpp b/Engine/Source/Runtime/Engine/Private/InstancedStaticMesh.cpp
--- a/Engine/Source/Runtime/Engine/Private/InstancedStaticMesh.cpp
+++ b/Engine/Source/Runtime/Engine/Private/InstancedStaticMesh.cpp
@@ -1211,8 +1211,8 @@ bool FInstancedStaticMeshSceneProxy::GetWireframeMeshElement(int32 LODIndex, int
 	UInstancedStaticMeshComponent
 -----------------------------------------------------------------------------*/
 
-UInstancedStaticMeshComponent::UInstancedStaticMeshComponent(const class FPostConstructInitializeProperties& PCIP)
-	: Super(PCIP)
+UInstancedStaticMeshComponent::UInstancedStaticMeshComponent(const FObjectInitializer& ObjectInitializer)
+	: Super(ObjectInitializer)
 {
 	Mobility = EComponentMobility::Movable;
 	BodyInstance.bSimulatePhysics = false;
diff --git a/Engine/Source/Runtime/Engine/Private/IntSerialization.cpp b/Engine/Source/Runtime/Engine/Private/IntSerialization.cpp
--- a/Engine/Source/Runtime/Engine/Private/IntSerialization.cpp
+++ b/Engine/Source/Runtime/Engine/Private/IntSerialization.cpp
@@ -7,8 +7,8 @@
 #include "EnginePrivate.h"
 #include "Runtime/Engine/Classes/Engine/IntSerialization.h"
 
-UIntSerialization::UIntSerialization(const class FPostConstructInitializeProperties& PCIP)
-	: UObject(PCIP)
+UIntSerialization::UIntSerialization(const FObjectInitializer& ObjectInitializer)
+	: UObject(ObjectInitializer)
 {
 }
 
diff --git a/Engine/Source/Runtime/Engine/Private/InteractiveFoliageActor.cpp b/Engine/Source/Runtime/Engine/Private/InteractiveFoliageActor.cpp
--- a/Engine/Source/Runtime/Engine/Private/InteractiveFoliageActor.cpp
+++ b/Engine/Source/Runtime/Engine/Private/InteractiveFoliageActor.cpp
@@ -5,8 +5,8 @@
 #include "Engine/InteractiveFoliageActor.h"
 #include "Components/InteractiveFoliageComponent.h"
 
-AInteractiveFoliageActor::AInteractiveFoliageActor(const class FPostConstructInitializeProperties& PCIP)
-	: Super(PCIP
+AInteractiveFoliageActor::AInteractiveFoliageActor(const FObjectInitializer& ObjectInitializer)
+	: Super(ObjectInitializer
 		.SetDefaultSubobjectClass<UInteractiveFoliageComponent>("StaticMeshComponent0"))
 {
 
@@ -15,7 +15,7 @@ AInteractiveFoliageActor::AInteractiveFoliageActor(const class FPostConstructIni
 	FoliageMeshComponent->SetCollisionProfileName(UCollisionProfile::NoCollision_ProfileName);
 	FoliageMeshComponent->Mobility = EComponentMobility::Static;
 
-	CapsuleComponent = PCIP.CreateDefaultSubobject<UCapsuleComponent>(this, TEXT("CollisionCylinder"));
+	CapsuleComponent = ObjectInitializer.CreateDefaultSubobject<UCapsuleComponent>(this, TEXT("CollisionCylinder"));
 	CapsuleComponent->InitCapsuleSize(60.0f, 200.0f);
 	CapsuleComponent->BodyInstance.bEnableCollision_DEPRECATED = true;
 	static FName CollisionProfileName(TEXT("OverlapAllDynamic"));
diff --git a/Engine/Source/Runtime/Engine/Private/Interface.cpp b/Engine/Source/Runtime/Engine/Private/Interface.cpp
--- a/Engine/Source/Runtime/Engine/Private/Interface.cpp
+++ b/Engine/Source/Runtime/Engine/Private/Interface.cpp
@@ -2,19 +2,19 @@
 
 #include "EnginePrivate.h"
 
-UInterface_CollisionDataProvider::UInterface_CollisionDataProvider(const class FPostConstructInitializeProperties& PCIP)
-	: Super(PCIP)
+UInterface_CollisionDataProvider::UInterface_CollisionDataProvider(const FObjectInitializer& ObjectInitializer)
+	: Super(ObjectInitializer)
 {
 }
 
-URuntimeMovieScenePlayerInterface::URuntimeMovieScenePlayerInterface( const class FPostConstructInitializeProperties& PCIP )
-	: Super(PCIP)
+URuntimeMovieScenePlayerInterface::URuntimeMovieScenePlayerInterface( const FObjectInitializer& ObjectInitializer )
+	: Super(ObjectInitializer)
 {
 
 }
 
-UInterface_PostProcessVolume::UInterface_PostProcessVolume( const class FPostConstructInitializeProperties& PCIP )
-	: Super(PCIP)
+UInterface_PostProcessVolume::UInterface_PostProcessVolume( const FObjectInitializer& ObjectInitializer )
+	: Super(ObjectInitializer)
 {
 
 }
diff --git a/Engine/Source/Runtime/Engine/Private/InterpCurveEdSetup.cpp b/Engine/Source/Runtime/Engine/Private/InterpCurveEdSetup.cpp
--- a/Engine/Source/Runtime/Engine/Private/InterpCurveEdSetup.cpp
+++ b/Engine/Source/Runtime/Engine/Private/InterpCurveEdSetup.cpp
@@ -49,8 +49,8 @@ void UInterpCurveEdSetup::Serialize(FArchive& Ar)
 	}
 }
 
-UInterpCurveEdSetup::UInterpCurveEdSetup(const class FPostConstructInitializeProperties& PCIP)
-	: Super(PCIP)
+UInterpCurveEdSetup::UInterpCurveEdSetup(const FObjectInitializer& ObjectInitializer)
+	: Super(ObjectInitializer)
 {
 
 	Tabs.Add(FCurveEdTab(TEXT("Default"), 0.0f, 1.0f, -1.0f, 1.0f));
diff --git a/Engine/Source/Runtime/Engine/Private/Interpolation.cpp b/Engine/Source/Runtime/Engine/Private/Interpolation.cpp
--- a/Engine/Source/Runtime/Engine/Private/Interpolation.cpp
+++ b/Engine/Source/Runtime/Engine/Private/Interpolation.cpp
@@ -317,14 +317,14 @@ bool AMatineeActor::IgnoreActorSelection()
 	return IgnoreActorSelectionCount > 0;
 }
 
-AMatineeActor::AMatineeActor(const class FPostConstructInitializeProperties& PCIP)
-	: Super(PCIP)
+AMatineeActor::AMatineeActor(const FObjectInitializer& ObjectInitializer)
+	: Super(ObjectInitializer)
 {
-	TSubobjectPtr<USceneComponent> SceneComponent = PCIP.CreateDefaultSubobject<USceneComponent>(this, TEXT("SceneComp"));
+	TSubobjectPtr<USceneComponent> SceneComponent = ObjectInitializer.CreateDefaultSubobject<USceneComponent>(this, TEXT("SceneComp"));
 	RootComponent = SceneComponent;
 
 #if WITH_EDITORONLY_DATA
-	SpriteComponent = PCIP.CreateEditorOnlyDefaultSubobject<UBillboardComponent>(this, TEXT("Sprite"));
+	SpriteComponent = ObjectInitializer.CreateEditorOnlyDefaultSubobject<UBillboardComponent>(this, TEXT("Sprite"));
 	if (!IsRunningCommandlet() && (SpriteComponent != nullptr))
 	{
 		// Structure to hold one-time initialization
@@ -1734,8 +1734,8 @@ void AMatineeActor::CheckPriorityRefresh()
   AMatineeActorCameraAnim
 -----------------------------------------------------------------------------*/
 
-AMatineeActorCameraAnim::AMatineeActorCameraAnim(const class FPostConstructInitializeProperties& PCIP)
-	: Super(PCIP)
+AMatineeActorCameraAnim::AMatineeActorCameraAnim(const FObjectInitializer& ObjectInitializer)
+	: Super(ObjectInitializer)
 {
 #if WITH_EDITORONLY_DATA
 	bEditable = false;
@@ -1746,8 +1746,8 @@ AMatineeActorCameraAnim::AMatineeActorCameraAnim(const class FPostConstructIniti
 /*-----------------------------------------------------------------------------
   UInterpData
 -----------------------------------------------------------------------------*/
-UInterpData::UInterpData(const class FPostConstructInitializeProperties& PCIP)
-	: Super(PCIP)
+UInterpData::UInterpData(const FObjectInitializer& ObjectInitializer)
+	: Super(ObjectInitializer)
 {
 	InterpLength = 5.0f;
 	EdSectionStart = 0.0f;
@@ -2022,8 +2022,8 @@ void UInterpData::UpdateBakeAndPruneStatus()
 /*-----------------------------------------------------------------------------
  UInterpGroup
 -----------------------------------------------------------------------------*/
-UInterpGroup::UInterpGroup(const class FPostConstructInitializeProperties& PCIP)
-	: Super(PCIP)
+UInterpGroup::UInterpGroup(const FObjectInitializer& ObjectInitializer)
+	: Super(ObjectInitializer)
 {
 	// Structure to hold one-time initialization
 	struct FConstructorStatics
@@ -2365,8 +2365,8 @@ AActor* UInterpGroup::DeselectGroupActor( UInterpGroupInst* GrInst )
 /*-----------------------------------------------------------------------------
 	UInterpFilter
 -----------------------------------------------------------------------------*/
-UInterpFilter::UInterpFilter(const class FPostConstructInitializeProperties& PCIP)
-	: Super(PCIP)
+UInterpFilter::UInterpFilter(const FObjectInitializer& ObjectInitializer)
+	: Super(ObjectInitializer)
 {
 }
 
@@ -2390,8 +2390,8 @@ void UInterpFilter::FilterData(class AMatineeActor* InMatineeActor)
 /*-----------------------------------------------------------------------------
 	UInterpFilter_Classes
 -----------------------------------------------------------------------------*/
-UInterpFilter_Classes::UInterpFilter_Classes(const class FPostConstructInitializeProperties& PCIP)
-	: Super(PCIP)
+UInterpFilter_Classes::UInterpFilter_Classes(const FObjectInitializer& ObjectInitializer)
+	: Super(ObjectInitializer)
 {
 }
 
@@ -2505,8 +2505,8 @@ void UInterpFilter_Classes::FilterData(class AMatineeActor* InMatineeActor)
 /*-----------------------------------------------------------------------------
 	UInterpFilter_Custom
 -----------------------------------------------------------------------------*/
-UInterpFilter_Custom::UInterpFilter_Custom(const class FPostConstructInitializeProperties& PCIP)
-	: Super(PCIP)
+UInterpFilter_Custom::UInterpFilter_Custom(const FObjectInitializer& ObjectInitializer)
+	: Super(ObjectInitializer)
 {
 }
 
@@ -2532,8 +2532,8 @@ void UInterpFilter_Custom::FilterData(AMatineeActor* InMatineeActor)
 /*-----------------------------------------------------------------------------
  UInterpGroupInst
 -----------------------------------------------------------------------------*/
-UInterpGroupInst::UInterpGroupInst(const class FPostConstructInitializeProperties& PCIP)
-	: Super(PCIP)
+UInterpGroupInst::UInterpGroupInst(const FObjectInitializer& ObjectInitializer)
+	: Super(ObjectInitializer)
 {
 }
 
@@ -2658,8 +2658,8 @@ void UInterpGroupInst::TermGroupInst(bool bDeleteTrackInst)
 /*-----------------------------------------------------------------------------
  UInterpGroupDirector
 -----------------------------------------------------------------------------*/
-UInterpGroupDirector::UInterpGroupDirector(const class FPostConstructInitializeProperties& PCIP)
-	: Super(PCIP)
+UInterpGroupDirector::UInterpGroupDirector(const FObjectInitializer& ObjectInitializer)
+	: Super(ObjectInitializer)
 {
 	// Structure to hold one-time initialization
 	struct FConstructorStatics
@@ -2899,8 +2899,8 @@ FString FInterpEdSelKey::GetOwningTrackSubGroupName( int32* piSubTrack )
 /*-----------------------------------------------------------------------------
  UInterpTrack
 -----------------------------------------------------------------------------*/
-UInterpTrack::UInterpTrack(const class FPostConstructInitializeProperties& PCIP)
-	: Super(PCIP)
+UInterpTrack::UInterpTrack(const FObjectInitializer& ObjectInitializer)
+	: Super(ObjectInitializer)
 {
 	TrackInstClass = UInterpTrackInst::StaticClass();
 	ActiveCondition = ETAC_Always;
@@ -3007,8 +3007,8 @@ void UInterpTrack::EnableTrack( bool bInEnable, bool bPropagateToSubTracks )
 /*-----------------------------------------------------------------------------
  UInterpTrackInst
 -----------------------------------------------------------------------------*/
-UInterpTrackInst::UInterpTrackInst(const class FPostConstructInitializeProperties& PCIP)
-	: Super(PCIP)
+UInterpTrackInst::UInterpTrackInst(const FObjectInitializer& ObjectInitializer)
+	: Super(ObjectInitializer)
 {
 }
 
@@ -3037,8 +3037,8 @@ UWorld* UInterpTrackInst::GetWorld() const
 	UInterpTrackInstProperty
 -----------------------------------------------------------------------------*/
 
-UInterpTrackInstProperty::UInterpTrackInstProperty(const class FPostConstructInitializeProperties& PCIP)
-	: Super(PCIP)
+UInterpTrackInstProperty::UInterpTrackInstProperty(const FObjectInitializer& ObjectInitializer)
+	: Super(ObjectInitializer)
 {
 }
 
@@ -3076,8 +3076,8 @@ void UInterpTrackInstProperty::TermTrackInst(UInterpTrack* Track)
 /*-----------------------------------------------------------------------------
  UInterpTrackMoveAxis
 -----------------------------------------------------------------------------*/
-UInterpTrackMoveAxis::UInterpTrackMoveAxis(const class FPostConstructInitializeProperties& PCIP)
-	: Super(PCIP)
+UInterpTrackMoveAxis::UInterpTrackMoveAxis(const FObjectInitializer& ObjectInitializer)
+	: Super(ObjectInitializer)
 {
 	CurveTension = 0.0f;
 	bSubTrackOnly = true;
@@ -3396,8 +3396,8 @@ void UInterpTrackMoveAxis::ReduceKeys( float IntervalStart, float IntervalEnd, f
 /*----------------------------------------------------------------------------
  UInterpTrackMove
 -----------------------------------------------------------------------------*/
-UInterpTrackMove::UInterpTrackMove(const class FPostConstructInitializeProperties& PCIP)
-	: Super(PCIP)
+UInterpTrackMove::UInterpTrackMove(const FObjectInitializer& ObjectInitializer)
+	: Super(ObjectInitializer)
 {
 	TrackInstClass = UInterpTrackInstMove::StaticClass();
 	bOnePerGroup = true;
@@ -4777,8 +4777,8 @@ FTransform UInterpTrackMove::GetMoveRefFrame(UInterpTrackInstMove* MoveTrackInst
 /*-----------------------------------------------------------------------------
  UInterpTrackInstMove
 -----------------------------------------------------------------------------*/
-UInterpTrackInstMove::UInterpTrackInstMove(const class FPostConstructInitializeProperties& PCIP)
-	: Super(PCIP)
+UInterpTrackInstMove::UInterpTrackInstMove(const FObjectInitializer& ObjectInitializer)
+	: Super(ObjectInitializer)
 {
 }
 
@@ -4976,8 +4976,8 @@ STRUCTTRACK_GETKEYFRAMETIME(UInterpTrackFloatBase, FloatTrack.Points, InVal)
 STRUCTTRACK_GETKEYFRAMEINDEX(UInterpTrackFloatBase, FloatTrack.Points, InVal)
 STRUCTTRACK_GETCLOSESTSNAPPOSITION(UInterpTrackFloatBase, FloatTrack.Points, InVal)
 
-UInterpTrackFloatBase::UInterpTrackFloatBase(const class FPostConstructInitializeProperties& PCIP)
-	: Super(PCIP)
+UInterpTrackFloatBase::UInterpTrackFloatBase(const FObjectInitializer& ObjectInitializer)
+	: Super(ObjectInitializer)
 {
 	TrackTitle = TEXT("Generic Float Track");
 	CurveTension = 0.0f;
@@ -5060,8 +5060,8 @@ STRUCTTRACK_REMOVEKEYFRAME(UInterpTrackToggle, ToggleTrack)
 STRUCTTRACK_DUPLICATEKEYFRAME(UInterpTrackToggle, ToggleTrack, Time, FToggleTrackKey)
 STRUCTTRACK_GETCLOSESTSNAPPOSITION(UInterpTrackToggle, ToggleTrack, Time)
 
-UInterpTrackToggle::UInterpTrackToggle(const class FPostConstructInitializeProperties& PCIP)
-: Super(PCIP)
+UInterpTrackToggle::UInterpTrackToggle(const FObjectInitializer& ObjectInitializer)
+: Super(ObjectInitializer)
 {
 	TrackInstClass = UInterpTrackInstToggle::StaticClass();
 	TrackTitle = TEXT("Toggle");
@@ -5368,8 +5368,8 @@ STRUCTTRACK_GETKEYFRAMETIME(UInterpTrackVectorBase, VectorTrack.Points, InVal)
 STRUCTTRACK_GETKEYFRAMEINDEX(UInterpTrackVectorBase, VectorTrack.Points, InVal)
 STRUCTTRACK_GETCLOSESTSNAPPOSITION(UInterpTrackVectorBase, VectorTrack.Points, InVal)
 
-UInterpTrackVectorBase::UInterpTrackVectorBase(const class FPostConstructInitializeProperties& PCIP)
-	: Super(PCIP)
+UInterpTrackVectorBase::UInterpTrackVectorBase(const FObjectInitializer& ObjectInitializer)
+	: Super(ObjectInitializer)
 {
 	TrackTitle = TEXT("Generic Vector Track");
 	CurveTension = 0.0f;
@@ -5450,8 +5450,8 @@ STRUCTTRACK_GETKEYFRAMETIME(UInterpTrackLinearColorBase, LinearColorTrack.Points
 STRUCTTRACK_GETKEYFRAMEINDEX(UInterpTrackLinearColorBase, LinearColorTrack.Points, InVal)
 STRUCTTRACK_GETCLOSESTSNAPPOSITION(UInterpTrackLinearColorBase, LinearColorTrack.Points, InVal)
 
-UInterpTrackLinearColorBase::UInterpTrackLinearColorBase(const class FPostConstructInitializeProperties& PCIP)
-	: Super(PCIP)
+UInterpTrackLinearColorBase::UInterpTrackLinearColorBase(const FObjectInitializer& ObjectInitializer)
+	: Super(ObjectInitializer)
 {
 	TrackTitle = TEXT("Generic LinearColor Track");
 	CurveTension = 0.0f;
@@ -5525,8 +5525,8 @@ void UInterpTrackLinearColorBase::PostEditChangeProperty(FPropertyChangedEvent&
 /*-----------------------------------------------------------------------------
 	UInterpTrackFloatProp
 -----------------------------------------------------------------------------*/
-UInterpTrackFloatProp::UInterpTrackFloatProp(const class FPostConstructInitializeProperties& PCIP)
-	: Super(PCIP)
+UInterpTrackFloatProp::UInterpTrackFloatProp(const FObjectInitializer& ObjectInitializer)
+	: Super(ObjectInitializer)
 {
 	TrackInstClass = UInterpTrackInstFloatProp::StaticClass();
 	TrackTitle = TEXT("Float Property");
@@ -5630,8 +5630,8 @@ void UInterpTrackFloatProp::ReduceKeys( float IntervalStart, float IntervalEnd,
 /*-----------------------------------------------------------------------------
   UInterpTrackInstFloatProp
 -----------------------------------------------------------------------------*/
-UInterpTrackInstFloatProp::UInterpTrackInstFloatProp(const class FPostConstructInitializeProperties& PCIP)
-	: Super(PCIP)
+UInterpTrackInstFloatProp::UInterpTrackInstFloatProp(const FObjectInitializer& ObjectInitializer)
+	: Super(ObjectInitializer)
 {
 }
 
@@ -5687,8 +5687,8 @@ void UInterpTrackInstFloatProp::InitTrackInst(UInterpTrack* Track)
 /*-----------------------------------------------------------------------------
 	UInterpTrackVectorProp
 -----------------------------------------------------------------------------*/
-UInterpTrackVectorProp::UInterpTrackVectorProp(const class FPostConstructInitializeProperties& PCIP)
-	: Super(PCIP)
+UInterpTrackVectorProp::UInterpTrackVectorProp(const FObjectInitializer& ObjectInitializer)
+	: Super(ObjectInitializer)
 {
 
 	TrackInstClass = UInterpTrackInstVectorProp::StaticClass();
@@ -5771,8 +5771,8 @@ const FString	UInterpTrackVectorProp::GetSlateHelperClassName() const
 /*-----------------------------------------------------------------------------
 	UInterpTrackInstVectorProp
 -----------------------------------------------------------------------------*/
-UInterpTrackInstVectorProp::UInterpTrackInstVectorProp(const class FPostConstructInitializeProperties& PCIP)
-	: Super(PCIP)
+UInterpTrackInstVectorProp::UInterpTrackInstVectorProp(const FObjectInitializer& ObjectInitializer)
+	: Super(ObjectInitializer)
 {
 }
 
@@ -5838,8 +5838,8 @@ STRUCTTRACK_REMOVEKEYFRAME(UInterpTrackBoolProp, BoolTrack)
 STRUCTTRACK_DUPLICATEKEYFRAME(UInterpTrackBoolProp, BoolTrack, Time, FBoolTrackKey)
 STRUCTTRACK_GETCLOSESTSNAPPOSITION(UInterpTrackBoolProp, BoolTrack, Time)
 
-UInterpTrackBoolProp::UInterpTrackBoolProp(const class FPostConstructInitializeProperties& PCIP)
-: Super(PCIP)
+UInterpTrackBoolProp::UInterpTrackBoolProp(const FObjectInitializer& ObjectInitializer)
+: Super(ObjectInitializer)
 {
 	TrackInstClass = UInterpTrackInstBoolProp::StaticClass();
 	TrackTitle = TEXT("Bool Property");
@@ -5961,8 +5961,8 @@ const FString UInterpTrackBoolProp::GetSlateHelperClassName() const
 /*-----------------------------------------------------------------------------
 	UInterpTrackInstBoolProp
 -----------------------------------------------------------------------------*/
-UInterpTrackInstBoolProp::UInterpTrackInstBoolProp(const class FPostConstructInitializeProperties& PCIP)
-	: Super(PCIP)
+UInterpTrackInstBoolProp::UInterpTrackInstBoolProp(const FObjectInitializer& ObjectInitializer)
+	: Super(ObjectInitializer)
 {
 }
 
@@ -6013,8 +6013,8 @@ void UInterpTrackInstBoolProp::RestoreActorState( UInterpTrack* Track )
 /*-----------------------------------------------------------------------------
 	UInterpTrackColorProp
 -----------------------------------------------------------------------------*/
-UInterpTrackColorProp::UInterpTrackColorProp(const class FPostConstructInitializeProperties& PCIP)
-	: Super(PCIP)
+UInterpTrackColorProp::UInterpTrackColorProp(const FObjectInitializer& ObjectInitializer)
+	: Super(ObjectInitializer)
 {
 
 	TrackInstClass = UInterpTrackInstColorProp::StaticClass();
@@ -6109,8 +6109,8 @@ const FString	UInterpTrackColorProp::GetSlateHelperClassName() const
 /*-----------------------------------------------------------------------------
 UInterpTrackInstColorProp
 -----------------------------------------------------------------------------*/
-UInterpTrackInstColorProp::UInterpTrackInstColorProp(const class FPostConstructInitializeProperties& PCIP)
-	: Super(PCIP)
+UInterpTrackInstColorProp::UInterpTrackInstColorProp(const FObjectInitializer& ObjectInitializer)
+	: Super(ObjectInitializer)
 {
 }
 
@@ -6170,8 +6170,8 @@ void UInterpTrackInstColorProp::InitTrackInst(UInterpTrack* Track)
 /*-----------------------------------------------------------------------------
 UInterpTrackLinearColorProp
 -----------------------------------------------------------------------------*/
-UInterpTrackLinearColorProp::UInterpTrackLinearColorProp(const class FPostConstructInitializeProperties& PCIP)
-	: Super(PCIP)
+UInterpTrackLinearColorProp::UInterpTrackLinearColorProp(const FObjectInitializer& ObjectInitializer)
+	: Super(ObjectInitializer)
 {
 	TrackInstClass = UInterpTrackInstLinearColorProp::StaticClass();
 	TrackTitle = TEXT("LinearColor Property");
@@ -6258,8 +6258,8 @@ const FString	UInterpTrackLinearColorProp::GetSlateHelperClassName() const
 /*-----------------------------------------------------------------------------
 UInterpTrackInstLinearColorProp
 -----------------------------------------------------------------------------*/
-UInterpTrackInstLinearColorProp::UInterpTrackInstLinearColorProp(const class FPostConstructInitializeProperties& PCIP)
-	: Super(PCIP)
+UInterpTrackInstLinearColorProp::UInterpTrackInstLinearColorProp(const FObjectInitializer& ObjectInitializer)
+	: Super(ObjectInitializer)
 {
 }
 
@@ -6325,8 +6325,8 @@ STRUCTTRACK_SETKEYFRAMETIME(UInterpTrackEvent, EventTrack, Time, FEventTrackKey)
 STRUCTTRACK_DUPLICATEKEYFRAME(UInterpTrackEvent, EventTrack, Time, FEventTrackKey)
 STRUCTTRACK_GETCLOSESTSNAPPOSITION(UInterpTrackEvent, EventTrack, Time)
 
-UInterpTrackEvent::UInterpTrackEvent(const class FPostConstructInitializeProperties& PCIP)
-: Super(PCIP)
+UInterpTrackEvent::UInterpTrackEvent(const FObjectInitializer& ObjectInitializer)
+: Super(ObjectInitializer)
 {
 	TrackInstClass = UInterpTrackInstEvent::StaticClass();
 	TrackTitle = TEXT("Event");
@@ -6484,8 +6484,8 @@ const FString UInterpTrackEvent::GetSlateHelperClassName() const
 /*-----------------------------------------------------------------------------
 	UInterpTrackInstEvent
 -----------------------------------------------------------------------------*/
-UInterpTrackInstEvent::UInterpTrackInstEvent(const class FPostConstructInitializeProperties& PCIP)
-	: Super(PCIP)
+UInterpTrackInstEvent::UInterpTrackInstEvent(const FObjectInitializer& ObjectInitializer)
+	: Super(ObjectInitializer)
 {
 }
 
@@ -6512,8 +6512,8 @@ STRUCTTRACK_REMOVEKEYFRAME(UInterpTrackDirector, CutTrack)
 STRUCTTRACK_DUPLICATEKEYFRAME(UInterpTrackDirector, CutTrack, Time, FDirectorTrackCut)
 STRUCTTRACK_GETCLOSESTSNAPPOSITION(UInterpTrackDirector, CutTrack, Time)
 
-UInterpTrackDirector::UInterpTrackDirector(const class FPostConstructInitializeProperties& PCIP)
-: Super(PCIP)
+UInterpTrackDirector::UInterpTrackDirector(const FObjectInitializer& ObjectInitializer)
+: Super(ObjectInitializer)
 {
 
 	bOnePerGroup = true;
@@ -6814,8 +6814,8 @@ const FString UInterpTrackDirector::GetFormattedCameraShotName(int32 KeyIndex) c
 /*-----------------------------------------------------------------------------
 	UInterpTrackInstDirector
 -----------------------------------------------------------------------------*/
-UInterpTrackInstDirector::UInterpTrackInstDirector(const class FPostConstructInitializeProperties& PCIP)
-	: Super(PCIP)
+UInterpTrackInstDirector::UInterpTrackInstDirector(const FObjectInitializer& ObjectInitializer)
+	: Super(ObjectInitializer)
 {
 }
 
@@ -6849,8 +6849,8 @@ void UInterpTrackInstDirector::TermTrackInst(UInterpTrack* Track)
 /*-----------------------------------------------------------------------------
 	UInterpTrackFade
 -----------------------------------------------------------------------------*/
-UInterpTrackFade::UInterpTrackFade(const class FPostConstructInitializeProperties& PCIP)
-	: Super(PCIP)
+UInterpTrackFade::UInterpTrackFade(const FObjectInitializer& ObjectInitializer)
+	: Super(ObjectInitializer)
 {
 	bOnePerGroup = true;
 	bDirGroupOnly = true;
@@ -6908,8 +6908,8 @@ float UInterpTrackFade::GetFadeAmountAtTime(float Time)
 /*-----------------------------------------------------------------------------
 	UInterpTrackInstFade
 -----------------------------------------------------------------------------*/
-UInterpTrackInstFade::UInterpTrackInstFade(const class FPostConstructInitializeProperties& PCIP)
-	: Super(PCIP)
+UInterpTrackInstFade::UInterpTrackInstFade(const FObjectInitializer& ObjectInitializer)
+	: Super(ObjectInitializer)
 {
 }
 
@@ -6940,8 +6940,8 @@ void UInterpTrackInstFade::TermTrackInst(UInterpTrack* Track)
 /*-----------------------------------------------------------------------------
 	UInterpTrackSlomo
 -----------------------------------------------------------------------------*/
-UInterpTrackSlomo::UInterpTrackSlomo(const class FPostConstructInitializeProperties& PCIP)
-	: Super(PCIP)
+UInterpTrackSlomo::UInterpTrackSlomo(const FObjectInitializer& ObjectInitializer)
+	: Super(ObjectInitializer)
 {
 	bOnePerGroup = true;
 	bDirGroupOnly = true;
@@ -6999,8 +6999,8 @@ void UInterpTrackSlomo::SetTrackToSensibleDefault()
 /*-----------------------------------------------------------------------------
 	UInterpTrackInstSlomo
 -----------------------------------------------------------------------------*/
-UInterpTrackInstSlomo::UInterpTrackInstSlomo(const class FPostConstructInitializeProperties& PCIP)
-	: Super(PCIP)
+UInterpTrackInstSlomo::UInterpTrackInstSlomo(const FObjectInitializer& ObjectInitializer)
+	: Super(ObjectInitializer)
 {
 }
 
@@ -7072,8 +7072,8 @@ STRUCTTRACK_SETKEYFRAMETIME(UInterpTrackAnimControl, AnimSeqs, StartTime, FAnimC
 STRUCTTRACK_REMOVEKEYFRAME(UInterpTrackAnimControl, AnimSeqs)
 STRUCTTRACK_DUPLICATEKEYFRAME(UInterpTrackAnimControl, AnimSeqs, StartTime, FAnimControlTrackKey)
 
-UInterpTrackAnimControl::UInterpTrackAnimControl(const class FPostConstructInitializeProperties& PCIP)
-	: Super(PCIP)
+UInterpTrackAnimControl::UInterpTrackAnimControl(const FObjectInitializer& ObjectInitializer)
+	: Super(ObjectInitializer)
 {
 
 	TrackInstClass = UInterpTrackInstAnimControl::StaticClass();
@@ -7640,8 +7640,8 @@ const FString UInterpTrackAnimControl::GetSlateHelperClassName() const
 /*-----------------------------------------------------------------------------
 	UInterpTrackInstAnimControl
 -----------------------------------------------------------------------------*/
-UInterpTrackInstAnimControl::UInterpTrackInstAnimControl(const class FPostConstructInitializeProperties& PCIP)
-	: Super(PCIP)
+UInterpTrackInstAnimControl::UInterpTrackInstAnimControl(const FObjectInitializer& ObjectInitializer)
+	: Super(ObjectInitializer)
 {
 }
 
@@ -7674,8 +7674,8 @@ STRUCTTRACK_SETKEYFRAMETIME(UInterpTrackSound, Sounds, Time, FSoundTrackKey)
 STRUCTTRACK_REMOVEKEYFRAME(UInterpTrackSound, Sounds)
 STRUCTTRACK_DUPLICATEKEYFRAME(UInterpTrackSound, Sounds, Time, FSoundTrackKey)
 
-UInterpTrackSound::UInterpTrackSound(const class FPostConstructInitializeProperties& PCIP)
-: Super(PCIP)
+UInterpTrackSound::UInterpTrackSound(const FObjectInitializer& ObjectInitializer)
+: Super(ObjectInitializer)
 {
 
 	TrackInstClass = UInterpTrackInstSound::StaticClass();
@@ -8077,8 +8077,8 @@ void UInterpTrackSound::PreviewStopPlayback(class UInterpTrackInst* TrInst)
 /*-----------------------------------------------------------------------------
 	UInterpTrackInstSound
 -----------------------------------------------------------------------------*/
-UInterpTrackInstSound::UInterpTrackInstSound(const class FPostConstructInitializeProperties& PCIP)
-	: Super(PCIP)
+UInterpTrackInstSound::UInterpTrackInstSound(const FObjectInitializer& ObjectInitializer)
+	: Super(ObjectInitializer)
 {
 }
 
@@ -8116,8 +8116,8 @@ void UInterpTrackInstSound::TermTrackInst(UInterpTrack* Track)
 /*-----------------------------------------------------------------------------
 	UInterpTrackFloatParticleParam
 -----------------------------------------------------------------------------*/
-UInterpTrackFloatParticleParam::UInterpTrackFloatParticleParam(const class FPostConstructInitializeProperties& PCIP)
-	: Super(PCIP)
+UInterpTrackFloatParticleParam::UInterpTrackFloatParticleParam(const FObjectInitializer& ObjectInitializer)
+	: Super(ObjectInitializer)
 {
 	
 	TrackInstClass = UInterpTrackInstFloatParticleParam::StaticClass();
@@ -8155,8 +8155,8 @@ void UInterpTrackFloatParticleParam::UpdateTrack(float NewPosition, UInterpTrack
 /*-----------------------------------------------------------------------------
 	UInterpTrackInstFloatParticleParam
 -----------------------------------------------------------------------------*/
-UInterpTrackInstFloatParticleParam::UInterpTrackInstFloatParticleParam(const class FPostConstructInitializeProperties& PCIP)
-	: Super(PCIP)
+UInterpTrackInstFloatParticleParam::UInterpTrackInstFloatParticleParam(const FObjectInitializer& ObjectInitializer)
+	: Super(ObjectInitializer)
 {
 }
 
@@ -8462,8 +8462,8 @@ static void PostEditChangeMaterialParamTrack()
 /*-----------------------------------------------------------------------------
 	Float material parameter track.
 -----------------------------------------------------------------------------*/
-UInterpTrackFloatMaterialParam::UInterpTrackFloatMaterialParam(const class FPostConstructInitializeProperties& PCIP)
-	: Super(PCIP)
+UInterpTrackFloatMaterialParam::UInterpTrackFloatMaterialParam(const FObjectInitializer& ObjectInitializer)
+	: Super(ObjectInitializer)
 {
 	TrackInstClass = UInterpTrackInstFloatMaterialParam::StaticClass();
 	TrackTitle = TEXT("Float UMaterial Param");
@@ -8541,8 +8541,8 @@ void UInterpTrackFloatMaterialParam::UpdateTrack(float NewPosition, UInterpTrack
 	}
 }
 
-UInterpTrackInstFloatMaterialParam::UInterpTrackInstFloatMaterialParam(const class FPostConstructInitializeProperties& PCIP)
-	: Super(PCIP)
+UInterpTrackInstFloatMaterialParam::UInterpTrackInstFloatMaterialParam(const FObjectInitializer& ObjectInitializer)
+	: Super(ObjectInitializer)
 {
 }
 
@@ -8619,8 +8619,8 @@ void UInterpTrackInstFloatMaterialParam::RestoreActorState(UInterpTrack* Track)
 /*-----------------------------------------------------------------------------
 	Vector material parameter track.
 -----------------------------------------------------------------------------*/
-UInterpTrackVectorMaterialParam::UInterpTrackVectorMaterialParam(const class FPostConstructInitializeProperties& PCIP)
-	: Super(PCIP)
+UInterpTrackVectorMaterialParam::UInterpTrackVectorMaterialParam(const FObjectInitializer& ObjectInitializer)
+	: Super(ObjectInitializer)
 {
 
 	TrackInstClass = UInterpTrackInstVectorMaterialParam::StaticClass();
@@ -8673,8 +8673,8 @@ void UInterpTrackVectorMaterialParam::UpdateTrack(float NewPosition, UInterpTrac
 	}
 }
 
-UInterpTrackInstVectorMaterialParam::UInterpTrackInstVectorMaterialParam(const class FPostConstructInitializeProperties& PCIP)
-	: Super(PCIP)
+UInterpTrackInstVectorMaterialParam::UInterpTrackInstVectorMaterialParam(const FObjectInitializer& ObjectInitializer)
+	: Super(ObjectInitializer)
 {
 }
 
@@ -8771,14 +8771,14 @@ static void GetListOfStaticActors(FString& OutString, const TArray<AActor*>& Act
 	}
 }
 
-AMaterialInstanceActor::AMaterialInstanceActor(const class FPostConstructInitializeProperties& PCIP)
-	: Super(PCIP)
+AMaterialInstanceActor::AMaterialInstanceActor(const FObjectInitializer& ObjectInitializer)
+	: Super(ObjectInitializer)
 {
-	TSubobjectPtr<USceneComponent> SceneComponent = PCIP.CreateDefaultSubobject<USceneComponent>(this, TEXT("SceneComp"));
+	TSubobjectPtr<USceneComponent> SceneComponent = ObjectInitializer.CreateDefaultSubobject<USceneComponent>(this, TEXT("SceneComp"));
 	RootComponent = SceneComponent;
 
 #if WITH_EDITORONLY_DATA
-	SpriteComponent = PCIP.CreateEditorOnlyDefaultSubobject<UBillboardComponent>(this, TEXT("Sprite"));
+	SpriteComponent = ObjectInitializer.CreateEditorOnlyDefaultSubobject<UBillboardComponent>(this, TEXT("Sprite"));
 	if (!IsRunningCommandlet() && (SpriteComponent != NULL))
 	{
 		// Structure to hold one-time initialization
@@ -8837,8 +8837,8 @@ void AMaterialInstanceActor::PostEditChangeProperty(FPropertyChangedEvent& Prope
 /*-----------------------------------------------------------------------------
 	UInterpTrackInstToggle
 -----------------------------------------------------------------------------*/
-UInterpTrackInstToggle::UInterpTrackInstToggle(const class FPostConstructInitializeProperties& PCIP)
-	: Super(PCIP)
+UInterpTrackInstToggle::UInterpTrackInstToggle(const FObjectInitializer& ObjectInitializer)
+	: Super(ObjectInitializer)
 {
 }
 
@@ -8896,8 +8896,8 @@ void UInterpTrackInstToggle::RestoreActorState(UInterpTrack* Track)
 /*-----------------------------------------------------------------------------
 	UInterpTrackColorScale
 -----------------------------------------------------------------------------*/
-UInterpTrackColorScale::UInterpTrackColorScale(const class FPostConstructInitializeProperties& PCIP)
-	: Super(PCIP)
+UInterpTrackColorScale::UInterpTrackColorScale(const FObjectInitializer& ObjectInitializer)
+	: Super(ObjectInitializer)
 {
 
 	bOnePerGroup = true;
@@ -8960,8 +8960,8 @@ void UInterpTrackColorScale::SetTrackToSensibleDefault()
 /*-----------------------------------------------------------------------------
 	UInterpTrackInstColorScale
 -----------------------------------------------------------------------------*/
-UInterpTrackInstColorScale::UInterpTrackInstColorScale(const class FPostConstructInitializeProperties& PCIP)
-	: Super(PCIP)
+UInterpTrackInstColorScale::UInterpTrackInstColorScale(const FObjectInitializer& ObjectInitializer)
+	: Super(ObjectInitializer)
 {
 }
 
@@ -8981,8 +8981,8 @@ void UInterpTrackInstColorScale::TermTrackInst(UInterpTrack* Track)
 /*-----------------------------------------------------------------------------
 	UInterpTrackAudioMaster
 -----------------------------------------------------------------------------*/
-UInterpTrackAudioMaster::UInterpTrackAudioMaster(const class FPostConstructInitializeProperties& PCIP)
-	: Super(PCIP)
+UInterpTrackAudioMaster::UInterpTrackAudioMaster(const FObjectInitializer& ObjectInitializer)
+	: Super(ObjectInitializer)
 {
 	bOnePerGroup = true;
 	bDirGroupOnly = true;
@@ -9054,8 +9054,8 @@ float UInterpTrackAudioMaster::GetPitchScaleForTime( float Time ) const
 /*-----------------------------------------------------------------------------
 	UInterpTrackInstAudioMaster
 -----------------------------------------------------------------------------*/
-UInterpTrackInstAudioMaster::UInterpTrackInstAudioMaster(const class FPostConstructInitializeProperties& PCIP)
-	: Super(PCIP)
+UInterpTrackInstAudioMaster::UInterpTrackInstAudioMaster(const FObjectInitializer& ObjectInitializer)
+	: Super(ObjectInitializer)
 {
 }
 
@@ -9077,8 +9077,8 @@ STRUCTTRACK_REMOVEKEYFRAME(UInterpTrackVisibility, VisibilityTrack)
 STRUCTTRACK_DUPLICATEKEYFRAME(UInterpTrackVisibility, VisibilityTrack, Time, FVisibilityTrackKey)
 STRUCTTRACK_GETCLOSESTSNAPPOSITION(UInterpTrackVisibility, VisibilityTrack, Time)
 
-UInterpTrackVisibility::UInterpTrackVisibility(const class FPostConstructInitializeProperties& PCIP)
-: Super(PCIP)
+UInterpTrackVisibility::UInterpTrackVisibility(const FObjectInitializer& ObjectInitializer)
+: Super(ObjectInitializer)
 {
 
 	TrackInstClass = UInterpTrackInstVisibility::StaticClass();
@@ -9313,8 +9313,8 @@ void UInterpTrackVisibility::HideActor( AActor* Actor, bool bHidden )
 /*-----------------------------------------------------------------------------
 	UInterpTrackInstVisibility
 -----------------------------------------------------------------------------*/
-UInterpTrackInstVisibility::UInterpTrackInstVisibility(const class FPostConstructInitializeProperties& PCIP)
-	: Super(PCIP)
+UInterpTrackInstVisibility::UInterpTrackInstVisibility(const FObjectInitializer& ObjectInitializer)
+	: Super(ObjectInitializer)
 {
 }
 
@@ -9340,8 +9340,8 @@ STRUCTTRACK_REMOVEKEYFRAME(UInterpTrackParticleReplay, TrackKeys)
 STRUCTTRACK_DUPLICATEKEYFRAME(UInterpTrackParticleReplay, TrackKeys, Time, FParticleReplayTrackKey)
 STRUCTTRACK_GETCLOSESTSNAPPOSITION(UInterpTrackParticleReplay, TrackKeys, Time)
 
-UInterpTrackParticleReplay::UInterpTrackParticleReplay(const class FPostConstructInitializeProperties& PCIP)
-: Super(PCIP)
+UInterpTrackParticleReplay::UInterpTrackParticleReplay(const FObjectInitializer& ObjectInitializer)
+: Super(ObjectInitializer)
 {
 
 	TrackInstClass = UInterpTrackInstParticleReplay::StaticClass();
@@ -9592,8 +9592,8 @@ const FString UInterpTrackParticleReplay::GetSlateHelperClassName() const
 /*-----------------------------------------------------------------------------
 	UInterpTrackInstParticleReplay
 -----------------------------------------------------------------------------*/
-UInterpTrackInstParticleReplay::UInterpTrackInstParticleReplay(const class FPostConstructInitializeProperties& PCIP)
-	: Super(PCIP)
+UInterpTrackInstParticleReplay::UInterpTrackInstParticleReplay(const FObjectInitializer& ObjectInitializer)
+	: Super(ObjectInitializer)
 {
 }
 
@@ -9624,8 +9624,8 @@ void UInterpTrackInstParticleReplay::RestoreActorState(UInterpTrack* Track)
 }
 
 
-UInterpGroupCamera::UInterpGroupCamera(const class FPostConstructInitializeProperties& PCIP)
-	: Super(PCIP)
+UInterpGroupCamera::UInterpGroupCamera(const FObjectInitializer& ObjectInitializer)
+	: Super(ObjectInitializer)
 {
 #if WITH_EDITORONLY_DATA
 	Target.Location = FVector(140.0f, 0.0f, -40.0f);
@@ -9634,13 +9634,13 @@ UInterpGroupCamera::UInterpGroupCamera(const class FPostConstructInitializePrope
 	CompressTolerance = 5.0f;
 }
 
-UInterpGroupInstCamera::UInterpGroupInstCamera(const class FPostConstructInitializeProperties& PCIP)
-	: Super(PCIP)
+UInterpGroupInstCamera::UInterpGroupInstCamera(const FObjectInitializer& ObjectInitializer)
+	: Super(ObjectInitializer)
 {
 }
 
-UInterpGroupInstDirector::UInterpGroupInstDirector(const class FPostConstructInitializeProperties& PCIP)
-	: Super(PCIP)
+UInterpGroupInstDirector::UInterpGroupInstDirector(const FObjectInitializer& ObjectInitializer)
+	: Super(ObjectInitializer)
 {
 }
 
diff --git a/Engine/Source/Runtime/Engine/Private/KillZVolume.cpp b/Engine/Source/Runtime/Engine/Private/KillZVolume.cpp
--- a/Engine/Source/Runtime/Engine/Private/KillZVolume.cpp
+++ b/Engine/Source/Runtime/Engine/Private/KillZVolume.cpp
@@ -3,8 +3,8 @@
 #include "EnginePrivate.h"
 #include "GameFramework/KillZVolume.h"
 
-AKillZVolume::AKillZVolume(const class FPostConstructInitializeProperties& PCIP)
-	: Super(PCIP)
+AKillZVolume::AKillZVolume(const FObjectInitializer& ObjectInitializer)
+	: Super(ObjectInitializer)
 {
 }
 
diff --git a/Engine/Source/Runtime/Engine/Private/KismetArrayLibrary.cpp b/Engine/Source/Runtime/Engine/Private/KismetArrayLibrary.cpp
--- a/Engine/Source/Runtime/Engine/Private/KismetArrayLibrary.cpp
+++ b/Engine/Source/Runtime/Engine/Private/KismetArrayLibrary.cpp
@@ -8,8 +8,8 @@ DEFINE_LOG_CATEGORY_STATIC(LogArray, Warning, All);
 //////////////////////////////////////////////////////////////////////////
 // UKismetArrayLibrary
 
-UKismetArrayLibrary::UKismetArrayLibrary(const FPostConstructInitializeProperties& PCIP)
-	: Super(PCIP)
+UKismetArrayLibrary::UKismetArrayLibrary(const FObjectInitializer& ObjectInitializer)
+	: Super(ObjectInitializer)
 {
 }
 
diff --git a/Engine/Source/Runtime/Engine/Private/KismetFunctionLibrary.cpp b/Engine/Source/Runtime/Engine/Private/KismetFunctionLibrary.cpp
--- a/Engine/Source/Runtime/Engine/Private/KismetFunctionLibrary.cpp
+++ b/Engine/Source/Runtime/Engine/Private/KismetFunctionLibrary.cpp
@@ -2,8 +2,8 @@
 
 #include "EnginePrivate.h"
 
-UBlueprintFunctionLibrary::UBlueprintFunctionLibrary(const class FPostConstructInitializeProperties& PCIP)
-	: Super(PCIP)
+UBlueprintFunctionLibrary::UBlueprintFunctionLibrary(const FObjectInitializer& ObjectInitializer)
+	: Super(ObjectInitializer)
 {
 }
 
diff --git a/Engine/Source/Runtime/Engine/Private/KismetGuidLibrary.cpp b/Engine/Source/Runtime/Engine/Private/KismetGuidLibrary.cpp
--- a/Engine/Source/Runtime/Engine/Private/KismetGuidLibrary.cpp
+++ b/Engine/Source/Runtime/Engine/Private/KismetGuidLibrary.cpp
@@ -7,8 +7,8 @@
 /* Guid functions
  *****************************************************************************/
 
-UKismetGuidLibrary::UKismetGuidLibrary( const class FPostConstructInitializeProperties& PCIP )
-	: Super(PCIP)
+UKismetGuidLibrary::UKismetGuidLibrary( const FObjectInitializer& ObjectInitializer )
+	: Super(ObjectInitializer)
 { }
 
 
diff --git a/Engine/Source/Runtime/Engine/Private/KismetInputLibrary.cpp b/Engine/Source/Runtime/Engine/Private/KismetInputLibrary.cpp
--- a/Engine/Source/Runtime/Engine/Private/KismetInputLibrary.cpp
+++ b/Engine/Source/Runtime/Engine/Private/KismetInputLibrary.cpp
@@ -12,8 +12,8 @@
 #define LOCTEXT_NAMESPACE "KismetInputLibrary"
 
 
-UKismetInputLibrary::UKismetInputLibrary(const class FPostConstructInitializeProperties& PCIP)
-	: Super(PCIP)
+UKismetInputLibrary::UKismetInputLibrary(const FObjectInitializer& ObjectInitializer)
+	: Super(ObjectInitializer)
 {
 }
 
diff --git a/Engine/Source/Runtime/Engine/Private/KismetMaterialLibrary.cpp b/Engine/Source/Runtime/Engine/Private/KismetMaterialLibrary.cpp
--- a/Engine/Source/Runtime/Engine/Private/KismetMaterialLibrary.cpp
+++ b/Engine/Source/Runtime/Engine/Private/KismetMaterialLibrary.cpp
@@ -11,8 +11,8 @@
 #define LOCTEXT_NAMESPACE "KismetMaterialLibrary"
 
 
-UKismetMaterialLibrary::UKismetMaterialLibrary(const class FPostConstructInitializeProperties& PCIP)
-	: Super(PCIP)
+UKismetMaterialLibrary::UKismetMaterialLibrary(const FObjectInitializer& ObjectInitializer)
+	: Super(ObjectInitializer)
 {
 }
 
diff --git a/Engine/Source/Runtime/Engine/Private/KismetMathLibrary.cpp b/Engine/Source/Runtime/Engine/Private/KismetMathLibrary.cpp
--- a/Engine/Source/Runtime/Engine/Private/KismetMathLibrary.cpp
+++ b/Engine/Source/Runtime/Engine/Private/KismetMathLibrary.cpp
@@ -26,8 +26,8 @@ float EaseAlpha(float InAlpha, uint8 EasingFunc, float BlendExp, int32 Steps)
 	return InAlpha;
 }
 
-UKismetMathLibrary::UKismetMathLibrary(const class FPostConstructInitializeProperties& PCIP)
-	: Super(PCIP)
+UKismetMathLibrary::UKismetMathLibrary(const FObjectInitializer& ObjectInitializer)
+	: Super(ObjectInitializer)
 {
 }
 
diff --git a/Engine/Source/Runtime/Engine/Private/KismetNodeHelperLibrary.cpp b/Engine/Source/Runtime/Engine/Private/KismetNodeHelperLibrary.cpp
--- a/Engine/Source/Runtime/Engine/Private/KismetNodeHelperLibrary.cpp
+++ b/Engine/Source/Runtime/Engine/Private/KismetNodeHelperLibrary.cpp
@@ -3,8 +3,8 @@
 #include "EnginePrivate.h"
 #include "Kismet/KismetNodeHelperLibrary.h"
 
-UKismetNodeHelperLibrary::UKismetNodeHelperLibrary(const class FPostConstructInitializeProperties& PCIP)
-	: Super(PCIP)
+UKismetNodeHelperLibrary::UKismetNodeHelperLibrary(const FObjectInitializer& ObjectInitializer)
+	: Super(ObjectInitializer)
 {
 }
 
diff --git a/Engine/Source/Runtime/Engine/Private/KismetStringLibrary.cpp b/Engine/Source/Runtime/Engine/Private/KismetStringLibrary.cpp
--- a/Engine/Source/Runtime/Engine/Private/KismetStringLibrary.cpp
+++ b/Engine/Source/Runtime/Engine/Private/KismetStringLibrary.cpp
@@ -6,8 +6,8 @@
 //////////////////////////////////////////////////////////////////////////
 // UKismetStringLibrary
 
-UKismetStringLibrary::UKismetStringLibrary(const class FPostConstructInitializeProperties& PCIP)
-	: Super(PCIP)
+UKismetStringLibrary::UKismetStringLibrary(const FObjectInitializer& ObjectInitializer)
+	: Super(ObjectInitializer)
 {
 }
 
diff --git a/Engine/Source/Runtime/Engine/Private/KismetSystemLibrary.cpp b/Engine/Source/Runtime/Engine/Private/KismetSystemLibrary.cpp
--- a/Engine/Source/Runtime/Engine/Private/KismetSystemLibrary.cpp
+++ b/Engine/Source/Runtime/Engine/Private/KismetSystemLibrary.cpp
@@ -13,8 +13,8 @@
 //////////////////////////////////////////////////////////////////////////
 // UKismetSystemLibrary
 
-UKismetSystemLibrary::UKismetSystemLibrary(const class FPostConstructInitializeProperties& PCIP)
-	: Super(PCIP)
+UKismetSystemLibrary::UKismetSystemLibrary(const FObjectInitializer& ObjectInitializer)
+	: Super(ObjectInitializer)
 {
 }
 
diff --git a/Engine/Source/Runtime/Engine/Private/KismetTextLibrary.cpp b/Engine/Source/Runtime/Engine/Private/KismetTextLibrary.cpp
--- a/Engine/Source/Runtime/Engine/Private/KismetTextLibrary.cpp
+++ b/Engine/Source/Runtime/Engine/Private/KismetTextLibrary.cpp
@@ -9,8 +9,8 @@
 /* UKismetTextLibrary structors
  *****************************************************************************/
 
-UKismetTextLibrary::UKismetTextLibrary(const class FPostConstructInitializeProperties& PCIP)
-	: Super(PCIP)
+UKismetTextLibrary::UKismetTextLibrary(const FObjectInitializer& ObjectInitializer)
+	: Super(ObjectInitializer)
 { }
 
 
