Rename FPostConstructInitializeProperties to something simpler

diff --git a/Engine/Source/Runtime/Engine/Private/Landscape/Landscape.cpp b/Engine/Source/Runtime/Engine/Private/Landscape/Landscape.cpp
--- a/Engine/Source/Runtime/Engine/Private/Landscape/Landscape.cpp
+++ b/Engine/Source/Runtime/Engine/Private/Landscape/Landscape.cpp
@@ -50,8 +50,8 @@ FAutoConsoleCommand CmdPrintNumLandscapeShadows(
 	FConsoleCommandDelegate::CreateStatic(PrintNumLandscapeShadows)
 	);
 
-ULandscapeComponent::ULandscapeComponent(const class FPostConstructInitializeProperties& PCIP)
-	: Super(PCIP)
+ULandscapeComponent::ULandscapeComponent(const FObjectInitializer& ObjectInitializer)
+	: Super(ObjectInitializer)
 {
 	BodyInstance.bEnableCollision_DEPRECATED = true;
 	SetCollisionProfileName(UCollisionProfile::BlockAll_ProfileName);
@@ -314,15 +314,15 @@ void ULandscapeComponent::GetLayerDebugColorKey(int32& R, int32& G, int32& B) co
 }
 #endif	//WITH_EDITOR
 
-ULandscapeMeshCollisionComponent::ULandscapeMeshCollisionComponent(const class FPostConstructInitializeProperties& PCIP)
-	: Super(PCIP)
+ULandscapeMeshCollisionComponent::ULandscapeMeshCollisionComponent(const FObjectInitializer& ObjectInitializer)
+	: Super(ObjectInitializer)
 {
 	// make landscape always create? 
 	bAlwaysCreatePhysicsState = true;
 }
 
-ULandscapeInfo::ULandscapeInfo(const class FPostConstructInitializeProperties& PCIP)
-	: Super(PCIP)
+ULandscapeInfo::ULandscapeInfo(const FObjectInitializer& ObjectInitializer)
+	: Super(ObjectInitializer)
 {
 #if WITH_EDITOR
 	bIsValid = false;
@@ -457,8 +457,8 @@ void ULandscapeComponent::PostLoad()
 
 #endif // WITH_EDITOR
 
-ALandscape::ALandscape(const class FPostConstructInitializeProperties& PCIP)
-	: Super(PCIP)
+ALandscape::ALandscape(const FObjectInitializer& ObjectInitializer)
+	: Super(ObjectInitializer)
 {
 	bIsProxy = false;
 
@@ -472,8 +472,8 @@ ALandscape* ALandscape::GetLandscapeActor()
 	return this;
 }
 
-ALandscapeProxy::ALandscapeProxy(const class FPostConstructInitializeProperties& PCIP)
-	: Super(PCIP)
+ALandscapeProxy::ALandscapeProxy(const FObjectInitializer& ObjectInitializer)
+	: Super(ObjectInitializer)
 {
 	PrimaryActorTick.bCanEverTick = false;
 	bReplicates = false;
@@ -482,7 +482,7 @@ ALandscapeProxy::ALandscapeProxy(const class FPostConstructInitializeProperties&
 	bReplicateMovement = false;
 	bCanBeDamaged = false;
 
-	TSubobjectPtr<USceneComponent> SceneComponent = PCIP.CreateDefaultSubobject<USceneComponent>(this, TEXT("RootComponent0"));
+	TSubobjectPtr<USceneComponent> SceneComponent = ObjectInitializer.CreateDefaultSubobject<USceneComponent>(this, TEXT("RootComponent0"));
 	RootComponent = SceneComponent;
 	RootComponent->RelativeScale3D = FVector(128.0f, 128.0f, 256.0f); // Old default scale, preserved for compatibility. See ULandscapeEditorObject::NewLandscape_Scale
 	RootComponent->Mobility = EComponentMobility::Static;
diff --git a/Engine/Source/Runtime/Engine/Private/Landscape/LandscapeCollision.cpp b/Engine/Source/Runtime/Engine/Private/Landscape/LandscapeCollision.cpp
--- a/Engine/Source/Runtime/Engine/Private/Landscape/LandscapeCollision.cpp
+++ b/Engine/Source/Runtime/Engine/Private/Landscape/LandscapeCollision.cpp
@@ -1569,8 +1569,8 @@ void ULandscapeHeightfieldCollisionComponent::SetSectionBase(FIntPoint InSection
 	SectionBaseY = InSectionBase.Y;
 }
 
-ULandscapeHeightfieldCollisionComponent::ULandscapeHeightfieldCollisionComponent(const class FPostConstructInitializeProperties& PCIP)
-	: Super(PCIP)
+ULandscapeHeightfieldCollisionComponent::ULandscapeHeightfieldCollisionComponent(const FObjectInitializer& ObjectInitializer)
+	: Super(ObjectInitializer)
 {
 	BodyInstance.bEnableCollision_DEPRECATED = true;
 	SetCollisionProfileName(UCollisionProfile::BlockAll_ProfileName);
diff --git a/Engine/Source/Runtime/Engine/Private/Landscape/LandscapeEdit.cpp b/Engine/Source/Runtime/Engine/Private/Landscape/LandscapeEdit.cpp
--- a/Engine/Source/Runtime/Engine/Private/Landscape/LandscapeEdit.cpp
+++ b/Engine/Source/Runtime/Engine/Private/Landscape/LandscapeEdit.cpp
@@ -3369,8 +3369,8 @@ void ALandscape::PostEditMove(bool bFinished)
 }
 #endif	//WITH_EDITOR
 
-ULandscapeLayerInfoObject::ULandscapeLayerInfoObject(const class FPostConstructInitializeProperties& PCIP)
-	: Super(PCIP)
+ULandscapeLayerInfoObject::ULandscapeLayerInfoObject(const FObjectInitializer& ObjectInitializer)
+	: Super(ObjectInitializer)
 {
 	Hardness = 0.5f;
 #if WITH_EDITORONLY_DATA
diff --git a/Engine/Source/Runtime/Engine/Private/Landscape/LandscapeGizmoActor.cpp b/Engine/Source/Runtime/Engine/Private/Landscape/LandscapeGizmoActor.cpp
--- a/Engine/Source/Runtime/Engine/Private/Landscape/LandscapeGizmoActor.cpp
+++ b/Engine/Source/Runtime/Engine/Private/Landscape/LandscapeGizmoActor.cpp
@@ -405,8 +405,8 @@ public:
 	uint32 GetAllocatedSize( void ) const { return( FPrimitiveSceneProxy::GetAllocatedSize() ); }
 };
 
-ULandscapeGizmoRenderComponent::ULandscapeGizmoRenderComponent(const class FPostConstructInitializeProperties& PCIP)
-	: Super(PCIP)
+ULandscapeGizmoRenderComponent::ULandscapeGizmoRenderComponent(const FObjectInitializer& ObjectInitializer)
+	: Super(ObjectInitializer)
 {
 	bHiddenInGame = true;
 	AlwaysLoadOnClient = false;
@@ -435,11 +435,11 @@ FBoxSphereBounds ULandscapeGizmoRenderComponent::CalcBounds(const FTransform& Lo
 	}
 }
 
-ALandscapeGizmoActor::ALandscapeGizmoActor(const class FPostConstructInitializeProperties& PCIP)
-	: Super(PCIP)
+ALandscapeGizmoActor::ALandscapeGizmoActor(const FObjectInitializer& ObjectInitializer)
+	: Super(ObjectInitializer)
 {
 #if WITH_EDITORONLY_DATA
-	SpriteComponent = PCIP.CreateEditorOnlyDefaultSubobject<UBillboardComponent>(this, TEXT("Sprite"));
+	SpriteComponent = ObjectInitializer.CreateEditorOnlyDefaultSubobject<UBillboardComponent>(this, TEXT("Sprite"));
 	if (!IsRunningCommandlet() && (SpriteComponent != nullptr))
 	{
 		// Structure to hold one-time initialization
@@ -466,7 +466,7 @@ ALandscapeGizmoActor::ALandscapeGizmoActor(const class FPostConstructInitializeP
 	}
 #endif
 
-	TSubobjectPtr<USceneComponent> SceneComponent = PCIP.CreateDefaultSubobject<USceneComponent>(this, TEXT("RootComponent0"));
+	TSubobjectPtr<USceneComponent> SceneComponent = ObjectInitializer.CreateDefaultSubobject<USceneComponent>(this, TEXT("RootComponent0"));
 	RootComponent = SceneComponent;
 	RootComponent->Mobility = EComponentMobility::Static;
 
@@ -506,8 +506,8 @@ void ALandscapeGizmoActor::Duplicate(ALandscapeGizmoActor* Gizmo)
 }
 #endif	//WITH_EDITOR
 
-ALandscapeGizmoActiveActor::ALandscapeGizmoActiveActor(const class FPostConstructInitializeProperties& PCIP)
-	: Super(PCIP
+ALandscapeGizmoActiveActor::ALandscapeGizmoActiveActor(const FObjectInitializer& ObjectInitializer)
+	: Super(ObjectInitializer
 		.DoNotCreateDefaultSubobject(TEXT("Sprite"))
 	)
 {
@@ -535,7 +535,7 @@ ALandscapeGizmoActiveActor::ALandscapeGizmoActiveActor(const class FPostConstruc
 	}
 #endif // WITH_EDITORONLY_DATA
 
-	TSubobjectPtr<ULandscapeGizmoRenderComponent> LandscapeGizmoRenderComponent = PCIP.CreateDefaultSubobject<ULandscapeGizmoRenderComponent>(this, TEXT("GizmoRendererComponent0"));
+	TSubobjectPtr<ULandscapeGizmoRenderComponent> LandscapeGizmoRenderComponent = ObjectInitializer.CreateDefaultSubobject<ULandscapeGizmoRenderComponent>(this, TEXT("GizmoRendererComponent0"));
 	LandscapeGizmoRenderComponent->SetCollisionProfileName(UCollisionProfile::BlockAllDynamic_ProfileName);
 
 	RootComponent = LandscapeGizmoRenderComponent;
diff --git a/Engine/Source/Runtime/Engine/Private/Landscape/LandscapeRender.cpp b/Engine/Source/Runtime/Engine/Private/Landscape/LandscapeRender.cpp
--- a/Engine/Source/Runtime/Engine/Private/Landscape/LandscapeRender.cpp
+++ b/Engine/Source/Runtime/Engine/Private/Landscape/LandscapeRender.cpp
@@ -2675,8 +2675,8 @@ void FLandscapeXYOffsetVertexFactory::ModifyCompilationEnvironment( EShaderPlatf
 IMPLEMENT_VERTEX_FACTORY_TYPE(FLandscapeXYOffsetVertexFactory, "LandscapeVertexFactory", true, true, true, false, false);
 
 /** ULandscapeMaterialInstanceConstant */
-ULandscapeMaterialInstanceConstant::ULandscapeMaterialInstanceConstant(const class FPostConstructInitializeProperties& PCIP)
-	: Super(PCIP)
+ULandscapeMaterialInstanceConstant::ULandscapeMaterialInstanceConstant(const FObjectInitializer& ObjectInitializer)
+	: Super(ObjectInitializer)
 {
 	bIsLayerThumbnail = false;
 	DataWeightmapIndex = -1;
diff --git a/Engine/Source/Runtime/Engine/Private/Landscape/LandscapeSplines.cpp b/Engine/Source/Runtime/Engine/Private/Landscape/LandscapeSplines.cpp
--- a/Engine/Source/Runtime/Engine/Private/Landscape/LandscapeSplines.cpp
+++ b/Engine/Source/Runtime/Engine/Private/Landscape/LandscapeSplines.cpp
@@ -447,8 +447,8 @@ public:
 //////////////////////////////////////////////////////////////////////////
 // SPLINE COMPONENT
 
-ULandscapeSplinesComponent::ULandscapeSplinesComponent(const class FPostConstructInitializeProperties& PCIP)
-	: Super(PCIP)
+ULandscapeSplinesComponent::ULandscapeSplinesComponent(const FObjectInitializer& ObjectInitializer)
+	: Super(ObjectInitializer)
 {
 	Mobility = EComponentMobility::Static;
 
@@ -621,8 +621,8 @@ void ULandscapeSplinesComponent::ShowSplineEditorMesh(bool bShow)
 //////////////////////////////////////////////////////////////////////////
 // CONTROL POINT MESH COMPONENT
 
-UControlPointMeshComponent::UControlPointMeshComponent(const class FPostConstructInitializeProperties& PCIP)
-	: Super(PCIP)
+UControlPointMeshComponent::UControlPointMeshComponent(const FObjectInitializer& ObjectInitializer)
+	: Super(ObjectInitializer)
 {
 	SetCollisionProfileName(UCollisionProfile::BlockAll_ProfileName);
 	Mobility = EComponentMobility::Static;
@@ -635,8 +635,8 @@ UControlPointMeshComponent::UControlPointMeshComponent(const class FPostConstruc
 //////////////////////////////////////////////////////////////////////////
 // SPLINE CONTROL POINT
 
-ULandscapeSplineControlPoint::ULandscapeSplineControlPoint(const class FPostConstructInitializeProperties& PCIP)
-	: Super(PCIP)
+ULandscapeSplineControlPoint::ULandscapeSplineControlPoint(const FObjectInitializer& ObjectInitializer)
+	: Super(ObjectInitializer)
 {
 	Width = 1000;
 	SideFalloff = 1000;
@@ -1109,8 +1109,8 @@ void ULandscapeSplineControlPoint::PostEditChangeProperty(FPropertyChangedEvent&
 //////////////////////////////////////////////////////////////////////////
 // SPLINE SEGMENT
 
-ULandscapeSplineSegment::ULandscapeSplineSegment(const class FPostConstructInitializeProperties& PCIP)
-	: Super(PCIP)
+ULandscapeSplineSegment::ULandscapeSplineSegment(const FObjectInitializer& ObjectInitializer)
+	: Super(ObjectInitializer)
 {
 
 	Connections[0].ControlPoint = NULL;
diff --git a/Engine/Source/Runtime/Engine/Private/Layers/Layer.cpp b/Engine/Source/Runtime/Engine/Private/Layers/Layer.cpp
--- a/Engine/Source/Runtime/Engine/Private/Layers/Layer.cpp
+++ b/Engine/Source/Runtime/Engine/Private/Layers/Layer.cpp
@@ -1,8 +1,8 @@
 // Copyright 1998-2014 Epic Games, Inc. All Rights Reserved.
 #include "EnginePrivate.h"
 
-ULayer::ULayer(const class FPostConstructInitializeProperties& PCIP)
-	: Super(PCIP)
+ULayer::ULayer(const FObjectInitializer& ObjectInitializer)
+	: Super(ObjectInitializer)
 	, LayerName( NAME_None )
 	, bIsVisible( true )
 {
diff --git a/Engine/Source/Runtime/Engine/Private/Level.cpp b/Engine/Source/Runtime/Engine/Private/Level.cpp
--- a/Engine/Source/Runtime/Engine/Private/Level.cpp
+++ b/Engine/Source/Runtime/Engine/Private/Level.cpp
@@ -34,8 +34,8 @@ DEFINE_LOG_CATEGORY(LogLevel);
 ULevelBase implementation.
 -----------------------------------------------------------------------------*/
 
-ULevelBase::ULevelBase( const class FPostConstructInitializeProperties& PCIP,const FURL& InURL )
-	: UObject(PCIP)
+ULevelBase::ULevelBase( const FObjectInitializer& ObjectInitializer,const FURL& InURL )
+	: UObject(ObjectInitializer)
 	, URL( InURL )
 	, Actors( this )
 {
@@ -253,16 +253,16 @@ FArchive& operator<<( FArchive& Ar, FPrecomputedVolumeDistanceField& D )
 
 TMap<FName, UWorld*> ULevel::StreamedLevelsOwningWorld;
 
-ULevel::ULevel( const class FPostConstructInitializeProperties& PCIP )
-	:	ULevelBase( PCIP )
+ULevel::ULevel( const FObjectInitializer& ObjectInitializer )
+	:	ULevelBase( ObjectInitializer )
 	,	OwningWorld(NULL)
 	,	TickTaskLevel(FTickTaskManagerInterface::Get().AllocateTickTaskLevel())
 {
 	PrecomputedLightVolume = new FPrecomputedLightVolume();
 }
 
-ULevel::ULevel( const class FPostConstructInitializeProperties& PCIP,const FURL& InURL )
-	:	ULevelBase( PCIP, InURL )
+ULevel::ULevel( const FObjectInitializer& ObjectInitializer,const FURL& InURL )
+	:	ULevelBase( ObjectInitializer, InURL )
 	,	OwningWorld(NULL)
 	,	TickTaskLevel(FTickTaskManagerInterface::Get().AllocateTickTaskLevel())
 {
@@ -927,7 +927,7 @@ void ULevel::CreateModelComponents()
 			Model->Nodes[Nodes[NodeIndex]].ComponentNodeIndex = NodeIndex;
 		}
 
-		UModelComponent* ModelComponent = new(this) UModelComponent(FPostConstructInitializeProperties(),Model,ModelComponents.Num(),Key.MaskedPolyFlags,Nodes);
+		UModelComponent* ModelComponent = new(this) UModelComponent(FObjectInitializer(),Model,ModelComponents.Num(),Key.MaskedPolyFlags,Nodes);
 		ModelComponents.Add(ModelComponent);
 
 		for(int32 NodeIndex = 0;NodeIndex < Nodes.Num();NodeIndex++)
@@ -2120,8 +2120,8 @@ private:
 	FPrimitiveViewRelevance ViewRelevance;
 };
 
-ULineBatchComponent::ULineBatchComponent(const class FPostConstructInitializeProperties& PCIP)
-	: Super(PCIP)
+ULineBatchComponent::ULineBatchComponent(const FObjectInitializer& ObjectInitializer)
+	: Super(ObjectInitializer)
 {
 	bAutoActivate = true;
 	bTickInEditor = true;
diff --git a/Engine/Source/Runtime/Engine/Private/LevelBounds.cpp b/Engine/Source/Runtime/Engine/Private/LevelBounds.cpp
--- a/Engine/Source/Runtime/Engine/Private/LevelBounds.cpp
+++ b/Engine/Source/Runtime/Engine/Private/LevelBounds.cpp
@@ -6,10 +6,10 @@
 // Default size of the box (scale)
 static const FVector DefaultLevelSize = FVector(1000.f);
 
-ALevelBounds::ALevelBounds(const class FPostConstructInitializeProperties& PCIP)
-	: Super(PCIP)
+ALevelBounds::ALevelBounds(const FObjectInitializer& ObjectInitializer)
+	: Super(ObjectInitializer)
 {
-	TSubobjectPtr<UBoxComponent> BoxComponent = PCIP.CreateDefaultSubobject<UBoxComponent>(this, TEXT("BoxComponent0"));
+	TSubobjectPtr<UBoxComponent> BoxComponent = ObjectInitializer.CreateDefaultSubobject<UBoxComponent>(this, TEXT("BoxComponent0"));
 	RootComponent = BoxComponent;
 	RootComponent->Mobility = EComponentMobility::Static;
 	RootComponent->RelativeScale3D = DefaultLevelSize;
diff --git a/Engine/Source/Runtime/Engine/Private/LevelScriptActor.cpp b/Engine/Source/Runtime/Engine/Private/LevelScriptActor.cpp
--- a/Engine/Source/Runtime/Engine/Private/LevelScriptActor.cpp
+++ b/Engine/Source/Runtime/Engine/Private/LevelScriptActor.cpp
@@ -11,8 +11,8 @@
 //////////////////////////////////////////////////////////////////////////
 // ALevelScriptActor
 
-ALevelScriptActor::ALevelScriptActor(const class FPostConstructInitializeProperties& PCIP)
-	: Super(PCIP)
+ALevelScriptActor::ALevelScriptActor(const FObjectInitializer& ObjectInitializer)
+	: Super(ObjectInitializer)
 {
 
 #if WITH_EDITORONLY_DATA
diff --git a/Engine/Source/Runtime/Engine/Private/LevelScriptBlueprint.cpp b/Engine/Source/Runtime/Engine/Private/LevelScriptBlueprint.cpp
--- a/Engine/Source/Runtime/Engine/Private/LevelScriptBlueprint.cpp
+++ b/Engine/Source/Runtime/Engine/Private/LevelScriptBlueprint.cpp
@@ -10,8 +10,8 @@
 //////////////////////////////////////////////////////////////////////////
 // ULevelScriptBlueprint
 
-ULevelScriptBlueprint::ULevelScriptBlueprint(const class FPostConstructInitializeProperties& PCIP)
-	: Super(PCIP)
+ULevelScriptBlueprint::ULevelScriptBlueprint(const FObjectInitializer& ObjectInitializer)
+	: Super(ObjectInitializer)
 {
 }
 
diff --git a/Engine/Source/Runtime/Engine/Private/LevelStreaming.cpp b/Engine/Source/Runtime/Engine/Private/LevelStreaming.cpp
--- a/Engine/Source/Runtime/Engine/Private/LevelStreaming.cpp
+++ b/Engine/Source/Runtime/Engine/Private/LevelStreaming.cpp
@@ -868,8 +868,8 @@ void ULevelStreaming::RemoveStreamingVolumeDuplicates()
 
 #endif // WITH_EDITOR
 
-ULevelStreaming::ULevelStreaming(const class FPostConstructInitializeProperties& PCIP)
-	: Super(PCIP)
+ULevelStreaming::ULevelStreaming(const FObjectInitializer& ObjectInitializer)
+	: Super(ObjectInitializer)
 {
 	bShouldBeVisibleInEditor = true;
 	LevelColor = FLinearColor::White;
@@ -895,16 +895,16 @@ void ULevelStreaming::PostEditUndo()
 /*-----------------------------------------------------------------------------
 	ULevelStreamingPersistent implementation.
 -----------------------------------------------------------------------------*/
-ULevelStreamingPersistent::ULevelStreamingPersistent(const class FPostConstructInitializeProperties& PCIP)
-	: Super(PCIP)
+ULevelStreamingPersistent::ULevelStreamingPersistent(const FObjectInitializer& ObjectInitializer)
+	: Super(ObjectInitializer)
 {
 }
 
 /*-----------------------------------------------------------------------------
 	ULevelStreamingKismet implementation.
 -----------------------------------------------------------------------------*/
-ULevelStreamingKismet::ULevelStreamingKismet(const class FPostConstructInitializeProperties& PCIP)
-	: Super(PCIP)
+ULevelStreamingKismet::ULevelStreamingKismet(const FObjectInitializer& ObjectInitializer)
+	: Super(ObjectInitializer)
 {
 }
 
@@ -933,8 +933,8 @@ bool ULevelStreamingKismet::ShouldBeLoaded( const FVector& ViewLocation )
 /*-----------------------------------------------------------------------------
 	ULevelStreamingAlwaysLoaded implementation.
 -----------------------------------------------------------------------------*/
-ULevelStreamingAlwaysLoaded::ULevelStreamingAlwaysLoaded(const class FPostConstructInitializeProperties& PCIP)
-	: Super(PCIP)
+ULevelStreamingAlwaysLoaded::ULevelStreamingAlwaysLoaded(const FObjectInitializer& ObjectInitializer)
+	: Super(ObjectInitializer)
 {
 	bShouldBeVisible = true;
 }
@@ -947,8 +947,8 @@ bool ULevelStreamingAlwaysLoaded::ShouldBeLoaded( const FVector& ViewLocation )
 /*-----------------------------------------------------------------------------
 	ULevelStreamingBounds implementation.
 -----------------------------------------------------------------------------*/
-ULevelStreamingBounds::ULevelStreamingBounds(const class FPostConstructInitializeProperties& PCIP)
-	: Super(PCIP)
+ULevelStreamingBounds::ULevelStreamingBounds(const FObjectInitializer& ObjectInitializer)
+	: Super(ObjectInitializer)
 {
 }
 
diff --git a/Engine/Source/Runtime/Engine/Private/LevelStreamingVolume.cpp b/Engine/Source/Runtime/Engine/Private/LevelStreamingVolume.cpp
--- a/Engine/Source/Runtime/Engine/Private/LevelStreamingVolume.cpp
+++ b/Engine/Source/Runtime/Engine/Private/LevelStreamingVolume.cpp
@@ -3,8 +3,8 @@
 #include "EnginePrivate.h"
 #include "Engine/LevelStreamingVolume.h"
 
-ALevelStreamingVolume::ALevelStreamingVolume(const class FPostConstructInitializeProperties& PCIP)
-	: Super(PCIP)
+ALevelStreamingVolume::ALevelStreamingVolume(const FObjectInitializer& ObjectInitializer)
+	: Super(ObjectInitializer)
 {
 	BrushComponent->BodyInstance.bEnableCollision_DEPRECATED = false;
 
diff --git a/Engine/Source/Runtime/Engine/Private/Light.cpp b/Engine/Source/Runtime/Engine/Private/Light.cpp
--- a/Engine/Source/Runtime/Engine/Private/Light.cpp
+++ b/Engine/Source/Runtime/Engine/Private/Light.cpp
@@ -8,10 +8,10 @@
 #include "LightingBuildOptions.h"
 #include "Net/UnrealNetwork.h"
 
-ALight::ALight(const class FPostConstructInitializeProperties& PCIP)
-	: Super(PCIP)
+ALight::ALight(const FObjectInitializer& ObjectInitializer)
+	: Super(ObjectInitializer)
 {
-	LightComponent = PCIP.CreateAbstractDefaultSubobject<ULightComponent>(this, TEXT("LightComponent0"));
+	LightComponent = ObjectInitializer.CreateAbstractDefaultSubobject<ULightComponent>(this, TEXT("LightComponent0"));
 
 	bHidden = true;
 	bCollideWhenPlacing = true;
@@ -152,8 +152,8 @@ void ALight::SetAffectTranslucentLighting(bool bNewValue)
 
 ///////////////////////////////////////////////////////////////////////////////////////////////////////////////////////
 
-APointLight::APointLight(const class FPostConstructInitializeProperties& PCIP)
-	: Super(PCIP.SetDefaultSubobjectClass<UPointLightComponent>(TEXT("LightComponent0")))
+APointLight::APointLight(const FObjectInitializer& ObjectInitializer)
+	: Super(ObjectInitializer.SetDefaultSubobjectClass<UPointLightComponent>(TEXT("LightComponent0")))
 {
 	PointLightComponent = CastChecked<UPointLightComponent>(LightComponent);
 	PointLightComponent->Mobility = EComponentMobility::Stationary;
@@ -200,8 +200,8 @@ void APointLight::LoadedFromAnotherClass(const FName& OldClassName)
 }
 #endif // WITH_EDITOR
 
-ADirectionalLight::ADirectionalLight(const class FPostConstructInitializeProperties& PCIP)
-	: Super(PCIP.SetDefaultSubobjectClass<UDirectionalLightComponent>(TEXT("LightComponent0")))
+ADirectionalLight::ADirectionalLight(const FObjectInitializer& ObjectInitializer)
+	: Super(ObjectInitializer.SetDefaultSubobjectClass<UDirectionalLightComponent>(TEXT("LightComponent0")))
 {
 	// Structure to hold one-time initialization
 	struct FConstructorStatics
@@ -226,7 +226,7 @@ ADirectionalLight::ADirectionalLight(const class FPostConstructInitializePropert
 	RootComponent = DirectionalLightComponent;
 
 #if WITH_EDITORONLY_DATA
-	ArrowComponent = PCIP.CreateEditorOnlyDefaultSubobject<UArrowComponent>(this, TEXT("ArrowComponent0"));
+	ArrowComponent = ObjectInitializer.CreateEditorOnlyDefaultSubobject<UArrowComponent>(this, TEXT("ArrowComponent0"));
 	if (ArrowComponent)
 	{
 		ArrowComponent->ArrowColor = FColor(150, 200, 255);
@@ -324,8 +324,8 @@ void APointLight::EditorApplyScale(const FVector& DeltaScale, const FVector* Piv
 }
 #endif
 
-AGeneratedMeshAreaLight::AGeneratedMeshAreaLight(const class FPostConstructInitializeProperties& PCIP)
-	: Super(PCIP)
+AGeneratedMeshAreaLight::AGeneratedMeshAreaLight(const FObjectInitializer& ObjectInitializer)
+	: Super(ObjectInitializer)
 {
 #if WITH_EDITORONLY_DATA
 	bEditable = false;
diff --git a/Engine/Source/Runtime/Engine/Private/LightComponent.cpp b/Engine/Source/Runtime/Engine/Private/LightComponent.cpp
--- a/Engine/Source/Runtime/Engine/Private/LightComponent.cpp
+++ b/Engine/Source/Runtime/Engine/Private/LightComponent.cpp
@@ -274,8 +274,8 @@ void FLightSceneProxy::ApplyWorldOffset(FVector InOffset)
 	SetTransform(NewLightToWorld, NewPosition);
 }
 
-ULightComponentBase::ULightComponentBase(const class FPostConstructInitializeProperties& PCIP)
-	: Super(PCIP)
+ULightComponentBase::ULightComponentBase(const FObjectInitializer& ObjectInitializer)
+	: Super(ObjectInitializer)
 {
 	Brightness_DEPRECATED = 3.1415926535897932f;
 	Intensity = 3.1415926535897932f;
@@ -290,8 +290,8 @@ ULightComponentBase::ULightComponentBase(const class FPostConstructInitializePro
 /**
  * Updates/ resets light GUIDs.
  */
-ULightComponent::ULightComponent(const class FPostConstructInitializeProperties& PCIP)
-	: Super(PCIP)
+ULightComponent::ULightComponent(const FObjectInitializer& ObjectInitializer)
+	: Super(ObjectInitializer)
 {
 	ShadowMapChannel = INDEX_NONE;
 	PreviewShadowMapChannel = INDEX_NONE;
diff --git a/Engine/Source/Runtime/Engine/Private/LightMap.cpp b/Engine/Source/Runtime/Engine/Private/LightMap.cpp
--- a/Engine/Source/Runtime/Engine/Private/LightMap.cpp
+++ b/Engine/Source/Runtime/Engine/Private/LightMap.cpp
@@ -91,8 +91,8 @@ void FLightMap::FinishCleanup()
 	delete this;
 }
 
-ULightMapTexture2D::ULightMapTexture2D(const class FPostConstructInitializeProperties& PCIP)
-	: Super(PCIP)
+ULightMapTexture2D::ULightMapTexture2D(const FObjectInitializer& ObjectInitializer)
+	: Super(ObjectInitializer)
 {
 }
 
@@ -711,7 +711,7 @@ void FLightMapPendingTexture::StartEncoding()
 
 	if (bNeedsSkyOcclusionTexture)
 	{
-		ULightMapTexture2D* Texture = new(Outer, GetSkyOcclusionTextureName(GLightmapCounter)) ULightMapTexture2D(FPostConstructInitializeProperties());
+		ULightMapTexture2D* Texture = new(Outer, GetSkyOcclusionTextureName(GLightmapCounter)) ULightMapTexture2D(FObjectInitializer());
 		SkyOcclusionTexture = Texture;
 
 		Texture->Source.Init2DWithMipChain(GetSizeX(), GetSizeY(), TSF_BGRA8);
@@ -810,7 +810,7 @@ void FLightMapPendingTexture::StartEncoding()
 		}
 
 		// Create the light-map texture for this coefficient.
-		ULightMapTexture2D* Texture = new(Outer, GetLightmapName(GLightmapCounter, CoefficientIndex)) ULightMapTexture2D(FPostConstructInitializeProperties());
+		ULightMapTexture2D* Texture = new(Outer, GetLightmapName(GLightmapCounter, CoefficientIndex)) ULightMapTexture2D(FObjectInitializer());
 		Textures[CoefficientIndex] = Texture;
 		Texture->Source.Init2DWithMipChain(GetSizeX(), GetSizeY() * 2, TSF_BGRA8);	// Top/bottom atlased
 		Texture->MipGenSettings = TMGS_LeaveExistingMips;
diff --git a/Engine/Source/Runtime/Engine/Private/LightmappedSurfaceCollection.cpp b/Engine/Source/Runtime/Engine/Private/LightmappedSurfaceCollection.cpp
--- a/Engine/Source/Runtime/Engine/Private/LightmappedSurfaceCollection.cpp
+++ b/Engine/Source/Runtime/Engine/Private/LightmappedSurfaceCollection.cpp
@@ -2,7 +2,7 @@
 
 #include "EnginePrivate.h"
 
-ULightmappedSurfaceCollection::ULightmappedSurfaceCollection(const class FPostConstructInitializeProperties& PCIP)
-	: Super(PCIP)
+ULightmappedSurfaceCollection::ULightmappedSurfaceCollection(const FObjectInitializer& ObjectInitializer)
+	: Super(ObjectInitializer)
 {
 }
diff --git a/Engine/Source/Runtime/Engine/Private/LightmassCharacterIndirectDetailVolume.cpp b/Engine/Source/Runtime/Engine/Private/LightmassCharacterIndirectDetailVolume.cpp
--- a/Engine/Source/Runtime/Engine/Private/LightmassCharacterIndirectDetailVolume.cpp
+++ b/Engine/Source/Runtime/Engine/Private/LightmassCharacterIndirectDetailVolume.cpp
@@ -4,8 +4,8 @@
 
 #include "../Classes/Lightmass/LightmassCharacterIndirectDetailVolume.h"
 
-ALightmassCharacterIndirectDetailVolume::ALightmassCharacterIndirectDetailVolume(const class FPostConstructInitializeProperties& PCIP)
-	: Super(PCIP)
+ALightmassCharacterIndirectDetailVolume::ALightmassCharacterIndirectDetailVolume(const FObjectInitializer& ObjectInitializer)
+	: Super(ObjectInitializer)
 {
 	BrushComponent->BodyInstance.bEnableCollision_DEPRECATED = false;
 
diff --git a/Engine/Source/Runtime/Engine/Private/LightmassImportanceVolume.cpp b/Engine/Source/Runtime/Engine/Private/LightmassImportanceVolume.cpp
--- a/Engine/Source/Runtime/Engine/Private/LightmassImportanceVolume.cpp
+++ b/Engine/Source/Runtime/Engine/Private/LightmassImportanceVolume.cpp
@@ -4,8 +4,8 @@
 
 #include "../Classes/Lightmass/LightmassImportanceVolume.h"
 
-ALightmassImportanceVolume::ALightmassImportanceVolume(const class FPostConstructInitializeProperties& PCIP)
-	: Super(PCIP)
+ALightmassImportanceVolume::ALightmassImportanceVolume(const FObjectInitializer& ObjectInitializer)
+	: Super(ObjectInitializer)
 {
 	BrushComponent->BodyInstance.bEnableCollision_DEPRECATED = false;
 
diff --git a/Engine/Source/Runtime/Engine/Private/LightmassPrimitiveSettingsObject.cpp b/Engine/Source/Runtime/Engine/Private/LightmassPrimitiveSettingsObject.cpp
--- a/Engine/Source/Runtime/Engine/Private/LightmassPrimitiveSettingsObject.cpp
+++ b/Engine/Source/Runtime/Engine/Private/LightmassPrimitiveSettingsObject.cpp
@@ -2,7 +2,7 @@
 
 #include "EnginePrivate.h"
 
-ULightmassPrimitiveSettingsObject::ULightmassPrimitiveSettingsObject(const class FPostConstructInitializeProperties& PCIP)
-	: Super(PCIP)
+ULightmassPrimitiveSettingsObject::ULightmassPrimitiveSettingsObject(const FObjectInitializer& ObjectInitializer)
+	: Super(ObjectInitializer)
 {
 }
diff --git a/Engine/Source/Runtime/Engine/Private/LocalMessage.cpp b/Engine/Source/Runtime/Engine/Private/LocalMessage.cpp
--- a/Engine/Source/Runtime/Engine/Private/LocalMessage.cpp
+++ b/Engine/Source/Runtime/Engine/Private/LocalMessage.cpp
@@ -14,7 +14,7 @@ FClientReceiveData::FClientReceiveData()
 	
 }
 
-ULocalMessage::ULocalMessage(const class FPostConstructInitializeProperties& PCIP)
-	: Super(PCIP)
+ULocalMessage::ULocalMessage(const FObjectInitializer& ObjectInitializer)
+	: Super(ObjectInitializer)
 {
 }
diff --git a/Engine/Source/Runtime/Engine/Private/LocalPlayer.cpp b/Engine/Source/Runtime/Engine/Private/LocalPlayer.cpp
--- a/Engine/Source/Runtime/Engine/Private/LocalPlayer.cpp
+++ b/Engine/Source/Runtime/Engine/Private/LocalPlayer.cpp
@@ -138,8 +138,8 @@ void FLocalPlayerContext::SetPlayerController( const APlayerController* InPlayer
 //////////////////////////////////////////////////////////////////////////
 // ULocalPlayer
 
-ULocalPlayer::ULocalPlayer(const class FPostConstructInitializeProperties& PCIP)
-	: Super(PCIP)
+ULocalPlayer::ULocalPlayer(const FObjectInitializer& ObjectInitializer)
+	: Super(ObjectInitializer)
 {
 	PendingLevelPlayerControllerClass = APlayerController::StaticClass();
 }
diff --git a/Engine/Source/Runtime/Engine/Private/MaterialBillboardComponent.cpp b/Engine/Source/Runtime/Engine/Private/MaterialBillboardComponent.cpp
--- a/Engine/Source/Runtime/Engine/Private/MaterialBillboardComponent.cpp
+++ b/Engine/Source/Runtime/Engine/Private/MaterialBillboardComponent.cpp
@@ -357,8 +357,8 @@ private:
 	FColor BaseColor;
 };
 
-UMaterialBillboardComponent::UMaterialBillboardComponent(const class FPostConstructInitializeProperties& PCIP)
-: Super(PCIP)
+UMaterialBillboardComponent::UMaterialBillboardComponent(const FObjectInitializer& ObjectInitializer)
+: Super(ObjectInitializer)
 {
 	BodyInstance.bEnableCollision_DEPRECATED = false;
 
diff --git a/Engine/Source/Runtime/Engine/Private/Materials/Material.cpp b/Engine/Source/Runtime/Engine/Private/Materials/Material.cpp
--- a/Engine/Source/Runtime/Engine/Private/Materials/Material.cpp
+++ b/Engine/Source/Runtime/Engine/Private/Materials/Material.cpp
@@ -491,8 +491,8 @@ bool UMaterial::IsDefaultMaterial() const
 	return bDefault;
 }
 
-UMaterial::UMaterial(const class FPostConstructInitializeProperties& PCIP)
-	: Super(PCIP)
+UMaterial::UMaterial(const FObjectInitializer& ObjectInitializer)
+	: Super(ObjectInitializer)
 {
 	BlendMode = BLEND_Opaque;
 	ShadingModel = MSM_DefaultLit;
