Rename FPostConstructInitializeProperties to something simpler

diff --git a/Engine/Source/Runtime/Engine/Private/Materials/MaterialExpressions.cpp b/Engine/Source/Runtime/Engine/Private/Materials/MaterialExpressions.cpp
--- a/Engine/Source/Runtime/Engine/Private/Materials/MaterialExpressions.cpp
+++ b/Engine/Source/Runtime/Engine/Private/Materials/MaterialExpressions.cpp
@@ -304,8 +304,8 @@ bool CanConnectMaterialValueTypes(uint32 InputType, uint32 OutputType)
 	return false;
 }
 
-UMaterialExpression::UMaterialExpression(const class FPostConstructInitializeProperties& PCIP)
-	: Super(PCIP)
+UMaterialExpression::UMaterialExpression(const FObjectInitializer& ObjectInitializer)
+	: Super(ObjectInitializer)
 #if WITH_EDITORONLY_DATA
 	, GraphNode(NULL)
 #endif
@@ -1016,8 +1016,8 @@ bool UMaterialExpression::ContainsInputLoopInternal(TArray<FMaterialExpressionKe
 	return false;
 }
 
-UMaterialExpressionTextureBase::UMaterialExpressionTextureBase(const class FPostConstructInitializeProperties& PCIP)
-	: Super(PCIP)
+UMaterialExpressionTextureBase::UMaterialExpressionTextureBase(const FObjectInitializer& ObjectInitializer)
+	: Super(ObjectInitializer)
 	, IsDefaultMeshpaintTexture(false)
 {}
 
@@ -1109,8 +1109,8 @@ EMaterialSamplerType UMaterialExpressionTextureBase::GetSamplerTypeForTexture(co
 
 
 
-UMaterialExpressionTextureSample::UMaterialExpressionTextureSample(const class FPostConstructInitializeProperties& PCIP)
-	: Super(PCIP)
+UMaterialExpressionTextureSample::UMaterialExpressionTextureSample(const FObjectInitializer& ObjectInitializer)
+	: Super(ObjectInitializer)
 {
 	// Structure to hold one-time initialization
 	struct FConstructorStatics
@@ -1409,8 +1409,8 @@ uint32 UMaterialExpressionTextureSample::GetInputType(int32 InputIndex)
 
 
 
-UMaterialExpressionAdd::UMaterialExpressionAdd(const class FPostConstructInitializeProperties& PCIP)
-	: Super(PCIP)
+UMaterialExpressionAdd::UMaterialExpressionAdd(const FObjectInitializer& ObjectInitializer)
+	: Super(ObjectInitializer)
 {
 	// Structure to hold one-time initialization
 	struct FConstructorStatics
@@ -1433,8 +1433,8 @@ UMaterialExpressionAdd::UMaterialExpressionAdd(const class FPostConstructInitial
 //
 //  UMaterialExpressionTextureSampleParameter
 //
-UMaterialExpressionTextureSampleParameter::UMaterialExpressionTextureSampleParameter(const class FPostConstructInitializeProperties& PCIP)
-	: Super(PCIP)
+UMaterialExpressionTextureSampleParameter::UMaterialExpressionTextureSampleParameter(const FObjectInitializer& ObjectInitializer)
+	: Super(ObjectInitializer)
 {
 	// Structure to hold one-time initialization
 	struct FConstructorStatics
@@ -1522,8 +1522,8 @@ void UMaterialExpressionTextureSampleParameter::GetAllParameterNames(TArray<FNam
 //
 //  UMaterialExpressionTextureObjectParameter
 //
-UMaterialExpressionTextureObjectParameter::UMaterialExpressionTextureObjectParameter(const class FPostConstructInitializeProperties& PCIP)
-	: Super(PCIP)
+UMaterialExpressionTextureObjectParameter::UMaterialExpressionTextureObjectParameter(const FObjectInitializer& ObjectInitializer)
+	: Super(ObjectInitializer)
 {
 	// Structure to hold one-time initialization
 	struct FConstructorStatics
@@ -1591,8 +1591,8 @@ int32 UMaterialExpressionTextureObjectParameter::CompilePreview(class FMaterialC
 //
 //  UMaterialExpressionTextureObject
 //
-UMaterialExpressionTextureObject::UMaterialExpressionTextureObject(const class FPostConstructInitializeProperties& PCIP)
-	: Super(PCIP)
+UMaterialExpressionTextureObject::UMaterialExpressionTextureObject(const FObjectInitializer& ObjectInitializer)
+	: Super(ObjectInitializer)
 {
 	// Structure to hold one-time initialization
 	struct FConstructorStatics
@@ -1678,8 +1678,8 @@ uint32 UMaterialExpressionTextureObject::GetOutputType(int32 OutputIndex)
 //
 //  UMaterialExpressionTextureSampleParameter2D
 //
-UMaterialExpressionTextureSampleParameter2D::UMaterialExpressionTextureSampleParameter2D(const class FPostConstructInitializeProperties& PCIP)
-	: Super(PCIP)
+UMaterialExpressionTextureSampleParameter2D::UMaterialExpressionTextureSampleParameter2D(const FObjectInitializer& ObjectInitializer)
+	: Super(ObjectInitializer)
 {
 	// Structure to hold one-time initialization
 	struct FConstructorStatics
@@ -1767,8 +1767,8 @@ void UMaterialExpressionTextureSampleParameter::SetEditableName(const FString& N
 //
 //  UMaterialExpressionTextureSampleParameterCube
 //
-UMaterialExpressionTextureSampleParameterCube::UMaterialExpressionTextureSampleParameterCube(const class FPostConstructInitializeProperties& PCIP)
-	: Super(PCIP)
+UMaterialExpressionTextureSampleParameterCube::UMaterialExpressionTextureSampleParameterCube(const FObjectInitializer& ObjectInitializer)
+	: Super(ObjectInitializer)
 {
 	// Structure to hold one-time initialization
 	struct FConstructorStatics
@@ -1844,8 +1844,8 @@ static int32 ParticleSubUV(FMaterialCompiler* Compiler, int32 TextureIndex, UTex
 /** 
  *	UMaterialExpressionTextureSampleParameterSubUV
  */
-UMaterialExpressionTextureSampleParameterSubUV::UMaterialExpressionTextureSampleParameterSubUV(const class FPostConstructInitializeProperties& PCIP)
-	: Super(PCIP)
+UMaterialExpressionTextureSampleParameterSubUV::UMaterialExpressionTextureSampleParameterSubUV(const FObjectInitializer& ObjectInitializer)
+	: Super(ObjectInitializer)
 {
 	// Structure to hold one-time initialization
 	struct FConstructorStatics
@@ -1928,8 +1928,8 @@ void UMaterialExpressionAdd::GetCaption(TArray<FString>& OutCaptions) const
 //
 //	UMaterialExpressionMultiply
 //
-UMaterialExpressionMultiply::UMaterialExpressionMultiply(const class FPostConstructInitializeProperties& PCIP)
-	: Super(PCIP)
+UMaterialExpressionMultiply::UMaterialExpressionMultiply(const FObjectInitializer& ObjectInitializer)
+	: Super(ObjectInitializer)
 {
 	// Structure to hold one-time initialization
 	struct FConstructorStatics
@@ -1973,8 +1973,8 @@ void UMaterialExpressionMultiply::GetCaption(TArray<FString>& OutCaptions) const
 	OutCaptions.Add(ret);
 }
 
-UMaterialExpressionDivide::UMaterialExpressionDivide(const class FPostConstructInitializeProperties& PCIP)
-	: Super(PCIP)
+UMaterialExpressionDivide::UMaterialExpressionDivide(const FObjectInitializer& ObjectInitializer)
+	: Super(ObjectInitializer)
 {
 	// Structure to hold one-time initialization
 	struct FConstructorStatics
@@ -2019,8 +2019,8 @@ void UMaterialExpressionDivide::GetCaption(TArray<FString>& OutCaptions) const
 //
 //	UMaterialExpressionSubtract
 //
-UMaterialExpressionSubtract::UMaterialExpressionSubtract(const class FPostConstructInitializeProperties& PCIP)
-	: Super(PCIP)
+UMaterialExpressionSubtract::UMaterialExpressionSubtract(const FObjectInitializer& ObjectInitializer)
+	: Super(ObjectInitializer)
 {
 	// Structure to hold one-time initialization
 	struct FConstructorStatics
@@ -2067,8 +2067,8 @@ void UMaterialExpressionSubtract::GetCaption(TArray<FString>& OutCaptions) const
 //	UMaterialExpressionLinearInterpolate
 //
 
-UMaterialExpressionLinearInterpolate::UMaterialExpressionLinearInterpolate(const class FPostConstructInitializeProperties& PCIP)
-	: Super(PCIP)
+UMaterialExpressionLinearInterpolate::UMaterialExpressionLinearInterpolate(const FObjectInitializer& ObjectInitializer)
+	: Super(ObjectInitializer)
 {
 	// Structure to hold one-time initialization
 	struct FConstructorStatics
@@ -2118,8 +2118,8 @@ void UMaterialExpressionLinearInterpolate::GetCaption(TArray<FString>& OutCaptio
 	OutCaptions.Add(ret);
 }
 
-UMaterialExpressionConstant::UMaterialExpressionConstant(const class FPostConstructInitializeProperties& PCIP)
-	: Super(PCIP)
+UMaterialExpressionConstant::UMaterialExpressionConstant(const FObjectInitializer& ObjectInitializer)
+	: Super(ObjectInitializer)
 {
 	// Structure to hold one-time initialization
 	struct FConstructorStatics
@@ -2159,8 +2159,8 @@ FString UMaterialExpressionConstant::GetDescription() const
 }
 #endif // WITH_EDITOR
 
-UMaterialExpressionConstant2Vector::UMaterialExpressionConstant2Vector(const class FPostConstructInitializeProperties& PCIP)
-	: Super(PCIP)
+UMaterialExpressionConstant2Vector::UMaterialExpressionConstant2Vector(const FObjectInitializer& ObjectInitializer)
+	: Super(ObjectInitializer)
 {
 	// Structure to hold one-time initialization
 	struct FConstructorStatics
@@ -2202,8 +2202,8 @@ FString UMaterialExpressionConstant2Vector::GetDescription() const
 }
 #endif // WITH_EDITOR
 
-UMaterialExpressionConstant3Vector::UMaterialExpressionConstant3Vector(const class FPostConstructInitializeProperties& PCIP)
-	: Super(PCIP)
+UMaterialExpressionConstant3Vector::UMaterialExpressionConstant3Vector(const FObjectInitializer& ObjectInitializer)
+	: Super(ObjectInitializer)
 {
 	// Structure to hold one-time initialization
 	struct FConstructorStatics
@@ -2257,8 +2257,8 @@ FString UMaterialExpressionConstant3Vector::GetDescription() const
 }
 #endif // WITH_EDITOR
 
-UMaterialExpressionConstant4Vector::UMaterialExpressionConstant4Vector(const class FPostConstructInitializeProperties& PCIP)
-	: Super(PCIP)
+UMaterialExpressionConstant4Vector::UMaterialExpressionConstant4Vector(const FObjectInitializer& ObjectInitializer)
+	: Super(ObjectInitializer)
 {
 	// Structure to hold one-time initialization
 	struct FConstructorStatics
@@ -2314,8 +2314,8 @@ FString UMaterialExpressionConstant4Vector::GetDescription() const
 }
 #endif // WITH_EDITOR
 
-UMaterialExpressionClamp::UMaterialExpressionClamp(const class FPostConstructInitializeProperties& PCIP)
-	: Super(PCIP)
+UMaterialExpressionClamp::UMaterialExpressionClamp(const FObjectInitializer& ObjectInitializer)
+	: Super(ObjectInitializer)
 {
 	// Structure to hold one-time initialization
 	struct FConstructorStatics
@@ -2399,8 +2399,8 @@ void UMaterialExpressionClamp::GetCaption(TArray<FString>& OutCaptions) const
 //	UMaterialExpressionMin
 //
 
-UMaterialExpressionMin::UMaterialExpressionMin(const class FPostConstructInitializeProperties& PCIP)
-	: Super(PCIP)
+UMaterialExpressionMin::UMaterialExpressionMin(const FObjectInitializer& ObjectInitializer)
+	: Super(ObjectInitializer)
 {
 	// Structure to hold one-time initialization
 	struct FConstructorStatics
@@ -2447,8 +2447,8 @@ void UMaterialExpressionMin::GetCaption(TArray<FString>& OutCaptions) const
 //	UMaterialExpressionMax
 //
 
-UMaterialExpressionMax::UMaterialExpressionMax(const class FPostConstructInitializeProperties& PCIP)
-	: Super(PCIP)
+UMaterialExpressionMax::UMaterialExpressionMax(const FObjectInitializer& ObjectInitializer)
+	: Super(ObjectInitializer)
 {
 	// Structure to hold one-time initialization
 	struct FConstructorStatics
@@ -2495,8 +2495,8 @@ void UMaterialExpressionMax::GetCaption(TArray<FString>& OutCaptions) const
 //	UMaterialExpressionTextureCoordinate
 //
 
-UMaterialExpressionTextureCoordinate::UMaterialExpressionTextureCoordinate(const class FPostConstructInitializeProperties& PCIP)
-	: Super(PCIP)
+UMaterialExpressionTextureCoordinate::UMaterialExpressionTextureCoordinate(const FObjectInitializer& ObjectInitializer)
+	: Super(ObjectInitializer)
 {
 	// Structure to hold one-time initialization
 	struct FConstructorStatics
@@ -2539,8 +2539,8 @@ void UMaterialExpressionTextureCoordinate::GetCaption(TArray<FString>& OutCaptio
 }
 
 
-UMaterialExpressionDotProduct::UMaterialExpressionDotProduct(const class FPostConstructInitializeProperties& PCIP)
-	: Super(PCIP)
+UMaterialExpressionDotProduct::UMaterialExpressionDotProduct(const FObjectInitializer& ObjectInitializer)
+	: Super(ObjectInitializer)
 {
 	// Structure to hold one-time initialization
 	struct FConstructorStatics
@@ -2585,8 +2585,8 @@ void UMaterialExpressionDotProduct::GetCaption(TArray<FString>& OutCaptions) con
 	OutCaptions.Add(TEXT("Dot"));
 }
 
-UMaterialExpressionCrossProduct::UMaterialExpressionCrossProduct(const class FPostConstructInitializeProperties& PCIP)
-	: Super(PCIP)
+UMaterialExpressionCrossProduct::UMaterialExpressionCrossProduct(const FObjectInitializer& ObjectInitializer)
+	: Super(ObjectInitializer)
 {
 	// Structure to hold one-time initialization
 	struct FConstructorStatics
@@ -2631,8 +2631,8 @@ void UMaterialExpressionCrossProduct::GetCaption(TArray<FString>& OutCaptions) c
 	OutCaptions.Add(TEXT("Cross"));
 }
 
-UMaterialExpressionComponentMask::UMaterialExpressionComponentMask(const class FPostConstructInitializeProperties& PCIP)
-	: Super(PCIP)
+UMaterialExpressionComponentMask::UMaterialExpressionComponentMask(const FObjectInitializer& ObjectInitializer)
+	: Super(ObjectInitializer)
 {
 	// Structure to hold one-time initialization
 	struct FConstructorStatics
@@ -2682,8 +2682,8 @@ void UMaterialExpressionComponentMask::GetCaption(TArray<FString>& OutCaptions)
 //
 //	UMaterialExpressionStaticComponentMaskParameter
 //
-UMaterialExpressionStaticComponentMaskParameter::UMaterialExpressionStaticComponentMaskParameter(const class FPostConstructInitializeProperties& PCIP)
-	: Super(PCIP)
+UMaterialExpressionStaticComponentMaskParameter::UMaterialExpressionStaticComponentMaskParameter(const FObjectInitializer& ObjectInitializer)
+	: Super(ObjectInitializer)
 {
 	// Structure to hold one-time initialization
 	struct FConstructorStatics
@@ -2729,8 +2729,8 @@ void UMaterialExpressionStaticComponentMaskParameter::GetCaption(TArray<FString>
 //	UMaterialExpressionTime
 //
 
-UMaterialExpressionTime::UMaterialExpressionTime(const class FPostConstructInitializeProperties& PCIP)
-	: Super(PCIP)
+UMaterialExpressionTime::UMaterialExpressionTime(const FObjectInitializer& ObjectInitializer)
+	: Super(ObjectInitializer)
 {
 	// Structure to hold one-time initialization
 	struct FConstructorStatics
@@ -2758,8 +2758,8 @@ void UMaterialExpressionTime::GetCaption(TArray<FString>& OutCaptions) const
 }
 
 
-UMaterialExpressionCameraVectorWS::UMaterialExpressionCameraVectorWS(const class FPostConstructInitializeProperties& PCIP)
-	: Super(PCIP)
+UMaterialExpressionCameraVectorWS::UMaterialExpressionCameraVectorWS(const FObjectInitializer& ObjectInitializer)
+	: Super(ObjectInitializer)
 {
 	// Structure to hold one-time initialization
 	struct FConstructorStatics
@@ -2789,8 +2789,8 @@ void UMaterialExpressionCameraVectorWS::GetCaption(TArray<FString>& OutCaptions)
 }
 
 
-UMaterialExpressionCameraPositionWS::UMaterialExpressionCameraPositionWS(const class FPostConstructInitializeProperties& PCIP)
-	: Super(PCIP)
+UMaterialExpressionCameraPositionWS::UMaterialExpressionCameraPositionWS(const FObjectInitializer& ObjectInitializer)
+	: Super(ObjectInitializer)
 {
 	// Structure to hold one-time initialization
 	struct FConstructorStatics
@@ -2825,8 +2825,8 @@ void UMaterialExpressionCameraPositionWS::GetCaption(TArray<FString>& OutCaption
 //	UMaterialExpressionReflectionVectorWS
 //
 
-UMaterialExpressionReflectionVectorWS::UMaterialExpressionReflectionVectorWS(const class FPostConstructInitializeProperties& PCIP)
-	: Super(PCIP)
+UMaterialExpressionReflectionVectorWS::UMaterialExpressionReflectionVectorWS(const FObjectInitializer& ObjectInitializer)
+	: Super(ObjectInitializer)
 {
 	// Structure to hold one-time initialization
 	struct FConstructorStatics
@@ -2865,8 +2865,8 @@ void UMaterialExpressionReflectionVectorWS::GetCaption(TArray<FString>& OutCapti
 //
 //	UMaterialExpressionPanner
 //
-UMaterialExpressionPanner::UMaterialExpressionPanner(const class FPostConstructInitializeProperties& PCIP)
-	: Super(PCIP)
+UMaterialExpressionPanner::UMaterialExpressionPanner(const FObjectInitializer& ObjectInitializer)
+	: Super(ObjectInitializer)
 {
 	// Structure to hold one-time initialization
 	struct FConstructorStatics
@@ -2920,8 +2920,8 @@ void UMaterialExpressionPanner::GetCaption(TArray<FString>& OutCaptions) const
 //
 //	UMaterialExpressionRotator
 //
-UMaterialExpressionRotator::UMaterialExpressionRotator(const class FPostConstructInitializeProperties& PCIP)
-	: Super(PCIP)
+UMaterialExpressionRotator::UMaterialExpressionRotator(const FObjectInitializer& ObjectInitializer)
+	: Super(ObjectInitializer)
 {
 	// Structure to hold one-time initialization
 	struct FConstructorStatics
@@ -2991,8 +2991,8 @@ void UMaterialExpressionRotator::GetCaption(TArray<FString>& OutCaptions) const
 //
 //	UMaterialExpressionSine
 //
-UMaterialExpressionSine::UMaterialExpressionSine(const class FPostConstructInitializeProperties& PCIP)
-	: Super(PCIP)
+UMaterialExpressionSine::UMaterialExpressionSine(const FObjectInitializer& ObjectInitializer)
+	: Super(ObjectInitializer)
 {
 	// Structure to hold one-time initialization
 	struct FConstructorStatics
@@ -3025,8 +3025,8 @@ void UMaterialExpressionSine::GetCaption(TArray<FString>& OutCaptions) const
 }
 
 
-UMaterialExpressionCosine::UMaterialExpressionCosine(const class FPostConstructInitializeProperties& PCIP)
-	: Super(PCIP)
+UMaterialExpressionCosine::UMaterialExpressionCosine(const FObjectInitializer& ObjectInitializer)
+	: Super(ObjectInitializer)
 {
 	// Structure to hold one-time initialization
 	struct FConstructorStatics
@@ -3060,8 +3060,8 @@ void UMaterialExpressionCosine::GetCaption(TArray<FString>& OutCaptions) const
 	OutCaptions.Add(TEXT("Cosine"));
 }
 
-UMaterialExpressionBumpOffset::UMaterialExpressionBumpOffset(const class FPostConstructInitializeProperties& PCIP)
-	: Super(PCIP)
+UMaterialExpressionBumpOffset::UMaterialExpressionBumpOffset(const FObjectInitializer& ObjectInitializer)
+	: Super(ObjectInitializer)
 {
 	// Structure to hold one-time initialization
 	struct FConstructorStatics
@@ -3110,8 +3110,8 @@ void UMaterialExpressionBumpOffset::GetCaption(TArray<FString>& OutCaptions) con
 	OutCaptions.Add(TEXT("BumpOffset"));
 }
 
-UMaterialExpressionAppendVector::UMaterialExpressionAppendVector(const class FPostConstructInitializeProperties& PCIP)
-	: Super(PCIP)
+UMaterialExpressionAppendVector::UMaterialExpressionAppendVector(const FObjectInitializer& ObjectInitializer)
+	: Super(ObjectInitializer)
 {
 	// Structure to hold one-time initialization
 	struct FConstructorStatics
@@ -3159,8 +3159,8 @@ void UMaterialExpressionAppendVector::GetCaption(TArray<FString>& OutCaptions) c
 
 // -----
 
-UMaterialExpressionMakeMaterialAttributes::UMaterialExpressionMakeMaterialAttributes(const class FPostConstructInitializeProperties& PCIP)
-	: Super(PCIP)
+UMaterialExpressionMakeMaterialAttributes::UMaterialExpressionMakeMaterialAttributes(const FObjectInitializer& ObjectInitializer)
+	: Super(ObjectInitializer)
 {
 	// Structure to hold one-time initialization
 	struct FConstructorStatics
@@ -3224,8 +3224,8 @@ void UMaterialExpressionMakeMaterialAttributes::GetCaption(TArray<FString>& OutC
 
 // -----
 
-UMaterialExpressionBreakMaterialAttributes::UMaterialExpressionBreakMaterialAttributes(const class FPostConstructInitializeProperties& PCIP)
-	: Super(PCIP)
+UMaterialExpressionBreakMaterialAttributes::UMaterialExpressionBreakMaterialAttributes(const FObjectInitializer& ObjectInitializer)
+	: Super(ObjectInitializer)
 {
 	// Structure to hold one-time initialization
 	struct FConstructorStatics
@@ -3331,8 +3331,8 @@ bool UMaterialExpressionBreakMaterialAttributes::IsInputConnectionRequired(int32
 
 // -----
 
-UMaterialExpressionFloor::UMaterialExpressionFloor(const class FPostConstructInitializeProperties& PCIP)
-	: Super(PCIP)
+UMaterialExpressionFloor::UMaterialExpressionFloor(const FObjectInitializer& ObjectInitializer)
+	: Super(ObjectInitializer)
 {
 	// Structure to hold one-time initialization
 	struct FConstructorStatics
@@ -3363,8 +3363,8 @@ void UMaterialExpressionFloor::GetCaption(TArray<FString>& OutCaptions) const
 	OutCaptions.Add(TEXT("Floor"));
 }
 
-UMaterialExpressionCeil::UMaterialExpressionCeil(const class FPostConstructInitializeProperties& PCIP)
-	: Super(PCIP)
+UMaterialExpressionCeil::UMaterialExpressionCeil(const FObjectInitializer& ObjectInitializer)
+	: Super(ObjectInitializer)
 {
 	// Structure to hold one-time initialization
 	struct FConstructorStatics
@@ -3401,8 +3401,8 @@ void UMaterialExpressionCeil::GetCaption(TArray<FString>& OutCaptions) const
 //	UMaterialExpressionFmod
 //
 
-UMaterialExpressionFmod::UMaterialExpressionFmod(const class FPostConstructInitializeProperties& PCIP)
-	: Super(PCIP)
+UMaterialExpressionFmod::UMaterialExpressionFmod(const FObjectInitializer& ObjectInitializer)
+	: Super(ObjectInitializer)
 {
 	// Structure to hold one-time initialization
 	struct FConstructorStatics
@@ -3438,8 +3438,8 @@ void UMaterialExpressionFmod::GetCaption(TArray<FString>& OutCaptions) const
 	OutCaptions.Add(TEXT("Fmod"));
 }
 
-UMaterialExpressionFrac::UMaterialExpressionFrac(const class FPostConstructInitializeProperties& PCIP)
-	: Super(PCIP)
+UMaterialExpressionFrac::UMaterialExpressionFrac(const FObjectInitializer& ObjectInitializer)
+	: Super(ObjectInitializer)
 {
 	// Structure to hold one-time initialization
 	struct FConstructorStatics
@@ -3471,8 +3471,8 @@ void UMaterialExpressionFrac::GetCaption(TArray<FString>& OutCaptions) const
 	OutCaptions.Add(TEXT("Frac"));
 }
 
-UMaterialExpressionDesaturation::UMaterialExpressionDesaturation(const class FPostConstructInitializeProperties& PCIP)
-	: Super(PCIP)
+UMaterialExpressionDesaturation::UMaterialExpressionDesaturation(const FObjectInitializer& ObjectInitializer)
+	: Super(ObjectInitializer)
 {
 	// Structure to hold one-time initialization
 	struct FConstructorStatics
@@ -3511,8 +3511,8 @@ int32 UMaterialExpressionDesaturation::Compile(class FMaterialCompiler* Compiler
 //	UMaterialExpressionParameter
 //
 
-UMaterialExpressionParameter::UMaterialExpressionParameter(const class FPostConstructInitializeProperties& PCIP)
-	: Super(PCIP)
+UMaterialExpressionParameter::UMaterialExpressionParameter(const FObjectInitializer& ObjectInitializer)
+	: Super(ObjectInitializer)
 {
 	// Structure to hold one-time initialization
 	struct FConstructorStatics
@@ -3566,8 +3566,8 @@ void UMaterialExpressionParameter::GetAllParameterNames(TArray<FName> &OutParame
 //
 //	UMaterialExpressionVectorParameter
 //
-UMaterialExpressionVectorParameter::UMaterialExpressionVectorParameter(const class FPostConstructInitializeProperties& PCIP)
-	: Super(PCIP)
+UMaterialExpressionVectorParameter::UMaterialExpressionVectorParameter(const FObjectInitializer& ObjectInitializer)
+	: Super(ObjectInitializer)
 {
 
 	Outputs.Reset();
@@ -3598,8 +3598,8 @@ void UMaterialExpressionVectorParameter::GetCaption(TArray<FString>& OutCaptions
 //
 //	UMaterialExpressionScalarParameter
 //
-UMaterialExpressionScalarParameter::UMaterialExpressionScalarParameter(const class FPostConstructInitializeProperties& PCIP)
-	: Super(PCIP)
+UMaterialExpressionScalarParameter::UMaterialExpressionScalarParameter(const FObjectInitializer& ObjectInitializer)
+	: Super(ObjectInitializer)
 {
 
 	bCollapsed = true;
@@ -3623,8 +3623,8 @@ void UMaterialExpressionScalarParameter::GetCaption(TArray<FString>& OutCaptions
 //
 //	UMaterialExpressionStaticSwitchParameter
 //
-UMaterialExpressionStaticSwitchParameter::UMaterialExpressionStaticSwitchParameter(const class FPostConstructInitializeProperties& PCIP)
-	: Super(PCIP)
+UMaterialExpressionStaticSwitchParameter::UMaterialExpressionStaticSwitchParameter(const FObjectInitializer& ObjectInitializer)
+	: Super(ObjectInitializer)
 {
 }
 
@@ -3696,8 +3696,8 @@ FString UMaterialExpressionStaticSwitchParameter::GetInputName(int32 InputIndex)
 //
 //	UMaterialExpressionStaticBoolParameter
 //
-UMaterialExpressionStaticBoolParameter::UMaterialExpressionStaticBoolParameter(const class FPostConstructInitializeProperties& PCIP)
-	: Super(PCIP)
+UMaterialExpressionStaticBoolParameter::UMaterialExpressionStaticBoolParameter(const FObjectInitializer& ObjectInitializer)
+	: Super(ObjectInitializer)
 {
 	bHidePreviewWindow = true;
 }
@@ -3722,8 +3722,8 @@ void UMaterialExpressionStaticBoolParameter::GetCaption(TArray<FString>& OutCapt
 //
 //	UMaterialExpressionStaticBool
 //
-UMaterialExpressionStaticBool::UMaterialExpressionStaticBool(const class FPostConstructInitializeProperties& PCIP)
-	: Super(PCIP)
+UMaterialExpressionStaticBool::UMaterialExpressionStaticBool(const FObjectInitializer& ObjectInitializer)
+	: Super(ObjectInitializer)
 {
 	// Structure to hold one-time initialization
 	struct FConstructorStatics
@@ -3761,8 +3761,8 @@ void UMaterialExpressionStaticBool::GetCaption(TArray<FString>& OutCaptions) con
 //
 //	UMaterialExpressionStaticSwitch
 //
-UMaterialExpressionStaticSwitch::UMaterialExpressionStaticSwitch(const class FPostConstructInitializeProperties& PCIP)
-	: Super(PCIP)
+UMaterialExpressionStaticSwitch::UMaterialExpressionStaticSwitch(const FObjectInitializer& ObjectInitializer)
+	: Super(ObjectInitializer)
 {
 	// Structure to hold one-time initialization
 	struct FConstructorStatics
@@ -3858,8 +3858,8 @@ uint32 UMaterialExpressionStaticSwitch::GetInputType(int32 InputIndex)
 //	UMaterialExpressionQualitySwitch
 //
 
-UMaterialExpressionQualitySwitch::UMaterialExpressionQualitySwitch(const class FPostConstructInitializeProperties& PCIP)
-: Super(PCIP)
+UMaterialExpressionQualitySwitch::UMaterialExpressionQualitySwitch(const FObjectInitializer& ObjectInitializer)
+: Super(ObjectInitializer)
 {
 	// Structure to hold one-time initialization
 	struct FConstructorStatics
@@ -3961,8 +3961,8 @@ bool UMaterialExpressionQualitySwitch::IsResultMaterialAttributes(int32 OutputIn
 //	UMaterialExpressionFeatureLevelSwitch
 //
 
-UMaterialExpressionFeatureLevelSwitch::UMaterialExpressionFeatureLevelSwitch(const class FPostConstructInitializeProperties& PCIP)
-: Super(PCIP)
+UMaterialExpressionFeatureLevelSwitch::UMaterialExpressionFeatureLevelSwitch(const FObjectInitializer& ObjectInitializer)
+: Super(ObjectInitializer)
 {
 	// Structure to hold one-time initialization
 	struct FConstructorStatics
@@ -4073,8 +4073,8 @@ void UMaterialExpressionFeatureLevelSwitch::Serialize(FArchive& Ar)
 //
 //	UMaterialExpressionNormalize
 //
-UMaterialExpressionNormalize::UMaterialExpressionNormalize(const class FPostConstructInitializeProperties& PCIP)
-	: Super(PCIP)
+UMaterialExpressionNormalize::UMaterialExpressionNormalize(const FObjectInitializer& ObjectInitializer)
+	: Super(ObjectInitializer)
 {
 	// Structure to hold one-time initialization
 	struct FConstructorStatics
@@ -4104,8 +4104,8 @@ int32 UMaterialExpressionNormalize::Compile(class FMaterialCompiler* Compiler, i
 	return Compiler->Div(V,Compiler->SquareRoot(Compiler->Dot(V,V)));
 }
 
-UMaterialExpressionVertexColor::UMaterialExpressionVertexColor(const class FPostConstructInitializeProperties& PCIP)
-	: Super(PCIP)
+UMaterialExpressionVertexColor::UMaterialExpressionVertexColor(const FObjectInitializer& ObjectInitializer)
+	: Super(ObjectInitializer)
 {
 	// Structure to hold one-time initialization
 	struct FConstructorStatics
@@ -4138,8 +4138,8 @@ void UMaterialExpressionVertexColor::GetCaption(TArray<FString>& OutCaptions) co
 	OutCaptions.Add(TEXT("Vertex Color"));
 }
 
-UMaterialExpressionParticleColor::UMaterialExpressionParticleColor(const class FPostConstructInitializeProperties& PCIP)
-	: Super(PCIP)
+UMaterialExpressionParticleColor::UMaterialExpressionParticleColor(const FObjectInitializer& ObjectInitializer)
+	: Super(ObjectInitializer)
 {
 	// Structure to hold one-time initialization
 	struct FConstructorStatics
@@ -4177,8 +4177,8 @@ void UMaterialExpressionParticleColor::GetCaption(TArray<FString>& OutCaptions)
 	OutCaptions.Add(TEXT("Particle Color"));
 }
 
-UMaterialExpressionParticlePositionWS::UMaterialExpressionParticlePositionWS(const class FPostConstructInitializeProperties& PCIP)
-	: Super(PCIP)
+UMaterialExpressionParticlePositionWS::UMaterialExpressionParticlePositionWS(const FObjectInitializer& ObjectInitializer)
+	: Super(ObjectInitializer)
 {
 	// Structure to hold one-time initialization
 	struct FConstructorStatics
@@ -4211,8 +4211,8 @@ void UMaterialExpressionParticlePositionWS::GetCaption(TArray<FString>& OutCapti
 	OutCaptions.Add(TEXT("Particle Position"));
 }
 
-UMaterialExpressionParticleRadius::UMaterialExpressionParticleRadius(const class FPostConstructInitializeProperties& PCIP)
-	: Super(PCIP)
+UMaterialExpressionParticleRadius::UMaterialExpressionParticleRadius(const FObjectInitializer& ObjectInitializer)
+	: Super(ObjectInitializer)
 {
 	// Structure to hold one-time initialization
 	struct FConstructorStatics
@@ -4243,8 +4243,8 @@ void UMaterialExpressionParticleRadius::GetCaption(TArray<FString>& OutCaptions)
 	OutCaptions.Add(TEXT("Particle Radius"));
 }
 
-UMaterialExpressionDynamicParameter::UMaterialExpressionDynamicParameter(const class FPostConstructInitializeProperties& PCIP)
-	: Super(PCIP)
+UMaterialExpressionDynamicParameter::UMaterialExpressionDynamicParameter(const FObjectInitializer& ObjectInitializer)
+	: Super(ObjectInitializer)
 {
 	// Structure to hold one-time initialization
 	struct FConstructorStatics
@@ -4371,8 +4371,8 @@ bool UMaterialExpressionDynamicParameter::CopyDynamicParameterNames(const UMater
 //
 //	MaterialExpressionParticleSubUV
 //
-UMaterialExpressionParticleSubUV::UMaterialExpressionParticleSubUV(const class FPostConstructInitializeProperties& PCIP)
-	: Super(PCIP)
+UMaterialExpressionParticleSubUV::UMaterialExpressionParticleSubUV(const FObjectInitializer& ObjectInitializer)
+	: Super(ObjectInitializer)
 {
 	// Structure to hold one-time initialization
 	struct FConstructorStatics
@@ -4421,8 +4421,8 @@ void UMaterialExpressionParticleSubUV::GetCaption(TArray<FString>& OutCaptions)
 //
 //	MaterialExpressionParticleMacroUV
 //
-UMaterialExpressionParticleMacroUV::UMaterialExpressionParticleMacroUV(const class FPostConstructInitializeProperties& PCIP)
-	: Super(PCIP)
+UMaterialExpressionParticleMacroUV::UMaterialExpressionParticleMacroUV(const FObjectInitializer& ObjectInitializer)
+	: Super(ObjectInitializer)
 {
 	// Structure to hold one-time initialization
 	struct FConstructorStatics
@@ -4452,8 +4452,8 @@ void UMaterialExpressionParticleMacroUV::GetCaption(TArray<FString>& OutCaptions
 }
 
 
-UMaterialExpressionLightVector::UMaterialExpressionLightVector(const class FPostConstructInitializeProperties& PCIP)
-	: Super(PCIP)
+UMaterialExpressionLightVector::UMaterialExpressionLightVector(const FObjectInitializer& ObjectInitializer)
+	: Super(ObjectInitializer)
 {
 	// Structure to hold one-time initialization
 	struct FConstructorStatics
@@ -4482,8 +4482,8 @@ void UMaterialExpressionLightVector::GetCaption(TArray<FString>& OutCaptions) co
 	OutCaptions.Add(TEXT("Light Vector"));
 }
 
-UMaterialExpressionScreenPosition::UMaterialExpressionScreenPosition(const class FPostConstructInitializeProperties& PCIP)
-	: Super(PCIP)
+UMaterialExpressionScreenPosition::UMaterialExpressionScreenPosition(const FObjectInitializer& ObjectInitializer)
+	: Super(ObjectInitializer)
 {
 	// Structure to hold one-time initialization
 	struct FConstructorStatics
@@ -4510,8 +4510,8 @@ void UMaterialExpressionScreenPosition::GetCaption(TArray<FString>& OutCaptions)
 	OutCaptions.Add(TEXT("ScreenPosition"));
 }
 
-UMaterialExpressionViewSize::UMaterialExpressionViewSize(const class FPostConstructInitializeProperties& PCIP)
-	: Super(PCIP)
+UMaterialExpressionViewSize::UMaterialExpressionViewSize(const FObjectInitializer& ObjectInitializer)
+	: Super(ObjectInitializer)
 {
 	// Structure to hold one-time initialization
 	struct FConstructorStatics
@@ -4538,8 +4538,8 @@ void UMaterialExpressionViewSize::GetCaption(TArray<FString>& OutCaptions) const
 	OutCaptions.Add(TEXT("ViewSize"));
 }
 
-UMaterialExpressionSceneTexelSize::UMaterialExpressionSceneTexelSize(const class FPostConstructInitializeProperties& PCIP)
-	: Super(PCIP)
+UMaterialExpressionSceneTexelSize::UMaterialExpressionSceneTexelSize(const FObjectInitializer& ObjectInitializer)
+	: Super(ObjectInitializer)
 {
 	// Structure to hold one-time initialization
 	struct FConstructorStatics
@@ -4566,8 +4566,8 @@ void UMaterialExpressionSceneTexelSize::GetCaption(TArray<FString>& OutCaptions)
 	OutCaptions.Add(TEXT("SceneTexelSize"));
 }
 
-UMaterialExpressionSquareRoot::UMaterialExpressionSquareRoot(const class FPostConstructInitializeProperties& PCIP)
-	: Super(PCIP)
+UMaterialExpressionSquareRoot::UMaterialExpressionSquareRoot(const FObjectInitializer& ObjectInitializer)
+	: Super(ObjectInitializer)
 {
 	// Structure to hold one-time initialization
 	struct FConstructorStatics
@@ -4600,8 +4600,8 @@ void UMaterialExpressionSquareRoot::GetCaption(TArray<FString>& OutCaptions) con
 ///////////////////////////////////////////////////////////////////////////////
 // UMaterialExpressionPixelDepth
 ///////////////////////////////////////////////////////////////////////////////
-UMaterialExpressionPixelDepth::UMaterialExpressionPixelDepth(const class FPostConstructInitializeProperties& PCIP)
-	: Super(PCIP)
+UMaterialExpressionPixelDepth::UMaterialExpressionPixelDepth(const FObjectInitializer& ObjectInitializer)
+	: Super(ObjectInitializer)
 {
 	// Structure to hold one-time initialization
 	struct FConstructorStatics
@@ -4637,8 +4637,8 @@ void UMaterialExpressionPixelDepth::GetCaption(TArray<FString>& OutCaptions) con
 ///////////////////////////////////////////////////////////////////////////////
 // UMaterialExpressionSceneDepth
 ///////////////////////////////////////////////////////////////////////////////
-UMaterialExpressionSceneDepth::UMaterialExpressionSceneDepth(const class FPostConstructInitializeProperties& PCIP)
-	: Super(PCIP)
+UMaterialExpressionSceneDepth::UMaterialExpressionSceneDepth(const FObjectInitializer& ObjectInitializer)
+	: Super(ObjectInitializer)
 {
 	// Structure to hold one-time initialization
 	struct FConstructorStatics
@@ -4721,8 +4721,8 @@ FString UMaterialExpressionSceneDepth::GetInputName(int32 InputIndex) const
 ///////////////////////////////////////////////////////////////////////////////
 // UMaterialExpressionSceneTexture
 ///////////////////////////////////////////////////////////////////////////////
-UMaterialExpressionSceneTexture::UMaterialExpressionSceneTexture(const class FPostConstructInitializeProperties& PCIP)
-	: Super(PCIP)
+UMaterialExpressionSceneTexture::UMaterialExpressionSceneTexture(const FObjectInitializer& ObjectInitializer)
+	: Super(ObjectInitializer)
 {
 	// Structure to hold one-time initialization
 	struct FConstructorStatics
@@ -4806,8 +4806,8 @@ void UMaterialExpressionSceneTexture::GetCaption(TArray<FString>& OutCaptions) c
 ///////////////////////////////////////////////////////////////////////////////
 // UMaterialExpressionSceneColor
 ///////////////////////////////////////////////////////////////////////////////
-UMaterialExpressionSceneColor::UMaterialExpressionSceneColor(const class FPostConstructInitializeProperties& PCIP)
-	: Super(PCIP)
+UMaterialExpressionSceneColor::UMaterialExpressionSceneColor(const FObjectInitializer& ObjectInitializer)
+	: Super(ObjectInitializer)
 {
 	// Structure to hold one-time initialization
 	struct FConstructorStatics
@@ -4874,8 +4874,8 @@ void UMaterialExpressionSceneColor::GetCaption(TArray<FString>& OutCaptions) con
 	OutCaptions.Add(TEXT("Scene Color"));
 }
 
-UMaterialExpressionPower::UMaterialExpressionPower(const class FPostConstructInitializeProperties& PCIP)
-	: Super(PCIP)
+UMaterialExpressionPower::UMaterialExpressionPower(const FObjectInitializer& ObjectInitializer)
+	: Super(ObjectInitializer)
 {
 	// Structure to hold one-time initialization
 	struct FConstructorStatics
@@ -4931,8 +4931,8 @@ FString UMaterialExpressionSceneColor::GetInputName(int32 InputIndex) const
 	return TEXT("");
 }
 
-UMaterialExpressionIf::UMaterialExpressionIf(const class FPostConstructInitializeProperties& PCIP)
-	: Super(PCIP)
+UMaterialExpressionIf::UMaterialExpressionIf(const FObjectInitializer& ObjectInitializer)
+	: Super(ObjectInitializer)
 {
 	// Structure to hold one-time initialization
 	struct FConstructorStatics
@@ -5009,8 +5009,8 @@ uint32 UMaterialExpressionIf::GetInputType(int32 InputIndex)
 }
 #endif
 
-UMaterialExpressionOneMinus::UMaterialExpressionOneMinus(const class FPostConstructInitializeProperties& PCIP)
-	: Super(PCIP)
+UMaterialExpressionOneMinus::UMaterialExpressionOneMinus(const FObjectInitializer& ObjectInitializer)
+	: Super(ObjectInitializer)
 {
 	// Structure to hold one-time initialization
 	struct FConstructorStatics
@@ -5040,8 +5040,8 @@ void UMaterialExpressionOneMinus::GetCaption(TArray<FString>& OutCaptions) const
 	OutCaptions.Add(TEXT("1-x"));
 }
 
-UMaterialExpressionAbs::UMaterialExpressionAbs(const class FPostConstructInitializeProperties& PCIP)
-	: Super(PCIP)
+UMaterialExpressionAbs::UMaterialExpressionAbs(const FObjectInitializer& ObjectInitializer)
+	: Super(ObjectInitializer)
 {
 	// Structure to hold one-time initialization
 	struct FConstructorStatics
@@ -5147,8 +5147,8 @@ void UMaterialExpressionTransform::GetCaption(TArray<FString>& OutCaptions) cons
 // UMaterialExpressionTransform
 ///////////////////////////////////////////////////////////////////////////////
 
-UMaterialExpressionTransform::UMaterialExpressionTransform(const class FPostConstructInitializeProperties& PCIP)
-	: Super(PCIP)
+UMaterialExpressionTransform::UMaterialExpressionTransform(const FObjectInitializer& ObjectInitializer)
+	: Super(ObjectInitializer)
 {
 	// Structure to hold one-time initialization
 	struct FConstructorStatics
@@ -5169,8 +5169,8 @@ UMaterialExpressionTransform::UMaterialExpressionTransform(const class FPostCons
 ///////////////////////////////////////////////////////////////////////////////
 // UMaterialExpressionTransformPosition
 ///////////////////////////////////////////////////////////////////////////////
-UMaterialExpressionTransformPosition::UMaterialExpressionTransformPosition(const class FPostConstructInitializeProperties& PCIP)
-	: Super(PCIP)
+UMaterialExpressionTransformPosition::UMaterialExpressionTransformPosition(const FObjectInitializer& ObjectInitializer)
+	: Super(ObjectInitializer)
 {
 	// Structure to hold one-time initialization
 	struct FConstructorStatics
@@ -5232,8 +5232,8 @@ void UMaterialExpressionTransformPosition::GetCaption(TArray<FString>& OutCaptio
 ///////////////////////////////////////////////////////////////////////////////
 // UMaterialExpressionComment
 ///////////////////////////////////////////////////////////////////////////////
-UMaterialExpressionComment::UMaterialExpressionComment(const class FPostConstructInitializeProperties& PCIP)
-	: Super(PCIP)
+UMaterialExpressionComment::UMaterialExpressionComment(const FObjectInitializer& ObjectInitializer)
+	: Super(ObjectInitializer)
 	, CommentColor(FLinearColor::White)
 {
 	// Structure to hold one-time initialization
@@ -5308,8 +5308,8 @@ bool UMaterialExpressionComment::MatchesSearchQuery( const TCHAR* SearchQuery )
 	return Super::MatchesSearchQuery(SearchQuery);
 }
 
-UMaterialExpressionFresnel::UMaterialExpressionFresnel(const class FPostConstructInitializeProperties& PCIP)
-	: Super(PCIP)
+UMaterialExpressionFresnel::UMaterialExpressionFresnel(const FObjectInitializer& ObjectInitializer)
+	: Super(ObjectInitializer)
 {
 	// Structure to hold one-time initialization
 	struct FConstructorStatics
@@ -5348,8 +5348,8 @@ int32 UMaterialExpressionFresnel::Compile(class FMaterialCompiler* Compiler, int
 /*-----------------------------------------------------------------------------
 UMaterialExpressionFontSample
 -----------------------------------------------------------------------------*/
-UMaterialExpressionFontSample::UMaterialExpressionFontSample(const class FPostConstructInitializeProperties& PCIP)
-	: Super(PCIP)
+UMaterialExpressionFontSample::UMaterialExpressionFontSample(const FObjectInitializer& ObjectInitializer)
+	: Super(ObjectInitializer)
 {
 	// Structure to hold one-time initialization
 	struct FConstructorStatics
@@ -5475,8 +5475,8 @@ UTexture* UMaterialExpressionFontSample::GetReferencedTexture()
 /*-----------------------------------------------------------------------------
 UMaterialExpressionFontSampleParameter
 -----------------------------------------------------------------------------*/
-UMaterialExpressionFontSampleParameter::UMaterialExpressionFontSampleParameter(const class FPostConstructInitializeProperties& PCIP)
-	: Super(PCIP)
+UMaterialExpressionFontSampleParameter::UMaterialExpressionFontSampleParameter(const FObjectInitializer& ObjectInitializer)
+	: Super(ObjectInitializer)
 {
 	// Structure to hold one-time initialization
 	struct FConstructorStatics
@@ -5590,8 +5590,8 @@ void UMaterialExpressionFontSampleParameter::GetAllParameterNames(TArray<FName>
 ///////////////////////////////////////////////////////////////////////////////
 // UMaterialExpressionWorldPosition
 ///////////////////////////////////////////////////////////////////////////////
-UMaterialExpressionWorldPosition::UMaterialExpressionWorldPosition(const class FPostConstructInitializeProperties& PCIP)
-	: Super(PCIP)
+UMaterialExpressionWorldPosition::UMaterialExpressionWorldPosition(const FObjectInitializer& ObjectInitializer)
+	: Super(ObjectInitializer)
 {
 	// Structure to hold one-time initialization
 	struct FConstructorStatics
@@ -5652,8 +5652,8 @@ void UMaterialExpressionWorldPosition::GetCaption(TArray<FString>& OutCaptions)
 ///////////////////////////////////////////////////////////////////////////////
 // UMaterialExpressionObjectPositionWS
 ///////////////////////////////////////////////////////////////////////////////
-UMaterialExpressionObjectPositionWS::UMaterialExpressionObjectPositionWS(const class FPostConstructInitializeProperties& PCIP)
-	: Super(PCIP)
+UMaterialExpressionObjectPositionWS::UMaterialExpressionObjectPositionWS(const FObjectInitializer& ObjectInitializer)
+	: Super(ObjectInitializer)
 {
 	// Structure to hold one-time initialization
 	struct FConstructorStatics
@@ -5686,8 +5686,8 @@ void UMaterialExpressionObjectPositionWS::GetCaption(TArray<FString>& OutCaption
 ///////////////////////////////////////////////////////////////////////////////
 // UMaterialExpressionObjectRadius
 ///////////////////////////////////////////////////////////////////////////////
-UMaterialExpressionObjectRadius::UMaterialExpressionObjectRadius(const class FPostConstructInitializeProperties& PCIP)
-	: Super(PCIP)
+UMaterialExpressionObjectRadius::UMaterialExpressionObjectRadius(const FObjectInitializer& ObjectInitializer)
+	: Super(ObjectInitializer)
 {
 	// Structure to hold one-time initialization
 	struct FConstructorStatics
@@ -5718,8 +5718,8 @@ void UMaterialExpressionObjectRadius::GetCaption(TArray<FString>& OutCaptions) c
 ///////////////////////////////////////////////////////////////////////////////
 // UMaterialExpressionObjectBoundingBox
 ///////////////////////////////////////////////////////////////////////////////
-UMaterialExpressionObjectBounds::UMaterialExpressionObjectBounds(const class FPostConstructInitializeProperties& PCIP)
-	: Super(PCIP)
+UMaterialExpressionObjectBounds::UMaterialExpressionObjectBounds(const FObjectInitializer& ObjectInitializer)
+	: Super(ObjectInitializer)
 {
 	// Structure to hold one-time initialization
 	struct FConstructorStatics
@@ -5750,8 +5750,8 @@ void UMaterialExpressionObjectBounds::GetCaption(TArray<FString>& OutCaptions) c
 ///////////////////////////////////////////////////////////////////////////////
 // UMaterialExpressionDistanceCullFade
 ///////////////////////////////////////////////////////////////////////////////
-UMaterialExpressionDistanceCullFade::UMaterialExpressionDistanceCullFade(const class FPostConstructInitializeProperties& PCIP)
-	: Super(PCIP)
+UMaterialExpressionDistanceCullFade::UMaterialExpressionDistanceCullFade(const FObjectInitializer& ObjectInitializer)
+	: Super(ObjectInitializer)
 {
 	// Structure to hold one-time initialization
 	struct FConstructorStatics
@@ -5783,8 +5783,8 @@ void UMaterialExpressionDistanceCullFade::GetCaption(TArray<FString>& OutCaption
 ///////////////////////////////////////////////////////////////////////////////
 // UMaterialExpressionActorPositionWS
 ///////////////////////////////////////////////////////////////////////////////
-UMaterialExpressionActorPositionWS::UMaterialExpressionActorPositionWS(const class FPostConstructInitializeProperties& PCIP)
-	: Super(PCIP)
+UMaterialExpressionActorPositionWS::UMaterialExpressionActorPositionWS(const FObjectInitializer& ObjectInitializer)
+	: Super(ObjectInitializer)
 {
 	// Structure to hold one-time initialization
 	struct FConstructorStatics
@@ -5819,8 +5819,8 @@ void UMaterialExpressionActorPositionWS::GetCaption(TArray<FString>& OutCaptions
 ///////////////////////////////////////////////////////////////////////////////
 // UMaterialExpressionDeriveNormalZ
 ///////////////////////////////////////////////////////////////////////////////
-UMaterialExpressionDeriveNormalZ::UMaterialExpressionDeriveNormalZ(const class FPostConstructInitializeProperties& PCIP)
-	: Super(PCIP)
+UMaterialExpressionDeriveNormalZ::UMaterialExpressionDeriveNormalZ(const FObjectInitializer& ObjectInitializer)
+	: Super(ObjectInitializer)
 {
 	// Structure to hold one-time initialization
 	struct FConstructorStatics
@@ -5862,8 +5862,8 @@ void UMaterialExpressionDeriveNormalZ::GetCaption(TArray<FString>& OutCaptions)
 ///////////////////////////////////////////////////////////////////////////////
 // UMaterialExpressionConstantBiasScale
 ///////////////////////////////////////////////////////////////////////////////
-UMaterialExpressionConstantBiasScale::UMaterialExpressionConstantBiasScale(const class FPostConstructInitializeProperties& PCIP)
-	: Super(PCIP)
+UMaterialExpressionConstantBiasScale::UMaterialExpressionConstantBiasScale(const FObjectInitializer& ObjectInitializer)
+	: Super(ObjectInitializer)
 {
 	// Structure to hold one-time initialization
 	struct FConstructorStatics
@@ -5902,8 +5902,8 @@ void UMaterialExpressionConstantBiasScale::GetCaption(TArray<FString>& OutCaptio
 ///////////////////////////////////////////////////////////////////////////////
 // UMaterialExpressionCustom
 ///////////////////////////////////////////////////////////////////////////////
-UMaterialExpressionCustom::UMaterialExpressionCustom(const class FPostConstructInitializeProperties& PCIP)
-	: Super(PCIP)
+UMaterialExpressionCustom::UMaterialExpressionCustom(const FObjectInitializer& ObjectInitializer)
+	: Super(ObjectInitializer)
 {
 	// Structure to hold one-time initialization
 	struct FConstructorStatics
@@ -6041,8 +6041,8 @@ uint32 UMaterialExpressionCustom::GetOutputType(int32 OutputIndex)
 ///////////////////////////////////////////////////////////////////////////////
 // UMaterialFunction
 ///////////////////////////////////////////////////////////////////////////////
-UMaterialFunction::UMaterialFunction(const class FPostConstructInitializeProperties& PCIP)
-	: Super(PCIP)
+UMaterialFunction::UMaterialFunction(const FObjectInitializer& ObjectInitializer)
+	: Super(ObjectInitializer)
 {
 	LibraryCategories.Add(TEXT("Misc"));
 
@@ -6558,8 +6558,8 @@ bool UMaterialFunction::HasFlippedCoordinates() const
 
 UMaterialFunction* SavedMaterialFunction = NULL;
 
-UMaterialExpressionMaterialFunctionCall::UMaterialExpressionMaterialFunctionCall(const class FPostConstructInitializeProperties& PCIP)
-	: Super(PCIP)
+UMaterialExpressionMaterialFunctionCall::UMaterialExpressionMaterialFunctionCall(const FObjectInitializer& ObjectInitializer)
+	: Super(ObjectInitializer)
 {
 	// Structure to hold one-time initialization
 	struct FConstructorStatics
@@ -7092,8 +7092,8 @@ uint32 UMaterialExpressionMaterialFunctionCall::GetInputType(int32 InputIndex)
 ///////////////////////////////////////////////////////////////////////////////
 // UMaterialExpressionFunctionInput
 ///////////////////////////////////////////////////////////////////////////////
-UMaterialExpressionFunctionInput::UMaterialExpressionFunctionInput(const class FPostConstructInitializeProperties& PCIP)
-	: Super(PCIP)
+UMaterialExpressionFunctionInput::UMaterialExpressionFunctionInput(const FObjectInitializer& ObjectInitializer)
+	: Super(ObjectInitializer)
 {
 	// Structure to hold one-time initialization
 	struct FConstructorStatics
@@ -7386,8 +7386,8 @@ uint32 UMaterialExpressionFunctionInput::GetOutputType(int32 OutputIndex)
 ///////////////////////////////////////////////////////////////////////////////
 // UMaterialExpressionFunctionOutput
 ///////////////////////////////////////////////////////////////////////////////
-UMaterialExpressionFunctionOutput::UMaterialExpressionFunctionOutput(const class FPostConstructInitializeProperties& PCIP)
-	: Super(PCIP)
+UMaterialExpressionFunctionOutput::UMaterialExpressionFunctionOutput(const FObjectInitializer& ObjectInitializer)
+	: Super(ObjectInitializer)
 {
 	// Structure to hold one-time initialization
 	struct FConstructorStatics
@@ -7551,8 +7551,8 @@ bool UMaterialExpressionFunctionOutput::IsResultMaterialAttributes(int32 OutputI
 	}
 }
 
-UMaterialExpressionCollectionParameter::UMaterialExpressionCollectionParameter(const class FPostConstructInitializeProperties& PCIP)
-	: Super(PCIP)
+UMaterialExpressionCollectionParameter::UMaterialExpressionCollectionParameter(const FObjectInitializer& ObjectInitializer)
+	: Super(ObjectInitializer)
 {
 	// Structure to hold one-time initialization
 	struct FConstructorStatics
@@ -7677,8 +7677,8 @@ bool UMaterialExpressionCollectionParameter::MatchesSearchQuery(const TCHAR* Sea
 //
 //	UMaterialExpressionLightmapUVs
 //
-UMaterialExpressionLightmapUVs::UMaterialExpressionLightmapUVs(const class FPostConstructInitializeProperties& PCIP)
-	: Super(PCIP)
+UMaterialExpressionLightmapUVs::UMaterialExpressionLightmapUVs(const FObjectInitializer& ObjectInitializer)
+	: Super(ObjectInitializer)
 {
 	// Structure to hold one-time initialization
 	struct FConstructorStatics
@@ -7717,8 +7717,8 @@ void UMaterialExpressionLightmapUVs::GetCaption(TArray<FString>& OutCaptions) co
 //
 //	UMaterialExpressionLightmassReplace
 //
-UMaterialExpressionLightmassReplace::UMaterialExpressionLightmassReplace(const class FPostConstructInitializeProperties& PCIP)
-	: Super(PCIP)
+UMaterialExpressionLightmassReplace::UMaterialExpressionLightmassReplace(const FObjectInitializer& ObjectInitializer)
+	: Super(ObjectInitializer)
 {
 	// Structure to hold one-time initialization
 	struct FConstructorStatics
@@ -7761,8 +7761,8 @@ void UMaterialExpressionLightmassReplace::GetCaption(TArray<FString>& OutCaption
 //
 //	UMaterialExpressionGIReplace
 //
-UMaterialExpressionGIReplace::UMaterialExpressionGIReplace(const class FPostConstructInitializeProperties& PCIP)
-	: Super(PCIP)
+UMaterialExpressionGIReplace::UMaterialExpressionGIReplace(const FObjectInitializer& ObjectInitializer)
+	: Super(ObjectInitializer)
 {
 	// Structure to hold one-time initialization
 	struct FConstructorStatics
@@ -7801,8 +7801,8 @@ void UMaterialExpressionGIReplace::GetCaption(TArray<FString>& OutCaptions) cons
 	OutCaptions.Add(TEXT("GIReplace"));
 }
 
-UMaterialExpressionObjectOrientation::UMaterialExpressionObjectOrientation(const class FPostConstructInitializeProperties& PCIP)
-	: Super(PCIP)
+UMaterialExpressionObjectOrientation::UMaterialExpressionObjectOrientation(const FObjectInitializer& ObjectInitializer)
+	: Super(ObjectInitializer)
 {
 	// Structure to hold one-time initialization
 	struct FConstructorStatics
@@ -7832,8 +7832,8 @@ void UMaterialExpressionObjectOrientation::GetCaption(TArray<FString>& OutCaptio
 	OutCaptions.Add(TEXT("ObjectOrientation"));
 }
 
-UMaterialExpressionRotateAboutAxis::UMaterialExpressionRotateAboutAxis(const class FPostConstructInitializeProperties& PCIP)
-	: Super(PCIP)
+UMaterialExpressionRotateAboutAxis::UMaterialExpressionRotateAboutAxis(const FObjectInitializer& ObjectInitializer)
+	: Super(ObjectInitializer)
 {
 	// Structure to hold one-time initialization
 	struct FConstructorStatics
@@ -7915,8 +7915,8 @@ static int32 CompileHelperSaturate( FMaterialCompiler* Compiler, int32 A )
 ///////////////////////////////////////////////////////////////////////////////
 // UMaterialExpressionSphereMask
 ///////////////////////////////////////////////////////////////////////////////
-UMaterialExpressionSphereMask::UMaterialExpressionSphereMask(const class FPostConstructInitializeProperties& PCIP)
-	: Super(PCIP)
+UMaterialExpressionSphereMask::UMaterialExpressionSphereMask(const FObjectInitializer& ObjectInitializer)
+	: Super(ObjectInitializer)
 {
 	// Structure to hold one-time initialization
 	struct FConstructorStatics
@@ -7997,8 +7997,8 @@ void UMaterialExpressionSphereMask::GetCaption(TArray<FString>& OutCaptions) con
 ///////////////////////////////////////////////////////////////////////////////
 // UMaterialExpressionNoise
 ///////////////////////////////////////////////////////////////////////////////
-UMaterialExpressionNoise::UMaterialExpressionNoise(const class FPostConstructInitializeProperties& PCIP)
-	: Super(PCIP)
+UMaterialExpressionNoise::UMaterialExpressionNoise(const FObjectInitializer& ObjectInitializer)
+	: Super(ObjectInitializer)
 {
 	// Structure to hold one-time initialization
 	struct FConstructorStatics
@@ -8059,8 +8059,8 @@ void UMaterialExpressionNoise::GetCaption(TArray<FString>& OutCaptions) const
 ///////////////////////////////////////////////////////////////////////////////
 // UMaterialExpressionBlackBody
 ///////////////////////////////////////////////////////////////////////////////
-UMaterialExpressionBlackBody::UMaterialExpressionBlackBody(const class FPostConstructInitializeProperties& PCIP)
-	: Super(PCIP)
+UMaterialExpressionBlackBody::UMaterialExpressionBlackBody(const FObjectInitializer& ObjectInitializer)
+	: Super(ObjectInitializer)
 {
 	// Structure to hold one-time initialization
 	struct FConstructorStatics
@@ -8102,8 +8102,8 @@ void UMaterialExpressionBlackBody::GetCaption(TArray<FString>& OutCaptions) cons
 ///////////////////////////////////////////////////////////////////////////////
 // UMaterialExpressionDistance
 ///////////////////////////////////////////////////////////////////////////////
-UMaterialExpressionDistance::UMaterialExpressionDistance(const class FPostConstructInitializeProperties& PCIP)
-	: Super(PCIP)
+UMaterialExpressionDistance::UMaterialExpressionDistance(const FObjectInitializer& ObjectInitializer)
+	: Super(ObjectInitializer)
 {
 	// Structure to hold one-time initialization
 	struct FConstructorStatics
@@ -8143,8 +8143,8 @@ void UMaterialExpressionDistance::GetCaption(TArray<FString>& OutCaptions) const
 	OutCaptions.Add(TEXT("Distance"));
 }
 
-UMaterialExpressionTwoSidedSign::UMaterialExpressionTwoSidedSign(const class FPostConstructInitializeProperties& PCIP)
-	: Super(PCIP)
+UMaterialExpressionTwoSidedSign::UMaterialExpressionTwoSidedSign(const FObjectInitializer& ObjectInitializer)
+	: Super(ObjectInitializer)
 {
 	// Structure to hold one-time initialization
 	struct FConstructorStatics
@@ -8171,8 +8171,8 @@ void UMaterialExpressionTwoSidedSign::GetCaption(TArray<FString>& OutCaptions) c
 	OutCaptions.Add(TEXT("TwoSidedSign"));
 }
 
-UMaterialExpressionVertexNormalWS::UMaterialExpressionVertexNormalWS(const class FPostConstructInitializeProperties& PCIP)
-	: Super(PCIP)
+UMaterialExpressionVertexNormalWS::UMaterialExpressionVertexNormalWS(const FObjectInitializer& ObjectInitializer)
+	: Super(ObjectInitializer)
 {
 	// Structure to hold one-time initialization
 	struct FConstructorStatics
@@ -8202,8 +8202,8 @@ void UMaterialExpressionVertexNormalWS::GetCaption(TArray<FString>& OutCaptions)
 	OutCaptions.Add(TEXT("VertexNormalWS"));
 }
 
-UMaterialExpressionPixelNormalWS::UMaterialExpressionPixelNormalWS(const class FPostConstructInitializeProperties& PCIP)
-	: Super(PCIP)
+UMaterialExpressionPixelNormalWS::UMaterialExpressionPixelNormalWS(const FObjectInitializer& ObjectInitializer)
+	: Super(ObjectInitializer)
 {
 	// Structure to hold one-time initialization
 	struct FConstructorStatics
@@ -8236,8 +8236,8 @@ void UMaterialExpressionPixelNormalWS::GetCaption(TArray<FString>& OutCaptions)
 ///////////////////////////////////////////////////////////////////////////////
 // UMaterialExpressionPerInstanceRandom
 ///////////////////////////////////////////////////////////////////////////////
-UMaterialExpressionPerInstanceRandom::UMaterialExpressionPerInstanceRandom(const class FPostConstructInitializeProperties& PCIP)
-	: Super(PCIP)
+UMaterialExpressionPerInstanceRandom::UMaterialExpressionPerInstanceRandom(const FObjectInitializer& ObjectInitializer)
+	: Super(ObjectInitializer)
 {
 	// Structure to hold one-time initialization
 	struct FConstructorStatics
@@ -8267,8 +8267,8 @@ void UMaterialExpressionPerInstanceRandom::GetCaption(TArray<FString>& OutCaptio
 ///////////////////////////////////////////////////////////////////////////////
 // UMaterialExpressionPerInstanceFadeAmount
 ///////////////////////////////////////////////////////////////////////////////
-UMaterialExpressionPerInstanceFadeAmount::UMaterialExpressionPerInstanceFadeAmount(const class FPostConstructInitializeProperties& PCIP)
-	: Super(PCIP)
+UMaterialExpressionPerInstanceFadeAmount::UMaterialExpressionPerInstanceFadeAmount(const FObjectInitializer& ObjectInitializer)
+	: Super(ObjectInitializer)
 {
 	// Structure to hold one-time initialization
 	struct FConstructorStatics
@@ -8299,8 +8299,8 @@ void UMaterialExpressionPerInstanceFadeAmount::GetCaption(TArray<FString>& OutCa
 // UMaterialExpressionAntialiasedTextureMask
 ///////////////////////////////////////////////////////////////////////////////
 
-UMaterialExpressionAntialiasedTextureMask::UMaterialExpressionAntialiasedTextureMask(const class FPostConstructInitializeProperties& PCIP)
-	: Super(PCIP)
+UMaterialExpressionAntialiasedTextureMask::UMaterialExpressionAntialiasedTextureMask(const FObjectInitializer& ObjectInitializer)
+	: Super(ObjectInitializer)
 {
 	// Structure to hold one-time initialization
 	struct FConstructorStatics
@@ -8399,8 +8399,8 @@ void UMaterialExpressionAntialiasedTextureMask::SetDefaultTexture()
 ///////////////////////////////////////////////////////////////////////////////
 // UMaterialExpressionLandscapeLayerSample
 ///////////////////////////////////////////////////////////////////////////////
-UMaterialExpressionLandscapeLayerSample::UMaterialExpressionLandscapeLayerSample(const class FPostConstructInitializeProperties& PCIP)
-	: Super(PCIP)
+UMaterialExpressionLandscapeLayerSample::UMaterialExpressionLandscapeLayerSample(const FObjectInitializer& ObjectInitializer)
+	: Super(ObjectInitializer)
 {
 	// Structure to hold one-time initialization
 	struct FConstructorStatics
@@ -8446,8 +8446,8 @@ void UMaterialExpressionLandscapeLayerSample::GetAllParameterNames(TArray<FName>
 ///////////////////////////////////////////////////////////////////////////////
 // UMaterialExpressionLandscapeLayerWeight
 ///////////////////////////////////////////////////////////////////////////////
-UMaterialExpressionLandscapeLayerWeight::UMaterialExpressionLandscapeLayerWeight(const class FPostConstructInitializeProperties& PCIP)
-	: Super(PCIP)
+UMaterialExpressionLandscapeLayerWeight::UMaterialExpressionLandscapeLayerWeight(const FObjectInitializer& ObjectInitializer)
+	: Super(ObjectInitializer)
 {
 	// Structure to hold one-time initialization
 	struct FConstructorStatics
@@ -8528,8 +8528,8 @@ void UMaterialExpressionLandscapeLayerWeight::GetAllParameterNames(TArray<FName>
 ///////////////////////////////////////////////////////////////////////////////
 // UMaterialExpressionLandscapeLayerSwitch
 ///////////////////////////////////////////////////////////////////////////////
-UMaterialExpressionLandscapeLayerSwitch::UMaterialExpressionLandscapeLayerSwitch(const class FPostConstructInitializeProperties& PCIP)
-	: Super(PCIP)
+UMaterialExpressionLandscapeLayerSwitch::UMaterialExpressionLandscapeLayerSwitch(const FObjectInitializer& ObjectInitializer)
+	: Super(ObjectInitializer)
 {
 	// Structure to hold one-time initialization
 	struct FConstructorStatics
@@ -8622,8 +8622,8 @@ void UMaterialExpressionLandscapeLayerSwitch::GetAllParameterNames(TArray<FName>
 ///////////////////////////////////////////////////////////////////////////////
 // UMaterialExpressionLandscapeLayerBlend
 ///////////////////////////////////////////////////////////////////////////////
-UMaterialExpressionLandscapeLayerBlend::UMaterialExpressionLandscapeLayerBlend(const class FPostConstructInitializeProperties& PCIP)
-	: Super(PCIP)
+UMaterialExpressionLandscapeLayerBlend::UMaterialExpressionLandscapeLayerBlend(const FObjectInitializer& ObjectInitializer)
+	: Super(ObjectInitializer)
 {
 	// Structure to hold one-time initialization
 	struct FConstructorStatics
@@ -8919,8 +8919,8 @@ void UMaterialExpressionLandscapeLayerBlend::GetAllParameterNames(TArray<FName>
 
 FName UMaterialExpressionLandscapeVisibilityMask::ParameterName = FName("__LANDSCAPE_VISIBILITY__");
 
-UMaterialExpressionLandscapeVisibilityMask::UMaterialExpressionLandscapeVisibilityMask(const class FPostConstructInitializeProperties& PCIP)
-: Super(PCIP)
+UMaterialExpressionLandscapeVisibilityMask::UMaterialExpressionLandscapeVisibilityMask(const FObjectInitializer& ObjectInitializer)
+: Super(ObjectInitializer)
 {
 	// Structure to hold one-time initialization
 	struct FConstructorStatics
@@ -8966,8 +8966,8 @@ void UMaterialExpressionLandscapeVisibilityMask::GetCaption(TArray<FString>& Out
 ///////////////////////////////////////////////////////////////////////////////
 // UMaterialExpressionLandscapeLayerCoords
 ///////////////////////////////////////////////////////////////////////////////
-UMaterialExpressionLandscapeLayerCoords::UMaterialExpressionLandscapeLayerCoords(const class FPostConstructInitializeProperties& PCIP)
-	: Super(PCIP)
+UMaterialExpressionLandscapeLayerCoords::UMaterialExpressionLandscapeLayerCoords(const FObjectInitializer& ObjectInitializer)
+	: Super(ObjectInitializer)
 {
 	// Structure to hold one-time initialization
 	struct FConstructorStatics
@@ -9054,8 +9054,8 @@ void UMaterialExpressionLandscapeLayerCoords::GetCaption(TArray<FString>& OutCap
 	OutCaptions.Add(FString(TEXT("LandscapeCoords")));
 }
 
-UMaterialExpressionDepthFade::UMaterialExpressionDepthFade(const class FPostConstructInitializeProperties& PCIP)
-	: Super(PCIP)
+UMaterialExpressionDepthFade::UMaterialExpressionDepthFade(const FObjectInitializer& ObjectInitializer)
+	: Super(ObjectInitializer)
 {
 	// Structure to hold one-time initialization
 	struct FConstructorStatics
@@ -9090,8 +9090,8 @@ int32 UMaterialExpressionDepthFade::Compile(class FMaterialCompiler* Compiler, i
 //
 //	UMaterialExpressionSphericalParticleOpacity
 //
-UMaterialExpressionSphericalParticleOpacity::UMaterialExpressionSphericalParticleOpacity(const class FPostConstructInitializeProperties& PCIP)
-	: Super(PCIP)
+UMaterialExpressionSphericalParticleOpacity::UMaterialExpressionSphericalParticleOpacity(const FObjectInitializer& ObjectInitializer)
+	: Super(ObjectInitializer)
 {
 	// Structure to hold one-time initialization
 	struct FConstructorStatics
@@ -9118,8 +9118,8 @@ int32 UMaterialExpressionSphericalParticleOpacity::Compile(class FMaterialCompil
 ///////////////////////////////////////////////////////////////////////////////
 // UMaterialExpressionDepthOfFieldFunction
 ///////////////////////////////////////////////////////////////////////////////
-UMaterialExpressionDepthOfFieldFunction::UMaterialExpressionDepthOfFieldFunction(const class FPostConstructInitializeProperties& PCIP)
-	: Super(PCIP)
+UMaterialExpressionDepthOfFieldFunction::UMaterialExpressionDepthOfFieldFunction(const FObjectInitializer& ObjectInitializer)
+	: Super(ObjectInitializer)
 {
 	// Structure to hold one-time initialization
 	struct FConstructorStatics
@@ -9169,8 +9169,8 @@ void UMaterialExpressionDepthOfFieldFunction::GetCaption(TArray<FString>& OutCap
 ///////////////////////////////////////////////////////////////////////////////
 // UMaterialExpressionDDX
 ///////////////////////////////////////////////////////////////////////////////
-UMaterialExpressionDDX::UMaterialExpressionDDX(const class FPostConstructInitializeProperties& PCIP)
-	: Super(PCIP)
+UMaterialExpressionDDX::UMaterialExpressionDDX(const FObjectInitializer& ObjectInitializer)
+	: Super(ObjectInitializer)
 {
 	// Structure to hold one-time initialization
 	struct FConstructorStatics
@@ -9214,8 +9214,8 @@ void UMaterialExpressionDDX::GetCaption(TArray<FString>& OutCaptions) const
 ///////////////////////////////////////////////////////////////////////////////
 // UMaterialExpressionDDY
 ///////////////////////////////////////////////////////////////////////////////
-UMaterialExpressionDDY::UMaterialExpressionDDY(const class FPostConstructInitializeProperties& PCIP)
-	: Super(PCIP)
+UMaterialExpressionDDY::UMaterialExpressionDDY(const FObjectInitializer& ObjectInitializer)
+	: Super(ObjectInitializer)
 {
 	// Structure to hold one-time initialization
 	struct FConstructorStatics
@@ -9259,8 +9259,8 @@ void UMaterialExpressionDDY::GetCaption(TArray<FString>& OutCaptions) const
 /*------------------------------------------------------------------------------
 	Particle relative time material expression.
 ------------------------------------------------------------------------------*/
-UMaterialExpressionParticleRelativeTime::UMaterialExpressionParticleRelativeTime(const class FPostConstructInitializeProperties& PCIP)
-	: Super(PCIP)
+UMaterialExpressionParticleRelativeTime::UMaterialExpressionParticleRelativeTime(const FObjectInitializer& ObjectInitializer)
+	: Super(ObjectInitializer)
 {
 	// Structure to hold one-time initialization
 	struct FConstructorStatics
@@ -9293,8 +9293,8 @@ void UMaterialExpressionParticleRelativeTime::GetCaption(TArray<FString>& OutCap
 /*------------------------------------------------------------------------------
 	Particle motion blur fade material expression.
 ------------------------------------------------------------------------------*/
-UMaterialExpressionParticleMotionBlurFade::UMaterialExpressionParticleMotionBlurFade(const class FPostConstructInitializeProperties& PCIP)
-	: Super(PCIP)
+UMaterialExpressionParticleMotionBlurFade::UMaterialExpressionParticleMotionBlurFade(const FObjectInitializer& ObjectInitializer)
+	: Super(ObjectInitializer)
 {
 	// Structure to hold one-time initialization
 	struct FConstructorStatics
@@ -9327,8 +9327,8 @@ void UMaterialExpressionParticleMotionBlurFade::GetCaption(TArray<FString>& OutC
 /*------------------------------------------------------------------------------
 	Particle direction material expression.
 ------------------------------------------------------------------------------*/
-UMaterialExpressionParticleDirection::UMaterialExpressionParticleDirection(const class FPostConstructInitializeProperties& PCIP)
-	: Super(PCIP)
+UMaterialExpressionParticleDirection::UMaterialExpressionParticleDirection(const FObjectInitializer& ObjectInitializer)
+	: Super(ObjectInitializer)
 {
 	// Structure to hold one-time initialization
 	struct FConstructorStatics
@@ -9361,8 +9361,8 @@ void UMaterialExpressionParticleDirection::GetCaption(TArray<FString>& OutCaptio
 /*------------------------------------------------------------------------------
 	Particle speed material expression.
 ------------------------------------------------------------------------------*/
-UMaterialExpressionParticleSpeed::UMaterialExpressionParticleSpeed(const class FPostConstructInitializeProperties& PCIP)
-	: Super(PCIP)
+UMaterialExpressionParticleSpeed::UMaterialExpressionParticleSpeed(const FObjectInitializer& ObjectInitializer)
+	: Super(ObjectInitializer)
 {
 	// Structure to hold one-time initialization
 	struct FConstructorStatics
@@ -9395,8 +9395,8 @@ void UMaterialExpressionParticleSpeed::GetCaption(TArray<FString>& OutCaptions)
 /*------------------------------------------------------------------------------
 	Particle size material expression.
 ------------------------------------------------------------------------------*/
-UMaterialExpressionParticleSize::UMaterialExpressionParticleSize(const class FPostConstructInitializeProperties& PCIP)
-	: Super(PCIP)
+UMaterialExpressionParticleSize::UMaterialExpressionParticleSize(const FObjectInitializer& ObjectInitializer)
+	: Super(ObjectInitializer)
 {
 	// Structure to hold one-time initialization
 	struct FConstructorStatics
@@ -9429,8 +9429,8 @@ void UMaterialExpressionParticleSize::GetCaption(TArray<FString>& OutCaptions) c
 /*------------------------------------------------------------------------------
 	Atmospheric fog material expression.
 ------------------------------------------------------------------------------*/
-UMaterialExpressionAtmosphericFogColor::UMaterialExpressionAtmosphericFogColor(const class FPostConstructInitializeProperties& PCIP)
-	: Super(PCIP)
+UMaterialExpressionAtmosphericFogColor::UMaterialExpressionAtmosphericFogColor(const FObjectInitializer& ObjectInitializer)
+	: Super(ObjectInitializer)
 {
 	// Structure to hold one-time initialization
 	struct FConstructorStatics
@@ -9467,8 +9467,8 @@ void UMaterialExpressionAtmosphericFogColor::GetCaption(TArray<FString>& OutCapt
 /*------------------------------------------------------------------------------
 	SpeedTree material expression.
 ------------------------------------------------------------------------------*/
-UMaterialExpressionSpeedTree::UMaterialExpressionSpeedTree(const class FPostConstructInitializeProperties& PCIP)
-	: Super(PCIP)
+UMaterialExpressionSpeedTree::UMaterialExpressionSpeedTree(const FObjectInitializer& ObjectInitializer)
+	: Super(ObjectInitializer)
 {
 	// Structure to hold one-time initialization
 	struct FConstructorStatics
@@ -9553,8 +9553,8 @@ bool UMaterialExpressionSpeedTree::CanEditChange(const UProperty* InProperty) co
 ///////////////////////////////////////////////////////////////////////////////
 // UMaterialExpressionEyeAdaptation
 ///////////////////////////////////////////////////////////////////////////////
-UMaterialExpressionEyeAdaptation::UMaterialExpressionEyeAdaptation(const class FPostConstructInitializeProperties& PCIP)
-	: Super(PCIP)
+UMaterialExpressionEyeAdaptation::UMaterialExpressionEyeAdaptation(const FObjectInitializer& ObjectInitializer)
+	: Super(ObjectInitializer)
 {
 	// Structure to hold one-time initialization
 	struct FConstructorStatics
diff --git a/Engine/Source/Runtime/Engine/Private/Materials/MaterialInstance.cpp b/Engine/Source/Runtime/Engine/Private/Materials/MaterialInstance.cpp
--- a/Engine/Source/Runtime/Engine/Private/Materials/MaterialInstance.cpp
+++ b/Engine/Source/Runtime/Engine/Private/Materials/MaterialInstance.cpp
@@ -421,8 +421,8 @@ bool UMaterialInstance::UpdateParameters()
 	return bDirty;
 }
 
-UMaterialInstance::UMaterialInstance(const class FPostConstructInitializeProperties& PCIP)
-	: Super(PCIP)
+UMaterialInstance::UMaterialInstance(const FObjectInitializer& ObjectInitializer)
+	: Super(ObjectInitializer)
 {
 	bHasStaticPermutationResource = false;
 	bOverrideBaseProperties = false;
diff --git a/Engine/Source/Runtime/Engine/Private/Materials/MaterialInstanceConstant.cpp b/Engine/Source/Runtime/Engine/Private/Materials/MaterialInstanceConstant.cpp
--- a/Engine/Source/Runtime/Engine/Private/Materials/MaterialInstanceConstant.cpp
+++ b/Engine/Source/Runtime/Engine/Private/Materials/MaterialInstanceConstant.cpp
@@ -8,8 +8,8 @@
 #include "MaterialInstanceSupport.h"
 
 
-UMaterialInstanceConstant::UMaterialInstanceConstant(const class FPostConstructInitializeProperties& PCIP)
-	: Super(PCIP)
+UMaterialInstanceConstant::UMaterialInstanceConstant(const FObjectInitializer& ObjectInitializer)
+	: Super(ObjectInitializer)
 {
 }
 
diff --git a/Engine/Source/Runtime/Engine/Private/Materials/MaterialInstanceDynamic.cpp b/Engine/Source/Runtime/Engine/Private/Materials/MaterialInstanceDynamic.cpp
--- a/Engine/Source/Runtime/Engine/Private/Materials/MaterialInstanceDynamic.cpp
+++ b/Engine/Source/Runtime/Engine/Private/Materials/MaterialInstanceDynamic.cpp
@@ -7,8 +7,8 @@
 #include "EnginePrivate.h"
 #include "MaterialInstanceSupport.h"
 
-UMaterialInstanceDynamic::UMaterialInstanceDynamic(const class FPostConstructInitializeProperties& PCIP)
-	: Super(PCIP)
+UMaterialInstanceDynamic::UMaterialInstanceDynamic(const FObjectInitializer& ObjectInitializer)
+	: Super(ObjectInitializer)
 {
 }
 
diff --git a/Engine/Source/Runtime/Engine/Private/Materials/MaterialInterface.cpp b/Engine/Source/Runtime/Engine/Private/Materials/MaterialInterface.cpp
--- a/Engine/Source/Runtime/Engine/Private/Materials/MaterialInterface.cpp
+++ b/Engine/Source/Runtime/Engine/Private/Materials/MaterialInterface.cpp
@@ -22,8 +22,8 @@ void FMaterialRelevance::SetPrimitiveViewRelevance(FPrimitiveViewRelevance& OutV
 
 //////////////////////////////////////////////////////////////////////////
 
-UMaterialInterface::UMaterialInterface(const class FPostConstructInitializeProperties& PCIP)
-	: Super(PCIP)
+UMaterialInterface::UMaterialInterface(const FObjectInitializer& ObjectInitializer)
+	: Super(ObjectInitializer)
 {
 	if (!HasAnyFlags(RF_ClassDefaultObject))
 	{
diff --git a/Engine/Source/Runtime/Engine/Private/Materials/ParameterCollection.cpp b/Engine/Source/Runtime/Engine/Private/Materials/ParameterCollection.cpp
--- a/Engine/Source/Runtime/Engine/Private/Materials/ParameterCollection.cpp
+++ b/Engine/Source/Runtime/Engine/Private/Materials/ParameterCollection.cpp
@@ -3,8 +3,8 @@
 #include "EnginePrivate.h"
 #include "ParameterCollection.h"
 
-UMaterialParameterCollection::UMaterialParameterCollection(const class FPostConstructInitializeProperties& PCIP)
-	: Super(PCIP)
+UMaterialParameterCollection::UMaterialParameterCollection(const FObjectInitializer& ObjectInitializer)
+	: Super(ObjectInitializer)
 {}
 
 void UMaterialParameterCollection::PostLoad()
@@ -340,8 +340,8 @@ void UMaterialParameterCollection::CreateBufferStruct()
 		);
 }
 
-UMaterialParameterCollectionInstance::UMaterialParameterCollectionInstance(const class FPostConstructInitializeProperties& PCIP)
-	: Super(PCIP)
+UMaterialParameterCollectionInstance::UMaterialParameterCollectionInstance(const FObjectInitializer& ObjectInitializer)
+	: Super(ObjectInitializer)
 {
 	Resource = NULL;
 }
diff --git a/Engine/Source/Runtime/Engine/Private/MatineeAnimInterface.cpp b/Engine/Source/Runtime/Engine/Private/MatineeAnimInterface.cpp
--- a/Engine/Source/Runtime/Engine/Private/MatineeAnimInterface.cpp
+++ b/Engine/Source/Runtime/Engine/Private/MatineeAnimInterface.cpp
@@ -3,7 +3,7 @@
 #include "EnginePrivate.h"
 #include "Matinee/MatineeAnimInterface.h"
 
-UMatineeAnimInterface::UMatineeAnimInterface(const class FPostConstructInitializeProperties& PCIP)
-	: Super(PCIP)
+UMatineeAnimInterface::UMatineeAnimInterface(const FObjectInitializer& ObjectInitializer)
+	: Super(ObjectInitializer)
 {
 }
diff --git a/Engine/Source/Runtime/Engine/Private/MatineeInterface.cpp b/Engine/Source/Runtime/Engine/Private/MatineeInterface.cpp
--- a/Engine/Source/Runtime/Engine/Private/MatineeInterface.cpp
+++ b/Engine/Source/Runtime/Engine/Private/MatineeInterface.cpp
@@ -3,7 +3,7 @@
 #include "EnginePrivate.h"
 #include "Matinee/MatineeInterface.h"
 
-UMatineeInterface::UMatineeInterface(const class FPostConstructInitializeProperties& PCIP)
-	: Super(PCIP)
+UMatineeInterface::UMatineeInterface(const FObjectInitializer& ObjectInitializer)
+	: Super(ObjectInitializer)
 {
 }
diff --git a/Engine/Source/Runtime/Engine/Private/MeshComponent.cpp b/Engine/Source/Runtime/Engine/Private/MeshComponent.cpp
--- a/Engine/Source/Runtime/Engine/Private/MeshComponent.cpp
+++ b/Engine/Source/Runtime/Engine/Private/MeshComponent.cpp
@@ -3,8 +3,8 @@
 #include "EnginePrivate.h"
 
 
-UMeshComponent::UMeshComponent(const class FPostConstructInitializeProperties& PCIP)
-	: Super(PCIP)
+UMeshComponent::UMeshComponent(const FObjectInitializer& ObjectInitializer)
+	: Super(ObjectInitializer)
 {
 	CastShadow = true;
 	bUseAsOccluder = true;
diff --git a/Engine/Source/Runtime/Engine/Private/Model.cpp b/Engine/Source/Runtime/Engine/Private/Model.cpp
--- a/Engine/Source/Runtime/Engine/Private/Model.cpp
+++ b/Engine/Source/Runtime/Engine/Private/Model.cpp
@@ -490,7 +490,7 @@ void UModel::EmptyModel( int32 EmptySurfInfo, int32 EmptyPolys )
 #if WITH_EDITOR
 	if( EmptyPolys )
 	{
-		Polys = new( GetOuter(), NAME_None, RF_Transactional )UPolys(FPostConstructInitializeProperties());
+		Polys = new( GetOuter(), NAME_None, RF_Transactional )UPolys(FObjectInitializer());
 	}
 #endif // WITH_EDITOR
 
@@ -501,8 +501,8 @@ void UModel::EmptyModel( int32 EmptySurfInfo, int32 EmptyPolys )
 //
 // Create a new model and allocate all objects needed for it.
 //
-UModel::UModel( const class FPostConstructInitializeProperties& PCIP,ABrush* Owner, bool InRootOutside )
-:	UObject(PCIP)
+UModel::UModel( const class FObjectInitializer& ObjectInitializer,ABrush* Owner, bool InRootOutside )
+:	UObject(ObjectInitializer)
 ,	Nodes		( this )
 ,	Verts		( this )
 ,	Vectors		( this )
diff --git a/Engine/Source/Runtime/Engine/Private/ModelComponent.cpp b/Engine/Source/Runtime/Engine/Private/ModelComponent.cpp
--- a/Engine/Source/Runtime/Engine/Private/ModelComponent.cpp
+++ b/Engine/Source/Runtime/Engine/Private/ModelComponent.cpp
@@ -51,8 +51,8 @@ FArchive& operator<<(FArchive& Ar,FModelElement& Element)
 	return Ar;
 }
 
-UModelComponent::UModelComponent(const class FPostConstructInitializeProperties& PCIP)
-	: Super(PCIP)
+UModelComponent::UModelComponent(const FObjectInitializer& ObjectInitializer)
+	: Super(ObjectInitializer)
 {
 	CastShadow = true;
 	bUseAsOccluder = true;
@@ -63,8 +63,8 @@ UModelComponent::UModelComponent(const class FPostConstructInitializeProperties&
 }
 
 #if WITH_EDITOR
-UModelComponent::UModelComponent(const class FPostConstructInitializeProperties& PCIP,UModel* InModel,uint16 InComponentIndex,uint32 MaskedSurfaceFlags,const TArray<uint16>& InNodes):
-	UPrimitiveComponent(PCIP),
+UModelComponent::UModelComponent(const FObjectInitializer& ObjectInitializer,UModel* InModel,uint16 InComponentIndex,uint32 MaskedSurfaceFlags,const TArray<uint16>& InNodes):
+	UPrimitiveComponent(ObjectInitializer),
 	Model(InModel),
 	ComponentIndex(InComponentIndex),
 	Nodes(InNodes)
diff --git a/Engine/Source/Runtime/Engine/Private/MorphMesh.cpp b/Engine/Source/Runtime/Engine/Private/MorphMesh.cpp
--- a/Engine/Source/Runtime/Engine/Private/MorphMesh.cpp
+++ b/Engine/Source/Runtime/Engine/Private/MorphMesh.cpp
@@ -8,15 +8,15 @@
 #include "Animation/VertexAnim/VertexAnimBase.h"
 #include "Animation/VertexAnim/MorphTarget.h"
 
-UVertexAnimBase::UVertexAnimBase(const class FPostConstructInitializeProperties& PCIP)
-	: Super(PCIP)
+UVertexAnimBase::UVertexAnimBase(const FObjectInitializer& ObjectInitializer)
+	: Super(ObjectInitializer)
 {
 }
 
 //////////////////////////////////////////////////////////////////////////
 
-UMorphTarget::UMorphTarget(const class FPostConstructInitializeProperties& PCIP)
-	: Super(PCIP)
+UMorphTarget::UMorphTarget(const FObjectInitializer& ObjectInitializer)
+	: Super(ObjectInitializer)
 {
 }
 
diff --git a/Engine/Source/Runtime/Engine/Private/MovementComponent.cpp b/Engine/Source/Runtime/Engine/Private/MovementComponent.cpp
--- a/Engine/Source/Runtime/Engine/Private/MovementComponent.cpp
+++ b/Engine/Source/Runtime/Engine/Private/MovementComponent.cpp
@@ -13,8 +13,8 @@ DEFINE_LOG_CATEGORY_STATIC(LogMovement, Log, All);
 //----------------------------------------------------------------------//
 // UMovementComponent
 //----------------------------------------------------------------------//
-UMovementComponent::UMovementComponent(const class FPostConstructInitializeProperties& PCIP)
-	: Super(PCIP)
+UMovementComponent::UMovementComponent(const FObjectInitializer& ObjectInitializer)
+	: Super(ObjectInitializer)
 {
 	PrimaryComponentTick.TickGroup = TG_PrePhysics;
 	PrimaryComponentTick.bCanEverTick = true;
diff --git a/Engine/Source/Runtime/Engine/Private/NavMeshBoundsVolume.cpp b/Engine/Source/Runtime/Engine/Private/NavMeshBoundsVolume.cpp
--- a/Engine/Source/Runtime/Engine/Private/NavMeshBoundsVolume.cpp
+++ b/Engine/Source/Runtime/Engine/Private/NavMeshBoundsVolume.cpp
@@ -3,8 +3,8 @@
 #include "EnginePrivate.h"
 #include "AI/Navigation/NavMeshBoundsVolume.h"
 
-ANavMeshBoundsVolume::ANavMeshBoundsVolume(const class FPostConstructInitializeProperties& PCIP)
-	: Super(PCIP)
+ANavMeshBoundsVolume::ANavMeshBoundsVolume(const FObjectInitializer& ObjectInitializer)
+	: Super(ObjectInitializer)
 {
 	BrushComponent->BodyInstance.bEnableCollision_DEPRECATED = false;
 
diff --git a/Engine/Source/Runtime/Engine/Private/NavMeshRenderingComponent.cpp b/Engine/Source/Runtime/Engine/Private/NavMeshRenderingComponent.cpp
--- a/Engine/Source/Runtime/Engine/Private/NavMeshRenderingComponent.cpp
+++ b/Engine/Source/Runtime/Engine/Private/NavMeshRenderingComponent.cpp
@@ -13,8 +13,8 @@
 
 #include "NavMeshRenderingHelpers.h"
 
-UNavMeshRenderingComponent::UNavMeshRenderingComponent(const class FPostConstructInitializeProperties& PCIP)
-	: Super(PCIP)
+UNavMeshRenderingComponent::UNavMeshRenderingComponent(const FObjectInitializer& ObjectInitializer)
+	: Super(ObjectInitializer)
 {
 	SetCollisionProfileName(UCollisionProfile::NoCollision_ProfileName);
 	AlwaysLoadOnClient = false;
diff --git a/Engine/Source/Runtime/Engine/Private/NavMovementComponent.cpp b/Engine/Source/Runtime/Engine/Private/NavMovementComponent.cpp
--- a/Engine/Source/Runtime/Engine/Private/NavMovementComponent.cpp
+++ b/Engine/Source/Runtime/Engine/Private/NavMovementComponent.cpp
@@ -16,8 +16,8 @@ void FNavAgentProperties::UpdateWithCollisionComponent(UShapeComponent* Collisio
 //----------------------------------------------------------------------//
 // UMovementComponent
 //----------------------------------------------------------------------//
-UNavMovementComponent::UNavMovementComponent(const class FPostConstructInitializeProperties& PCIP)
-	: Super(PCIP)
+UNavMovementComponent::UNavMovementComponent(const FObjectInitializer& ObjectInitializer)
+	: Super(ObjectInitializer)
 	, bUpdateNavAgentWithOwnersCollision(true)
 	, bStopMovementAbortPaths(true)
 {
diff --git a/Engine/Source/Runtime/Engine/Private/NavigationObjectBase.cpp b/Engine/Source/Runtime/Engine/Private/NavigationObjectBase.cpp
--- a/Engine/Source/Runtime/Engine/Private/NavigationObjectBase.cpp
+++ b/Engine/Source/Runtime/Engine/Private/NavigationObjectBase.cpp
@@ -5,10 +5,10 @@
 
 DEFINE_LOG_CATEGORY_STATIC(LogNavigationPoint, Log, All);
 
-ANavigationObjectBase::ANavigationObjectBase(const class FPostConstructInitializeProperties& PCIP)
-	: Super(PCIP)
+ANavigationObjectBase::ANavigationObjectBase(const FObjectInitializer& ObjectInitializer)
+	: Super(ObjectInitializer)
 {
-	CapsuleComponent = PCIP.CreateDefaultSubobject<UCapsuleComponent>(this, TEXT("CollisionCapsule"));
+	CapsuleComponent = ObjectInitializer.CreateDefaultSubobject<UCapsuleComponent>(this, TEXT("CollisionCapsule"));
 	CapsuleComponent->ShapeColor = FColor(255, 138, 5, 255);
 	CapsuleComponent->bDrawOnlyIfSelected = true;
 	CapsuleComponent->InitCapsuleSize(50.0f, 50.0f);
@@ -20,9 +20,9 @@ ANavigationObjectBase::ANavigationObjectBase(const class FPostConstructInitializ
 	RootComponent = CapsuleComponent;
 	bCollideWhenPlacing = true;
 
-	GoodSprite = PCIP.CreateEditorOnlyDefaultSubobject<UBillboardComponent>(this, TEXT("Sprite"));
+	GoodSprite = ObjectInitializer.CreateEditorOnlyDefaultSubobject<UBillboardComponent>(this, TEXT("Sprite"));
 
-	BadSprite = PCIP.CreateEditorOnlyDefaultSubobject<UBillboardComponent>(this, TEXT("Sprite2"));
+	BadSprite = ObjectInitializer.CreateEditorOnlyDefaultSubobject<UBillboardComponent>(this, TEXT("Sprite2"));
 
 #if WITH_EDITORONLY_DATA
 	if (!IsRunningCommandlet())
diff --git a/Engine/Source/Runtime/Engine/Private/NetConnection.cpp b/Engine/Source/Runtime/Engine/Private/NetConnection.cpp
--- a/Engine/Source/Runtime/Engine/Private/NetConnection.cpp
+++ b/Engine/Source/Runtime/Engine/Private/NetConnection.cpp
@@ -17,8 +17,8 @@
 
 UNetConnection* UNetConnection::GNetConnectionBeingCleanedUp = NULL;
 
-UNetConnection::UNetConnection(const class FPostConstructInitializeProperties& PCIP)
-:	UPlayer(PCIP)
+UNetConnection::UNetConnection(const FObjectInitializer& ObjectInitializer)
+:	UPlayer(ObjectInitializer)
 ,	Driver				( NULL )
 ,	PackageMap			( NULL )
 ,	Viewer				( NULL )
@@ -115,7 +115,7 @@ void UNetConnection::InitBase(UNetDriver* InDriver,class FSocket* InSocket, cons
 	}
 
 	// Create package map.
-	PackageMap = new( this )UPackageMapClient( FPostConstructInitializeProperties(), this, Driver->GuidCache );
+	PackageMap = new( this )UPackageMapClient( FObjectInitializer(), this, Driver->GuidCache );
 
 	// Create the voice channel
 	CreateChannel(CHTYPE_Voice, true, VOICE_CHANNEL_INDEX);
@@ -162,7 +162,7 @@ void UNetConnection::InitConnection(UNetDriver* InDriver, EConnectionState InSta
 	}
 
 	// Create package map.
-	PackageMap = new( this )UPackageMapClient( FPostConstructInitializeProperties(), this, Driver->GuidCache );
+	PackageMap = new( this )UPackageMapClient( FObjectInitializer(), this, Driver->GuidCache );
 }
 
 void UNetConnection::Serialize( FArchive& Ar )
@@ -298,8 +298,8 @@ void UNetConnection::CleanUp()
 	Driver = NULL;
 }
 
-UChildConnection::UChildConnection(const class FPostConstructInitializeProperties& PCIP)
-	: Super(PCIP)
+UChildConnection::UChildConnection(const FObjectInitializer& ObjectInitializer)
+	: Super(ObjectInitializer)
 {
 }
 
diff --git a/Engine/Source/Runtime/Engine/Private/NetworkDriver.cpp b/Engine/Source/Runtime/Engine/Private/NetworkDriver.cpp
--- a/Engine/Source/Runtime/Engine/Private/NetworkDriver.cpp
+++ b/Engine/Source/Runtime/Engine/Private/NetworkDriver.cpp
@@ -94,8 +94,8 @@ static TAutoConsoleVariable<int32> CVarNetDormancyValidate(
 	UNetDriver implementation.
 -----------------------------------------------------------------------------*/
 
-UNetDriver::UNetDriver(const class FPostConstructInitializeProperties& PCIP)
-:	UObject(PCIP)
+UNetDriver::UNetDriver(const FObjectInitializer& ObjectInitializer)
+:	UObject(ObjectInitializer)
 ,	MaxInternetClientRate(10000)
 , 	MaxClientRate(15000)
 ,	RequireEngineVersionMatch(true)
@@ -1499,7 +1499,7 @@ void UNetDriver::FlushActorDormancy(AActor* Actor)
 UChildConnection* UNetDriver::CreateChild(UNetConnection* Parent)
 {
 	UE_LOG(LogNet, Log, TEXT("Creating child connection with %s parent"), *Parent->GetName());
-	UChildConnection* Child = new UChildConnection(FPostConstructInitializeProperties());
+	UChildConnection* Child = new UChildConnection(FObjectInitializer());
 	Child->Driver = this;
 	Child->URL = FURL();
 	Child->State = Parent->State;
diff --git a/Engine/Source/Runtime/Engine/Private/NetworkPrediction.cpp b/Engine/Source/Runtime/Engine/Private/NetworkPrediction.cpp
--- a/Engine/Source/Runtime/Engine/Private/NetworkPrediction.cpp
+++ b/Engine/Source/Runtime/Engine/Private/NetworkPrediction.cpp
@@ -3,7 +3,7 @@
 #include "EnginePrivate.h"
 
 
-UNetworkPredictionInterface::UNetworkPredictionInterface(const class FPostConstructInitializeProperties& PCIP)
-	: Super(PCIP)
+UNetworkPredictionInterface::UNetworkPredictionInterface(const FObjectInitializer& ObjectInitializer)
+	: Super(ObjectInitializer)
 {
 }
diff --git a/Engine/Source/Runtime/Engine/Private/NetworkSettings.cpp b/Engine/Source/Runtime/Engine/Private/NetworkSettings.cpp
--- a/Engine/Source/Runtime/Engine/Private/NetworkSettings.cpp
+++ b/Engine/Source/Runtime/Engine/Private/NetworkSettings.cpp
@@ -2,8 +2,8 @@
 
 #include "EnginePrivate.h"
 
-UNetworkSettings::UNetworkSettings(const class FPostConstructInitializeProperties& PCIP)
-	: Super(PCIP)
+UNetworkSettings::UNetworkSettings(const FObjectInitializer& ObjectInitializer)
+	: Super(ObjectInitializer)
 {
 }
 
diff --git a/Engine/Source/Runtime/Engine/Private/NiagaraActor.cpp b/Engine/Source/Runtime/Engine/Private/NiagaraActor.cpp
--- a/Engine/Source/Runtime/Engine/Private/NiagaraActor.cpp
+++ b/Engine/Source/Runtime/Engine/Private/NiagaraActor.cpp
@@ -5,16 +5,16 @@
 #include "Engine/NiagaraActor.h"
 #include "Components/NiagaraComponent.h"
 
-ANiagaraActor::ANiagaraActor(const class FPostConstructInitializeProperties& PCIP)
-	: Super(PCIP)
+ANiagaraActor::ANiagaraActor(const FObjectInitializer& ObjectInitializer)
+	: Super(ObjectInitializer)
 {
-	NiagaraComponent = PCIP.CreateDefaultSubobject<UNiagaraComponent>(this, TEXT("NiagaraComponent0"));
+	NiagaraComponent = ObjectInitializer.CreateDefaultSubobject<UNiagaraComponent>(this, TEXT("NiagaraComponent0"));
 
 	RootComponent = NiagaraComponent;
 
 #if WITH_EDITORONLY_DATA
-	SpriteComponent = PCIP.CreateEditorOnlyDefaultSubobject<UBillboardComponent>(this, TEXT("Sprite"));
-	ArrowComponent = PCIP.CreateEditorOnlyDefaultSubobject<UArrowComponent>(this, TEXT("ArrowComponent0"));
+	SpriteComponent = ObjectInitializer.CreateEditorOnlyDefaultSubobject<UBillboardComponent>(this, TEXT("Sprite"));
+	ArrowComponent = ObjectInitializer.CreateEditorOnlyDefaultSubobject<UArrowComponent>(this, TEXT("ArrowComponent0"));
 
 	if (!IsRunningCommandlet())
 	{
diff --git a/Engine/Source/Runtime/Engine/Private/NiagaraComponent.cpp b/Engine/Source/Runtime/Engine/Private/NiagaraComponent.cpp
--- a/Engine/Source/Runtime/Engine/Private/NiagaraComponent.cpp
+++ b/Engine/Source/Runtime/Engine/Private/NiagaraComponent.cpp
@@ -392,8 +392,8 @@ private:
 
 //////////////////////////////////////////////////////////////////////////
 
-UNiagaraComponent::UNiagaraComponent(const class FPostConstructInitializeProperties& PCIP)
-	: Super(PCIP)
+UNiagaraComponent::UNiagaraComponent(const FObjectInitializer& ObjectInitializer)
+	: Super(ObjectInitializer)
 {
 	PrimaryComponentTick.bCanEverTick = true;
 	bTickInEditor = true;
diff --git a/Engine/Source/Runtime/Engine/Private/NiagaraEffect.cpp b/Engine/Source/Runtime/Engine/Private/NiagaraEffect.cpp
--- a/Engine/Source/Runtime/Engine/Private/NiagaraEffect.cpp
+++ b/Engine/Source/Runtime/Engine/Private/NiagaraEffect.cpp
@@ -5,8 +5,8 @@
 
 
 
-UNiagaraEffect::UNiagaraEffect(const class FPostConstructInitializeProperties& PCIP)
-: Super(PCIP)
+UNiagaraEffect::UNiagaraEffect(const FObjectInitializer& ObjectInitializer)
+: Super(ObjectInitializer)
 {
 
 }
diff --git a/Engine/Source/Runtime/Engine/Private/NiagaraScript.cpp b/Engine/Source/Runtime/Engine/Private/NiagaraScript.cpp
--- a/Engine/Source/Runtime/Engine/Private/NiagaraScript.cpp
+++ b/Engine/Source/Runtime/Engine/Private/NiagaraScript.cpp
@@ -6,14 +6,14 @@
 
 
 
-UNiagaraScript::UNiagaraScript(const class FPostConstructInitializeProperties& PCIP)
-	: Super(PCIP)
+UNiagaraScript::UNiagaraScript(const FObjectInitializer& ObjectInitializer)
+	: Super(ObjectInitializer)
 {
 
 }
 
-UNiagaraScriptSourceBase::UNiagaraScriptSourceBase(const class FPostConstructInitializeProperties& PCIP)
-	: Super(PCIP)
+UNiagaraScriptSourceBase::UNiagaraScriptSourceBase(const FObjectInitializer& ObjectInitializer)
+	: Super(ObjectInitializer)
 {
 
 }
diff --git a/Engine/Source/Runtime/Engine/Private/Note.cpp b/Engine/Source/Runtime/Engine/Private/Note.cpp
--- a/Engine/Source/Runtime/Engine/Private/Note.cpp
+++ b/Engine/Source/Runtime/Engine/Private/Note.cpp
@@ -2,16 +2,16 @@
 
 #include "EnginePrivate.h"
 
-ANote::ANote(const class FPostConstructInitializeProperties& PCIP)
-	: Super(PCIP)
+ANote::ANote(const FObjectInitializer& ObjectInitializer)
+	: Super(ObjectInitializer)
 {
-	TSubobjectPtr<USceneComponent> SceneComponent = PCIP.CreateDefaultSubobject<USceneComponent>(this, TEXT("SceneComp"));
+	TSubobjectPtr<USceneComponent> SceneComponent = ObjectInitializer.CreateDefaultSubobject<USceneComponent>(this, TEXT("SceneComp"));
 	RootComponent = SceneComponent;
 	RootComponent->Mobility = EComponentMobility::Static;
 
 #if WITH_EDITORONLY_DATA
-	ArrowComponent = PCIP.CreateEditorOnlyDefaultSubobject<UArrowComponent>(this, TEXT("Arrow"));
-	SpriteComponent = PCIP.CreateEditorOnlyDefaultSubobject<UBillboardComponent>(this, TEXT("Sprite"));
+	ArrowComponent = ObjectInitializer.CreateEditorOnlyDefaultSubobject<UArrowComponent>(this, TEXT("Arrow"));
+	SpriteComponent = ObjectInitializer.CreateEditorOnlyDefaultSubobject<UBillboardComponent>(this, TEXT("Sprite"));
 
 	if (!IsRunningCommandlet())
 	{
diff --git a/Engine/Source/Runtime/Engine/Private/ObjectLibrary.cpp b/Engine/Source/Runtime/Engine/Private/ObjectLibrary.cpp
--- a/Engine/Source/Runtime/Engine/Private/ObjectLibrary.cpp
+++ b/Engine/Source/Runtime/Engine/Private/ObjectLibrary.cpp
@@ -3,8 +3,8 @@
 #include "EnginePrivate.h"
 #include "AssetRegistryModule.h"
 
-UObjectLibrary::UObjectLibrary(const class FPostConstructInitializeProperties& PCIP)
-	: Super(PCIP)
+UObjectLibrary::UObjectLibrary(const FObjectInitializer& ObjectInitializer)
+	: Super(ObjectInitializer)
 {
 	bIsFullyLoaded = false;
 	bUseWeakReferences = false;
diff --git a/Engine/Source/Runtime/Engine/Private/ObjectReferencer.cpp b/Engine/Source/Runtime/Engine/Private/ObjectReferencer.cpp
--- a/Engine/Source/Runtime/Engine/Private/ObjectReferencer.cpp
+++ b/Engine/Source/Runtime/Engine/Private/ObjectReferencer.cpp
@@ -2,7 +2,7 @@
 
 #include "EnginePrivate.h"
 
-UObjectReferencer::UObjectReferencer(const class FPostConstructInitializeProperties& PCIP)
-	: Super(PCIP)
+UObjectReferencer::UObjectReferencer(const FObjectInitializer& ObjectInitializer)
+	: Super(ObjectInitializer)
 {
 }
diff --git a/Engine/Source/Runtime/Engine/Private/OnlineSession.cpp b/Engine/Source/Runtime/Engine/Private/OnlineSession.cpp
--- a/Engine/Source/Runtime/Engine/Private/OnlineSession.cpp
+++ b/Engine/Source/Runtime/Engine/Private/OnlineSession.cpp
@@ -7,8 +7,8 @@
 #include "EnginePrivate.h"
 
 
-UOnlineSession::UOnlineSession(const class FPostConstructInitializeProperties& PCIP)
-	: Super(PCIP)
+UOnlineSession::UOnlineSession(const FObjectInitializer& ObjectInitializer)
+	: Super(ObjectInitializer)
 {
 }
 
diff --git a/Engine/Source/Runtime/Engine/Private/PackageMapClient.cpp b/Engine/Source/Runtime/Engine/Private/PackageMapClient.cpp
--- a/Engine/Source/Runtime/Engine/Private/PackageMapClient.cpp
+++ b/Engine/Source/Runtime/Engine/Private/PackageMapClient.cpp
@@ -28,8 +28,8 @@ static TAutoConsoleVariable<int32> CVarIgnorePackageMismatch( TEXT( "net.IgnoreP
 /*-----------------------------------------------------------------------------
 	UPackageMapClient implementation.
 -----------------------------------------------------------------------------*/
-UPackageMapClient::UPackageMapClient(const class FPostConstructInitializeProperties& PCIP)
-	: Super(PCIP)
+UPackageMapClient::UPackageMapClient(const FObjectInitializer& ObjectInitializer)
+	: Super(ObjectInitializer)
 {
 }
 
diff --git a/Engine/Source/Runtime/Engine/Private/PainCausingVolume.cpp b/Engine/Source/Runtime/Engine/Private/PainCausingVolume.cpp
--- a/Engine/Source/Runtime/Engine/Private/PainCausingVolume.cpp
+++ b/Engine/Source/Runtime/Engine/Private/PainCausingVolume.cpp
@@ -3,8 +3,8 @@
 #include "EnginePrivate.h"
 #include "GameFramework/PainCausingVolume.h"
 
-APainCausingVolume::APainCausingVolume(const class FPostConstructInitializeProperties& PCIP)
-	: Super(PCIP)
+APainCausingVolume::APainCausingVolume(const FObjectInitializer& ObjectInitializer)
+	: Super(ObjectInitializer)
 {
 
 	bPainCausing = true;
diff --git a/Engine/Source/Runtime/Engine/Private/Particles/ParticleBeamModules.cpp b/Engine/Source/Runtime/Engine/Private/Particles/ParticleBeamModules.cpp
--- a/Engine/Source/Runtime/Engine/Private/Particles/ParticleBeamModules.cpp
+++ b/Engine/Source/Runtime/Engine/Private/Particles/ParticleBeamModules.cpp
@@ -23,8 +23,8 @@
 /*-----------------------------------------------------------------------------
 	UParticleModuleTypeDataBeam2 implementation.
 -----------------------------------------------------------------------------*/
-UParticleModuleTypeDataBeam2::UParticleModuleTypeDataBeam2(const class FPostConstructInitializeProperties& PCIP)
-	: Super(PCIP)
+UParticleModuleTypeDataBeam2::UParticleModuleTypeDataBeam2(const FObjectInitializer& ObjectInitializer)
+	: Super(ObjectInitializer)
 {
 	// Structure to hold one-time initialization
 	struct FConstructorStatics
@@ -897,8 +897,8 @@ void UParticleModuleTypeDataBeam2::GetNoiseRange(FVector& NoiseMin, FVector& Noi
 /*-----------------------------------------------------------------------------
 	UParticleModuleBaseModifier implementation.
 -----------------------------------------------------------------------------*/
-UParticleModuleBeamBase::UParticleModuleBeamBase(const class FPostConstructInitializeProperties& PCIP)
-	: Super(PCIP)
+UParticleModuleBeamBase::UParticleModuleBeamBase(const FObjectInitializer& ObjectInitializer)
+	: Super(ObjectInitializer)
 {
 	bSpawnModule = false;
 	bUpdateModule = false;
@@ -907,8 +907,8 @@ UParticleModuleBeamBase::UParticleModuleBeamBase(const class FPostConstructIniti
 /*-----------------------------------------------------------------------------
 	UParticleModuleBeamModifier implementation.
 -----------------------------------------------------------------------------*/
-UParticleModuleBeamModifier::UParticleModuleBeamModifier(const class FPostConstructInitializeProperties& PCIP)
-	: Super(PCIP)
+UParticleModuleBeamModifier::UParticleModuleBeamModifier(const FObjectInitializer& ObjectInitializer)
+	: Super(ObjectInitializer)
 {
 	ModifierType = PEB2MT_Source;
 }
@@ -1176,8 +1176,8 @@ void UParticleModuleBeamModifier::GetDataPointerOffsets(FParticleEmitterInstance
 /*-----------------------------------------------------------------------------
 	UParticleModuleBeamNoise implementation.
 -----------------------------------------------------------------------------*/
-UParticleModuleBeamNoise::UParticleModuleBeamNoise(const class FPostConstructInitializeProperties& PCIP)
-	: Super(PCIP)
+UParticleModuleBeamNoise::UParticleModuleBeamNoise(const FObjectInitializer& ObjectInitializer)
+	: Super(ObjectInitializer)
 {
 	Frequency = 0;
 	NoiseLockRadius = 1.0f;
@@ -1520,8 +1520,8 @@ void UParticleModuleBeamNoise::GetNoiseRange(FVector& NoiseMin, FVector& NoiseMa
 /*-----------------------------------------------------------------------------
 	UParticleModuleBeamSource implementation.
 -----------------------------------------------------------------------------*/
-UParticleModuleBeamSource::UParticleModuleBeamSource(const class FPostConstructInitializeProperties& PCIP)
-	: Super(PCIP)
+UParticleModuleBeamSource::UParticleModuleBeamSource(const FObjectInitializer& ObjectInitializer)
+	: Super(ObjectInitializer)
 {
 	// Structure to hold one-time initialization
 	struct FConstructorStatics
@@ -1952,8 +1952,8 @@ bool UParticleModuleBeamSource::ResolveSourceData(FParticleBeam2EmitterInstance*
 /*-----------------------------------------------------------------------------
 	UParticleModuleBeamTarget implementation.
 -----------------------------------------------------------------------------*/
-UParticleModuleBeamTarget::UParticleModuleBeamTarget(const class FPostConstructInitializeProperties& PCIP)
-	: Super(PCIP)
+UParticleModuleBeamTarget::UParticleModuleBeamTarget(const FObjectInitializer& ObjectInitializer)
+	: Super(ObjectInitializer)
 {
 	// Structure to hold one-time initialization
 	struct FConstructorStatics
