Rename FPostConstructInitializeProperties to something simpler

diff --git a/Engine/Source/Runtime/Engine/Private/Particles/ParticleComponents.cpp b/Engine/Source/Runtime/Engine/Private/Particles/ParticleComponents.cpp
--- a/Engine/Source/Runtime/Engine/Private/Particles/ParticleComponents.cpp
+++ b/Engine/Source/Runtime/Engine/Private/Particles/ParticleComponents.cpp
@@ -171,16 +171,16 @@ void Particle_ModifyVectorDistribution(UDistributionVector* pkDistribution, FVec
 /*-----------------------------------------------------------------------------
 	AEmitter implementation.
 -----------------------------------------------------------------------------*/
-AEmitter::AEmitter(const class FPostConstructInitializeProperties& PCIP)
-	: Super(PCIP)
+AEmitter::AEmitter(const FObjectInitializer& ObjectInitializer)
+	: Super(ObjectInitializer)
 {
-	ParticleSystemComponent = PCIP.CreateDefaultSubobject<UParticleSystemComponent>(this, TEXT("ParticleSystemComponent0"));
+	ParticleSystemComponent = ObjectInitializer.CreateDefaultSubobject<UParticleSystemComponent>(this, TEXT("ParticleSystemComponent0"));
 	ParticleSystemComponent->SecondsBeforeInactive = 1;
 	RootComponent = ParticleSystemComponent;
 
 #if WITH_EDITORONLY_DATA
-	SpriteComponent = PCIP.CreateEditorOnlyDefaultSubobject<UBillboardComponent>(this, TEXT("Sprite"));
-	ArrowComponent = PCIP.CreateEditorOnlyDefaultSubobject<UArrowComponent>(this, TEXT("ArrowComponent0"));
+	SpriteComponent = ObjectInitializer.CreateEditorOnlyDefaultSubobject<UBillboardComponent>(this, TEXT("Sprite"));
+	ArrowComponent = ObjectInitializer.CreateEditorOnlyDefaultSubobject<UArrowComponent>(this, TEXT("ArrowComponent0"));
 
 	if (!IsRunningCommandlet())
 	{
@@ -471,8 +471,8 @@ bool AEmitter::GetReferencedContentObjects( TArray<UObject*>& Objects ) const
 /*-----------------------------------------------------------------------------
 	UParticleLODLevel implementation.
 -----------------------------------------------------------------------------*/
-UParticleLODLevel::UParticleLODLevel(const class FPostConstructInitializeProperties& PCIP)
-	: Super(PCIP)
+UParticleLODLevel::UParticleLODLevel(const FObjectInitializer& ObjectInitializer)
+	: Super(ObjectInitializer)
 {
 	bEnabled = true;
 	ConvertedModules = true;
@@ -948,8 +948,8 @@ bool UParticleLODLevel::IsModuleEditable(UParticleModule* InModule)
 /*-----------------------------------------------------------------------------
 	UParticleEmitter implementation.
 -----------------------------------------------------------------------------*/
-UParticleEmitter::UParticleEmitter(const class FPostConstructInitializeProperties& PCIP)
-	: Super(PCIP),
+UParticleEmitter::UParticleEmitter(const FObjectInitializer& ObjectInitializer)
+	: Super(ObjectInitializer),
 	QualityLevelSpawnRateScale(1.0f)
 {
 	// Structure to hold one-time initialization
@@ -1780,8 +1780,8 @@ void UParticleEmitter::Build()
 /*-----------------------------------------------------------------------------
 	UParticleSpriteEmitter implementation.
 -----------------------------------------------------------------------------*/
-UParticleSpriteEmitter::UParticleSpriteEmitter(const class FPostConstructInitializeProperties& PCIP)
-	: Super(PCIP)
+UParticleSpriteEmitter::UParticleSpriteEmitter(const FObjectInitializer& ObjectInitializer)
+	: Super(ObjectInitializer)
 {
 }
 
@@ -1948,8 +1948,8 @@ void UParticleSpriteEmitter::SetToSensibleDefaults()
 /*-----------------------------------------------------------------------------
 	UParticleSystem implementation.
 -----------------------------------------------------------------------------*/
-UParticleSystem::UParticleSystem(const class FPostConstructInitializeProperties& PCIP)
-	: Super(PCIP)
+UParticleSystem::UParticleSystem(const FObjectInitializer& ObjectInitializer)
+	: Super(ObjectInitializer)
 {
 #if WITH_EDITORONLY_DATA
 	ThumbnailDistance = 200.0;
@@ -2887,8 +2887,8 @@ bool UParticleSystem::HasGPUEmitter() const
 	return false;
 }
 
-UParticleSystemComponent::UParticleSystemComponent(const class FPostConstructInitializeProperties& PCIP)
-	: Super(PCIP)
+UParticleSystemComponent::UParticleSystemComponent(const FObjectInitializer& ObjectInitializer)
+	: Super(ObjectInitializer)
 {
 	PrimaryComponentTick.bCanEverTick = true;
 	PrimaryComponentTick.TickGroup = TG_DuringPhysics;
@@ -5905,8 +5905,8 @@ int32 UParticleSystemComponent::GetNamedMaterialIndex(FName Name) const
 	return INDEX_NONE;
 }
 
-UParticleSystemReplay::UParticleSystemReplay(const class FPostConstructInitializeProperties& PCIP)
-	: Super(PCIP)
+UParticleSystemReplay::UParticleSystemReplay(const FObjectInitializer& ObjectInitializer)
+	: Super(ObjectInitializer)
 {
 }
 
@@ -5995,8 +5995,8 @@ FArchive& operator<<( FArchive& Ar, FParticleEmitterReplayFrame& Obj )
 	return Ar;
 }
 
-AEmitterCameraLensEffectBase::AEmitterCameraLensEffectBase(const class FPostConstructInitializeProperties& PCIP)
-	: Super(PCIP
+AEmitterCameraLensEffectBase::AEmitterCameraLensEffectBase(const FObjectInitializer& ObjectInitializer)
+	: Super(ObjectInitializer
 		.DoNotCreateDefaultSubobject(TEXT("Sprite"))
 		.DoNotCreateDefaultSubobject(TEXT("ArrowComponent0"))
 	)
diff --git a/Engine/Source/Runtime/Engine/Private/Particles/ParticleEventManager.cpp b/Engine/Source/Runtime/Engine/Private/Particles/ParticleEventManager.cpp
--- a/Engine/Source/Runtime/Engine/Private/Particles/ParticleEventManager.cpp
+++ b/Engine/Source/Runtime/Engine/Private/Particles/ParticleEventManager.cpp
@@ -9,8 +9,8 @@
 /*-----------------------------------------------------------------------------
 	AParticleEventManager implementation.
 -----------------------------------------------------------------------------*/
-AParticleEventManager::AParticleEventManager(const class FPostConstructInitializeProperties& PCIP)
-	: Super(PCIP)
+AParticleEventManager::AParticleEventManager(const FObjectInitializer& ObjectInitializer)
+	: Super(ObjectInitializer)
 {
 	bCanBeDamaged = false;
 }
diff --git a/Engine/Source/Runtime/Engine/Private/Particles/ParticleModules.cpp b/Engine/Source/Runtime/Engine/Private/Particles/ParticleModules.cpp
--- a/Engine/Source/Runtime/Engine/Private/Particles/ParticleModules.cpp
+++ b/Engine/Source/Runtime/Engine/Private/Particles/ParticleModules.cpp
@@ -84,8 +84,8 @@
 /*-----------------------------------------------------------------------------
 	UParticleModule implementation.
 -----------------------------------------------------------------------------*/
-UParticleModule::UParticleModule(const class FPostConstructInitializeProperties& PCIP)
-	: Super(PCIP)
+UParticleModule::UParticleModule(const FObjectInitializer& ObjectInitializer)
+	: Super(ObjectInitializer)
 {
 	bSupported3DDrawMode = false;
 	b3DDrawMode = false;
@@ -944,8 +944,8 @@ FText UParticleModule::GetDistributionNotAllowedOnGPUText(const FString& ModuleN
 /*-----------------------------------------------------------------------------
 	UParticleModuleSourceMovement implementation.
 -----------------------------------------------------------------------------*/
-UParticleModuleSourceMovement::UParticleModuleSourceMovement(const class FPostConstructInitializeProperties& PCIP)
-	: Super(PCIP)
+UParticleModuleSourceMovement::UParticleModuleSourceMovement(const FObjectInitializer& ObjectInitializer)
+	: Super(ObjectInitializer)
 {
 	bFinalUpdateModule = true;
 }
@@ -993,16 +993,16 @@ void UParticleModuleSourceMovement::FinalUpdate(FParticleEmitterInstance* Owner,
 /*-----------------------------------------------------------------------------
 	UParticleModuleOrientationBase implementation.
 -----------------------------------------------------------------------------*/
-UParticleModuleOrientationBase::UParticleModuleOrientationBase(const class FPostConstructInitializeProperties& PCIP)
-	: Super(PCIP)
+UParticleModuleOrientationBase::UParticleModuleOrientationBase(const FObjectInitializer& ObjectInitializer)
+	: Super(ObjectInitializer)
 {
 }
 
 /*-----------------------------------------------------------------------------
 	UParticleModuleOrientationAxisLock implementation.
 -----------------------------------------------------------------------------*/
-UParticleModuleOrientationAxisLock::UParticleModuleOrientationAxisLock(const class FPostConstructInitializeProperties& PCIP)
-	: Super(PCIP)
+UParticleModuleOrientationAxisLock::UParticleModuleOrientationAxisLock(const FObjectInitializer& ObjectInitializer)
+	: Super(ObjectInitializer)
 {
 	bSpawnModule = true;
 	bUpdateModule = true;
@@ -1056,8 +1056,8 @@ void UParticleModuleOrientationAxisLock::SetLockAxis(EParticleAxisLock eLockFlag
 /*-----------------------------------------------------------------------------
 	UParticleModuleRequired implementation.
 -----------------------------------------------------------------------------*/
-UParticleModuleRequired::UParticleModuleRequired(const class FPostConstructInitializeProperties& PCIP)
-	: Super(PCIP)
+UParticleModuleRequired::UParticleModuleRequired(const FObjectInitializer& ObjectInitializer)
+	: Super(ObjectInitializer)
 {
 	ScreenAlignment = PSA_Square;
 	bSpawnModule = true;
@@ -1221,16 +1221,16 @@ bool UParticleModuleRequired::GenerateLODModuleValues(UParticleModule* SourceMod
 /*-----------------------------------------------------------------------------
 	UParticleModuleRotationBase implementation.
 -----------------------------------------------------------------------------*/
-UParticleModuleRotationBase::UParticleModuleRotationBase(const class FPostConstructInitializeProperties& PCIP)
-	: Super(PCIP)
+UParticleModuleRotationBase::UParticleModuleRotationBase(const FObjectInitializer& ObjectInitializer)
+	: Super(ObjectInitializer)
 {
 }
 
 /*-----------------------------------------------------------------------------
 	UParticleModuleMeshRotation implementation.
 -----------------------------------------------------------------------------*/
-UParticleModuleMeshRotation::UParticleModuleMeshRotation(const class FPostConstructInitializeProperties& PCIP)
-	: Super(PCIP)
+UParticleModuleMeshRotation::UParticleModuleMeshRotation(const FObjectInitializer& ObjectInitializer)
+	: Super(ObjectInitializer)
 {
 	bSpawnModule = true;
 	bInheritParent = false;
@@ -1292,8 +1292,8 @@ void UParticleModuleMeshRotation::SpawnEx(FParticleEmitterInstance* Owner, int32
 /*-----------------------------------------------------------------------------
 	UParticleModuleMeshRotation_Seeded implementation.
 -----------------------------------------------------------------------------*/
-UParticleModuleMeshRotation_Seeded::UParticleModuleMeshRotation_Seeded(const class FPostConstructInitializeProperties& PCIP)
-	: Super(PCIP)
+UParticleModuleMeshRotation_Seeded::UParticleModuleMeshRotation_Seeded(const FObjectInitializer& ObjectInitializer)
+	: Super(ObjectInitializer)
 {
 	bSpawnModule = true;
 	bSupportsRandomSeed = true;
@@ -1328,16 +1328,16 @@ void UParticleModuleMeshRotation_Seeded::EmitterLoopingNotify(FParticleEmitterIn
 /*-----------------------------------------------------------------------------
 	UParticleModuleRotationRateBase implementation.
 -----------------------------------------------------------------------------*/
-UParticleModuleRotationRateBase::UParticleModuleRotationRateBase(const class FPostConstructInitializeProperties& PCIP)
-	: Super(PCIP)
+UParticleModuleRotationRateBase::UParticleModuleRotationRateBase(const FObjectInitializer& ObjectInitializer)
+	: Super(ObjectInitializer)
 {
 }
 
 /*-----------------------------------------------------------------------------
 	UParticleModuleMeshRotationRate implementation.
 -----------------------------------------------------------------------------*/
-UParticleModuleMeshRotationRate::UParticleModuleMeshRotationRate(const class FPostConstructInitializeProperties& PCIP)
-	: Super(PCIP)
+UParticleModuleMeshRotationRate::UParticleModuleMeshRotationRate(const FObjectInitializer& ObjectInitializer)
+	: Super(ObjectInitializer)
 {
 	bSpawnModule = true;
 	bUpdateModule = true;
@@ -1403,8 +1403,8 @@ void UParticleModuleMeshRotationRate::SetToSensibleDefaults(UParticleEmitter* Ow
 /*-----------------------------------------------------------------------------
 	UParticleModuleMeshRotationRate_Seeded implementation.
 -----------------------------------------------------------------------------*/
-UParticleModuleMeshRotationRate_Seeded::UParticleModuleMeshRotationRate_Seeded(const class FPostConstructInitializeProperties& PCIP)
-	: Super(PCIP)
+UParticleModuleMeshRotationRate_Seeded::UParticleModuleMeshRotationRate_Seeded(const FObjectInitializer& ObjectInitializer)
+	: Super(ObjectInitializer)
 {
 	bSpawnModule = true;
 	bSupportsRandomSeed = true;
@@ -1440,8 +1440,8 @@ void UParticleModuleMeshRotationRate_Seeded::EmitterLoopingNotify(FParticleEmitt
 /*-----------------------------------------------------------------------------
 	UParticleModuleMeshRotationRateMultiplyLife implementation.
 -----------------------------------------------------------------------------*/
-UParticleModuleMeshRotationRateMultiplyLife::UParticleModuleMeshRotationRateMultiplyLife(const class FPostConstructInitializeProperties& PCIP)
-	: Super(PCIP)
+UParticleModuleMeshRotationRateMultiplyLife::UParticleModuleMeshRotationRateMultiplyLife(const FObjectInitializer& ObjectInitializer)
+	: Super(ObjectInitializer)
 {
 	bSpawnModule = true;
 	bUpdateModule = true;
@@ -1508,8 +1508,8 @@ void UParticleModuleMeshRotationRateMultiplyLife::SetToSensibleDefaults(UParticl
 /*-----------------------------------------------------------------------------
 	UParticleModuleMeshRotationRateOverLife implementation.
 -----------------------------------------------------------------------------*/
-UParticleModuleMeshRotationRateOverLife::UParticleModuleMeshRotationRateOverLife(const class FPostConstructInitializeProperties& PCIP)
-	: Super(PCIP)
+UParticleModuleMeshRotationRateOverLife::UParticleModuleMeshRotationRateOverLife(const FObjectInitializer& ObjectInitializer)
+	: Super(ObjectInitializer)
 {
 	bSpawnModule = true;
 	bUpdateModule = true;
@@ -1599,8 +1599,8 @@ void UParticleModuleMeshRotationRateOverLife::SetToSensibleDefaults(UParticleEmi
 /*-----------------------------------------------------------------------------
 	UParticleModuleRotation implementation.
 -----------------------------------------------------------------------------*/
-UParticleModuleRotation::UParticleModuleRotation(const class FPostConstructInitializeProperties& PCIP)
-	: Super(PCIP)
+UParticleModuleRotation::UParticleModuleRotation(const FObjectInitializer& ObjectInitializer)
+	: Super(ObjectInitializer)
 {
 	bSpawnModule = true;
 }
@@ -1658,8 +1658,8 @@ void UParticleModuleRotation::SpawnEx(FParticleEmitterInstance* Owner, int32 Off
 /*-----------------------------------------------------------------------------
 	UParticleModuleRotation_Seeded implementation.
 -----------------------------------------------------------------------------*/
-UParticleModuleRotation_Seeded::UParticleModuleRotation_Seeded(const class FPostConstructInitializeProperties& PCIP)
-	: Super(PCIP)
+UParticleModuleRotation_Seeded::UParticleModuleRotation_Seeded(const FObjectInitializer& ObjectInitializer)
+	: Super(ObjectInitializer)
 {
 	bSpawnModule = true;
 	bSupportsRandomSeed = true;
@@ -1694,8 +1694,8 @@ void UParticleModuleRotation_Seeded::EmitterLoopingNotify(FParticleEmitterInstan
 /*-----------------------------------------------------------------------------
 	UParticleModuleRotationRate implementation.
 -----------------------------------------------------------------------------*/
-UParticleModuleRotationRate::UParticleModuleRotationRate(const class FPostConstructInitializeProperties& PCIP)
-	: Super(PCIP)
+UParticleModuleRotationRate::UParticleModuleRotationRate(const FObjectInitializer& ObjectInitializer)
+	: Super(ObjectInitializer)
 {
 	bSpawnModule = true;
 }
@@ -1776,8 +1776,8 @@ void UParticleModuleRotationRate::SetToSensibleDefaults(UParticleEmitter* Owner)
 /*-----------------------------------------------------------------------------
 	UParticleModuleRotationRate_Seeded implementation.
 -----------------------------------------------------------------------------*/
-UParticleModuleRotationRate_Seeded::UParticleModuleRotationRate_Seeded(const class FPostConstructInitializeProperties& PCIP)
-	: Super(PCIP)
+UParticleModuleRotationRate_Seeded::UParticleModuleRotationRate_Seeded(const FObjectInitializer& ObjectInitializer)
+	: Super(ObjectInitializer)
 {
 	bSpawnModule = true;
 	bSupportsRandomSeed = true;
@@ -1812,8 +1812,8 @@ void UParticleModuleRotationRate_Seeded::EmitterLoopingNotify(FParticleEmitterIn
 /*-----------------------------------------------------------------------------
 	UParticleModuleRotationOverLifetime implementation.
 -----------------------------------------------------------------------------*/
-UParticleModuleRotationOverLifetime::UParticleModuleRotationOverLifetime(const class FPostConstructInitializeProperties& PCIP)
-	: Super(PCIP)
+UParticleModuleRotationOverLifetime::UParticleModuleRotationOverLifetime(const FObjectInitializer& ObjectInitializer)
+	: Super(ObjectInitializer)
 {
 	bSpawnModule = false;
 	bUpdateModule = true;	
@@ -1872,16 +1872,16 @@ void UParticleModuleRotationOverLifetime::Update(FParticleEmitterInstance* Owner
 /*-----------------------------------------------------------------------------
 	UParticleModuleSubUVBase implementation.
 -----------------------------------------------------------------------------*/
-UParticleModuleSubUVBase::UParticleModuleSubUVBase(const class FPostConstructInitializeProperties& PCIP)
-	: Super(PCIP)
+UParticleModuleSubUVBase::UParticleModuleSubUVBase(const FObjectInitializer& ObjectInitializer)
+	: Super(ObjectInitializer)
 {
 }
 
 /*-----------------------------------------------------------------------------
 	UParticleModuleSubUV implementation.
 -----------------------------------------------------------------------------*/
-UParticleModuleSubUV::UParticleModuleSubUV(const class FPostConstructInitializeProperties& PCIP)
-	: Super(PCIP)
+UParticleModuleSubUV::UParticleModuleSubUV(const FObjectInitializer& ObjectInitializer)
+	: Super(ObjectInitializer)
 {
 	bSpawnModule = true;
 	bUpdateModule = true;	
@@ -2096,8 +2096,8 @@ void UParticleModuleSubUV::SetToSensibleDefaults(UParticleEmitter* Owner)
 /*-----------------------------------------------------------------------------
 	UParticleModuleSubUVMovie implementation.
 -----------------------------------------------------------------------------*/
-UParticleModuleSubUVMovie::UParticleModuleSubUVMovie(const class FPostConstructInitializeProperties& PCIP)
-	: Super(PCIP)
+UParticleModuleSubUVMovie::UParticleModuleSubUVMovie(const FObjectInitializer& ObjectInitializer)
+	: Super(ObjectInitializer)
 {
 	bSpawnModule = true;
 	bUpdateModule = true;
@@ -2274,8 +2274,8 @@ void UParticleModuleSubUVMovie::GetCurveObjects(TArray<FParticleCurvePair>& OutC
 /*-----------------------------------------------------------------------------
 	UParticleModuleRotationRateMultiplyLife implementation.
 -----------------------------------------------------------------------------*/
-UParticleModuleRotationRateMultiplyLife::UParticleModuleRotationRateMultiplyLife(const class FPostConstructInitializeProperties& PCIP)
-	: Super(PCIP)
+UParticleModuleRotationRateMultiplyLife::UParticleModuleRotationRateMultiplyLife(const FObjectInitializer& ObjectInitializer)
+	: Super(ObjectInitializer)
 {
 	bSpawnModule = true;
 	bUpdateModule = true;
@@ -2349,8 +2349,8 @@ void UParticleModuleRotationRateMultiplyLife::SetToSensibleDefaults(UParticleEmi
 /*-----------------------------------------------------------------------------
 	UParticleModuleAcceleration implementation.
 -----------------------------------------------------------------------------*/
-UParticleModuleAccelerationBase::UParticleModuleAccelerationBase(const class FPostConstructInitializeProperties& PCIP)
-	: Super(PCIP)
+UParticleModuleAccelerationBase::UParticleModuleAccelerationBase(const FObjectInitializer& ObjectInitializer)
+	: Super(ObjectInitializer)
 {
 }
 
@@ -2363,8 +2363,8 @@ void UParticleModuleAccelerationBase::SetToSensibleDefaults(UParticleEmitter* Ow
 /*-----------------------------------------------------------------------------
 	ParticleModuleAccelerationConstant implementation.
 -----------------------------------------------------------------------------*/
-UParticleModuleAccelerationConstant::UParticleModuleAccelerationConstant(const class FPostConstructInitializeProperties& PCIP)
-	: Super(PCIP)
+UParticleModuleAccelerationConstant::UParticleModuleAccelerationConstant(const FObjectInitializer& ObjectInitializer)
+	: Super(ObjectInitializer)
 {
 	bSpawnModule = true;
 	bUpdateModule = true;
@@ -2444,8 +2444,8 @@ void UParticleModuleAccelerationConstant::Update(FParticleEmitterInstance* Owner
 	ParticleModuleAccelerationDrag implementation.
 -----------------------------------------------------------------------------*/
 
-UParticleModuleAccelerationDrag::UParticleModuleAccelerationDrag(const class FPostConstructInitializeProperties& PCIP)
-	: Super(PCIP)
+UParticleModuleAccelerationDrag::UParticleModuleAccelerationDrag(const FObjectInitializer& ObjectInitializer)
+	: Super(ObjectInitializer)
 {
 	bSpawnModule = false;
 	bUpdateModule = true;
@@ -2520,8 +2520,8 @@ void UParticleModuleAccelerationDrag::Update(FParticleEmitterInstance* Owner, in
 /*-----------------------------------------------------------------------------
 	ParticleModuleAccelerationDragScaleOverLife implementation.
 -----------------------------------------------------------------------------*/
-UParticleModuleAccelerationDragScaleOverLife::UParticleModuleAccelerationDragScaleOverLife(const class FPostConstructInitializeProperties& PCIP)
-	: Super(PCIP)
+UParticleModuleAccelerationDragScaleOverLife::UParticleModuleAccelerationDragScaleOverLife(const FObjectInitializer& ObjectInitializer)
+	: Super(ObjectInitializer)
 {
 }
 
@@ -2585,8 +2585,8 @@ bool UParticleModuleAccelerationDragScaleOverLife::IsValidForLODLevel(UParticleL
 /*-----------------------------------------------------------------------------
 	ParticleModuleAttractorPointGravity implementation.
 -----------------------------------------------------------------------------*/
-UParticleModuleAttractorPointGravity::UParticleModuleAttractorPointGravity(const class FPostConstructInitializeProperties& PCIP)
-	: Super(PCIP)
+UParticleModuleAttractorPointGravity::UParticleModuleAttractorPointGravity(const FObjectInitializer& ObjectInitializer)
+	: Super(ObjectInitializer)
 {
 	bSupported3DDrawMode = true;
 }
@@ -2650,8 +2650,8 @@ void UParticleModuleAttractorPointGravity::Render3DPreview(FParticleEmitterInsta
 /*-----------------------------------------------------------------------------
 	UParticleModuleAcceleration implementation.
 -----------------------------------------------------------------------------*/
-UParticleModuleAcceleration::UParticleModuleAcceleration(const class FPostConstructInitializeProperties& PCIP)
-	: Super(PCIP)
+UParticleModuleAcceleration::UParticleModuleAcceleration(const FObjectInitializer& ObjectInitializer)
+	: Super(ObjectInitializer)
 {
 	bSpawnModule = true;
 	bUpdateModule = true;
@@ -2790,8 +2790,8 @@ bool UParticleModuleAcceleration::IsValidForLODLevel(UParticleLODLevel* LODLevel
 	UParticleModuleAccelerationOverLifetime implementation.
 -----------------------------------------------------------------------------*/
 
-UParticleModuleAccelerationOverLifetime::UParticleModuleAccelerationOverLifetime(const class FPostConstructInitializeProperties& PCIP)
-	: Super(PCIP)
+UParticleModuleAccelerationOverLifetime::UParticleModuleAccelerationOverLifetime(const FObjectInitializer& ObjectInitializer)
+	: Super(ObjectInitializer)
 {
 	bSpawnModule = false;
 	bUpdateModule = true;
@@ -2852,13 +2852,13 @@ void UParticleModuleAccelerationOverLifetime::Update(FParticleEmitterInstance* O
 	UParticleModuleLight implementation.
 -----------------------------------------------------------------------------*/
 
-UParticleModuleLightBase::UParticleModuleLightBase(const class FPostConstructInitializeProperties& PCIP)
-	: Super(PCIP)
+UParticleModuleLightBase::UParticleModuleLightBase(const FObjectInitializer& ObjectInitializer)
+	: Super(ObjectInitializer)
 {
 }
 
-UParticleModuleLight::UParticleModuleLight(const class FPostConstructInitializeProperties& PCIP)
-	: Super(PCIP)
+UParticleModuleLight::UParticleModuleLight(const FObjectInitializer& ObjectInitializer)
+	: Super(ObjectInitializer)
 {
 	bSpawnModule = true;
 	bUpdateModule = true;
@@ -3028,8 +3028,8 @@ void UParticleModuleLight::Render3DPreview(FParticleEmitterInstance* Owner, cons
 #endif	//#if WITH_EDITOR
 }
 
-UParticleModuleLight_Seeded::UParticleModuleLight_Seeded(const class FPostConstructInitializeProperties& PCIP)
-	: Super(PCIP)
+UParticleModuleLight_Seeded::UParticleModuleLight_Seeded(const FObjectInitializer& ObjectInitializer)
+	: Super(ObjectInitializer)
 {
 	bSpawnModule = true;
 	bSupportsRandomSeed = true;
@@ -3064,8 +3064,8 @@ void UParticleModuleLight_Seeded::EmitterLoopingNotify(FParticleEmitterInstance*
 /*-----------------------------------------------------------------------------
 	UParticleModuleTypeDataBase implementation.
 -----------------------------------------------------------------------------*/
-UParticleModuleTypeDataBase::UParticleModuleTypeDataBase(const class FPostConstructInitializeProperties& PCIP)
-	: Super(PCIP)
+UParticleModuleTypeDataBase::UParticleModuleTypeDataBase(const FObjectInitializer& ObjectInitializer)
+	: Super(ObjectInitializer)
 {
 	bSpawnModule = false;
 	bUpdateModule = false;
@@ -3079,8 +3079,8 @@ FParticleEmitterInstance* UParticleModuleTypeDataBase::CreateInstance(UParticleE
 /*-----------------------------------------------------------------------------
 	UParticleModuleTypeDataMesh implementation.
 -----------------------------------------------------------------------------*/
-UParticleModuleTypeDataMesh::UParticleModuleTypeDataMesh(const class FPostConstructInitializeProperties& PCIP)
-	: Super(PCIP)
+UParticleModuleTypeDataMesh::UParticleModuleTypeDataMesh(const FObjectInitializer& ObjectInitializer)
+	: Super(ObjectInitializer)
 {
 	CastShadows = false;
 	DoCollisions = false;
@@ -3165,16 +3165,16 @@ void UParticleModuleTypeDataMesh::PostEditChangeProperty(FPropertyChangedEvent&
 }
 #endif // WITH_EDITOR
 
-UParticleModuleKillBase::UParticleModuleKillBase(const class FPostConstructInitializeProperties& PCIP)
-	: Super(PCIP)
+UParticleModuleKillBase::UParticleModuleKillBase(const FObjectInitializer& ObjectInitializer)
+	: Super(ObjectInitializer)
 {
 }
 
 /*-----------------------------------------------------------------------------
 	UParticleModuleKillBox implementation.
 -----------------------------------------------------------------------------*/
-UParticleModuleKillBox::UParticleModuleKillBox(const class FPostConstructInitializeProperties& PCIP)
-	: Super(PCIP)
+UParticleModuleKillBox::UParticleModuleKillBox(const FObjectInitializer& ObjectInitializer)
+	: Super(ObjectInitializer)
 {
 	bUpdateModule = true;
 	bSupported3DDrawMode = true;
@@ -3334,8 +3334,8 @@ void UParticleModuleKillBox::Render3DPreview(FParticleEmitterInstance* Owner, co
 /*-----------------------------------------------------------------------------
 	UParticleModuleKillHeight implementation.
 -----------------------------------------------------------------------------*/
-UParticleModuleKillHeight::UParticleModuleKillHeight(const class FPostConstructInitializeProperties& PCIP)
-	: Super(PCIP)
+UParticleModuleKillHeight::UParticleModuleKillHeight(const FObjectInitializer& ObjectInitializer)
+	: Super(ObjectInitializer)
 {
 	bUpdateModule = true;
 	bSupported3DDrawMode = true;
@@ -3459,16 +3459,16 @@ void UParticleModuleKillHeight::Render3DPreview(FParticleEmitterInstance* Owner,
 /*-----------------------------------------------------------------------------
 	UParticleModuleLifetimeBase implementation.
 -----------------------------------------------------------------------------*/
-UParticleModuleLifetimeBase::UParticleModuleLifetimeBase(const class FPostConstructInitializeProperties& PCIP)
-	: Super(PCIP)
+UParticleModuleLifetimeBase::UParticleModuleLifetimeBase(const FObjectInitializer& ObjectInitializer)
+	: Super(ObjectInitializer)
 {
 }
 /*-----------------------------------------------------------------------------
 	UParticleModuleLifetime implementation.
 -----------------------------------------------------------------------------*/
 
-UParticleModuleLifetime::UParticleModuleLifetime(const class FPostConstructInitializeProperties& PCIP)
-	: Super(PCIP)
+UParticleModuleLifetime::UParticleModuleLifetime(const FObjectInitializer& ObjectInitializer)
+	: Super(ObjectInitializer)
 {
 	bSpawnModule = true;
 }
@@ -3570,8 +3570,8 @@ void UParticleModuleLifetime::PostEditChangeProperty(FPropertyChangedEvent& Prop
 /*-----------------------------------------------------------------------------
 	UParticleModuleLifetime_Seeded implementation.
 -----------------------------------------------------------------------------*/
-UParticleModuleLifetime_Seeded::UParticleModuleLifetime_Seeded(const class FPostConstructInitializeProperties& PCIP)
-	: Super(PCIP)
+UParticleModuleLifetime_Seeded::UParticleModuleLifetime_Seeded(const FObjectInitializer& ObjectInitializer)
+	: Super(ObjectInitializer)
 {
 	bSpawnModule = true;
 	bSupportsRandomSeed = true;
@@ -3616,16 +3616,16 @@ float UParticleModuleLifetime_Seeded::GetLifetimeValue(FParticleEmitterInstance*
 /*-----------------------------------------------------------------------------
 	UParticleModuleAttractorBase implementation.
 -----------------------------------------------------------------------------*/
-UParticleModuleAttractorBase::UParticleModuleAttractorBase(const class FPostConstructInitializeProperties& PCIP)
-	: Super(PCIP)
+UParticleModuleAttractorBase::UParticleModuleAttractorBase(const FObjectInitializer& ObjectInitializer)
+	: Super(ObjectInitializer)
 {
 }
 
 /*-----------------------------------------------------------------------------
 	UParticleModuleAttractorLine implementation.
 -----------------------------------------------------------------------------*/
-UParticleModuleAttractorLine::UParticleModuleAttractorLine(const class FPostConstructInitializeProperties& PCIP)
-	: Super(PCIP)
+UParticleModuleAttractorLine::UParticleModuleAttractorLine(const FObjectInitializer& ObjectInitializer)
+	: Super(ObjectInitializer)
 {
 	bUpdateModule = true;
 	bSupported3DDrawMode = true;
@@ -3759,8 +3759,8 @@ void UParticleModuleAttractorLine::Render3DPreview(FParticleEmitterInstance* Own
 /*-----------------------------------------------------------------------------
 	UParticleModuleAttractorParticle implementation.
 -----------------------------------------------------------------------------*/
-UParticleModuleAttractorParticle::UParticleModuleAttractorParticle(const class FPostConstructInitializeProperties& PCIP)
-	: Super(PCIP)
+UParticleModuleAttractorParticle::UParticleModuleAttractorParticle(const FObjectInitializer& ObjectInitializer)
+	: Super(ObjectInitializer)
 {
 	// Structure to hold one-time initialization
 	struct FConstructorStatics
@@ -4016,8 +4016,8 @@ uint32 UParticleModuleAttractorParticle::RequiredBytes(FParticleEmitterInstance*
 /*-----------------------------------------------------------------------------
 	UParticleModuleAttractorPoint implementation.
 -----------------------------------------------------------------------------*/
-UParticleModuleAttractorPoint::UParticleModuleAttractorPoint(const class FPostConstructInitializeProperties& PCIP)
-	: Super(PCIP)
+UParticleModuleAttractorPoint::UParticleModuleAttractorPoint(const FObjectInitializer& ObjectInitializer)
+	: Super(ObjectInitializer)
 {
 	bUpdateModule = true;
 
@@ -4217,13 +4217,13 @@ void UParticleModuleAttractorPoint::Render3DPreview(FParticleEmitterInstance* Ow
 /*-----------------------------------------------------------------------------
 	Parameter-based distributions
 -----------------------------------------------------------------------------*/
-UDistributionFloatParticleParameter::UDistributionFloatParticleParameter(const class FPostConstructInitializeProperties& PCIP)
-	: Super(PCIP)
+UDistributionFloatParticleParameter::UDistributionFloatParticleParameter(const FObjectInitializer& ObjectInitializer)
+	: Super(ObjectInitializer)
 {
 }
 
-UDistributionVectorParticleParameter::UDistributionVectorParticleParameter(const class FPostConstructInitializeProperties& PCIP)
-	: Super(PCIP)
+UDistributionVectorParticleParameter::UDistributionVectorParticleParameter(const FObjectInitializer& ObjectInitializer)
+	: Super(ObjectInitializer)
 {
 }
 
@@ -4276,8 +4276,8 @@ bool UDistributionVectorParticleParameter::GetParamValue(UObject* Data, FName Pa
 /*-----------------------------------------------------------------------------
 	Type data module for GPU particles.
 -----------------------------------------------------------------------------*/
-UParticleModuleTypeDataGpu::UParticleModuleTypeDataGpu(const class FPostConstructInitializeProperties& PCIP)
-	: Super(PCIP)
+UParticleModuleTypeDataGpu::UParticleModuleTypeDataGpu(const FObjectInitializer& ObjectInitializer)
+	: Super(ObjectInitializer)
 {
 }
 
@@ -4624,8 +4624,8 @@ FParticleEmitterInstance* UParticleModuleTypeDataGpu::CreateInstance(UParticleEm
 	UParticleModulePivotOffset implementation.
 -----------------------------------------------------------------------------*/
 
-UParticleModulePivotOffset::UParticleModulePivotOffset(const class FPostConstructInitializeProperties& PCIP)
-	: Super(PCIP)
+UParticleModulePivotOffset::UParticleModulePivotOffset(const FObjectInitializer& ObjectInitializer)
+	: Super(ObjectInitializer)
 {
 	bSpawnModule = false;
 	bUpdateModule = false;
diff --git a/Engine/Source/Runtime/Engine/Private/Particles/ParticleModules_Camera.cpp b/Engine/Source/Runtime/Engine/Private/Particles/ParticleModules_Camera.cpp
--- a/Engine/Source/Runtime/Engine/Private/Particles/ParticleModules_Camera.cpp
+++ b/Engine/Source/Runtime/Engine/Private/Particles/ParticleModules_Camera.cpp
@@ -14,8 +14,8 @@
 #include "Particles/ParticleSpriteEmitter.h"
 #include "Particles/ParticleSystemComponent.h"
 
-UParticleModuleCameraBase::UParticleModuleCameraBase(const class FPostConstructInitializeProperties& PCIP)
-	: Super(PCIP)
+UParticleModuleCameraBase::UParticleModuleCameraBase(const FObjectInitializer& ObjectInitializer)
+	: Super(ObjectInitializer)
 {
 }
 
@@ -26,8 +26,8 @@ UParticleModuleCameraBase::UParticleModuleCameraBase(const class FPostConstructI
 /*-----------------------------------------------------------------------------
 	UParticleModuleCameraOffset
 -----------------------------------------------------------------------------*/
-UParticleModuleCameraOffset::UParticleModuleCameraOffset(const class FPostConstructInitializeProperties& PCIP)
-	: Super(PCIP)
+UParticleModuleCameraOffset::UParticleModuleCameraOffset(const FObjectInitializer& ObjectInitializer)
+	: Super(ObjectInitializer)
 {
 	bSpawnModule = true;
 	bUpdateModule = true;
diff --git a/Engine/Source/Runtime/Engine/Private/Particles/ParticleModules_Collision.cpp b/Engine/Source/Runtime/Engine/Private/Particles/ParticleModules_Collision.cpp
--- a/Engine/Source/Runtime/Engine/Private/Particles/ParticleModules_Collision.cpp
+++ b/Engine/Source/Runtime/Engine/Private/Particles/ParticleModules_Collision.cpp
@@ -19,8 +19,8 @@
 #include "Particles/ParticleSpriteEmitter.h"
 #include "Particles/ParticleSystemComponent.h"
 
-UParticleModuleCollisionBase::UParticleModuleCollisionBase(const class FPostConstructInitializeProperties& PCIP)
-	: Super(PCIP)
+UParticleModuleCollisionBase::UParticleModuleCollisionBase(const FObjectInitializer& ObjectInitializer)
+	: Super(ObjectInitializer)
 {
 }
 
@@ -36,8 +36,8 @@ DEFINE_STAT(STAT_ParticleCollisionTime);
 /*-----------------------------------------------------------------------------
 	UParticleModuleCollision implementation.
 -----------------------------------------------------------------------------*/
-UParticleModuleCollision::UParticleModuleCollision(const class FPostConstructInitializeProperties& PCIP)
-	: Super(PCIP)
+UParticleModuleCollision::UParticleModuleCollision(const FObjectInitializer& ObjectInitializer)
+	: Super(ObjectInitializer)
 {
 	bSpawnModule = true;
 	bUpdateModule = true;
@@ -576,8 +576,8 @@ bool UParticleModuleCollision::PerformCollisionCheck(FParticleEmitterInstance* O
 	GPU particle collision module.
 ------------------------------------------------------------------------------*/
 
-UParticleModuleCollisionGPU::UParticleModuleCollisionGPU(const class FPostConstructInitializeProperties& PCIP)
-	: Super(PCIP)
+UParticleModuleCollisionGPU::UParticleModuleCollisionGPU(const FObjectInitializer& ObjectInitializer)
+	: Super(ObjectInitializer)
 	, Friction(0.0f)
 	, RadiusScale(1.0f)
 	, RadiusBias(0.0f)
diff --git a/Engine/Source/Runtime/Engine/Private/Particles/ParticleModules_Color.cpp b/Engine/Source/Runtime/Engine/Private/Particles/ParticleModules_Color.cpp
--- a/Engine/Source/Runtime/Engine/Private/Particles/ParticleModules_Color.cpp
+++ b/Engine/Source/Runtime/Engine/Private/Particles/ParticleModules_Color.cpp
@@ -17,8 +17,8 @@
 #include "Particles/ParticleSystem.h"
 #include "Particles/ParticleSystemComponent.h"
 
-UParticleModuleColorBase::UParticleModuleColorBase(const class FPostConstructInitializeProperties& PCIP)
-	: Super(PCIP)
+UParticleModuleColorBase::UParticleModuleColorBase(const FObjectInitializer& ObjectInitializer)
+	: Super(ObjectInitializer)
 {
 }
 
@@ -30,8 +30,8 @@ UParticleModuleColorBase::UParticleModuleColorBase(const class FPostConstructIni
 /*-----------------------------------------------------------------------------
 	UParticleModuleColor implementation.
 -----------------------------------------------------------------------------*/
-UParticleModuleColor::UParticleModuleColor(const class FPostConstructInitializeProperties& PCIP)
-	: Super(PCIP)
+UParticleModuleColor::UParticleModuleColor(const FObjectInitializer& ObjectInitializer)
+	: Super(ObjectInitializer)
 {
 	bSpawnModule = true;
 	bUpdateModule = false;
@@ -182,8 +182,8 @@ void UParticleModuleColor::SetToSensibleDefaults(UParticleEmitter* Owner)
 /*-----------------------------------------------------------------------------
 	UParticleModuleColor_Seeded implementation.
 -----------------------------------------------------------------------------*/
-UParticleModuleColor_Seeded::UParticleModuleColor_Seeded(const class FPostConstructInitializeProperties& PCIP)
-	: Super(PCIP)
+UParticleModuleColor_Seeded::UParticleModuleColor_Seeded(const FObjectInitializer& ObjectInitializer)
+	: Super(ObjectInitializer)
 {
 	bSpawnModule = true;
 	bSupportsRandomSeed = true;
@@ -218,8 +218,8 @@ void UParticleModuleColor_Seeded::EmitterLoopingNotify(FParticleEmitterInstance*
 /*-----------------------------------------------------------------------------
 	UParticleModuleColorOverLife implementation.
 -----------------------------------------------------------------------------*/
-UParticleModuleColorOverLife::UParticleModuleColorOverLife(const class FPostConstructInitializeProperties& PCIP)
-	: Super(PCIP)
+UParticleModuleColorOverLife::UParticleModuleColorOverLife(const FObjectInitializer& ObjectInitializer)
+	: Super(ObjectInitializer)
 {
 	bSpawnModule = true;
 	bUpdateModule = true;
@@ -478,8 +478,8 @@ void UParticleModuleColorOverLife::SetToSensibleDefaults(UParticleEmitter* Owner
 /*-----------------------------------------------------------------------------
 	UParticleModuleColorScaleOverLife implementation.
 -----------------------------------------------------------------------------*/
-UParticleModuleColorScaleOverLife::UParticleModuleColorScaleOverLife(const class FPostConstructInitializeProperties& PCIP)
-	: Super(PCIP)
+UParticleModuleColorScaleOverLife::UParticleModuleColorScaleOverLife(const FObjectInitializer& ObjectInitializer)
+	: Super(ObjectInitializer)
 {
 	bSpawnModule = true;
 	bUpdateModule = true;
diff --git a/Engine/Source/Runtime/Engine/Private/Particles/ParticleModules_Event.cpp b/Engine/Source/Runtime/Engine/Private/Particles/ParticleModules_Event.cpp
--- a/Engine/Source/Runtime/Engine/Private/Particles/ParticleModules_Event.cpp
+++ b/Engine/Source/Runtime/Engine/Private/Particles/ParticleModules_Event.cpp
@@ -22,24 +22,24 @@
 /*-----------------------------------------------------------------------------
 	UParticleModuleEventSendToGame implementation.
 -----------------------------------------------------------------------------*/
-UParticleModuleEventSendToGame::UParticleModuleEventSendToGame(const class FPostConstructInitializeProperties& PCIP)
-	: Super(PCIP)
+UParticleModuleEventSendToGame::UParticleModuleEventSendToGame(const FObjectInitializer& ObjectInitializer)
+	: Super(ObjectInitializer)
 {
 }
 
 /*-----------------------------------------------------------------------------
 	UParticleModuleEventBase implementation.
 -----------------------------------------------------------------------------*/
-UParticleModuleEventBase::UParticleModuleEventBase(const class FPostConstructInitializeProperties& PCIP)
-	: Super(PCIP)
+UParticleModuleEventBase::UParticleModuleEventBase(const FObjectInitializer& ObjectInitializer)
+	: Super(ObjectInitializer)
 {
 }
 
 /*-----------------------------------------------------------------------------
 	UParticleModuleEventGenerator implementation.
 -----------------------------------------------------------------------------*/
-UParticleModuleEventGenerator::UParticleModuleEventGenerator(const class FPostConstructInitializeProperties& PCIP)
-	: Super(PCIP)
+UParticleModuleEventGenerator::UParticleModuleEventGenerator(const FObjectInitializer& ObjectInitializer)
+	: Super(ObjectInitializer)
 {
 	bSpawnModule = true;
 	bUpdateModule = true;
@@ -254,16 +254,16 @@ bool UParticleModuleEventGenerator::HandleParticleBurst(FParticleEmitterInstance
 /*-----------------------------------------------------------------------------
 	UParticleModuleEventReceiverBase implementation.
 -----------------------------------------------------------------------------*/
-UParticleModuleEventReceiverBase::UParticleModuleEventReceiverBase(const class FPostConstructInitializeProperties& PCIP)
-	: Super(PCIP)
+UParticleModuleEventReceiverBase::UParticleModuleEventReceiverBase(const FObjectInitializer& ObjectInitializer)
+	: Super(ObjectInitializer)
 {
 }
 
 /*-----------------------------------------------------------------------------
 	UParticleModuleEventReceiver implementation.
 -----------------------------------------------------------------------------*/
-UParticleModuleEventReceiverKillParticles::UParticleModuleEventReceiverKillParticles(const class FPostConstructInitializeProperties& PCIP)
-	: Super(PCIP)
+UParticleModuleEventReceiverKillParticles::UParticleModuleEventReceiverKillParticles(const FObjectInitializer& ObjectInitializer)
+	: Super(ObjectInitializer)
 {
 	bSpawnModule = true;
 	bUpdateModule = true;
@@ -287,8 +287,8 @@ bool UParticleModuleEventReceiverKillParticles::ProcessParticleEvent(FParticleEm
 /*-----------------------------------------------------------------------------
 	UParticleModuleEventReceiver implementation.
 -----------------------------------------------------------------------------*/
-UParticleModuleEventReceiverSpawn::UParticleModuleEventReceiverSpawn(const class FPostConstructInitializeProperties& PCIP)
-	: Super(PCIP)
+UParticleModuleEventReceiverSpawn::UParticleModuleEventReceiverSpawn(const FObjectInitializer& ObjectInitializer)
+	: Super(ObjectInitializer)
 {
 	bSpawnModule = true;
 	bUpdateModule = true;
diff --git a/Engine/Source/Runtime/Engine/Private/Particles/ParticleModules_Location.cpp b/Engine/Source/Runtime/Engine/Private/Particles/ParticleModules_Location.cpp
--- a/Engine/Source/Runtime/Engine/Private/Particles/ParticleModules_Location.cpp
+++ b/Engine/Source/Runtime/Engine/Private/Particles/ParticleModules_Location.cpp
@@ -31,8 +31,8 @@
 #include "Particles/ParticleSystemComponent.h"
 #include "Animation/SkeletalMeshActor.h"
 
-UParticleModuleLocationBase::UParticleModuleLocationBase(const class FPostConstructInitializeProperties& PCIP)
-	: Super(PCIP)
+UParticleModuleLocationBase::UParticleModuleLocationBase(const FObjectInitializer& ObjectInitializer)
+	: Super(ObjectInitializer)
 {
 }
 
@@ -45,8 +45,8 @@ UParticleModuleLocationBase::UParticleModuleLocationBase(const class FPostConstr
 	UParticleModuleLocation implementation.
 -----------------------------------------------------------------------------*/
 
-UParticleModuleLocation::UParticleModuleLocation(const class FPostConstructInitializeProperties& PCIP)
-	: Super(PCIP)
+UParticleModuleLocation::UParticleModuleLocation(const FObjectInitializer& ObjectInitializer)
+	: Super(ObjectInitializer)
 {
 	bSpawnModule = true;
 	bSupported3DDrawMode = true;
@@ -179,8 +179,8 @@ void UParticleModuleLocation::Render3DPreview(FParticleEmitterInstance* Owner, c
 /*-----------------------------------------------------------------------------
 	UParticleModuleLocation_Seeded implementation.
 -----------------------------------------------------------------------------*/
-UParticleModuleLocation_Seeded::UParticleModuleLocation_Seeded(const class FPostConstructInitializeProperties& PCIP)
-	: Super(PCIP)
+UParticleModuleLocation_Seeded::UParticleModuleLocation_Seeded(const FObjectInitializer& ObjectInitializer)
+	: Super(ObjectInitializer)
 {
 	bSpawnModule = true;
 	bSupportsRandomSeed = true;
@@ -215,8 +215,8 @@ void UParticleModuleLocation_Seeded::EmitterLoopingNotify(FParticleEmitterInstan
 /*-----------------------------------------------------------------------------
 	UParticleModuleLocationWorldOffset implementation.
 -----------------------------------------------------------------------------*/
-UParticleModuleLocationWorldOffset::UParticleModuleLocationWorldOffset(const class FPostConstructInitializeProperties& PCIP)
-	: Super(PCIP)
+UParticleModuleLocationWorldOffset::UParticleModuleLocationWorldOffset(const FObjectInitializer& ObjectInitializer)
+	: Super(ObjectInitializer)
 {
 }
 
@@ -242,8 +242,8 @@ void UParticleModuleLocationWorldOffset::SpawnEx(FParticleEmitterInstance* Owner
 /*-----------------------------------------------------------------------------
 	UParticleModuleLocationWorldOffset_Seeded implementation.
 -----------------------------------------------------------------------------*/
-UParticleModuleLocationWorldOffset_Seeded::UParticleModuleLocationWorldOffset_Seeded(const class FPostConstructInitializeProperties& PCIP)
-	: Super(PCIP)
+UParticleModuleLocationWorldOffset_Seeded::UParticleModuleLocationWorldOffset_Seeded(const FObjectInitializer& ObjectInitializer)
+	: Super(ObjectInitializer)
 {
 	bSpawnModule = true;
 	bSupportsRandomSeed = true;
@@ -278,8 +278,8 @@ void UParticleModuleLocationWorldOffset_Seeded::EmitterLoopingNotify(FParticleEm
 /*-----------------------------------------------------------------------------
 	UParticleModuleLocationDirect implementation.
 -----------------------------------------------------------------------------*/
-UParticleModuleLocationDirect::UParticleModuleLocationDirect(const class FPostConstructInitializeProperties& PCIP)
-	: Super(PCIP)
+UParticleModuleLocationDirect::UParticleModuleLocationDirect(const FObjectInitializer& ObjectInitializer)
+	: Super(ObjectInitializer)
 {
 	bSpawnModule = true;
 	bUpdateModule = true;
@@ -401,8 +401,8 @@ uint32 UParticleModuleLocationDirect::RequiredBytes(FParticleEmitterInstance* Ow
 /*-----------------------------------------------------------------------------
 	UParticleModuleLocationEmitter implementation.
 -----------------------------------------------------------------------------*/
-UParticleModuleLocationEmitter::UParticleModuleLocationEmitter(const class FPostConstructInitializeProperties& PCIP)
-	: Super(PCIP)
+UParticleModuleLocationEmitter::UParticleModuleLocationEmitter(const FObjectInitializer& ObjectInitializer)
+	: Super(ObjectInitializer)
 {
 	// Structure to hold one-time initialization
 	struct FConstructorStatics
@@ -550,8 +550,8 @@ uint32 UParticleModuleLocationEmitter::RequiredBytesPerInstance(FParticleEmitter
 /*-----------------------------------------------------------------------------
 	UParticleModuleLocationEmitterDirect implementation.
 -----------------------------------------------------------------------------*/
-UParticleModuleLocationEmitterDirect::UParticleModuleLocationEmitterDirect(const class FPostConstructInitializeProperties& PCIP)
-	: Super(PCIP)
+UParticleModuleLocationEmitterDirect::UParticleModuleLocationEmitterDirect(const FObjectInitializer& ObjectInitializer)
+	: Super(ObjectInitializer)
 {
 	// Structure to hold one-time initialization
 	struct FConstructorStatics
@@ -651,8 +651,8 @@ void UParticleModuleLocationEmitterDirect::Update(FParticleEmitterInstance* Owne
 /*-----------------------------------------------------------------------------
 	UParticleModuleLocationPrimitiveBase implementation.
 -----------------------------------------------------------------------------*/
-UParticleModuleLocationPrimitiveBase::UParticleModuleLocationPrimitiveBase(const class FPostConstructInitializeProperties& PCIP)
-	: Super(PCIP)
+UParticleModuleLocationPrimitiveBase::UParticleModuleLocationPrimitiveBase(const FObjectInitializer& ObjectInitializer)
+	: Super(ObjectInitializer)
 {
 	bSpawnModule = true;
 	Positive_X = true;
@@ -783,8 +783,8 @@ void UParticleModuleLocationPrimitiveBase::DetermineUnitDirection(FParticleEmitt
 /*-----------------------------------------------------------------------------
 	UParticleModuleLocationPrimitiveTriangle implementation.
 -----------------------------------------------------------------------------*/
-UParticleModuleLocationPrimitiveTriangle::UParticleModuleLocationPrimitiveTriangle(const class FPostConstructInitializeProperties& PCIP)
-	: Super(PCIP)
+UParticleModuleLocationPrimitiveTriangle::UParticleModuleLocationPrimitiveTriangle(const FObjectInitializer& ObjectInitializer)
+	: Super(ObjectInitializer)
 {
 
 	bSupported3DDrawMode = true;
@@ -962,8 +962,8 @@ void UParticleModuleLocationPrimitiveTriangle::Render3DPreview(FParticleEmitterI
 /*-----------------------------------------------------------------------------
 	UParticleModuleLocationPrimitiveCylinder implementation.
 -----------------------------------------------------------------------------*/
-UParticleModuleLocationPrimitiveCylinder::UParticleModuleLocationPrimitiveCylinder(const class FPostConstructInitializeProperties& PCIP)
-	: Super(PCIP)
+UParticleModuleLocationPrimitiveCylinder::UParticleModuleLocationPrimitiveCylinder(const FObjectInitializer& ObjectInitializer)
+	: Super(ObjectInitializer)
 {
 	RadialVelocity = true;
 	bSupported3DDrawMode = true;
@@ -1243,8 +1243,8 @@ void UParticleModuleLocationPrimitiveCylinder::Render3DPreview(FParticleEmitterI
 /*-----------------------------------------------------------------------------
 	UParticleModuleLocationPrimitiveCylinder_Seeded implementation.
 -----------------------------------------------------------------------------*/
-UParticleModuleLocationPrimitiveCylinder_Seeded::UParticleModuleLocationPrimitiveCylinder_Seeded(const class FPostConstructInitializeProperties& PCIP)
-	: Super(PCIP)
+UParticleModuleLocationPrimitiveCylinder_Seeded::UParticleModuleLocationPrimitiveCylinder_Seeded(const FObjectInitializer& ObjectInitializer)
+	: Super(ObjectInitializer)
 {
 	bSpawnModule = true;
 	bSupportsRandomSeed = true;
@@ -1279,8 +1279,8 @@ void UParticleModuleLocationPrimitiveCylinder_Seeded::EmitterLoopingNotify(FPart
 /*-----------------------------------------------------------------------------
 	UParticleModuleLocationPrimitiveSphere implementation.
 -----------------------------------------------------------------------------*/
-UParticleModuleLocationPrimitiveSphere::UParticleModuleLocationPrimitiveSphere(const class FPostConstructInitializeProperties& PCIP)
-	: Super(PCIP)
+UParticleModuleLocationPrimitiveSphere::UParticleModuleLocationPrimitiveSphere(const FObjectInitializer& ObjectInitializer)
+	: Super(ObjectInitializer)
 {
 
 	bSupported3DDrawMode = true;
@@ -1463,8 +1463,8 @@ void UParticleModuleLocationPrimitiveSphere::Render3DPreview(FParticleEmitterIns
 /*-----------------------------------------------------------------------------
 	UParticleModuleLocationPrimitiveSphere_Seeded implementation.
 -----------------------------------------------------------------------------*/
-UParticleModuleLocationPrimitiveSphere_Seeded::UParticleModuleLocationPrimitiveSphere_Seeded(const class FPostConstructInitializeProperties& PCIP)
-	: Super(PCIP)
+UParticleModuleLocationPrimitiveSphere_Seeded::UParticleModuleLocationPrimitiveSphere_Seeded(const FObjectInitializer& ObjectInitializer)
+	: Super(ObjectInitializer)
 {
 	bSpawnModule = true;
 	bSupportsRandomSeed = true;
@@ -1499,8 +1499,8 @@ void UParticleModuleLocationPrimitiveSphere_Seeded::EmitterLoopingNotify(FPartic
 /*-----------------------------------------------------------------------------
 	UParticleModuleLocationBoneSocket implementation.
 -----------------------------------------------------------------------------*/
-UParticleModuleLocationBoneSocket::UParticleModuleLocationBoneSocket(const class FPostConstructInitializeProperties& PCIP)
-	: Super(PCIP)
+UParticleModuleLocationBoneSocket::UParticleModuleLocationBoneSocket(const FObjectInitializer& ObjectInitializer)
+	: Super(ObjectInitializer)
 {
 	// Structure to hold one-time initialization
 	struct FConstructorStatics
@@ -1993,8 +1993,8 @@ bool UParticleModuleLocationBoneSocket::GetParticleLocation(FParticleEmitterInst
 /*-----------------------------------------------------------------------------
 	UParticleModuleLocationVertSurface implementation.
 -----------------------------------------------------------------------------*/
-UParticleModuleLocationSkelVertSurface::UParticleModuleLocationSkelVertSurface(const class FPostConstructInitializeProperties& PCIP)
-	: Super(PCIP)
+UParticleModuleLocationSkelVertSurface::UParticleModuleLocationSkelVertSurface(const FObjectInitializer& ObjectInitializer)
+	: Super(ObjectInitializer)
 {
 	// Structure to hold one-time initialization
 	struct FConstructorStatics
diff --git a/Engine/Source/Runtime/Engine/Private/Particles/ParticleModules_Material.cpp b/Engine/Source/Runtime/Engine/Private/Particles/ParticleModules_Material.cpp
--- a/Engine/Source/Runtime/Engine/Private/Particles/ParticleModules_Material.cpp
+++ b/Engine/Source/Runtime/Engine/Private/Particles/ParticleModules_Material.cpp
@@ -11,8 +11,8 @@
 #include "Particles/ParticleModuleRequired.h"
 #include "Particles/TypeData/ParticleModuleTypeDataMesh.h"
 
-UParticleModuleMaterialBase::UParticleModuleMaterialBase(const class FPostConstructInitializeProperties& PCIP)
-	: Super(PCIP)
+UParticleModuleMaterialBase::UParticleModuleMaterialBase(const FObjectInitializer& ObjectInitializer)
+	: Super(ObjectInitializer)
 {
 }
 
@@ -23,8 +23,8 @@ UParticleModuleMaterialBase::UParticleModuleMaterialBase(const class FPostConstr
 /*-----------------------------------------------------------------------------
 	UParticleModuleMeshMaterial
 -----------------------------------------------------------------------------*/
-UParticleModuleMeshMaterial::UParticleModuleMeshMaterial(const class FPostConstructInitializeProperties& PCIP)
-	: Super(PCIP)
+UParticleModuleMeshMaterial::UParticleModuleMeshMaterial(const FObjectInitializer& ObjectInitializer)
+	: Super(ObjectInitializer)
 {
 	bSpawnModule = true;
 	bUpdateModule = true;
diff --git a/Engine/Source/Runtime/Engine/Private/Particles/ParticleModules_Orbit.cpp b/Engine/Source/Runtime/Engine/Private/Particles/ParticleModules_Orbit.cpp
--- a/Engine/Source/Runtime/Engine/Private/Particles/ParticleModules_Orbit.cpp
+++ b/Engine/Source/Runtime/Engine/Private/Particles/ParticleModules_Orbit.cpp
@@ -11,8 +11,8 @@
 #include "Particles/ParticleLODLevel.h"
 #include "Particles/ParticleSystemComponent.h"
 
-UParticleModuleOrbitBase::UParticleModuleOrbitBase(const class FPostConstructInitializeProperties& PCIP)
-	: Super(PCIP)
+UParticleModuleOrbitBase::UParticleModuleOrbitBase(const FObjectInitializer& ObjectInitializer)
+	: Super(ObjectInitializer)
 {
 }
 
@@ -24,8 +24,8 @@ UParticleModuleOrbitBase::UParticleModuleOrbitBase(const class FPostConstructIni
 /*-----------------------------------------------------------------------------
 	UParticleModuleOrbit implementation.
 -----------------------------------------------------------------------------*/
-UParticleModuleOrbit::UParticleModuleOrbit(const class FPostConstructInitializeProperties& PCIP)
-	: Super(PCIP)
+UParticleModuleOrbit::UParticleModuleOrbit(const FObjectInitializer& ObjectInitializer)
+	: Super(ObjectInitializer)
 {
 	bSpawnModule = true;
 	bUpdateModule = true;
diff --git a/Engine/Source/Runtime/Engine/Private/Particles/ParticleModules_Parameter.cpp b/Engine/Source/Runtime/Engine/Private/Particles/ParticleModules_Parameter.cpp
--- a/Engine/Source/Runtime/Engine/Private/Particles/ParticleModules_Parameter.cpp
+++ b/Engine/Source/Runtime/Engine/Private/Particles/ParticleModules_Parameter.cpp
@@ -16,8 +16,8 @@
 #include "Particles/ParticleLODLevel.h"
 #include "Particles/ParticleModuleRequired.h"
 
-UParticleModuleParameterBase::UParticleModuleParameterBase(const class FPostConstructInitializeProperties& PCIP)
-	: Super(PCIP)
+UParticleModuleParameterBase::UParticleModuleParameterBase(const FObjectInitializer& ObjectInitializer)
+	: Super(ObjectInitializer)
 {
 }
 
@@ -28,8 +28,8 @@ UParticleModuleParameterBase::UParticleModuleParameterBase(const class FPostCons
 /*-----------------------------------------------------------------------------
 	UParticleModuleParameterDynamic implementation.
 -----------------------------------------------------------------------------*/
-UParticleModuleParameterDynamic::UParticleModuleParameterDynamic(const class FPostConstructInitializeProperties& PCIP)
-	: Super(PCIP)
+UParticleModuleParameterDynamic::UParticleModuleParameterDynamic(const FObjectInitializer& ObjectInitializer)
+	: Super(ObjectInitializer)
 {
 	bSpawnModule = true;
 	bUpdateModule = true;
@@ -534,8 +534,8 @@ void UParticleModuleParameterDynamic::UpdateUsageFlags()
 /*-----------------------------------------------------------------------------
 	UParticleModuleParameterDynamic_Seeded implementation.
 -----------------------------------------------------------------------------*/
-UParticleModuleParameterDynamic_Seeded::UParticleModuleParameterDynamic_Seeded(const class FPostConstructInitializeProperties& PCIP)
-	: Super(PCIP)
+UParticleModuleParameterDynamic_Seeded::UParticleModuleParameterDynamic_Seeded(const FObjectInitializer& ObjectInitializer)
+	: Super(ObjectInitializer)
 {
 	bSpawnModule = true;
 	bUpdateModule = true;
diff --git a/Engine/Source/Runtime/Engine/Private/Particles/ParticleModules_Size.cpp b/Engine/Source/Runtime/Engine/Private/Particles/ParticleModules_Size.cpp
--- a/Engine/Source/Runtime/Engine/Private/Particles/ParticleModules_Size.cpp
+++ b/Engine/Source/Runtime/Engine/Private/Particles/ParticleModules_Size.cpp
@@ -16,8 +16,8 @@
 #include "Particles/ParticleLODLevel.h"
 #include "Particles/ParticleSystemComponent.h"
 
-UParticleModuleSizeBase::UParticleModuleSizeBase(const class FPostConstructInitializeProperties& PCIP)
-	: Super(PCIP)
+UParticleModuleSizeBase::UParticleModuleSizeBase(const FObjectInitializer& ObjectInitializer)
+	: Super(ObjectInitializer)
 {
 }
 
@@ -30,8 +30,8 @@ UParticleModuleSizeBase::UParticleModuleSizeBase(const class FPostConstructIniti
 	UParticleModuleSize implementation.
 -----------------------------------------------------------------------------*/
 
-UParticleModuleSize::UParticleModuleSize(const class FPostConstructInitializeProperties& PCIP)
-	: Super(PCIP)
+UParticleModuleSize::UParticleModuleSize(const FObjectInitializer& ObjectInitializer)
+	: Super(ObjectInitializer)
 {
 	bSpawnModule = true;
 	bUpdateModule = false;
@@ -102,8 +102,8 @@ void UParticleModuleSize::SpawnEx(FParticleEmitterInstance* Owner, int32 Offset,
 /*-----------------------------------------------------------------------------
 	UParticleModuleSize_Seeded implementation.
 -----------------------------------------------------------------------------*/
-UParticleModuleSize_Seeded::UParticleModuleSize_Seeded(const class FPostConstructInitializeProperties& PCIP)
-	: Super(PCIP)
+UParticleModuleSize_Seeded::UParticleModuleSize_Seeded(const FObjectInitializer& ObjectInitializer)
+	: Super(ObjectInitializer)
 {
 	bSpawnModule = true;
 	bSupportsRandomSeed = true;
@@ -138,8 +138,8 @@ void UParticleModuleSize_Seeded::EmitterLoopingNotify(FParticleEmitterInstance*
 /*-----------------------------------------------------------------------------
 	UParticleModuleSizeMultiplyLife implementation.
 -----------------------------------------------------------------------------*/
-UParticleModuleSizeMultiplyLife::UParticleModuleSizeMultiplyLife(const class FPostConstructInitializeProperties& PCIP)
-	: Super(PCIP)
+UParticleModuleSizeMultiplyLife::UParticleModuleSizeMultiplyLife(const FObjectInitializer& ObjectInitializer)
+	: Super(ObjectInitializer)
 {
 	bSpawnModule = true;
 	bUpdateModule = true;
@@ -337,8 +337,8 @@ void UParticleModuleSizeMultiplyLife::SetToSensibleDefaults(UParticleEmitter* Ow
 /*-----------------------------------------------------------------------------
 	UParticleModuleSizeScale implementation.
 -----------------------------------------------------------------------------*/
-UParticleModuleSizeScale::UParticleModuleSizeScale(const class FPostConstructInitializeProperties& PCIP)
-	: Super(PCIP)
+UParticleModuleSizeScale::UParticleModuleSizeScale(const FObjectInitializer& ObjectInitializer)
+	: Super(ObjectInitializer)
 {
 	bSpawnModule = true;
 	bUpdateModule = true;
@@ -429,8 +429,8 @@ void UParticleModuleSizeScale::SetToSensibleDefaults(UParticleEmitter* Owner)
 /*------------------------------------------------------------------------------
 	Scale size by speed module.
 ------------------------------------------------------------------------------*/
-UParticleModuleSizeScaleBySpeed::UParticleModuleSizeScaleBySpeed(const class FPostConstructInitializeProperties& PCIP)
-	: Super(PCIP)
+UParticleModuleSizeScaleBySpeed::UParticleModuleSizeScaleBySpeed(const FObjectInitializer& ObjectInitializer)
+	: Super(ObjectInitializer)
 {
 	bUpdateModule = true;
 	MaxScale.X = 1;
diff --git a/Engine/Source/Runtime/Engine/Private/Particles/ParticleModules_Spawn.cpp b/Engine/Source/Runtime/Engine/Private/Particles/ParticleModules_Spawn.cpp
--- a/Engine/Source/Runtime/Engine/Private/Particles/ParticleModules_Spawn.cpp
+++ b/Engine/Source/Runtime/Engine/Private/Particles/ParticleModules_Spawn.cpp
@@ -12,8 +12,8 @@
 #include "Particles/ParticleEmitter.h"
 #include "Particles/ParticleSystemComponent.h"
 
-UParticleModuleSpawnBase::UParticleModuleSpawnBase(const class FPostConstructInitializeProperties& PCIP)
-	: Super(PCIP)
+UParticleModuleSpawnBase::UParticleModuleSpawnBase(const FObjectInitializer& ObjectInitializer)
+	: Super(ObjectInitializer)
 {
 	bProcessSpawnRate = true;
 	bProcessBurstList = true;
@@ -26,8 +26,8 @@ UParticleModuleSpawnBase::UParticleModuleSpawnBase(const class FPostConstructIni
 /*-----------------------------------------------------------------------------
 	UParticleModuleSpawn implementation.
 -----------------------------------------------------------------------------*/
-UParticleModuleSpawn::UParticleModuleSpawn(const class FPostConstructInitializeProperties& PCIP)
-	: Super(PCIP)
+UParticleModuleSpawn::UParticleModuleSpawn(const FObjectInitializer& ObjectInitializer)
+	: Super(ObjectInitializer)
 {
 	bProcessSpawnRate = true;
 	LODDuplicate = false;
@@ -270,8 +270,8 @@ float UParticleModuleSpawn::GetGlobalRateScale()const
 /*-----------------------------------------------------------------------------
 	UParticleModuleSpawnPerUnit implementation.
 -----------------------------------------------------------------------------*/
-UParticleModuleSpawnPerUnit::UParticleModuleSpawnPerUnit(const class FPostConstructInitializeProperties& PCIP)
-	: Super(PCIP)
+UParticleModuleSpawnPerUnit::UParticleModuleSpawnPerUnit(const FObjectInitializer& ObjectInitializer)
+	: Super(ObjectInitializer)
 {
 	bSpawnModule = false;
 	bUpdateModule = false;
diff --git a/Engine/Source/Runtime/Engine/Private/Particles/ParticleModules_VectorField.cpp b/Engine/Source/Runtime/Engine/Private/Particles/ParticleModules_VectorField.cpp
--- a/Engine/Source/Runtime/Engine/Private/Particles/ParticleModules_VectorField.cpp
+++ b/Engine/Source/Runtime/Engine/Private/Particles/ParticleModules_VectorField.cpp
@@ -20,13 +20,13 @@
 /*------------------------------------------------------------------------------
 	Global vector field scale.
 ------------------------------------------------------------------------------*/
-UParticleModuleVectorFieldBase::UParticleModuleVectorFieldBase(const class FPostConstructInitializeProperties& PCIP)
-	: Super(PCIP)
+UParticleModuleVectorFieldBase::UParticleModuleVectorFieldBase(const FObjectInitializer& ObjectInitializer)
+	: Super(ObjectInitializer)
 {
 }
 
-UParticleModuleVectorFieldGlobal::UParticleModuleVectorFieldGlobal(const class FPostConstructInitializeProperties& PCIP)
-	: Super(PCIP)
+UParticleModuleVectorFieldGlobal::UParticleModuleVectorFieldGlobal(const FObjectInitializer& ObjectInitializer)
+	: Super(ObjectInitializer)
 {
 }
 
@@ -39,8 +39,8 @@ void UParticleModuleVectorFieldGlobal::CompileModule(FParticleEmitterBuildInfo&
 /*------------------------------------------------------------------------------
 	Per-particle vector field scale.
 ------------------------------------------------------------------------------*/
-UParticleModuleVectorFieldScale::UParticleModuleVectorFieldScale(const class FPostConstructInitializeProperties& PCIP)
-	: Super(PCIP)
+UParticleModuleVectorFieldScale::UParticleModuleVectorFieldScale(const FObjectInitializer& ObjectInitializer)
+	: Super(ObjectInitializer)
 {
 }
 
@@ -89,8 +89,8 @@ bool UParticleModuleVectorFieldScale::IsValidForLODLevel(UParticleLODLevel* LODL
 /*------------------------------------------------------------------------------
 	Per-particle vector field scale over life.
 ------------------------------------------------------------------------------*/
-UParticleModuleVectorFieldScaleOverLife::UParticleModuleVectorFieldScaleOverLife(const class FPostConstructInitializeProperties& PCIP)
-	: Super(PCIP)
+UParticleModuleVectorFieldScaleOverLife::UParticleModuleVectorFieldScaleOverLife(const FObjectInitializer& ObjectInitializer)
+	: Super(ObjectInitializer)
 {
 }
 
@@ -139,8 +139,8 @@ bool UParticleModuleVectorFieldScaleOverLife::IsValidForLODLevel(UParticleLODLev
 /*------------------------------------------------------------------------------
 	Local vector fields.
 ------------------------------------------------------------------------------*/
-UParticleModuleVectorFieldLocal::UParticleModuleVectorFieldLocal(const class FPostConstructInitializeProperties& PCIP)
-	: Super(PCIP)
+UParticleModuleVectorFieldLocal::UParticleModuleVectorFieldLocal(const FObjectInitializer& ObjectInitializer)
+	: Super(ObjectInitializer)
 {
 	RelativeScale3D = FVector(1.0f, 1.0f, 1.0f);
 
@@ -165,8 +165,8 @@ void UParticleModuleVectorFieldLocal::CompileModule(FParticleEmitterBuildInfo& E
 /*------------------------------------------------------------------------------
 	Local vector initial rotation.
 ------------------------------------------------------------------------------*/
-UParticleModuleVectorFieldRotation::UParticleModuleVectorFieldRotation(const class FPostConstructInitializeProperties& PCIP)
-	: Super(PCIP)
+UParticleModuleVectorFieldRotation::UParticleModuleVectorFieldRotation(const FObjectInitializer& ObjectInitializer)
+	: Super(ObjectInitializer)
 {
 }
 
@@ -179,8 +179,8 @@ void UParticleModuleVectorFieldRotation::CompileModule(FParticleEmitterBuildInfo
 /*------------------------------------------------------------------------------
 	Local vector field rotation rate.
 ------------------------------------------------------------------------------*/
-UParticleModuleVectorFieldRotationRate::UParticleModuleVectorFieldRotationRate(const class FPostConstructInitializeProperties& PCIP)
-	: Super(PCIP)
+UParticleModuleVectorFieldRotationRate::UParticleModuleVectorFieldRotationRate(const FObjectInitializer& ObjectInitializer)
+	: Super(ObjectInitializer)
 {
 }
 
diff --git a/Engine/Source/Runtime/Engine/Private/Particles/ParticleModules_Velocity.cpp b/Engine/Source/Runtime/Engine/Private/Particles/ParticleModules_Velocity.cpp
--- a/Engine/Source/Runtime/Engine/Private/Particles/ParticleModules_Velocity.cpp
+++ b/Engine/Source/Runtime/Engine/Private/Particles/ParticleModules_Velocity.cpp
@@ -21,8 +21,8 @@
 #include "Particles/ParticleSpriteEmitter.h"
 #include "Particles/ParticleSystemComponent.h"
 
-UParticleModuleVelocityBase::UParticleModuleVelocityBase(const class FPostConstructInitializeProperties& PCIP)
-	: Super(PCIP)
+UParticleModuleVelocityBase::UParticleModuleVelocityBase(const FObjectInitializer& ObjectInitializer)
+	: Super(ObjectInitializer)
 {
 }
 
@@ -34,8 +34,8 @@ UParticleModuleVelocityBase::UParticleModuleVelocityBase(const class FPostConstr
 	UParticleModuleVelocity implementation.
 -----------------------------------------------------------------------------*/
 
-UParticleModuleVelocity::UParticleModuleVelocity(const class FPostConstructInitializeProperties& PCIP)
-	: Super(PCIP)
+UParticleModuleVelocity::UParticleModuleVelocity(const FObjectInitializer& ObjectInitializer)
+	: Super(ObjectInitializer)
 {
 	bSpawnModule = true;
 }
@@ -125,8 +125,8 @@ void UParticleModuleVelocity::SpawnEx(FParticleEmitterInstance* Owner, int32 Off
 /*-----------------------------------------------------------------------------
 	UParticleModuleVelocityInheritParent implementation.
 -----------------------------------------------------------------------------*/
-UParticleModuleVelocity_Seeded::UParticleModuleVelocity_Seeded(const class FPostConstructInitializeProperties& PCIP)
-	: Super(PCIP)
+UParticleModuleVelocity_Seeded::UParticleModuleVelocity_Seeded(const FObjectInitializer& ObjectInitializer)
+	: Super(ObjectInitializer)
 {
 	bSpawnModule = true;
 	bSupportsRandomSeed = true;
@@ -161,8 +161,8 @@ void UParticleModuleVelocity_Seeded::EmitterLoopingNotify(FParticleEmitterInstan
 /*-----------------------------------------------------------------------------
 	UParticleModuleVelocityInheritParent implementation.
 -----------------------------------------------------------------------------*/
-UParticleModuleVelocityInheritParent::UParticleModuleVelocityInheritParent(const class FPostConstructInitializeProperties& PCIP)
-	: Super(PCIP)
+UParticleModuleVelocityInheritParent::UParticleModuleVelocityInheritParent(const FObjectInitializer& ObjectInitializer)
+	: Super(ObjectInitializer)
 {
 	bSpawnModule = true;
 }
@@ -231,8 +231,8 @@ void UParticleModuleVelocityInheritParent::Spawn(FParticleEmitterInstance* Owner
 /*-----------------------------------------------------------------------------
 	UParticleModuleVelocityOverLifetime implementation.
 -----------------------------------------------------------------------------*/
-UParticleModuleVelocityOverLifetime::UParticleModuleVelocityOverLifetime(const class FPostConstructInitializeProperties& PCIP)
-	: Super(PCIP)
+UParticleModuleVelocityOverLifetime::UParticleModuleVelocityOverLifetime(const FObjectInitializer& ObjectInitializer)
+	: Super(ObjectInitializer)
 {
 	bSpawnModule = true;
 	bUpdateModule = true;
@@ -392,8 +392,8 @@ void UParticleModuleVelocityOverLifetime::Update(FParticleEmitterInstance* Owner
 /*-----------------------------------------------------------------------------
 	UParticleModuleVelocityCone implementation.
 -----------------------------------------------------------------------------*/
-UParticleModuleVelocityCone::UParticleModuleVelocityCone(const class FPostConstructInitializeProperties& PCIP)
-	: Super(PCIP)
+UParticleModuleVelocityCone::UParticleModuleVelocityCone(const FObjectInitializer& ObjectInitializer)
+	: Super(ObjectInitializer)
 {
 	bSpawnModule = true;
 	bSupported3DDrawMode = true;
diff --git a/Engine/Source/Runtime/Engine/Private/Particles/ParticleTrailModules.cpp b/Engine/Source/Runtime/Engine/Private/Particles/ParticleTrailModules.cpp
--- a/Engine/Source/Runtime/Engine/Private/Particles/ParticleTrailModules.cpp
+++ b/Engine/Source/Runtime/Engine/Private/Particles/ParticleTrailModules.cpp
@@ -12,8 +12,8 @@
 #include "Particles/TypeData/ParticleModuleTypeDataRibbon.h"
 #include "Particles/ParticleSystemComponent.h"
 
-UParticleModuleTrailBase::UParticleModuleTrailBase(const class FPostConstructInitializeProperties& PCIP)
-	: Super(PCIP)
+UParticleModuleTrailBase::UParticleModuleTrailBase(const FObjectInitializer& ObjectInitializer)
+	: Super(ObjectInitializer)
 {
 	bSpawnModule = false;
 	bUpdateModule = false;
@@ -27,8 +27,8 @@ UParticleModuleTrailBase::UParticleModuleTrailBase(const class FPostConstructIni
 	UParticleModuleTrailSource implementation.
 -----------------------------------------------------------------------------*/
 
-UParticleModuleTrailSource::UParticleModuleTrailSource(const class FPostConstructInitializeProperties& PCIP)
-	: Super(PCIP)
+UParticleModuleTrailSource::UParticleModuleTrailSource(const FObjectInitializer& ObjectInitializer)
+	: Super(ObjectInitializer)
 {
 	// Structure to hold one-time initialization
 	struct FConstructorStatics
@@ -172,8 +172,8 @@ bool UParticleModuleTrailSource::ResolveSourceOffset(int32 InTrailIdx, FParticle
 /*-----------------------------------------------------------------------------
 	UParticleModuleTypeDataRibbon implementation.
 -----------------------------------------------------------------------------*/
-UParticleModuleTypeDataRibbon::UParticleModuleTypeDataRibbon(const class FPostConstructInitializeProperties& PCIP)
-	: Super(PCIP)
+UParticleModuleTypeDataRibbon::UParticleModuleTypeDataRibbon(const FObjectInitializer& ObjectInitializer)
+	: Super(ObjectInitializer)
 {
 	MaxTessellationBetweenParticles = 25;
 	SheetsPerTrail = 1;
@@ -246,8 +246,8 @@ FParticleEmitterInstance* UParticleModuleTypeDataRibbon::CreateInstance(UParticl
 /*-----------------------------------------------------------------------------
 	UParticleModuleTypeDataAnimTrail implementation.
 -----------------------------------------------------------------------------*/
-UParticleModuleTypeDataAnimTrail::UParticleModuleTypeDataAnimTrail(const class FPostConstructInitializeProperties& PCIP)
-	: Super(PCIP)
+UParticleModuleTypeDataAnimTrail::UParticleModuleTypeDataAnimTrail(const FObjectInitializer& ObjectInitializer)
+	: Super(ObjectInitializer)
 {
 	bDeadTrailsOnDeactivate = true;
 	bEnablePreviousTangentRecalculation = true;
diff --git a/Engine/Source/Runtime/Engine/Private/Pawn.cpp b/Engine/Source/Runtime/Engine/Private/Pawn.cpp
--- a/Engine/Source/Runtime/Engine/Private/Pawn.cpp
+++ b/Engine/Source/Runtime/Engine/Private/Pawn.cpp
@@ -25,8 +25,8 @@ DEFINE_LOG_CATEGORY_STATIC(LogPawn, Warning, All);
 FDetailedTickStats GDetailedPathFindingStats( 30, 10, 1, 20, TEXT("pathfinding") );
 #endif
 
-APawn::APawn(const class FPostConstructInitializeProperties& PCIP)
-	: Super(PCIP)
+APawn::APawn(const FObjectInitializer& ObjectInitializer)
+	: Super(ObjectInitializer)
 {
 	PrimaryActorTick.bCanEverTick = true;
 	PrimaryActorTick.TickGroup = TG_PrePhysics;
diff --git a/Engine/Source/Runtime/Engine/Private/PawnNoiseEmitterComponent.cpp b/Engine/Source/Runtime/Engine/Private/PawnNoiseEmitterComponent.cpp
--- a/Engine/Source/Runtime/Engine/Private/PawnNoiseEmitterComponent.cpp
+++ b/Engine/Source/Runtime/Engine/Private/PawnNoiseEmitterComponent.cpp
@@ -6,8 +6,8 @@
 
 
 
-UPawnNoiseEmitterComponent::UPawnNoiseEmitterComponent(const class FPostConstructInitializeProperties& PCIP)
-	: Super(PCIP)
+UPawnNoiseEmitterComponent::UPawnNoiseEmitterComponent(const FObjectInitializer& ObjectInitializer)
+	: Super(ObjectInitializer)
 {
 	NoiseLifetime = 1.0f;
 	PrimaryComponentTick.bCanEverTick = false;
diff --git a/Engine/Source/Runtime/Engine/Private/PendingNetGame.cpp b/Engine/Source/Runtime/Engine/Private/PendingNetGame.cpp
--- a/Engine/Source/Runtime/Engine/Private/PendingNetGame.cpp
+++ b/Engine/Source/Runtime/Engine/Private/PendingNetGame.cpp
@@ -9,8 +9,8 @@
 #include "Net/NetworkProfiler.h"
 #include "Net/DataChannel.h"
 
-UPendingNetGame::UPendingNetGame( const class FPostConstructInitializeProperties& PCIP, const FURL& InURL )
-	: Super(PCIP)
+UPendingNetGame::UPendingNetGame( const FObjectInitializer& ObjectInitializer, const FURL& InURL )
+	: Super(ObjectInitializer)
 	, NetDriver(NULL)
 	, URL(InURL)
 	, bSuccessfullyConnected(false)
@@ -18,8 +18,8 @@ UPendingNetGame::UPendingNetGame( const class FPostConstructInitializeProperties
 {
 }
 
-UPendingNetGame::UPendingNetGame(const class FPostConstructInitializeProperties& PCIP)
-	: Super(PCIP)
+UPendingNetGame::UPendingNetGame(const FObjectInitializer& ObjectInitializer)
+	: Super(ObjectInitializer)
 {
 }
 
diff --git a/Engine/Source/Runtime/Engine/Private/PhysicsEngine/BodySetup.cpp b/Engine/Source/Runtime/Engine/Private/PhysicsEngine/BodySetup.cpp
--- a/Engine/Source/Runtime/Engine/Private/PhysicsEngine/BodySetup.cpp
+++ b/Engine/Source/Runtime/Engine/Private/PhysicsEngine/BodySetup.cpp
@@ -39,8 +39,8 @@ static TAutoConsoleVariable<float> CVarMaxContactOffset(
 	ECVF_Default);
 
 DEFINE_LOG_CATEGORY(LogPhysics);
-UBodySetup::UBodySetup(const class FPostConstructInitializeProperties& PCIP)
-	: Super(PCIP)
+UBodySetup::UBodySetup(const FObjectInitializer& ObjectInitializer)
+	: Super(ObjectInitializer)
 {
 	bConsiderForBounds = true;
 	bMeshCollideAll = false;
diff --git a/Engine/Source/Runtime/Engine/Private/PhysicsEngine/PhysicalMaterial.cpp b/Engine/Source/Runtime/Engine/Private/PhysicsEngine/PhysicalMaterial.cpp
--- a/Engine/Source/Runtime/Engine/Private/PhysicsEngine/PhysicalMaterial.cpp
+++ b/Engine/Source/Runtime/Engine/Private/PhysicsEngine/PhysicalMaterial.cpp
@@ -13,13 +13,13 @@
 	#include "../Vehicles/PhysXVehicleManager.h"
 #endif // WITH_PHYSX
 
-UDEPRECATED_PhysicalMaterialPropertyBase::UDEPRECATED_PhysicalMaterialPropertyBase(const class FPostConstructInitializeProperties& PCIP)
-	: Super(PCIP)
+UDEPRECATED_PhysicalMaterialPropertyBase::UDEPRECATED_PhysicalMaterialPropertyBase(const FObjectInitializer& ObjectInitializer)
+	: Super(ObjectInitializer)
 {
 }
 
-UPhysicalMaterial::UPhysicalMaterial(const class FPostConstructInitializeProperties& PCIP)
-	: Super(PCIP)
+UPhysicalMaterial::UPhysicalMaterial(const FObjectInitializer& ObjectInitializer)
+	: Super(ObjectInitializer)
 {
 	Friction = 0.7f;
 	Restitution = 0.3f;
diff --git a/Engine/Source/Runtime/Engine/Private/PhysicsEngine/PhysicsAsset.cpp b/Engine/Source/Runtime/Engine/Private/PhysicsEngine/PhysicsAsset.cpp
--- a/Engine/Source/Runtime/Engine/Private/PhysicsEngine/PhysicsAsset.cpp
+++ b/Engine/Source/Runtime/Engine/Private/PhysicsEngine/PhysicsAsset.cpp
@@ -16,8 +16,8 @@
 //////////// UPhysicsAsset ////////////
 ///////////////////////////////////////
 
-UPhysicsAsset::UPhysicsAsset(const class FPostConstructInitializeProperties& PCIP)
-	: Super(PCIP)
+UPhysicsAsset::UPhysicsAsset(const FObjectInitializer& ObjectInitializer)
+	: Super(ObjectInitializer)
 {
 }
 
diff --git a/Engine/Source/Runtime/Engine/Private/PhysicsEngine/PhysicsCollisionHandler.cpp b/Engine/Source/Runtime/Engine/Private/PhysicsEngine/PhysicsCollisionHandler.cpp
--- a/Engine/Source/Runtime/Engine/Private/PhysicsEngine/PhysicsCollisionHandler.cpp
+++ b/Engine/Source/Runtime/Engine/Private/PhysicsEngine/PhysicsCollisionHandler.cpp
@@ -3,8 +3,8 @@
 #include "EnginePrivate.h"
 #include "PhysicsPublic.h"
 
-UPhysicsCollisionHandler::UPhysicsCollisionHandler(const class FPostConstructInitializeProperties& PCIP)
-	: Super(PCIP)
+UPhysicsCollisionHandler::UPhysicsCollisionHandler(const FObjectInitializer& ObjectInitializer)
+	: Super(ObjectInitializer)
 {
 	ImpactReFireDelay = 0.1f;
 }
diff --git a/Engine/Source/Runtime/Engine/Private/PhysicsEngine/PhysicsConstraintActor.cpp b/Engine/Source/Runtime/Engine/Private/PhysicsEngine/PhysicsConstraintActor.cpp
--- a/Engine/Source/Runtime/Engine/Private/PhysicsEngine/PhysicsConstraintActor.cpp
+++ b/Engine/Source/Runtime/Engine/Private/PhysicsEngine/PhysicsConstraintActor.cpp
@@ -3,8 +3,8 @@
 #include "EnginePrivate.h"
 #include "PhysXSupport.h"
 
-APhysicsConstraintActor::APhysicsConstraintActor(const class FPostConstructInitializeProperties& PCIP)
-	: Super(PCIP)
+APhysicsConstraintActor::APhysicsConstraintActor(const FObjectInitializer& ObjectInitializer)
+	: Super(ObjectInitializer)
 {
 	// Structure to hold one-time initialization
 	struct FConstructorStatics
@@ -19,7 +19,7 @@ APhysicsConstraintActor::APhysicsConstraintActor(const class FPostConstructIniti
 	};
 	static FConstructorStatics ConstructorStatics;
 
-	ConstraintComp = PCIP.CreateDefaultSubobject<UPhysicsConstraintComponent>(this, TEXT("MyConstraintComp"));
+	ConstraintComp = ObjectInitializer.CreateDefaultSubobject<UPhysicsConstraintComponent>(this, TEXT("MyConstraintComp"));
 	RootComponent = ConstraintComp;
 	bHidden = true;
 }
diff --git a/Engine/Source/Runtime/Engine/Private/PhysicsEngine/PhysicsConstraintComponent.cpp b/Engine/Source/Runtime/Engine/Private/PhysicsEngine/PhysicsConstraintComponent.cpp
--- a/Engine/Source/Runtime/Engine/Private/PhysicsEngine/PhysicsConstraintComponent.cpp
+++ b/Engine/Source/Runtime/Engine/Private/PhysicsEngine/PhysicsConstraintComponent.cpp
@@ -7,8 +7,8 @@
 
 #define LOCTEXT_NAMESPACE "ConstraintComponent"
 
-UPhysicsConstraintComponent::UPhysicsConstraintComponent(const class FPostConstructInitializeProperties& PCIP)
-	: Super(PCIP)
+UPhysicsConstraintComponent::UPhysicsConstraintComponent(const FObjectInitializer& ObjectInitializer)
+	: Super(ObjectInitializer)
 {
 	bWantsInitializeComponent = true;
 }
diff --git a/Engine/Source/Runtime/Engine/Private/PhysicsEngine/PhysicsConstraintTemplate.cpp b/Engine/Source/Runtime/Engine/Private/PhysicsEngine/PhysicsConstraintTemplate.cpp
--- a/Engine/Source/Runtime/Engine/Private/PhysicsEngine/PhysicsConstraintTemplate.cpp
+++ b/Engine/Source/Runtime/Engine/Private/PhysicsEngine/PhysicsConstraintTemplate.cpp
@@ -3,8 +3,8 @@
 #include "EnginePrivate.h"
 #include "PhysXSupport.h"
 
-UPhysicsConstraintTemplate::UPhysicsConstraintTemplate(const class FPostConstructInitializeProperties& PCIP)
-	: Super(PCIP)
+UPhysicsConstraintTemplate::UPhysicsConstraintTemplate(const FObjectInitializer& ObjectInitializer)
+	: Super(ObjectInitializer)
 {
 	LinearXMotion_DEPRECATED = LCM_Locked;
 	LinearYMotion_DEPRECATED = LCM_Locked;
diff --git a/Engine/Source/Runtime/Engine/Private/PhysicsEngine/PhysicsHandleComponent.cpp b/Engine/Source/Runtime/Engine/Private/PhysicsEngine/PhysicsHandleComponent.cpp
--- a/Engine/Source/Runtime/Engine/Private/PhysicsEngine/PhysicsHandleComponent.cpp
+++ b/Engine/Source/Runtime/Engine/Private/PhysicsEngine/PhysicsHandleComponent.cpp
@@ -6,8 +6,8 @@
 	#include "PhysXSupport.h"
 #endif // WITH_PHYSX
 
-UPhysicsHandleComponent::UPhysicsHandleComponent(const class FPostConstructInitializeProperties& PCIP)
-	: Super(PCIP)
+UPhysicsHandleComponent::UPhysicsHandleComponent(const FObjectInitializer& ObjectInitializer)
+	: Super(ObjectInitializer)
 {
 	bAutoActivate = true;
 	PrimaryComponentTick.bCanEverTick = true;
diff --git a/Engine/Source/Runtime/Engine/Private/PhysicsEngine/PhysicsSettings.cpp b/Engine/Source/Runtime/Engine/Private/PhysicsEngine/PhysicsSettings.cpp
--- a/Engine/Source/Runtime/Engine/Private/PhysicsEngine/PhysicsSettings.cpp
+++ b/Engine/Source/Runtime/Engine/Private/PhysicsEngine/PhysicsSettings.cpp
@@ -2,8 +2,8 @@
 
 #include "EnginePrivate.h"
 
-UPhysicsSettings::UPhysicsSettings(const class FPostConstructInitializeProperties& PCIP)
-	: Super(PCIP)
+UPhysicsSettings::UPhysicsSettings(const FObjectInitializer& ObjectInitializer)
+	: Super(ObjectInitializer)
 	, DefaultGravityZ(-980.f)
 	, DefaultTerminalVelocity(4000.f)
 	, DefaultFluidFriction(0.3f)
diff --git a/Engine/Source/Runtime/Engine/Private/PhysicsEngine/PhysicsThruster.cpp b/Engine/Source/Runtime/Engine/Private/PhysicsEngine/PhysicsThruster.cpp
--- a/Engine/Source/Runtime/Engine/Private/PhysicsEngine/PhysicsThruster.cpp
+++ b/Engine/Source/Runtime/Engine/Private/PhysicsEngine/PhysicsThruster.cpp
@@ -4,8 +4,8 @@
 #include "PhysXSupport.h"
 
 
-UPhysicsThrusterComponent::UPhysicsThrusterComponent(const class FPostConstructInitializeProperties& PCIP)
-	: Super(PCIP)
+UPhysicsThrusterComponent::UPhysicsThrusterComponent(const FObjectInitializer& ObjectInitializer)
+	: Super(ObjectInitializer)
 {
 	PrimaryComponentTick.bCanEverTick = true;
 	PrimaryComponentTick.TickGroup = TG_PrePhysics;
@@ -41,15 +41,15 @@ void UPhysicsThrusterComponent::PostLoad()
 
 //////////////////////////////////////////////////////////////////////////
 
-APhysicsThruster::APhysicsThruster(const class FPostConstructInitializeProperties& PCIP)
-	: Super(PCIP)
+APhysicsThruster::APhysicsThruster(const FObjectInitializer& ObjectInitializer)
+	: Super(ObjectInitializer)
 {
-	ThrusterComponent = PCIP.CreateDefaultSubobject<UPhysicsThrusterComponent>(this, TEXT("Thruster0"));
+	ThrusterComponent = ObjectInitializer.CreateDefaultSubobject<UPhysicsThrusterComponent>(this, TEXT("Thruster0"));
 	RootComponent = ThrusterComponent;
 
 #if WITH_EDITORONLY_DATA
-	ArrowComponent = PCIP.CreateEditorOnlyDefaultSubobject<UArrowComponent>(this, TEXT("ArrowComponent0"));
-	SpriteComponent = PCIP.CreateEditorOnlyDefaultSubobject<UBillboardComponent>(this, TEXT("Sprite"));
+	ArrowComponent = ObjectInitializer.CreateEditorOnlyDefaultSubobject<UArrowComponent>(this, TEXT("ArrowComponent0"));
+	SpriteComponent = ObjectInitializer.CreateEditorOnlyDefaultSubobject<UBillboardComponent>(this, TEXT("Sprite"));
 
 	if (!IsRunningCommandlet())
 	{
diff --git a/Engine/Source/Runtime/Engine/Private/PhysicsEngine/RadialForceComponent.cpp b/Engine/Source/Runtime/Engine/Private/PhysicsEngine/RadialForceComponent.cpp
--- a/Engine/Source/Runtime/Engine/Private/PhysicsEngine/RadialForceComponent.cpp
+++ b/Engine/Source/Runtime/Engine/Private/PhysicsEngine/RadialForceComponent.cpp
@@ -6,8 +6,8 @@
 
 //////////////////////////////////////////////////////////////////////////
 // RADIALFORCECOMPONENT
-URadialForceComponent::URadialForceComponent(const class FPostConstructInitializeProperties& PCIP)
-	: Super(PCIP)
+URadialForceComponent::URadialForceComponent(const FObjectInitializer& ObjectInitializer)
+	: Super(ObjectInitializer)
 {
 	PrimaryComponentTick.bCanEverTick = true;
 	PrimaryComponentTick.TickGroup = TG_PrePhysics;
@@ -172,13 +172,13 @@ void URadialForceComponent::UpdateCollisionObjectQueryParams()
 
 //////////////////////////////////////////////////////////////////////////
 // ARB_RADIALFORCEACTOR
-ARadialForceActor::ARadialForceActor(const class FPostConstructInitializeProperties& PCIP)
-	: Super(PCIP)
+ARadialForceActor::ARadialForceActor(const FObjectInitializer& ObjectInitializer)
+	: Super(ObjectInitializer)
 {
-	ForceComponent = PCIP.CreateDefaultSubobject<URadialForceComponent>(this, TEXT("ForceComponent0"));
+	ForceComponent = ObjectInitializer.CreateDefaultSubobject<URadialForceComponent>(this, TEXT("ForceComponent0"));
 
 #if WITH_EDITOR
-	SpriteComponent = PCIP.CreateEditorOnlyDefaultSubobject<UBillboardComponent>(this, TEXT("Sprite"));
+	SpriteComponent = ObjectInitializer.CreateEditorOnlyDefaultSubobject<UBillboardComponent>(this, TEXT("Sprite"));
 	if (SpriteComponent)
 	{
 		// Structure to hold one-time initialization
diff --git a/Engine/Source/Runtime/Engine/Private/PhysicsEngine2D/BodySetup2D.cpp b/Engine/Source/Runtime/Engine/Private/PhysicsEngine2D/BodySetup2D.cpp
--- a/Engine/Source/Runtime/Engine/Private/PhysicsEngine2D/BodySetup2D.cpp
+++ b/Engine/Source/Runtime/Engine/Private/PhysicsEngine2D/BodySetup2D.cpp
@@ -6,8 +6,8 @@
 //////////////////////////////////////////////////////////////////////////
 // UBodySetup2D
 
-UBodySetup2D::UBodySetup2D(const FPostConstructInitializeProperties& PCIP)
-	: Super(PCIP)
+UBodySetup2D::UBodySetup2D(const FObjectInitializer& ObjectInitializer)
+	: Super(ObjectInitializer)
 {
 }
 
diff --git a/Engine/Source/Runtime/Engine/Private/PhysicsVolume.cpp b/Engine/Source/Runtime/Engine/Private/PhysicsVolume.cpp
--- a/Engine/Source/Runtime/Engine/Private/PhysicsVolume.cpp
+++ b/Engine/Source/Runtime/Engine/Private/PhysicsVolume.cpp
@@ -3,8 +3,8 @@
 #include "EnginePrivate.h"
 #include "GameFramework/PhysicsVolume.h"
 
-APhysicsVolume::APhysicsVolume(const class FPostConstructInitializeProperties& PCIP)
-	: Super(PCIP)
+APhysicsVolume::APhysicsVolume(const FObjectInitializer& ObjectInitializer)
+	: Super(ObjectInitializer)
 {
 	BrushComponent->BodyInstance.bEnableCollision_DEPRECATED = true;
 	static FName CollisionProfileName(TEXT("OverlapAllDynamic"));
diff --git a/Engine/Source/Runtime/Engine/Private/Player.cpp b/Engine/Source/Runtime/Engine/Private/Player.cpp
--- a/Engine/Source/Runtime/Engine/Private/Player.cpp
+++ b/Engine/Source/Runtime/Engine/Private/Player.cpp
@@ -18,8 +18,8 @@
 //////////////////////////////////////////////////////////////////////////
 // UPlayer
 
-UPlayer::UPlayer(const class FPostConstructInitializeProperties& PCIP)
-	: Super(PCIP)
+UPlayer::UPlayer(const FObjectInitializer& ObjectInitializer)
+	: Super(ObjectInitializer)
 {
 }
 
diff --git a/Engine/Source/Runtime/Engine/Private/PlayerCameraManager.cpp b/Engine/Source/Runtime/Engine/Private/PlayerCameraManager.cpp
--- a/Engine/Source/Runtime/Engine/Private/PlayerCameraManager.cpp
+++ b/Engine/Source/Runtime/Engine/Private/PlayerCameraManager.cpp
@@ -14,8 +14,8 @@ DEFINE_LOG_CATEGORY_STATIC(LogPlayerCameraManager, Log, All);
 //////////////////////////////////////////////////////////////////////////
 // APlayerCameraManager
 
-APlayerCameraManager::APlayerCameraManager(const class FPostConstructInitializeProperties& PCIP)
-	: Super(PCIP)
+APlayerCameraManager::APlayerCameraManager(const FObjectInitializer& ObjectInitializer)
+	: Super(ObjectInitializer)
 {
 	static FName NAME_Default(TEXT("Default"));
 
@@ -40,7 +40,7 @@ APlayerCameraManager::APlayerCameraManager(const class FPostConstructInitializeP
 	bFollowHmdOrientation = false;
 
 	// create dummy transform component
-	TransformComponent = PCIP.CreateDefaultSubobject<USceneComponent>(this, TEXT("TransformComponent0"));
+	TransformComponent = ObjectInitializer.CreateDefaultSubobject<USceneComponent>(this, TEXT("TransformComponent0"));
 	RootComponent = TransformComponent;
 }
 
diff --git a/Engine/Source/Runtime/Engine/Private/PlayerController.cpp b/Engine/Source/Runtime/Engine/Private/PlayerController.cpp
--- a/Engine/Source/Runtime/Engine/Private/PlayerController.cpp
+++ b/Engine/Source/Runtime/Engine/Private/PlayerController.cpp
@@ -37,8 +37,8 @@ const float RetryServerCheckSpectatorThrottleTime = 0.25f;
 //////////////////////////////////////////////////////////////////////////
 // APlayerController
 
-APlayerController::APlayerController(const class FPostConstructInitializeProperties& PCIP)
-	: Super(PCIP)
+APlayerController::APlayerController(const FObjectInitializer& ObjectInitializer)
+	: Super(ObjectInitializer)
 	, SlateOperations(FReply::Unhandled())
 {
 	NetPriority = 3.0f;
diff --git a/Engine/Source/Runtime/Engine/Private/PlayerStart.cpp b/Engine/Source/Runtime/Engine/Private/PlayerStart.cpp
--- a/Engine/Source/Runtime/Engine/Private/PlayerStart.cpp
+++ b/Engine/Source/Runtime/Engine/Private/PlayerStart.cpp
@@ -2,13 +2,13 @@
 
 #include "EnginePrivate.h"
 
-APlayerStart::APlayerStart(const class FPostConstructInitializeProperties& PCIP)
-	: Super(PCIP)
+APlayerStart::APlayerStart(const FObjectInitializer& ObjectInitializer)
+	: Super(ObjectInitializer)
 {
 	CapsuleComponent->InitCapsuleSize(40.0f, 92.0f);
 
 #if WITH_EDITORONLY_DATA
-	ArrowComponent = PCIP.CreateEditorOnlyDefaultSubobject<UArrowComponent>(this, TEXT("Arrow"));
+	ArrowComponent = ObjectInitializer.CreateEditorOnlyDefaultSubobject<UArrowComponent>(this, TEXT("Arrow"));
 
 	if (!IsRunningCommandlet())
 	{
diff --git a/Engine/Source/Runtime/Engine/Private/PlayerStartPIE.cpp b/Engine/Source/Runtime/Engine/Private/PlayerStartPIE.cpp
--- a/Engine/Source/Runtime/Engine/Private/PlayerStartPIE.cpp
+++ b/Engine/Source/Runtime/Engine/Private/PlayerStartPIE.cpp
@@ -2,8 +2,8 @@
 
 #include "EnginePrivate.h"
 
-APlayerStartPIE::APlayerStartPIE(const class FPostConstructInitializeProperties& PCIP)
-	: Super(PCIP)
+APlayerStartPIE::APlayerStartPIE(const FObjectInitializer& ObjectInitializer)
+	: Super(ObjectInitializer)
 {
 	
 }
diff --git a/Engine/Source/Runtime/Engine/Private/PlayerState.cpp b/Engine/Source/Runtime/Engine/Private/PlayerState.cpp
--- a/Engine/Source/Runtime/Engine/Private/PlayerState.cpp
+++ b/Engine/Source/Runtime/Engine/Private/PlayerState.cpp
@@ -8,8 +8,8 @@
 #include "Net/UnrealNetwork.h"
 #include "OnlineSubsystemUtils.h"
 
-APlayerState::APlayerState(const class FPostConstructInitializeProperties& PCIP)
-	: Super(PCIP
+APlayerState::APlayerState(const FObjectInitializer& ObjectInitializer)
+	: Super(ObjectInitializer
 		.DoNotCreateDefaultSubobject(TEXT("Sprite")) )
 {
 	SetRemoteRoleForBackwardsCompat(ROLE_SimulatedProxy);
diff --git a/Engine/Source/Runtime/Engine/Private/PointLightComponent.cpp b/Engine/Source/Runtime/Engine/Private/PointLightComponent.cpp
--- a/Engine/Source/Runtime/Engine/Private/PointLightComponent.cpp
+++ b/Engine/Source/Runtime/Engine/Private/PointLightComponent.cpp
@@ -97,8 +97,8 @@ public:
 	{}
 };
 
-UPointLightComponent::UPointLightComponent(const class FPostConstructInitializeProperties& PCIP)
-	: Super(PCIP)
+UPointLightComponent::UPointLightComponent(const FObjectInitializer& ObjectInitializer)
+	: Super(ObjectInitializer)
 {
 #if WITH_EDITORONLY_DATA
 	if (!IsRunningCommandlet())
diff --git a/Engine/Source/Runtime/Engine/Private/PoseableMeshComponent.cpp b/Engine/Source/Runtime/Engine/Private/PoseableMeshComponent.cpp
--- a/Engine/Source/Runtime/Engine/Private/PoseableMeshComponent.cpp
+++ b/Engine/Source/Runtime/Engine/Private/PoseableMeshComponent.cpp
@@ -8,8 +8,8 @@
 #include "AnimTree.h"
 #include "Animation/AnimInstance.h"
 
-UPoseableMeshComponent::UPoseableMeshComponent(const class FPostConstructInitializeProperties& PCIP)
-	: Super(PCIP)
+UPoseableMeshComponent::UPoseableMeshComponent(const FObjectInitializer& ObjectInitializer)
+	: Super(ObjectInitializer)
 {
 
 }
diff --git a/Engine/Source/Runtime/Engine/Private/PostProcessComponent.cpp b/Engine/Source/Runtime/Engine/Private/PostProcessComponent.cpp
--- a/Engine/Source/Runtime/Engine/Private/PostProcessComponent.cpp
+++ b/Engine/Source/Runtime/Engine/Private/PostProcessComponent.cpp
@@ -2,8 +2,8 @@
 
 #include "EnginePrivate.h"
 
-UPostProcessComponent::UPostProcessComponent(const class FPostConstructInitializeProperties& PCIP)
-	: Super(PCIP)
+UPostProcessComponent::UPostProcessComponent(const FObjectInitializer& ObjectInitializer)
+	: Super(ObjectInitializer)
 {
 	bEnabled = true;
 	BlendRadius = 100.0f;
diff --git a/Engine/Source/Runtime/Engine/Private/PostProcessVolume.cpp b/Engine/Source/Runtime/Engine/Private/PostProcessVolume.cpp
--- a/Engine/Source/Runtime/Engine/Private/PostProcessVolume.cpp
+++ b/Engine/Source/Runtime/Engine/Private/PostProcessVolume.cpp
@@ -3,8 +3,8 @@
 #include "EnginePrivate.h"
 #include "Engine/PostProcessVolume.h"
 
-APostProcessVolume::APostProcessVolume(const class FPostConstructInitializeProperties& PCIP)
-	: Super(PCIP)
+APostProcessVolume::APostProcessVolume(const FObjectInitializer& ObjectInitializer)
+	: Super(ObjectInitializer)
 {
 	BrushComponent->BodyInstance.bEnableCollision_DEPRECATED = false;
 	BrushComponent->SetCollisionProfileName(UCollisionProfile::NoCollision_ProfileName);
diff --git a/Engine/Source/Runtime/Engine/Private/PrecomputedVisibilityVolume.cpp b/Engine/Source/Runtime/Engine/Private/PrecomputedVisibilityVolume.cpp
--- a/Engine/Source/Runtime/Engine/Private/PrecomputedVisibilityVolume.cpp
+++ b/Engine/Source/Runtime/Engine/Private/PrecomputedVisibilityVolume.cpp
@@ -3,8 +3,8 @@
 #include "EnginePrivate.h"
 #include "Lightmass/PrecomputedVisibilityVolume.h"
 
-APrecomputedVisibilityVolume::APrecomputedVisibilityVolume(const class FPostConstructInitializeProperties& PCIP)
-	: Super(PCIP)
+APrecomputedVisibilityVolume::APrecomputedVisibilityVolume(const FObjectInitializer& ObjectInitializer)
+	: Super(ObjectInitializer)
 {
 	BrushComponent->BodyInstance.bEnableCollision_DEPRECATED = false;
 	BrushComponent->SetCollisionProfileName(UCollisionProfile::NoCollision_ProfileName);
diff --git a/Engine/Source/Runtime/Engine/Private/PrecomputedVisiblityOverrideVolume.cpp b/Engine/Source/Runtime/Engine/Private/PrecomputedVisiblityOverrideVolume.cpp
--- a/Engine/Source/Runtime/Engine/Private/PrecomputedVisiblityOverrideVolume.cpp
+++ b/Engine/Source/Runtime/Engine/Private/PrecomputedVisiblityOverrideVolume.cpp
@@ -3,8 +3,8 @@
 #include "EnginePrivate.h"
 #include "Lightmass/PrecomputedVisibilityOverrideVolume.h"
 
-APrecomputedVisibilityOverrideVolume::APrecomputedVisibilityOverrideVolume(const class FPostConstructInitializeProperties& PCIP)
-	: Super(PCIP)
+APrecomputedVisibilityOverrideVolume::APrecomputedVisibilityOverrideVolume(const FObjectInitializer& ObjectInitializer)
+	: Super(ObjectInitializer)
 {
 	BrushComponent->BodyInstance.bEnableCollision_DEPRECATED = false;
 	BrushComponent->SetCollisionProfileName(UCollisionProfile::NoCollision_ProfileName);
diff --git a/Engine/Source/Runtime/Engine/Private/PreviewScene.cpp b/Engine/Source/Runtime/Engine/Private/PreviewScene.cpp
--- a/Engine/Source/Runtime/Engine/Private/PreviewScene.cpp
+++ b/Engine/Source/Runtime/Engine/Private/PreviewScene.cpp
@@ -12,7 +12,7 @@ FPreviewScene::FPreviewScene(FPreviewScene::ConstructionValues CVS)
 	: PreviewWorld(NULL)
 	, bForceAllUsedMipsResident(CVS.bForceMipsResident)
 {
-	PreviewWorld = new UWorld(FPostConstructInitializeProperties(),FURL(NULL));
+	PreviewWorld = new UWorld(FObjectInitializer(),FURL(NULL));
 	PreviewWorld->WorldType = EWorldType::Preview;
 	if (!CVS.bTransactional)
 	{
diff --git a/Engine/Source/Runtime/Engine/Private/PrimitiveComponent.cpp b/Engine/Source/Runtime/Engine/Private/PrimitiveComponent.cpp
--- a/Engine/Source/Runtime/Engine/Private/PrimitiveComponent.cpp
+++ b/Engine/Source/Runtime/Engine/Private/PrimitiveComponent.cpp
@@ -59,8 +59,8 @@ int32 UPrimitiveComponent::CurrentTag = 2147483647 / 4;
 // 0 is reserved to mean invalid
 uint64 UPrimitiveComponent::NextComponentId = 1;
 
-UPrimitiveComponent::UPrimitiveComponent(const class FPostConstructInitializeProperties& PCIP)
-	: Super(PCIP)
+UPrimitiveComponent::UPrimitiveComponent(const FObjectInitializer& ObjectInitializer)
+	: Super(ObjectInitializer)
 {
 	PostPhysicsComponentTick.bCanEverTick = false;
 	PostPhysicsComponentTick.bStartWithTickEnabled = true;
