Rename FPostConstructInitializeProperties to something simpler

diff --git a/Engine/Source/Runtime/Engine/Private/ProjectileMovementComponent.cpp b/Engine/Source/Runtime/Engine/Private/ProjectileMovementComponent.cpp
--- a/Engine/Source/Runtime/Engine/Private/ProjectileMovementComponent.cpp
+++ b/Engine/Source/Runtime/Engine/Private/ProjectileMovementComponent.cpp
@@ -7,8 +7,8 @@ DEFINE_LOG_CATEGORY_STATIC(LogProjectileMovement, Log, All);
 
 const float UProjectileMovementComponent::MIN_TICK_TIME = 0.0002f;
 
-UProjectileMovementComponent::UProjectileMovementComponent(const class FPostConstructInitializeProperties& PCIP)
-	: Super(PCIP)
+UProjectileMovementComponent::UProjectileMovementComponent(const FObjectInitializer& ObjectInitializer)
+	: Super(ObjectInitializer)
 {
 	bUpdateOnlyIfRendered = false;
 	bInitialVelocityInLocalSpace = true;
diff --git a/Engine/Source/Runtime/Engine/Private/ReflectionCaptureComponent.cpp b/Engine/Source/Runtime/Engine/Private/ReflectionCaptureComponent.cpp
--- a/Engine/Source/Runtime/Engine/Private/ReflectionCaptureComponent.cpp
+++ b/Engine/Source/Runtime/Engine/Private/ReflectionCaptureComponent.cpp
@@ -43,13 +43,13 @@ void UWorld::UpdateAllReflectionCaptures()
 	UReflectionCaptureComponent::UpdateReflectionCaptureContents(this);
 }
 
-AReflectionCapture::AReflectionCapture(const class FPostConstructInitializeProperties& PCIP)
-	: Super(PCIP)
+AReflectionCapture::AReflectionCapture(const FObjectInitializer& ObjectInitializer)
+	: Super(ObjectInitializer)
 {
-	CaptureComponent = PCIP.CreateAbstractDefaultSubobject<UReflectionCaptureComponent>(this, TEXT("NewReflectionComponent"));
+	CaptureComponent = ObjectInitializer.CreateAbstractDefaultSubobject<UReflectionCaptureComponent>(this, TEXT("NewReflectionComponent"));
 
 #if WITH_EDITORONLY_DATA
-	SpriteComponent = PCIP.CreateEditorOnlyDefaultSubobject<UBillboardComponent>(this, TEXT("Sprite"));
+	SpriteComponent = ObjectInitializer.CreateEditorOnlyDefaultSubobject<UBillboardComponent>(this, TEXT("Sprite"));
 	if (!IsRunningCommandlet() && (SpriteComponent != nullptr))
 	{
 		// Structure to hold one-time initialization
@@ -85,8 +85,8 @@ void AReflectionCapture::PostEditMove(bool bFinished)
 }
 #endif // WITH_EDITOR
 
-ASphereReflectionCapture::ASphereReflectionCapture(const class FPostConstructInitializeProperties& PCIP)
-	: Super(PCIP.SetDefaultSubobjectClass<USphereReflectionCaptureComponent>(TEXT("NewReflectionComponent")))
+ASphereReflectionCapture::ASphereReflectionCapture(const FObjectInitializer& ObjectInitializer)
+	: Super(ObjectInitializer.SetDefaultSubobjectClass<USphereReflectionCaptureComponent>(TEXT("NewReflectionComponent")))
 {
 	USphereReflectionCaptureComponent* SphereComponent = CastChecked<USphereReflectionCaptureComponent>(CaptureComponent);
 	RootComponent = SphereComponent;
@@ -97,7 +97,7 @@ ASphereReflectionCapture::ASphereReflectionCapture(const class FPostConstructIni
 	}
 #endif	//WITH_EDITORONLY_DATA
 
-	TSubobjectPtr<UDrawSphereComponent> DrawInfluenceRadius = PCIP.CreateDefaultSubobject<UDrawSphereComponent>(this, TEXT("DrawRadius0"));
+	TSubobjectPtr<UDrawSphereComponent> DrawInfluenceRadius = ObjectInitializer.CreateDefaultSubobject<UDrawSphereComponent>(this, TEXT("DrawRadius0"));
 	DrawInfluenceRadius->AttachParent = CaptureComponent;
 	DrawInfluenceRadius->bDrawOnlyIfSelected = true;
 	DrawInfluenceRadius->bUseEditorCompositing = true;
@@ -105,7 +105,7 @@ ASphereReflectionCapture::ASphereReflectionCapture(const class FPostConstructIni
 	DrawInfluenceRadius->SetCollisionProfileName(UCollisionProfile::NoCollision_ProfileName);
 	SphereComponent->PreviewInfluenceRadius = DrawInfluenceRadius;
 
-	DrawCaptureRadius = PCIP.CreateDefaultSubobject<UDrawSphereComponent>(this, TEXT("DrawRadius1"));
+	DrawCaptureRadius = ObjectInitializer.CreateDefaultSubobject<UDrawSphereComponent>(this, TEXT("DrawRadius1"));
 	DrawCaptureRadius->AttachParent = CaptureComponent;
 	DrawCaptureRadius->bDrawOnlyIfSelected = true;
 	DrawCaptureRadius->bUseEditorCompositing = true;
@@ -136,8 +136,8 @@ void APlaneReflectionCapture::EditorApplyScale(const FVector& DeltaScale, const
 }
 #endif
 
-ABoxReflectionCapture::ABoxReflectionCapture(const class FPostConstructInitializeProperties& PCIP)
-	: Super(PCIP.SetDefaultSubobjectClass<UBoxReflectionCaptureComponent>(TEXT("NewReflectionComponent")))
+ABoxReflectionCapture::ABoxReflectionCapture(const FObjectInitializer& ObjectInitializer)
+	: Super(ObjectInitializer.SetDefaultSubobjectClass<UBoxReflectionCaptureComponent>(TEXT("NewReflectionComponent")))
 {
 	UBoxReflectionCaptureComponent* BoxComponent = CastChecked<UBoxReflectionCaptureComponent>(CaptureComponent);
 	BoxComponent->RelativeScale3D = FVector(1000, 1000, 400);
@@ -148,7 +148,7 @@ ABoxReflectionCapture::ABoxReflectionCapture(const class FPostConstructInitializ
 		SpriteComponent->AttachParent = BoxComponent;
 	}
 #endif	//WITH_EDITORONLY_DATA
-	TSubobjectPtr<UBoxComponent> DrawInfluenceBox = PCIP.CreateDefaultSubobject<UBoxComponent>(this, TEXT("DrawBox0"));
+	TSubobjectPtr<UBoxComponent> DrawInfluenceBox = ObjectInitializer.CreateDefaultSubobject<UBoxComponent>(this, TEXT("DrawBox0"));
 	DrawInfluenceBox->AttachParent = CaptureComponent;
 	DrawInfluenceBox->bDrawOnlyIfSelected = true;
 	DrawInfluenceBox->bUseEditorCompositing = true;
@@ -157,7 +157,7 @@ ABoxReflectionCapture::ABoxReflectionCapture(const class FPostConstructInitializ
 	DrawInfluenceBox->InitBoxExtent(FVector(1, 1, 1));
 	BoxComponent->PreviewInfluenceBox = DrawInfluenceBox;
 
-	TSubobjectPtr<UBoxComponent> DrawCaptureBox = PCIP.CreateDefaultSubobject<UBoxComponent>(this, TEXT("DrawBox1"));
+	TSubobjectPtr<UBoxComponent> DrawCaptureBox = ObjectInitializer.CreateDefaultSubobject<UBoxComponent>(this, TEXT("DrawBox1"));
 	DrawCaptureBox->AttachParent = CaptureComponent;
 	DrawCaptureBox->bDrawOnlyIfSelected = true;
 	DrawCaptureBox->bUseEditorCompositing = true;
@@ -168,8 +168,8 @@ ABoxReflectionCapture::ABoxReflectionCapture(const class FPostConstructInitializ
 	BoxComponent->PreviewCaptureBox = DrawCaptureBox;
 }
 
-APlaneReflectionCapture::APlaneReflectionCapture(const class FPostConstructInitializeProperties& PCIP)
-	: Super(PCIP.SetDefaultSubobjectClass<UPlaneReflectionCaptureComponent>(TEXT("NewReflectionComponent")))
+APlaneReflectionCapture::APlaneReflectionCapture(const FObjectInitializer& ObjectInitializer)
+	: Super(ObjectInitializer.SetDefaultSubobjectClass<UPlaneReflectionCaptureComponent>(TEXT("NewReflectionComponent")))
 {
 	UPlaneReflectionCaptureComponent* PlaneComponent = CastChecked<UPlaneReflectionCaptureComponent>(CaptureComponent);
 	PlaneComponent->RelativeScale3D = FVector(1, 1000, 1000);
@@ -180,7 +180,7 @@ APlaneReflectionCapture::APlaneReflectionCapture(const class FPostConstructIniti
 		SpriteComponent->AttachParent = PlaneComponent;
 	}
 #endif	//#if WITH_EDITORONLY_DATA
-	TSubobjectPtr<UDrawSphereComponent> DrawInfluenceRadius = PCIP.CreateDefaultSubobject<UDrawSphereComponent>(this, TEXT("DrawRadius0"));
+	TSubobjectPtr<UDrawSphereComponent> DrawInfluenceRadius = ObjectInitializer.CreateDefaultSubobject<UDrawSphereComponent>(this, TEXT("DrawRadius0"));
 	DrawInfluenceRadius->AttachParent = CaptureComponent;
 	DrawInfluenceRadius->bDrawOnlyIfSelected = true;
 	DrawInfluenceRadius->bAbsoluteScale = true;
@@ -189,7 +189,7 @@ APlaneReflectionCapture::APlaneReflectionCapture(const class FPostConstructIniti
 	DrawInfluenceRadius->SetCollisionProfileName(UCollisionProfile::NoCollision_ProfileName);
 	PlaneComponent->PreviewInfluenceRadius = DrawInfluenceRadius;
 
-	TSubobjectPtr<UBoxComponent> DrawCaptureBox = PCIP.CreateDefaultSubobject<UBoxComponent>(this, TEXT("DrawBox1"));
+	TSubobjectPtr<UBoxComponent> DrawCaptureBox = ObjectInitializer.CreateDefaultSubobject<UBoxComponent>(this, TEXT("DrawBox1"));
 	DrawCaptureBox->AttachParent = CaptureComponent;
 	DrawCaptureBox->bDrawOnlyIfSelected = true;
 	DrawCaptureBox->bUseEditorCompositing = true;
@@ -722,8 +722,8 @@ private:
 TArray<UReflectionCaptureComponent*> UReflectionCaptureComponent::ReflectionCapturesToUpdate;
 TArray<UReflectionCaptureComponent*> UReflectionCaptureComponent::ReflectionCapturesToUpdateForLoad;
 
-UReflectionCaptureComponent::UReflectionCaptureComponent(const class FPostConstructInitializeProperties& PCIP)
-	: Super(PCIP)
+UReflectionCaptureComponent::UReflectionCaptureComponent(const FObjectInitializer& ObjectInitializer)
+	: Super(ObjectInitializer)
 {
 	Brightness = 1;
 	// Shouldn't be able to change reflection captures at runtime
@@ -1404,8 +1404,8 @@ void UReflectionCaptureComponent::PreFeatureLevelChange(ERHIFeatureLevel::Type P
 }
 #endif // WITH_EDITOR
 
-USphereReflectionCaptureComponent::USphereReflectionCaptureComponent(const class FPostConstructInitializeProperties& PCIP)
-	: Super(PCIP)
+USphereReflectionCaptureComponent::USphereReflectionCaptureComponent(const FObjectInitializer& ObjectInitializer)
+	: Super(ObjectInitializer)
 {
 	InfluenceRadius = 3000;
 }
@@ -1432,8 +1432,8 @@ void USphereReflectionCaptureComponent::PostEditChangeProperty(FPropertyChangedE
 }
 #endif // WITH_EDITOR
 
-UBoxReflectionCaptureComponent::UBoxReflectionCaptureComponent(const class FPostConstructInitializeProperties& PCIP)
-	: Super(PCIP)
+UBoxReflectionCaptureComponent::UBoxReflectionCaptureComponent(const FObjectInitializer& ObjectInitializer)
+	: Super(ObjectInitializer)
 {
 	BoxTransitionDistance = 100;
 }
@@ -1464,8 +1464,8 @@ void UBoxReflectionCaptureComponent::PostEditChangeProperty(FPropertyChangedEven
 }
 #endif // WITH_EDITOR
 
-UPlaneReflectionCaptureComponent::UPlaneReflectionCaptureComponent(const class FPostConstructInitializeProperties& PCIP)
-	: Super(PCIP)
+UPlaneReflectionCaptureComponent::UPlaneReflectionCaptureComponent(const FObjectInitializer& ObjectInitializer)
+	: Super(ObjectInitializer)
 {
 	InfluenceRadiusScale = 2;
 }
diff --git a/Engine/Source/Runtime/Engine/Private/Rendering/SubsurfaceProfile.cpp b/Engine/Source/Runtime/Engine/Private/Rendering/SubsurfaceProfile.cpp
--- a/Engine/Source/Runtime/Engine/Private/Rendering/SubsurfaceProfile.cpp
+++ b/Engine/Source/Runtime/Engine/Private/Rendering/SubsurfaceProfile.cpp
@@ -289,8 +289,8 @@ ENGINE_API const IPooledRenderTarget* GetSubsufaceProfileTexture_RT(FRHICommandL
 
 // ------------------------------------------------------
 
-USubsurfaceProfile::USubsurfaceProfile(const class FPostConstructInitializeProperties& PCIP)
-	: Super(PCIP)
+USubsurfaceProfile::USubsurfaceProfile(const FObjectInitializer& ObjectInitializer)
+	: Super(ObjectInitializer)
 {
 }
 
diff --git a/Engine/Source/Runtime/Engine/Private/RenderingSettings.cpp b/Engine/Source/Runtime/Engine/Private/RenderingSettings.cpp
--- a/Engine/Source/Runtime/Engine/Private/RenderingSettings.cpp
+++ b/Engine/Source/Runtime/Engine/Private/RenderingSettings.cpp
@@ -3,8 +3,8 @@
 #include "EnginePrivate.h"
 #include "Engine/RendererSettings.h"
 
-URendererSettings::URendererSettings(const class FPostConstructInitializeProperties& PCIP)
-	: Super(PCIP)
+URendererSettings::URendererSettings(const FObjectInitializer& ObjectInitializer)
+	: Super(ObjectInitializer)
 {
 }
 
diff --git a/Engine/Source/Runtime/Engine/Private/ReverbEffect.cpp b/Engine/Source/Runtime/Engine/Private/ReverbEffect.cpp
--- a/Engine/Source/Runtime/Engine/Private/ReverbEffect.cpp
+++ b/Engine/Source/Runtime/Engine/Private/ReverbEffect.cpp
@@ -3,8 +3,8 @@
 #include "EnginePrivate.h"
 #include "Sound/ReverbEffect.h"
 
-UReverbEffect::UReverbEffect(const class FPostConstructInitializeProperties& PCIP)
-	: Super(PCIP)
+UReverbEffect::UReverbEffect(const FObjectInitializer& ObjectInitializer)
+	: Super(ObjectInitializer)
 {
 	Density = 1.0f;
 	Diffusion = 1.0f;
diff --git a/Engine/Source/Runtime/Engine/Private/RigidBodyBase.cpp b/Engine/Source/Runtime/Engine/Private/RigidBodyBase.cpp
--- a/Engine/Source/Runtime/Engine/Private/RigidBodyBase.cpp
+++ b/Engine/Source/Runtime/Engine/Private/RigidBodyBase.cpp
@@ -2,8 +2,8 @@
 
 #include "EnginePrivate.h"
 
-ARigidBodyBase::ARigidBodyBase(const class FPostConstructInitializeProperties& PCIP)
-	: Super(PCIP)
+ARigidBodyBase::ARigidBodyBase(const FObjectInitializer& ObjectInitializer)
+	: Super(ObjectInitializer)
 {
 }
 
diff --git a/Engine/Source/Runtime/Engine/Private/RotatingMovementComponent.cpp b/Engine/Source/Runtime/Engine/Private/RotatingMovementComponent.cpp
--- a/Engine/Source/Runtime/Engine/Private/RotatingMovementComponent.cpp
+++ b/Engine/Source/Runtime/Engine/Private/RotatingMovementComponent.cpp
@@ -3,8 +3,8 @@
 #include "EnginePrivate.h"
 #include "GameFramework/RotatingMovementComponent.h"
 
-URotatingMovementComponent::URotatingMovementComponent(const class FPostConstructInitializeProperties& PCIP)
-	: Super(PCIP)
+URotatingMovementComponent::URotatingMovementComponent(const FObjectInitializer& ObjectInitializer)
+	: Super(ObjectInitializer)
 {
 	RotationRate.Yaw = 180.0f;
 	bRotationInLocalSpace = true;
diff --git a/Engine/Source/Runtime/Engine/Private/SCS_Node.cpp b/Engine/Source/Runtime/Engine/Private/SCS_Node.cpp
--- a/Engine/Source/Runtime/Engine/Private/SCS_Node.cpp
+++ b/Engine/Source/Runtime/Engine/Private/SCS_Node.cpp
@@ -11,8 +11,8 @@
 //////////////////////////////////////////////////////////////////////////
 // USCS_Node
 
-USCS_Node::USCS_Node(const class FPostConstructInitializeProperties& PCIP)
-	: Super(PCIP)
+USCS_Node::USCS_Node(const FObjectInitializer& ObjectInitializer)
+	: Super(ObjectInitializer)
 {
 	bIsFalseRoot_DEPRECATED = false;
 	bIsNative_DEPRECATED = false;
diff --git a/Engine/Source/Runtime/Engine/Private/SaveGameSummary.cpp b/Engine/Source/Runtime/Engine/Private/SaveGameSummary.cpp
--- a/Engine/Source/Runtime/Engine/Private/SaveGameSummary.cpp
+++ b/Engine/Source/Runtime/Engine/Private/SaveGameSummary.cpp
@@ -2,7 +2,7 @@
 
 #include "EnginePrivate.h"
 
-UDEPRECATED_SaveGameSummary::UDEPRECATED_SaveGameSummary(const class FPostConstructInitializeProperties& PCIP)
-	: Super(PCIP)
+UDEPRECATED_SaveGameSummary::UDEPRECATED_SaveGameSummary(const FObjectInitializer& ObjectInitializer)
+	: Super(ObjectInitializer)
 {
 }
diff --git a/Engine/Source/Runtime/Engine/Private/Scene.cpp b/Engine/Source/Runtime/Engine/Private/Scene.cpp
--- a/Engine/Source/Runtime/Engine/Private/Scene.cpp
+++ b/Engine/Source/Runtime/Engine/Private/Scene.cpp
@@ -2,7 +2,7 @@
 
 #include "EnginePrivate.h"
 
-UScene::UScene(const class FPostConstructInitializeProperties& PCIP)
-	: Super(PCIP)
+UScene::UScene(const FObjectInitializer& ObjectInitializer)
+	: Super(ObjectInitializer)
 {
 }
\ No newline at end of file
diff --git a/Engine/Source/Runtime/Engine/Private/SceneCaptureComponent.cpp b/Engine/Source/Runtime/Engine/Private/SceneCaptureComponent.cpp
--- a/Engine/Source/Runtime/Engine/Private/SceneCaptureComponent.cpp
+++ b/Engine/Source/Runtime/Engine/Private/SceneCaptureComponent.cpp
@@ -7,10 +7,10 @@
 #include "EnginePrivate.h"
 #include "../../Renderer/Private/ScenePrivate.h"
 
-ASceneCapture::ASceneCapture(const class FPostConstructInitializeProperties& PCIP)
-	: Super(PCIP)
+ASceneCapture::ASceneCapture(const FObjectInitializer& ObjectInitializer)
+	: Super(ObjectInitializer)
 {
-	MeshComp = PCIP.CreateDefaultSubobject<UStaticMeshComponent>(this, TEXT("CamMesh0"));
+	MeshComp = ObjectInitializer.CreateDefaultSubobject<UStaticMeshComponent>(this, TEXT("CamMesh0"));
 	MeshComp->BodyInstance.bEnableCollision_DEPRECATED = false;
 
 	MeshComp->SetCollisionProfileName(UCollisionProfile::NoCollision_ProfileName);
@@ -22,15 +22,15 @@ ASceneCapture::ASceneCapture(const class FPostConstructInitializeProperties& PCI
 }
 // -----------------------------------------------
 
-ASceneCapture2D::ASceneCapture2D(const class FPostConstructInitializeProperties& PCIP)
-	: Super(PCIP)
+ASceneCapture2D::ASceneCapture2D(const FObjectInitializer& ObjectInitializer)
+	: Super(ObjectInitializer)
 {
-	DrawFrustum = PCIP.CreateDefaultSubobject<UDrawFrustumComponent>(this, TEXT("DrawFrust0"));
+	DrawFrustum = ObjectInitializer.CreateDefaultSubobject<UDrawFrustumComponent>(this, TEXT("DrawFrust0"));
 	DrawFrustum->AlwaysLoadOnClient = false;
 	DrawFrustum->AlwaysLoadOnServer = false;
 	DrawFrustum->AttachParent = MeshComp;
 
-	CaptureComponent2D = PCIP.CreateDefaultSubobject<USceneCaptureComponent2D>(this, TEXT("NewSceneCaptureComponent2D"));
+	CaptureComponent2D = ObjectInitializer.CreateDefaultSubobject<USceneCaptureComponent2D>(this, TEXT("NewSceneCaptureComponent2D"));
 	CaptureComponent2D->AttachParent = MeshComp;
 }
 
@@ -78,15 +78,15 @@ void ASceneCapture2D::PostActorCreated()
 }
 // -----------------------------------------------
 
-ASceneCaptureCube::ASceneCaptureCube(const class FPostConstructInitializeProperties& PCIP)
-	: Super(PCIP)
+ASceneCaptureCube::ASceneCaptureCube(const FObjectInitializer& ObjectInitializer)
+	: Super(ObjectInitializer)
 {
-	DrawFrustum = PCIP.CreateDefaultSubobject<UDrawFrustumComponent>(this, TEXT("DrawFrust0"));
+	DrawFrustum = ObjectInitializer.CreateDefaultSubobject<UDrawFrustumComponent>(this, TEXT("DrawFrust0"));
 	DrawFrustum->AlwaysLoadOnClient = false;
 	DrawFrustum->AlwaysLoadOnServer = false;
 	DrawFrustum->AttachParent = MeshComp;
 
-	CaptureComponentCube = PCIP.CreateDefaultSubobject<USceneCaptureComponentCube>(this, TEXT("NewSceneCaptureComponentCube"));
+	CaptureComponentCube = ObjectInitializer.CreateDefaultSubobject<USceneCaptureComponentCube>(this, TEXT("NewSceneCaptureComponentCube"));
 	CaptureComponentCube->AttachParent = MeshComp;
 }
 
@@ -143,8 +143,8 @@ void ASceneCaptureCube::PostEditMove(bool bFinished)
 }
 #endif
 // -----------------------------------------------
-USceneCaptureComponent::USceneCaptureComponent(const class FPostConstructInitializeProperties& PCIP)
-	: Super(PCIP)
+USceneCaptureComponent::USceneCaptureComponent(const FObjectInitializer& ObjectInitializer)
+	: Super(ObjectInitializer)
 {
 	bCaptureEveryFrame = true;
 	MaxViewDistanceOverride = -1;
@@ -175,8 +175,8 @@ void USceneCaptureComponent::HideActorComponents(AActor* InActor)
 // -----------------------------------------------
 
 
-USceneCaptureComponent2D::USceneCaptureComponent2D(const class FPostConstructInitializeProperties& PCIP)
-	: Super(PCIP)
+USceneCaptureComponent2D::USceneCaptureComponent2D(const FObjectInitializer& ObjectInitializer)
+	: Super(ObjectInitializer)
 {
 	FOVAngle = 90.0f;
 	bAutoActivate = true;
@@ -239,8 +239,8 @@ void USceneCaptureComponent2D::PostEditChangeProperty(FPropertyChangedEvent& Pro
 // -----------------------------------------------
 
 
-USceneCaptureComponentCube::USceneCaptureComponentCube(const class FPostConstructInitializeProperties& PCIP)
-	: Super(PCIP)
+USceneCaptureComponentCube::USceneCaptureComponentCube(const FObjectInitializer& ObjectInitializer)
+	: Super(ObjectInitializer)
 {
 	bAutoActivate = true;
 	PrimaryComponentTick.bCanEverTick = true;
diff --git a/Engine/Source/Runtime/Engine/Private/SceneComponent.cpp b/Engine/Source/Runtime/Engine/Private/SceneComponent.cpp
--- a/Engine/Source/Runtime/Engine/Private/SceneComponent.cpp
+++ b/Engine/Source/Runtime/Engine/Private/SceneComponent.cpp
@@ -22,8 +22,8 @@ FOverlapInfo::FOverlapInfo(UPrimitiveComponent* InComponent, int32 InBodyIndex)
 	OverlapInfo.Item = InBodyIndex;
 }
 
-USceneComponent::USceneComponent(const class FPostConstructInitializeProperties& PCIP)
-	: Super(PCIP)
+USceneComponent::USceneComponent(const FObjectInitializer& ObjectInitializer)
+	: Super(ObjectInitializer)
 {
 	Mobility = EComponentMobility::Movable;
 	RelativeScale_DEPRECATED = 1.0f;
diff --git a/Engine/Source/Runtime/Engine/Private/SceneView.cpp b/Engine/Source/Runtime/Engine/Private/SceneView.cpp
--- a/Engine/Source/Runtime/Engine/Private/SceneView.cpp
+++ b/Engine/Source/Runtime/Engine/Private/SceneView.cpp
@@ -818,7 +818,7 @@ void FSceneView::OverridePostProcessSettings(const FPostProcessSettings& Src, fl
 }
 
 /** Dummy class needed to support Cast<IBlendableInterface>(Object) */
-UBlendableInterface::UBlendableInterface(const class FPostConstructInitializeProperties& PCIP) : Super(PCIP) 
+UBlendableInterface::UBlendableInterface(const FObjectInitializer& ObjectInitializer) : Super(ObjectInitializer) 
 {
 }
 
diff --git a/Engine/Source/Runtime/Engine/Private/ScriptPlatformInterface.cpp b/Engine/Source/Runtime/Engine/Private/ScriptPlatformInterface.cpp
--- a/Engine/Source/Runtime/Engine/Private/ScriptPlatformInterface.cpp
+++ b/Engine/Source/Runtime/Engine/Private/ScriptPlatformInterface.cpp
@@ -11,8 +11,8 @@ DEFINE_LOG_CATEGORY_STATIC(LogScriptPlatformInterface, Log, All);
 /*******************************************
  * Platform Interface Base                 *
  *******************************************/
-UPlatformInterfaceBase::UPlatformInterfaceBase(const class FPostConstructInitializeProperties& PCIP)
-	: Super(PCIP)
+UPlatformInterfaceBase::UPlatformInterfaceBase(const FObjectInitializer& ObjectInitializer)
+	: Super(ObjectInitializer)
 {
 }
 
@@ -166,8 +166,8 @@ IMPLEMENT_PLATFORM_INTERFACE_SINGLETON(UTwitterIntegrationBase,TwitterIntegratio
  * Cloud Storage                           *
  *******************************************/
 
-UCloudStorageBase::UCloudStorageBase(const class FPostConstructInitializeProperties& PCIP)
- : Super(PCIP)
+UCloudStorageBase::UCloudStorageBase(const FObjectInitializer& ObjectInitializer)
+ : Super(ObjectInitializer)
 {
 }
 
@@ -489,8 +489,8 @@ bool UCloudStorageBase::ResolveConflictWithVersionIndex(int32 /*Index*/)
  *******************************************/
 
 
-UMicroTransactionBase::UMicroTransactionBase(const class FPostConstructInitializeProperties& PCIP)
-	: Super(PCIP)
+UMicroTransactionBase::UMicroTransactionBase(const FObjectInitializer& ObjectInitializer)
+	: Super(ObjectInitializer)
 {
 }
 
@@ -521,8 +521,8 @@ bool UMicroTransactionBase::BeginPurchase(int Index)
 /*******************************************
  * Twitter Integration
  *******************************************/
-UTwitterIntegrationBase::UTwitterIntegrationBase(const class FPostConstructInitializeProperties& PCIP)
-	: Super(PCIP)
+UTwitterIntegrationBase::UTwitterIntegrationBase(const FObjectInitializer& ObjectInitializer)
+	: Super(ObjectInitializer)
 {
 }
 
@@ -571,8 +571,8 @@ bool UTwitterIntegrationBase::TwitterRequest(const FString& URL, const TArray<FS
 /*******************************************
  * Platform Interface Web Response         *
  *******************************************/
-UPlatformInterfaceWebResponse::UPlatformInterfaceWebResponse(const class FPostConstructInitializeProperties& PCIP)
-	: Super(PCIP)
+UPlatformInterfaceWebResponse::UPlatformInterfaceWebResponse(const FObjectInitializer& ObjectInitializer)
+	: Super(ObjectInitializer)
 {
 }
 
diff --git a/Engine/Source/Runtime/Engine/Private/ScriptViewportClient.cpp b/Engine/Source/Runtime/Engine/Private/ScriptViewportClient.cpp
--- a/Engine/Source/Runtime/Engine/Private/ScriptViewportClient.cpp
+++ b/Engine/Source/Runtime/Engine/Private/ScriptViewportClient.cpp
@@ -2,7 +2,7 @@
 
 #include "EnginePrivate.h"
 
-UScriptViewportClient::UScriptViewportClient(const class FPostConstructInitializeProperties& PCIP)
-	: Super(PCIP)
+UScriptViewportClient::UScriptViewportClient(const FObjectInitializer& ObjectInitializer)
+	: Super(ObjectInitializer)
 {
 }
diff --git a/Engine/Source/Runtime/Engine/Private/Selection.cpp b/Engine/Source/Runtime/Engine/Private/Selection.cpp
--- a/Engine/Source/Runtime/Engine/Private/Selection.cpp
+++ b/Engine/Source/Runtime/Engine/Private/Selection.cpp
@@ -11,8 +11,8 @@ USelection::FOnSelectionChanged	USelection::SelectObjectEvent;
 FSimpleMulticastDelegate		USelection::SelectNoneEvent;
 
 
-USelection::USelection(const class FPostConstructInitializeProperties& PCIP)
-:	UObject(PCIP)
+USelection::USelection(const FObjectInitializer& ObjectInitializer)
+:	UObject(ObjectInitializer)
 ,	SelectionMutex( 0 )
 ,	bIsBatchDirty(false)
 {
diff --git a/Engine/Source/Runtime/Engine/Private/ShadowMap.cpp b/Engine/Source/Runtime/Engine/Private/ShadowMap.cpp
--- a/Engine/Source/Runtime/Engine/Private/ShadowMap.cpp
+++ b/Engine/Source/Runtime/Engine/Private/ShadowMap.cpp
@@ -28,8 +28,8 @@
 	extern ENGINE_API bool GAllowStreamingLightmaps;
 #endif
 
-UShadowMapTexture2D::UShadowMapTexture2D(const class FPostConstructInitializeProperties& PCIP)
-	: Super(PCIP)
+UShadowMapTexture2D::UShadowMapTexture2D(const FObjectInitializer& ObjectInitializer)
+	: Super(ObjectInitializer)
 {
 	LODGroup = TEXTUREGROUP_Shadowmap;
 }
@@ -176,7 +176,7 @@ bool FShadowMapPendingTexture::AddElement(FShadowMapAllocation& Allocation, cons
 void FShadowMapPendingTexture::StartEncoding(UWorld* InWorld)
 {
 	// Create the shadow-map texture.
-	UShadowMapTexture2D* Texture = new(Outer) UShadowMapTexture2D(FPostConstructInitializeProperties());
+	UShadowMapTexture2D* Texture = new(Outer) UShadowMapTexture2D(FObjectInitializer());
 	Texture->Filter			= GUseBilinearLightmaps ? TF_Default : TF_Nearest;
 	// Signed distance field textures get stored in linear space, since they need more precision near .5.
 	Texture->SRGB			= false;
diff --git a/Engine/Source/Runtime/Engine/Private/ShapeComponent.cpp b/Engine/Source/Runtime/Engine/Private/ShapeComponent.cpp
--- a/Engine/Source/Runtime/Engine/Private/ShapeComponent.cpp
+++ b/Engine/Source/Runtime/Engine/Private/ShapeComponent.cpp
@@ -3,8 +3,8 @@
 
 #include "EnginePrivate.h"
 
-UShapeComponent::UShapeComponent(const class FPostConstructInitializeProperties& PCIP)
-	: Super(PCIP)
+UShapeComponent::UShapeComponent(const FObjectInitializer& ObjectInitializer)
+	: Super(ObjectInitializer)
 {
 	BodyInstance.SetCollisionProfileName(FName(TEXT("OverlapAllDynamic")));
 	// when we deprecated this variable, we switched on for the shapecomponent collision profile
diff --git a/Engine/Source/Runtime/Engine/Private/SimpleConstructionScript.cpp b/Engine/Source/Runtime/Engine/Private/SimpleConstructionScript.cpp
--- a/Engine/Source/Runtime/Engine/Private/SimpleConstructionScript.cpp
+++ b/Engine/Source/Runtime/Engine/Private/SimpleConstructionScript.cpp
@@ -15,8 +15,8 @@
 const FName USimpleConstructionScript::DefaultSceneRootVariableName = FName(TEXT("DefaultSceneRoot"));
 #endif
 
-USimpleConstructionScript::USimpleConstructionScript(const class FPostConstructInitializeProperties& PCIP)
-	: Super(PCIP)
+USimpleConstructionScript::USimpleConstructionScript(const FObjectInitializer& ObjectInitializer)
+	: Super(ObjectInitializer)
 {
 	RootNode_DEPRECATED = NULL;
 	DefaultSceneRootNode = NULL;
diff --git a/Engine/Source/Runtime/Engine/Private/SkeletalMesh.cpp b/Engine/Source/Runtime/Engine/Private/SkeletalMesh.cpp
--- a/Engine/Source/Runtime/Engine/Private/SkeletalMesh.cpp
+++ b/Engine/Source/Runtime/Engine/Private/SkeletalMesh.cpp
@@ -1926,8 +1926,8 @@ bool FSkeletalMeshResource::RequiresCPUSkinning(ERHIFeatureLevel::Type FeatureLe
 /*-----------------------------------------------------------------------------
 	USkeletalMesh
 -----------------------------------------------------------------------------*/
-USkeletalMesh::USkeletalMesh(const class FPostConstructInitializeProperties& PCIP)
-	: Super(PCIP)
+USkeletalMesh::USkeletalMesh(const FObjectInitializer& ObjectInitializer)
+	: Super(ObjectInitializer)
 {
 	SkelMirrorAxis = EAxis::X;
 	SkelMirrorFlipAxis = EAxis::Z;
@@ -3362,8 +3362,8 @@ FString USkeletalMesh::GetDesc()
 /*-----------------------------------------------------------------------------
 USkeletalMeshSocket
 -----------------------------------------------------------------------------*/
-USkeletalMeshSocket::USkeletalMeshSocket(const class FPostConstructInitializeProperties& PCIP)
-	: Super(PCIP)
+USkeletalMeshSocket::USkeletalMeshSocket(const FObjectInitializer& ObjectInitializer)
+	: Super(ObjectInitializer)
 {
 	RelativeScale = FVector(1.0f, 1.0f, 1.0f);
 }
@@ -3497,11 +3497,11 @@ bool USkeletalMeshSocket::AttachActor(AActor* Actor, class USkeletalMeshComponen
 /*-----------------------------------------------------------------------------
 	ASkeletalMeshActor
 -----------------------------------------------------------------------------*/
-ASkeletalMeshActor::ASkeletalMeshActor(const class FPostConstructInitializeProperties& PCIP)
-	: Super(PCIP)
+ASkeletalMeshActor::ASkeletalMeshActor(const FObjectInitializer& ObjectInitializer)
+	: Super(ObjectInitializer)
 {
 
-	SkeletalMeshComponent = PCIP.CreateDefaultSubobject<USkeletalMeshComponent>(this, TEXT("SkeletalMeshComponent0"));
+	SkeletalMeshComponent = ObjectInitializer.CreateDefaultSubobject<USkeletalMeshComponent>(this, TEXT("SkeletalMeshComponent0"));
 	SkeletalMeshComponent->MeshComponentUpdateFlag = EMeshComponentUpdateFlag::AlwaysTickPose;
 	SkeletalMeshComponent->BodyInstance.bEnableCollision_DEPRECATED = true;
 	// check BaseEngine.ini for profile setup
@@ -4656,8 +4656,8 @@ void FSkeletalMeshSceneProxy::UpdateMorphMaterialUsage_GameThread(bool bNeedsMor
 }
 
 
-USkinnedMeshComponent::USkinnedMeshComponent(const class FPostConstructInitializeProperties& PCIP)
-	: Super(PCIP)
+USkinnedMeshComponent::USkinnedMeshComponent(const FObjectInitializer& ObjectInitializer)
+	: Super(ObjectInitializer)
 {
 	bAutoActivate = true;
 	PrimaryComponentTick.bCanEverTick = true;
diff --git a/Engine/Source/Runtime/Engine/Private/SkeletalMeshComponent.cpp b/Engine/Source/Runtime/Engine/Private/SkeletalMeshComponent.cpp
--- a/Engine/Source/Runtime/Engine/Private/SkeletalMeshComponent.cpp
+++ b/Engine/Source/Runtime/Engine/Private/SkeletalMeshComponent.cpp
@@ -85,8 +85,8 @@ public:
 	}
 };
 
-USkeletalMeshComponent::USkeletalMeshComponent(const class FPostConstructInitializeProperties& PCIP)
-	: Super(PCIP)
+USkeletalMeshComponent::USkeletalMeshComponent(const FObjectInitializer& ObjectInitializer)
+	: Super(ObjectInitializer)
 {
 	bAutoActivate = true;
 	PrimaryComponentTick.bCanEverTick = true;
diff --git a/Engine/Source/Runtime/Engine/Private/SkyLightComponent.cpp b/Engine/Source/Runtime/Engine/Private/SkyLightComponent.cpp
--- a/Engine/Source/Runtime/Engine/Private/SkyLightComponent.cpp
+++ b/Engine/Source/Runtime/Engine/Private/SkyLightComponent.cpp
@@ -86,8 +86,8 @@ FSkyLightSceneProxy::FSkyLightSceneProxy(const USkyLightComponent* InLightCompon
 {
 }
 
-USkyLightComponent::USkyLightComponent(const class FPostConstructInitializeProperties& PCIP)
-	: Super(PCIP)
+USkyLightComponent::USkyLightComponent(const FObjectInitializer& ObjectInitializer)
+	: Super(ObjectInitializer)
 {
 #if WITH_EDITORONLY_DATA
 	if (!IsRunningCommandlet())
@@ -471,10 +471,10 @@ void USkyLightComponent::RecaptureSky()
 	SetCaptureIsDirty();
 }
 
-ASkyLight::ASkyLight(const class FPostConstructInitializeProperties& PCIP)
-	: Super(PCIP)
+ASkyLight::ASkyLight(const FObjectInitializer& ObjectInitializer)
+	: Super(ObjectInitializer)
 {
-	LightComponent = PCIP.CreateDefaultSubobject<USkyLightComponent>(this, TEXT("SkyLightComponent0"));
+	LightComponent = ObjectInitializer.CreateDefaultSubobject<USkyLightComponent>(this, TEXT("SkyLightComponent0"));
 	RootComponent = LightComponent;
 
 #if WITH_EDITORONLY_DATA
diff --git a/Engine/Source/Runtime/Engine/Private/Slate/ButtonStyleAsset.cpp b/Engine/Source/Runtime/Engine/Private/Slate/ButtonStyleAsset.cpp
--- a/Engine/Source/Runtime/Engine/Private/Slate/ButtonStyleAsset.cpp
+++ b/Engine/Source/Runtime/Engine/Private/Slate/ButtonStyleAsset.cpp
@@ -3,8 +3,8 @@
 #include "EnginePrivate.h"
 #include "Slate/ButtonStyleAsset.h"
 
-UButtonStyleAsset::UButtonStyleAsset( const class FPostConstructInitializeProperties& PCIP )
-	: Super(PCIP)
+UButtonStyleAsset::UButtonStyleAsset( const FObjectInitializer& ObjectInitializer )
+	: Super(ObjectInitializer)
 {
 	
 }
diff --git a/Engine/Source/Runtime/Engine/Private/Slate/CheckboxStyleAsset.cpp b/Engine/Source/Runtime/Engine/Private/Slate/CheckboxStyleAsset.cpp
--- a/Engine/Source/Runtime/Engine/Private/Slate/CheckboxStyleAsset.cpp
+++ b/Engine/Source/Runtime/Engine/Private/Slate/CheckboxStyleAsset.cpp
@@ -3,8 +3,8 @@
 #include "EnginePrivate.h"
 #include "Slate/CheckboxStyleAsset.h"
 
-UCheckBoxStyleAsset::UCheckBoxStyleAsset( const class FPostConstructInitializeProperties& PCIP )
-	: Super(PCIP)
+UCheckBoxStyleAsset::UCheckBoxStyleAsset( const FObjectInitializer& ObjectInitializer )
+	: Super(ObjectInitializer)
 {
 	
 }
diff --git a/Engine/Source/Runtime/Engine/Private/Slate/SlateBrushAsset.cpp b/Engine/Source/Runtime/Engine/Private/Slate/SlateBrushAsset.cpp
--- a/Engine/Source/Runtime/Engine/Private/Slate/SlateBrushAsset.cpp
+++ b/Engine/Source/Runtime/Engine/Private/Slate/SlateBrushAsset.cpp
@@ -3,8 +3,8 @@
 #include "EnginePrivate.h"
 #include "Slate/SlateBrushAsset.h"
 
-USlateBrushAsset::USlateBrushAsset( const class FPostConstructInitializeProperties& PCIP )
-	: Super(PCIP)
+USlateBrushAsset::USlateBrushAsset( const FObjectInitializer& ObjectInitializer )
+	: Super(ObjectInitializer)
 {
 	
 }
diff --git a/Engine/Source/Runtime/Engine/Private/SoundAttenuation.cpp b/Engine/Source/Runtime/Engine/Private/SoundAttenuation.cpp
--- a/Engine/Source/Runtime/Engine/Private/SoundAttenuation.cpp
+++ b/Engine/Source/Runtime/Engine/Private/SoundAttenuation.cpp
@@ -248,7 +248,7 @@ void FAttenuationSettings::CollectAttenuationShapesForVisualization(TMultiMap<EA
 	}
 }
 
-USoundAttenuation::USoundAttenuation(const class FPostConstructInitializeProperties& PCIP)
-	: Super(PCIP)
+USoundAttenuation::USoundAttenuation(const FObjectInitializer& ObjectInitializer)
+	: Super(ObjectInitializer)
 {
 }
\ No newline at end of file
diff --git a/Engine/Source/Runtime/Engine/Private/SoundBase.cpp b/Engine/Source/Runtime/Engine/Private/SoundBase.cpp
--- a/Engine/Source/Runtime/Engine/Private/SoundBase.cpp
+++ b/Engine/Source/Runtime/Engine/Private/SoundBase.cpp
@@ -5,8 +5,8 @@
 #include "Sound/AudioSettings.h"
 #include "AudioDevice.h"
 
-USoundBase::USoundBase(const class FPostConstructInitializeProperties& PCIP)
-	: Super(PCIP)
+USoundBase::USoundBase(const FObjectInitializer& ObjectInitializer)
+	: Super(ObjectInitializer)
 {
 	MaxConcurrentPlayCount = 16;
 }
diff --git a/Engine/Source/Runtime/Engine/Private/SoundClass.cpp b/Engine/Source/Runtime/Engine/Private/SoundClass.cpp
--- a/Engine/Source/Runtime/Engine/Private/SoundClass.cpp
+++ b/Engine/Source/Runtime/Engine/Private/SoundClass.cpp
@@ -14,8 +14,8 @@
 	USoundClass implementation.
 -----------------------------------------------------------------------------*/
 
-USoundClass::USoundClass(const class FPostConstructInitializeProperties& PCIP)
-	: Super(PCIP)
+USoundClass::USoundClass(const FObjectInitializer& ObjectInitializer)
+	: Super(ObjectInitializer)
 #if WITH_EDITORONLY_DATA
 	, SoundClassGraph(NULL)
 #endif
diff --git a/Engine/Source/Runtime/Engine/Private/SoundCue.cpp b/Engine/Source/Runtime/Engine/Private/SoundCue.cpp
--- a/Engine/Source/Runtime/Engine/Private/SoundCue.cpp
+++ b/Engine/Source/Runtime/Engine/Private/SoundCue.cpp
@@ -15,8 +15,8 @@
 	USoundCue implementation.
 -----------------------------------------------------------------------------*/
 
-USoundCue::USoundCue(const class FPostConstructInitializeProperties& PCIP)
-	: Super(PCIP)
+USoundCue::USoundCue(const FObjectInitializer& ObjectInitializer)
+	: Super(ObjectInitializer)
 {
 	VolumeMultiplier = 0.75f;
 	PitchMultiplier = 1.0f;
diff --git a/Engine/Source/Runtime/Engine/Private/SoundGroups.cpp b/Engine/Source/Runtime/Engine/Private/SoundGroups.cpp
--- a/Engine/Source/Runtime/Engine/Private/SoundGroups.cpp
+++ b/Engine/Source/Runtime/Engine/Private/SoundGroups.cpp
@@ -4,8 +4,8 @@
 #include "Audio.h"
 #include "Sound/SoundGroups.h"
 
-USoundGroups::USoundGroups(const class FPostConstructInitializeProperties& PCIP)
-	: Super(PCIP)
+USoundGroups::USoundGroups(const FObjectInitializer& ObjectInitializer)
+	: Super(ObjectInitializer)
 {
 }
 
diff --git a/Engine/Source/Runtime/Engine/Private/SoundMix.cpp b/Engine/Source/Runtime/Engine/Private/SoundMix.cpp
--- a/Engine/Source/Runtime/Engine/Private/SoundMix.cpp
+++ b/Engine/Source/Runtime/Engine/Private/SoundMix.cpp
@@ -12,8 +12,8 @@
 	USoundMix implementation.
 -----------------------------------------------------------------------------*/
 
-USoundMix::USoundMix(const class FPostConstructInitializeProperties& PCIP)
-	: Super(PCIP)
+USoundMix::USoundMix(const FObjectInitializer& ObjectInitializer)
+	: Super(ObjectInitializer)
 {
 	bApplyEQ = false;
 	InitialDelay = 0.0f;
diff --git a/Engine/Source/Runtime/Engine/Private/SoundNode.cpp b/Engine/Source/Runtime/Engine/Private/SoundNode.cpp
--- a/Engine/Source/Runtime/Engine/Private/SoundNode.cpp
+++ b/Engine/Source/Runtime/Engine/Private/SoundNode.cpp
@@ -12,8 +12,8 @@
 /*-----------------------------------------------------------------------------
 	USoundNode implementation.
 -----------------------------------------------------------------------------*/
-USoundNode::USoundNode(const class FPostConstructInitializeProperties& PCIP)
-	: Super(PCIP)
+USoundNode::USoundNode(const FObjectInitializer& ObjectInitializer)
+	: Super(ObjectInitializer)
 {
 }
 
diff --git a/Engine/Source/Runtime/Engine/Private/SoundNodeAmbient.cpp b/Engine/Source/Runtime/Engine/Private/SoundNodeAmbient.cpp
--- a/Engine/Source/Runtime/Engine/Private/SoundNodeAmbient.cpp
+++ b/Engine/Source/Runtime/Engine/Private/SoundNodeAmbient.cpp
@@ -8,8 +8,8 @@
 /*-----------------------------------------------------------------------------
 	USoundNodeAmbient implementation.
 -----------------------------------------------------------------------------*/
-UDEPRECATED_SoundNodeAmbient::UDEPRECATED_SoundNodeAmbient(const class FPostConstructInitializeProperties& PCIP)
-	: Super(PCIP)
+UDEPRECATED_SoundNodeAmbient::UDEPRECATED_SoundNodeAmbient(const FObjectInitializer& ObjectInitializer)
+	: Super(ObjectInitializer)
 {
 	bAttenuate = true;
 	bSpatialize = true;
diff --git a/Engine/Source/Runtime/Engine/Private/SoundNodeAmbientNonLoop.cpp b/Engine/Source/Runtime/Engine/Private/SoundNodeAmbientNonLoop.cpp
--- a/Engine/Source/Runtime/Engine/Private/SoundNodeAmbientNonLoop.cpp
+++ b/Engine/Source/Runtime/Engine/Private/SoundNodeAmbientNonLoop.cpp
@@ -7,7 +7,7 @@
 /*-----------------------------------------------------------------------------
 	USoundNodeAmbientNonLoop implementation.
 -----------------------------------------------------------------------------*/
-UDEPRECATED_SoundNodeAmbientNonLoop::UDEPRECATED_SoundNodeAmbientNonLoop(const class FPostConstructInitializeProperties& PCIP)
-	: Super(PCIP)
+UDEPRECATED_SoundNodeAmbientNonLoop::UDEPRECATED_SoundNodeAmbientNonLoop(const FObjectInitializer& ObjectInitializer)
+	: Super(ObjectInitializer)
 {
 }
diff --git a/Engine/Source/Runtime/Engine/Private/SoundNodeAmbientNonLoopToggle.cpp b/Engine/Source/Runtime/Engine/Private/SoundNodeAmbientNonLoopToggle.cpp
--- a/Engine/Source/Runtime/Engine/Private/SoundNodeAmbientNonLoopToggle.cpp
+++ b/Engine/Source/Runtime/Engine/Private/SoundNodeAmbientNonLoopToggle.cpp
@@ -7,7 +7,7 @@
 /*-----------------------------------------------------------------------------
 	USoundNodeAmbientNonLoopToggle implementation
 -----------------------------------------------------------------------------*/
-UDEPRECATED_SoundNodeAmbientNonLoopToggle::UDEPRECATED_SoundNodeAmbientNonLoopToggle(const class FPostConstructInitializeProperties& PCIP)
-	: Super(PCIP)
+UDEPRECATED_SoundNodeAmbientNonLoopToggle::UDEPRECATED_SoundNodeAmbientNonLoopToggle(const FObjectInitializer& ObjectInitializer)
+	: Super(ObjectInitializer)
 {
 }
diff --git a/Engine/Source/Runtime/Engine/Private/SoundNodeAttenuation.cpp b/Engine/Source/Runtime/Engine/Private/SoundNodeAttenuation.cpp
--- a/Engine/Source/Runtime/Engine/Private/SoundNodeAttenuation.cpp
+++ b/Engine/Source/Runtime/Engine/Private/SoundNodeAttenuation.cpp
@@ -9,8 +9,8 @@
 	USoundNodeAttenuation implementation.
 -----------------------------------------------------------------------------*/
 
-USoundNodeAttenuation::USoundNodeAttenuation(const class FPostConstructInitializeProperties& PCIP)
-	: Super(PCIP)
+USoundNodeAttenuation::USoundNodeAttenuation(const FObjectInitializer& ObjectInitializer)
+	: Super(ObjectInitializer)
 {
 	bAttenuate_DEPRECATED = true;
 	bSpatialize_DEPRECATED = true;
diff --git a/Engine/Source/Runtime/Engine/Private/SoundNodeBranch.cpp b/Engine/Source/Runtime/Engine/Private/SoundNodeBranch.cpp
--- a/Engine/Source/Runtime/Engine/Private/SoundNodeBranch.cpp
+++ b/Engine/Source/Runtime/Engine/Private/SoundNodeBranch.cpp
@@ -10,8 +10,8 @@
 /*-----------------------------------------------------------------------------
     USoundNodeLocality implementation.
 -----------------------------------------------------------------------------*/
-USoundNodeBranch::USoundNodeBranch(const class FPostConstructInitializeProperties& PCIP)
-	: Super(PCIP)
+USoundNodeBranch::USoundNodeBranch(const FObjectInitializer& ObjectInitializer)
+	: Super(ObjectInitializer)
 {
 }
 
diff --git a/Engine/Source/Runtime/Engine/Private/SoundNodeConcatenator.cpp b/Engine/Source/Runtime/Engine/Private/SoundNodeConcatenator.cpp
--- a/Engine/Source/Runtime/Engine/Private/SoundNodeConcatenator.cpp
+++ b/Engine/Source/Runtime/Engine/Private/SoundNodeConcatenator.cpp
@@ -8,8 +8,8 @@
 /*-----------------------------------------------------------------------------
          USoundNodeConcatenator implementation.
 -----------------------------------------------------------------------------*/
-USoundNodeConcatenator::USoundNodeConcatenator(const class FPostConstructInitializeProperties& PCIP)
-	: Super(PCIP)
+USoundNodeConcatenator::USoundNodeConcatenator(const FObjectInitializer& ObjectInitializer)
+	: Super(ObjectInitializer)
 {
 }
 
diff --git a/Engine/Source/Runtime/Engine/Private/SoundNodeConcatenatorRadio.cpp b/Engine/Source/Runtime/Engine/Private/SoundNodeConcatenatorRadio.cpp
--- a/Engine/Source/Runtime/Engine/Private/SoundNodeConcatenatorRadio.cpp
+++ b/Engine/Source/Runtime/Engine/Private/SoundNodeConcatenatorRadio.cpp
@@ -5,7 +5,7 @@
 #include "SoundDefinitions.h"
 #include "Sound/SoundNodeConcatenatorRadio.h"
 
-UDEPRECATED_SoundNodeConcatenatorRadio::UDEPRECATED_SoundNodeConcatenatorRadio(const class FPostConstructInitializeProperties& PCIP)
-	: Super(PCIP)
+UDEPRECATED_SoundNodeConcatenatorRadio::UDEPRECATED_SoundNodeConcatenatorRadio(const FObjectInitializer& ObjectInitializer)
+	: Super(ObjectInitializer)
 {
 }
diff --git a/Engine/Source/Runtime/Engine/Private/SoundNodeDelay.cpp b/Engine/Source/Runtime/Engine/Private/SoundNodeDelay.cpp
--- a/Engine/Source/Runtime/Engine/Private/SoundNodeDelay.cpp
+++ b/Engine/Source/Runtime/Engine/Private/SoundNodeDelay.cpp
@@ -8,8 +8,8 @@
 /*-----------------------------------------------------------------------------
          USoundNodeDelay implementation.
 -----------------------------------------------------------------------------*/
-USoundNodeDelay::USoundNodeDelay(const class FPostConstructInitializeProperties& PCIP)
-	: Super(PCIP)
+USoundNodeDelay::USoundNodeDelay(const FObjectInitializer& ObjectInitializer)
+	: Super(ObjectInitializer)
 {
 	DelayMin = 0.0f;
 	DelayMax = 0.0f;
diff --git a/Engine/Source/Runtime/Engine/Private/SoundNodeDeprecated.cpp b/Engine/Source/Runtime/Engine/Private/SoundNodeDeprecated.cpp
--- a/Engine/Source/Runtime/Engine/Private/SoundNodeDeprecated.cpp
+++ b/Engine/Source/Runtime/Engine/Private/SoundNodeDeprecated.cpp
@@ -5,7 +5,7 @@
 #include "SoundDefinitions.h"
 #include "Sound/SoundNodeDeprecated.h"
 
-UDEPRECATED_SoundNodeDeprecated::UDEPRECATED_SoundNodeDeprecated(const class FPostConstructInitializeProperties& PCIP)
-	: Super(PCIP)
+UDEPRECATED_SoundNodeDeprecated::UDEPRECATED_SoundNodeDeprecated(const FObjectInitializer& ObjectInitializer)
+	: Super(ObjectInitializer)
 {
 }
diff --git a/Engine/Source/Runtime/Engine/Private/SoundNodeDialoguePlayer.cpp b/Engine/Source/Runtime/Engine/Private/SoundNodeDialoguePlayer.cpp
--- a/Engine/Source/Runtime/Engine/Private/SoundNodeDialoguePlayer.cpp
+++ b/Engine/Source/Runtime/Engine/Private/SoundNodeDialoguePlayer.cpp
@@ -4,8 +4,8 @@
 #include "Sound/SoundBase.h"
 #include "Sound/SoundNodeDialoguePlayer.h"
 
-USoundNodeDialoguePlayer::USoundNodeDialoguePlayer(const class FPostConstructInitializeProperties& PCIP)
-	: Super(PCIP)
+USoundNodeDialoguePlayer::USoundNodeDialoguePlayer(const FObjectInitializer& ObjectInitializer)
+	: Super(ObjectInitializer)
 {
 }
 
diff --git a/Engine/Source/Runtime/Engine/Private/SoundNodeDistanceCrossFade.cpp b/Engine/Source/Runtime/Engine/Private/SoundNodeDistanceCrossFade.cpp
--- a/Engine/Source/Runtime/Engine/Private/SoundNodeDistanceCrossFade.cpp
+++ b/Engine/Source/Runtime/Engine/Private/SoundNodeDistanceCrossFade.cpp
@@ -8,8 +8,8 @@
 /*-----------------------------------------------------------------------------
 	USoundNodeDistanceCrossFade implementation.
 -----------------------------------------------------------------------------*/
-USoundNodeDistanceCrossFade::USoundNodeDistanceCrossFade(const class FPostConstructInitializeProperties& PCIP)
-	: Super(PCIP)
+USoundNodeDistanceCrossFade::USoundNodeDistanceCrossFade(const FObjectInitializer& ObjectInitializer)
+	: Super(ObjectInitializer)
 {
 }
 
diff --git a/Engine/Source/Runtime/Engine/Private/SoundNodeDoppler.cpp b/Engine/Source/Runtime/Engine/Private/SoundNodeDoppler.cpp
--- a/Engine/Source/Runtime/Engine/Private/SoundNodeDoppler.cpp
+++ b/Engine/Source/Runtime/Engine/Private/SoundNodeDoppler.cpp
@@ -8,8 +8,8 @@
 /*-----------------------------------------------------------------------------
          USoundNodeDoppler implementation.
 -----------------------------------------------------------------------------*/
-USoundNodeDoppler::USoundNodeDoppler(const class FPostConstructInitializeProperties& PCIP)
-	: Super(PCIP)
+USoundNodeDoppler::USoundNodeDoppler(const FObjectInitializer& ObjectInitializer)
+	: Super(ObjectInitializer)
 {
 	DopplerIntensity = 1.0f;
 }
diff --git a/Engine/Source/Runtime/Engine/Private/SoundNodeEnveloper.cpp b/Engine/Source/Runtime/Engine/Private/SoundNodeEnveloper.cpp
--- a/Engine/Source/Runtime/Engine/Private/SoundNodeEnveloper.cpp
+++ b/Engine/Source/Runtime/Engine/Private/SoundNodeEnveloper.cpp
@@ -8,8 +8,8 @@
 /*-----------------------------------------------------------------------------
 USoundNodeEnveloper implementation.
 -----------------------------------------------------------------------------*/
-USoundNodeEnveloper::USoundNodeEnveloper(const class FPostConstructInitializeProperties& PCIP)
-	: Super(PCIP)
+USoundNodeEnveloper::USoundNodeEnveloper(const FObjectInitializer& ObjectInitializer)
+	: Super(ObjectInitializer)
 {
 	PitchMin = 1.0f;
 	PitchMax = 1.0f;
diff --git a/Engine/Source/Runtime/Engine/Private/SoundNodeGroupControl.cpp b/Engine/Source/Runtime/Engine/Private/SoundNodeGroupControl.cpp
--- a/Engine/Source/Runtime/Engine/Private/SoundNodeGroupControl.cpp
+++ b/Engine/Source/Runtime/Engine/Private/SoundNodeGroupControl.cpp
@@ -12,8 +12,8 @@ TMap< USoundNodeGroupControl*, TArray< TMap< FActiveSound*, int32> > > USoundNod
 /*-----------------------------------------------------------------------------
     USoundNodeGroupControl implementation.
 -----------------------------------------------------------------------------*/
-USoundNodeGroupControl::USoundNodeGroupControl(const class FPostConstructInitializeProperties& PCIP)
-	: Super(PCIP)
+USoundNodeGroupControl::USoundNodeGroupControl(const FObjectInitializer& ObjectInitializer)
+	: Super(ObjectInitializer)
 {
 }
 
diff --git a/Engine/Source/Runtime/Engine/Private/SoundNodeLooping.cpp b/Engine/Source/Runtime/Engine/Private/SoundNodeLooping.cpp
--- a/Engine/Source/Runtime/Engine/Private/SoundNodeLooping.cpp
+++ b/Engine/Source/Runtime/Engine/Private/SoundNodeLooping.cpp
@@ -9,8 +9,8 @@
 /*-----------------------------------------------------------------------------
 	USoundNodeLooping implementation.
 -----------------------------------------------------------------------------*/
-USoundNodeLooping::USoundNodeLooping(const class FPostConstructInitializeProperties& PCIP)
-	: Super(PCIP)
+USoundNodeLooping::USoundNodeLooping(const FObjectInitializer& ObjectInitializer)
+	: Super(ObjectInitializer)
 {
 }
 
diff --git a/Engine/Source/Runtime/Engine/Private/SoundNodeMature.cpp b/Engine/Source/Runtime/Engine/Private/SoundNodeMature.cpp
--- a/Engine/Source/Runtime/Engine/Private/SoundNodeMature.cpp
+++ b/Engine/Source/Runtime/Engine/Private/SoundNodeMature.cpp
@@ -10,8 +10,8 @@
 /*-----------------------------------------------------------------------------
        USoundNodeMature implementation.
 -----------------------------------------------------------------------------*/
-USoundNodeMature::USoundNodeMature(const class FPostConstructInitializeProperties& PCIP)
-	: Super(PCIP)
+USoundNodeMature::USoundNodeMature(const FObjectInitializer& ObjectInitializer)
+	: Super(ObjectInitializer)
 {
 }
 
diff --git a/Engine/Source/Runtime/Engine/Private/SoundNodeMixer.cpp b/Engine/Source/Runtime/Engine/Private/SoundNodeMixer.cpp
--- a/Engine/Source/Runtime/Engine/Private/SoundNodeMixer.cpp
+++ b/Engine/Source/Runtime/Engine/Private/SoundNodeMixer.cpp
@@ -8,8 +8,8 @@
 /*-----------------------------------------------------------------------------
 	USoundNodeMixer implementation.
 -----------------------------------------------------------------------------*/
-USoundNodeMixer::USoundNodeMixer(const class FPostConstructInitializeProperties& PCIP)
-	: Super(PCIP)
+USoundNodeMixer::USoundNodeMixer(const FObjectInitializer& ObjectInitializer)
+	: Super(ObjectInitializer)
 {
 }
 
diff --git a/Engine/Source/Runtime/Engine/Private/SoundNodeModulator.cpp b/Engine/Source/Runtime/Engine/Private/SoundNodeModulator.cpp
--- a/Engine/Source/Runtime/Engine/Private/SoundNodeModulator.cpp
+++ b/Engine/Source/Runtime/Engine/Private/SoundNodeModulator.cpp
@@ -8,8 +8,8 @@
 /*-----------------------------------------------------------------------------
 	USoundNodeModulator implementation.
 -----------------------------------------------------------------------------*/
-USoundNodeModulator::USoundNodeModulator(const class FPostConstructInitializeProperties& PCIP)
-	: Super(PCIP)
+USoundNodeModulator::USoundNodeModulator(const FObjectInitializer& ObjectInitializer)
+	: Super(ObjectInitializer)
 {
 	PitchMin = 0.95f;
 	PitchMax = 1.05f;
diff --git a/Engine/Source/Runtime/Engine/Private/SoundNodeModulatorContinuous.cpp b/Engine/Source/Runtime/Engine/Private/SoundNodeModulatorContinuous.cpp
--- a/Engine/Source/Runtime/Engine/Private/SoundNodeModulatorContinuous.cpp
+++ b/Engine/Source/Runtime/Engine/Private/SoundNodeModulatorContinuous.cpp
@@ -42,8 +42,8 @@ float FModulatorContinuousParams::GetValue(const FActiveSound& ActiveSound) cons
 /*-----------------------------------------------------------------------------
 	USoundNodeModulatorContinuous implementation.
 -----------------------------------------------------------------------------*/
-USoundNodeModulatorContinuous::USoundNodeModulatorContinuous(const class FPostConstructInitializeProperties& PCIP)
-	: Super(PCIP)
+USoundNodeModulatorContinuous::USoundNodeModulatorContinuous(const FObjectInitializer& ObjectInitializer)
+	: Super(ObjectInitializer)
 {
 }
 
diff --git a/Engine/Source/Runtime/Engine/Private/SoundNodeOscillator.cpp b/Engine/Source/Runtime/Engine/Private/SoundNodeOscillator.cpp
--- a/Engine/Source/Runtime/Engine/Private/SoundNodeOscillator.cpp
+++ b/Engine/Source/Runtime/Engine/Private/SoundNodeOscillator.cpp
@@ -8,8 +8,8 @@
 /*-----------------------------------------------------------------------------
 	USoundNodeOscillator implementation.
 -----------------------------------------------------------------------------*/
-USoundNodeOscillator::USoundNodeOscillator(const class FPostConstructInitializeProperties& PCIP)
-	: Super(PCIP)
+USoundNodeOscillator::USoundNodeOscillator(const FObjectInitializer& ObjectInitializer)
+	: Super(ObjectInitializer)
 {
 	AmplitudeMin = 0;
 	AmplitudeMax = 0;
diff --git a/Engine/Source/Runtime/Engine/Private/SoundNodeParamCrossFade.cpp b/Engine/Source/Runtime/Engine/Private/SoundNodeParamCrossFade.cpp
--- a/Engine/Source/Runtime/Engine/Private/SoundNodeParamCrossFade.cpp
+++ b/Engine/Source/Runtime/Engine/Private/SoundNodeParamCrossFade.cpp
@@ -8,7 +8,7 @@
 /*-----------------------------------------------------------------------------
 	USoundNodeParamCrossFade implementation.
 -----------------------------------------------------------------------------*/
-USoundNodeParamCrossFade::USoundNodeParamCrossFade(const class FPostConstructInitializeProperties& PCIP) : Super(PCIP)
+USoundNodeParamCrossFade::USoundNodeParamCrossFade(const FObjectInitializer& ObjectInitializer) : Super(ObjectInitializer)
 {
 }
 
diff --git a/Engine/Source/Runtime/Engine/Private/SoundNodeRandom.cpp b/Engine/Source/Runtime/Engine/Private/SoundNodeRandom.cpp
--- a/Engine/Source/Runtime/Engine/Private/SoundNodeRandom.cpp
+++ b/Engine/Source/Runtime/Engine/Private/SoundNodeRandom.cpp
@@ -12,8 +12,8 @@
 /*-----------------------------------------------------------------------------
     USoundNodeRandom implementation.
 -----------------------------------------------------------------------------*/
-USoundNodeRandom::USoundNodeRandom(const class FPostConstructInitializeProperties& PCIP)
-	: Super(PCIP)
+USoundNodeRandom::USoundNodeRandom(const FObjectInitializer& ObjectInitializer)
+	: Super(ObjectInitializer)
 {
 	bRandomizeWithoutReplacement = true;
 	NumRandomUsed = 0;
diff --git a/Engine/Source/Runtime/Engine/Private/SoundNodeSoundClass.cpp b/Engine/Source/Runtime/Engine/Private/SoundNodeSoundClass.cpp
--- a/Engine/Source/Runtime/Engine/Private/SoundNodeSoundClass.cpp
+++ b/Engine/Source/Runtime/Engine/Private/SoundNodeSoundClass.cpp
@@ -9,8 +9,8 @@
 	USoundNodeSoundClass implementation.
 -----------------------------------------------------------------------------*/
 
-USoundNodeSoundClass::USoundNodeSoundClass(const class FPostConstructInitializeProperties& PCIP)
-	: Super(PCIP)
+USoundNodeSoundClass::USoundNodeSoundClass(const FObjectInitializer& ObjectInitializer)
+	: Super(ObjectInitializer)
 {
 }
 
diff --git a/Engine/Source/Runtime/Engine/Private/SoundNodeSwitch.cpp b/Engine/Source/Runtime/Engine/Private/SoundNodeSwitch.cpp
--- a/Engine/Source/Runtime/Engine/Private/SoundNodeSwitch.cpp
+++ b/Engine/Source/Runtime/Engine/Private/SoundNodeSwitch.cpp
@@ -11,8 +11,8 @@
 
 #define LOCTEXT_NAMESPACE "SoundNodeSwitch"
 
-USoundNodeSwitch::USoundNodeSwitch(const class FPostConstructInitializeProperties& PCIP)
-	: Super(PCIP)
+USoundNodeSwitch::USoundNodeSwitch(const FObjectInitializer& ObjectInitializer)
+	: Super(ObjectInitializer)
 {
 }
 
diff --git a/Engine/Source/Runtime/Engine/Private/SoundNodeWave.cpp b/Engine/Source/Runtime/Engine/Private/SoundNodeWave.cpp
--- a/Engine/Source/Runtime/Engine/Private/SoundNodeWave.cpp
+++ b/Engine/Source/Runtime/Engine/Private/SoundNodeWave.cpp
@@ -12,8 +12,8 @@
 /*-----------------------------------------------------------------------------
 	USoundNodeWave implementation.
 -----------------------------------------------------------------------------*/
-UDEPRECATED_SoundNodeWave::UDEPRECATED_SoundNodeWave(const class FPostConstructInitializeProperties& PCIP)
-	: Super(PCIP)
+UDEPRECATED_SoundNodeWave::UDEPRECATED_SoundNodeWave(const FObjectInitializer& ObjectInitializer)
+	: Super(ObjectInitializer)
 {
 	Volume = 0.75;
 	Pitch = 1.0;
diff --git a/Engine/Source/Runtime/Engine/Private/SoundNodeWaveParam.cpp b/Engine/Source/Runtime/Engine/Private/SoundNodeWaveParam.cpp
--- a/Engine/Source/Runtime/Engine/Private/SoundNodeWaveParam.cpp
+++ b/Engine/Source/Runtime/Engine/Private/SoundNodeWaveParam.cpp
@@ -9,8 +9,8 @@
 /*-----------------------------------------------------------------------------
 	USoundNodeWaveParam implementation
 -----------------------------------------------------------------------------*/
-USoundNodeWaveParam::USoundNodeWaveParam(const class FPostConstructInitializeProperties& PCIP)
-	: Super(PCIP)
+USoundNodeWaveParam::USoundNodeWaveParam(const FObjectInitializer& ObjectInitializer)
+	: Super(ObjectInitializer)
 {
 }
 
diff --git a/Engine/Source/Runtime/Engine/Private/SoundNodeWavePlayer.cpp b/Engine/Source/Runtime/Engine/Private/SoundNodeWavePlayer.cpp
--- a/Engine/Source/Runtime/Engine/Private/SoundNodeWavePlayer.cpp
+++ b/Engine/Source/Runtime/Engine/Private/SoundNodeWavePlayer.cpp
@@ -6,8 +6,8 @@
 
 #define LOCTEXT_NAMESPACE "SoundNodeWavePlayer"
 
-USoundNodeWavePlayer::USoundNodeWavePlayer(const class FPostConstructInitializeProperties& PCIP)
-	: Super(PCIP)
+USoundNodeWavePlayer::USoundNodeWavePlayer(const FObjectInitializer& ObjectInitializer)
+	: Super(ObjectInitializer)
 {
 }
 
diff --git a/Engine/Source/Runtime/Engine/Private/SoundWave.cpp b/Engine/Source/Runtime/Engine/Private/SoundWave.cpp
--- a/Engine/Source/Runtime/Engine/Private/SoundWave.cpp
+++ b/Engine/Source/Runtime/Engine/Private/SoundWave.cpp
@@ -51,8 +51,8 @@ void FStreamedAudioChunk::StoreInDerivedDataCache(const FString& InDerivedDataKe
 }
 #endif // #if WITH_EDITORONLY_DATA
 
-USoundWave::USoundWave(const class FPostConstructInitializeProperties& PCIP)
-	: Super(PCIP)
+USoundWave::USoundWave(const FObjectInitializer& ObjectInitializer)
+	: Super(ObjectInitializer)
 {
 	Volume = 1.0;
 	Pitch = 1.0;
diff --git a/Engine/Source/Runtime/Engine/Private/SoundWaveStreaming.cpp b/Engine/Source/Runtime/Engine/Private/SoundWaveStreaming.cpp
--- a/Engine/Source/Runtime/Engine/Private/SoundWaveStreaming.cpp
+++ b/Engine/Source/Runtime/Engine/Private/SoundWaveStreaming.cpp
@@ -3,8 +3,8 @@
 #include "EnginePrivate.h"
 #include "Sound/SoundWaveStreaming.h"
 
-USoundWaveStreaming::USoundWaveStreaming(const class FPostConstructInitializeProperties& PCIP)
-	: Super(PCIP)
+USoundWaveStreaming::USoundWaveStreaming(const FObjectInitializer& ObjectInitializer)
+	: Super(ObjectInitializer)
 {
 	bProcedural = true;
 
diff --git a/Engine/Source/Runtime/Engine/Private/SpectatorPawn.cpp b/Engine/Source/Runtime/Engine/Private/SpectatorPawn.cpp
--- a/Engine/Source/Runtime/Engine/Private/SpectatorPawn.cpp
+++ b/Engine/Source/Runtime/Engine/Private/SpectatorPawn.cpp
@@ -5,8 +5,8 @@
 #include "GameFramework/SpectatorPawn.h"
 #include "GameFramework/SpectatorPawnMovement.h"
 
-ASpectatorPawn::ASpectatorPawn(const class FPostConstructInitializeProperties& PCIP)
-	: Super(PCIP
+ASpectatorPawn::ASpectatorPawn(const FObjectInitializer& ObjectInitializer)
+	: Super(ObjectInitializer
 	.SetDefaultSubobjectClass<USpectatorPawnMovement>(Super::MovementComponentName)
 	.DoNotCreateDefaultSubobject(Super::MeshComponentName)
 	)
diff --git a/Engine/Source/Runtime/Engine/Private/SpectatorPawnMovement.cpp b/Engine/Source/Runtime/Engine/Private/SpectatorPawnMovement.cpp
--- a/Engine/Source/Runtime/Engine/Private/SpectatorPawnMovement.cpp
+++ b/Engine/Source/Runtime/Engine/Private/SpectatorPawnMovement.cpp
@@ -8,8 +8,8 @@
 #include "EnginePrivate.h"
 #include "GameFramework/SpectatorPawnMovement.h"
 
-USpectatorPawnMovement::USpectatorPawnMovement(const class FPostConstructInitializeProperties& PCIP)
-	: Super(PCIP)
+USpectatorPawnMovement::USpectatorPawnMovement(const FObjectInitializer& ObjectInitializer)
+	: Super(ObjectInitializer)
 {
 	MaxSpeed = 1200.f;
 	Acceleration = 4000.f;
diff --git a/Engine/Source/Runtime/Engine/Private/SphereComponent.cpp b/Engine/Source/Runtime/Engine/Private/SphereComponent.cpp
--- a/Engine/Source/Runtime/Engine/Private/SphereComponent.cpp
+++ b/Engine/Source/Runtime/Engine/Private/SphereComponent.cpp
@@ -3,8 +3,8 @@
 
 #include "EnginePrivate.h"
 
-USphereComponent::USphereComponent(const class FPostConstructInitializeProperties& PCIP)
-	: Super(PCIP)
+USphereComponent::USphereComponent(const FObjectInitializer& ObjectInitializer)
+	: Super(ObjectInitializer)
 {
 	SphereRadius = 32.0f;
 	ShapeColor = FColor(255, 0, 0, 255);
diff --git a/Engine/Source/Runtime/Engine/Private/SplineComponent.cpp b/Engine/Source/Runtime/Engine/Private/SplineComponent.cpp
--- a/Engine/Source/Runtime/Engine/Private/SplineComponent.cpp
+++ b/Engine/Source/Runtime/Engine/Private/SplineComponent.cpp
@@ -9,8 +9,8 @@
 #include "ComponentInstanceDataCache.h"
 
 
-USplineComponent::USplineComponent(const class FPostConstructInitializeProperties& PCIP)
-	: Super(PCIP)
+USplineComponent::USplineComponent(const FObjectInitializer& ObjectInitializer)
+	: Super(ObjectInitializer)
 	, bAllowSplineEditingPerInstance(true)
 	, ReparamStepsPerSegment(10)
 	, Duration(1.0f)
diff --git a/Engine/Source/Runtime/Engine/Private/SplineMeshComponent.cpp b/Engine/Source/Runtime/Engine/Private/SplineMeshComponent.cpp
--- a/Engine/Source/Runtime/Engine/Private/SplineMeshComponent.cpp
+++ b/Engine/Source/Runtime/Engine/Private/SplineMeshComponent.cpp
@@ -215,8 +215,8 @@ void FSplineMeshSceneProxy::ReleaseResources()
 //////////////////////////////////////////////////////////////////////////
 // SplineMeshComponent
 
-USplineMeshComponent::USplineMeshComponent(const class FPostConstructInitializeProperties& PCIP)
-	: Super(PCIP)
+USplineMeshComponent::USplineMeshComponent(const FObjectInitializer& ObjectInitializer)
+	: Super(ObjectInitializer)
 {
 	Mobility = EComponentMobility::Static;
 
diff --git a/Engine/Source/Runtime/Engine/Private/SpotLight.cpp b/Engine/Source/Runtime/Engine/Private/SpotLight.cpp
--- a/Engine/Source/Runtime/Engine/Private/SpotLight.cpp
+++ b/Engine/Source/Runtime/Engine/Private/SpotLight.cpp
@@ -2,8 +2,8 @@
 
 #include "EnginePrivate.h"
 
-ASpotLight::ASpotLight(const class FPostConstructInitializeProperties& PCIP)
-	: Super(PCIP.SetDefaultSubobjectClass<USpotLightComponent>(TEXT("LightComponent0")))
+ASpotLight::ASpotLight(const FObjectInitializer& ObjectInitializer)
+	: Super(ObjectInitializer.SetDefaultSubobjectClass<USpotLightComponent>(TEXT("LightComponent0")))
 {
 	// Structure to hold one-time initialization
 	struct FConstructorStatics
@@ -26,7 +26,7 @@ ASpotLight::ASpotLight(const class FPostConstructInitializeProperties& PCIP)
 	RootComponent = SpotLightComponent;
 
 #if WITH_EDITORONLY_DATA
-	ArrowComponent = PCIP.CreateEditorOnlyDefaultSubobject<UArrowComponent>(this, TEXT("ArrowComponent0"));
+	ArrowComponent = ObjectInitializer.CreateEditorOnlyDefaultSubobject<UArrowComponent>(this, TEXT("ArrowComponent0"));
 	if (ArrowComponent)
 	{
 		ArrowComponent->ArrowColor = GetLightColor();
diff --git a/Engine/Source/Runtime/Engine/Private/SpotLightComponent.cpp b/Engine/Source/Runtime/Engine/Private/SpotLightComponent.cpp
--- a/Engine/Source/Runtime/Engine/Private/SpotLightComponent.cpp
+++ b/Engine/Source/Runtime/Engine/Private/SpotLightComponent.cpp
@@ -153,8 +153,8 @@ public:
 	}
 };
 
-USpotLightComponent::USpotLightComponent(const class FPostConstructInitializeProperties& PCIP)
-	: Super(PCIP)
+USpotLightComponent::USpotLightComponent(const FObjectInitializer& ObjectInitializer)
+	: Super(ObjectInitializer)
 {
 #if WITH_EDITORONLY_DATA
 	if (!IsRunningCommandlet())
diff --git a/Engine/Source/Runtime/Engine/Private/StaticMesh.cpp b/Engine/Source/Runtime/Engine/Private/StaticMesh.cpp
--- a/Engine/Source/Runtime/Engine/Private/StaticMesh.cpp
+++ b/Engine/Source/Runtime/Engine/Private/StaticMesh.cpp
@@ -1174,8 +1174,8 @@ void FStaticMeshRenderData::Cache(UStaticMesh* Owner, const FStaticMeshLODSettin
 UStaticMesh
 -----------------------------------------------------------------------------*/
 
-UStaticMesh::UStaticMesh(const FPostConstructInitializeProperties& PCIP)
-	: UObject(PCIP)
+UStaticMesh::UStaticMesh(const FObjectInitializer& ObjectInitializer)
+	: UObject(ObjectInitializer)
 {
 	ElementToIgnoreForTexFactor = -1;
 	StreamingDistanceMultiplier=1.0f;
@@ -2562,8 +2562,8 @@ UStaticMeshSocket* UStaticMesh::FindSocket(FName InSocketName)
 UStaticMeshSocket
 -----------------------------------------------------------------------------*/
 
-UStaticMeshSocket::UStaticMeshSocket(const class FPostConstructInitializeProperties& PCIP)
-	: Super(PCIP)
+UStaticMeshSocket::UStaticMeshSocket(const class FObjectInitializer& ObjectInitializer)
+	: Super(ObjectInitializer)
 {
 	RelativeScale = FVector(1.0f, 1.0f, 1.0f);
 }
diff --git a/Engine/Source/Runtime/Engine/Private/StaticMeshActor.cpp b/Engine/Source/Runtime/Engine/Private/StaticMeshActor.cpp
--- a/Engine/Source/Runtime/Engine/Private/StaticMeshActor.cpp
+++ b/Engine/Source/Runtime/Engine/Private/StaticMeshActor.cpp
@@ -13,12 +13,12 @@
 
 #define LOCTEXT_NAMESPACE "StaticMeshActor"
 
-AStaticMeshActor::AStaticMeshActor(const class FPostConstructInitializeProperties& PCIP)
-	: Super(PCIP)
+AStaticMeshActor::AStaticMeshActor(const FObjectInitializer& ObjectInitializer)
+	: Super(ObjectInitializer)
 {
 	bCanBeDamaged = false;
 
-	StaticMeshComponent = PCIP.CreateDefaultSubobject<UStaticMeshComponent>(this, TEXT("StaticMeshComponent0"));
+	StaticMeshComponent = ObjectInitializer.CreateDefaultSubobject<UStaticMeshComponent>(this, TEXT("StaticMeshComponent0"));
 	StaticMeshComponent->SetCollisionProfileName(UCollisionProfile::BlockAll_ProfileName);
 	StaticMeshComponent->Mobility = EComponentMobility::Static;
 	StaticMeshComponent->bGenerateOverlapEvents = false;
diff --git a/Engine/Source/Runtime/Engine/Private/StaticMeshComponent.cpp b/Engine/Source/Runtime/Engine/Private/StaticMeshComponent.cpp
--- a/Engine/Source/Runtime/Engine/Private/StaticMeshComponent.cpp
+++ b/Engine/Source/Runtime/Engine/Private/StaticMeshComponent.cpp
@@ -123,8 +123,8 @@ public:
 	FLightMapInstanceData CachedStaticLighting;
 };
 
-UStaticMeshComponent::UStaticMeshComponent(const class FPostConstructInitializeProperties& PCIP)
-	: Super(PCIP)
+UStaticMeshComponent::UStaticMeshComponent(const FObjectInitializer& ObjectInitializer)
+	: Super(ObjectInitializer)
 {
 	PrimaryComponentTick.bCanEverTick = false;
 
diff --git a/Engine/Source/Runtime/Engine/Private/TargetPoint.cpp b/Engine/Source/Runtime/Engine/Private/TargetPoint.cpp
--- a/Engine/Source/Runtime/Engine/Private/TargetPoint.cpp
+++ b/Engine/Source/Runtime/Engine/Private/TargetPoint.cpp
@@ -2,16 +2,16 @@
 
 #include "EnginePrivate.h"
 
-ATargetPoint::ATargetPoint(const class FPostConstructInitializeProperties& PCIP)
-	: Super(PCIP)
+ATargetPoint::ATargetPoint(const FObjectInitializer& ObjectInitializer)
+	: Super(ObjectInitializer)
 {
-	TSubobjectPtr<USceneComponent> SceneComponent = PCIP.CreateDefaultSubobject<USceneComponent>(this, TEXT("SceneComp"));
+	TSubobjectPtr<USceneComponent> SceneComponent = ObjectInitializer.CreateDefaultSubobject<USceneComponent>(this, TEXT("SceneComp"));
 
 	RootComponent = SceneComponent;
 
 #if WITH_EDITORONLY_DATA
-	SpriteComponent = PCIP.CreateEditorOnlyDefaultSubobject<UBillboardComponent>(this, TEXT("Sprite"));
-	ArrowComponent = PCIP.CreateEditorOnlyDefaultSubobject<UArrowComponent>(this, TEXT("Arrow"));
+	SpriteComponent = ObjectInitializer.CreateEditorOnlyDefaultSubobject<UBillboardComponent>(this, TEXT("Sprite"));
+	ArrowComponent = ObjectInitializer.CreateEditorOnlyDefaultSubobject<UArrowComponent>(this, TEXT("Arrow"));
 
 	if (!IsRunningCommandlet())
 	{
diff --git a/Engine/Source/Runtime/Engine/Private/Tests/EngineAutomationTests.cpp b/Engine/Source/Runtime/Engine/Private/Tests/EngineAutomationTests.cpp
--- a/Engine/Source/Runtime/Engine/Private/Tests/EngineAutomationTests.cpp
+++ b/Engine/Source/Runtime/Engine/Private/Tests/EngineAutomationTests.cpp
@@ -446,7 +446,7 @@ bool FCinematicFPSPerfTest::RunTest(const FString& Parameters)
 /* UAutomationTestSettings interface
  *****************************************************************************/
 
-UAutomationTestSettings::UAutomationTestSettings( const class FPostConstructInitializeProperties& PCIP )
-	: Super(PCIP)
+UAutomationTestSettings::UAutomationTestSettings( const FObjectInitializer& ObjectInitializer )
+	: Super(ObjectInitializer)
 {
 }
diff --git a/Engine/Source/Runtime/Engine/Private/Tests/IntSerializationTest.cpp b/Engine/Source/Runtime/Engine/Private/Tests/IntSerializationTest.cpp
--- a/Engine/Source/Runtime/Engine/Private/Tests/IntSerializationTest.cpp
+++ b/Engine/Source/Runtime/Engine/Private/Tests/IntSerializationTest.cpp
@@ -11,7 +11,7 @@ bool FIntSerializationTest::RunTest (const FString& Parameters)
 	//Create object for serialization
 	//UIntSerialization SerializableObject;
 	//SerializableObject->
-	UIntSerialization* SerializableObject = new UIntSerialization(FPostConstructInitializeProperties());
+	UIntSerialization* SerializableObject = new UIntSerialization(FObjectInitializer());
 	SerializableObject->UnsignedInt8Variable = 255U;
 	SerializableObject->UnsignedInt16Variable = 65535U;
 	SerializableObject->UnsignedInt32Variable = 4294967295U;
@@ -30,7 +30,7 @@ bool FIntSerializationTest::RunTest (const FString& Parameters)
 
 	//Deserialize into new object
 	FMemoryReader InAr(SaveData, true);
-	UIntSerialization* DeSerializableObject = new UIntSerialization(FPostConstructInitializeProperties());
+	UIntSerialization* DeSerializableObject = new UIntSerialization(FObjectInitializer());
 	DeSerializableObject->Serialize(InAr);
 
 	//Compare test
diff --git a/Engine/Source/Runtime/Engine/Private/Tests/Internationalization/TextPropertyTest.cpp b/Engine/Source/Runtime/Engine/Private/Tests/Internationalization/TextPropertyTest.cpp
--- a/Engine/Source/Runtime/Engine/Private/Tests/Internationalization/TextPropertyTest.cpp
+++ b/Engine/Source/Runtime/Engine/Private/Tests/Internationalization/TextPropertyTest.cpp
@@ -5,8 +5,8 @@
 
 #define LOCTEXT_NAMESPACE "TextPropertyTest"
 
-UTextPropertyTestObject::UTextPropertyTestObject(const class FPostConstructInitializeProperties& PCIP)
-	: Super(PCIP), DefaultedText( LOCTEXT("DefaultedText", "DefaultValue") )
+UTextPropertyTestObject::UTextPropertyTestObject(const FObjectInitializer& ObjectInitializer)
+	: Super(ObjectInitializer), DefaultedText( LOCTEXT("DefaultedText", "DefaultValue") )
 {
 }
 
diff --git a/Engine/Source/Runtime/Engine/Private/TextRenderComponent.cpp b/Engine/Source/Runtime/Engine/Private/TextRenderComponent.cpp
--- a/Engine/Source/Runtime/Engine/Private/TextRenderComponent.cpp
+++ b/Engine/Source/Runtime/Engine/Private/TextRenderComponent.cpp
@@ -5,14 +5,14 @@
 #include "Engine/TextRenderActor.h"
 #include "LocalVertexFactory.h"
 
-ATextRenderActor::ATextRenderActor(const class FPostConstructInitializeProperties& PCIP)
-	: Super(PCIP)
+ATextRenderActor::ATextRenderActor(const FObjectInitializer& ObjectInitializer)
+	: Super(ObjectInitializer)
 {
-	TextRender = PCIP.CreateDefaultSubobject<UTextRenderComponent>(this, TEXT("NewTextRenderComponent"));
+	TextRender = ObjectInitializer.CreateDefaultSubobject<UTextRenderComponent>(this, TEXT("NewTextRenderComponent"));
 	RootComponent = TextRender;
 
 #if WITH_EDITORONLY_DATA
-	SpriteComponent = PCIP.CreateEditorOnlyDefaultSubobject<UBillboardComponent>(this, TEXT("Sprite"));
+	SpriteComponent = ObjectInitializer.CreateEditorOnlyDefaultSubobject<UBillboardComponent>(this, TEXT("Sprite"));
 
 	if (!IsRunningCommandlet() && (SpriteComponent != nullptr))
 	{
@@ -726,8 +726,8 @@ bool  FTextRenderSceneProxy::BuildStringMesh( TArray<FDynamicMeshVertex>& OutVer
 
 // ------------------------------------------------------
 
-UTextRenderComponent::UTextRenderComponent(const class FPostConstructInitializeProperties& PCIP)
-	: Super(PCIP)
+UTextRenderComponent::UTextRenderComponent(const FObjectInitializer& ObjectInitializer)
+	: Super(ObjectInitializer)
 {
 	// Structure to hold one-time initialization
 	struct FConstructorStatics
diff --git a/Engine/Source/Runtime/Engine/Private/Texture.cpp b/Engine/Source/Runtime/Engine/Private/Texture.cpp
--- a/Engine/Source/Runtime/Engine/Private/Texture.cpp
+++ b/Engine/Source/Runtime/Engine/Private/Texture.cpp
@@ -35,8 +35,8 @@ FName FTextureResource::TextureGroupStatFNames[TEXTUREGROUP_MAX] =
 
 UTexture::FOnTextureSaved UTexture::PreSaveEvent;
 
-UTexture::UTexture(const class FPostConstructInitializeProperties& PCIP)
-	: Super(PCIP)
+UTexture::UTexture(const FObjectInitializer& ObjectInitializer)
+	: Super(ObjectInitializer)
 {
 	SRGB = true;
 	Filter = TF_Default;
diff --git a/Engine/Source/Runtime/Engine/Private/Texture2D.cpp b/Engine/Source/Runtime/Engine/Private/Texture2D.cpp
--- a/Engine/Source/Runtime/Engine/Private/Texture2D.cpp
+++ b/Engine/Source/Runtime/Engine/Private/Texture2D.cpp
@@ -18,8 +18,8 @@
 #include "TargetPlatform.h"
 #include "ContentStreaming.h"
 
-UTexture2D::UTexture2D(const class FPostConstructInitializeProperties& PCIP)
-	: Super(PCIP)
+UTexture2D::UTexture2D(const FObjectInitializer& ObjectInitializer)
+	: Super(ObjectInitializer)
 {
 	bHasCancelationPending = false;
 	StreamingIndex = -1;
diff --git a/Engine/Source/Runtime/Engine/Private/Texture2DDynamic.cpp b/Engine/Source/Runtime/Engine/Private/Texture2DDynamic.cpp
--- a/Engine/Source/Runtime/Engine/Private/Texture2DDynamic.cpp
+++ b/Engine/Source/Runtime/Engine/Private/Texture2DDynamic.cpp
@@ -80,8 +80,8 @@ FTexture2DRHIRef FTexture2DDynamicResource::GetTexture2DRHI()
 /*-----------------------------------------------------------------------------
 	UTexture2DDynamic
 -----------------------------------------------------------------------------*/
-UTexture2DDynamic::UTexture2DDynamic(const class FPostConstructInitializeProperties& PCIP)
-	: Super(PCIP)
+UTexture2DDynamic::UTexture2DDynamic(const FObjectInitializer& ObjectInitializer)
+	: Super(ObjectInitializer)
 {
 	NeverStream = true;
 	Format = PF_B8G8R8A8;
diff --git a/Engine/Source/Runtime/Engine/Private/TextureCube.cpp b/Engine/Source/Runtime/Engine/Private/TextureCube.cpp
--- a/Engine/Source/Runtime/Engine/Private/TextureCube.cpp
+++ b/Engine/Source/Runtime/Engine/Private/TextureCube.cpp
@@ -8,8 +8,8 @@
 #include "RenderUtils.h"
 #include "DDSLoader.h"
 
-UTextureCube::UTextureCube(const class FPostConstructInitializeProperties& PCIP)
-	: Super(PCIP)
+UTextureCube::UTextureCube(const FObjectInitializer& ObjectInitializer)
+	: Super(ObjectInitializer)
 {
 	SRGB = true;
 }
diff --git a/Engine/Source/Runtime/Engine/Private/TextureLightProfile.cpp b/Engine/Source/Runtime/Engine/Private/TextureLightProfile.cpp
--- a/Engine/Source/Runtime/Engine/Private/TextureLightProfile.cpp
+++ b/Engine/Source/Runtime/Engine/Private/TextureLightProfile.cpp
@@ -10,8 +10,8 @@
 /*-----------------------------------------------------------------------------
 	UTextureLightProfile
 -----------------------------------------------------------------------------*/
-UTextureLightProfile::UTextureLightProfile(const class FPostConstructInitializeProperties& PCIP)
-	: Super(PCIP)
+UTextureLightProfile::UTextureLightProfile(const FObjectInitializer& ObjectInitializer)
+	: Super(ObjectInitializer)
 {
 	NeverStream = true;
 	Brightness = -1;
diff --git a/Engine/Source/Runtime/Engine/Private/TextureRenderTarget.cpp b/Engine/Source/Runtime/Engine/Private/TextureRenderTarget.cpp
--- a/Engine/Source/Runtime/Engine/Private/TextureRenderTarget.cpp
+++ b/Engine/Source/Runtime/Engine/Private/TextureRenderTarget.cpp
@@ -10,8 +10,8 @@
 UTextureRenderTarget
 -----------------------------------------------------------------------------*/
 
-UTextureRenderTarget::UTextureRenderTarget(const class FPostConstructInitializeProperties& PCIP)
-	: Super(PCIP)
+UTextureRenderTarget::UTextureRenderTarget(const FObjectInitializer& ObjectInitializer)
+	: Super(ObjectInitializer)
 {
 	NeverStream = true;
 	SRGB = true;
diff --git a/Engine/Source/Runtime/Engine/Private/TextureRenderTarget2D.cpp b/Engine/Source/Runtime/Engine/Private/TextureRenderTarget2D.cpp
--- a/Engine/Source/Runtime/Engine/Private/TextureRenderTarget2D.cpp
+++ b/Engine/Source/Runtime/Engine/Private/TextureRenderTarget2D.cpp
@@ -10,8 +10,8 @@
 	UTextureRenderTarget2D
 -----------------------------------------------------------------------------*/
 
-UTextureRenderTarget2D::UTextureRenderTarget2D(const class FPostConstructInitializeProperties& PCIP)
-	: Super(PCIP)
+UTextureRenderTarget2D::UTextureRenderTarget2D(const FObjectInitializer& ObjectInitializer)
+	: Super(ObjectInitializer)
 {
 	bHDR = true;
 	ClearColor = FLinearColor(0.0f, 1.0f, 0.0f, 1.0f);
diff --git a/Engine/Source/Runtime/Engine/Private/TextureRenderTargetCube.cpp b/Engine/Source/Runtime/Engine/Private/TextureRenderTargetCube.cpp
--- a/Engine/Source/Runtime/Engine/Private/TextureRenderTargetCube.cpp
+++ b/Engine/Source/Runtime/Engine/Private/TextureRenderTargetCube.cpp
@@ -11,8 +11,8 @@
 	UTextureRenderTargetCube
 -----------------------------------------------------------------------------*/
 
-UTextureRenderTargetCube::UTextureRenderTargetCube(const class FPostConstructInitializeProperties& PCIP)
-	: Super(PCIP)
+UTextureRenderTargetCube::UTextureRenderTargetCube(const FObjectInitializer& ObjectInitializer)
+	: Super(ObjectInitializer)
 {
 	bHDR = true;
 	ClearColor = FLinearColor(0.0f, 1.0f, 0.0f, 1.0f);
diff --git a/Engine/Source/Runtime/Engine/Private/Timeline.cpp b/Engine/Source/Runtime/Engine/Private/Timeline.cpp
--- a/Engine/Source/Runtime/Engine/Private/Timeline.cpp
+++ b/Engine/Source/Runtime/Engine/Private/Timeline.cpp
@@ -553,8 +553,8 @@ void FTimeline::GetAllCurves(TSet<class UCurveBase*>& InOutCurves) const
 //////////////////////////////////////////////////////////////////////////
 // UTimelineComponent
 
-UTimelineComponent::UTimelineComponent(const class FPostConstructInitializeProperties& PCIP)
-	: Super(PCIP)
+UTimelineComponent::UTimelineComponent(const FObjectInitializer& ObjectInitializer)
+	: Super(ObjectInitializer)
 {
 	PrimaryComponentTick.bCanEverTick = true;
 	PrimaryComponentTick.TickGroup = TG_PrePhysics;
diff --git a/Engine/Source/Runtime/Engine/Private/TimelineTemplate.cpp b/Engine/Source/Runtime/Engine/Private/TimelineTemplate.cpp
--- a/Engine/Source/Runtime/Engine/Private/TimelineTemplate.cpp
+++ b/Engine/Source/Runtime/Engine/Private/TimelineTemplate.cpp
@@ -29,8 +29,8 @@ namespace
 //////////////////////////////////////////////////////////////////////////
 // UTimelineTemplate
 
-UTimelineTemplate::UTimelineTemplate(const class FPostConstructInitializeProperties& PCIP)
-	: Super(PCIP)
+UTimelineTemplate::UTimelineTemplate(const FObjectInitializer& ObjectInitializer)
+	: Super(ObjectInitializer)
 {
 	TimelineLength = 5.0f;
 	TimelineGuid = FGuid::NewGuid();
diff --git a/Engine/Source/Runtime/Engine/Private/TriggerActors.cpp b/Engine/Source/Runtime/Engine/Private/TriggerActors.cpp
--- a/Engine/Source/Runtime/Engine/Private/TriggerActors.cpp
+++ b/Engine/Source/Runtime/Engine/Private/TriggerActors.cpp
@@ -13,8 +13,8 @@ namespace
 {
 	static const FColor TriggerBaseColor(100, 255, 100, 255);
 }
-ATriggerCapsule::ATriggerCapsule(const class FPostConstructInitializeProperties& PCIP)
-	: Super(PCIP.SetDefaultSubobjectClass<UCapsuleComponent>(TEXT("CollisionComp")))
+ATriggerCapsule::ATriggerCapsule(const FObjectInitializer& ObjectInitializer)
+	: Super(ObjectInitializer.SetDefaultSubobjectClass<UCapsuleComponent>(TEXT("CollisionComp")))
 {
 	UCapsuleComponent* CapsuleCollisionComponent = CastChecked<UCapsuleComponent>(CollisionComponent);
 	CapsuleCollisionComponent->ShapeColor = TriggerBaseColor;
@@ -81,8 +81,8 @@ void ATriggerCapsule::EditorApplyScale(const FVector& DeltaScale, const FVector*
 }
 #endif	//WITH_EDITOR
 
-ATriggerBox::ATriggerBox(const class FPostConstructInitializeProperties& PCIP)
-	: Super(PCIP.SetDefaultSubobjectClass<UBoxComponent>(TEXT("CollisionComp")))
+ATriggerBox::ATriggerBox(const FObjectInitializer& ObjectInitializer)
+	: Super(ObjectInitializer.SetDefaultSubobjectClass<UBoxComponent>(TEXT("CollisionComp")))
 {
 	UBoxComponent* BoxCollisionComponent = CastChecked<UBoxComponent>(CollisionComponent);
 
@@ -125,8 +125,8 @@ void ATriggerBox::EditorApplyScale(const FVector& DeltaScale, const FVector* Piv
 }
 #endif	//WITH_EDITOR
 
-ATriggerSphere::ATriggerSphere(const class FPostConstructInitializeProperties& PCIP)
-	: Super(PCIP.SetDefaultSubobjectClass<USphereComponent>(TEXT("CollisionComp")))
+ATriggerSphere::ATriggerSphere(const FObjectInitializer& ObjectInitializer)
+	: Super(ObjectInitializer.SetDefaultSubobjectClass<USphereComponent>(TEXT("CollisionComp")))
 {
 	USphereComponent* SphereCollisionComponent = CastChecked<USphereComponent>(CollisionComponent);
 
diff --git a/Engine/Source/Runtime/Engine/Private/TriggerBase.cpp b/Engine/Source/Runtime/Engine/Private/TriggerBase.cpp
--- a/Engine/Source/Runtime/Engine/Private/TriggerBase.cpp
+++ b/Engine/Source/Runtime/Engine/Private/TriggerBase.cpp
@@ -3,8 +3,8 @@
 #include "EnginePrivate.h"
 #include "Engine/TriggerBase.h"
 
-ATriggerBase::ATriggerBase(const class FPostConstructInitializeProperties& PCIP)
-	: Super(PCIP)
+ATriggerBase::ATriggerBase(const FObjectInitializer& ObjectInitializer)
+	: Super(ObjectInitializer)
 {
 	// Structure to hold one-time initialization
 	struct FConstructorStatics
@@ -25,15 +25,15 @@ ATriggerBase::ATriggerBase(const class FPostConstructInitializeProperties& PCIP)
 	bCanBeDamaged = false;
 
 	// ATriggerBase is requesting UShapeComponent which is abstract, however it is responsibility
-	// of a derived class to override this type with PCIP.SetDefaultSubobjectClass.
-	CollisionComponent = PCIP.CreateAbstractDefaultSubobject<UShapeComponent>(this, TEXT("CollisionComp"));
+	// of a derived class to override this type with ObjectInitializer.SetDefaultSubobjectClass.
+	CollisionComponent = ObjectInitializer.CreateAbstractDefaultSubobject<UShapeComponent>(this, TEXT("CollisionComp"));
 	if (CollisionComponent)
 	{
 		RootComponent = CollisionComponent;
 		CollisionComponent->bHiddenInGame = false;
 	}
 
-	SpriteComponent = PCIP.CreateDefaultSubobject<UBillboardComponent>(this, TEXT("Sprite"));
+	SpriteComponent = ObjectInitializer.CreateDefaultSubobject<UBillboardComponent>(this, TEXT("Sprite"));
 	if (SpriteComponent)
 	{
 		SpriteComponent->Sprite = ConstructorStatics.TriggerTextureObject.Get();
diff --git a/Engine/Source/Runtime/Engine/Private/TriggerVolume.cpp b/Engine/Source/Runtime/Engine/Private/TriggerVolume.cpp
--- a/Engine/Source/Runtime/Engine/Private/TriggerVolume.cpp
+++ b/Engine/Source/Runtime/Engine/Private/TriggerVolume.cpp
@@ -3,8 +3,8 @@
 #include "EnginePrivate.h"
 #include "Engine/TriggerVolume.h"
 
-ATriggerVolume::ATriggerVolume(const class FPostConstructInitializeProperties& PCIP)
-	: Super(PCIP)
+ATriggerVolume::ATriggerVolume(const FObjectInitializer& ObjectInitializer)
+	: Super(ObjectInitializer)
 {
 	static FName CollisionProfileName(TEXT("Trigger"));
 	BrushComponent->SetCollisionProfileName(CollisionProfileName);
diff --git a/Engine/Source/Runtime/Engine/Private/UnrealEngine.cpp b/Engine/Source/Runtime/Engine/Private/UnrealEngine.cpp
--- a/Engine/Source/Runtime/Engine/Private/UnrealEngine.cpp
+++ b/Engine/Source/Runtime/Engine/Private/UnrealEngine.cpp
@@ -8086,7 +8086,7 @@ EBrowseReturnVal::Type UEngine::Browse( FWorldContext& WorldContext, FURL URL, F
 			ShutdownWorldNetDriver(WorldContext.World());
 		}
 
-		WorldContext.PendingNetGame = new UPendingNetGame(FPostConstructInitializeProperties(), URL);
+		WorldContext.PendingNetGame = new UPendingNetGame(FObjectInitializer(), URL);
 		WorldContext.PendingNetGame->InitNetDriver();
 		if( !WorldContext.PendingNetGame->NetDriver )
 		{
diff --git a/Engine/Source/Runtime/Engine/Private/UnrealExporter.cpp b/Engine/Source/Runtime/Engine/Private/UnrealExporter.cpp
--- a/Engine/Source/Runtime/Engine/Private/UnrealExporter.cpp
+++ b/Engine/Source/Runtime/Engine/Private/UnrealExporter.cpp
@@ -14,8 +14,8 @@ FString UExporter::CurrentFilename(TEXT(""));
 
 TSet< TWeakObjectPtr<UExporter> > UExporter::RegisteredExporters;
 
-UExporter::UExporter(const class FPostConstructInitializeProperties& PCIP)
-	: Super(PCIP)
+UExporter::UExporter(const FObjectInitializer& ObjectInitializer)
+	: Super(ObjectInitializer)
 {
 	if ( HasAnyFlags(RF_ClassDefaultObject) )
 	{
diff --git a/Engine/Source/Runtime/Engine/Private/UserDefinedEnum.cpp b/Engine/Source/Runtime/Engine/Private/UserDefinedEnum.cpp
--- a/Engine/Source/Runtime/Engine/Private/UserDefinedEnum.cpp
+++ b/Engine/Source/Runtime/Engine/Private/UserDefinedEnum.cpp
@@ -6,8 +6,8 @@
 #include "Editor/UnrealEd/Public/Kismet2/EnumEditorUtils.h"
 #endif	// WITH_EDITOR
 
-UUserDefinedEnum::UUserDefinedEnum(const class FPostConstructInitializeProperties& PCIP)
-	: Super(PCIP)
+UUserDefinedEnum::UUserDefinedEnum(const FObjectInitializer& ObjectInitializer)
+	: Super(ObjectInitializer)
 {
 
 }
diff --git a/Engine/Source/Runtime/Engine/Private/UserDefinedStruct.cpp b/Engine/Source/Runtime/Engine/Private/UserDefinedStruct.cpp
--- a/Engine/Source/Runtime/Engine/Private/UserDefinedStruct.cpp
+++ b/Engine/Source/Runtime/Engine/Private/UserDefinedStruct.cpp
@@ -6,8 +6,8 @@
 #include "StructureEditorUtils.h"
 #endif //WITH_EDITOR
 
-UUserDefinedStruct::UUserDefinedStruct(const class FPostConstructInitializeProperties& PCIP)
-	: Super(PCIP)
+UUserDefinedStruct::UUserDefinedStruct(const FObjectInitializer& ObjectInitializer)
+	: Super(ObjectInitializer)
 {
 }
 
diff --git a/Engine/Source/Runtime/Engine/Private/UserInterface/Canvas.cpp b/Engine/Source/Runtime/Engine/Private/UserInterface/Canvas.cpp
--- a/Engine/Source/Runtime/Engine/Private/UserInterface/Canvas.cpp
+++ b/Engine/Source/Runtime/Engine/Private/UserInterface/Canvas.cpp
@@ -1089,8 +1089,8 @@ void UCanvas::WrapString( FTextSizingParameters& Parameters, const float InCurX,
 	UCanvas object functions.
 -----------------------------------------------------------------------------*/
 
-UCanvas::UCanvas(const class FPostConstructInitializeProperties& PCIP)
-	: Super(PCIP)
+UCanvas::UCanvas(const FObjectInitializer& ObjectInitializer)
+	: Super(ObjectInitializer)
 {
 	// Structure to hold one-time initialization
 	struct FConstructorStatics
diff --git a/Engine/Source/Runtime/Engine/Private/UserInterface/Console.cpp b/Engine/Source/Runtime/Engine/Private/UserInterface/Console.cpp
--- a/Engine/Source/Runtime/Engine/Private/UserInterface/Console.cpp
+++ b/Engine/Source/Runtime/Engine/Private/UserInterface/Console.cpp
@@ -42,8 +42,8 @@ public:
 	}
 };
 
-UConsole::UConsole(const class FPostConstructInitializeProperties& PCIP)
-	: Super(PCIP)
+UConsole::UConsole(const FObjectInitializer& ObjectInitializer)
+	: Super(ObjectInitializer)
 //	, HistoryBufferCur(0)
 //	, AutoCompleteCursor(-1)
 {
diff --git a/Engine/Source/Runtime/Engine/Private/UserInterface/ForceFeedbackEffect.cpp b/Engine/Source/Runtime/Engine/Private/UserInterface/ForceFeedbackEffect.cpp
--- a/Engine/Source/Runtime/Engine/Private/UserInterface/ForceFeedbackEffect.cpp
+++ b/Engine/Source/Runtime/Engine/Private/UserInterface/ForceFeedbackEffect.cpp
@@ -3,8 +3,8 @@
 #include "EnginePrivate.h"
 #include "IForceFeedbackSystem.h"
 
-UForceFeedbackEffect::UForceFeedbackEffect(const class FPostConstructInitializeProperties& PCIP)
-	: Super(PCIP)
+UForceFeedbackEffect::UForceFeedbackEffect(const FObjectInitializer& ObjectInitializer)
+	: Super(ObjectInitializer)
 {
 	// Make sure that by default the force feedback effect has an entry
 	FForceFeedbackChannelDetails ChannelDetail;
diff --git a/Engine/Source/Runtime/Engine/Private/UserInterface/InputSettings.cpp b/Engine/Source/Runtime/Engine/Private/UserInterface/InputSettings.cpp
--- a/Engine/Source/Runtime/Engine/Private/UserInterface/InputSettings.cpp
+++ b/Engine/Source/Runtime/Engine/Private/UserInterface/InputSettings.cpp
@@ -6,8 +6,8 @@
 
 #include "EnginePrivate.h"
 
-UInputSettings::UInputSettings(const class FPostConstructInitializeProperties& PCIP)
-: Super(PCIP)
+UInputSettings::UInputSettings(const FObjectInitializer& ObjectInitializer)
+: Super(ObjectInitializer)
 {
 }
 
diff --git a/Engine/Source/Runtime/Engine/Private/UserInterface/PlayerInput.cpp b/Engine/Source/Runtime/Engine/Private/UserInterface/PlayerInput.cpp
--- a/Engine/Source/Runtime/Engine/Private/UserInterface/PlayerInput.cpp
+++ b/Engine/Source/Runtime/Engine/Private/UserInterface/PlayerInput.cpp
@@ -21,8 +21,8 @@ const TArray<FInputActionKeyMapping> UPlayerInput::NoKeyMappings;
 TArray<FInputActionKeyMapping> UPlayerInput::EngineDefinedActionMappings;
 TArray<FInputAxisKeyMapping> UPlayerInput::EngineDefinedAxisMappings;
 
-UPlayerInput::UPlayerInput(const class FPostConstructInitializeProperties& PCIP)
-	: Super(PCIP)
+UPlayerInput::UPlayerInput(const FObjectInitializer& ObjectInitializer)
+	: Super(ObjectInitializer)
 {
 	SetFlags(RF_Transactional);
 	MouseSamplingTotal = +0.0083f;
diff --git a/Engine/Source/Runtime/Engine/Private/VectorField.cpp b/Engine/Source/Runtime/Engine/Private/VectorField.cpp
--- a/Engine/Source/Runtime/Engine/Private/VectorField.cpp
+++ b/Engine/Source/Runtime/Engine/Private/VectorField.cpp
@@ -86,8 +86,8 @@ void FVectorFieldInstance::UpdateTransforms(const FMatrix& LocalToWorld)
 ------------------------------------------------------------------------------*/
 
 
-UVectorField::UVectorField(const class FPostConstructInitializeProperties& PCIP)
-	: Super(PCIP)
+UVectorField::UVectorField(const FObjectInitializer& ObjectInitializer)
+	: Super(ObjectInitializer)
 {
 	Intensity = 1.0f;
 }
@@ -252,8 +252,8 @@ private:
 	void* VolumeData;
 };
 
-UVectorFieldStatic::UVectorFieldStatic(const class FPostConstructInitializeProperties& PCIP)
-	: Super(PCIP)
+UVectorFieldStatic::UVectorFieldStatic(const FObjectInitializer& ObjectInitializer)
+	: Super(ObjectInitializer)
 {
 }
 
@@ -440,8 +440,8 @@ private:
 /*------------------------------------------------------------------------------
 	UVectorFieldComponent implementation.
 ------------------------------------------------------------------------------*/
-UVectorFieldComponent::UVectorFieldComponent(const class FPostConstructInitializeProperties& PCIP)
-	: Super(PCIP)
+UVectorFieldComponent::UVectorFieldComponent(const FObjectInitializer& ObjectInitializer)
+	: Super(ObjectInitializer)
 {
 	BodyInstance.bEnableCollision_DEPRECATED = false;
 	SetCollisionProfileName(UCollisionProfile::NoCollision_ProfileName);
@@ -877,8 +877,8 @@ public:
 	}
 };
 
-UVectorFieldAnimated::UVectorFieldAnimated(const class FPostConstructInitializeProperties& PCIP)
-	: Super(PCIP)
+UVectorFieldAnimated::UVectorFieldAnimated(const FObjectInitializer& ObjectInitializer)
+	: Super(ObjectInitializer)
 {
 	VolumeSizeX = 16;
 	VolumeSizeY = 16;
diff --git a/Engine/Source/Runtime/Engine/Private/VectorFieldVolume.cpp b/Engine/Source/Runtime/Engine/Private/VectorFieldVolume.cpp
--- a/Engine/Source/Runtime/Engine/Private/VectorFieldVolume.cpp
+++ b/Engine/Source/Runtime/Engine/Private/VectorFieldVolume.cpp
@@ -3,14 +3,14 @@
 #include "EnginePrivate.h"
 #include "VectorField/VectorFieldVolume.h"
 
-AVectorFieldVolume::AVectorFieldVolume(const class FPostConstructInitializeProperties& PCIP)
-	: Super(PCIP)
+AVectorFieldVolume::AVectorFieldVolume(const FObjectInitializer& ObjectInitializer)
+	: Super(ObjectInitializer)
 {
-	VectorFieldComponent = PCIP.CreateDefaultSubobject<UVectorFieldComponent>(this, TEXT("VectorFieldComponent0"));
+	VectorFieldComponent = ObjectInitializer.CreateDefaultSubobject<UVectorFieldComponent>(this, TEXT("VectorFieldComponent0"));
 	RootComponent = VectorFieldComponent;
 
 #if WITH_EDITORONLY_DATA
-	SpriteComponent = PCIP.CreateEditorOnlyDefaultSubobject<UBillboardComponent>(this, TEXT("Sprite"));
+	SpriteComponent = ObjectInitializer.CreateEditorOnlyDefaultSubobject<UBillboardComponent>(this, TEXT("Sprite"));
 
 	if (!IsRunningCommandlet() && (SpriteComponent != nullptr))
 	{
diff --git a/Engine/Source/Runtime/Engine/Private/Vehicles/TireType.cpp b/Engine/Source/Runtime/Engine/Private/Vehicles/TireType.cpp
--- a/Engine/Source/Runtime/Engine/Private/Vehicles/TireType.cpp
+++ b/Engine/Source/Runtime/Engine/Private/Vehicles/TireType.cpp
@@ -9,8 +9,8 @@
 
 TArray<TWeakObjectPtr<UTireType>> UTireType::AllTireTypes;
 
-UTireType::UTireType(const class FPostConstructInitializeProperties& PCIP)
-	: Super(PCIP)
+UTireType::UTireType(const FObjectInitializer& ObjectInitializer)
+	: Super(ObjectInitializer)
 {
 	// Property initialization
 	FrictionScale = 1.0f;
diff --git a/Engine/Source/Runtime/Engine/Private/Vehicles/VehicleAnimInstance.cpp b/Engine/Source/Runtime/Engine/Private/Vehicles/VehicleAnimInstance.cpp
--- a/Engine/Source/Runtime/Engine/Private/Vehicles/VehicleAnimInstance.cpp
+++ b/Engine/Source/Runtime/Engine/Private/Vehicles/VehicleAnimInstance.cpp
@@ -14,8 +14,8 @@
 // UVehicleAnimInstance
 /////////////////////////////////////////////////////
 
-UVehicleAnimInstance::UVehicleAnimInstance(const class FPostConstructInitializeProperties& PCIP)
-	: Super(PCIP)
+UVehicleAnimInstance::UVehicleAnimInstance(const FObjectInitializer& ObjectInitializer)
+	: Super(ObjectInitializer)
 {
 }
 
diff --git a/Engine/Source/Runtime/Engine/Private/Vehicles/VehicleWheel.cpp b/Engine/Source/Runtime/Engine/Private/Vehicles/VehicleWheel.cpp
--- a/Engine/Source/Runtime/Engine/Private/Vehicles/VehicleWheel.cpp
+++ b/Engine/Source/Runtime/Engine/Private/Vehicles/VehicleWheel.cpp
@@ -13,8 +13,8 @@
 #endif // WITH_PHYSX
 
 
-UVehicleWheel::UVehicleWheel(const class FPostConstructInitializeProperties& PCIP)
-	: Super(PCIP)
+UVehicleWheel::UVehicleWheel(const FObjectInitializer& ObjectInitializer)
+	: Super(ObjectInitializer)
 {
 	static ConstructorHelpers::FObjectFinder<UStaticMesh> CollisionMeshObj(TEXT("/Engine/EngineMeshes/Cylinder"));
 	CollisionMesh = CollisionMeshObj.Object;
diff --git a/Engine/Source/Runtime/Engine/Private/Vehicles/WheeledVehicleMovementComponent.cpp b/Engine/Source/Runtime/Engine/Private/Vehicles/WheeledVehicleMovementComponent.cpp
--- a/Engine/Source/Runtime/Engine/Private/Vehicles/WheeledVehicleMovementComponent.cpp
+++ b/Engine/Source/Runtime/Engine/Private/Vehicles/WheeledVehicleMovementComponent.cpp
@@ -83,8 +83,8 @@ void PTireShader(const void* shaderData, const PxF32 tireFriction,
 
 #endif // WITH_PHYSX
 
-UWheeledVehicleMovementComponent::UWheeledVehicleMovementComponent(const class FPostConstructInitializeProperties& PCIP)
-	: Super(PCIP)
+UWheeledVehicleMovementComponent::UWheeledVehicleMovementComponent(const FObjectInitializer& ObjectInitializer)
+	: Super(ObjectInitializer)
 {
 	Mass = 1500.0f;
 	DragCoefficient = 0.3f;
diff --git a/Engine/Source/Runtime/Engine/Private/Vehicles/WheeledVehicleMovementComponent4W.cpp b/Engine/Source/Runtime/Engine/Private/Vehicles/WheeledVehicleMovementComponent4W.cpp
--- a/Engine/Source/Runtime/Engine/Private/Vehicles/WheeledVehicleMovementComponent4W.cpp
+++ b/Engine/Source/Runtime/Engine/Private/Vehicles/WheeledVehicleMovementComponent4W.cpp
@@ -8,7 +8,7 @@
 #include "../PhysicsEngine/PhysXSupport.h"
 #endif // WITH_VEHICLE
 
-UWheeledVehicleMovementComponent4W::UWheeledVehicleMovementComponent4W(const class FPostConstructInitializeProperties& PCIP) : Super(PCIP)
+UWheeledVehicleMovementComponent4W::UWheeledVehicleMovementComponent4W(const FObjectInitializer& ObjectInitializer) : Super(ObjectInitializer)
 {
 #if WITH_VEHICLE
 	// grab default values from physx
diff --git a/Engine/Source/Runtime/Engine/Private/VisualLogger/VisualLoggerAutomationTests.cpp b/Engine/Source/Runtime/Engine/Private/VisualLogger/VisualLoggerAutomationTests.cpp
--- a/Engine/Source/Runtime/Engine/Private/VisualLogger/VisualLoggerAutomationTests.cpp
+++ b/Engine/Source/Runtime/Engine/Private/VisualLogger/VisualLoggerAutomationTests.cpp
@@ -6,7 +6,7 @@
 #include "VisualLogger/VisualLogger.h"
 #include "VisualLogger/VisualLoggerAutomationTests.h"
 
-UVisualLoggerAutomationTests::UVisualLoggerAutomationTests(const class FPostConstructInitializeProperties& PCIP) : Super(PCIP)
+UVisualLoggerAutomationTests::UVisualLoggerAutomationTests(const FObjectInitializer& ObjectInitializer) : Super(ObjectInitializer)
 {
 
 }
diff --git a/Engine/Source/Runtime/Engine/Private/VisualLogger/VisualLoggerKismetLibrary.cpp b/Engine/Source/Runtime/Engine/Private/VisualLogger/VisualLoggerKismetLibrary.cpp
--- a/Engine/Source/Runtime/Engine/Private/VisualLogger/VisualLoggerKismetLibrary.cpp
+++ b/Engine/Source/Runtime/Engine/Private/VisualLogger/VisualLoggerKismetLibrary.cpp
@@ -5,8 +5,8 @@
 #include "VisualLogger/VisualLogger.h"
 #include "VisualLogger/VisualLoggerKismetLibrary.h"
 
-UVisualLoggerKismetLibrary::UVisualLoggerKismetLibrary(const class FPostConstructInitializeProperties& PCIP)
-	: Super(PCIP)
+UVisualLoggerKismetLibrary::UVisualLoggerKismetLibrary(const FObjectInitializer& ObjectInitializer)
+	: Super(ObjectInitializer)
 {
 }
 
diff --git a/Engine/Source/Runtime/Engine/Private/Volume.cpp b/Engine/Source/Runtime/Engine/Private/Volume.cpp
--- a/Engine/Source/Runtime/Engine/Private/Volume.cpp
+++ b/Engine/Source/Runtime/Engine/Private/Volume.cpp
@@ -13,8 +13,8 @@ AVolume::FOnVolumeShapeChanged AVolume::OnVolumeShapeChanged;
 
 DEFINE_LOG_CATEGORY_STATIC(LogVolume, Log, All);
 
-AVolume::AVolume(const class FPostConstructInitializeProperties& PCIP)
-	: Super(PCIP)
+AVolume::AVolume(const FObjectInitializer& ObjectInitializer)
+	: Super(ObjectInitializer)
 {
 	BrushComponent->AlwaysLoadOnClient = true;
 	BrushComponent->AlwaysLoadOnServer = true;
diff --git a/Engine/Source/Runtime/Engine/Private/WheeledVehicle.cpp b/Engine/Source/Runtime/Engine/Private/WheeledVehicle.cpp
--- a/Engine/Source/Runtime/Engine/Private/WheeledVehicle.cpp
+++ b/Engine/Source/Runtime/Engine/Private/WheeledVehicle.cpp
@@ -13,10 +13,10 @@
 FName AWheeledVehicle::VehicleMovementComponentName(TEXT("MovementComp"));
 FName AWheeledVehicle::VehicleMeshComponentName(TEXT("VehicleMesh"));
 
-AWheeledVehicle::AWheeledVehicle(const class FPostConstructInitializeProperties& PCIP)
-	: Super(PCIP)
+AWheeledVehicle::AWheeledVehicle(const FObjectInitializer& ObjectInitializer)
+	: Super(ObjectInitializer)
 {
-	Mesh = PCIP.CreateDefaultSubobject<USkeletalMeshComponent>(this, VehicleMeshComponentName);
+	Mesh = ObjectInitializer.CreateDefaultSubobject<USkeletalMeshComponent>(this, VehicleMeshComponentName);
 	Mesh->SetCollisionProfileName(UCollisionProfile::Vehicle_ProfileName);
 	Mesh->BodyInstance.bSimulatePhysics = true;
 	Mesh->BodyInstance.bNotifyRigidBodyCollision = true;
@@ -26,7 +26,7 @@ AWheeledVehicle::AWheeledVehicle(const class FPostConstructInitializeProperties&
 	Mesh->bCanEverAffectNavigation = false;
 	RootComponent = Mesh;
 
-	VehicleMovement = PCIP.CreateDefaultSubobject<UWheeledVehicleMovementComponent, UWheeledVehicleMovementComponent4W>(this, VehicleMovementComponentName);
+	VehicleMovement = ObjectInitializer.CreateDefaultSubobject<UWheeledVehicleMovementComponent, UWheeledVehicleMovementComponent4W>(this, VehicleMovementComponentName);
 	VehicleMovement->SetIsReplicated(true); // Enable replication by default
 	VehicleMovement->UpdatedComponent = Mesh;
 }
