Rename FPostConstructInitializeProperties to something simpler

diff --git a/Engine/Source/Runtime/Engine/Private/WindDirectionalSource.cpp b/Engine/Source/Runtime/Engine/Private/WindDirectionalSource.cpp
--- a/Engine/Source/Runtime/Engine/Private/WindDirectionalSource.cpp
+++ b/Engine/Source/Runtime/Engine/Private/WindDirectionalSource.cpp
@@ -3,14 +3,14 @@
 #include "EnginePrivate.h"
 #include "Engine/WindDirectionalSource.h"
 
-AWindDirectionalSource::AWindDirectionalSource(const class FPostConstructInitializeProperties& PCIP)
-	: Super(PCIP)
+AWindDirectionalSource::AWindDirectionalSource(const FObjectInitializer& ObjectInitializer)
+	: Super(ObjectInitializer)
 {
-	Component = PCIP.CreateDefaultSubobject<UWindDirectionalSourceComponent>(this, TEXT("WindDirectionalSourceComponent0"));
+	Component = ObjectInitializer.CreateDefaultSubobject<UWindDirectionalSourceComponent>(this, TEXT("WindDirectionalSourceComponent0"));
 	RootComponent = Component;
 
 #if WITH_EDITORONLY_DATA
-	ArrowComponent = PCIP.CreateEditorOnlyDefaultSubobject<UArrowComponent>(this, TEXT("ArrowComponent0"));
+	ArrowComponent = ObjectInitializer.CreateEditorOnlyDefaultSubobject<UArrowComponent>(this, TEXT("ArrowComponent0"));
 
 	if (!IsRunningCommandlet())
 	{
@@ -98,8 +98,8 @@ void FWindSourceSceneProxy::ApplyWorldOffset(FVector InOffset)
 	}
 }
 
-UWindDirectionalSourceComponent::UWindDirectionalSourceComponent(const class FPostConstructInitializeProperties& PCIP)
-	: Super(PCIP)
+UWindDirectionalSourceComponent::UWindDirectionalSourceComponent(const FObjectInitializer& ObjectInitializer)
+	: Super(ObjectInitializer)
 {
 	Strength = 1.0f;
 	Speed = 1.0f;
diff --git a/Engine/Source/Runtime/Engine/Private/World.cpp b/Engine/Source/Runtime/Engine/Private/World.cpp
--- a/Engine/Source/Runtime/Engine/Private/World.cpp
+++ b/Engine/Source/Runtime/Engine/Private/World.cpp
@@ -76,8 +76,8 @@ FWorldDelegates::FWorldEvent FWorldDelegates::OnPreWorldFinishDestroy;
 FWorldDelegates::FOnLevelChanged FWorldDelegates::LevelAddedToWorld;
 FWorldDelegates::FOnLevelChanged FWorldDelegates::LevelRemovedFromWorld;
 
-UWorld::UWorld( const class FPostConstructInitializeProperties& PCIP )
-:	UObject(PCIP)
+UWorld::UWorld( const FObjectInitializer& ObjectInitializer )
+:	UObject(ObjectInitializer)
 ,	FeatureLevel(GMaxRHIFeatureLevel)
 ,	TickTaskLevel(FTickTaskManagerInterface::Get().AllocateTickTaskLevel())
 ,   bIsBuilt(false)
@@ -90,8 +90,8 @@ UWorld::UWorld( const class FPostConstructInitializeProperties& PCIP )
 }
 
 
-UWorld::UWorld( const class FPostConstructInitializeProperties& PCIP,const FURL& InURL )
-:	UObject(PCIP)
+UWorld::UWorld( const FObjectInitializer& ObjectInitializer,const FURL& InURL )
+:	UObject(ObjectInitializer)
 ,	FeatureLevel(GMaxRHIFeatureLevel)
 ,	URL(InURL)
 ,	TickTaskLevel(FTickTaskManagerInterface::Get().AllocateTickTaskLevel())
@@ -937,7 +937,7 @@ void UWorld::InitWorld(const InitializationValues IVS)
 			// Spawn the default brush.
 			DefaultBrush = SpawnBrush();
 			check(DefaultBrush->BrushComponent);
-			DefaultBrush->Brush = new( DefaultBrush->GetOuter(), TEXT("Brush") )UModel( FPostConstructInitializeProperties(), DefaultBrush, 1 );
+			DefaultBrush->Brush = new( DefaultBrush->GetOuter(), TEXT("Brush") )UModel( FObjectInitializer(), DefaultBrush, 1 );
 			DefaultBrush->BrushComponent->Brush = DefaultBrush->Brush;
 			DefaultBrush->SetNotForClientOrServer();
 			DefaultBrush->Brush->SetFlags( RF_Transactional );
@@ -1007,8 +1007,8 @@ void UWorld::InitializeNewWorld(const InitializationValues IVS)
 		ClearFlags(RF_Transactional);
 	}
 
-	PersistentLevel			= new( this, TEXT("PersistentLevel") ) ULevel(FPostConstructInitializeProperties(),FURL(NULL));
-	PersistentLevel->Model	= new( PersistentLevel				 ) UModel(FPostConstructInitializeProperties(),NULL, 1 );
+	PersistentLevel			= new( this, TEXT("PersistentLevel") ) ULevel(FObjectInitializer(),FURL(NULL));
+	PersistentLevel->Model	= new( PersistentLevel				 ) UModel(FObjectInitializer(),NULL, 1 );
 	PersistentLevel->OwningWorld = this;
 
 	// Mark objects are transactional for undo/ redo.
@@ -1093,7 +1093,7 @@ UWorld* UWorld::CreateWorld( const EWorldType::Type InWorldType, bool bInformEng
 
 	// Create new UWorld, ULevel and UModel.
 	const FString WorldNameString = (WorldName != NAME_None) ? WorldName.ToString() : TEXT("NewWorld");
-	UWorld* NewWorld = new( WorldPackage				, *WorldNameString			) UWorld(FPostConstructInitializeProperties(),FURL(NULL));
+	UWorld* NewWorld = new( WorldPackage				, *WorldNameString			) UWorld(FObjectInitializer(),FURL(NULL));
 	NewWorld->WorldType = InWorldType;
 	NewWorld->InitializeNewWorld(UWorld::InitializationValues().ShouldSimulatePhysics(false).EnableTraceCollision(true).CreateNavigation(InWorldType == EWorldType::Editor).CreateAISystem(InWorldType == EWorldType::Editor));
 
@@ -2707,7 +2707,7 @@ bool UWorld::HandleDemoPlayCommand( const TCHAR* Cmd, FOutputDevice& Ar, UWorld*
 		DemoURL.Map = DemoDir + TEXT( "/" ) + Temp + TEXT( ".demo" );
 
 #if 0
-		UPendingNetGame * NewPendingNetGame = new UDemoPendingNetGame( FPostConstructInitializeProperties(), DemoURL );
+		UPendingNetGame * NewPendingNetGame = new UDemoPendingNetGame( FObjectInitializer(), DemoURL );
 
 		if ( !NewPendingNetGame->DemoNetDriver )
 		{
diff --git a/Engine/Source/Runtime/Engine/Private/WorldComposition.cpp b/Engine/Source/Runtime/Engine/Private/WorldComposition.cpp
--- a/Engine/Source/Runtime/Engine/Private/WorldComposition.cpp
+++ b/Engine/Source/Runtime/Engine/Private/WorldComposition.cpp
@@ -14,8 +14,8 @@ UWorldComposition::FWorldCompositionEvent UWorldComposition::OnWorldCompositionC
 UWorldComposition::FWorldCompositionEvent UWorldComposition::OnWorldCompositionDestroyed;
 #endif // WITH_EDITOR
 
-UWorldComposition::UWorldComposition(const class FPostConstructInitializeProperties& PCIP)
-	: Super(PCIP)
+UWorldComposition::UWorldComposition(const FObjectInitializer& ObjectInitializer)
+	: Super(ObjectInitializer)
 {
 }
 
diff --git a/Engine/Source/Runtime/Engine/Private/WorldSettings.cpp b/Engine/Source/Runtime/Engine/Private/WorldSettings.cpp
--- a/Engine/Source/Runtime/Engine/Private/WorldSettings.cpp
+++ b/Engine/Source/Runtime/Engine/Private/WorldSettings.cpp
@@ -13,8 +13,8 @@
 
 #define LOCTEXT_NAMESPACE "ErrorChecking"
 
-AWorldSettings::AWorldSettings(const class FPostConstructInitializeProperties& PCIP)
-	: Super(PCIP.DoNotCreateDefaultSubobject(TEXT("Sprite")))
+AWorldSettings::AWorldSettings(const FObjectInitializer& ObjectInitializer)
+	: Super(ObjectInitializer.DoNotCreateDefaultSubobject(TEXT("Sprite")))
 {
 	// Structure to hold one-time initialization
 	struct FConstructorStatics
@@ -56,7 +56,7 @@ AWorldSettings::AWorldSettings(const class FPostConstructInitializeProperties& P
 	VisibilityAggressiveness = VIS_LeastAggressive;
 	LevelLightingQuality = Quality_MAX;
 
-	TSubobjectPtr<UStaticMeshComponent> StaticMeshComponent = PCIP.CreateDefaultSubobject<UStaticMeshComponent>(this, TEXT("StaticMeshComponent0"));
+	TSubobjectPtr<UStaticMeshComponent> StaticMeshComponent = ObjectInitializer.CreateDefaultSubobject<UStaticMeshComponent>(this, TEXT("StaticMeshComponent0"));
 	StaticMeshComponent->bHiddenInGame = true;
 	StaticMeshComponent->BodyInstance.bEnableCollision_DEPRECATED = true;
 	StaticMeshComponent->SetCollisionProfileName(UCollisionProfile::BlockAll_ProfileName);
diff --git a/Engine/Source/Runtime/Engine/Public/Model.h b/Engine/Source/Runtime/Engine/Public/Model.h
--- a/Engine/Source/Runtime/Engine/Public/Model.h
+++ b/Engine/Source/Runtime/Engine/Public/Model.h
@@ -410,8 +410,8 @@ private:
 public:
 
 	// Constructors.
-	UModel(const class FPostConstructInitializeProperties& PCIP)
-	: UObject(PCIP)
+	UModel(const FObjectInitializer& ObjectInitializer)
+	: UObject(ObjectInitializer)
 	, Nodes( this )
 	, Verts( this )
 	, Vectors( this )
@@ -432,7 +432,7 @@ public:
 			}
 		}
 	}
-	ENGINE_API UModel( const class FPostConstructInitializeProperties& PCIP,ABrush* Owner, bool InRootOutside=1 );
+	ENGINE_API UModel( const FObjectInitializer& ObjectInitializer,ABrush* Owner, bool InRootOutside=1 );
 
 	// UObject interface.
 	virtual void Serialize( FArchive& Ar );	
diff --git a/Engine/Source/Runtime/EngineSettings/Private/EngineSettingsModule.cpp b/Engine/Source/Runtime/EngineSettings/Private/EngineSettingsModule.cpp
--- a/Engine/Source/Runtime/EngineSettings/Private/EngineSettingsModule.cpp
+++ b/Engine/Source/Runtime/EngineSettings/Private/EngineSettingsModule.cpp
@@ -26,41 +26,41 @@ public:
 /* Class constructors
  *****************************************************************************/
 
-UConsoleSettings::UConsoleSettings( const class FPostConstructInitializeProperties& PCIP )
-	: Super(PCIP)
+UConsoleSettings::UConsoleSettings( const FObjectInitializer& ObjectInitializer )
+	: Super(ObjectInitializer)
 { }
 
 
-UGameMapsSettings::UGameMapsSettings( const class FPostConstructInitializeProperties& PCIP )
-	: Super(PCIP)
+UGameMapsSettings::UGameMapsSettings( const FObjectInitializer& ObjectInitializer )
+	: Super(ObjectInitializer)
 	, bUseSplitscreen(true)
 	, TwoPlayerSplitscreenLayout(ETwoPlayerSplitScreenType::Horizontal)
 	, ThreePlayerSplitscreenLayout(EThreePlayerSplitScreenType::FavorTop)
 { }
 
 
-UGameNetworkManagerSettings::UGameNetworkManagerSettings( const class FPostConstructInitializeProperties& PCIP )
-	: Super(PCIP)
+UGameNetworkManagerSettings::UGameNetworkManagerSettings( const FObjectInitializer& ObjectInitializer )
+	: Super(ObjectInitializer)
 { }
 
 
-UGameSessionSettings::UGameSessionSettings( const class FPostConstructInitializeProperties& PCIP )
-	: Super(PCIP)
+UGameSessionSettings::UGameSessionSettings( const FObjectInitializer& ObjectInitializer )
+	: Super(ObjectInitializer)
 { }
 
 
-UGeneralEngineSettings::UGeneralEngineSettings( const class FPostConstructInitializeProperties& PCIP )
-	: Super(PCIP)
+UGeneralEngineSettings::UGeneralEngineSettings( const FObjectInitializer& ObjectInitializer )
+	: Super(ObjectInitializer)
 { }
 
 
-UGeneralProjectSettings::UGeneralProjectSettings( const class FPostConstructInitializeProperties& PCIP )
-	: Super(PCIP)
+UGeneralProjectSettings::UGeneralProjectSettings( const FObjectInitializer& ObjectInitializer )
+	: Super(ObjectInitializer)
 { }
 
 
-UHudSettings::UHudSettings( const class FPostConstructInitializeProperties& PCIP )
-	: Super(PCIP)
+UHudSettings::UHudSettings( const FObjectInitializer& ObjectInitializer )
+	: Super(ObjectInitializer)
 { }
 
 
diff --git a/Engine/Source/Runtime/GameLiveStreaming/GameLiveStreamingFunctionLibrary.cpp b/Engine/Source/Runtime/GameLiveStreaming/GameLiveStreamingFunctionLibrary.cpp
--- a/Engine/Source/Runtime/GameLiveStreaming/GameLiveStreamingFunctionLibrary.cpp
+++ b/Engine/Source/Runtime/GameLiveStreaming/GameLiveStreamingFunctionLibrary.cpp
@@ -6,8 +6,8 @@
 #define LOCTEXT_NAMESPACE "GameLiveStreaming"
 
 
-UGameLiveStreamingFunctionLibrary::UGameLiveStreamingFunctionLibrary( const FPostConstructInitializeProperties& PCIP )
-	: Super(PCIP)
+UGameLiveStreamingFunctionLibrary::UGameLiveStreamingFunctionLibrary( const FObjectInitializer& ObjectInitializer )
+	: Super(ObjectInitializer)
 {
 }
 
diff --git a/Engine/Source/Runtime/GameLiveStreaming/QueryLiveStreamsCallbackProxy.cpp b/Engine/Source/Runtime/GameLiveStreaming/QueryLiveStreamsCallbackProxy.cpp
--- a/Engine/Source/Runtime/GameLiveStreaming/QueryLiveStreamsCallbackProxy.cpp
+++ b/Engine/Source/Runtime/GameLiveStreaming/QueryLiveStreamsCallbackProxy.cpp
@@ -4,8 +4,8 @@
 #include "Runtime/Engine/Public/Features/ILiveStreamingService.h"
 #include "IGameLiveStreaming.h"
 
-UQueryLiveStreamsCallbackProxy::UQueryLiveStreamsCallbackProxy( const FPostConstructInitializeProperties& PCIP )
-	: Super(PCIP)
+UQueryLiveStreamsCallbackProxy::UQueryLiveStreamsCallbackProxy( const FObjectInitializer& ObjectInitializer )
+	: Super(ObjectInitializer)
 {
 }
 
diff --git a/Engine/Source/Runtime/GameMenuBuilder/Private/GameMenuWidgetStyle.cpp b/Engine/Source/Runtime/GameMenuBuilder/Private/GameMenuWidgetStyle.cpp
--- a/Engine/Source/Runtime/GameMenuBuilder/Private/GameMenuWidgetStyle.cpp
+++ b/Engine/Source/Runtime/GameMenuBuilder/Private/GameMenuWidgetStyle.cpp
@@ -36,8 +36,8 @@ void FGameMenuStyle::GetResources(TArray<const FSlateBrush*>& OutBrushes) const
 }
 
 
-UGameMenuWidgetStyle::UGameMenuWidgetStyle(const class FPostConstructInitializeProperties& PCIP)
-	: Super(PCIP)
+UGameMenuWidgetStyle::UGameMenuWidgetStyle(const FObjectInitializer& ObjectInitializer)
+	: Super(ObjectInitializer)
 {
 	
 }
diff --git a/Engine/Source/Runtime/GameplayAbilities/Private/Abilities/GameplayAbility.cpp b/Engine/Source/Runtime/GameplayAbilities/Private/Abilities/GameplayAbility.cpp
--- a/Engine/Source/Runtime/GameplayAbilities/Private/Abilities/GameplayAbility.cpp
+++ b/Engine/Source/Runtime/GameplayAbilities/Private/Abilities/GameplayAbility.cpp
@@ -14,8 +14,8 @@
 //
 // --------------------------------------------------------------------------------------------------------------------------------------------------------
 
-UGameplayAbility::UGameplayAbility(const class FPostConstructInitializeProperties& PCIP)
-	: Super(PCIP)
+UGameplayAbility::UGameplayAbility(const FObjectInitializer& ObjectInitializer)
+	: Super(ObjectInitializer)
 {
 	CostGameplayEffect = NULL;
 	CooldownGameplayEffect = NULL;
diff --git a/Engine/Source/Runtime/GameplayAbilities/Private/Abilities/GameplayAbilitySet.cpp b/Engine/Source/Runtime/GameplayAbilities/Private/Abilities/GameplayAbilitySet.cpp
--- a/Engine/Source/Runtime/GameplayAbilities/Private/Abilities/GameplayAbilitySet.cpp
+++ b/Engine/Source/Runtime/GameplayAbilities/Private/Abilities/GameplayAbilitySet.cpp
@@ -13,8 +13,8 @@
 // --------------------------------------------------------------------------------------------------------------------------------------------------------
 
 
-UGameplayAbilitySet::UGameplayAbilitySet(const class FPostConstructInitializeProperties& PCIP)
-	: Super(PCIP)
+UGameplayAbilitySet::UGameplayAbilitySet(const FObjectInitializer& ObjectInitializer)
+	: Super(ObjectInitializer)
 {
 
 }
diff --git a/Engine/Source/Runtime/GameplayAbilities/Private/Abilities/GameplayAbilityTargetActor.cpp b/Engine/Source/Runtime/GameplayAbilities/Private/Abilities/GameplayAbilityTargetActor.cpp
--- a/Engine/Source/Runtime/GameplayAbilities/Private/Abilities/GameplayAbilityTargetActor.cpp
+++ b/Engine/Source/Runtime/GameplayAbilities/Private/Abilities/GameplayAbilityTargetActor.cpp
@@ -12,8 +12,8 @@
 //
 // --------------------------------------------------------------------------------------------------------------------------------------------------------
 
-AGameplayAbilityTargetActor::AGameplayAbilityTargetActor(const class FPostConstructInitializeProperties& PCIP)
-	: Super(PCIP)
+AGameplayAbilityTargetActor::AGameplayAbilityTargetActor(const FObjectInitializer& ObjectInitializer)
+	: Super(ObjectInitializer)
 {
 	StaticTargetFunction = false;
 	ShouldProduceTargetDataOnServer = false;
diff --git a/Engine/Source/Runtime/GameplayAbilities/Private/Abilities/GameplayAbilityTargetActor_GroundTrace.cpp b/Engine/Source/Runtime/GameplayAbilities/Private/Abilities/GameplayAbilityTargetActor_GroundTrace.cpp
--- a/Engine/Source/Runtime/GameplayAbilities/Private/Abilities/GameplayAbilityTargetActor_GroundTrace.cpp
+++ b/Engine/Source/Runtime/GameplayAbilities/Private/Abilities/GameplayAbilityTargetActor_GroundTrace.cpp
@@ -12,8 +12,8 @@
 //
 // --------------------------------------------------------------------------------------------------------------------------------------------------------
 
-AGameplayAbilityTargetActor_GroundTrace::AGameplayAbilityTargetActor_GroundTrace(const class FPostConstructInitializeProperties& PCIP)
-: Super(PCIP)
+AGameplayAbilityTargetActor_GroundTrace::AGameplayAbilityTargetActor_GroundTrace(const FObjectInitializer& ObjectInitializer)
+: Super(ObjectInitializer)
 {
 	CollisionHeightOffset = 0.0f;
 }
diff --git a/Engine/Source/Runtime/GameplayAbilities/Private/Abilities/GameplayAbilityTargetActor_Radius.cpp b/Engine/Source/Runtime/GameplayAbilities/Private/Abilities/GameplayAbilityTargetActor_Radius.cpp
--- a/Engine/Source/Runtime/GameplayAbilities/Private/Abilities/GameplayAbilityTargetActor_Radius.cpp
+++ b/Engine/Source/Runtime/GameplayAbilities/Private/Abilities/GameplayAbilityTargetActor_Radius.cpp
@@ -13,8 +13,8 @@
 //
 // --------------------------------------------------------------------------------------------------------------------------------------------------------
 
-AGameplayAbilityTargetActor_Radius::AGameplayAbilityTargetActor_Radius(const class FPostConstructInitializeProperties& PCIP)
-	: Super(PCIP)
+AGameplayAbilityTargetActor_Radius::AGameplayAbilityTargetActor_Radius(const FObjectInitializer& ObjectInitializer)
+	: Super(ObjectInitializer)
 {
 	PrimaryActorTick.bCanEverTick = true;
 	PrimaryActorTick.TickGroup = TG_PrePhysics;
diff --git a/Engine/Source/Runtime/GameplayAbilities/Private/Abilities/GameplayAbilityTargetActor_SingleLineTrace.cpp b/Engine/Source/Runtime/GameplayAbilities/Private/Abilities/GameplayAbilityTargetActor_SingleLineTrace.cpp
--- a/Engine/Source/Runtime/GameplayAbilities/Private/Abilities/GameplayAbilityTargetActor_SingleLineTrace.cpp
+++ b/Engine/Source/Runtime/GameplayAbilities/Private/Abilities/GameplayAbilityTargetActor_SingleLineTrace.cpp
@@ -12,8 +12,8 @@
 //
 // --------------------------------------------------------------------------------------------------------------------------------------------------------
 
-AGameplayAbilityTargetActor_SingleLineTrace::AGameplayAbilityTargetActor_SingleLineTrace(const class FPostConstructInitializeProperties& PCIP)
-	: Super(PCIP)
+AGameplayAbilityTargetActor_SingleLineTrace::AGameplayAbilityTargetActor_SingleLineTrace(const FObjectInitializer& ObjectInitializer)
+	: Super(ObjectInitializer)
 {
 }
 
diff --git a/Engine/Source/Runtime/GameplayAbilities/Private/Abilities/GameplayAbilityTargetActor_Trace.cpp b/Engine/Source/Runtime/GameplayAbilities/Private/Abilities/GameplayAbilityTargetActor_Trace.cpp
--- a/Engine/Source/Runtime/GameplayAbilities/Private/Abilities/GameplayAbilityTargetActor_Trace.cpp
+++ b/Engine/Source/Runtime/GameplayAbilities/Private/Abilities/GameplayAbilityTargetActor_Trace.cpp
@@ -12,8 +12,8 @@
 //
 // --------------------------------------------------------------------------------------------------------------------------------------------------------
 
-AGameplayAbilityTargetActor_Trace::AGameplayAbilityTargetActor_Trace(const class FPostConstructInitializeProperties& PCIP)
-	: Super(PCIP)
+AGameplayAbilityTargetActor_Trace::AGameplayAbilityTargetActor_Trace(const FObjectInitializer& ObjectInitializer)
+	: Super(ObjectInitializer)
 {
 	PrimaryActorTick.bCanEverTick = true;
 	PrimaryActorTick.TickGroup = TG_PrePhysics;
diff --git a/Engine/Source/Runtime/GameplayAbilities/Private/Abilities/GameplayAbilityWorldReticle.cpp b/Engine/Source/Runtime/GameplayAbilities/Private/Abilities/GameplayAbilityWorldReticle.cpp
--- a/Engine/Source/Runtime/GameplayAbilities/Private/Abilities/GameplayAbilityWorldReticle.cpp
+++ b/Engine/Source/Runtime/GameplayAbilities/Private/Abilities/GameplayAbilityWorldReticle.cpp
@@ -10,8 +10,8 @@
 //
 // --------------------------------------------------------------------------------------------------------------------------------------------------------
 
-AGameplayAbilityWorldReticle::AGameplayAbilityWorldReticle(const class FPostConstructInitializeProperties& PCIP)
-	: Super(PCIP)
+AGameplayAbilityWorldReticle::AGameplayAbilityWorldReticle(const FObjectInitializer& ObjectInitializer)
+	: Super(ObjectInitializer)
 {
 	PrimaryActorTick.bCanEverTick = true;
 	PrimaryActorTick.TickGroup = TG_PrePhysics;
diff --git a/Engine/Source/Runtime/GameplayAbilities/Private/Abilities/GameplayAbility_CharacterJump.cpp b/Engine/Source/Runtime/GameplayAbilities/Private/Abilities/GameplayAbility_CharacterJump.cpp
--- a/Engine/Source/Runtime/GameplayAbilities/Private/Abilities/GameplayAbility_CharacterJump.cpp
+++ b/Engine/Source/Runtime/GameplayAbilities/Private/Abilities/GameplayAbility_CharacterJump.cpp
@@ -9,8 +9,8 @@
 //
 // --------------------------------------------------------------------------------------------------------------------------------------------------------
 
-UGameplayAbility_CharacterJump::UGameplayAbility_CharacterJump(const class FPostConstructInitializeProperties& PCIP)
-	: Super(PCIP)
+UGameplayAbility_CharacterJump::UGameplayAbility_CharacterJump(const FObjectInitializer& ObjectInitializer)
+	: Super(ObjectInitializer)
 {
 	NetExecutionPolicy = EGameplayAbilityNetExecutionPolicy::Predictive;
 	InstancingPolicy = EGameplayAbilityInstancingPolicy::NonInstanced;
diff --git a/Engine/Source/Runtime/GameplayAbilities/Private/Abilities/GameplayAbility_Montage.cpp b/Engine/Source/Runtime/GameplayAbilities/Private/Abilities/GameplayAbility_Montage.cpp
--- a/Engine/Source/Runtime/GameplayAbilities/Private/Abilities/GameplayAbility_Montage.cpp
+++ b/Engine/Source/Runtime/GameplayAbilities/Private/Abilities/GameplayAbility_Montage.cpp
@@ -11,8 +11,8 @@
 //
 // --------------------------------------------------------------------------------------------------------------------------------------------------------
 
-UGameplayAbility_Montage::UGameplayAbility_Montage(const class FPostConstructInitializeProperties& PCIP)
-	: Super(PCIP)
+UGameplayAbility_Montage::UGameplayAbility_Montage(const FObjectInitializer& ObjectInitializer)
+	: Super(ObjectInitializer)
 {
 	PlayRate = 1.f;
 	NetExecutionPolicy = EGameplayAbilityNetExecutionPolicy::Server;
diff --git a/Engine/Source/Runtime/GameplayAbilities/Private/Abilities/Tasks/AbilityTask.cpp b/Engine/Source/Runtime/GameplayAbilities/Private/Abilities/Tasks/AbilityTask.cpp
--- a/Engine/Source/Runtime/GameplayAbilities/Private/Abilities/Tasks/AbilityTask.cpp
+++ b/Engine/Source/Runtime/GameplayAbilities/Private/Abilities/Tasks/AbilityTask.cpp
@@ -3,8 +3,8 @@
 #include "Abilities/Tasks/AbilityTask.h"
 #include "AbilitySystemComponent.h"
 
-UAbilityTask::UAbilityTask(const class FPostConstructInitializeProperties& PCIP)
-	: Super(PCIP)
+UAbilityTask::UAbilityTask(const FObjectInitializer& ObjectInitializer)
+	: Super(ObjectInitializer)
 {
 	bTickingTask = false;
 	bSimulatedTask = false;
diff --git a/Engine/Source/Runtime/GameplayAbilities/Private/Abilities/Tasks/AbilityTask_MoveToLocation.cpp b/Engine/Source/Runtime/GameplayAbilities/Private/Abilities/Tasks/AbilityTask_MoveToLocation.cpp
--- a/Engine/Source/Runtime/GameplayAbilities/Private/Abilities/Tasks/AbilityTask_MoveToLocation.cpp
+++ b/Engine/Source/Runtime/GameplayAbilities/Private/Abilities/Tasks/AbilityTask_MoveToLocation.cpp
@@ -5,8 +5,8 @@
 #include "GameFramework/CharacterMovementComponent.h"
 #include "Net/UnrealNetwork.h"
 
-UAbilityTask_MoveToLocation::UAbilityTask_MoveToLocation(const class FPostConstructInitializeProperties& PCIP)
-: Super(PCIP)
+UAbilityTask_MoveToLocation::UAbilityTask_MoveToLocation(const FObjectInitializer& ObjectInitializer)
+: Super(ObjectInitializer)
 {
 	bTickingTask = true;
 	bSimulatedTask = true;
diff --git a/Engine/Source/Runtime/GameplayAbilities/Private/Abilities/Tasks/AbilityTask_PlayMontageAndWait.cpp b/Engine/Source/Runtime/GameplayAbilities/Private/Abilities/Tasks/AbilityTask_PlayMontageAndWait.cpp
--- a/Engine/Source/Runtime/GameplayAbilities/Private/Abilities/Tasks/AbilityTask_PlayMontageAndWait.cpp
+++ b/Engine/Source/Runtime/GameplayAbilities/Private/Abilities/Tasks/AbilityTask_PlayMontageAndWait.cpp
@@ -3,8 +3,8 @@
 #include "Abilities/Tasks/AbilityTask_PlayMontageAndWait.h"
 #include "AbilitySystemComponent.h"
 
-UAbilityTask_PlayMontageAndWait::UAbilityTask_PlayMontageAndWait(const class FPostConstructInitializeProperties& PCIP)
-	: Super(PCIP)
+UAbilityTask_PlayMontageAndWait::UAbilityTask_PlayMontageAndWait(const FObjectInitializer& ObjectInitializer)
+	: Super(ObjectInitializer)
 {
 	Rate = 1.f;
 }
diff --git a/Engine/Source/Runtime/GameplayAbilities/Private/Abilities/Tasks/AbilityTask_SpawnActor.cpp b/Engine/Source/Runtime/GameplayAbilities/Private/Abilities/Tasks/AbilityTask_SpawnActor.cpp
--- a/Engine/Source/Runtime/GameplayAbilities/Private/Abilities/Tasks/AbilityTask_SpawnActor.cpp
+++ b/Engine/Source/Runtime/GameplayAbilities/Private/Abilities/Tasks/AbilityTask_SpawnActor.cpp
@@ -5,8 +5,8 @@
 #include "Abilities/Tasks/AbilityTask_SpawnActor.h"
 
 
-UAbilityTask_SpawnActor::UAbilityTask_SpawnActor(const class FPostConstructInitializeProperties& PCIP)
-	: Super(PCIP)
+UAbilityTask_SpawnActor::UAbilityTask_SpawnActor(const FObjectInitializer& ObjectInitializer)
+	: Super(ObjectInitializer)
 {
 	
 }
diff --git a/Engine/Source/Runtime/GameplayAbilities/Private/Abilities/Tasks/AbilityTask_VisualizeTargeting.cpp b/Engine/Source/Runtime/GameplayAbilities/Private/Abilities/Tasks/AbilityTask_VisualizeTargeting.cpp
--- a/Engine/Source/Runtime/GameplayAbilities/Private/Abilities/Tasks/AbilityTask_VisualizeTargeting.cpp
+++ b/Engine/Source/Runtime/GameplayAbilities/Private/Abilities/Tasks/AbilityTask_VisualizeTargeting.cpp
@@ -4,8 +4,8 @@
 #include "AbilitySystemComponent.h"
 #include "Abilities/Tasks/AbilityTask_VisualizeTargeting.h"
 
-UAbilityTask_VisualizeTargeting::UAbilityTask_VisualizeTargeting(const class FPostConstructInitializeProperties& PCIP)
-	: Super(PCIP)
+UAbilityTask_VisualizeTargeting::UAbilityTask_VisualizeTargeting(const FObjectInitializer& ObjectInitializer)
+	: Super(ObjectInitializer)
 {
 
 }
diff --git a/Engine/Source/Runtime/GameplayAbilities/Private/Abilities/Tasks/AbilityTask_WaitAbilityActivate.cpp b/Engine/Source/Runtime/GameplayAbilities/Private/Abilities/Tasks/AbilityTask_WaitAbilityActivate.cpp
--- a/Engine/Source/Runtime/GameplayAbilities/Private/Abilities/Tasks/AbilityTask_WaitAbilityActivate.cpp
+++ b/Engine/Source/Runtime/GameplayAbilities/Private/Abilities/Tasks/AbilityTask_WaitAbilityActivate.cpp
@@ -6,8 +6,8 @@
 
 #include "Abilities/GameplayAbility.h"
 
-UAbilityTask_WaitAbilityActivate::UAbilityTask_WaitAbilityActivate(const class FPostConstructInitializeProperties& PCIP)
-	: Super(PCIP)
+UAbilityTask_WaitAbilityActivate::UAbilityTask_WaitAbilityActivate(const FObjectInitializer& ObjectInitializer)
+	: Super(ObjectInitializer)
 {
 	IncludeTriggeredAbilities = false;
 }
diff --git a/Engine/Source/Runtime/GameplayAbilities/Private/Abilities/Tasks/AbilityTask_WaitAbilityCommit.cpp b/Engine/Source/Runtime/GameplayAbilities/Private/Abilities/Tasks/AbilityTask_WaitAbilityCommit.cpp
--- a/Engine/Source/Runtime/GameplayAbilities/Private/Abilities/Tasks/AbilityTask_WaitAbilityCommit.cpp
+++ b/Engine/Source/Runtime/GameplayAbilities/Private/Abilities/Tasks/AbilityTask_WaitAbilityCommit.cpp
@@ -7,8 +7,8 @@
 #include "Abilities/GameplayAbility.h"
 
 
-UAbilityTask_WaitAbilityCommit::UAbilityTask_WaitAbilityCommit(const class FPostConstructInitializeProperties& PCIP)
-	: Super(PCIP)
+UAbilityTask_WaitAbilityCommit::UAbilityTask_WaitAbilityCommit(const FObjectInitializer& ObjectInitializer)
+	: Super(ObjectInitializer)
 {
 }
 
diff --git a/Engine/Source/Runtime/GameplayAbilities/Private/Abilities/Tasks/AbilityTask_WaitAttributeChange.cpp b/Engine/Source/Runtime/GameplayAbilities/Private/Abilities/Tasks/AbilityTask_WaitAttributeChange.cpp
--- a/Engine/Source/Runtime/GameplayAbilities/Private/Abilities/Tasks/AbilityTask_WaitAttributeChange.cpp
+++ b/Engine/Source/Runtime/GameplayAbilities/Private/Abilities/Tasks/AbilityTask_WaitAttributeChange.cpp
@@ -7,8 +7,8 @@
 #include "Abilities/GameplayAbility.h"
 #include "GameplayEffectExtension.h"
 
-UAbilityTask_WaitAttributeChange::UAbilityTask_WaitAttributeChange(const class FPostConstructInitializeProperties& PCIP)
-	: Super(PCIP)
+UAbilityTask_WaitAttributeChange::UAbilityTask_WaitAttributeChange(const FObjectInitializer& ObjectInitializer)
+	: Super(ObjectInitializer)
 {
 }
 
diff --git a/Engine/Source/Runtime/GameplayAbilities/Private/Abilities/Tasks/AbilityTask_WaitConfirm.cpp b/Engine/Source/Runtime/GameplayAbilities/Private/Abilities/Tasks/AbilityTask_WaitConfirm.cpp
--- a/Engine/Source/Runtime/GameplayAbilities/Private/Abilities/Tasks/AbilityTask_WaitConfirm.cpp
+++ b/Engine/Source/Runtime/GameplayAbilities/Private/Abilities/Tasks/AbilityTask_WaitConfirm.cpp
@@ -2,8 +2,8 @@
 #include "AbilitySystemPrivatePCH.h"
 #include "Abilities/Tasks/AbilityTask_WaitConfirm.h"
 
-UAbilityTask_WaitConfirm::UAbilityTask_WaitConfirm(const class FPostConstructInitializeProperties& PCIP)
-	: Super(PCIP)
+UAbilityTask_WaitConfirm::UAbilityTask_WaitConfirm(const FObjectInitializer& ObjectInitializer)
+	: Super(ObjectInitializer)
 {
 	RegisteredCallback = false;
 }
diff --git a/Engine/Source/Runtime/GameplayAbilities/Private/Abilities/Tasks/AbilityTask_WaitConfirmCancel.cpp b/Engine/Source/Runtime/GameplayAbilities/Private/Abilities/Tasks/AbilityTask_WaitConfirmCancel.cpp
--- a/Engine/Source/Runtime/GameplayAbilities/Private/Abilities/Tasks/AbilityTask_WaitConfirmCancel.cpp
+++ b/Engine/Source/Runtime/GameplayAbilities/Private/Abilities/Tasks/AbilityTask_WaitConfirmCancel.cpp
@@ -6,8 +6,8 @@
 
 #include "Abilities/GameplayAbility.h"
 
-UAbilityTask_WaitConfirmCancel::UAbilityTask_WaitConfirmCancel(const class FPostConstructInitializeProperties& PCIP)
-	: Super(PCIP)
+UAbilityTask_WaitConfirmCancel::UAbilityTask_WaitConfirmCancel(const FObjectInitializer& ObjectInitializer)
+	: Super(ObjectInitializer)
 {
 	RegisteredCallbacks = false;
 
diff --git a/Engine/Source/Runtime/GameplayAbilities/Private/Abilities/Tasks/AbilityTask_WaitDelay.cpp b/Engine/Source/Runtime/GameplayAbilities/Private/Abilities/Tasks/AbilityTask_WaitDelay.cpp
--- a/Engine/Source/Runtime/GameplayAbilities/Private/Abilities/Tasks/AbilityTask_WaitDelay.cpp
+++ b/Engine/Source/Runtime/GameplayAbilities/Private/Abilities/Tasks/AbilityTask_WaitDelay.cpp
@@ -2,8 +2,8 @@
 #include "AbilitySystemPrivatePCH.h"
 #include "Abilities/Tasks/AbilityTask_WaitDelay.h"
 
-UAbilityTask_WaitDelay::UAbilityTask_WaitDelay(const class FPostConstructInitializeProperties& PCIP)
-	: Super(PCIP)
+UAbilityTask_WaitDelay::UAbilityTask_WaitDelay(const FObjectInitializer& ObjectInitializer)
+	: Super(ObjectInitializer)
 {
 	Time = 0.f;
 	TimeStarted = 0.f;
diff --git a/Engine/Source/Runtime/GameplayAbilities/Private/Abilities/Tasks/AbilityTask_WaitGameplayEffectRemoved.cpp b/Engine/Source/Runtime/GameplayAbilities/Private/Abilities/Tasks/AbilityTask_WaitGameplayEffectRemoved.cpp
--- a/Engine/Source/Runtime/GameplayAbilities/Private/Abilities/Tasks/AbilityTask_WaitGameplayEffectRemoved.cpp
+++ b/Engine/Source/Runtime/GameplayAbilities/Private/Abilities/Tasks/AbilityTask_WaitGameplayEffectRemoved.cpp
@@ -6,8 +6,8 @@
 
 #include "Abilities/GameplayAbility.h"
 
-UAbilityTask_WaitGameplayEffectRemoved::UAbilityTask_WaitGameplayEffectRemoved(const class FPostConstructInitializeProperties& PCIP)
-	: Super(PCIP)
+UAbilityTask_WaitGameplayEffectRemoved::UAbilityTask_WaitGameplayEffectRemoved(const FObjectInitializer& ObjectInitializer)
+	: Super(ObjectInitializer)
 {
 	Registered = false;
 }
diff --git a/Engine/Source/Runtime/GameplayAbilities/Private/Abilities/Tasks/AbilityTask_WaitInputRelease.cpp b/Engine/Source/Runtime/GameplayAbilities/Private/Abilities/Tasks/AbilityTask_WaitInputRelease.cpp
--- a/Engine/Source/Runtime/GameplayAbilities/Private/Abilities/Tasks/AbilityTask_WaitInputRelease.cpp
+++ b/Engine/Source/Runtime/GameplayAbilities/Private/Abilities/Tasks/AbilityTask_WaitInputRelease.cpp
@@ -3,8 +3,8 @@
 #include "Abilities/Tasks/AbilityTask_WaitInputRelease.h"
 #include "GameplayPrediction.h"
 
-UAbilityTask_WaitInputRelease::UAbilityTask_WaitInputRelease(const class FPostConstructInitializeProperties& PCIP)
-	: Super(PCIP)
+UAbilityTask_WaitInputRelease::UAbilityTask_WaitInputRelease(const FObjectInitializer& ObjectInitializer)
+	: Super(ObjectInitializer)
 {
 	StartTime = 0.f;
 	RegisteredCallback= false;
diff --git a/Engine/Source/Runtime/GameplayAbilities/Private/Abilities/Tasks/AbilityTask_WaitMovementModeChange.cpp b/Engine/Source/Runtime/GameplayAbilities/Private/Abilities/Tasks/AbilityTask_WaitMovementModeChange.cpp
--- a/Engine/Source/Runtime/GameplayAbilities/Private/Abilities/Tasks/AbilityTask_WaitMovementModeChange.cpp
+++ b/Engine/Source/Runtime/GameplayAbilities/Private/Abilities/Tasks/AbilityTask_WaitMovementModeChange.cpp
@@ -2,8 +2,8 @@
 #include "AbilitySystemPrivatePCH.h"
 #include "Abilities/Tasks/AbilityTask_WaitMovementModeChange.h"
 
-UAbilityTask_WaitMovementModeChange::UAbilityTask_WaitMovementModeChange(const class FPostConstructInitializeProperties& PCIP)
-	: Super(PCIP)
+UAbilityTask_WaitMovementModeChange::UAbilityTask_WaitMovementModeChange(const FObjectInitializer& ObjectInitializer)
+	: Super(ObjectInitializer)
 {
 	RequiredMode = MOVE_None;
 }
diff --git a/Engine/Source/Runtime/GameplayAbilities/Private/Abilities/Tasks/AbilityTask_WaitOverlap.cpp b/Engine/Source/Runtime/GameplayAbilities/Private/Abilities/Tasks/AbilityTask_WaitOverlap.cpp
--- a/Engine/Source/Runtime/GameplayAbilities/Private/Abilities/Tasks/AbilityTask_WaitOverlap.cpp
+++ b/Engine/Source/Runtime/GameplayAbilities/Private/Abilities/Tasks/AbilityTask_WaitOverlap.cpp
@@ -2,8 +2,8 @@
 #include "AbilitySystemPrivatePCH.h"
 #include "Abilities/Tasks/AbilityTask_WaitOverlap.h"
 
-UAbilityTask_WaitOverlap::UAbilityTask_WaitOverlap(const class FPostConstructInitializeProperties& PCIP)
-	: Super(PCIP)
+UAbilityTask_WaitOverlap::UAbilityTask_WaitOverlap(const FObjectInitializer& ObjectInitializer)
+	: Super(ObjectInitializer)
 {
 
 }
diff --git a/Engine/Source/Runtime/GameplayAbilities/Private/Abilities/Tasks/AbilityTask_WaitTargetData.cpp b/Engine/Source/Runtime/GameplayAbilities/Private/Abilities/Tasks/AbilityTask_WaitTargetData.cpp
--- a/Engine/Source/Runtime/GameplayAbilities/Private/Abilities/Tasks/AbilityTask_WaitTargetData.cpp
+++ b/Engine/Source/Runtime/GameplayAbilities/Private/Abilities/Tasks/AbilityTask_WaitTargetData.cpp
@@ -4,8 +4,8 @@
 #include "AbilitySystemComponent.h"
 #include "Abilities/Tasks/AbilityTask_WaitTargetData.h"
 
-UAbilityTask_WaitTargetData::UAbilityTask_WaitTargetData(const class FPostConstructInitializeProperties& PCIP)
-	: Super(PCIP)
+UAbilityTask_WaitTargetData::UAbilityTask_WaitTargetData(const FObjectInitializer& ObjectInitializer)
+	: Super(ObjectInitializer)
 {
 	
 }
diff --git a/Engine/Source/Runtime/GameplayAbilities/Private/Abilities/Tasks/AbilityTask_WaitVelocityChange.cpp b/Engine/Source/Runtime/GameplayAbilities/Private/Abilities/Tasks/AbilityTask_WaitVelocityChange.cpp
--- a/Engine/Source/Runtime/GameplayAbilities/Private/Abilities/Tasks/AbilityTask_WaitVelocityChange.cpp
+++ b/Engine/Source/Runtime/GameplayAbilities/Private/Abilities/Tasks/AbilityTask_WaitVelocityChange.cpp
@@ -3,8 +3,8 @@
 #include "AbilitySystemComponent.h"
 #include "AbilityTask_WaitVelocityChange.h"
 
-UAbilityTask_WaitVelocityChange::UAbilityTask_WaitVelocityChange(const class FPostConstructInitializeProperties& PCIP)
-	: Super(PCIP)
+UAbilityTask_WaitVelocityChange::UAbilityTask_WaitVelocityChange(const FObjectInitializer& ObjectInitializer)
+	: Super(ObjectInitializer)
 {
 	bTickingTask = true;
 }
diff --git a/Engine/Source/Runtime/GameplayAbilities/Private/AbilitySystemBlueprintLibrary.cpp b/Engine/Source/Runtime/GameplayAbilities/Private/AbilitySystemBlueprintLibrary.cpp
--- a/Engine/Source/Runtime/GameplayAbilities/Private/AbilitySystemBlueprintLibrary.cpp
+++ b/Engine/Source/Runtime/GameplayAbilities/Private/AbilitySystemBlueprintLibrary.cpp
@@ -10,8 +10,8 @@
 #include "Abilities/Tasks/AbilityTask_WaitConfirmCancel.h"
 #include "LatentActions.h"
 
-UAbilitySystemBlueprintLibrary::UAbilitySystemBlueprintLibrary(const class FPostConstructInitializeProperties& PCIP)
-: Super(PCIP)
+UAbilitySystemBlueprintLibrary::UAbilitySystemBlueprintLibrary(const FObjectInitializer& ObjectInitializer)
+: Super(ObjectInitializer)
 {
 }
 
diff --git a/Engine/Source/Runtime/GameplayAbilities/Private/AbilitySystemComponent.cpp b/Engine/Source/Runtime/GameplayAbilities/Private/AbilitySystemComponent.cpp
--- a/Engine/Source/Runtime/GameplayAbilities/Private/AbilitySystemComponent.cpp
+++ b/Engine/Source/Runtime/GameplayAbilities/Private/AbilitySystemComponent.cpp
@@ -29,8 +29,8 @@ static FAutoConsoleVariableRef CVarDebugGameplayCues(
 /** Enable to log out all render state create, destroy and updatetransform events */
 #define LOG_RENDER_STATE 0
 
-UAbilitySystemComponent::UAbilitySystemComponent(const class FPostConstructInitializeProperties& PCIP)
-	: Super(PCIP)
+UAbilitySystemComponent::UAbilitySystemComponent(const FObjectInitializer& ObjectInitializer)
+	: Super(ObjectInitializer)
 {
 	bWantsInitializeComponent = true;
 
diff --git a/Engine/Source/Runtime/GameplayAbilities/Private/AbilitySystemDebugHUD.cpp b/Engine/Source/Runtime/GameplayAbilities/Private/AbilitySystemDebugHUD.cpp
--- a/Engine/Source/Runtime/GameplayAbilities/Private/AbilitySystemDebugHUD.cpp
+++ b/Engine/Source/Runtime/GameplayAbilities/Private/AbilitySystemDebugHUD.cpp
@@ -5,7 +5,7 @@
 #include "AbilitySystemDebugHUD.h"
 #include "DebugRenderSceneProxy.h"
 
-AAbilitySystemDebugHUD::AAbilitySystemDebugHUD(const class FPostConstructInitializeProperties& PCIP) : Super(PCIP)
+AAbilitySystemDebugHUD::AAbilitySystemDebugHUD(const FObjectInitializer& ObjectInitializer) : Super(ObjectInitializer)
 {
 
 }
diff --git a/Engine/Source/Runtime/GameplayAbilities/Private/AbilitySystemGlobals.cpp b/Engine/Source/Runtime/GameplayAbilities/Private/AbilitySystemGlobals.cpp
--- a/Engine/Source/Runtime/GameplayAbilities/Private/AbilitySystemGlobals.cpp
+++ b/Engine/Source/Runtime/GameplayAbilities/Private/AbilitySystemGlobals.cpp
@@ -8,8 +8,8 @@
 #include "GameplayTagsModule.h"
 #include "GameplayCueInterface.h"
 
-UAbilitySystemGlobals::UAbilitySystemGlobals(const class FPostConstructInitializeProperties& PCIP)
-: Super(PCIP)
+UAbilitySystemGlobals::UAbilitySystemGlobals(const FObjectInitializer& ObjectInitializer)
+: Super(ObjectInitializer)
 {
 	AbilitySystemGlobalsClassName = FStringClassReference(TEXT("/Script/GameplayAbilities.AbilitySystemGlobals"));
 
diff --git a/Engine/Source/Runtime/GameplayAbilities/Private/AbilitySystemInterface.cpp b/Engine/Source/Runtime/GameplayAbilities/Private/AbilitySystemInterface.cpp
--- a/Engine/Source/Runtime/GameplayAbilities/Private/AbilitySystemInterface.cpp
+++ b/Engine/Source/Runtime/GameplayAbilities/Private/AbilitySystemInterface.cpp
@@ -3,7 +3,7 @@
 #include "AbilitySystemPrivatePCH.h"
 #include "AbilitySystemInterface.h"
 
-UAbilitySystemInterface::UAbilitySystemInterface(const FPostConstructInitializeProperties& PCIP)
-: Super(PCIP)
+UAbilitySystemInterface::UAbilitySystemInterface(const FObjectInitializer& ObjectInitializer)
+: Super(ObjectInitializer)
 {
 }
\ No newline at end of file
diff --git a/Engine/Source/Runtime/GameplayAbilities/Private/AbilitySystemTestAttributeSet.cpp b/Engine/Source/Runtime/GameplayAbilities/Private/AbilitySystemTestAttributeSet.cpp
--- a/Engine/Source/Runtime/GameplayAbilities/Private/AbilitySystemTestAttributeSet.cpp
+++ b/Engine/Source/Runtime/GameplayAbilities/Private/AbilitySystemTestAttributeSet.cpp
@@ -7,8 +7,8 @@
 #include "AbilitySystemComponent.h"
 
 
-UAbilitySystemTestAttributeSet::UAbilitySystemTestAttributeSet(const class FPostConstructInitializeProperties& PCIP)
-: Super(PCIP)
+UAbilitySystemTestAttributeSet::UAbilitySystemTestAttributeSet(const FObjectInitializer& ObjectInitializer)
+: Super(ObjectInitializer)
 {
 	Health = MaxHealth = 100.f;
 	Mana = MaxMana = 100.f;
diff --git a/Engine/Source/Runtime/GameplayAbilities/Private/AbilitySystemTestPawn.cpp b/Engine/Source/Runtime/GameplayAbilities/Private/AbilitySystemTestPawn.cpp
--- a/Engine/Source/Runtime/GameplayAbilities/Private/AbilitySystemTestPawn.cpp
+++ b/Engine/Source/Runtime/GameplayAbilities/Private/AbilitySystemTestPawn.cpp
@@ -7,10 +7,10 @@
 
 FName  AAbilitySystemTestPawn::AbilitySystemComponentName(TEXT("AbilitySystemComponent0"));
 
-AAbilitySystemTestPawn::AAbilitySystemTestPawn(const class FPostConstructInitializeProperties& PCIP)
-: Super(PCIP)
+AAbilitySystemTestPawn::AAbilitySystemTestPawn(const FObjectInitializer& ObjectInitializer)
+: Super(ObjectInitializer)
 {
-	AbilitySystemComponent = PCIP.CreateDefaultSubobject<UAbilitySystemComponent>(this, AAbilitySystemTestPawn::AbilitySystemComponentName);
+	AbilitySystemComponent = ObjectInitializer.CreateDefaultSubobject<UAbilitySystemComponent>(this, AAbilitySystemTestPawn::AbilitySystemComponentName);
 	AbilitySystemComponent->SetIsReplicated(true);
 
 	//DefaultAbilitySet = NULL;
diff --git a/Engine/Source/Runtime/GameplayAbilities/Private/AttributeSet.cpp b/Engine/Source/Runtime/GameplayAbilities/Private/AttributeSet.cpp
--- a/Engine/Source/Runtime/GameplayAbilities/Private/AttributeSet.cpp
+++ b/Engine/Source/Runtime/GameplayAbilities/Private/AttributeSet.cpp
@@ -23,8 +23,8 @@ float FGameplayAttribute::GetNumericValueChecked(const UAttributeSet* Src) const
 	return NumericProperty->GetFloatingPointPropertyValue(ValuePtr);
 }
 
-UAttributeSet::UAttributeSet(const class FPostConstructInitializeProperties& PCIP)
-	: Super(PCIP)
+UAttributeSet::UAttributeSet(const FObjectInitializer& ObjectInitializer)
+	: Super(ObjectInitializer)
 {
 
 }
diff --git a/Engine/Source/Runtime/GameplayAbilities/Private/GameplayAbilityBlueprint.cpp b/Engine/Source/Runtime/GameplayAbilities/Private/GameplayAbilityBlueprint.cpp
--- a/Engine/Source/Runtime/GameplayAbilities/Private/GameplayAbilityBlueprint.cpp
+++ b/Engine/Source/Runtime/GameplayAbilities/Private/GameplayAbilityBlueprint.cpp
@@ -7,8 +7,8 @@
 //////////////////////////////////////////////////////////////////////////
 // UGameplayAbilityBlueprint
 
-UGameplayAbilityBlueprint::UGameplayAbilityBlueprint(const class FPostConstructInitializeProperties& PCIP)
-	: Super(PCIP)
+UGameplayAbilityBlueprint::UGameplayAbilityBlueprint(const FObjectInitializer& ObjectInitializer)
+	: Super(ObjectInitializer)
 {
 }
 
diff --git a/Engine/Source/Runtime/GameplayAbilities/Private/GameplayAbilityBlueprintGeneratedClass.cpp b/Engine/Source/Runtime/GameplayAbilities/Private/GameplayAbilityBlueprintGeneratedClass.cpp
--- a/Engine/Source/Runtime/GameplayAbilities/Private/GameplayAbilityBlueprintGeneratedClass.cpp
+++ b/Engine/Source/Runtime/GameplayAbilities/Private/GameplayAbilityBlueprintGeneratedClass.cpp
@@ -8,8 +8,8 @@
 /////////////////////////////////////////////////////
 // UGameplayAbilityBlueprintGeneratedClass
 
-UGameplayAbilityBlueprintGeneratedClass::UGameplayAbilityBlueprintGeneratedClass(const FPostConstructInitializeProperties& PCIP)
-	: Super(PCIP)
+UGameplayAbilityBlueprintGeneratedClass::UGameplayAbilityBlueprintGeneratedClass(const FObjectInitializer& ObjectInitializer)
+	: Super(ObjectInitializer)
 {
 }
 
diff --git a/Engine/Source/Runtime/GameplayAbilities/Private/GameplayCueActor.cpp b/Engine/Source/Runtime/GameplayAbilities/Private/GameplayCueActor.cpp
--- a/Engine/Source/Runtime/GameplayAbilities/Private/GameplayCueActor.cpp
+++ b/Engine/Source/Runtime/GameplayAbilities/Private/GameplayCueActor.cpp
@@ -3,8 +3,8 @@
 #include "AbilitySystemPrivatePCH.h"
 #include "GameplayCueActor.h"
 
-AGameplayCueActor::AGameplayCueActor(const class FPostConstructInitializeProperties& PCIP)
-: Super(PCIP)
+AGameplayCueActor::AGameplayCueActor(const FObjectInitializer& ObjectInitializer)
+: Super(ObjectInitializer)
 {
 	AttachToOwner = false;
 }
diff --git a/Engine/Source/Runtime/GameplayAbilities/Private/GameplayCueInterface.cpp b/Engine/Source/Runtime/GameplayAbilities/Private/GameplayCueInterface.cpp
--- a/Engine/Source/Runtime/GameplayAbilities/Private/GameplayCueInterface.cpp
+++ b/Engine/Source/Runtime/GameplayAbilities/Private/GameplayCueInterface.cpp
@@ -5,8 +5,8 @@
 #include "GameplayCueInterface.h"
 #include "GameplayTagsModule.h"
 
-UGameplayCueInterface::UGameplayCueInterface(const FPostConstructInitializeProperties& PCIP)
-: Super(PCIP)
+UGameplayCueInterface::UGameplayCueInterface(const FObjectInitializer& ObjectInitializer)
+: Super(ObjectInitializer)
 {
 }
 
diff --git a/Engine/Source/Runtime/GameplayAbilities/Private/GameplayCueView.cpp b/Engine/Source/Runtime/GameplayAbilities/Private/GameplayCueView.cpp
--- a/Engine/Source/Runtime/GameplayAbilities/Private/GameplayCueView.cpp
+++ b/Engine/Source/Runtime/GameplayAbilities/Private/GameplayCueView.cpp
@@ -14,8 +14,8 @@
 #include "GameplayCueActor.h"
 #include "AbilitySystemComponent.h"
 
-UGameplayCueView::UGameplayCueView(const class FPostConstructInitializeProperties& PCIP)
-: Super(PCIP)
+UGameplayCueView::UGameplayCueView(const FObjectInitializer& ObjectInitializer)
+: Super(ObjectInitializer)
 {
 
 }
diff --git a/Engine/Source/Runtime/GameplayAbilities/Private/GameplayEffect.cpp b/Engine/Source/Runtime/GameplayAbilities/Private/GameplayEffect.cpp
--- a/Engine/Source/Runtime/GameplayAbilities/Private/GameplayEffect.cpp
+++ b/Engine/Source/Runtime/GameplayAbilities/Private/GameplayEffect.cpp
@@ -25,8 +25,8 @@ FAggregator::FAllocationStats FAggregator::AllocationStats;
 // --------------------------------------------------------------------------------------------------------------------------------------------------------
 
 
-UGameplayEffect::UGameplayEffect(const class FPostConstructInitializeProperties& PCIP)
-	: Super(PCIP)
+UGameplayEffect::UGameplayEffect(const FObjectInitializer& ObjectInitializer)
+	: Super(ObjectInitializer)
 {
 	ChanceToApplyToTarget.SetValue(1.f);
 	ChanceToExecuteOnGameplayEffect.SetValue(1.f);
diff --git a/Engine/Source/Runtime/GameplayAbilities/Private/GameplayEffectExtension.cpp b/Engine/Source/Runtime/GameplayAbilities/Private/GameplayEffectExtension.cpp
--- a/Engine/Source/Runtime/GameplayAbilities/Private/GameplayEffectExtension.cpp
+++ b/Engine/Source/Runtime/GameplayAbilities/Private/GameplayEffectExtension.cpp
@@ -3,8 +3,8 @@
 #include "AbilitySystemPrivatePCH.h"
 #include "GameplayEffectExtension.h"
 
-UGameplayEffectExtension::UGameplayEffectExtension(const class FPostConstructInitializeProperties& PCIP)
-: Super(PCIP)
+UGameplayEffectExtension::UGameplayEffectExtension(const FObjectInitializer& ObjectInitializer)
+: Super(ObjectInitializer)
 {
 
 }
diff --git a/Engine/Source/Runtime/GameplayAbilities/Private/GameplayEffectExtension_LifestealTest.cpp b/Engine/Source/Runtime/GameplayAbilities/Private/GameplayEffectExtension_LifestealTest.cpp
--- a/Engine/Source/Runtime/GameplayAbilities/Private/GameplayEffectExtension_LifestealTest.cpp
+++ b/Engine/Source/Runtime/GameplayAbilities/Private/GameplayEffectExtension_LifestealTest.cpp
@@ -6,8 +6,8 @@
 #include "GameplayTagsModule.h"
 #include "AbilitySystemComponent.h"
 
-UGameplayEffectExtension_LifestealTest::UGameplayEffectExtension_LifestealTest(const class FPostConstructInitializeProperties& PCIP)
-: Super(PCIP)
+UGameplayEffectExtension_LifestealTest::UGameplayEffectExtension_LifestealTest(const FObjectInitializer& ObjectInitializer)
+: Super(ObjectInitializer)
 {
 	HealthRestoreGameplayEffect = NULL;
 }
diff --git a/Engine/Source/Runtime/GameplayAbilities/Private/GameplayEffectExtension_ShieldTest.cpp b/Engine/Source/Runtime/GameplayAbilities/Private/GameplayEffectExtension_ShieldTest.cpp
--- a/Engine/Source/Runtime/GameplayAbilities/Private/GameplayEffectExtension_ShieldTest.cpp
+++ b/Engine/Source/Runtime/GameplayAbilities/Private/GameplayEffectExtension_ShieldTest.cpp
@@ -6,8 +6,8 @@
 #include "AbilitySystemComponent.h"
 #include "AbilitySystemTestAttributeSet.h"
 
-UGameplayEffectExtension_ShieldTest::UGameplayEffectExtension_ShieldTest(const class FPostConstructInitializeProperties& PCIP)
-: Super(PCIP)
+UGameplayEffectExtension_ShieldTest::UGameplayEffectExtension_ShieldTest(const FObjectInitializer& ObjectInitializer)
+: Super(ObjectInitializer)
 {
 
 }
diff --git a/Engine/Source/Runtime/GameplayAbilities/Private/GameplayEffectStackingExtension.cpp b/Engine/Source/Runtime/GameplayAbilities/Private/GameplayEffectStackingExtension.cpp
--- a/Engine/Source/Runtime/GameplayAbilities/Private/GameplayEffectStackingExtension.cpp
+++ b/Engine/Source/Runtime/GameplayAbilities/Private/GameplayEffectStackingExtension.cpp
@@ -3,8 +3,8 @@
 #include "AbilitySystemPrivatePCH.h"
 #include "GameplayEffectStackingExtension.h"
 
-UGameplayEffectStackingExtension::UGameplayEffectStackingExtension(const class FPostConstructInitializeProperties& PCIP)
-: Super(PCIP)
+UGameplayEffectStackingExtension::UGameplayEffectStackingExtension(const FObjectInitializer& ObjectInitializer)
+: Super(ObjectInitializer)
 {
 	
 }
diff --git a/Engine/Source/Runtime/GameplayAbilities/Private/GameplayEffectStackingExtension_CappedNumberTest.cpp b/Engine/Source/Runtime/GameplayAbilities/Private/GameplayEffectStackingExtension_CappedNumberTest.cpp
--- a/Engine/Source/Runtime/GameplayAbilities/Private/GameplayEffectStackingExtension_CappedNumberTest.cpp
+++ b/Engine/Source/Runtime/GameplayAbilities/Private/GameplayEffectStackingExtension_CappedNumberTest.cpp
@@ -7,8 +7,8 @@
 
 const int32 MaxNumInStack = 2;
 
-UGameplayEffectStackingExtension_CappedNumberTest::UGameplayEffectStackingExtension_CappedNumberTest(const class FPostConstructInitializeProperties& PCIP)
-: Super(PCIP)
+UGameplayEffectStackingExtension_CappedNumberTest::UGameplayEffectStackingExtension_CappedNumberTest(const FObjectInitializer& ObjectInitializer)
+: Super(ObjectInitializer)
 {
 	Handle = FCrc::StrCrc32("UGameplayEffectStackingExtension_CappedNumberTest");
 }
diff --git a/Engine/Source/Runtime/GameplayAbilities/Private/GameplayEffectStackingExtension_DiminishingReturnsTest.cpp b/Engine/Source/Runtime/GameplayAbilities/Private/GameplayEffectStackingExtension_DiminishingReturnsTest.cpp
--- a/Engine/Source/Runtime/GameplayAbilities/Private/GameplayEffectStackingExtension_DiminishingReturnsTest.cpp
+++ b/Engine/Source/Runtime/GameplayAbilities/Private/GameplayEffectStackingExtension_DiminishingReturnsTest.cpp
@@ -14,8 +14,8 @@ UCurveTable* GetCurveTable()
 	return CurveTable;
 }
 
-UGameplayEffectStackingExtension_DiminishingReturnsTest::UGameplayEffectStackingExtension_DiminishingReturnsTest(const class FPostConstructInitializeProperties& PCIP)
-: Super(PCIP)
+UGameplayEffectStackingExtension_DiminishingReturnsTest::UGameplayEffectStackingExtension_DiminishingReturnsTest(const FObjectInitializer& ObjectInitializer)
+: Super(ObjectInitializer)
 {
 	Handle = FCrc::StrCrc32("UGameplayEffectStackingExtension_DiminishingReturnsTest");
 }
diff --git a/Engine/Source/Runtime/GameplayAbilities/Private/GameplayEffectTemplate.cpp b/Engine/Source/Runtime/GameplayAbilities/Private/GameplayEffectTemplate.cpp
--- a/Engine/Source/Runtime/GameplayAbilities/Private/GameplayEffectTemplate.cpp
+++ b/Engine/Source/Runtime/GameplayAbilities/Private/GameplayEffectTemplate.cpp
@@ -3,8 +3,8 @@
 #include "AbilitySystemPrivatePCH.h"
 #include "GameplayEffectTemplate.h"
 
-UGameplayEffectTemplate::UGameplayEffectTemplate(const class FPostConstructInitializeProperties& PCIP)
-: Super(PCIP)
+UGameplayEffectTemplate::UGameplayEffectTemplate(const FObjectInitializer& ObjectInitializer)
+: Super(ObjectInitializer)
 {
 	
 }
\ No newline at end of file
diff --git a/Engine/Source/Runtime/GameplayAbilities/Private/GameplayModCalculation.cpp b/Engine/Source/Runtime/GameplayAbilities/Private/GameplayModCalculation.cpp
--- a/Engine/Source/Runtime/GameplayAbilities/Private/GameplayModCalculation.cpp
+++ b/Engine/Source/Runtime/GameplayAbilities/Private/GameplayModCalculation.cpp
@@ -11,8 +11,8 @@ FGameplayModCalculationSpecHandle FGameplayModCalculationSpecHandle::GenerateNew
 	return NewHandle;
 }
 
-UGameplayModCalculation::UGameplayModCalculation(const class FPostConstructInitializeProperties& PCIP)
-: Super(PCIP)
+UGameplayModCalculation::UGameplayModCalculation(const FObjectInitializer& ObjectInitializer)
+: Super(ObjectInitializer)
 {
 
 }
diff --git a/Engine/Source/Runtime/GameplayAbilities/Public/Abilities/GameplayAbilityTargetActor_ActorPlacement.cpp b/Engine/Source/Runtime/GameplayAbilities/Public/Abilities/GameplayAbilityTargetActor_ActorPlacement.cpp
--- a/Engine/Source/Runtime/GameplayAbilities/Public/Abilities/GameplayAbilityTargetActor_ActorPlacement.cpp
+++ b/Engine/Source/Runtime/GameplayAbilities/Public/Abilities/GameplayAbilityTargetActor_ActorPlacement.cpp
@@ -11,8 +11,8 @@
 //
 // --------------------------------------------------------------------------------------------------------------------------------------------------------
 
-AGameplayAbilityTargetActor_ActorPlacement::AGameplayAbilityTargetActor_ActorPlacement(const class FPostConstructInitializeProperties& PCIP)
-	: Super(PCIP)
+AGameplayAbilityTargetActor_ActorPlacement::AGameplayAbilityTargetActor_ActorPlacement(const FObjectInitializer& ObjectInitializer)
+	: Super(ObjectInitializer)
 {
 }
 
diff --git a/Engine/Source/Runtime/GameplayAbilities/Public/Abilities/GameplayAbilityWorldReticle_ActorVisualization.cpp b/Engine/Source/Runtime/GameplayAbilities/Public/Abilities/GameplayAbilityWorldReticle_ActorVisualization.cpp
--- a/Engine/Source/Runtime/GameplayAbilities/Public/Abilities/GameplayAbilityWorldReticle_ActorVisualization.cpp
+++ b/Engine/Source/Runtime/GameplayAbilities/Public/Abilities/GameplayAbilityWorldReticle_ActorVisualization.cpp
@@ -10,13 +10,13 @@
 //
 // --------------------------------------------------------------------------------------------------------------------------------------------------------
 
-AGameplayAbilityWorldReticle_ActorVisualization::AGameplayAbilityWorldReticle_ActorVisualization(const class FPostConstructInitializeProperties& PCIP)
-	: Super(PCIP)
+AGameplayAbilityWorldReticle_ActorVisualization::AGameplayAbilityWorldReticle_ActorVisualization(const FObjectInitializer& ObjectInitializer)
+	: Super(ObjectInitializer)
 {
 	//PrimaryActorTick.bCanEverTick = true;
 	//PrimaryActorTick.TickGroup = TG_PrePhysics;
 
-	CollisionComponent = PCIP.CreateDefaultSubobject<UCapsuleComponent>(this, TEXT("CollisionCapsule0"));
+	CollisionComponent = ObjectInitializer.CreateDefaultSubobject<UCapsuleComponent>(this, TEXT("CollisionCapsule0"));
 	CollisionComponent->InitCapsuleSize(0.f, 0.f);
 	CollisionComponent->AlwaysLoadOnClient = true;
 	CollisionComponent->bAbsoluteScale = true;
@@ -25,7 +25,7 @@ AGameplayAbilityWorldReticle_ActorVisualization::AGameplayAbilityWorldReticle_Ac
 	CollisionComponent->BodyInstance.bEnableCollision_DEPRECATED = false;
 	CollisionComponent->SetCollisionEnabled(ECollisionEnabled::NoCollision);
 
-	//TSubobjectPtr<USceneComponent> SceneComponent = PCIP.CreateDefaultSubobject<USceneComponent>(this, TEXT("RootComponent0"));
+	//TSubobjectPtr<USceneComponent> SceneComponent = ObjectInitializer.CreateDefaultSubobject<USceneComponent>(this, TEXT("RootComponent0"));
 	RootComponent = CollisionComponent;
 }
 
diff --git a/Engine/Source/Runtime/GameplayTags/Private/BlueprintGameplayTagLibrary.cpp b/Engine/Source/Runtime/GameplayTags/Private/BlueprintGameplayTagLibrary.cpp
--- a/Engine/Source/Runtime/GameplayTags/Private/BlueprintGameplayTagLibrary.cpp
+++ b/Engine/Source/Runtime/GameplayTags/Private/BlueprintGameplayTagLibrary.cpp
@@ -2,8 +2,8 @@
 
 #include "GameplayTagsModulePrivatePCH.h"
 
-UBlueprintGameplayTagLibrary::UBlueprintGameplayTagLibrary(const FPostConstructInitializeProperties& PCIP)
-	: Super(PCIP)
+UBlueprintGameplayTagLibrary::UBlueprintGameplayTagLibrary(const FObjectInitializer& ObjectInitializer)
+	: Super(ObjectInitializer)
 {
 }
 
diff --git a/Engine/Source/Runtime/GameplayTags/Private/GameplayTagAssetInterface.cpp b/Engine/Source/Runtime/GameplayTags/Private/GameplayTagAssetInterface.cpp
--- a/Engine/Source/Runtime/GameplayTags/Private/GameplayTagAssetInterface.cpp
+++ b/Engine/Source/Runtime/GameplayTags/Private/GameplayTagAssetInterface.cpp
@@ -2,8 +2,8 @@
 
 #include "GameplayTagsModulePrivatePCH.h"
 
-UGameplayTagAssetInterface::UGameplayTagAssetInterface(const FPostConstructInitializeProperties& PCIP)
-	: Super(PCIP)
+UGameplayTagAssetInterface::UGameplayTagAssetInterface(const FObjectInitializer& ObjectInitializer)
+	: Super(ObjectInitializer)
 {
 }
 
diff --git a/Engine/Source/Runtime/GameplayTags/Private/GameplayTagsManager.cpp b/Engine/Source/Runtime/GameplayTags/Private/GameplayTagsManager.cpp
--- a/Engine/Source/Runtime/GameplayTags/Private/GameplayTagsManager.cpp
+++ b/Engine/Source/Runtime/GameplayTags/Private/GameplayTagsManager.cpp
@@ -9,8 +9,8 @@
 #include "UnrealEd.h"
 #endif
 
-UGameplayTagsManager::UGameplayTagsManager(const class FPostConstructInitializeProperties& PCIP)
-	: Super(PCIP)
+UGameplayTagsManager::UGameplayTagsManager(const FObjectInitializer& ObjectInitializer)
+	: Super(ObjectInitializer)
 {
 #if WITH_EDITOR
 	RegisteredObjectReimport = false;
diff --git a/Engine/Source/Runtime/GameplayTags/Private/GameplayTagsSettings.cpp b/Engine/Source/Runtime/GameplayTags/Private/GameplayTagsSettings.cpp
--- a/Engine/Source/Runtime/GameplayTags/Private/GameplayTagsSettings.cpp
+++ b/Engine/Source/Runtime/GameplayTags/Private/GameplayTagsSettings.cpp
@@ -3,8 +3,8 @@
 #include "GameplayTagsModulePrivatePCH.h"
 #include "GameplayTagsModule.h"
 
-UGameplayTagsSettings::UGameplayTagsSettings(const class FPostConstructInitializeProperties& PCIP)
-: Super(PCIP)
+UGameplayTagsSettings::UGameplayTagsSettings(const FObjectInitializer& ObjectInitializer)
+: Super(ObjectInitializer)
 {
 
 
diff --git a/Engine/Source/Runtime/IOS/IOSRuntimeSettings/Private/IOSRuntimeSettings.cpp b/Engine/Source/Runtime/IOS/IOSRuntimeSettings/Private/IOSRuntimeSettings.cpp
--- a/Engine/Source/Runtime/IOS/IOSRuntimeSettings/Private/IOSRuntimeSettings.cpp
+++ b/Engine/Source/Runtime/IOS/IOSRuntimeSettings/Private/IOSRuntimeSettings.cpp
@@ -4,8 +4,8 @@
 
 #include "IOSRuntimeSettings.h"
 
-UIOSRuntimeSettings::UIOSRuntimeSettings(const class FPostConstructInitializeProperties& PCIP)
-	: Super(PCIP)
+UIOSRuntimeSettings::UIOSRuntimeSettings(const FObjectInitializer& ObjectInitializer)
+	: Super(ObjectInitializer)
 {
 	bEnableGameCenterSupport = true;
 	bSupportsPortraitOrientation = true;
diff --git a/Engine/Source/Runtime/InputCore/Private/InputCoreTypes.cpp b/Engine/Source/Runtime/InputCore/Private/InputCoreTypes.cpp
--- a/Engine/Source/Runtime/InputCore/Private/InputCoreTypes.cpp
+++ b/Engine/Source/Runtime/InputCore/Private/InputCoreTypes.cpp
@@ -3,8 +3,8 @@
 #include "InputCorePrivatePCH.h"
 #include "PropertyTag.h"
 
-UInputCoreTypes::UInputCoreTypes(const class FPostConstructInitializeProperties& PCIP)
-: Super(PCIP)
+UInputCoreTypes::UInputCoreTypes(const FObjectInitializer& ObjectInitializer)
+: Super(ObjectInitializer)
 {
 }
 
diff --git a/Engine/Source/Runtime/MediaAssets/Private/Assets/MediaPlayer.cpp b/Engine/Source/Runtime/MediaAssets/Private/Assets/MediaPlayer.cpp
--- a/Engine/Source/Runtime/MediaAssets/Private/Assets/MediaPlayer.cpp
+++ b/Engine/Source/Runtime/MediaAssets/Private/Assets/MediaPlayer.cpp
@@ -6,8 +6,8 @@
 /* UMediaPlayer structors
  *****************************************************************************/
 
-UMediaPlayer::UMediaPlayer( const class FPostConstructInitializeProperties& PCIP )
-	: Super(PCIP)
+UMediaPlayer::UMediaPlayer( const FObjectInitializer& ObjectInitializer )
+	: Super(ObjectInitializer)
 	, AutoPlay(false)
 	, AutoPlayRate(1.0f)
 	, Looping(true)
diff --git a/Engine/Source/Runtime/MediaAssets/Private/Assets/MediaSoundWave.cpp b/Engine/Source/Runtime/MediaAssets/Private/Assets/MediaSoundWave.cpp
--- a/Engine/Source/Runtime/MediaAssets/Private/Assets/MediaSoundWave.cpp
+++ b/Engine/Source/Runtime/MediaAssets/Private/Assets/MediaSoundWave.cpp
@@ -6,8 +6,8 @@
 /* UMediaSoundWave structors
  *****************************************************************************/
 
-UMediaSoundWave::UMediaSoundWave( const class FPostConstructInitializeProperties& PCIP )
-	: Super(PCIP)
+UMediaSoundWave::UMediaSoundWave( const FObjectInitializer& ObjectInitializer )
+	: Super(ObjectInitializer)
 	, AudioQueue(MakeShareable(new FMediaSampleQueue))
 {
 	bLooping = false;
diff --git a/Engine/Source/Runtime/MediaAssets/Private/Assets/MediaTexture.cpp b/Engine/Source/Runtime/MediaAssets/Private/Assets/MediaTexture.cpp
--- a/Engine/Source/Runtime/MediaAssets/Private/Assets/MediaTexture.cpp
+++ b/Engine/Source/Runtime/MediaAssets/Private/Assets/MediaTexture.cpp
@@ -7,8 +7,8 @@
 /* UMediaTexture structors
  *****************************************************************************/
 
-UMediaTexture::UMediaTexture( const class FPostConstructInitializeProperties& PCIP )
-	: Super(PCIP)
+UMediaTexture::UMediaTexture( const FObjectInitializer& ObjectInitializer )
+	: Super(ObjectInitializer)
 	, ClearColor(FLinearColor::Black)
 	, MediaPlayer(nullptr)
 	, CurrentMediaPlayer(nullptr)
diff --git a/Engine/Source/Runtime/MoviePlayer/Private/MoviePlayerSettings.cpp b/Engine/Source/Runtime/MoviePlayer/Private/MoviePlayerSettings.cpp
--- a/Engine/Source/Runtime/MoviePlayer/Private/MoviePlayerSettings.cpp
+++ b/Engine/Source/Runtime/MoviePlayer/Private/MoviePlayerSettings.cpp
@@ -6,8 +6,8 @@
 
 #define LOCTEXT_NAMESPACE "MoviePlayerSettings"
 
-UMoviePlayerSettings::UMoviePlayerSettings( const class FPostConstructInitializeProperties& PCIP )
-	: Super(PCIP)
+UMoviePlayerSettings::UMoviePlayerSettings( const FObjectInitializer& ObjectInitializer )
+	: Super(ObjectInitializer)
 {
 }
 
diff --git a/Engine/Source/Runtime/MovieSceneCore/Private/MovieScene.cpp b/Engine/Source/Runtime/MovieSceneCore/Private/MovieScene.cpp
--- a/Engine/Source/Runtime/MovieSceneCore/Private/MovieScene.cpp
+++ b/Engine/Source/Runtime/MovieSceneCore/Private/MovieScene.cpp
@@ -42,8 +42,8 @@ bool FMovieSceneObjectBinding::RemoveTrack( UMovieSceneTrack& Track )
 	return Tracks.RemoveSingle( &Track ) != 0;
 }
 
-UMovieScene::UMovieScene( const FPostConstructInitializeProperties& PCIP )
-	: Super( PCIP )
+UMovieScene::UMovieScene( const FObjectInitializer& ObjectInitializer )
+	: Super( ObjectInitializer )
 {
 }
 
diff --git a/Engine/Source/Runtime/MovieSceneCore/Private/MovieSceneBindings.cpp b/Engine/Source/Runtime/MovieSceneCore/Private/MovieSceneBindings.cpp
--- a/Engine/Source/Runtime/MovieSceneCore/Private/MovieSceneBindings.cpp
+++ b/Engine/Source/Runtime/MovieSceneCore/Private/MovieSceneBindings.cpp
@@ -11,8 +11,8 @@ FMovieSceneBoundObject::FMovieSceneBoundObject( const FGuid& InitPosessableGuid,
 
 
 
-UMovieSceneBindings::UMovieSceneBindings( const FPostConstructInitializeProperties& PCIP )
-	: Super( PCIP )
+UMovieSceneBindings::UMovieSceneBindings( const FObjectInitializer& ObjectInitializer )
+	: Super( ObjectInitializer )
 {
 
 	RootMovieScene = NULL;
diff --git a/Engine/Source/Runtime/MovieSceneCore/Private/MovieScenePropertyTrack.cpp b/Engine/Source/Runtime/MovieSceneCore/Private/MovieScenePropertyTrack.cpp
--- a/Engine/Source/Runtime/MovieSceneCore/Private/MovieScenePropertyTrack.cpp
+++ b/Engine/Source/Runtime/MovieSceneCore/Private/MovieScenePropertyTrack.cpp
@@ -4,8 +4,8 @@
 #include "MovieScenePropertyTrack.h"
 
 
-UMovieScenePropertyTrack::UMovieScenePropertyTrack( const FPostConstructInitializeProperties& PCIP )
-	: Super( PCIP )
+UMovieScenePropertyTrack::UMovieScenePropertyTrack( const FObjectInitializer& ObjectInitializer )
+	: Super( ObjectInitializer )
 {
 	bSectionsAreShowable = false;
 }
diff --git a/Engine/Source/Runtime/MovieSceneCore/Private/MovieSceneSection.cpp b/Engine/Source/Runtime/MovieSceneCore/Private/MovieSceneSection.cpp
--- a/Engine/Source/Runtime/MovieSceneCore/Private/MovieSceneSection.cpp
+++ b/Engine/Source/Runtime/MovieSceneCore/Private/MovieSceneSection.cpp
@@ -3,8 +3,8 @@
 #include "MovieSceneCorePrivatePCH.h"
 #include "MovieSceneSection.h"
 
-UMovieSceneSection::UMovieSceneSection( const FPostConstructInitializeProperties& PCIP )
-	: Super( PCIP )
+UMovieSceneSection::UMovieSceneSection( const FObjectInitializer& ObjectInitializer )
+	: Super( ObjectInitializer )
 {
 
 	StartTime = 0.0f;
diff --git a/Engine/Source/Runtime/MovieSceneCore/Private/MovieSceneTrack.cpp b/Engine/Source/Runtime/MovieSceneCore/Private/MovieSceneTrack.cpp
--- a/Engine/Source/Runtime/MovieSceneCore/Private/MovieSceneTrack.cpp
+++ b/Engine/Source/Runtime/MovieSceneCore/Private/MovieSceneTrack.cpp
@@ -6,7 +6,7 @@
 
 
 
-UMovieSceneTrack::UMovieSceneTrack( const FPostConstructInitializeProperties& PCIP )
-	: Super( PCIP )
+UMovieSceneTrack::UMovieSceneTrack( const FObjectInitializer& ObjectInitializer )
+	: Super( ObjectInitializer )
 {
 }
diff --git a/Engine/Source/Runtime/MovieSceneCore/Private/RuntimeMovieScenePlayer.cpp b/Engine/Source/Runtime/MovieSceneCore/Private/RuntimeMovieScenePlayer.cpp
--- a/Engine/Source/Runtime/MovieSceneCore/Private/RuntimeMovieScenePlayer.cpp
+++ b/Engine/Source/Runtime/MovieSceneCore/Private/RuntimeMovieScenePlayer.cpp
@@ -30,8 +30,8 @@ URuntimeMovieScenePlayer* URuntimeMovieScenePlayer::CreateRuntimeMovieScenePlaye
 }
 
 
-URuntimeMovieScenePlayer::URuntimeMovieScenePlayer( const FPostConstructInitializeProperties& PCIP )
-	: Super( PCIP )
+URuntimeMovieScenePlayer::URuntimeMovieScenePlayer( const FObjectInitializer& ObjectInitializer )
+	: Super( ObjectInitializer )
 {
 
 	MovieSceneBindings = NULL;
diff --git a/Engine/Source/Runtime/MovieSceneCoreTypes/Private/MovieScene3DTransformSection.cpp b/Engine/Source/Runtime/MovieSceneCoreTypes/Private/MovieScene3DTransformSection.cpp
--- a/Engine/Source/Runtime/MovieSceneCoreTypes/Private/MovieScene3DTransformSection.cpp
+++ b/Engine/Source/Runtime/MovieSceneCoreTypes/Private/MovieScene3DTransformSection.cpp
@@ -4,8 +4,8 @@
 #include "MovieScene3DTransformSection.h"
 #include "MovieScene3DTransformTrack.h"
 
-UMovieScene3DTransformSection::UMovieScene3DTransformSection( const FPostConstructInitializeProperties& PCIP )
-	: Super( PCIP )
+UMovieScene3DTransformSection::UMovieScene3DTransformSection( const FObjectInitializer& ObjectInitializer )
+	: Super( ObjectInitializer )
 {
 
 }
diff --git a/Engine/Source/Runtime/MovieSceneCoreTypes/Private/MovieScene3DTransformTrack.cpp b/Engine/Source/Runtime/MovieSceneCoreTypes/Private/MovieScene3DTransformTrack.cpp
--- a/Engine/Source/Runtime/MovieSceneCoreTypes/Private/MovieScene3DTransformTrack.cpp
+++ b/Engine/Source/Runtime/MovieSceneCoreTypes/Private/MovieScene3DTransformTrack.cpp
@@ -57,8 +57,8 @@ bool FTransformKey::GetVectorComponentIfDifferent( EAxis::Type Axis, const FVect
 	return bShouldAddKey;
 }
 
-UMovieScene3DTransformTrack::UMovieScene3DTransformTrack( const FPostConstructInitializeProperties& PCIP )
-	: Super( PCIP )
+UMovieScene3DTransformTrack::UMovieScene3DTransformTrack( const FObjectInitializer& ObjectInitializer )
+	: Super( ObjectInitializer )
 {
 }
 
diff --git a/Engine/Source/Runtime/MovieSceneCoreTypes/Private/MovieSceneAnimationSection.cpp b/Engine/Source/Runtime/MovieSceneCoreTypes/Private/MovieSceneAnimationSection.cpp
--- a/Engine/Source/Runtime/MovieSceneCoreTypes/Private/MovieSceneAnimationSection.cpp
+++ b/Engine/Source/Runtime/MovieSceneCoreTypes/Private/MovieSceneAnimationSection.cpp
@@ -3,8 +3,8 @@
 #include "MovieSceneCoreTypesPCH.h"
 #include "MovieSceneAnimationSection.h"
 
-UMovieSceneAnimationSection::UMovieSceneAnimationSection( const FPostConstructInitializeProperties& PCIP )
-	: Super( PCIP )
+UMovieSceneAnimationSection::UMovieSceneAnimationSection( const FObjectInitializer& ObjectInitializer )
+	: Super( ObjectInitializer )
 {
 	AnimSequence = NULL;
 	AnimationStartTime = 0.f;
diff --git a/Engine/Source/Runtime/MovieSceneCoreTypes/Private/MovieSceneAnimationTrack.cpp b/Engine/Source/Runtime/MovieSceneCoreTypes/Private/MovieSceneAnimationTrack.cpp
--- a/Engine/Source/Runtime/MovieSceneCoreTypes/Private/MovieSceneAnimationTrack.cpp
+++ b/Engine/Source/Runtime/MovieSceneCoreTypes/Private/MovieSceneAnimationTrack.cpp
@@ -8,8 +8,8 @@
 
 #define LOCTEXT_NAMESPACE "MovieSceneAnimationTrack"
 
-UMovieSceneAnimationTrack::UMovieSceneAnimationTrack( const FPostConstructInitializeProperties& PCIP )
-	: Super( PCIP )
+UMovieSceneAnimationTrack::UMovieSceneAnimationTrack( const FObjectInitializer& ObjectInitializer )
+	: Super( ObjectInitializer )
 {
 }
 
diff --git a/Engine/Source/Runtime/MovieSceneCoreTypes/Private/MovieSceneAudioSection.cpp b/Engine/Source/Runtime/MovieSceneCoreTypes/Private/MovieSceneAudioSection.cpp
--- a/Engine/Source/Runtime/MovieSceneCoreTypes/Private/MovieSceneAudioSection.cpp
+++ b/Engine/Source/Runtime/MovieSceneCoreTypes/Private/MovieSceneAudioSection.cpp
@@ -4,8 +4,8 @@
 #include "MovieSceneAudioSection.h"
 #include "SoundDefinitions.h"
 
-UMovieSceneAudioSection::UMovieSceneAudioSection( const FPostConstructInitializeProperties& PCIP )
-	: Super( PCIP )
+UMovieSceneAudioSection::UMovieSceneAudioSection( const FObjectInitializer& ObjectInitializer )
+	: Super( ObjectInitializer )
 {
 
 	Sound = NULL;
diff --git a/Engine/Source/Runtime/MovieSceneCoreTypes/Private/MovieSceneAudioTrack.cpp b/Engine/Source/Runtime/MovieSceneCoreTypes/Private/MovieSceneAudioTrack.cpp
--- a/Engine/Source/Runtime/MovieSceneCoreTypes/Private/MovieSceneAudioTrack.cpp
+++ b/Engine/Source/Runtime/MovieSceneCoreTypes/Private/MovieSceneAudioTrack.cpp
@@ -11,8 +11,8 @@
 
 #define LOCTEXT_NAMESPACE "MovieSceneDirectorTrack"
 
-UMovieSceneAudioTrack::UMovieSceneAudioTrack( const FPostConstructInitializeProperties& PCIP )
-	: Super( PCIP )
+UMovieSceneAudioTrack::UMovieSceneAudioTrack( const FObjectInitializer& ObjectInitializer )
+	: Super( ObjectInitializer )
 {
 }
 
diff --git a/Engine/Source/Runtime/MovieSceneCoreTypes/Private/MovieSceneBoolSection.cpp b/Engine/Source/Runtime/MovieSceneCoreTypes/Private/MovieSceneBoolSection.cpp
--- a/Engine/Source/Runtime/MovieSceneCoreTypes/Private/MovieSceneBoolSection.cpp
+++ b/Engine/Source/Runtime/MovieSceneCoreTypes/Private/MovieSceneBoolSection.cpp
@@ -4,8 +4,8 @@
 #include "MovieSceneBoolSection.h"
 
 
-UMovieSceneBoolSection::UMovieSceneBoolSection( const FPostConstructInitializeProperties& PCIP )
-	: Super( PCIP )
+UMovieSceneBoolSection::UMovieSceneBoolSection( const FObjectInitializer& ObjectInitializer )
+	: Super( ObjectInitializer )
 {
 }
 
diff --git a/Engine/Source/Runtime/MovieSceneCoreTypes/Private/MovieSceneBoolTrack.cpp b/Engine/Source/Runtime/MovieSceneCoreTypes/Private/MovieSceneBoolTrack.cpp
--- a/Engine/Source/Runtime/MovieSceneCoreTypes/Private/MovieSceneBoolTrack.cpp
+++ b/Engine/Source/Runtime/MovieSceneCoreTypes/Private/MovieSceneBoolTrack.cpp
@@ -7,8 +7,8 @@
 #include "MovieSceneBoolTrackInstance.h"
 
 
-UMovieSceneBoolTrack::UMovieSceneBoolTrack( const FPostConstructInitializeProperties& PCIP )
-	: Super( PCIP )
+UMovieSceneBoolTrack::UMovieSceneBoolTrack( const FObjectInitializer& ObjectInitializer )
+	: Super( ObjectInitializer )
 {
 }
 
diff --git a/Engine/Source/Runtime/MovieSceneCoreTypes/Private/MovieSceneColorSection.cpp b/Engine/Source/Runtime/MovieSceneCoreTypes/Private/MovieSceneColorSection.cpp
--- a/Engine/Source/Runtime/MovieSceneCoreTypes/Private/MovieSceneColorSection.cpp
+++ b/Engine/Source/Runtime/MovieSceneCoreTypes/Private/MovieSceneColorSection.cpp
@@ -4,8 +4,8 @@
 #include "MovieSceneColorSection.h"
 
 
-UMovieSceneColorSection::UMovieSceneColorSection( const FPostConstructInitializeProperties& PCIP )
-	: Super( PCIP )
+UMovieSceneColorSection::UMovieSceneColorSection( const FObjectInitializer& ObjectInitializer )
+	: Super( ObjectInitializer )
 {
 }
 
diff --git a/Engine/Source/Runtime/MovieSceneCoreTypes/Private/MovieSceneColorTrack.cpp b/Engine/Source/Runtime/MovieSceneCoreTypes/Private/MovieSceneColorTrack.cpp
--- a/Engine/Source/Runtime/MovieSceneCoreTypes/Private/MovieSceneColorTrack.cpp
+++ b/Engine/Source/Runtime/MovieSceneCoreTypes/Private/MovieSceneColorTrack.cpp
@@ -6,8 +6,8 @@
 #include "IMovieScenePlayer.h"
 #include "MovieSceneColorTrackInstance.h"
 
-UMovieSceneColorTrack::UMovieSceneColorTrack( const FPostConstructInitializeProperties& PCIP )
-	: Super( PCIP )
+UMovieSceneColorTrack::UMovieSceneColorTrack( const FObjectInitializer& ObjectInitializer )
+	: Super( ObjectInitializer )
 {
 }
 
diff --git a/Engine/Source/Runtime/MovieSceneCoreTypes/Private/MovieSceneDirectorTrack.cpp b/Engine/Source/Runtime/MovieSceneCoreTypes/Private/MovieSceneDirectorTrack.cpp
--- a/Engine/Source/Runtime/MovieSceneCoreTypes/Private/MovieSceneDirectorTrack.cpp
+++ b/Engine/Source/Runtime/MovieSceneCoreTypes/Private/MovieSceneDirectorTrack.cpp
@@ -8,8 +8,8 @@
 
 #define LOCTEXT_NAMESPACE "MovieSceneDirectorTrack"
 
-UMovieSceneDirectorTrack::UMovieSceneDirectorTrack( const FPostConstructInitializeProperties& PCIP )
-	: Super( PCIP )
+UMovieSceneDirectorTrack::UMovieSceneDirectorTrack( const FObjectInitializer& ObjectInitializer )
+	: Super( ObjectInitializer )
 {
 }
 
diff --git a/Engine/Source/Runtime/MovieSceneCoreTypes/Private/MovieSceneFloatSection.cpp b/Engine/Source/Runtime/MovieSceneCoreTypes/Private/MovieSceneFloatSection.cpp
--- a/Engine/Source/Runtime/MovieSceneCoreTypes/Private/MovieSceneFloatSection.cpp
+++ b/Engine/Source/Runtime/MovieSceneCoreTypes/Private/MovieSceneFloatSection.cpp
@@ -4,8 +4,8 @@
 #include "MovieSceneFloatSection.h"
 
 
-UMovieSceneFloatSection::UMovieSceneFloatSection( const FPostConstructInitializeProperties& PCIP )
-	: Super( PCIP )
+UMovieSceneFloatSection::UMovieSceneFloatSection( const FObjectInitializer& ObjectInitializer )
+	: Super( ObjectInitializer )
 {
 }
 
diff --git a/Engine/Source/Runtime/MovieSceneCoreTypes/Private/MovieSceneFloatTrack.cpp b/Engine/Source/Runtime/MovieSceneCoreTypes/Private/MovieSceneFloatTrack.cpp
--- a/Engine/Source/Runtime/MovieSceneCoreTypes/Private/MovieSceneFloatTrack.cpp
+++ b/Engine/Source/Runtime/MovieSceneCoreTypes/Private/MovieSceneFloatTrack.cpp
@@ -7,8 +7,8 @@
 #include "MovieSceneFloatTrackInstance.h"
 
 
-UMovieSceneFloatTrack::UMovieSceneFloatTrack( const FPostConstructInitializeProperties& PCIP )
-	: Super( PCIP )
+UMovieSceneFloatTrack::UMovieSceneFloatTrack( const FObjectInitializer& ObjectInitializer )
+	: Super( ObjectInitializer )
 {
 }
 
diff --git a/Engine/Source/Runtime/MovieSceneCoreTypes/Private/MovieSceneParticleSection.cpp b/Engine/Source/Runtime/MovieSceneCoreTypes/Private/MovieSceneParticleSection.cpp
--- a/Engine/Source/Runtime/MovieSceneCoreTypes/Private/MovieSceneParticleSection.cpp
+++ b/Engine/Source/Runtime/MovieSceneCoreTypes/Private/MovieSceneParticleSection.cpp
@@ -3,8 +3,8 @@
 #include "MovieSceneCoreTypesPCH.h"
 #include "MovieSceneParticleSection.h"
 
-UMovieSceneParticleSection::UMovieSceneParticleSection( const FPostConstructInitializeProperties& PCIP )
-	: Super( PCIP )
+UMovieSceneParticleSection::UMovieSceneParticleSection( const FObjectInitializer& ObjectInitializer )
+	: Super( ObjectInitializer )
 {
 	KeyType = EParticleKey::Toggle;
 }
diff --git a/Engine/Source/Runtime/MovieSceneCoreTypes/Private/MovieSceneParticleTrack.cpp b/Engine/Source/Runtime/MovieSceneCoreTypes/Private/MovieSceneParticleTrack.cpp
--- a/Engine/Source/Runtime/MovieSceneCoreTypes/Private/MovieSceneParticleTrack.cpp
+++ b/Engine/Source/Runtime/MovieSceneCoreTypes/Private/MovieSceneParticleTrack.cpp
@@ -8,8 +8,8 @@
 
 #define LOCTEXT_NAMESPACE "MovieSceneParticleTrack"
 
-UMovieSceneParticleTrack::UMovieSceneParticleTrack( const FPostConstructInitializeProperties& PCIP )
-	: Super( PCIP )
+UMovieSceneParticleTrack::UMovieSceneParticleTrack( const FObjectInitializer& ObjectInitializer )
+	: Super( ObjectInitializer )
 {
 }
 
diff --git a/Engine/Source/Runtime/MovieSceneCoreTypes/Private/MovieSceneShotSection.cpp b/Engine/Source/Runtime/MovieSceneCoreTypes/Private/MovieSceneShotSection.cpp
--- a/Engine/Source/Runtime/MovieSceneCoreTypes/Private/MovieSceneShotSection.cpp
+++ b/Engine/Source/Runtime/MovieSceneCoreTypes/Private/MovieSceneShotSection.cpp
@@ -3,8 +3,8 @@
 #include "MovieSceneCoreTypesPCH.h"
 #include "MovieSceneShotSection.h"
 
-UMovieSceneShotSection::UMovieSceneShotSection( const FPostConstructInitializeProperties& PCIP )
-	: Super( PCIP )
+UMovieSceneShotSection::UMovieSceneShotSection( const FObjectInitializer& ObjectInitializer )
+	: Super( ObjectInitializer )
 {
 }
 
diff --git a/Engine/Source/Runtime/MovieSceneCoreTypes/Private/MovieSceneVectorSection.cpp b/Engine/Source/Runtime/MovieSceneCoreTypes/Private/MovieSceneVectorSection.cpp
--- a/Engine/Source/Runtime/MovieSceneCoreTypes/Private/MovieSceneVectorSection.cpp
+++ b/Engine/Source/Runtime/MovieSceneCoreTypes/Private/MovieSceneVectorSection.cpp
@@ -4,8 +4,8 @@
 #include "MovieSceneVectorSection.h"
 #include "MovieSceneVectorTrack.h"
 
-UMovieSceneVectorSection::UMovieSceneVectorSection( const FPostConstructInitializeProperties& PCIP )
-	: Super( PCIP )
+UMovieSceneVectorSection::UMovieSceneVectorSection( const FObjectInitializer& ObjectInitializer )
+	: Super( ObjectInitializer )
 {
 
 }
diff --git a/Engine/Source/Runtime/MovieSceneCoreTypes/Private/MovieSceneVectorTrack.cpp b/Engine/Source/Runtime/MovieSceneCoreTypes/Private/MovieSceneVectorTrack.cpp
--- a/Engine/Source/Runtime/MovieSceneCoreTypes/Private/MovieSceneVectorTrack.cpp
+++ b/Engine/Source/Runtime/MovieSceneCoreTypes/Private/MovieSceneVectorTrack.cpp
@@ -6,8 +6,8 @@
 #include "IMovieScenePlayer.h"
 #include "MovieSceneVectorTrackInstance.h"
 
-UMovieSceneVectorTrack::UMovieSceneVectorTrack( const FPostConstructInitializeProperties& PCIP )
-	: Super( PCIP )
+UMovieSceneVectorTrack::UMovieSceneVectorTrack( const FObjectInitializer& ObjectInitializer )
+	: Super( ObjectInitializer )
 {
 	NumChannelsUsed = 0;
 }
diff --git a/Engine/Source/Runtime/MovieSceneCoreTypes/Private/SubMovieSceneSection.cpp b/Engine/Source/Runtime/MovieSceneCoreTypes/Private/SubMovieSceneSection.cpp
--- a/Engine/Source/Runtime/MovieSceneCoreTypes/Private/SubMovieSceneSection.cpp
+++ b/Engine/Source/Runtime/MovieSceneCoreTypes/Private/SubMovieSceneSection.cpp
@@ -4,7 +4,7 @@
 #include "SubMovieSceneSection.h"
 #include "IMovieScenePlayer.h"
 
-USubMovieSceneSection::USubMovieSceneSection( const FPostConstructInitializeProperties& PCIP )
-	: Super( PCIP )
+USubMovieSceneSection::USubMovieSceneSection( const FObjectInitializer& ObjectInitializer )
+	: Super( ObjectInitializer )
 {
 }
diff --git a/Engine/Source/Runtime/MovieSceneCoreTypes/Private/SubMovieSceneTrack.cpp b/Engine/Source/Runtime/MovieSceneCoreTypes/Private/SubMovieSceneTrack.cpp
--- a/Engine/Source/Runtime/MovieSceneCoreTypes/Private/SubMovieSceneTrack.cpp
+++ b/Engine/Source/Runtime/MovieSceneCoreTypes/Private/SubMovieSceneTrack.cpp
@@ -6,8 +6,8 @@
 #include "SubMovieSceneSection.h"
 #include "SubMovieSceneTrackInstance.h"
 
-USubMovieSceneTrack::USubMovieSceneTrack( const FPostConstructInitializeProperties& PCIP )
-	: Super( PCIP )
+USubMovieSceneTrack::USubMovieSceneTrack( const FObjectInitializer& ObjectInitializer )
+	: Super( ObjectInitializer )
 {
 }
 
diff --git a/Engine/Source/Runtime/Online/BuildPatchServices/Private/BuildPatchManifest.cpp b/Engine/Source/Runtime/Online/BuildPatchServices/Private/BuildPatchManifest.cpp
--- a/Engine/Source/Runtime/Online/BuildPatchServices/Private/BuildPatchManifest.cpp
+++ b/Engine/Source/Runtime/Online/BuildPatchServices/Private/BuildPatchManifest.cpp
@@ -379,8 +379,8 @@ private:
 
 /* UBuildPatchManifest implementation
 *****************************************************************************/
-UBuildPatchManifest::UBuildPatchManifest(const FPostConstructInitializeProperties& PCIP)
-	: Super(PCIP)
+UBuildPatchManifest::UBuildPatchManifest(const FObjectInitializer& ObjectInitializer)
+	: Super(ObjectInitializer)
 	, ManifestFileVersion(EBuildPatchAppManifestVersion::GetLatestVersion())
 	, bIsFileData(false)
 	, AppID(0)
diff --git a/Engine/Source/Runtime/Online/OnlineSubsystem/Private/NamedInterfaces.cpp b/Engine/Source/Runtime/Online/OnlineSubsystem/Private/NamedInterfaces.cpp
--- a/Engine/Source/Runtime/Online/OnlineSubsystem/Private/NamedInterfaces.cpp
+++ b/Engine/Source/Runtime/Online/OnlineSubsystem/Private/NamedInterfaces.cpp
@@ -3,8 +3,8 @@
 #include "OnlineSubsystemPrivatePCH.h"
 #include "NamedInterfaces.h"
 
-UNamedInterfaces::UNamedInterfaces(const class FPostConstructInitializeProperties& PCIP) :
-	Super(PCIP)
+UNamedInterfaces::UNamedInterfaces(const FObjectInitializer& ObjectInitializer) :
+	Super(ObjectInitializer)
 {
 }
 
diff --git a/Engine/Source/Runtime/Online/OnlineSubsystemSteam/Private/SteamNetConnection.cpp b/Engine/Source/Runtime/Online/OnlineSubsystemSteam/Private/SteamNetConnection.cpp
--- a/Engine/Source/Runtime/Online/OnlineSubsystemSteam/Private/SteamNetConnection.cpp
+++ b/Engine/Source/Runtime/Online/OnlineSubsystemSteam/Private/SteamNetConnection.cpp
@@ -8,8 +8,8 @@
 #include "SocketSubsystemSteam.h"
 #include "SocketsSteam.h"
 
-USteamNetConnection::USteamNetConnection(const class FPostConstructInitializeProperties& PCIP) :
-	Super(PCIP),
+USteamNetConnection::USteamNetConnection(const FObjectInitializer& ObjectInitializer) :
+	Super(ObjectInitializer),
 	bIsPassthrough(false)
 {
 }
diff --git a/Engine/Source/Runtime/Online/OnlineSubsystemSteam/Private/SteamNetDriver.cpp b/Engine/Source/Runtime/Online/OnlineSubsystemSteam/Private/SteamNetDriver.cpp
--- a/Engine/Source/Runtime/Online/OnlineSubsystemSteam/Private/SteamNetDriver.cpp
+++ b/Engine/Source/Runtime/Online/OnlineSubsystemSteam/Private/SteamNetDriver.cpp
@@ -9,8 +9,8 @@
 #include "SocketSubsystemSteam.h"
 #include "SocketsSteam.h"
 
-USteamNetDriver::USteamNetDriver(const class FPostConstructInitializeProperties& PCIP) :
-	Super(PCIP),
+USteamNetDriver::USteamNetDriver(const FObjectInitializer& ObjectInitializer) :
+	Super(ObjectInitializer),
 	bIsPassthrough(false)
 {
 }
diff --git a/Engine/Source/Runtime/Online/OnlineSubsystemUtils/Classes/LeaderboardFlushCallbackProxy.cpp b/Engine/Source/Runtime/Online/OnlineSubsystemUtils/Classes/LeaderboardFlushCallbackProxy.cpp
--- a/Engine/Source/Runtime/Online/OnlineSubsystemUtils/Classes/LeaderboardFlushCallbackProxy.cpp
+++ b/Engine/Source/Runtime/Online/OnlineSubsystemUtils/Classes/LeaderboardFlushCallbackProxy.cpp
@@ -5,8 +5,8 @@
 //////////////////////////////////////////////////////////////////////////
 // ULeaderboardFlushCallbackProxy
 
-ULeaderboardFlushCallbackProxy::ULeaderboardFlushCallbackProxy(const class FPostConstructInitializeProperties& PCIP)
-	: Super(PCIP)
+ULeaderboardFlushCallbackProxy::ULeaderboardFlushCallbackProxy(const FObjectInitializer& ObjectInitializer)
+	: Super(ObjectInitializer)
 {
 }
 
diff --git a/Engine/Source/Runtime/Online/OnlineSubsystemUtils/Classes/LeaderboardQueryCallbackProxy.cpp b/Engine/Source/Runtime/Online/OnlineSubsystemUtils/Classes/LeaderboardQueryCallbackProxy.cpp
--- a/Engine/Source/Runtime/Online/OnlineSubsystemUtils/Classes/LeaderboardQueryCallbackProxy.cpp
+++ b/Engine/Source/Runtime/Online/OnlineSubsystemUtils/Classes/LeaderboardQueryCallbackProxy.cpp
@@ -5,8 +5,8 @@
 //////////////////////////////////////////////////////////////////////////
 // ULeaderboardQueryCallbackProxy
 
-ULeaderboardQueryCallbackProxy::ULeaderboardQueryCallbackProxy(const class FPostConstructInitializeProperties& PCIP)
-	: Super(PCIP)
+ULeaderboardQueryCallbackProxy::ULeaderboardQueryCallbackProxy(const FObjectInitializer& ObjectInitializer)
+	: Super(ObjectInitializer)
 {
 }
 
diff --git a/Engine/Source/Runtime/Online/OnlineSubsystemUtils/Private/AchievementBlueprintLibrary.cpp b/Engine/Source/Runtime/Online/OnlineSubsystemUtils/Private/AchievementBlueprintLibrary.cpp
--- a/Engine/Source/Runtime/Online/OnlineSubsystemUtils/Private/AchievementBlueprintLibrary.cpp
+++ b/Engine/Source/Runtime/Online/OnlineSubsystemUtils/Private/AchievementBlueprintLibrary.cpp
@@ -4,8 +4,8 @@
 #include "AchievementBlueprintLibrary.h"
 #include "Runtime/Online/OnlineSubsystem/Public/Interfaces/OnlineAchievementsInterface.h"
 
-UAchievementBlueprintLibrary::UAchievementBlueprintLibrary(const FPostConstructInitializeProperties& PCIP)
-	: Super(PCIP)
+UAchievementBlueprintLibrary::UAchievementBlueprintLibrary(const FObjectInitializer& ObjectInitializer)
+	: Super(ObjectInitializer)
 {
 }
 
diff --git a/Engine/Source/Runtime/Online/OnlineSubsystemUtils/Private/AchievementQueryCallbackProxy.cpp b/Engine/Source/Runtime/Online/OnlineSubsystemUtils/Private/AchievementQueryCallbackProxy.cpp
--- a/Engine/Source/Runtime/Online/OnlineSubsystemUtils/Private/AchievementQueryCallbackProxy.cpp
+++ b/Engine/Source/Runtime/Online/OnlineSubsystemUtils/Private/AchievementQueryCallbackProxy.cpp
@@ -5,8 +5,8 @@
 //////////////////////////////////////////////////////////////////////////
 // UAchievementQueryCallbackProxy
 
-UAchievementQueryCallbackProxy::UAchievementQueryCallbackProxy(const class FPostConstructInitializeProperties& PCIP)
-	: Super(PCIP)
+UAchievementQueryCallbackProxy::UAchievementQueryCallbackProxy(const FObjectInitializer& ObjectInitializer)
+	: Super(ObjectInitializer)
 {
 }
 
diff --git a/Engine/Source/Runtime/Online/OnlineSubsystemUtils/Private/AchievementWriteCallbackProxy.cpp b/Engine/Source/Runtime/Online/OnlineSubsystemUtils/Private/AchievementWriteCallbackProxy.cpp
--- a/Engine/Source/Runtime/Online/OnlineSubsystemUtils/Private/AchievementWriteCallbackProxy.cpp
+++ b/Engine/Source/Runtime/Online/OnlineSubsystemUtils/Private/AchievementWriteCallbackProxy.cpp
@@ -5,8 +5,8 @@
 //////////////////////////////////////////////////////////////////////////
 // ULeaderboardQueryCallbackProxy
 
-UAchievementWriteCallbackProxy::UAchievementWriteCallbackProxy(const class FPostConstructInitializeProperties& PCIP)
-	: Super(PCIP)
+UAchievementWriteCallbackProxy::UAchievementWriteCallbackProxy(const FObjectInitializer& ObjectInitializer)
+	: Super(ObjectInitializer)
 {
 }
 
diff --git a/Engine/Source/Runtime/Online/OnlineSubsystemUtils/Private/ConnectionCallbackProxy.cpp b/Engine/Source/Runtime/Online/OnlineSubsystemUtils/Private/ConnectionCallbackProxy.cpp
--- a/Engine/Source/Runtime/Online/OnlineSubsystemUtils/Private/ConnectionCallbackProxy.cpp
+++ b/Engine/Source/Runtime/Online/OnlineSubsystemUtils/Private/ConnectionCallbackProxy.cpp
@@ -5,8 +5,8 @@
 //////////////////////////////////////////////////////////////////////////
 // UConnectionCallbackProxy
 
-UConnectionCallbackProxy::UConnectionCallbackProxy(const class FPostConstructInitializeProperties& PCIP)
-	: Super(PCIP)
+UConnectionCallbackProxy::UConnectionCallbackProxy(const FObjectInitializer& ObjectInitializer)
+	: Super(ObjectInitializer)
 {
 
 	OnLoginCompleteDelegate.BindUObject(this, &ThisClass::OnLoginCompleted);
diff --git a/Engine/Source/Runtime/Online/OnlineSubsystemUtils/Private/InAppPurchaseCallbackProxy.cpp b/Engine/Source/Runtime/Online/OnlineSubsystemUtils/Private/InAppPurchaseCallbackProxy.cpp
--- a/Engine/Source/Runtime/Online/OnlineSubsystemUtils/Private/InAppPurchaseCallbackProxy.cpp
+++ b/Engine/Source/Runtime/Online/OnlineSubsystemUtils/Private/InAppPurchaseCallbackProxy.cpp
@@ -5,8 +5,8 @@
 //////////////////////////////////////////////////////////////////////////
 // UInAppPurchaseCallbackProxy
 
-UInAppPurchaseCallbackProxy::UInAppPurchaseCallbackProxy(const class FPostConstructInitializeProperties& PCIP)
-	: Super(PCIP)
+UInAppPurchaseCallbackProxy::UInAppPurchaseCallbackProxy(const FObjectInitializer& ObjectInitializer)
+	: Super(ObjectInitializer)
 {
 	PurchaseRequest = nullptr;
 	WorldPtr = nullptr;
diff --git a/Engine/Source/Runtime/Online/OnlineSubsystemUtils/Private/InAppPurchaseQueryCallbackProxy.cpp b/Engine/Source/Runtime/Online/OnlineSubsystemUtils/Private/InAppPurchaseQueryCallbackProxy.cpp
--- a/Engine/Source/Runtime/Online/OnlineSubsystemUtils/Private/InAppPurchaseQueryCallbackProxy.cpp
+++ b/Engine/Source/Runtime/Online/OnlineSubsystemUtils/Private/InAppPurchaseQueryCallbackProxy.cpp
@@ -5,8 +5,8 @@
 //////////////////////////////////////////////////////////////////////////
 // UInAppPurchaseQueryCallbackProxy
 
-UInAppPurchaseQueryCallbackProxy::UInAppPurchaseQueryCallbackProxy(const class FPostConstructInitializeProperties& PCIP)
-	: Super(PCIP)
+UInAppPurchaseQueryCallbackProxy::UInAppPurchaseQueryCallbackProxy(const FObjectInitializer& ObjectInitializer)
+	: Super(ObjectInitializer)
 {
 }
 
diff --git a/Engine/Source/Runtime/Online/OnlineSubsystemUtils/Private/IpConnection.cpp b/Engine/Source/Runtime/Online/OnlineSubsystemUtils/Private/IpConnection.cpp
--- a/Engine/Source/Runtime/Online/OnlineSubsystemUtils/Private/IpConnection.cpp
+++ b/Engine/Source/Runtime/Online/OnlineSubsystemUtils/Private/IpConnection.cpp
@@ -24,8 +24,8 @@ Notes:
 #define SLIP_HEADER_SIZE   (UDP_HEADER_SIZE+4)
 #define WINSOCK_MAX_PACKET (512)
 
-UIpConnection::UIpConnection(const class FPostConstructInitializeProperties& PCIP) :
-	Super(PCIP),
+UIpConnection::UIpConnection(const FObjectInitializer& ObjectInitializer) :
+	Super(ObjectInitializer),
 	RemoteAddr(NULL),
 	Socket(NULL),
 	ResolveInfo(NULL)
diff --git a/Engine/Source/Runtime/Online/OnlineSubsystemUtils/Private/IpNetDriver.cpp b/Engine/Source/Runtime/Online/OnlineSubsystemUtils/Private/IpNetDriver.cpp
--- a/Engine/Source/Runtime/Online/OnlineSubsystemUtils/Private/IpNetDriver.cpp
+++ b/Engine/Source/Runtime/Online/OnlineSubsystemUtils/Private/IpNetDriver.cpp
@@ -20,8 +20,8 @@ Notes:
 /** Size of the network recv buffer */
 #define NETWORK_MAX_PACKET (576)
 
-UIpNetDriver::UIpNetDriver(const class FPostConstructInitializeProperties& PCIP)
-	: Super(PCIP)
+UIpNetDriver::UIpNetDriver(const FObjectInitializer& ObjectInitializer)
+	: Super(ObjectInitializer)
 {
 }
 
diff --git a/Engine/Source/Runtime/Online/OnlineSubsystemUtils/Private/LeaderboardBlueprintLibrary.cpp b/Engine/Source/Runtime/Online/OnlineSubsystemUtils/Private/LeaderboardBlueprintLibrary.cpp
--- a/Engine/Source/Runtime/Online/OnlineSubsystemUtils/Private/LeaderboardBlueprintLibrary.cpp
+++ b/Engine/Source/Runtime/Online/OnlineSubsystemUtils/Private/LeaderboardBlueprintLibrary.cpp
@@ -7,8 +7,8 @@
 //////////////////////////////////////////////////////////////////////////
 // ULeaderboardBlueprintLibrary
 
-ULeaderboardBlueprintLibrary::ULeaderboardBlueprintLibrary(const FPostConstructInitializeProperties& PCIP)
-	: Super(PCIP)
+ULeaderboardBlueprintLibrary::ULeaderboardBlueprintLibrary(const FObjectInitializer& ObjectInitializer)
+	: Super(ObjectInitializer)
 {
 }
 
diff --git a/Engine/Source/Runtime/Online/OnlineSubsystemUtils/Private/OnlineBeacon.cpp b/Engine/Source/Runtime/Online/OnlineSubsystemUtils/Private/OnlineBeacon.cpp
--- a/Engine/Source/Runtime/Online/OnlineSubsystemUtils/Private/OnlineBeacon.cpp
+++ b/Engine/Source/Runtime/Online/OnlineSubsystemUtils/Private/OnlineBeacon.cpp
@@ -5,8 +5,8 @@
 
 DEFINE_LOG_CATEGORY(LogBeacon);
 
-AOnlineBeacon::AOnlineBeacon(const FPostConstructInitializeProperties& PCIP) :
-	Super(PCIP)
+AOnlineBeacon::AOnlineBeacon(const FObjectInitializer& ObjectInitializer) :
+	Super(ObjectInitializer)
 {
 	NetDriverName = FName(TEXT("BeaconDriver"));
 }
diff --git a/Engine/Source/Runtime/Online/OnlineSubsystemUtils/Private/OnlineBeaconClient.cpp b/Engine/Source/Runtime/Online/OnlineSubsystemUtils/Private/OnlineBeaconClient.cpp
--- a/Engine/Source/Runtime/Online/OnlineSubsystemUtils/Private/OnlineBeaconClient.cpp
+++ b/Engine/Source/Runtime/Online/OnlineSubsystemUtils/Private/OnlineBeaconClient.cpp
@@ -5,8 +5,8 @@
 
 #define BEACON_RPC_TIMEOUT 15.0f
 
-AOnlineBeaconClient::AOnlineBeaconClient(const class FPostConstructInitializeProperties& PCIP) :
-	Super(PCIP),
+AOnlineBeaconClient::AOnlineBeaconClient(const FObjectInitializer& ObjectInitializer) :
+	Super(ObjectInitializer),
 	BeaconOwner(NULL)
 {
 	NetDriverName = FName(TEXT("BeaconDriver"));
diff --git a/Engine/Source/Runtime/Online/OnlineSubsystemUtils/Private/OnlineBeaconHost.cpp b/Engine/Source/Runtime/Online/OnlineSubsystemUtils/Private/OnlineBeaconHost.cpp
--- a/Engine/Source/Runtime/Online/OnlineSubsystemUtils/Private/OnlineBeaconHost.cpp
+++ b/Engine/Source/Runtime/Online/OnlineSubsystemUtils/Private/OnlineBeaconHost.cpp
@@ -4,8 +4,8 @@
 #include "Net/UnrealNetwork.h"
 #include "Net/DataChannel.h"
 
-AOnlineBeaconHost::AOnlineBeaconHost(const class FPostConstructInitializeProperties& PCIP) :
-	Super(PCIP)
+AOnlineBeaconHost::AOnlineBeaconHost(const FObjectInitializer& ObjectInitializer) :
+	Super(ObjectInitializer)
 {
 	NetDriverName = FName(TEXT("BeaconDriver"));
 }
@@ -344,8 +344,8 @@ AOnlineBeaconHost::FOnBeaconConnected& AOnlineBeaconHost::OnBeaconConnected(cons
 }
 
 
-AOnlineBeaconHostObject::AOnlineBeaconHostObject(const FPostConstructInitializeProperties& PCIP) :
-	Super(PCIP),
+AOnlineBeaconHostObject::AOnlineBeaconHostObject(const FObjectInitializer& ObjectInitializer) :
+	Super(ObjectInitializer),
 	BeaconTypeName(TEXT("UNDEFINED"))
 {
 	PrimaryActorTick.bCanEverTick = true;
diff --git a/Engine/Source/Runtime/Online/OnlineSubsystemUtils/Private/OnlineBlueprintCallProxyBase.cpp b/Engine/Source/Runtime/Online/OnlineSubsystemUtils/Private/OnlineBlueprintCallProxyBase.cpp
--- a/Engine/Source/Runtime/Online/OnlineSubsystemUtils/Private/OnlineBlueprintCallProxyBase.cpp
+++ b/Engine/Source/Runtime/Online/OnlineSubsystemUtils/Private/OnlineBlueprintCallProxyBase.cpp
@@ -6,8 +6,8 @@
 //////////////////////////////////////////////////////////////////////////
 // UOnlineBlueprintCallProxyBase
 
-UOnlineBlueprintCallProxyBase::UOnlineBlueprintCallProxyBase(const FPostConstructInitializeProperties& PCIP)
-	: Super(PCIP)
+UOnlineBlueprintCallProxyBase::UOnlineBlueprintCallProxyBase(const FObjectInitializer& ObjectInitializer)
+	: Super(ObjectInitializer)
 {
 }
 
diff --git a/Engine/Source/Runtime/Online/OnlineSubsystemUtils/Private/OnlineSessionClient.cpp b/Engine/Source/Runtime/Online/OnlineSubsystemUtils/Private/OnlineSessionClient.cpp
--- a/Engine/Source/Runtime/Online/OnlineSubsystemUtils/Private/OnlineSessionClient.cpp
+++ b/Engine/Source/Runtime/Online/OnlineSubsystemUtils/Private/OnlineSessionClient.cpp
@@ -11,8 +11,8 @@
 
 #define INVALID_CONTROLLERID 255
 
-UOnlineSessionClient::UOnlineSessionClient(const class FPostConstructInitializeProperties& PCIP)
-	: Super(PCIP)
+UOnlineSessionClient::UOnlineSessionClient(const FObjectInitializer& ObjectInitializer)
+	: Super(ObjectInitializer)
 {
 	bHandlingDisconnect = false;
 	bIsFromInvite = false;
diff --git a/Engine/Source/Runtime/Online/OnlineSubsystemUtils/Private/PartyBeaconClient.cpp b/Engine/Source/Runtime/Online/OnlineSubsystemUtils/Private/PartyBeaconClient.cpp
--- a/Engine/Source/Runtime/Online/OnlineSubsystemUtils/Private/PartyBeaconClient.cpp
+++ b/Engine/Source/Runtime/Online/OnlineSubsystemUtils/Private/PartyBeaconClient.cpp
@@ -3,8 +3,8 @@
 #include "OnlineSubsystemUtilsPrivatePCH.h"
 #include "OnlineSessionInterface.h"
 
-APartyBeaconClient::APartyBeaconClient(const class FPostConstructInitializeProperties& PCIP) :
-	Super(PCIP),
+APartyBeaconClient::APartyBeaconClient(const FObjectInitializer& ObjectInitializer) :
+	Super(ObjectInitializer),
 	bPendingReservationSent(false),
 	bCancelReservation(false)
 {
diff --git a/Engine/Source/Runtime/Online/OnlineSubsystemUtils/Private/PartyBeaconHost.cpp b/Engine/Source/Runtime/Online/OnlineSubsystemUtils/Private/PartyBeaconHost.cpp
--- a/Engine/Source/Runtime/Online/OnlineSubsystemUtils/Private/PartyBeaconHost.cpp
+++ b/Engine/Source/Runtime/Online/OnlineSubsystemUtils/Private/PartyBeaconHost.cpp
@@ -2,8 +2,8 @@
 
 #include "OnlineSubsystemUtilsPrivatePCH.h"
 
-APartyBeaconHost::APartyBeaconHost(const FPostConstructInitializeProperties& PCIP) :
-	Super(PCIP),
+APartyBeaconHost::APartyBeaconHost(const FObjectInitializer& ObjectInitializer) :
+	Super(ObjectInitializer),
 	State(NULL)
 {
 	BeaconTypeName = TEXT("PartyBeacon");
diff --git a/Engine/Source/Runtime/Online/OnlineSubsystemUtils/Private/PartyBeaconState.cpp b/Engine/Source/Runtime/Online/OnlineSubsystemUtils/Private/PartyBeaconState.cpp
--- a/Engine/Source/Runtime/Online/OnlineSubsystemUtils/Private/PartyBeaconState.cpp
+++ b/Engine/Source/Runtime/Online/OnlineSubsystemUtils/Private/PartyBeaconState.cpp
@@ -4,8 +4,8 @@
 #include "OnlineSubsystemUtilsClasses.h"
 
 
-UPartyBeaconState::UPartyBeaconState(const FPostConstructInitializeProperties& PCIP) :
-	Super(PCIP),
+UPartyBeaconState::UPartyBeaconState(const FObjectInitializer& ObjectInitializer) :
+	Super(ObjectInitializer),
 	SessionName(NAME_None),
 	NumConsumedReservations(0),
 	MaxReservations(0)
diff --git a/Engine/Source/Runtime/Online/OnlineSubsystemUtils/Private/TestBeaconClient.cpp b/Engine/Source/Runtime/Online/OnlineSubsystemUtils/Private/TestBeaconClient.cpp
--- a/Engine/Source/Runtime/Online/OnlineSubsystemUtils/Private/TestBeaconClient.cpp
+++ b/Engine/Source/Runtime/Online/OnlineSubsystemUtils/Private/TestBeaconClient.cpp
@@ -3,8 +3,8 @@
 #include "OnlineSubsystemUtilsPrivatePCH.h"
 #include "TestBeaconClient.h"
 
-ATestBeaconClient::ATestBeaconClient(const class FPostConstructInitializeProperties& PCIP) :
-	Super(PCIP)
+ATestBeaconClient::ATestBeaconClient(const FObjectInitializer& ObjectInitializer) :
+	Super(ObjectInitializer)
 {
 }
 
diff --git a/Engine/Source/Runtime/Online/OnlineSubsystemUtils/Private/TestBeaconHost.cpp b/Engine/Source/Runtime/Online/OnlineSubsystemUtils/Private/TestBeaconHost.cpp
--- a/Engine/Source/Runtime/Online/OnlineSubsystemUtils/Private/TestBeaconHost.cpp
+++ b/Engine/Source/Runtime/Online/OnlineSubsystemUtils/Private/TestBeaconHost.cpp
@@ -2,8 +2,8 @@
 
 #include "OnlineSubsystemUtilsPrivatePCH.h"
 
-ATestBeaconHost::ATestBeaconHost(const FPostConstructInitializeProperties& PCIP) :
-	Super(PCIP)
+ATestBeaconHost::ATestBeaconHost(const FObjectInitializer& ObjectInitializer) :
+	Super(ObjectInitializer)
 {
 	BeaconTypeName = TEXT("TestBeacon");
 }
diff --git a/Engine/Source/Runtime/Slate/Private/Framework/Styling/ButtonWidgetStyle.cpp b/Engine/Source/Runtime/Slate/Private/Framework/Styling/ButtonWidgetStyle.cpp
--- a/Engine/Source/Runtime/Slate/Private/Framework/Styling/ButtonWidgetStyle.cpp
+++ b/Engine/Source/Runtime/Slate/Private/Framework/Styling/ButtonWidgetStyle.cpp
@@ -3,8 +3,8 @@
 #include "SlatePrivatePCH.h"
 
 
-UButtonWidgetStyle::UButtonWidgetStyle( const class FPostConstructInitializeProperties& PCIP )
-	: Super(PCIP)
+UButtonWidgetStyle::UButtonWidgetStyle( const FObjectInitializer& ObjectInitializer )
+	: Super(ObjectInitializer)
 {
 	
 }
\ No newline at end of file
diff --git a/Engine/Source/Runtime/Slate/Private/Framework/Styling/CheckBoxWidgetStyle.cpp b/Engine/Source/Runtime/Slate/Private/Framework/Styling/CheckBoxWidgetStyle.cpp
--- a/Engine/Source/Runtime/Slate/Private/Framework/Styling/CheckBoxWidgetStyle.cpp
+++ b/Engine/Source/Runtime/Slate/Private/Framework/Styling/CheckBoxWidgetStyle.cpp
@@ -3,8 +3,8 @@
 #include "SlatePrivatePCH.h"
 
 
-UCheckBoxWidgetStyle::UCheckBoxWidgetStyle( const class FPostConstructInitializeProperties& PCIP )
-	: Super(PCIP)
+UCheckBoxWidgetStyle::UCheckBoxWidgetStyle( const FObjectInitializer& ObjectInitializer )
+	: Super(ObjectInitializer)
 {
 	
 }
\ No newline at end of file
diff --git a/Engine/Source/Runtime/Slate/Private/Framework/Styling/ComboBoxWidgetStyle.cpp b/Engine/Source/Runtime/Slate/Private/Framework/Styling/ComboBoxWidgetStyle.cpp
--- a/Engine/Source/Runtime/Slate/Private/Framework/Styling/ComboBoxWidgetStyle.cpp
+++ b/Engine/Source/Runtime/Slate/Private/Framework/Styling/ComboBoxWidgetStyle.cpp
@@ -3,8 +3,8 @@
 #include "SlatePrivatePCH.h"
 
 
-UComboBoxWidgetStyle::UComboBoxWidgetStyle( const class FPostConstructInitializeProperties& PCIP )
-	: Super(PCIP)
+UComboBoxWidgetStyle::UComboBoxWidgetStyle( const FObjectInitializer& ObjectInitializer )
+	: Super(ObjectInitializer)
 {
 	
 }
