New features and fixes for VIsual Logger.
- fixed crash in visual logger while using EnableAILoggind command line
- LogVisualizer flushes data from visual logger to be updated with all logs
- added visual log events
- first version of reports generated from visual log events as part LogVisualizer

diff --git a/Engine/Source/Developer/LogVisualizer/Private/SLogVisualizer.cpp b/Engine/Source/Developer/LogVisualizer/Private/SLogVisualizer.cpp
--- a/Engine/Source/Developer/LogVisualizer/Private/SLogVisualizer.cpp
+++ b/Engine/Source/Developer/LogVisualizer/Private/SLogVisualizer.cpp
@@ -21,6 +21,7 @@
 
 #include "GameplayDebuggingComponent.h"
 #include "LogVisualizerModule.h"
+#include "SLogVisualizerReport.h"
 
 #include "SFilterList.h"
 
@@ -373,6 +374,7 @@ void SLogVisualizer::Construct(const FArguments& InArgs, FLogVisualizer* InLogVi
 						.ItemHeight(20)
 						// Called when the user double-clicks with LMB on an item in the list
 						.OnMouseButtonDoubleClick(this, &SLogVisualizer::OnListDoubleClick)
+						.OnContextMenuOpening(this, &SLogVisualizer::GetListRightClickMenuContent)
 						.ListItemsSource(&LogsList)
 						.SelectionMode(ESelectionMode::Multi)
 						.OnGenerateRow(this, &SLogVisualizer::LogsListGenerateRow)
@@ -648,6 +650,54 @@ void SLogVisualizer::OnListDoubleClick(TSharedPtr<FLogsListItem> LogListItem)
 #endif
 }
 
+TSharedPtr<SWidget> SLogVisualizer::GetListRightClickMenuContent()
+{
+	FMenuBuilder MenuBuilder(true, NULL);
+
+	if (IsEnabled())
+	{
+		MenuBuilder.BeginSection("VisualLogReports", LOCTEXT("VisualLogReports", "VisualLog Reports"));
+		{
+			MenuBuilder.AddMenuEntry(
+				LOCTEXT("GenarateReport", "Genarate Report"),
+				LOCTEXT("GenarateReportTooltip", "Genarate report from Visual Log events."),
+				FSlateIcon(),
+				FUIAction(FExecuteAction::CreateSP(this, &SLogVisualizer::GenerateReport))
+				);
+		}
+		MenuBuilder.EndSection();
+	}
+
+	return MenuBuilder.MakeWidget();
+}
+
+void SLogVisualizer::GenerateReport()
+{
+	TArray< TSharedPtr<FLogsListItem> > ItemsToSave = LogsListWidget->GetSelectedItems();
+	TArray< TSharedPtr<FActorsVisLog> > AllLogs;
+
+	TSharedPtr<FLogsListItem>* LogListItem = ItemsToSave.GetData();
+	for (int32 ItemIndex = 0; ItemIndex < ItemsToSave.Num(); ++ItemIndex, ++LogListItem)
+	{
+		if (LogListItem->IsValid() && LogVisualizer->Logs.IsValidIndex((*LogListItem)->LogIndex))
+		{
+			TSharedPtr<FActorsVisLog> Log = LogVisualizer->Logs[(*LogListItem)->LogIndex];
+			AllLogs.Add(Log);
+		}
+	}
+
+
+	TSharedRef<SWindow> NewWindow = SNew(SWindow)
+		.ClientSize(FVector2D(720, 768))
+		.Title(NSLOCTEXT("LogVisualizerReport", "WindowTitle", "Log Visualizer Report"))
+		[
+			SNew(SLogVisualizerReport, this, AllLogs)
+		];
+
+	FSlateApplication::Get().AddWindow(NewWindow);
+}
+
+
 int32 SLogVisualizer::GetCurrentVisibleLogEntryIndex(const TArray<TSharedPtr<FVisualLogEntry> >& InVisibleEntries)
 {
 	if (LogVisualizer->Logs.IsValidIndex(SelectedLogIndex))
@@ -774,6 +824,26 @@ void SLogVisualizer::UpdateVisibleEntriesCache(const TSharedPtr<FActorsVisLog>&
 				continue;
 			}
 
+			if (!bAddedEntry)
+			{
+				for (int32 LogLineIndex = 0; LogLineIndex < CurrentEntry->Events.Num(); ++LogLineIndex)
+				{
+					const FVisualLogEntry::FLogEvent &CurrentLine = CurrentEntry->Events[LogLineIndex];
+					const FString CurrentCategory = CurrentLine.Name;
+					if (FilterListPtr->IsFilterEnabled(CurrentCategory, CurrentLine.Verbosity))
+					{
+						OutEntriesCached[CachedLogIndex].CachedEntries.Add(CurrentEntry);
+						bAddedEntry = true;
+						break;
+					}
+				}
+			}
+
+			if (bAddedEntry)
+			{
+				continue;
+			}
+
 			for (const auto CurrentBlock : CurrentEntry->DataBlocks)
 			{
 				const bool bCurrentCategoryPassed = FilterListPtr->IsFilterEnabled(CurrentBlock.Category.ToString(), ELogVerbosity::All);
@@ -884,6 +954,22 @@ void SLogVisualizer::GetVisibleEntries(const TSharedPtr<FActorsVisLog>& Log, TAr
 					continue;
 				}
 
+				for (int32 j = 0; j < CurrentEntry->Events.Num(); ++j)
+				{
+					const FString CurrentCategory = CurrentEntry->Events[j].Name;
+					FVisualLogEntry::FLogEvent &CurrentElement = CurrentEntry->Events[j];
+					if (CurrentElement.Name.Len() > 0 && (FilterListPtr->IsFilterEnabled(CurrentCategory, CurrentElement.Verbosity)))
+					{
+						OutEntriesCached[CachedLogIndex].CachedEntries.AddUnique(CurrentEntry);
+						bAddedEntry = true;
+						break;
+					}
+				}
+				if (bAddedEntry)
+				{
+					continue;
+				}
+
 				for (const auto CurrentData : CurrentEntry->DataBlocks)
 				{
 					const bool bCurrentCategoryPassed = FilterListPtr->IsFilterEnabled(CurrentData.Category.ToString(), ELogVerbosity::All);
@@ -975,6 +1061,22 @@ void SLogVisualizer::GetVisibleEntries(const TSharedPtr<FActorsVisLog>& Log, TAr
 				continue;
 			}
 
+			for (int32 ElementIndex = 0; ElementIndex < LogEntry->Events.Num(); ++ElementIndex)
+			{
+				FString CurrentCategory = LogEntry->Events[ElementIndex].Name;
+				if (bHistogramGraphsFilter || CurrentCategory.Find(QuickFilterText) != INDEX_NONE)
+				{
+					OutEntries.AddUnique(LogEntry);
+					bAddedEntry = true;
+					break;
+				}
+			}
+			if (bAddedEntry)
+			{
+				continue;
+			}
+
+
 			for (const auto CurrentData : LogEntry->DataBlocks)
 			{
 				const bool bCurrentCategoryPassed = bHistogramGraphsFilter || (QuickFilterText.Len() == 0 || CurrentData.Category.ToString().Find(QuickFilterText) != INDEX_NONE);
@@ -1040,6 +1142,7 @@ void SLogVisualizer::Tick( const FGeometry& AllottedGeometry, const double InCur
 	UWorld* World = GetWorld();
 	if (World && !World->bPlayersOnly && TimeTillNextUpdate < 0 && LogVisualizer->IsRecording())
 	{
+		FVisualLogger::Get().Flush();
 		DoTickUpdate();
 		//DoFullUpdate();
 	}
@@ -1188,6 +1291,18 @@ void SLogVisualizer::IncrementCurrentLogIndex(int32 IncrementBy)
 				}
 			}
 
+			if (!bShouldShow)
+			{
+				for (int32 LineIndex = 0; LineIndex < Log->Entries[NewEntryIndex]->Events.Num(); ++LineIndex)
+				{
+					if (Log->Entries[NewEntryIndex]->Events[LineIndex].Name.Len() == 0 || FilterListPtr->IsFilterEnabled(Log->Entries[NewEntryIndex]->Events[LineIndex].Name, Log->Entries[NewEntryIndex]->Events[LineIndex].Verbosity))
+					{
+						bShouldShow = true;
+						break;
+					}
+				}
+			}
+
 			if (!bShouldShow)
 			{
 				for (const auto CurrentData : Log->Entries[NewEntryIndex]->DataBlocks)
@@ -1294,6 +1409,18 @@ void SLogVisualizer::AddOrUpdateLog(int32 LogIndex, const FActorsVisLog* Log)
 			FilterListPtr->AddGraphFilter(GraphNameAsString, DataNameAsString, FColor::White);
 		}
 
+		for (auto Iter(Log->Entries[EntryIndex]->Events.CreateConstIterator()); Iter; Iter++)
+		{
+			const FString CategoryAsString = Iter->Name.Len() != 0 ? Iter->Name : TEXT("Event");
+
+			int32 Index = UsedCategories.Find(CategoryAsString);
+			if (Index == INDEX_NONE)
+			{
+				Index = UsedCategories.Add(CategoryAsString);
+				FilterListPtr->AddFilter(CategoryAsString, GetColorForUsedCategory(Index));
+			}
+		}
+
 		for (const auto CurrentData : Log->Entries[EntryIndex]->DataBlocks)
 		{
 			int32 Index = UsedCategories.Find(CurrentData.Category.ToString());
@@ -1765,6 +1892,42 @@ void SLogVisualizer::ShowEntry(const FVisualLogEntry* LogEntry)
 		}
 	}
 
+	for (auto& Event : LogEntry->Events)
+	{
+		bool bShowLine = true;
+
+		if (FilterListPtr.IsValid())
+		{
+			FString CurrentCategory = Event.Name;
+			bShowLine = FilterListPtr->IsFilterEnabled(CurrentCategory, Event.Verbosity) && (bHistogramGraphsFilter || (QuickFilterText.Len() == 0 || CurrentCategory.Find(QuickFilterText) != INDEX_NONE));
+		}
+
+		if (bShowLine)
+		{
+			FLogEntryItem EntryItem;
+			EntryItem.Category = Event.Name;
+
+			int32 Index = UsedCategories.Find(EntryItem.Category);
+			if (Index == INDEX_NONE)
+			{
+				Index = UsedCategories.Add(EntryItem.Category);
+			}
+			EntryItem.CategoryColor = GetColorForUsedCategory(Index);
+
+			EntryItem.Verbosity = Event.Verbosity;
+			EntryItem.Line = FString::Printf(TEXT("Registered event: '%s' (%d times)%s"), *Event.Name, Event.Counter, Event.EventTags.Num() ? TEXT("\n") : TEXT(""));
+			for (auto& EventTag : Event.EventTags)
+			{
+				EntryItem.Line += FString::Printf(TEXT("%d times for tag: '%s'\n"), EventTag.Value, *EventTag.Key.ToString());
+			}
+			EntryItem.UserData = Event.UserData;
+			EntryItem.TagName = Event.TagName;
+
+			LogEntryLines.Add(MakeShareable(new FLogEntryItem(EntryItem)));
+		}
+
+	}
+
 	SetCurrentViewedTime(LogEntry->TimeStamp);
 
 	UWorld* World = GetWorld();
diff --git a/Engine/Source/Developer/LogVisualizer/Private/SLogVisualizer.h b/Engine/Source/Developer/LogVisualizer/Private/SLogVisualizer.h
--- a/Engine/Source/Developer/LogVisualizer/Private/SLogVisualizer.h
+++ b/Engine/Source/Developer/LogVisualizer/Private/SLogVisualizer.h
@@ -114,6 +114,9 @@ public:
 
 	FOnHistogramWindowChanged& OnHistogramWindowChanged() { return HistogramWindowChangedNotify; }
 
+	TSharedPtr<SWidget> GetListRightClickMenuContent();
+	void GenerateReport();
+
 	void OnListDoubleClick(TSharedPtr<FLogsListItem>);
 	void DoFullUpdate();
 	void DoTickUpdate();
diff --git a/Engine/Source/Developer/LogVisualizer/Private/SLogVisualizerReport.cpp b/Engine/Source/Developer/LogVisualizer/Private/SLogVisualizerReport.cpp
new file mode 100644
--- /dev/null
+++ b/Engine/Source/Developer/LogVisualizer/Private/SLogVisualizerReport.cpp
@@ -0,0 +1,158 @@
+// Copyright 1998-2014 Epic Games, Inc. All Rights Reserved.
+
+#include "LogVisualizerPCH.h"
+#include "Debug/ReporterGraph.h"
+#include "MainFrame.h"
+#include "DesktopPlatformModule.h"
+#include "Json.h"
+#include "Editor/UnrealEd/Classes/Editor/EditorEngine.h"
+
+#if WITH_EDITOR
+#	include "Editor/UnrealEd/Public/EditorComponents.h"
+#	include "Editor/UnrealEd/Public/EditorReimportHandler.h"
+#	include "Editor/UnrealEd/Public/TexAlignTools.h"
+#	include "Editor/UnrealEd/Public/TickableEditorObject.h"
+#	include "UnrealEdClasses.h"
+#	include "Editor/UnrealEd/Public/Editor.h"
+#	include "Editor/UnrealEd/Public/EditorViewportClient.h"
+#endif
+
+#include "LogVisualizerModule.h"
+#include "SLogVisualizerReport.h"
+
+#define LOCTEXT_NAMESPACE "SLogVisualizerReport"
+
+BEGIN_SLATE_FUNCTION_BUILD_OPTIMIZATION
+void SLogVisualizerReport::Construct(const FArguments& InArgs, SLogVisualizer* InLogVisualizer, TArray< TSharedPtr<struct FActorsVisLog> >& InSelectedItems)
+{
+	SelectedItems = InSelectedItems;
+	this->ChildSlot
+		[
+			SNew(SScrollBox)
+			+ SScrollBox::Slot()
+			[
+				SNew(SBorder)
+				.BorderImage(FEditorStyle::Get().GetBrush("ToolPanel.GroupBorder"))
+				[
+					SNew(SVerticalBox)
+					+ SVerticalBox::Slot().AutoHeight().Padding(0)
+					[
+						SNew(SBorder)
+						.BorderImage(FEditorStyle::Get().GetBrush("Menu.Background"))
+						.Padding(0)
+						[
+							SNew(SRichTextBlock)
+							.Text(this, &SLogVisualizerReport::GenerateReportText)
+							.TextStyle(FEditorStyle::Get(), "ContentBrowser.Filters.Text")
+							.DecoratorStyleSet(&FEditorStyle::Get())
+							.WrapTextAt(800)
+							.Justification(ETextJustify::Left)
+							.Margin(FMargin(20))
+						]
+					]
+				]
+			]
+		];
+}
+END_SLATE_FUNCTION_BUILD_OPTIMIZATION
+
+SLogVisualizerReport::~SLogVisualizerReport()
+{
+}
+
+FText SLogVisualizerReport::GenerateReportText() const
+{
+	FString OutString;
+	TArray<FVisualLogEntry::FLogEvent> GlobalEventsStats;
+
+	for (auto LogItem : SelectedItems)
+	{
+		TArray<FVisualLogEntry::FLogEvent> AllEvents;
+
+		for (TSharedPtr<FVisualLogEntry> CurrentEntry : LogItem->Entries)
+		{
+			for (FVisualLogEntry::FLogEvent& Event : CurrentEntry->Events)
+			{
+				int32 Index = AllEvents.Find(FVisualLogEntry::FLogEvent(Event));
+				if (Index != INDEX_NONE)
+				{
+					for (auto& CurrentEventTag : Event.EventTags)
+					{
+						if (AllEvents[Index].EventTags.Contains(CurrentEventTag.Key))
+						{
+							AllEvents[Index].EventTags[CurrentEventTag.Key] += CurrentEventTag.Value;
+						}
+						else
+						{
+							AllEvents[Index].EventTags.Add(CurrentEventTag.Key, CurrentEventTag.Value);
+						}
+					}
+					AllEvents[Index].Counter++;
+				}
+				else
+				{
+					AllEvents.Add(Event);
+				}
+			}
+		}
+
+		bool bPrintNextLine = false;
+		if (AllEvents.Num() > 0)
+		{
+			OutString.Append(FString::Printf(TEXT("%s\n"), *LogItem->Name.ToString()));
+		}
+		for (auto& CurrentEvent : AllEvents)
+		{
+			bPrintNextLine = CurrentEvent.EventTags.Num() > 0;
+			for (auto& CurrentEventTag : CurrentEvent.EventTags)
+			{
+				OutString.Append(FString::Printf(TEXT("        %s <%s> occurred %d times\n"), *CurrentEvent.Name, *CurrentEventTag.Key.ToString(), CurrentEventTag.Value));
+			}
+			OutString.Append( FString::Printf(TEXT("        %s occurred %d times\n"), *CurrentEvent.Name, CurrentEvent.Counter) );
+
+			int32 Index = GlobalEventsStats.Find(FVisualLogEntry::FLogEvent(CurrentEvent));
+			if (Index != INDEX_NONE)
+			{
+				GlobalEventsStats[Index].Counter += CurrentEvent.Counter;
+				for (auto& CurrentEventTag : CurrentEvent.EventTags)
+				{
+					if (GlobalEventsStats[Index].EventTags.Contains(CurrentEventTag.Key))
+					{
+						GlobalEventsStats[Index].EventTags[CurrentEventTag.Key] += CurrentEventTag.Value;
+					}
+					else
+					{
+						GlobalEventsStats[Index].EventTags.Add(CurrentEventTag.Key, CurrentEventTag.Value);
+					}
+				}
+			}
+			else
+			{
+				GlobalEventsStats.Add(CurrentEvent);
+			}
+
+			if (bPrintNextLine)
+			{
+				OutString.Append(TEXT("\n"));
+			}
+		}
+
+		OutString.Append(TEXT("\n"));
+	}
+
+	OutString.Append(TEXT("-===========================================-\n"));
+
+	for (auto& CurrentEvent : GlobalEventsStats)
+	{
+		OutString.Append(FString::Printf(TEXT(" %s  (%s) occurred %d times\n"), *CurrentEvent.Name, *CurrentEvent.UserFriendlyDesc, CurrentEvent.Counter));
+		for (auto& CurrentEventTag : CurrentEvent.EventTags)
+		{
+			OutString.Append(FString::Printf(TEXT("        %d  times for <%s>\n"), CurrentEventTag.Value, *CurrentEventTag.Key.ToString()));
+		}
+	}
+
+	return FText::FromString(OutString);
+}
+
+
+#undef LOCTEXT_NAMESPACE
diff --git a/Engine/Source/Developer/LogVisualizer/Private/SLogVisualizerReport.h b/Engine/Source/Developer/LogVisualizer/Private/SLogVisualizerReport.h
new file mode 100644
--- /dev/null
+++ b/Engine/Source/Developer/LogVisualizer/Private/SLogVisualizerReport.h
@@ -0,0 +1,20 @@
+// Copyright 1998-2014 Epic Games, Inc. All Rights Reserved.
+
+#pragma once
+
+class SLogVisualizerReport : public SCompoundWidget
+{
+public:
+	SLATE_BEGIN_ARGS(SLogVisualizerReport) {}
+
+	SLATE_END_ARGS()
+
+	void Construct(const FArguments& InArgs, SLogVisualizer* InAnalyzer, TArray< TSharedPtr<struct FActorsVisLog> >&);
+
+	virtual ~SLogVisualizerReport();
+
+	FText GenerateReportText() const;
+
+protected:
+	TArray< TSharedPtr<struct FActorsVisLog> > SelectedItems;
+};
diff --git a/Engine/Source/Runtime/AIModule/Classes/EnvironmentQuery/EnvQueryDebugHelpers.h b/Engine/Source/Runtime/AIModule/Classes/EnvironmentQuery/EnvQueryDebugHelpers.h
--- a/Engine/Source/Runtime/AIModule/Classes/EnvironmentQuery/EnvQueryDebugHelpers.h
+++ b/Engine/Source/Runtime/AIModule/Classes/EnvironmentQuery/EnvQueryDebugHelpers.h
@@ -141,7 +141,7 @@ FArchive& operator<<(FArchive& Ar, EQSDebug::FQueryData& Data)
 
 #define UE_VLOG_EQS(Query, CategoryName, Verbosity) \
 { \
-	UEnvQueryDebugHelpers::LogQuery(Query->Owner.Get(), Query, CategoryName.GetCategoryName(), ELogVerbosity::Verbosity, \
+	UEnvQueryDebugHelpers::LogQuery(Query, CategoryName.GetCategoryName(), ELogVerbosity::Verbosity, \
 		FString::Printf(TEXT("Executed EQS: \n - Name: '%s' (id=%d, option=%d),\n - All Items: %d,\n - ValidItems: %d"), *Query->QueryName, Query->QueryID, Query->OptionIndex, Query->ItemDetails.Num(), Query->NumValidItems)); \
 }
 
@@ -160,6 +160,46 @@ class AIMODULE_API UEnvQueryDebugHelpers : public UObject
 	static void QueryToDebugData(struct FEnvQueryInstance* Query, EQSDebug::FQueryData& EQSLocalData);
 	static void QueryToBlobArray(struct FEnvQueryInstance* Query, TArray<uint8>& BlobArray, bool bUseCompression = false);
 	static void BlobArrayToDebugData(const TArray<uint8>& BlobArray, EQSDebug::FQueryData& EQSLocalData, bool bUseCompression = false);
-	static void LogQuery(const class UObject* QueryOwner, struct FEnvQueryInstance* Query, const FName& CategoryName, ELogVerbosity::Type, const FString& AdditionalLogInfo);
+	static void LogQuery(struct FEnvQueryInstance* Query, const FName& CategoryName, ELogVerbosity::Type, const FString& AdditionalLogInfo);
 #endif
 };
+
+#if USE_EQS_DEBUGGER
+FORCEINLINE void UEnvQueryDebugHelpers::LogQuery(struct FEnvQueryInstance* Query, const FName& CategoryName, ELogVerbosity::Type Type, const FString& AdditionalLogInfo)
+{
+#if ENABLE_VISUAL_LOG
+	if (GEngine && GEngine->bDisableAILogging || Query == NULL || FVisualLogger::Get().IsRecording() == false || (FVisualLogger::Get().IsBlockedForAllCategories() && FVisualLogger::Get().GetWhiteList().Find(CategoryName) == INDEX_NONE))
+	{
+		return;
+	}
+
+	const UObject* LogOwner = FVisualLogger::Get().FindRedirection(Query->Owner.Get());
+	if (LogOwner == NULL || LogOwner->HasAnyFlags(RF_ClassDefaultObject))
+	{
+		return;
+	}
+
+	UWorld* World = GEngine->GetWorldFromContextObject(Query->Owner.Get(), false);
+	if (!World)
+	{
+		return;
+	}
+
+	TArray<uint8> BlobArray;
+	UEnvQueryDebugHelpers::QueryToBlobArray(Query, BlobArray);
+
+	FVisualLogEntry* EntryToWrite = FVisualLogger::Get().GetEntryToWrite(LogOwner, World->TimeSeconds);
+	const int32 UniqueId = FVisualLogger::Get().GetUniqueId(World->TimeSeconds);
+	if (EntryToWrite)
+	{
+		FVisualLogEntry::FLogLine Line(CategoryName, Type, AdditionalLogInfo, Query->QueryID);
+		Line.TagName = *EVisLogTags::TAG_EQS;
+		Line.UniqueId = UniqueId;
+		EntryToWrite->LogLines.Add(Line);
+
+		EntryToWrite->AddDataBlock(EVisLogTags::TAG_EQS, BlobArray, CategoryName).UniqueId = UniqueId;
+	}
+#endif
+}
+#endif
+
diff --git a/Engine/Source/Runtime/AIModule/Private/EnvironmentQuery/EnvQueryDebugHelpers.cpp b/Engine/Source/Runtime/AIModule/Private/EnvironmentQuery/EnvQueryDebugHelpers.cpp
--- a/Engine/Source/Runtime/AIModule/Private/EnvironmentQuery/EnvQueryDebugHelpers.cpp
+++ b/Engine/Source/Runtime/AIModule/Private/EnvironmentQuery/EnvQueryDebugHelpers.cpp
@@ -129,27 +129,4 @@ void  UEnvQueryDebugHelpers::BlobArrayToDebugData(const TArray<uint8>& BlobArray
 	}
 }
 
-void UEnvQueryDebugHelpers::LogQuery(const class UObject* QueryOwner, struct FEnvQueryInstance* Query, const FName& CategoryName, ELogVerbosity::Type Type, const FString& AdditionalLogInfo)
-{
-#if ENABLE_VISUAL_LOG
-	TArray<uint8> BlobArray;
-	UEnvQueryDebugHelpers::QueryToBlobArray(Query, BlobArray);
-
-	UObject* LogOwner = FVisualLogger::Get().FindRedirection(QueryOwner);
-	UWorld* World = GEngine->GetWorldFromContextObject(QueryOwner);
-
-	FVisualLogEntry* EntryToWrite = FVisualLogger::Get().GetEntryToWrite(LogOwner, World->TimeSeconds);
-	const int32 UniqueId = FVisualLogger::Get().GetUniqueId(World->TimeSeconds);
-	if (EntryToWrite)
-	{
-		FVisualLogEntry::FLogLine Line(CategoryName, Type, AdditionalLogInfo, Query->QueryID);
-		Line.TagName = *EVisLogTags::TAG_EQS;
-		Line.UniqueId = UniqueId;
-		EntryToWrite->LogLines.Add(Line);
-
-		EntryToWrite->AddDataBlock(EVisLogTags::TAG_EQS, BlobArray, CategoryName).UniqueId = UniqueId;
-	}
-#endif
-}
-
 #endif //ENABLE_VISUAL_LOG
diff --git a/Engine/Source/Runtime/Engine/Private/VisualLog.cpp b/Engine/Source/Runtime/Engine/Private/VisualLog.cpp
--- a/Engine/Source/Runtime/Engine/Private/VisualLog.cpp
+++ b/Engine/Source/Runtime/Engine/Private/VisualLog.cpp
@@ -464,6 +464,7 @@ void FVisualLog::SetIsRecording(bool NewRecording, bool bRecordToFile)
 		if (FileAr)
 		{
 			// dump remaining logs
+			FVisualLogger::Get().Flush();
 			DumpRecordedLogs();
 
 			// close JSON data correctly
diff --git a/Engine/Source/Runtime/Engine/Private/VisualLogger/VisualLogger.cpp b/Engine/Source/Runtime/Engine/Private/VisualLogger/VisualLogger.cpp
--- a/Engine/Source/Runtime/Engine/Private/VisualLogger/VisualLogger.cpp
+++ b/Engine/Source/Runtime/Engine/Private/VisualLogger/VisualLogger.cpp
@@ -25,11 +25,14 @@ FVisualLogger::FVisualLogger()
 }
 
 FVisualLogger::~FVisualLogger()
+{
+	RemoveDevice(&FVisualLog::GetStatic());
+}
+
+void FVisualLogger::Shutdown()
 {
 	SetIsRecording(false);
 	SetIsRecordingToFile(false);
-
-	RemoveDevice(&FVisualLog::GetStatic());
 }
 
 void FVisualLogger::Cleanup(bool bReleaseMemory)
@@ -108,7 +111,7 @@ void FVisualLogger::SetIsRecordingToFile(bool InIsRecording)
 
 	FString BaseFileName = TEXT("VisualLog");
 	FString UserDefinedName;
-	FString MapName = GWorld ? GWorld->GetMapName() : TEXT("");
+	FString MapName = GWorld.GetReference() ? GWorld->GetMapName() : TEXT("");
 	if (LogFileNameGetter.IsBound())
 	{
 		UserDefinedName = LogFileNameGetter.Execute();
@@ -131,13 +134,13 @@ void FVisualLogger::SetIsRecordingToFile(bool InIsRecording)
 		{
 			if (Device->HasFlags(VisualLogger::CanSaveToFile))
 			{
-				Device->StopRecordingToFile(GWorld ? GWorld->TimeSeconds : StartRecordingToFileTime);
+				Device->StopRecordingToFile(GWorld.GetReference() ? GWorld->TimeSeconds : StartRecordingToFileTime);
 			}
 		}
 	}
 	else if (!bIsRecordingToFile && InIsRecording)
 	{
-		StartRecordingToFileTime = GWorld ? GWorld->TimeSeconds : 0;
+		StartRecordingToFileTime = GWorld.GetReference() ? GWorld->TimeSeconds : 0;
 		for (auto* Device : OutputDevices)
 		{
 			if (Device->HasFlags(VisualLogger::CanSaveToFile))
diff --git a/Engine/Source/Runtime/Engine/Private/VisualLogger/VisualLoggerAutomationTests.cpp b/Engine/Source/Runtime/Engine/Private/VisualLogger/VisualLoggerAutomationTests.cpp
--- a/Engine/Source/Runtime/Engine/Private/VisualLogger/VisualLoggerAutomationTests.cpp
+++ b/Engine/Source/Runtime/Engine/Private/VisualLogger/VisualLoggerAutomationTests.cpp
@@ -122,18 +122,19 @@ bool FVisualLogTest::RunTest(const FString& Parameters)
 		CHECK_SUCCESS(Context.Device.LastObject != GWorld);
 		CHECK_SUCCESS(Context.Device.LastEntry.TimeStamp == -1);
 		FVisualLogEntry* CurrentEntry = FVisualLogger::Get().GetEntryToWrite(GWorld, GWorld->TimeSeconds, VisualLogger::DontCreate);
-		CHECK_SUCCESS(CurrentEntry != NULL);
-		CHECK_SUCCESS(CurrentEntry->TimeStamp == GWorld->TimeSeconds);
-		CHECK_SUCCESS(CurrentEntry->LogLines.Num() == 1);
-		CHECK_SUCCESS(CurrentEntry->LogLines[0].Category == LogVisual.GetCategoryName());
-		CHECK_SUCCESS(CurrentEntry->LogLines[0].Line == TextToLog);
-
-		FVisualLogger::Get().GetEntryToWrite(GWorld, GWorld->TimeSeconds+0.1); //generate new entry and serialize old one
-		CHECK_SUCCESS(Context.Device.LastEntry.TimeStamp == GWorld->TimeSeconds);
-		CHECK_SUCCESS(Context.Device.LastObject == GWorld);
-		CHECK_SUCCESS(Context.Device.LastEntry.LogLines.Num() == 1);
-		CHECK_SUCCESS(Context.Device.LastEntry.LogLines[0].Category == LogVisual.GetCategoryName());
-		CHECK_SUCCESS(Context.Device.LastEntry.LogLines[0].Line == TextToLog);
+		
+		float CurrentTimestamp = GWorld->TimeSeconds;
+		for (int32 Index = 0; Index < 2; ++Index)
+		{
+			CHECK_SUCCESS(CurrentEntry != NULL);
+			CHECK_SUCCESS(CurrentEntry->TimeStamp == CurrentTimestamp);
+			CHECK_SUCCESS(CurrentEntry->LogLines.Num() == 1);
+			CHECK_SUCCESS(CurrentEntry->LogLines[0].Category == LogVisual.GetCategoryName());
+			CHECK_SUCCESS(CurrentEntry->LogLines[0].Line == TextToLog);
+
+			FVisualLogger::Get().GetEntryToWrite(GWorld, CurrentTimestamp + 0.1); //generate new entry and serialize old one
+			CurrentEntry = &Context.Device.LastEntry;
+		}
 	}
 
 	FVisualLogger::Get().SetIsRecording(false);
@@ -163,14 +164,118 @@ bool FVisualLogSegmentsTest::RunTest(const FString& Parameters)
 		CHECK_SUCCESS(Context.Device.LastObject == NULL);
 		CHECK_SUCCESS(Context.Device.LastEntry.TimeStamp == -1);
 		FVisualLogEntry* CurrentEntry = FVisualLogger::Get().GetEntryToWrite(GWorld, GWorld->TimeSeconds, VisualLogger::DontCreate);
+
+		float CurrentTimestamp = GWorld->TimeSeconds;
+		for (int32 Index = 0; Index < 2; ++Index)
+		{
+			CHECK_SUCCESS(CurrentEntry != NULL);
+			CHECK_SUCCESS(CurrentEntry->TimeStamp == CurrentTimestamp);
+			CHECK_SUCCESS(CurrentEntry->ElementsToDraw.Num() == 1);
+			CHECK_SUCCESS(CurrentEntry->ElementsToDraw[0].GetType() == FVisualLogEntry::FElementToDraw::Segment);
+			CHECK_SUCCESS(CurrentEntry->ElementsToDraw[0].Points.Num() == 2);
+			CHECK_SUCCESS(CurrentEntry->ElementsToDraw[0].Points[0] == StartPoint);
+			CHECK_SUCCESS(CurrentEntry->ElementsToDraw[0].Points[1] == EndPoint);
+
+			FVisualLogger::Get().GetEntryToWrite(GWorld, CurrentTimestamp + 0.1); //generate new entry and serialize old one
+			CurrentEntry = &Context.Device.LastEntry;
+		}
+	}
+	return true;
+}
+
+IMPLEMENT_SIMPLE_AUTOMATION_TEST(FVisualLogEventsTest, "Engine.VisualLogger.Logging events", EAutomationTestFlags::ATF_Game | EAutomationTestFlags::ATF_Editor)
+
+DECLARE_VLOG_EVENT(EventTest, All, "Simple event for tests")
+DEFINE_VLOG_EVENT(EventTest)
+
+DECLARE_VLOG_EVENT(EventTest2, All, "Second simple event for tests")
+DEFINE_VLOG_EVENT(EventTest2)
+
+DECLARE_VLOG_EVENT(EventTest3, All, "Third simple event for tests")
+DEFINE_VLOG_EVENT(EventTest3)
+
+bool FVisualLogEventsTest::RunTest(const FString& Parameters)
+{
+	FTestDeviceContext Context;
+	FVisualLogger::Get().SetIsRecording(true);
+
+	CHECK_SUCCESS(EventTest.GetName() == TEXT("EventTest"));
+	CHECK_SUCCESS(EventTest.GetUserFriendlyDesc() == TEXT("Simple event for tests"));
+
+	CHECK_SUCCESS(EventTest2.GetName() == TEXT("EventTest2"));
+	CHECK_SUCCESS(EventTest2.GetUserFriendlyDesc() == TEXT("Second simple event for tests"));
+
+	CHECK_SUCCESS(EventTest3.GetName() == TEXT("EventTest3"));
+	CHECK_SUCCESS(EventTest3.GetUserFriendlyDesc() == TEXT("Third simple event for tests"));
+
+	FVisualLogEntry* CurrentEntry = FVisualLogger::Get().GetEntryToWrite(GWorld, GWorld->TimeSeconds, VisualLogger::DontCreate);
+	float CurrentTimestamp = GWorld->TimeSeconds;
+	UE_VLOG_EVENTS(GWorld, NAME_None, EventTest);
+	CHECK_SUCCESS(CurrentEntry != NULL);
+	CHECK_SUCCESS(CurrentEntry->TimeStamp == CurrentTimestamp);
+	CHECK_SUCCESS(CurrentEntry->Events.Num() == 1);
+	CHECK_SUCCESS(CurrentEntry->Events[0].Name == TEXT("EventTest"));
+
+	UE_VLOG_EVENTS(GWorld, NAME_None, EventTest, EventTest2);
+	CHECK_SUCCESS(CurrentEntry != NULL);
+	CHECK_SUCCESS(CurrentEntry->TimeStamp == CurrentTimestamp);
+	CHECK_SUCCESS(CurrentEntry->Events.Num() == 2);
+	CHECK_SUCCESS(CurrentEntry->Events[0].Counter == 2);
+	CHECK_SUCCESS(CurrentEntry->Events[0].Name == TEXT("EventTest"));
+	CHECK_SUCCESS(CurrentEntry->Events[1].Counter == 1);
+	CHECK_SUCCESS(CurrentEntry->Events[1].Name == TEXT("EventTest2"));
+
+	CurrentTimestamp = GWorld->TimeSeconds;
+	UE_VLOG_EVENTS(GWorld, NAME_None, EventTest, EventTest2, EventTest3);
+
+	for (int32 Index = 0; Index < 2; ++Index)
+	{
 		CHECK_SUCCESS(CurrentEntry != NULL);
-		CHECK_SUCCESS(CurrentEntry->TimeStamp == GWorld->TimeSeconds);
-		CHECK_SUCCESS(CurrentEntry->ElementsToDraw.Num() == 1);
-		CHECK_SUCCESS(CurrentEntry->ElementsToDraw[0].GetType() == FVisualLogEntry::FElementToDraw::Segment);
-		CHECK_SUCCESS(CurrentEntry->ElementsToDraw[0].Points.Num() == 2);
-		CHECK_SUCCESS(CurrentEntry->ElementsToDraw[0].Points[0] == StartPoint);
-		CHECK_SUCCESS(CurrentEntry->ElementsToDraw[0].Points[1] == EndPoint);
+		CHECK_SUCCESS(CurrentEntry->TimeStamp == CurrentTimestamp);
+		CHECK_SUCCESS(CurrentEntry->Events.Num() == 3);
+		CHECK_SUCCESS(CurrentEntry->Events[0].Counter == 3);
+		CHECK_SUCCESS(CurrentEntry->Events[0].Name == TEXT("EventTest"));
+		CHECK_SUCCESS(CurrentEntry->Events[1].Counter == 2);
+		CHECK_SUCCESS(CurrentEntry->Events[1].Name == TEXT("EventTest2"));
+		CHECK_SUCCESS(CurrentEntry->Events[2].Counter == 1);
+		CHECK_SUCCESS(CurrentEntry->Events[2].Name == TEXT("EventTest3"));
+
+		CHECK_SUCCESS(CurrentEntry->Events[0].UserFriendlyDesc == TEXT("Simple event for tests"));
+		CHECK_SUCCESS(CurrentEntry->Events[1].UserFriendlyDesc == TEXT("Second simple event for tests"));
+		CHECK_SUCCESS(CurrentEntry->Events[2].UserFriendlyDesc == TEXT("Third simple event for tests"));
+
+		FVisualLogger::Get().GetEntryToWrite(GWorld, CurrentTimestamp+0.1); //generate new entry and serialize old one
+		CurrentEntry = &Context.Device.LastEntry;
 	}
+
+	const FName EventTag1 = TEXT("ATLAS_C_0");
+	const FName EventTag2 = TEXT("ATLAS_C_1");
+	const FName EventTag3 = TEXT("ATLAS_C_2");
+
+	CurrentTimestamp = GWorld->TimeSeconds + 0.2;
+	CurrentEntry = FVisualLogger::Get().GetEntryToWrite(GWorld, CurrentTimestamp); //generate new entry and serialize old one
+	UE_VLOG_EVENT_WITH_DATA(GWorld, EventTest, EventTag1);
+	CHECK_SUCCESS(CurrentEntry != NULL);
+	CHECK_SUCCESS(CurrentEntry->TimeStamp == CurrentTimestamp);
+	CHECK_SUCCESS(CurrentEntry->Events.Num() == 1);
+	CHECK_SUCCESS(CurrentEntry->Events[0].Name == TEXT("EventTest"));
+	CHECK_SUCCESS(CurrentEntry->Events[0].EventTags.Num() == 1);
+	CHECK_SUCCESS(CurrentEntry->Events[0].EventTags[EventTag1] == 1);
+
+	CurrentTimestamp = GWorld->TimeSeconds + 0.2;
+	CurrentEntry = FVisualLogger::Get().GetEntryToWrite(GWorld, CurrentTimestamp); //generate new entry and serialize old one
+	UE_VLOG_EVENT_WITH_DATA(GWorld, EventTest, EventTag1, EventTag2, EventTag3);
+	UE_VLOG_EVENT_WITH_DATA(GWorld, EventTest, EventTag3);
+	CHECK_SUCCESS(CurrentEntry != NULL);
+	CHECK_SUCCESS(CurrentEntry->TimeStamp == CurrentTimestamp);
+	CHECK_SUCCESS(CurrentEntry->Events.Num() == 1);
+	CHECK_SUCCESS(CurrentEntry->Events[0].Name == TEXT("EventTest"));
+	CHECK_SUCCESS(CurrentEntry->Events[0].EventTags.Num() == 3);
+	CHECK_SUCCESS(CurrentEntry->Events[0].EventTags[EventTag1] == 1);
+	CHECK_SUCCESS(CurrentEntry->Events[0].EventTags[EventTag2] == 1);
+	CHECK_SUCCESS(CurrentEntry->Events[0].EventTags[EventTag3] == 2);
+
+
 	return true;
 }
 
diff --git a/Engine/Source/Runtime/Engine/Private/VisualLogger/VisualLoggerTypes.cpp b/Engine/Source/Runtime/Engine/Private/VisualLogger/VisualLoggerTypes.cpp
--- a/Engine/Source/Runtime/Engine/Private/VisualLogger/VisualLoggerTypes.cpp
+++ b/Engine/Source/Runtime/Engine/Private/VisualLogger/VisualLoggerTypes.cpp
@@ -11,6 +11,7 @@ FVisualLogEntry::FVisualLogEntry(const FVisualLogEntry& Entry)
 	TimeStamp = Entry.TimeStamp;
 	Location = Entry.Location;
 
+	Events = Entry.Events;
 	LogLines = Entry.LogLines;
 	Status = Entry.Status;
 	ElementsToDraw = Entry.ElementsToDraw;
@@ -69,6 +70,11 @@ FVisualLogEntry::FVisualLogEntry(float InTimeStamp, FVector InLocation, const UO
 	}
 }
 
+int32 FVisualLogEntry::AddEvent(const FVisualLogEventBase& Event)
+{
+	return Events.Add(Event);
+}
+
 void FVisualLogEntry::AddText(const FString& TextLine, const FName& CategoryName)
 {
 	LogLines.Add(FLogLine(CategoryName, ELogVerbosity::All, TextLine));
diff --git a/Engine/Source/Runtime/Engine/Public/VisualLog.h b/Engine/Source/Runtime/Engine/Public/VisualLog.h
--- a/Engine/Source/Runtime/Engine/Public/VisualLog.h
+++ b/Engine/Source/Runtime/Engine/Public/VisualLog.h
@@ -143,7 +143,7 @@ protected:
 
 	void LogElementImpl(FVisLogSelfDrawingElement* Element);
 
-	FORCEINLINE_DEBUGGABLE TSharedPtr<FActorsVisLog> GetLog(const class UObject* Object)
+	FORCEINLINE TSharedPtr<FActorsVisLog> GetLog(const class UObject* Object)
 	{
 		TSharedPtr<FActorsVisLog>* Log = LogsMap.Find(Object);
 		if (Log == NULL)
diff --git a/Engine/Source/Runtime/Engine/Public/VisualLogger/VisualLogger.h b/Engine/Source/Runtime/Engine/Public/VisualLogger/VisualLogger.h
--- a/Engine/Source/Runtime/Engine/Public/VisualLogger/VisualLogger.h
+++ b/Engine/Source/Runtime/Engine/Public/VisualLogger/VisualLogger.h
@@ -39,6 +39,9 @@
 #define UE_VLOG_HISTOGRAM(LogOwner, CategoryName, Verbosity, GraphName, DataName, Data) FVisualLogger::HistogramDataLogf(LogOwner, CategoryName, ELogVerbosity::Verbosity, INDEX_NONE, GraphName, DataName, Data, FColor::White, TEXT(""))
 #define UE_CVLOG_HISTOGRAM(LogOwner, CategoryName, Verbosity, GraphName, DataName, Data) if(Condition) {UE_VLOG_HISTOGRAM(LogOwner, CategoryName, Verbosity, GraphName, DataName, Data);} 
 
+#define UE_VLOG_EVENTS(LogOwner, TagNameToLog, ...) FVisualLogger::EventLog(LogOwner, TagNameToLog, ##__VA_ARGS__)
+#define UE_VLOG_EVENT_WITH_DATA(LogOwner, LogEvent, ...) FVisualLogger::EventLog(LogOwner, LogEvent, ##__VA_ARGS__)
+
 #else
 #define REDIRECT_TO_VLOG(Dest)
 #define REDIRECT_OBJECT_TO_VLOG(Src, Dest)
@@ -64,6 +67,11 @@
 #define UE_VLOG_HISTOGRAM(Actor, CategoryName, Verbosity, GraphName, DataName, Data)
 #define UE_CVLOG_HISTOGRAM(Actor, CategoryName, Verbosity, GraphName, DataName, Data)
 
+#define DECLARE_VLOG_EVENT(EventName, UserFriendlyDesc)
+#define DEFINE_VLOG_EVENT(EventName) 
+#define UE_VLOG_EVENTS(LogOwner, TagNameToLog, ...) 
+#define UE_VLOG_EVENT_WITH_DATA(LogOwner, LogEvent, ...)
+
 #endif //ENABLE_VISUAL_LOG
 
 // helper macros
@@ -99,6 +107,20 @@ public:
 	//Histogram data
 	VARARG_DECL(static void, static void, return, HistogramDataLogf, VARARG_NONE, const TCHAR*, VARARG_EXTRA(const class UObject* LogOwner) VARARG_EXTRA(const struct FLogCategoryBase& Category) VARARG_EXTRA(ELogVerbosity::Type Verbosity) VARARG_EXTRA(int32 UniqueLogId) VARARG_EXTRA(FName GraphName) VARARG_EXTRA(FName DataName) VARARG_EXTRA(const FVector2D& Data) VARARG_EXTRA(const FColor& Color), VARARG_EXTRA(LogOwner) VARARG_EXTRA(Category) VARARG_EXTRA(Verbosity) VARARG_EXTRA(UniqueLogId) VARARG_EXTRA(GraphName) VARARG_EXTRA(DataName) VARARG_EXTRA(Data) VARARG_EXTRA(Color));
 
+	/** Log events */
+	static void EventLog(const class UObject* LogOwner, const FName EventTag1, const FVisualLogEventBase& Event1, const FName EventTag2 = NAME_None, const FName EventTag3 = NAME_None, const FName EventTag4 = NAME_None, const FName EventTag5 = NAME_None, const FName EventTag6 = NAME_None);
+	static void EventLog(const class UObject* LogOwner, const FName EventTag1, const FVisualLogEventBase& Event1, const FVisualLogEventBase& Event2);
+	static void EventLog(const class UObject* LogOwner, const FName EventTag1, const FVisualLogEventBase& Event1, const FVisualLogEventBase& Event2, const FVisualLogEventBase& Event3);
+	static void EventLog(const class UObject* LogOwner, const FName EventTag1, const FVisualLogEventBase& Event1, const FVisualLogEventBase& Event2, const FVisualLogEventBase& Event3, const FVisualLogEventBase& Event4);
+	static void EventLog(const class UObject* LogOwner, const FName EventTag1, const FVisualLogEventBase& Event1, const FVisualLogEventBase& Event2, const FVisualLogEventBase& Event3, const FVisualLogEventBase& Event4, const FVisualLogEventBase& Event5);
+	static void EventLog(const class UObject* LogOwner, const FName EventTag1, const FVisualLogEventBase& Event1, const FVisualLogEventBase& Event2, const FVisualLogEventBase& Event3, const FVisualLogEventBase& Event4, const FVisualLogEventBase& Event5, const FVisualLogEventBase& Event6);
+
+	FORCEINLINE
+	static void EventLog(const class UObject* LogOwner, const FVisualLogEventBase& Event1, const FName EventTag1 = NAME_None, const FName EventTag2 = NAME_None, const FName EventTag3 = NAME_None, const FName EventTag4 = NAME_None, const FName EventTag5 = NAME_None, const FName EventTag6 = NAME_None)
+	{
+		EventLog(LogOwner, EventTag1, Event1, EventTag2, EventTag3, EventTag4, EventTag5, EventTag6);
+	}
+
 	// static getter
 	static FVisualLogger& Get()
 	{
@@ -106,9 +128,11 @@ public:
 		return GVisLog;
 	}
 
-	FVisualLogger();
 	virtual ~FVisualLogger();
 
+	// called on engine shutdown to flush all, etc.
+	virtual void Shutdown();
+
 	// Removes all logged data 
 	void Cleanup(bool bReleaseMemory = false);
 
@@ -155,6 +179,8 @@ public:
 	TArray<FVisualLogDevice*>& GetDevices() { return OutputDevices; }
 	/** Returns  current entry for given TimeStap or creates another one  but first it serialize previous entry as completed to vislog devices. Use VisualLogger::DontCreate to get current entry without serialization*/
 	FVisualLogEntry* GetEntryToWrite(const class UObject* Object, float TimeStamp, VisualLogger::ECreateIfNeeded ShouldCreate = VisualLogger::Create);
+	/** flush and serialize data if timestamp allows it */
+	virtual void Flush() override;
 
 	/** FileName getter to set project specific file name for vlogs - highly encouradged to use FVisualLogFilenameGetterDelegate::CreateUObject with this */
 	void SetLogFileNameGetter(const FVisualLogFilenameGetterDelegate& InLogFileNameGetter) { LogFileNameGetter = InLogFileNameGetter; }
@@ -169,6 +195,7 @@ public:
 	TMap<FName, FVisualLogExtensionInterface*>& GetAllExtensions() { return AllExtensions; }
 
 private:
+	FVisualLogger();
 	virtual void Serialize(const TCHAR* V, ELogVerbosity::Type Verbosity, const class FName& Category) override { ensureMsgf(0, TEXT("Regular serialize is forbiden for visual logs")); }
 
 protected:
diff --git a/Engine/Source/Runtime/Engine/Public/VisualLogger/VisualLogger.inl b/Engine/Source/Runtime/Engine/Public/VisualLogger/VisualLogger.inl
--- a/Engine/Source/Runtime/Engine/Public/VisualLogger/VisualLogger.inl
+++ b/Engine/Source/Runtime/Engine/Public/VisualLogger/VisualLogger.inl
@@ -10,12 +10,15 @@
 		return; \
 	} \
 	UObject* LogOwner = FindRedirection(Object); \
-	UWorld* World = GEngine->GetWorldFromContextObject(Object); \
-	if (World == NULL) \
+	if (LogOwner == NULL) \
 	{ \
-		ensure(World); \
 		return; \
 	} \
+	UWorld* World = GEngine->GetWorldFromContextObject(Object); \
+	if(ensure(World) == false) \
+	{ \
+		return; \
+	}
 
 FORCEINLINE_DEBUGGABLE
 FVisualLogEntry* FVisualLogger::GetEntryToWrite(const class UObject* Object, float TimeStamp, VisualLogger::ECreateIfNeeded ShouldCreate)
@@ -28,9 +31,12 @@ FVisualLogEntry* FVisualLogger::GetEntryToWrite(const class UObject* Object, flo
 		CurrentEntry = &CurrentEntryPerObject[Object];
 		if (TimeStamp > CurrentEntry->TimeStamp && ShouldCreate == VisualLogger::Create)
 		{
-			for (auto* Device : OutputDevices)
+			if (CurrentEntry->TimeStamp >= 0) //-1 means not initialized entry information
 			{
-				Device->Serialize(Object, *CurrentEntry);
+				for (auto* Device : OutputDevices)
+				{
+					Device->Serialize(Object, *CurrentEntry);
+				}
 			}
 			InitializeNewEntry = true;
 		}
@@ -67,6 +73,23 @@ FVisualLogEntry* FVisualLogger::GetEntryToWrite(const class UObject* Object, flo
 	return CurrentEntry;
 }
 
+FORCEINLINE_DEBUGGABLE
+void FVisualLogger::Flush()
+{
+	for (auto &CurrentEntry : CurrentEntryPerObject)
+	{
+		if (CurrentEntry.Value.TimeStamp >= 0)
+		{
+			for (auto* Device : OutputDevices)
+			{
+				Device->Serialize(CurrentEntry.Key, CurrentEntry.Value);
+			}
+			CurrentEntry.Value.Reset();
+		}
+	}
+}
+
+
 FORCEINLINE
 VARARG_BODY(void, FVisualLogger::CategorizedLogf, const TCHAR*, VARARG_EXTRA(const class UObject* Object) VARARG_EXTRA(const struct FLogCategoryBase& Category) VARARG_EXTRA(ELogVerbosity::Type Verbosity) VARARG_EXTRA(int32 UniqueLogId))
 {
@@ -171,4 +194,80 @@ VARARG_BODY(void, FVisualLogger::HistogramDataLogf, const TCHAR*, VARARG_EXTRA(c
 	);
 }
 
+FORCEINLINE
+void FVisualLogger::EventLog(const class UObject* Object, const FName EventTag1, const FVisualLogEventBase& Event1, const FVisualLogEventBase& Event2, const FVisualLogEventBase& Event3, const FVisualLogEventBase& Event4, const FVisualLogEventBase& Event5, const FVisualLogEventBase& Event6)
+{
+	EventLog(Object, EventTag1, Event1, Event2, Event3, Event4, Event5);
+	EventLog(Object, EventTag1, Event6);
+}
+
+FORCEINLINE
+void FVisualLogger::EventLog(const class UObject* Object, const FName EventTag1, const FVisualLogEventBase& Event1, const FVisualLogEventBase& Event2, const FVisualLogEventBase& Event3, const FVisualLogEventBase& Event4, const FVisualLogEventBase& Event5)
+{
+	EventLog(Object, EventTag1, Event1, Event2, Event3, Event4);
+	EventLog(Object, EventTag1, Event5);
+}
+
+FORCEINLINE
+void FVisualLogger::EventLog(const class UObject* Object, const FName EventTag1, const FVisualLogEventBase& Event1, const FVisualLogEventBase& Event2, const FVisualLogEventBase& Event3, const FVisualLogEventBase& Event4)
+{
+	EventLog(Object, EventTag1, Event1, Event2, Event3);
+	EventLog(Object, EventTag1, Event4);
+}
+
+FORCEINLINE
+void FVisualLogger::EventLog(const class UObject* Object, const FName EventTag1, const FVisualLogEventBase& Event1, const FVisualLogEventBase& Event2, const FVisualLogEventBase& Event3)
+{
+	EventLog(Object, EventTag1, Event1, Event2);
+	EventLog(Object, EventTag1, Event3);
+}
+
+FORCEINLINE
+void FVisualLogger::EventLog(const class UObject* Object, const FName EventTag1, const FVisualLogEventBase& Event1, const FVisualLogEventBase& Event2)
+{
+	EventLog(Object, EventTag1, Event1);
+	EventLog(Object, EventTag1, Event2);
+}
+
+FORCEINLINE_DEBUGGABLE_ACTUAL
+void FVisualLogger::EventLog(const class UObject* Object, const FName EventTag1, const FVisualLogEventBase& Event, const FName EventTag2, const FName EventTag3, const FName EventTag4, const FName EventTag5, const FName EventTag6)
+{
+	const FName CategoryName = *Event.GetName();
+
+	float CurrentTime = 0;
+	if (GEngine && GEngine->bDisableAILogging || FVisualLogger::Get().bIsRecording == false || !Object || Object->HasAnyFlags(RF_ClassDefaultObject) || (FVisualLogger::Get().IsBlockedForAllCategories() && FVisualLogger::Get().CategoriesWhiteList.Find(CategoryName) == INDEX_NONE))
+	{
+		return;
+	}
+	UObject* LogOwner = FindRedirection(Object);
+	UWorld* World = GEngine->GetWorldFromContextObject(Object);
+	if (World == NULL) 
+	{ 
+		ensure(World); 
+		return; 
+	} 
+
+	FVisualLogEntry* CurrentEntry = FVisualLogger::Get().GetEntryToWrite(LogOwner, World->TimeSeconds);
+	if (ensure(CurrentEntry))
+	{
+		int32 Index = CurrentEntry->Events.Find(FVisualLogEntry::FLogEvent(Event));
+		if (Index != INDEX_NONE)
+		{
+			CurrentEntry->Events[Index].Counter++;
+		}
+		else
+		{
+			Index = CurrentEntry->AddEvent(Event);
+		}
+
+		CurrentEntry->Events[Index].EventTags.FindOrAdd(EventTag1)++;
+		CurrentEntry->Events[Index].EventTags.FindOrAdd(EventTag2)++;
+		CurrentEntry->Events[Index].EventTags.FindOrAdd(EventTag3)++;
+		CurrentEntry->Events[Index].EventTags.FindOrAdd(EventTag4)++;
+		CurrentEntry->Events[Index].EventTags.FindOrAdd(EventTag5)++;
+		CurrentEntry->Events[Index].EventTags.FindOrAdd(EventTag6)++;
+		CurrentEntry->Events[Index].EventTags.Remove(NAME_None);
+	}
+}
+
 #undef  CHECK_VLOG_DATA
diff --git a/Engine/Source/Runtime/Engine/Public/VisualLogger/VisualLoggerTypes.h b/Engine/Source/Runtime/Engine/Public/VisualLogger/VisualLoggerTypes.h
--- a/Engine/Source/Runtime/Engine/Public/VisualLogger/VisualLoggerTypes.h
+++ b/Engine/Source/Runtime/Engine/Public/VisualLogger/VisualLoggerTypes.h
@@ -18,8 +18,52 @@ namespace VisualLogger
 	};
 }
 
+struct FVisualLogEventBase
+{
+	virtual FString GetName() const { return FString(); }
+	virtual FString GetUserFriendlyDesc() const { return FString(); }
+	virtual ELogVerbosity::Type GetVerbosity() const { return ELogVerbosity::All; }
+};
+
+#define DECLARE_VLOG_EVENT(EventName, Verbosity, UserFriendlyDesc) \
+struct FVisualLogEventBase_##EventName : public FVisualLogEventBase \
+{ \
+	virtual FString GetName() const override { return TEXT(#EventName); } \
+	virtual FString GetUserFriendlyDesc() const override { return TEXT(UserFriendlyDesc); } \
+	virtual ELogVerbosity::Type GetVerbosity() const override { return ELogVerbosity::Verbosity; } \
+};
+
+#define DEFINE_VLOG_EVENT(EventName) FVisualLogEventBase_##EventName EventName;
+
+
 struct ENGINE_API FVisualLogEntry
 {
+	struct FLogEvent
+	{
+		FString Name;
+		FString UserFriendlyDesc;
+		TEnumAsByte<ELogVerbosity::Type> Verbosity;
+		TMap<FName, int32>	 EventTags;
+		int32 Counter;
+		int64 UserData;
+		FName TagName;
+
+		FLogEvent() : Counter(1) {}
+		FLogEvent(const FVisualLogEventBase& Event) : Counter(1)
+		{
+			Name = Event.GetName();
+			UserFriendlyDesc = Event.GetUserFriendlyDesc();
+			Verbosity = Event.GetVerbosity();
+		}
+
+		FLogEvent& operator = (const FVisualLogEventBase& Event)
+		{
+			Name = Event.GetName();
+			UserFriendlyDesc = Event.GetUserFriendlyDesc();
+			Verbosity = Event.GetVerbosity();
+			return *this;
+		}
+	};
 	struct FLogLine
 	{
 		FString Line;
@@ -162,6 +206,7 @@ struct ENGINE_API FVisualLogEntry
 		float TimeStamp;
 		FVector Location;
 
+		TArray<FLogEvent> Events;
 		TArray<FLogLine> LogLines;
 		TArray<FStatusCategory> Status;
 		TArray<FElementToDraw> ElementsToDraw;
@@ -177,6 +222,7 @@ struct ENGINE_API FVisualLogEntry
 		{
 			TimeStamp = -1;
 			Location = FVector::ZeroVector;
+			Events.Reset();
 			LogLines.Reset();
 			Status.Reset();
 			ElementsToDraw.Reset();
@@ -203,20 +249,32 @@ struct ENGINE_API FVisualLogEntry
 		FVisualLogEntry::FDataBlock& AddDataBlock(const FString& TagName, const TArray<uint8>& BlobDataArray, const FName& CategoryName);
 		// capsule
 		void AddCapsule(FVector const& Center, float HalfHeight, float Radius, const FQuat & Rotation, const FName& CategoryName, const FColor& Color = FColor::White, const FString& Description = TEXT(""));
+		// Event
+		int32 AddEvent(const FVisualLogEventBase& Event);
+
+		int32 FindStatusIndex(const FString& CategoryName);
 
 		// find index of status category
-		FORCEINLINE int32 FindStatusIndex(const FString& CategoryName)
+
+#endif // ENABLE_VISUAL_LOG
+	};
+
+#if  ENABLE_VISUAL_LOG
+	FORCEINLINE int32 FVisualLogEntry::FindStatusIndex(const FString& CategoryName)
+	{
+		for (int32 TestCategoryIndex = 0; TestCategoryIndex < Status.Num(); TestCategoryIndex++)
 		{
-			for (int32 TestCategoryIndex = 0; TestCategoryIndex < Status.Num(); TestCategoryIndex++)
+			if (Status[TestCategoryIndex].Category == CategoryName)
 			{
-				if (Status[TestCategoryIndex].Category == CategoryName)
-				{
-					return TestCategoryIndex;
-				}
+				return TestCategoryIndex;
 			}
-
-			return INDEX_NONE;
 		}
 
+		return INDEX_NONE;
+	}
+
+	FORCEINLINE bool operator==(const FVisualLogEntry::FLogEvent& Left, const FVisualLogEntry::FLogEvent& Right)
+	{
+		return Left.Name == Right.Name;
+	}
 #endif // ENABLE_VISUAL_LOG
-	};
diff --git a/Engine/Source/Runtime/Launch/Private/LaunchEngineLoop.cpp b/Engine/Source/Runtime/Launch/Private/LaunchEngineLoop.cpp
--- a/Engine/Source/Runtime/Launch/Private/LaunchEngineLoop.cpp
+++ b/Engine/Source/Runtime/Launch/Private/LaunchEngineLoop.cpp
@@ -85,6 +85,10 @@
 	#include "HideWindowsPlatformTypes.h"
 #endif
 
+#if ENABLE_VISUAL_LOG
+#	include "VisualLogger/VisualLogger.h"
+#endif
+
 // Pipe output to std output
 // This enables UBT to collect the output for it's own use
 class FOutputDeviceStdOutput : public FOutputDevice
@@ -1911,6 +1915,11 @@ void FEngineLoop::Exit()
 	GIsRunning	= 0;
 	GLogConsole	= NULL;
 
+	// shutdown visual logger and flush all data
+#if ENABLE_VISUAL_LOG
+	FVisualLogger::Get().Shutdown();
+#endif
+
 	GetMoviePlayer()->Shutdown();
 
 	// Make sure we're not in the middle of loading something.
