CIS fix.

diff --git a/Engine/Source/Runtime/Engine/Private/PhysicsEngine/PhysicsSpring.cpp b/Engine/Source/Runtime/Engine/Private/PhysicsEngine/PhysicsSpring.cpp
--- a/Engine/Source/Runtime/Engine/Private/PhysicsEngine/PhysicsSpring.cpp
+++ b/Engine/Source/Runtime/Engine/Private/PhysicsEngine/PhysicsSpring.cpp
@@ -5,8 +5,8 @@
 #include "PhysXSupport.h"
 
 
-UPhysicsSpringComponent::UPhysicsSpringComponent(const class FPostConstructInitializeProperties& PCIP)
-	: Super(PCIP)
+UPhysicsSpringComponent::UPhysicsSpringComponent(const FObjectInitializer& ObjectInitializer)
+	: Super(ObjectInitializer)
 {
 	PrimaryComponentTick.bCanEverTick = true;
 	PrimaryComponentTick.TickGroup = TG_PreCloth;
