Minor changes in header for visual logger #ue4

diff --git a/Engine/Source/Runtime/Engine/Public/VisualLogger/VisualLogger.h b/Engine/Source/Runtime/Engine/Public/VisualLogger/VisualLogger.h
--- a/Engine/Source/Runtime/Engine/Public/VisualLogger/VisualLogger.h
+++ b/Engine/Source/Runtime/Engine/Public/VisualLogger/VisualLogger.h
@@ -41,6 +41,15 @@
 
 #define UE_VLOG_EVENTS(LogOwner, TagNameToLog, ...) FVisualLogger::EventLog(LogOwner, TagNameToLog, ##__VA_ARGS__)
 #define UE_VLOG_EVENT_WITH_DATA(LogOwner, LogEvent, ...) FVisualLogger::EventLog(LogOwner, LogEvent, ##__VA_ARGS__)
+#define DECLARE_VLOG_EVENT(EventName, Verbosity, UserFriendlyDesc) \
+	struct FVisualLogEventBase_##EventName : public FVisualLogEventBase \
+	{ \
+		virtual FString GetName() const override { return TEXT(#EventName); } \
+		virtual FString GetUserFriendlyDesc() const override { return TEXT(UserFriendlyDesc); } \
+		virtual ELogVerbosity::Type GetVerbosity() const override { return ELogVerbosity::Verbosity; } \
+	};
+
+#define DEFINE_VLOG_EVENT(EventName) FVisualLogEventBase_##EventName EventName;
 
 #else
 #define REDIRECT_TO_VLOG(Dest)
@@ -67,10 +76,10 @@
 #define UE_VLOG_HISTOGRAM(Actor, CategoryName, Verbosity, GraphName, DataName, Data)
 #define UE_CVLOG_HISTOGRAM(Actor, CategoryName, Verbosity, GraphName, DataName, Data)
 
-#define DECLARE_VLOG_EVENT(EventName, UserFriendlyDesc)
-#define DEFINE_VLOG_EVENT(EventName) 
 #define UE_VLOG_EVENTS(LogOwner, TagNameToLog, ...) 
 #define UE_VLOG_EVENT_WITH_DATA(LogOwner, LogEvent, ...)
+#define DECLARE_VLOG_EVENT(EventName, Verbosity, UserFriendlyDesc)
+#define DEFINE_VLOG_EVENT(EventName) 
 
 #endif //ENABLE_VISUAL_LOG
 
diff --git a/Engine/Source/Runtime/Engine/Public/VisualLogger/VisualLoggerTypes.h b/Engine/Source/Runtime/Engine/Public/VisualLogger/VisualLoggerTypes.h
--- a/Engine/Source/Runtime/Engine/Public/VisualLogger/VisualLoggerTypes.h
+++ b/Engine/Source/Runtime/Engine/Public/VisualLogger/VisualLoggerTypes.h
@@ -25,17 +25,6 @@ struct FVisualLogEventBase
 	virtual ELogVerbosity::Type GetVerbosity() const { return ELogVerbosity::All; }
 };
 
-#define DECLARE_VLOG_EVENT(EventName, Verbosity, UserFriendlyDesc) \
-struct FVisualLogEventBase_##EventName : public FVisualLogEventBase \
-{ \
-	virtual FString GetName() const override { return TEXT(#EventName); } \
-	virtual FString GetUserFriendlyDesc() const override { return TEXT(UserFriendlyDesc); } \
-	virtual ELogVerbosity::Type GetVerbosity() const override { return ELogVerbosity::Verbosity; } \
-};
-
-#define DEFINE_VLOG_EVENT(EventName) FVisualLogEventBase_##EventName EventName;
-
-
 struct ENGINE_API FVisualLogEntry
 {
 	struct FLogEvent
