Fixed automation test for visual logger (broken test and not visual logger) #ue4

diff --git a/Engine/Source/Runtime/Engine/Private/VisualLogger/VisualLoggerAutomationTests.cpp b/Engine/Source/Runtime/Engine/Private/VisualLogger/VisualLoggerAutomationTests.cpp
--- a/Engine/Source/Runtime/Engine/Private/VisualLogger/VisualLoggerAutomationTests.cpp
+++ b/Engine/Source/Runtime/Engine/Private/VisualLogger/VisualLoggerAutomationTests.cpp
@@ -262,7 +262,7 @@ bool FVisualLogEventsTest::RunTest(const FString& Parameters)
 	CHECK_SUCCESS(CurrentEntry->Events[0].EventTags.Num() == 1);
 	CHECK_SUCCESS(CurrentEntry->Events[0].EventTags[EventTag1] == 1);
 
-	CurrentTimestamp = GWorld->TimeSeconds + 0.2;
+	CurrentTimestamp = GWorld->TimeSeconds + 0.3;
 	CurrentEntry = FVisualLogger::Get().GetEntryToWrite(GWorld, CurrentTimestamp); //generate new entry and serialize old one
 	UE_VLOG_EVENT_WITH_DATA(GWorld, EventTest, EventTag1, EventTag2, EventTag3);
 	UE_VLOG_EVENT_WITH_DATA(GWorld, EventTest, EventTag3);
