Fixed CLang compilation issue #ue4

diff --git a/Engine/Source/Runtime/Engine/Public/VisualLogger/VisualLogger.inl b/Engine/Source/Runtime/Engine/Public/VisualLogger/VisualLogger.inl
--- a/Engine/Source/Runtime/Engine/Public/VisualLogger/VisualLogger.inl
+++ b/Engine/Source/Runtime/Engine/Public/VisualLogger/VisualLogger.inl
@@ -235,7 +235,7 @@ void FVisualLogger::EventLog(const class UObject* Object, const FName EventTag1,
 	const FName CategoryName = *Event.GetName();
 
 	float CurrentTime = 0;
-	if (GEngine && GEngine->bDisableAILogging || FVisualLogger::Get().bIsRecording == false || !Object || Object->HasAnyFlags(RF_ClassDefaultObject) || (FVisualLogger::Get().IsBlockedForAllCategories() && FVisualLogger::Get().CategoriesWhiteList.Find(CategoryName) == INDEX_NONE))
+	if ((GEngine && GEngine->bDisableAILogging) || FVisualLogger::Get().bIsRecording == false || !Object || Object->HasAnyFlags(RF_ClassDefaultObject) || (FVisualLogger::Get().IsBlockedForAllCategories() && FVisualLogger::Get().CategoriesWhiteList.Find(CategoryName) == INDEX_NONE))
 	{
 		return;
 	}
