DataTable - UObjects are handled properly while reconstrucing.

diff --git a/Engine/Source/Runtime/Engine/Classes/Engine/DataTable.h b/Engine/Source/Runtime/Engine/Classes/Engine/DataTable.h
--- a/Engine/Source/Runtime/Engine/Classes/Engine/DataTable.h
+++ b/Engine/Source/Runtime/Engine/Classes/Engine/DataTable.h
@@ -124,7 +124,7 @@ class UDataTable : public UObject
 private:
 	//when RowStruct is being modified, row data is stored serialized with tags
 	TArray<uint8> RowsSerializedWithTags;
-
+	TSet<UObject*> TemporarilyReferencedObjects;
 public:
 	ENGINE_API void CleanBeforeStructChange();
 	ENGINE_API void RestoreAfterStructChange();
diff --git a/Engine/Source/Runtime/Engine/Private/DataTable.cpp b/Engine/Source/Runtime/Engine/Private/DataTable.cpp
--- a/Engine/Source/Runtime/Engine/Private/DataTable.cpp
+++ b/Engine/Source/Runtime/Engine/Private/DataTable.cpp
@@ -180,6 +180,10 @@ void UDataTable::AddReferencedObjects(UObject* InThis, FReferenceCollector& Coll
 		}
 	}
 
+#if WITH_EDITOR || HACK_HEADER_GENERATOR
+	Collector.AddReferencedObjects(This->TemporarilyReferencedObjects);
+#endif //WITH_EDITOR || HACK_HEADER_GENERATOR
+
 	Super::AddReferencedObjects( This, Collector );
 }
 
@@ -248,8 +252,23 @@ struct FPropertyDisplayNameHelper
 void UDataTable::CleanBeforeStructChange()
 {
 	RowsSerializedWithTags.Reset();
+	TemporarilyReferencedObjects.Empty();
 	{
-		FMemoryWriter MemoryWriter(RowsSerializedWithTags);
+		class FRawStructWriter : public FObjectWriter
+		{
+			TSet<UObject*>& TemporarilyReferencedObjects;
+		public: 
+			FRawStructWriter(TArray<uint8>& InBytes, TSet<UObject*>& InTemporarilyReferencedObjects) 
+				: FObjectWriter(InBytes), TemporarilyReferencedObjects(InTemporarilyReferencedObjects) {}
+			virtual FArchive& operator<<(class UObject*& Res) override
+			{
+				FObjectWriter::operator<<(Res);
+				TemporarilyReferencedObjects.Add(Res);
+				return *this;
+			}
+		};
+
+		FRawStructWriter MemoryWriter(RowsSerializedWithTags, TemporarilyReferencedObjects);
 		SaveStructData(MemoryWriter);
 	}
 	EmptyTable();
@@ -260,9 +279,24 @@ void UDataTable::RestoreAfterStructChange()
 {
 	EmptyTable();
 	{
-		FMemoryReader MemoryReader(RowsSerializedWithTags);
+		class FRawStructReader : public FObjectReader
+		{
+		public:
+			FRawStructReader(TArray<uint8>& InBytes) : FObjectReader(InBytes) {}
+			virtual FArchive& operator<<(class UObject*& Res) override
+			{
+				UObject* Object = NULL;
+				FObjectReader::operator<<(Object);
+				FWeakObjectPtr WeakObjectPtr = Object;
+				Res = WeakObjectPtr.Get();
+				return *this;
+			}
+		};
+
+		FRawStructReader MemoryReader(RowsSerializedWithTags);
 		LoadStructData(MemoryReader);
 	}
+	TemporarilyReferencedObjects.Empty();
 	RowsSerializedWithTags.Empty();
 }
 
