Increase log output to warning

diff --git a/Engine/Source/Runtime/Engine/Private/PackageMapClient.cpp b/Engine/Source/Runtime/Engine/Private/PackageMapClient.cpp
--- a/Engine/Source/Runtime/Engine/Private/PackageMapClient.cpp
+++ b/Engine/Source/Runtime/Engine/Private/PackageMapClient.cpp
@@ -102,7 +102,7 @@ bool UPackageMapClient::SerializeObject( FArchive& Ar, UClass* Class, UObject*&
 		// NULL if we haven't finished loading the objects level yet
 		if ( !ObjectLevelHasFinishedLoading( Object ) )
 		{
-			UE_LOG(LogNetPackageMap, Log, TEXT("Using None instead of replicated reference to %s because the level it's in has not been made visible"), *Object->GetFullName());
+			UE_LOG(LogNetPackageMap, Warning, TEXT("Using None instead of replicated reference to %s because the level it's in has not been made visible"), *Object->GetFullName());
 			Object = NULL;
 		}
 
