Set PIE flags when loading a new world due to a ClientTravel to another PIE instance

diff --git a/Engine/Source/Runtime/Engine/Private/UnrealEngine.cpp b/Engine/Source/Runtime/Engine/Private/UnrealEngine.cpp
--- a/Engine/Source/Runtime/Engine/Private/UnrealEngine.cpp
+++ b/Engine/Source/Runtime/Engine/Private/UnrealEngine.cpp
@@ -8715,7 +8715,14 @@ bool UEngine::LoadMap( FWorldContext& WorldContext, FURL URL, class UPendingNetG
 				const FName PIEPackageFName = FName(*PIEPackageName);
 				UWorld::WorldTypePreLoadMap.FindOrAdd( PIEPackageFName ) = WorldContext.WorldType;
 
-				WorldPackage = LoadPackage(CreatePackage(NULL, *PIEPackageName), *SourceWorldPackage, LOAD_None);
+				uint32 LoadFlags = LOAD_None;
+				auto NewPackage = CreatePackage(NULL, *PIEPackageName);
+				if (NewPackage != nullptr && WorldContext.WorldType == EWorldType::PIE)
+				{
+					NewPackage->PackageFlags |= PKG_PlayInEditor;
+					LoadFlags |= LOAD_PackageForPIE;
+				}
+				WorldPackage = LoadPackage(NewPackage, *SourceWorldPackage, LoadFlags);
 
 				// Clean up the world type list now that PostLoad has occurred
 				UWorld::WorldTypePreLoadMap.Remove( PIEPackageFName );
