[347201] Added error tolerance input to "Equal" function for vector/rotator types.

- Updated tooltip text.
- Remapped recently introduced "NearlyEqual" function for vector/rotator back to "Equal"
- Added a K2Node::PostLoad() override to clear the default error tolerance input value on already-placed "Equal" nodes to preserve existing behavior in current Blueprint assets.

diff --git a/Engine/Config/BaseEngine.ini b/Engine/Config/BaseEngine.ini
--- a/Engine/Config/BaseEngine.ini
+++ b/Engine/Config/BaseEngine.ini
@@ -1144,6 +1144,9 @@ bUseStreamingPause=false
 +K2FieldRedirects=(OldFieldName="KismetMathLibrary.FRandFromStream",NewFieldName="RandomFloatFromStream")
 +K2FieldRedirects=(OldFieldName="KismetMathLibrary.FRandRangeFromStream",NewFieldName="RandomFloatInRangeFromStream")
 
++K2FieldRedirects=(OldFieldName="KismetMathLibrary.NearlyEqual_VectorVector",NewFieldName="EqualEqual_VectorVector")
++K2FieldRedirects=(OldFieldName="KismetMathLibrary.NearlyEqual_RotatorRotator",NewFieldName="EqualEqual_RotatorRotator")
+
 +K2FieldRedirects=(OldFieldName="PrimitiveComponent.SetRBLinearVelocity",NewFieldName="SetPhysicsLinearVelocity")
 +K2FieldRedirects=(OldFieldName="PrimitiveComponent.GetRBLinearVelocity",NewFieldName="GetPhysicsLinearVelocity")
 +K2FieldRedirects=(OldFieldName="PrimitiveComponent.SetAllRBLinearVelocity",NewFieldName="SetAllPhysicsLinearVelocity")
diff --git a/Engine/Source/Editor/BlueprintGraph/Classes/K2Node.h b/Engine/Source/Editor/BlueprintGraph/Classes/K2Node.h
--- a/Engine/Source/Editor/BlueprintGraph/Classes/K2Node.h
+++ b/Engine/Source/Editor/BlueprintGraph/Classes/K2Node.h
@@ -131,6 +131,10 @@ class UK2Node : public UEdGraphNode
 {
 	GENERATED_UCLASS_BODY()
 
+	// UObject interface
+	BLUEPRINTGRAPH_API virtual void PostLoad() override;
+	// End of UObject interface
+
 	// UEdGraphNode interface
 	BLUEPRINTGRAPH_API virtual void ReconstructNode() override;
 	BLUEPRINTGRAPH_API virtual FLinearColor GetNodeTitleColor() const override;
@@ -411,6 +415,8 @@ protected:
 			ReferencedObject->GetLinker()->Preload(ReferencedObject);
 		}
 	}
+
+	void FixupPinDefaultValues();
 };
 
 
diff --git a/Engine/Source/Editor/BlueprintGraph/Private/K2Node.cpp b/Engine/Source/Editor/BlueprintGraph/Private/K2Node.cpp
--- a/Engine/Source/Editor/BlueprintGraph/Private/K2Node.cpp
+++ b/Engine/Source/Editor/BlueprintGraph/Private/K2Node.cpp
@@ -21,6 +21,52 @@ UK2Node::UK2Node(const FObjectInitializer& ObjectInitializer)
 {
 }
 
+void UK2Node::PostLoad()
+{
+	Super::PostLoad();
+
+	// fix up pin default values
+	FixupPinDefaultValues();
+}
+
+void UK2Node::FixupPinDefaultValues()
+{
+	const int32 LinkerUE4Version = GetLinkerUE4Version();
+	const UEdGraphSchema_K2* K2Schema = GetDefault<UEdGraphSchema_K2>();
+
+	// Swap "new" default error tolerance value with zero on vector/rotator equality nodes, in order to preserve current behavior in existing blueprints.
+	if(LinkerUE4Version < VER_UE4_BP_MATH_VECTOR_EQUALITY_USES_EPSILON)
+	{
+		static const FString VectorsEqualFunctionEpsilonPinName = TEXT("KismetMathLibrary.EqualEqual_VectorVector.ErrorTolerance");
+		static const FString VectorsNotEqualFunctionEpsilonPinName = TEXT("KismetMathLibrary.NotEqual_VectorVector.ErrorTolerance");
+		static const FString RotatorsEqualFunctionEpsilonPinName = TEXT("KismetMathLibrary.EqualEqual_RotatorRotator.ErrorTolerance");
+		static const FString RotatorsNotEqualFunctionEpsilonPinName = TEXT("KismetMathLibrary.NotEqual_RotatorRotator.ErrorTolerance");
+
+		bool bFoundPin = false;
+		for(int32 i = 0; i < Pins.Num() && !bFoundPin; ++i)
+		{
+			UEdGraphPin* Pin = Pins[i];
+			check(Pin);
+
+			TArray<FString> RedirectPinNames;
+			GetRedirectPinNames(*Pin, RedirectPinNames);
+
+			for(auto PinNameIter = RedirectPinNames.CreateConstIterator(); PinNameIter && !bFoundPin; ++PinNameIter)
+			{
+				if((Pin->DefaultValue == Pin->AutogeneratedDefaultValue)
+					&& (*PinNameIter == VectorsEqualFunctionEpsilonPinName
+						|| *PinNameIter == VectorsNotEqualFunctionEpsilonPinName
+						|| *PinNameIter == RotatorsEqualFunctionEpsilonPinName
+						|| *PinNameIter == RotatorsNotEqualFunctionEpsilonPinName))
+				{
+					bFoundPin = true;
+					K2Schema->TrySetDefaultValue(*Pin, TEXT("0.0"));
+				}
+			}
+		}
+	}
+}
+
 FText UK2Node::GetToolTipHeading() const
 {
 	FText Heading = FText::GetEmpty();
diff --git a/Engine/Source/Runtime/Core/Public/UObject/ObjectVersion.h b/Engine/Source/Runtime/Core/Public/UObject/ObjectVersion.h
--- a/Engine/Source/Runtime/Core/Public/UObject/ObjectVersion.h
+++ b/Engine/Source/Runtime/Core/Public/UObject/ObjectVersion.h
@@ -626,6 +626,8 @@ enum EUnrealEngineObjectUE4Version
 	VER_UE4_CLEAN_DESTRUCTIBLE_SETTINGS,
 	// CharacterMovementComponent refactor of AdjustUpperHemisphereImpact and deprecation of some associated vars.
 	VER_UE4_CHARACTER_MOVEMENT_UPPER_IMPACT_BEHAVIOR,
+	// Changed Blueprint math equality functions for vectors and rotators to operate as a "nearly" equals rather than "exact"
+	VER_UE4_BP_MATH_VECTOR_EQUALITY_USES_EPSILON,
 
 	// -----<new versions can be added before this line>-------------------------------------------------
 	// - this needs to be the last line (see note below)
diff --git a/Engine/Source/Runtime/Engine/Classes/Kismet/KismetMathLibrary.h b/Engine/Source/Runtime/Engine/Classes/Kismet/KismetMathLibrary.h
--- a/Engine/Source/Runtime/Engine/Classes/Kismet/KismetMathLibrary.h
+++ b/Engine/Source/Runtime/Engine/Classes/Kismet/KismetMathLibrary.h
@@ -521,17 +521,13 @@ class UKismetMathLibrary : public UBlueprintFunctionLibrary
 	UFUNCTION(BlueprintPure, meta=(FriendlyName = "RotateVectorAroundAxis"), Category="Math|Vector")
 	static FVector RotateAngleAxis(FVector InVect, float AngleDeg, FVector Axis);
 
-	/* Returns true if two vectors are exactly equal (A == B) */
+	/* Returns true if vector A is equal to vector B (A == B) within a specified error tolerance */
 	UFUNCTION(BlueprintPure, meta=(FriendlyName = "Equal (vector)", CompactNodeTitle = "==", Keywords = "== equal"), Category="Math|Vector")
-	static bool EqualEqual_VectorVector(FVector A, FVector B);
+	static bool EqualEqual_VectorVector(FVector A, FVector B, float ErrorTolerance = 1.e-4f);
 
-	/* Returns true if two vectors are nearly equal (within a specified error tolerance) */
-	UFUNCTION(BlueprintPure, meta=(FriendlyName = "NearlyEqual (vector)", Keywords = "== equal"), Category="Math|Vector")
-	static bool NearlyEqual_VectorVector(FVector A, FVector B, float ErrorTolerance = 1.e-6f);
-
-	/* Returns true if two vectors are not equal (A != B) */
-	UFUNCTION(BlueprintPure, meta=(FriendlyName = "NotEqual (vector)", CompactNodeTitle = "!=", Keywords = "!= not equal"), Category="Math|Vector")
-	static bool NotEqual_VectorVector(FVector A, FVector B);
+	/* Returns true if vector A is not equal to vector B (A != B) within a specified error tolerance */
+	UFUNCTION(BlueprintPure, meta=(FriendlyName = "Not Equal (vector)", CompactNodeTitle = "!=", Keywords = "!= not equal"), Category="Math|Vector")
+	static bool NotEqual_VectorVector(FVector A, FVector B, float ErrorTolerance = 1.e-4f);
 
 	/* Returns the dot product of two vectors */
 	UFUNCTION(BlueprintPure, meta=(FriendlyName = "Dot Product", CompactNodeTitle = "."), Category="Math|Vector" )
@@ -617,17 +613,13 @@ class UKismetMathLibrary : public UBlueprintFunctionLibrary
 	// Rotator functions.
 	//
 
-	//Returns true if rotator A is exactly equal to rotator B (A == B)
+	//Returns true if rotator A is equal to rotator B (A == B) within a specified error tolerance
 	UFUNCTION(BlueprintPure, meta=(FriendlyName = "Equal (Rotator)", CompactNodeTitle = "==", Keywords = "== equal"), Category="Math|Rotator")
-	static bool EqualEqual_RotatorRotator(FRotator A, FRotator B);
-
-	//Returns true if rotator A is nearly equal to rotator B (within a specified error tolerance)
-	UFUNCTION(BlueprintPure, meta=(FriendlyName = "NearlyEqual (Rotator)", Keywords = "== equal"), Category="Math|Rotator")
-	static bool NearlyEqual_RotatorRotator(FRotator A, FRotator B, float ErrorTolerance = 1.e-6f);
+	static bool EqualEqual_RotatorRotator(FRotator A, FRotator B, float ErrorTolerance = 1.e-4f);
 
-	//Returns true if rotator A does not equal rotator B (A != B)
-	UFUNCTION(BlueprintPure, meta=(FriendlyName = "NotEqual (Rotator)", CompactNodeTitle = "!=", Keywords = "!= not equal"), Category="Math|Rotator")
-	static bool NotEqual_RotatorRotator(FRotator A, FRotator B);
+	//Returns true if rotator A is not equal to rotator B (A != B) within a specified error tolerance
+	UFUNCTION(BlueprintPure, meta=(FriendlyName = "Not Equal (Rotator)", CompactNodeTitle = "!=", Keywords = "!= not equal"), Category="Math|Rotator")
+	static bool NotEqual_RotatorRotator(FRotator A, FRotator B, float ErrorTolerance = 1.e-4f);
 
 	//Returns rotator representing rotator A scaled by B 
 	UFUNCTION(BlueprintPure, meta=(FriendlyName = "ScaleRotator", CompactNodeTitle = "*", Keywords = "* multiply rotate rotation"), Category="Math|Rotator")
diff --git a/Engine/Source/Runtime/Engine/Private/KismetMathLibrary.cpp b/Engine/Source/Runtime/Engine/Private/KismetMathLibrary.cpp
--- a/Engine/Source/Runtime/Engine/Private/KismetMathLibrary.cpp
+++ b/Engine/Source/Runtime/Engine/Private/KismetMathLibrary.cpp
@@ -683,19 +683,14 @@ FVector  UKismetMathLibrary::RotateAngleAxis(FVector InVect, float AngleDeg, FVe
 	return InVect.RotateAngleAxis(AngleDeg, Axis.SafeNormal());
 }
 
-bool UKismetMathLibrary::EqualEqual_VectorVector(FVector A, FVector B)
-{
-	return (A.X == B.X) && (A.Y == B.Y) && (A.Z == B.Z);
-}	
-
-bool UKismetMathLibrary::NearlyEqual_VectorVector(FVector A, FVector B, float ErrorTolerance)
+bool UKismetMathLibrary::EqualEqual_VectorVector(FVector A, FVector B, float ErrorTolerance)
 {
 	return A.Equals(B, ErrorTolerance);
-}
+}	
 
-bool UKismetMathLibrary::NotEqual_VectorVector(FVector A, FVector B)
+bool UKismetMathLibrary::NotEqual_VectorVector(FVector A, FVector B, float ErrorTolerance)
 {
-	return (A.X != B.X) || (A.Y != B.Y) || (A.Z != B.Z);
+	return !A.Equals(B, ErrorTolerance);
 }	
 
 float UKismetMathLibrary::Dot_VectorVector(FVector A, FVector B)
@@ -838,19 +833,14 @@ FVector UKismetMathLibrary::GetDirectionVector(FVector From, FVector To)
 }
 
 
-bool UKismetMathLibrary::EqualEqual_RotatorRotator(FRotator A, FRotator B)
-{
-	return (A.Pitch == B.Pitch) && (A.Yaw == B.Yaw) && (A.Roll == B.Roll);
-}	
-
-bool UKismetMathLibrary::NearlyEqual_RotatorRotator(FRotator A, FRotator B, float ErrorTolerance)
+bool UKismetMathLibrary::EqualEqual_RotatorRotator(FRotator A, FRotator B, float ErrorTolerance)
 {
 	return A.Equals(B, ErrorTolerance);
-}
+}	
 
-bool UKismetMathLibrary::NotEqual_RotatorRotator(FRotator A, FRotator B)
+bool UKismetMathLibrary::NotEqual_RotatorRotator(FRotator A, FRotator B, float ErrorTolerance)
 {
-	return (A.Pitch != B.Pitch) || (A.Yaw != B.Yaw) || (A.Roll != B.Roll);
+	return !A.Equals(B, ErrorTolerance);
 }	
 
 FRotator UKismetMathLibrary::Multiply_RotatorFloat(FRotator A, float B)
