ChildActorComponent spam fix - attach happening in OnCreate causes issue since the mobility is not its own yet.(temporary mobility) I moved attach to happen in register, and it clears the warning.

diff --git a/Engine/Source/Runtime/Engine/Private/ChildActorComponent.cpp b/Engine/Source/Runtime/Engine/Private/ChildActorComponent.cpp
--- a/Engine/Source/Runtime/Engine/Private/ChildActorComponent.cpp
+++ b/Engine/Source/Runtime/Engine/Private/ChildActorComponent.cpp
@@ -17,6 +17,11 @@ void UChildActorComponent::OnRegister()
 	if (ChildActor)
 	{
 		ChildActorName = ChildActor->GetFName();
+		// attach new actor to this component
+		// we can't attach in CreateChildActor since it has intermediate Mobility set up
+		// causing spam with inconsistent mobility set up
+		// so moving Attach to happen in Register
+		ChildActor->AttachRootComponentTo(this, NAME_None, EAttachLocation::SnapToTarget);
 	}
 }
 
@@ -179,9 +184,6 @@ void UChildActorComponent::CreateChildActor()
 
 					// Parts that we deferred from SpawnActor
 					ChildActor->FinishSpawning(ComponentToWorld);
-
-					// attach new actor to this component
-					ChildActor->AttachRootComponentTo(this, NAME_None, EAttachLocation::SnapToTarget);
 				}
 			}
 		}
