Fixed physics blending issue
- when root body isn't the root mesh bone
- when root bone doesn't have identity transform
- cleaned up old code of optimization and use just local variable now

diff --git a/Engine/Source/Runtime/Engine/Classes/Components/SkeletalMeshComponent.h b/Engine/Source/Runtime/Engine/Classes/Components/SkeletalMeshComponent.h
--- a/Engine/Source/Runtime/Engine/Classes/Components/SkeletalMeshComponent.h
+++ b/Engine/Source/Runtime/Engine/Classes/Components/SkeletalMeshComponent.h
@@ -642,7 +642,15 @@ public:
 	 *	Index of the 'Root Body', or top body in the asset hierarchy. 
 	 *	Filled in by InitInstance, so we don't need to save it.
 	 */
-	int32 RootBodyIndex;
+	/** To save root body index/bone index consistently **/
+	struct 
+	{
+		int32 BodyIndex;
+		int32 BoneIndex;
+	} RootBodyData;
+
+	/** Set Root Body Index */
+	void SetRootBodyIndex(int32 InBodyIndex);
 
 	/** Array of FBodyInstance objects, storing per-instance state about about each body. */
 	TArray<struct FBodyInstance*> Bodies;
diff --git a/Engine/Source/Runtime/Engine/Private/PhysicsEngine/PhysAnim.cpp b/Engine/Source/Runtime/Engine/Private/PhysicsEngine/PhysAnim.cpp
--- a/Engine/Source/Runtime/Engine/Private/PhysicsEngine/PhysAnim.cpp
+++ b/Engine/Source/Runtime/Engine/Private/PhysicsEngine/PhysAnim.cpp
@@ -19,19 +19,16 @@ static const FColor AnimSkelDrawColor(255, 64, 64);
 // Temporary workspace for caching world-space matrices.
 struct FAssetWorldBoneTM
 {
-	FTransform	TM; // Should never contain scaling.
-	int32			UpdateNum; // If this equals PhysAssetUpdateNum, then the matrix is up to date.
+	FTransform	TM;			// Should never contain scaling.
+	bool bUpToDate;			// If this equals PhysAssetUpdateNum, then the matrix is up to date.
 };
 
-// @todo anim: this has to change to make multi-threading working. 
-static int32 PhysAssetUpdateNum = 0;
-static TArray<FAssetWorldBoneTM> WorldBoneTMs;
 
 // Use current pose to calculate world-space position of this bone without physics now.
-static void UpdateWorldBoneTM(int32 BoneIndex, USkeletalMeshComponent* SkelComp, const FVector& Scale3D)
+void UpdateWorldBoneTM(TArray<FAssetWorldBoneTM> & WorldBoneTMs, int32 BoneIndex, USkeletalMeshComponent* SkelComp, const FTransform &LocalToWorldTM, const FVector& Scale3D)
 {
 	// If its already up to date - do nothing
-	if(	WorldBoneTMs[BoneIndex].UpdateNum == PhysAssetUpdateNum )
+	if(	WorldBoneTMs[BoneIndex].bUpToDate )
 	{
 		return;
 	}
@@ -40,15 +37,13 @@ static void UpdateWorldBoneTM(int32 BoneIndex, USkeletalMeshComponent* SkelComp,
 	if(BoneIndex == 0)
 	{
 		// If this is the root bone, we use the mesh component LocalToWorld as the parent transform.
-		FTransform LocalToWorldTM = SkelComp->ComponentToWorld;
-		LocalToWorldTM.RemoveScaling();
 		ParentTM = LocalToWorldTM;
 	}
 	else
 	{
 		// If not root, use our cached world-space bone transforms.
 		int32 ParentIndex = SkelComp->SkeletalMesh->RefSkeleton.GetParentIndex(BoneIndex);
-		UpdateWorldBoneTM(ParentIndex, SkelComp, Scale3D);
+		UpdateWorldBoneTM(WorldBoneTMs, ParentIndex, SkelComp, LocalToWorldTM, Scale3D);
 		ParentTM = WorldBoneTMs[ParentIndex].TM;
 	}
 
@@ -56,7 +51,7 @@ static void UpdateWorldBoneTM(int32 BoneIndex, USkeletalMeshComponent* SkelComp,
 	RelTM.ScaleTranslation( Scale3D );
 
 	WorldBoneTMs[BoneIndex].TM = RelTM * ParentTM;
-	WorldBoneTMs[BoneIndex].UpdateNum = PhysAssetUpdateNum;
+	WorldBoneTMs[BoneIndex].bUpToDate = true;
 }
 
 
@@ -64,16 +59,15 @@ void USkeletalMeshComponent::BlendPhysicsBones( TArray<FBoneIndexType>& InRequir
 {
 	// Get drawscale from Owner (if there is one)
 	FVector TotalScale3D = ComponentToWorld.GetScale3D();
-
 	FVector RecipScale3D = TotalScale3D.Reciprocal();
 
 	UPhysicsAsset * const PhysicsAsset = GetPhysicsAsset();
 	check( PhysicsAsset );
 
 	// Make sure scratch space is big enough.
+	TArray<FAssetWorldBoneTM> WorldBoneTMs;
 	WorldBoneTMs.Reset();
-	WorldBoneTMs.AddUninitialized(SpaceBases.Num());
-	PhysAssetUpdateNum++;
+	WorldBoneTMs.AddZeroed(SpaceBases.Num());
 
 	FTransform LocalToWorldTM = ComponentToWorld;
 	LocalToWorldTM.RemoveScaling();
@@ -115,7 +109,7 @@ void USkeletalMeshComponent::BlendPhysicsBones( TArray<FBoneIndexType>& InRequir
 
 				// Store this world-space transform in cache.
 				WorldBoneTMs[BoneIndex].TM = PhysTM;
-				WorldBoneTMs[BoneIndex].UpdateNum = PhysAssetUpdateNum;
+				WorldBoneTMs[BoneIndex].bUpToDate = true;
 
 				float UsePhysWeight = (bBlendPhysics)? 1.f : BodyInstance->PhysicsBlendWeight;
 
@@ -133,7 +127,7 @@ void USkeletalMeshComponent::BlendPhysicsBones( TArray<FBoneIndexType>& InRequir
 					{
 						// If not root, get parent TM from cache (making sure its up-to-date).
 						int32 ParentIndex = SkeletalMesh->RefSkeleton.GetParentIndex(BoneIndex);
-						UpdateWorldBoneTM(ParentIndex, this, TotalScale3D);
+						UpdateWorldBoneTM(WorldBoneTMs, ParentIndex, this, LocalToWorldTM, TotalScale3D);
 						ParentWorldTM = WorldBoneTMs[ParentIndex].TM;
 					}
 
@@ -153,19 +147,7 @@ void USkeletalMeshComponent::BlendPhysicsBones( TArray<FBoneIndexType>& InRequir
 						bUpdatePhysics = true;
 					}
 				}
-				else
-				{
-					WorldBoneTMs[BoneIndex].UpdateNum = 0;
-				}
 			}
-			else
-			{
-				WorldBoneTMs[BoneIndex].UpdateNum = 0;
-			}
-		}
-		else
-		{
-			WorldBoneTMs[BoneIndex].UpdateNum = 0;
 		}
 
 		// Update SpaceBases entry for this bone now
diff --git a/Engine/Source/Runtime/Engine/Private/SkeletalMeshComponent.cpp b/Engine/Source/Runtime/Engine/Private/SkeletalMeshComponent.cpp
--- a/Engine/Source/Runtime/Engine/Private/SkeletalMeshComponent.cpp
+++ b/Engine/Source/Runtime/Engine/Private/SkeletalMeshComponent.cpp
@@ -1769,3 +1769,19 @@ bool USkeletalMeshComponent::IsPlayingRootMotion()
 {
 	return (AnimScriptInstance ? (AnimScriptInstance->GetRootMotionMontageInstance() != NULL) : false);
 }
+
+void USkeletalMeshComponent::SetRootBodyIndex(int32 InBodyIndex)
+{
+	RootBodyData.BodyIndex = InBodyIndex;
+
+	if(Bodies.IsValidIndex(RootBodyData.BodyIndex) && 
+		Bodies[RootBodyData.BodyIndex]->BodySetup.IsValid() && Bodies[RootBodyData.BodyIndex]->BodySetup.Get()->BoneName != NAME_None)
+	{
+		RootBodyData.BoneIndex = GetBoneIndex(Bodies[RootBodyData.BodyIndex]->BodySetup->BoneName);
+	}
+	else
+	{
+		// error
+		RootBodyData.BoneIndex = INDEX_NONE;
+	}
+}
\ No newline at end of file
diff --git a/Engine/Source/Runtime/Engine/Private/SkeletalMeshComponentPhysics.cpp b/Engine/Source/Runtime/Engine/Private/SkeletalMeshComponentPhysics.cpp
--- a/Engine/Source/Runtime/Engine/Private/SkeletalMeshComponentPhysics.cpp
+++ b/Engine/Source/Runtime/Engine/Private/SkeletalMeshComponentPhysics.cpp
@@ -626,10 +626,10 @@ void USkeletalMeshComponent::SetAllPhysicsLinearVelocity(FVector NewVel, bool bA
 
 void USkeletalMeshComponent::SetAllPhysicsAngularVelocity(FVector const& NewAngVel, bool bAddToCurrent)
 {
-	if(RootBodyIndex < Bodies.Num())
+	if(RootBodyData.BodyIndex < Bodies.Num())
 	{
 		// Find the root actor. We use its location as the center of the rotation.
-		FBodyInstance* RootBodyInst = Bodies[ RootBodyIndex ];
+		FBodyInstance* RootBodyInst = Bodies[ RootBodyData.BodyIndex ];
 		check(RootBodyInst);
 		FTransform RootTM = RootBodyInst->GetUnrealWorldTransform();
 
@@ -648,10 +648,10 @@ void USkeletalMeshComponent::SetAllPhysicsAngularVelocity(FVector const& NewAngV
 
 void USkeletalMeshComponent::SetAllPhysicsPosition(FVector NewPos)
 {
-	if(RootBodyIndex < Bodies.Num())
+	if(RootBodyData.BodyIndex < Bodies.Num())
 	{
 		// calculate the deltas to get the root body to NewPos
-		FBodyInstance* RootBI = Bodies[RootBodyIndex];
+		FBodyInstance* RootBI = Bodies[RootBodyData.BodyIndex];
 		check(RootBI);
 		if(RootBI->IsValidBodyInstance())
 		{
@@ -669,7 +669,7 @@ void USkeletalMeshComponent::SetAllPhysicsPosition(FVector NewPos)
 			// apply the delta to all the other bodies
 			for (int32 i = 0; i < Bodies.Num(); i++)
 			{
-				if (i != RootBodyIndex)
+				if (i != RootBodyData.BodyIndex)
 				{
 					FBodyInstance* BI = Bodies[i];
 					check(BI);
@@ -688,10 +688,10 @@ void USkeletalMeshComponent::SetAllPhysicsPosition(FVector NewPos)
 
 void USkeletalMeshComponent::SetAllPhysicsRotation(FRotator NewRot)
 {
-	if(RootBodyIndex < Bodies.Num())
+	if(RootBodyData.BodyIndex < Bodies.Num())
 	{
 		// calculate the deltas to get the root body to NewRot
-		FBodyInstance* RootBI = Bodies[RootBodyIndex];
+		FBodyInstance* RootBI = Bodies[RootBodyData.BodyIndex];
 		check(RootBI);
 		if(RootBI->IsValidBodyInstance())
 		{
@@ -705,7 +705,7 @@ void USkeletalMeshComponent::SetAllPhysicsRotation(FRotator NewRot)
 			// apply the delta to all the other bodies
 			for (int32 i = 0; i < Bodies.Num(); i++)
 			{
-				if (i != RootBodyIndex)
+				if (i != RootBodyData.BodyIndex)
 				{
 					FBodyInstance* BI = Bodies[i];
 					check(BI);
@@ -724,10 +724,10 @@ void USkeletalMeshComponent::SetAllPhysicsRotation(FRotator NewRot)
 
 void USkeletalMeshComponent::ApplyDeltaToAllPhysicsTransforms(const FVector& DeltaLocation, const FQuat& DeltaRotation)
 {
-	if(RootBodyIndex < Bodies.Num())
+	if(RootBodyData.BodyIndex < Bodies.Num())
 	{
 		// calculate the deltas to get the root body to NewRot
-		FBodyInstance* RootBI = Bodies[RootBodyIndex];
+		FBodyInstance* RootBI = Bodies[RootBodyData.BodyIndex];
 		check(RootBI);
 		if(RootBI->IsValidBodyInstance())
 		{
@@ -740,7 +740,7 @@ void USkeletalMeshComponent::ApplyDeltaToAllPhysicsTransforms(const FVector& Del
 			// apply the delta to all the other bodies
 			for (int32 i = 0; i < Bodies.Num(); i++)
 			{
-				if (i != RootBodyIndex)
+				if (i != RootBodyData.BodyIndex)
 				{
 					FBodyInstance* BI = Bodies[i];
 					check(BI);
@@ -790,7 +790,7 @@ void USkeletalMeshComponent::InitArticulated(FPhysScene* PhysScene)
 	float Scale = Scale3D.X;
 
 	// Find root physics body
-	RootBodyIndex = INDEX_NONE;
+	int32 RootBodyIndex = INDEX_NONE;
 	for(int32 i=0; i<SkeletalMesh->RefSkeleton.GetNum(); i++)
 	{
 		int32 BodyInstIndex = PhysicsAsset->FindBodyIndex( SkeletalMesh->RefSkeleton.GetBoneName(i) );
@@ -862,6 +862,9 @@ void USkeletalMeshComponent::InitArticulated(FPhysScene* PhysScene)
 		}
 	}
 
+	// now update root body index because body has BodySetup now
+	SetRootBodyIndex(RootBodyIndex);
+
 #if WITH_PHYSX
 	// add Aggregate into the scene
 	if(Aggregate && Aggregate->getNbActors() > 0 && PhysScene)
@@ -898,7 +901,6 @@ void USkeletalMeshComponent::InitArticulated(FPhysScene* PhysScene)
 
 	// Update Flag
 	ResetAllBodiesSimulatePhysics();
-
 #if WITH_APEX_CLOTHING
 	PrevRootBoneMatrix = GetBoneMatrix(0); // save the root bone transform
 	// pre-compute cloth teleport thresholds for performance
@@ -1499,9 +1501,9 @@ FBodyInstance* USkeletalMeshComponent::GetBodyInstance(FName BoneName, bool) con
 		// A name of NAME_None indicates 'root body'
 		if(BoneName == NAME_None)
 		{
-			if(Bodies.IsValidIndex(RootBodyIndex))
+			if(Bodies.IsValidIndex(RootBodyData.BodyIndex))
 			{
-				BodyInst = Bodies[RootBodyIndex];
+				BodyInst = Bodies[RootBodyData.BodyIndex];
 			}
 		}
 		// otherwise, look for the body
@@ -1912,7 +1914,7 @@ bool USkeletalMeshComponent::ComponentOverlapMulti(TArray<struct FOverlapResult>
 {
 	OutOverlaps.Reset();
 
-	if (!Bodies.IsValidIndex(RootBodyIndex))
+	if (!Bodies.IsValidIndex(RootBodyData.BodyIndex))
 	{
 		return false;
 	}
@@ -4950,7 +4952,12 @@ FTransform USkeletalMeshComponent::GetComponentTransformFromBodyInstance(FBodyIn
 {
 	// undo root transform so that it only moves according to what actor itself suppose to move
 	FTransform BodyTransform = UseBI->GetUnrealWorldTransform();
-	FTransform RootTransform(LocalAtoms[0].GetRotation());
-	return RootTransform.GetRelativeTransformReverse(BodyTransform);
+	if (RootBodyData.BoneIndex != INDEX_NONE)
+	{
+		FTransform RootTransform(SpaceBases[RootBodyData.BoneIndex]);
+		return RootTransform.GetRelativeTransformReverse(BodyTransform);
+	}
+
+	return BodyTransform;
 }
 #undef LOCTEXT_NAMESPACE
