Added json serialization for Visual Log's events. #ue4

diff --git a/Engine/Source/Runtime/Engine/Private/VisualLog.cpp b/Engine/Source/Runtime/Engine/Private/VisualLog.cpp
--- a/Engine/Source/Runtime/Engine/Private/VisualLog.cpp
+++ b/Engine/Source/Runtime/Engine/Private/VisualLog.cpp
@@ -106,6 +106,42 @@ TSharedPtr<FJsonValue> GenerateJson(FVisualLogEntry* InEntry)
 
 	JsonEntryObject->SetArrayField(VisualLogJson::TAG_HISTOGRAMSAMPLES, JsonHistogramSamples);
 
+	// generate json from events
+	{
+		const int32 EventsCount = InEntry->Events.Num();
+		TArray<TSharedPtr<FJsonValue> > JsonEventSamples;
+		JsonEventSamples.Reserve(HistogramSamplesCount);
+
+		const FVisualLogEntry::FHistogramSample* Sample = InEntry->HistogramSamples.GetData();
+		for (auto& CurrentEvent : InEntry->Events)
+		{
+			TSharedPtr<FJsonObject> JsonSampleObject = MakeShareable(new FJsonObject);
+			
+			JsonSampleObject->SetStringField(VisualLogJson::TAG_CATEGORY, CurrentEvent.Name);
+			JsonSampleObject->SetNumberField(VisualLogJson::TAG_VERBOSITY, CurrentEvent.Verbosity);
+			JsonSampleObject->SetStringField(VisualLogJson::TAG_TAGNAME, CurrentEvent.TagName.ToString());
+			JsonSampleObject->SetNumberField(VisualLogJson::TAG_USERDATA, CurrentEvent.UserData);
+			JsonSampleObject->SetStringField(VisualLogJson::TAG_DESCRIPTION, CurrentEvent.UserFriendlyDesc);
+			JsonSampleObject->SetNumberField(VisualLogJson::TAG_COUNTER, CurrentEvent.Counter);
+
+			TArray<TSharedPtr<FJsonValue> > JsonEventTags;
+			JsonEventTags.Reserve(HistogramSamplesCount);
+			for (auto& CurrentTag : CurrentEvent.EventTags)
+			{
+				TSharedPtr<FJsonObject> JsonEventObject = MakeShareable(new FJsonObject);
+				JsonEventObject->SetStringField(VisualLogJson::TAG_EVENTNAME, CurrentTag.Key.ToString());
+				JsonEventObject->SetNumberField(VisualLogJson::TAG_COUNTER, CurrentTag.Value);
+
+				JsonEventTags.Add(MakeShareable(new FJsonValueObject(JsonEventObject)));
+			}
+			JsonSampleObject->SetArrayField(VisualLogJson::TAG_EVENTAGS, JsonEventTags);
+
+			JsonEventSamples.Add(MakeShareable(new FJsonValueObject(JsonSampleObject)));
+		}
+
+		JsonEntryObject->SetArrayField(VisualLogJson::TAG_EVENTSAMPLES, JsonEventSamples);
+	}
+
 	int32 DataBlockIndex = 0;
 	TArray<TSharedPtr<FJsonValue> > JasonDataBlocks;
 	JasonDataBlocks.AddZeroed(InEntry->DataBlocks.Num());
@@ -249,6 +285,42 @@ TSharedPtr<FVisualLogEntry> JsonToVisualLogEntry(TSharedPtr<FJsonValue> FromJson
 		}
 	}
 
+	TArray< TSharedPtr<FJsonValue> > JsonEventSamples = JsonEntryObject->GetArrayField(VisualLogJson::TAG_EVENTSAMPLES);
+	if (JsonEventSamples.Num() > 0)
+	{
+		VisLogEntry->Events.Reserve(JsonEventSamples.Num());
+		for (int32 SampleIndex = 0; SampleIndex < JsonEventSamples.Num(); ++SampleIndex)
+		{
+			TSharedPtr<FJsonObject> JsonSampleObject = JsonEventSamples[SampleIndex]->AsObject();
+			if (JsonSampleObject.IsValid())
+			{
+				FVisualLogEntry::FLogEvent& Event = VisLogEntry->Events[VisLogEntry->Events.Add(FVisualLogEntry::FLogEvent())];
+				Event.Name = JsonSampleObject->GetStringField(VisualLogJson::TAG_CATEGORY);
+				Event.Verbosity = TEnumAsByte<ELogVerbosity::Type>((uint8)FMath::TruncToInt(JsonSampleObject->GetNumberField(VisualLogJson::TAG_VERBOSITY)));
+				Event.TagName = FName(*(JsonSampleObject->GetStringField(VisualLogJson::TAG_TAGNAME)));
+				Event.UserData = (int64)(JsonSampleObject->GetNumberField(VisualLogJson::TAG_USERDATA));
+				Event.UserFriendlyDesc = JsonSampleObject->GetStringField(VisualLogJson::TAG_DESCRIPTION);
+				Event.Counter = (int64)(JsonSampleObject->GetNumberField(VisualLogJson::TAG_COUNTER));
+
+				TArray< TSharedPtr<FJsonValue> > JsonEventTags = JsonEntryObject->GetArrayField(VisualLogJson::TAG_EVENTAGS);
+				if (JsonEventTags.Num() > 0)
+				{
+					Event.EventTags.Reserve(JsonEventSamples.Num());
+					for (int32 TagIndex = 0; TagIndex < JsonEventTags.Num(); ++TagIndex)
+					{
+						TSharedPtr<FJsonObject> JsonTagObject = JsonEventTags[SampleIndex]->AsObject();
+						if (JsonTagObject.IsValid())
+						{
+							FString TagName = JsonTagObject->GetStringField(VisualLogJson::TAG_EVENTNAME);
+							int32 TagCounter = JsonTagObject->GetNumberField(VisualLogJson::TAG_COUNTER);
+							Event.EventTags.Add(*TagName, TagCounter);
+						}
+					}
+				}
+			}
+		}
+	}
+
 	TArray< TSharedPtr<FJsonValue> > JasonDataBlocks = JsonEntryObject->GetArrayField(VisualLogJson::TAG_DATABLOCK);
 	if (JasonDataBlocks.Num() > 0)
 	{
diff --git a/Engine/Source/Runtime/Engine/Public/VisualLog.h b/Engine/Source/Runtime/Engine/Public/VisualLog.h
--- a/Engine/Source/Runtime/Engine/Public/VisualLog.h
+++ b/Engine/Source/Runtime/Engine/Public/VisualLog.h
@@ -55,6 +55,10 @@ namespace VisualLogJson
 	static const FString TAG_TAGNAME = TEXT("DataBlockTagName");
 	static const FString TAG_USERDATA = TEXT("UserData");
 	static const FString TAG_CAPSULESTODRAW = TEXT("CapsulesToDraw");
+	static const FString TAG_COUNTER = TEXT("Counter");
+	static const FString TAG_EVENTNAME = TEXT("EventName");
+	static const FString TAG_EVENTAGS = TEXT("EventTags");
+	static const FString TAG_EVENTSAMPLES = TEXT("EventTags");
 
 	static const FString TAG_HISTOGRAMSAMPLES = TEXT("HistogramSamples");
 	static const FString TAG_HISTOGRAMSAMPLE = TEXT("Sample");
