Fix for specific physics crash that happened, at least for us, when destroying a static mesh right before/at the same time as impact with the character.

diff --git a/Engine/Source/Runtime/Engine/Private/PrimitiveComponent.cpp b/Engine/Source/Runtime/Engine/Private/PrimitiveComponent.cpp
--- a/Engine/Source/Runtime/Engine/Private/PrimitiveComponent.cpp
+++ b/Engine/Source/Runtime/Engine/Private/PrimitiveComponent.cpp
@@ -1574,7 +1574,7 @@ void UPrimitiveComponent::DispatchBlockingHit(AActor& Owner, FHitResult const& B
 	Owner.DispatchBlockingHit(this, BlockingHit.Component.Get(), true, BlockingHit);
 
 	// BlockingHit.GetActor() could be marked for deletion in DispatchBlockingHit(), which would make the weak pointer return NULL.
-	if ( BlockingHit.GetActor() != NULL )
+	if ( BlockingHit.GetActor() != nullptr && BlockingHit.Component.Get() != nullptr)
 	{
 		BlockingHit.GetActor()->DispatchBlockingHit(BlockingHit.Component.Get(), this, false, BlockingHit);
 	}
@@ -1733,7 +1733,7 @@ bool UPrimitiveComponent::OverlapComponent(const FVector& Pos, const FQuat& Rot,
 	return BodyInstance.OverlapTest(Pos, Rot, CollisionShape);
 }
 
-bool UPrimitiveComponent::ComputePenetration(FMTDResult & OutMTD, const FCollisionShape& CollisionShape, const FVector& Pos, const FQuat& Rot)
+bool UPrimitiveComponent::ComputePenetration(FMTDResult & OutMTD, const FCollisionShape & CollisionShape, const FVector & Pos, const FQuat & Rot)
 {
 #if WITH_PHYSX
 	UCollision2PGeom GeomStorage0(CollisionShape);
@@ -2209,7 +2209,7 @@ void UPrimitiveComponent::UpdateOverlaps(TArray<FOverlapInfo> const* PendingOver
 		// End all overlaps that exist, in case bGenerateOverlapEvents was true last tick (i.e. was just turned off)
 		// Iterate backwards since EndComponentOverlap will remove items from OverlappingComponents.
 		for (int32 OverlapIdx = OverlappingComponents.Num()-1; OverlapIdx >= 0; --OverlapIdx)
-		{		
+		{
 			const FOverlapInfo& OtherOverlap = OverlappingComponents[OverlapIdx];
 			if (OtherOverlap.OverlapInfo.Component.IsValid())
 			{
