Prepare Engine module for private PCH (EnginePrivate.h) cleanup.

diff --git a/Engine/Source/Editor/UnrealEd/Public/Kismet2/EnumEditorUtils.h b/Engine/Source/Editor/UnrealEd/Public/Kismet2/EnumEditorUtils.h
--- a/Engine/Source/Editor/UnrealEd/Public/Kismet2/EnumEditorUtils.h
+++ b/Engine/Source/Editor/UnrealEd/Public/Kismet2/EnumEditorUtils.h
@@ -7,6 +7,8 @@
 
 #include "ListenerManager.h"
 
+class UUserDefinedEnum;
+
 class UNREALED_API FEnumEditorUtils
 {
 	static void PrepareForChange(const UUserDefinedEnum* Enum);
diff --git a/Engine/Source/Runtime/Engine/Classes/Distributions/DistributionFloatParameterBase.h b/Engine/Source/Runtime/Engine/Classes/Distributions/DistributionFloatParameterBase.h
--- a/Engine/Source/Runtime/Engine/Classes/Distributions/DistributionFloatParameterBase.h
+++ b/Engine/Source/Runtime/Engine/Classes/Distributions/DistributionFloatParameterBase.h
@@ -2,6 +2,7 @@
 
 
 #pragma once
+#include "Distributions/DistributionFloatConstant.h"
 #include "DistributionFloatParameterBase.generated.h"
 
 UENUM()
diff --git a/Engine/Source/Runtime/Engine/Classes/Distributions/DistributionFloatParticleParameter.h b/Engine/Source/Runtime/Engine/Classes/Distributions/DistributionFloatParticleParameter.h
--- a/Engine/Source/Runtime/Engine/Classes/Distributions/DistributionFloatParticleParameter.h
+++ b/Engine/Source/Runtime/Engine/Classes/Distributions/DistributionFloatParticleParameter.h
@@ -2,6 +2,7 @@
 
 
 #pragma once
+#include "Distributions/DistributionFloatParameterBase.h"
 #include "DistributionFloatParticleParameter.generated.h"
 
 UCLASS(collapsecategories, hidecategories=Object, editinlinenew, MinimalAPI)
diff --git a/Engine/Source/Runtime/Engine/Classes/Distributions/DistributionFloatSoundParameter.h b/Engine/Source/Runtime/Engine/Classes/Distributions/DistributionFloatSoundParameter.h
--- a/Engine/Source/Runtime/Engine/Classes/Distributions/DistributionFloatSoundParameter.h
+++ b/Engine/Source/Runtime/Engine/Classes/Distributions/DistributionFloatSoundParameter.h
@@ -2,6 +2,7 @@
 
 
 #pragma once
+#include "Distributions/DistributionFloatParameterBase.h"
 #include "DistributionFloatSoundParameter.generated.h"
 
 UCLASS(deprecated)
diff --git a/Engine/Source/Runtime/Engine/Classes/Distributions/DistributionVectorParameterBase.h b/Engine/Source/Runtime/Engine/Classes/Distributions/DistributionVectorParameterBase.h
--- a/Engine/Source/Runtime/Engine/Classes/Distributions/DistributionVectorParameterBase.h
+++ b/Engine/Source/Runtime/Engine/Classes/Distributions/DistributionVectorParameterBase.h
@@ -2,6 +2,7 @@
 
 
 #pragma once
+#include "Distributions/DistributionVectorConstant.h"
 #include "DistributionVectorParameterBase.generated.h"
 
 UCLASS(abstract, collapsecategories, hidecategories=Object, editinlinenew)
diff --git a/Engine/Source/Runtime/Engine/Classes/Distributions/DistributionVectorParticleParameter.h b/Engine/Source/Runtime/Engine/Classes/Distributions/DistributionVectorParticleParameter.h
--- a/Engine/Source/Runtime/Engine/Classes/Distributions/DistributionVectorParticleParameter.h
+++ b/Engine/Source/Runtime/Engine/Classes/Distributions/DistributionVectorParticleParameter.h
@@ -2,6 +2,7 @@
 
 
 #pragma once
+#include "Distributions/DistributionVectorParameterBase.h"
 #include "DistributionVectorParticleParameter.generated.h"
 
 UCLASS(collapsecategories, hidecategories=Object, editinlinenew, MinimalAPI)
diff --git a/Engine/Source/Runtime/Engine/Classes/Engine/SimpleConstructionScript.h b/Engine/Source/Runtime/Engine/Classes/Engine/SimpleConstructionScript.h
--- a/Engine/Source/Runtime/Engine/Classes/Engine/SimpleConstructionScript.h
+++ b/Engine/Source/Runtime/Engine/Classes/Engine/SimpleConstructionScript.h
@@ -4,6 +4,8 @@
 #pragma once
 #include "SimpleConstructionScript.generated.h"
 
+class USCS_Node;
+
 UCLASS(MinimalAPI)
 class USimpleConstructionScript : public UObject
 {
diff --git a/Engine/Source/Runtime/Engine/Classes/Kismet/DataTableFunctionLibrary.h b/Engine/Source/Runtime/Engine/Classes/Kismet/DataTableFunctionLibrary.h
--- a/Engine/Source/Runtime/Engine/Classes/Kismet/DataTableFunctionLibrary.h
+++ b/Engine/Source/Runtime/Engine/Classes/Kismet/DataTableFunctionLibrary.h
@@ -5,6 +5,8 @@
 #include "Kismet/BlueprintFunctionLibrary.h"
 #include "DataTableFunctionLibrary.generated.h"
 
+class UCurveTable;
+
 /** Enum used to indicate success or failure of EvaluateCurveTableRow */
 UENUM()
 namespace EEvaluateCurveTableResult
diff --git a/Engine/Source/Runtime/Engine/Classes/Matinee/InterpTrack.h b/Engine/Source/Runtime/Engine/Classes/Matinee/InterpTrack.h
--- a/Engine/Source/Runtime/Engine/Classes/Matinee/InterpTrack.h
+++ b/Engine/Source/Runtime/Engine/Classes/Matinee/InterpTrack.h
@@ -14,7 +14,7 @@
 // Forward declarations
 //
 class FCanvas;
-
+class UInterpGroup;
 
 /** Required condition for this track to be enabled */
 UENUM()
diff --git a/Engine/Source/Runtime/Engine/Classes/Matinee/InterpTrackDirector.h b/Engine/Source/Runtime/Engine/Classes/Matinee/InterpTrackDirector.h
--- a/Engine/Source/Runtime/Engine/Classes/Matinee/InterpTrackDirector.h
+++ b/Engine/Source/Runtime/Engine/Classes/Matinee/InterpTrackDirector.h
@@ -1,6 +1,7 @@
 // Copyright 1998-2014 Epic Games, Inc. All Rights Reserved.
 
 #pragma once
+#include "Matinee/InterpTrack.h"
 #include "InterpTrackDirector.generated.h"
 
 /**
diff --git a/Engine/Source/Runtime/Engine/Classes/Matinee/InterpTrackEvent.h b/Engine/Source/Runtime/Engine/Classes/Matinee/InterpTrackEvent.h
--- a/Engine/Source/Runtime/Engine/Classes/Matinee/InterpTrackEvent.h
+++ b/Engine/Source/Runtime/Engine/Classes/Matinee/InterpTrackEvent.h
@@ -1,6 +1,7 @@
 // Copyright 1998-2014 Epic Games, Inc. All Rights Reserved.
 
 #pragma once
+#include "Matinee/InterpTrack.h"
 #include "InterpTrackEvent.generated.h"
 
 /**
diff --git a/Engine/Source/Runtime/Engine/Classes/Matinee/InterpTrackMove.h b/Engine/Source/Runtime/Engine/Classes/Matinee/InterpTrackMove.h
--- a/Engine/Source/Runtime/Engine/Classes/Matinee/InterpTrackMove.h
+++ b/Engine/Source/Runtime/Engine/Classes/Matinee/InterpTrackMove.h
@@ -4,6 +4,8 @@
 #include "Matinee/InterpTrack.h"
 #include "InterpTrackMove.generated.h"
 
+class UInterpGroup;
+
 UENUM()
 enum EInterpTrackMoveRotMode
 {
diff --git a/Engine/Source/Runtime/Engine/Classes/Matinee/InterpTrackToggle.h b/Engine/Source/Runtime/Engine/Classes/Matinee/InterpTrackToggle.h
--- a/Engine/Source/Runtime/Engine/Classes/Matinee/InterpTrackToggle.h
+++ b/Engine/Source/Runtime/Engine/Classes/Matinee/InterpTrackToggle.h
@@ -1,6 +1,7 @@
 // Copyright 1998-2014 Epic Games, Inc. All Rights Reserved.
 
 #pragma once
+#include "Matinee/InterpTrack.h"
 #include "InterpTrackToggle.generated.h"
 
 /**
diff --git a/Engine/Source/Runtime/Engine/Classes/MovieScene/RuntimeMovieScenePlayerInterface.h b/Engine/Source/Runtime/Engine/Classes/MovieScene/RuntimeMovieScenePlayerInterface.h
--- a/Engine/Source/Runtime/Engine/Classes/MovieScene/RuntimeMovieScenePlayerInterface.h
+++ b/Engine/Source/Runtime/Engine/Classes/MovieScene/RuntimeMovieScenePlayerInterface.h
@@ -1,7 +1,7 @@
 // Copyright 1998-2014 Epic Games, Inc. All Rights Reserved.
 
 #pragma once
-
+#include "Interface.h"
 #include "RuntimeMovieScenePlayerInterface.generated.h"
 
 UINTERFACE(MinimalAPI, meta=(CannotImplementInterfaceInBlueprint))
diff --git a/Engine/Source/Runtime/Engine/Private/ActorComponent.cpp b/Engine/Source/Runtime/Engine/Private/ActorComponent.cpp
--- a/Engine/Source/Runtime/Engine/Private/ActorComponent.cpp
+++ b/Engine/Source/Runtime/Engine/Private/ActorComponent.cpp
@@ -6,6 +6,7 @@
 #include "MessageLog.h"
 #include "UObjectToken.h"
 #include "MapErrors.h"
+#include "ComponentReregisterContext.h"
 
 #define LOCTEXT_NAMESPACE "ActorComponent"
 
diff --git a/Engine/Source/Runtime/Engine/Private/Atmosphere/Atmosphere.cpp b/Engine/Source/Runtime/Engine/Private/Atmosphere/Atmosphere.cpp
--- a/Engine/Source/Runtime/Engine/Private/Atmosphere/Atmosphere.cpp
+++ b/Engine/Source/Runtime/Engine/Private/Atmosphere/Atmosphere.cpp
@@ -5,6 +5,7 @@
 #include "Atmosphere.h"
 #include "ComponentInstanceDataCache.h"
 #include "Atmosphere/AtmosphericFog.h"
+#include "ComponentReregisterContext.h"
 
 #if WITH_EDITOR
 #include "ObjectEditorUtils.h"
diff --git a/Engine/Source/Runtime/Engine/Private/Camera/CameraActor.cpp b/Engine/Source/Runtime/Engine/Private/Camera/CameraActor.cpp
--- a/Engine/Source/Runtime/Engine/Private/Camera/CameraActor.cpp
+++ b/Engine/Source/Runtime/Engine/Private/Camera/CameraActor.cpp
@@ -2,6 +2,7 @@
 
 #include "EnginePrivate.h"
 #include "Camera/CameraActor.h"
+#include "Camera/CameraAnim.h"
 #include "Net/UnrealNetwork.h"
 #include "MessageLog.h"
 #include "UObjectToken.h"
diff --git a/Engine/Source/Runtime/Engine/Private/Camera/CameraAnimInst.cpp b/Engine/Source/Runtime/Engine/Private/Camera/CameraAnimInst.cpp
--- a/Engine/Source/Runtime/Engine/Private/Camera/CameraAnimInst.cpp
+++ b/Engine/Source/Runtime/Engine/Private/Camera/CameraAnimInst.cpp
@@ -4,6 +4,7 @@
 #include "EnginePrivate.h"
 #include "Camera/CameraActor.h"
 #include "Camera/CameraAnim.h"
+#include "Camera/CameraAnimInst.h"
 #include "Matinee/InterpTrackMove.h"
 #include "Matinee/InterpGroup.h"
 #include "Matinee/InterpGroupInst.h"
diff --git a/Engine/Source/Runtime/Engine/Private/Camera/CameraModifier_CameraShake.cpp b/Engine/Source/Runtime/Engine/Private/Camera/CameraModifier_CameraShake.cpp
--- a/Engine/Source/Runtime/Engine/Private/Camera/CameraModifier_CameraShake.cpp
+++ b/Engine/Source/Runtime/Engine/Private/Camera/CameraModifier_CameraShake.cpp
@@ -1,6 +1,7 @@
 // Copyright 1998-2014 Epic Games, Inc. All Rights Reserved.
 
 #include "EnginePrivate.h"
+#include "Camera/CameraAnimInst.h"
 
 //////////////////////////////////////////////////////////////////////////
 // UCameraModifier_CameraShake
diff --git a/Engine/Source/Runtime/Engine/Private/CurveTable.cpp b/Engine/Source/Runtime/Engine/Private/CurveTable.cpp
--- a/Engine/Source/Runtime/Engine/Private/CurveTable.cpp
+++ b/Engine/Source/Runtime/Engine/Private/CurveTable.cpp
@@ -1,6 +1,7 @@
 // Copyright 1998-2014 Epic Games, Inc. All Rights Reserved.
 
 #include "EnginePrivate.h"
+#include "Engine/CurveTable.h"
 #include "Json.h"
 #include "CsvParser.h"
 
diff --git a/Engine/Source/Runtime/Engine/Private/DataChannel.cpp b/Engine/Source/Runtime/Engine/Private/DataChannel.cpp
--- a/Engine/Source/Runtime/Engine/Private/DataChannel.cpp
+++ b/Engine/Source/Runtime/Engine/Private/DataChannel.cpp
@@ -11,6 +11,7 @@
 #include "Net/UnrealNetwork.h"
 #include "Engine/ActorChannel.h"
 #include "Engine/ControlChannel.h"
+#include "Engine/PackageMapClient.h"
 
 DEFINE_LOG_CATEGORY(LogNet);
 DEFINE_LOG_CATEGORY(LogNetPlayerMovement);
diff --git a/Engine/Source/Runtime/Engine/Private/DataReplication.cpp b/Engine/Source/Runtime/Engine/Private/DataReplication.cpp
--- a/Engine/Source/Runtime/Engine/Private/DataReplication.cpp
+++ b/Engine/Source/Runtime/Engine/Private/DataReplication.cpp
@@ -10,6 +10,7 @@
 #include "Net/RepLayout.h"
 #include "Net/DataReplication.h"
 #include "Engine/ActorChannel.h"
+#include "Engine/PackageMapClient.h"
 
 class FNetSerializeCB : public INetSerializeCB
 {
diff --git a/Engine/Source/Runtime/Engine/Private/DataTable.cpp b/Engine/Source/Runtime/Engine/Private/DataTable.cpp
--- a/Engine/Source/Runtime/Engine/Private/DataTable.cpp
+++ b/Engine/Source/Runtime/Engine/Private/DataTable.cpp
@@ -3,6 +3,7 @@
 #include "EnginePrivate.h"
 #include "Json.h"
 #include "CsvParser.h"
+#include "Engine/DataTable.h"
 
 DEFINE_LOG_CATEGORY(LogDataTable);
 
diff --git a/Engine/Source/Runtime/Engine/Private/DataTableFunctionLibrary.cpp b/Engine/Source/Runtime/Engine/Private/DataTableFunctionLibrary.cpp
--- a/Engine/Source/Runtime/Engine/Private/DataTableFunctionLibrary.cpp
+++ b/Engine/Source/Runtime/Engine/Private/DataTableFunctionLibrary.cpp
@@ -2,6 +2,7 @@
 
 #include "EnginePrivate.h"
 #include "../Classes/Kismet/DataTableFunctionLibrary.h"
+#include "Engine/CurveTable.h"
 
 UDataTableFunctionLibrary::UDataTableFunctionLibrary(const FObjectInitializer& ObjectInitializer)
 	: Super(ObjectInitializer)
diff --git a/Engine/Source/Runtime/Engine/Private/DialogueTypes.cpp b/Engine/Source/Runtime/Engine/Private/DialogueTypes.cpp
--- a/Engine/Source/Runtime/Engine/Private/DialogueTypes.cpp
+++ b/Engine/Source/Runtime/Engine/Private/DialogueTypes.cpp
@@ -2,7 +2,7 @@
 
 #include "EnginePrivate.h"
 #include "SoundDefinitions.h"
-
+#include "Sound/DialogueVoice.h"
 bool operator==(const FDialogueContext& LHS, const FDialogueContext& RHS)
 {
 	return	LHS.Speaker == RHS.Speaker &&
diff --git a/Engine/Source/Runtime/Engine/Private/DialogueVoice.cpp b/Engine/Source/Runtime/Engine/Private/DialogueVoice.cpp
--- a/Engine/Source/Runtime/Engine/Private/DialogueVoice.cpp
+++ b/Engine/Source/Runtime/Engine/Private/DialogueVoice.cpp
@@ -6,6 +6,7 @@
 #include "TargetPlatform.h"
 #include "AudioDerivedData.h"
 #include "SubtitleManager.h"
+#include "Sound/DialogueVoice.h"
 
 UDialogueVoice::UDialogueVoice(const FObjectInitializer& ObjectInitializer)
 	: Super(ObjectInitializer)
diff --git a/Engine/Source/Runtime/Engine/Private/DialogueWave.cpp b/Engine/Source/Runtime/Engine/Private/DialogueWave.cpp
--- a/Engine/Source/Runtime/Engine/Private/DialogueWave.cpp
+++ b/Engine/Source/Runtime/Engine/Private/DialogueWave.cpp
@@ -2,6 +2,7 @@
 
 #include "EnginePrivate.h"
 #include "ActiveSound.h"
+#include "Sound/DialogueWave.h"
 #include "Sound/DialogueSoundWaveProxy.h"
 #include "Sound/SoundWave.h"
 #include "SubtitleManager.h"
diff --git a/Engine/Source/Runtime/Engine/Private/DistributionHelpers.h b/Engine/Source/Runtime/Engine/Private/DistributionHelpers.h
--- a/Engine/Source/Runtime/Engine/Private/DistributionHelpers.h
+++ b/Engine/Source/Runtime/Engine/Private/DistributionHelpers.h
@@ -1,5 +1,8 @@
 // Copyright 1998-2014 Epic Games, Inc. All Rights Reserved.
 #pragma once
+#include "Distributions/DistributionFloatConstant.h"
+#include "Distributions/DistributionVectorConstant.h"
+#include "Distributions/DistributionVectorUniform.h"
 
 /**
  * Helper class for restoring default values on UDistribution objects in particle modules.
diff --git a/Engine/Source/Runtime/Engine/Private/Distributions.cpp b/Engine/Source/Runtime/Engine/Private/Distributions.cpp
--- a/Engine/Source/Runtime/Engine/Private/Distributions.cpp
+++ b/Engine/Source/Runtime/Engine/Private/Distributions.cpp
@@ -6,6 +6,9 @@
 
 #include "EnginePrivate.h"
 #include "DistributionHelpers.h"
+#include "Distributions/DistributionVectorParameterBase.h"
+#include "Distributions/DistributionVectorConstantCurve.h"
+#include "Distributions/DistributionVectorUniformCurve.h"
 #include "ParticleDefinitions.h"
 #include "SoundDefinitions.h"
 #include "Sound/SoundNode.h"
diff --git a/Engine/Source/Runtime/Engine/Private/Engine.cpp b/Engine/Source/Runtime/Engine/Private/Engine.cpp
--- a/Engine/Source/Runtime/Engine/Private/Engine.cpp
+++ b/Engine/Source/Runtime/Engine/Private/Engine.cpp
@@ -7,6 +7,7 @@
 #include "EnginePrivate.h"
 #include "EngineModule.h"
 #include "EditorSupportDelegates.h"
+#include "ComponentReregisterContext.h"
 
 /*-----------------------------------------------------------------------------
 	Globals.
diff --git a/Engine/Source/Runtime/Engine/Private/Font.cpp b/Engine/Source/Runtime/Engine/Private/Font.cpp
--- a/Engine/Source/Runtime/Engine/Private/Font.cpp
+++ b/Engine/Source/Runtime/Engine/Private/Font.cpp
@@ -5,7 +5,8 @@
 =============================================================================*/
 
 #include "EnginePrivate.h"
-
+#include "Engine/Font.h"
+#include "Engine/FontImportOptions.h"
 
 UFontImportOptions::UFontImportOptions(const FObjectInitializer& ObjectInitializer)
 	: Super(ObjectInitializer)
diff --git a/Engine/Source/Runtime/Engine/Private/GameFramework/SaveGame.cpp b/Engine/Source/Runtime/Engine/Private/GameFramework/SaveGame.cpp
--- a/Engine/Source/Runtime/Engine/Private/GameFramework/SaveGame.cpp
+++ b/Engine/Source/Runtime/Engine/Private/GameFramework/SaveGame.cpp
@@ -1,6 +1,7 @@
 // Copyright 1998-2014 Epic Games, Inc. All Rights Reserved.
 
 #include "EnginePrivate.h"
+#include "GameFramework/SaveGame.h"
 
 USaveGame::USaveGame(const FObjectInitializer& ObjectInitializer)
 	: Super(ObjectInitializer)
diff --git a/Engine/Source/Runtime/Engine/Private/GameplayStatics.cpp b/Engine/Source/Runtime/Engine/Private/GameplayStatics.cpp
--- a/Engine/Source/Runtime/Engine/Private/GameplayStatics.cpp
+++ b/Engine/Source/Runtime/Engine/Private/GameplayStatics.cpp
@@ -9,9 +9,11 @@
 #include "SlateBasics.h"
 #include "Particles/ParticleSystemComponent.h"
 #include "GameFramework/Character.h"
+#include "Sound/DialogueWave.h"
 #include "Sound/SoundBase.h"
 #include "Sound/SoundCue.h"
 #include "Engine/GameInstance.h"
+#include "GameFramework/SaveGame.h"
 
 //////////////////////////////////////////////////////////////////////////
 // UGameplayStatics
diff --git a/Engine/Source/Runtime/Engine/Private/Interface.cpp b/Engine/Source/Runtime/Engine/Private/Interface.cpp
--- a/Engine/Source/Runtime/Engine/Private/Interface.cpp
+++ b/Engine/Source/Runtime/Engine/Private/Interface.cpp
@@ -1,6 +1,7 @@
 // Copyright 1998-2014 Epic Games, Inc. All Rights Reserved.
 
 #include "EnginePrivate.h"
+#include "MovieScene/RuntimeMovieScenePlayerInterface.h"
 
 UInterface_CollisionDataProvider::UInterface_CollisionDataProvider(const FObjectInitializer& ObjectInitializer)
 	: Super(ObjectInitializer)
diff --git a/Engine/Source/Runtime/Engine/Private/Landscape/LandscapeEditInterface.cpp b/Engine/Source/Runtime/Engine/Private/Landscape/LandscapeEditInterface.cpp
--- a/Engine/Source/Runtime/Engine/Private/Landscape/LandscapeEditInterface.cpp
+++ b/Engine/Source/Runtime/Engine/Private/Landscape/LandscapeEditInterface.cpp
@@ -11,6 +11,7 @@ LandscapeEditInterface.cpp: Landscape editing interface
 #include "Landscape/Landscape.h"
 #include "Landscape/LandscapeInfo.h"
 #include "Landscape/LandscapeLayerInfoObject.h"
+#include "ComponentReregisterContext.h"
 
 #if WITH_EDITOR
 
diff --git a/Engine/Source/Runtime/Engine/Private/Landscape/LandscapeLight.cpp b/Engine/Source/Runtime/Engine/Private/Landscape/LandscapeLight.cpp
--- a/Engine/Source/Runtime/Engine/Private/Landscape/LandscapeLight.cpp
+++ b/Engine/Source/Runtime/Engine/Private/Landscape/LandscapeLight.cpp
@@ -9,6 +9,7 @@ LandscapeLight.cpp: Static lighting for LandscapeComponents
 #include "Landscape/LandscapeInfo.h"
 #include "Landscape/LandscapeRender.h"
 #include "Landscape/LandscapeDataAccess.h"
+#include "ComponentReregisterContext.h"
 
 #if WITH_EDITOR
 
diff --git a/Engine/Source/Runtime/Engine/Private/LightComponent.cpp b/Engine/Source/Runtime/Engine/Private/LightComponent.cpp
--- a/Engine/Source/Runtime/Engine/Private/LightComponent.cpp
+++ b/Engine/Source/Runtime/Engine/Private/LightComponent.cpp
@@ -12,6 +12,7 @@
 #include "UObjectToken.h"
 #include "ComponentInstanceDataCache.h"
 #include "TargetPlatform.h"
+#include "ComponentReregisterContext.h"
 
 void FStaticShadowDepthMap::InitRHI()
 {
diff --git a/Engine/Source/Runtime/Engine/Private/LocalPlayer.cpp b/Engine/Source/Runtime/Engine/Private/LocalPlayer.cpp
--- a/Engine/Source/Runtime/Engine/Private/LocalPlayer.cpp
+++ b/Engine/Source/Runtime/Engine/Private/LocalPlayer.cpp
@@ -18,6 +18,7 @@
 #include "IHeadMountedDisplay.h"
 #include "SceneViewExtension.h"
 #include "DataChannel.h"
+#include "GameFramework/OnlineSession.h"
 
 DEFINE_LOG_CATEGORY(LogPlayerManagement);
 DEFINE_LOG_CATEGORY_STATIC(LogEngine, Log, All);
diff --git a/Engine/Source/Runtime/Engine/Private/Materials/MaterialExpressions.cpp b/Engine/Source/Runtime/Engine/Private/Materials/MaterialExpressions.cpp
--- a/Engine/Source/Runtime/Engine/Private/Materials/MaterialExpressions.cpp
+++ b/Engine/Source/Runtime/Engine/Private/Materials/MaterialExpressions.cpp
@@ -143,6 +143,8 @@
 #include "NotificationManager.h"
 #endif //WITH_EDITOR
 
+#include "Engine/Font.h"
+
 #define LOCTEXT_NAMESPACE "MaterialExpression"
 
 #define SWAP_REFERENCE_TO( ExpressionInput, ToBeRemovedExpression, ToReplaceWithExpression )	\
diff --git a/Engine/Source/Runtime/Engine/Private/Materials/MaterialShared.cpp b/Engine/Source/Runtime/Engine/Private/Materials/MaterialShared.cpp
--- a/Engine/Source/Runtime/Engine/Private/Materials/MaterialShared.cpp
+++ b/Engine/Source/Runtime/Engine/Private/Materials/MaterialShared.cpp
@@ -17,6 +17,7 @@
 #include "Developer/TargetPlatform/Public/TargetPlatform.h"
 #include "ComponentReregisterContext.h"
 #include "EngineModule.h"
+#include "Engine/Font.h"
 
 #include "LocalVertexFactory.h"
 
diff --git a/Engine/Source/Runtime/Engine/Private/ModelLight.cpp b/Engine/Source/Runtime/Engine/Private/ModelLight.cpp
--- a/Engine/Source/Runtime/Engine/Private/ModelLight.cpp
+++ b/Engine/Source/Runtime/Engine/Private/ModelLight.cpp
@@ -11,6 +11,7 @@
 #include "Model.h"
 #include "LightMap.h"
 #include "ShadowMap.h"
+#include "ComponentReregisterContext.h"
 
 //
 //	Static vars
diff --git a/Engine/Source/Runtime/Engine/Private/NetConnection.cpp b/Engine/Source/Runtime/Engine/Private/NetConnection.cpp
--- a/Engine/Source/Runtime/Engine/Private/NetConnection.cpp
+++ b/Engine/Source/Runtime/Engine/Private/NetConnection.cpp
@@ -10,6 +10,7 @@
 #include "Net/DataReplication.h"
 #include "Engine/ActorChannel.h"
 #include "DataChannel.h"
+#include "Engine/PackageMapClient.h"
 
 /*-----------------------------------------------------------------------------
 	UNetConnection implementation.
diff --git a/Engine/Source/Runtime/Engine/Private/NetworkDriver.cpp b/Engine/Source/Runtime/Engine/Private/NetworkDriver.cpp
--- a/Engine/Source/Runtime/Engine/Private/NetworkDriver.cpp
+++ b/Engine/Source/Runtime/Engine/Private/NetworkDriver.cpp
@@ -15,6 +15,7 @@
 #include "OnlineSubsystemUtils.h"
 #include "NetworkingDistanceConstants.h"
 #include "DataChannel.h"
+#include "Engine/PackageMapClient.h"
 
 // Default net driver stats
 DEFINE_STAT(STAT_Ping);
diff --git a/Engine/Source/Runtime/Engine/Private/NetworkSettings.cpp b/Engine/Source/Runtime/Engine/Private/NetworkSettings.cpp
--- a/Engine/Source/Runtime/Engine/Private/NetworkSettings.cpp
+++ b/Engine/Source/Runtime/Engine/Private/NetworkSettings.cpp
@@ -1,6 +1,7 @@
 // Copyright 1998-2014 Epic Games, Inc. All Rights Reserved.
 
 #include "EnginePrivate.h"
+#include "Engine/NetworkSettings.h"
 
 UNetworkSettings::UNetworkSettings(const FObjectInitializer& ObjectInitializer)
 	: Super(ObjectInitializer)
diff --git a/Engine/Source/Runtime/Engine/Private/NiagaraComponent.cpp b/Engine/Source/Runtime/Engine/Private/NiagaraComponent.cpp
--- a/Engine/Source/Runtime/Engine/Private/NiagaraComponent.cpp
+++ b/Engine/Source/Runtime/Engine/Private/NiagaraComponent.cpp
@@ -10,6 +10,7 @@
 #include "Engine/NiagaraEffectRenderer.h"
 #include "MeshBatch.h"
 #include "SceneUtils.h"
+#include "ComponentReregisterContext.h"
 
 DECLARE_CYCLE_STAT(TEXT("Tick"),STAT_NiagaraTick,STATGROUP_Niagara);
 DECLARE_CYCLE_STAT(TEXT("Simulate"),STAT_NiagaraSimulate,STATGROUP_Niagara);
diff --git a/Engine/Source/Runtime/Engine/Private/ObjectLibrary.cpp b/Engine/Source/Runtime/Engine/Private/ObjectLibrary.cpp
--- a/Engine/Source/Runtime/Engine/Private/ObjectLibrary.cpp
+++ b/Engine/Source/Runtime/Engine/Private/ObjectLibrary.cpp
@@ -2,6 +2,7 @@
 
 #include "EnginePrivate.h"
 #include "AssetRegistryModule.h"
+#include "Engine/ObjectLibrary.h"
 
 UObjectLibrary::UObjectLibrary(const FObjectInitializer& ObjectInitializer)
 	: Super(ObjectInitializer)
diff --git a/Engine/Source/Runtime/Engine/Private/ObjectReferencer.cpp b/Engine/Source/Runtime/Engine/Private/ObjectReferencer.cpp
--- a/Engine/Source/Runtime/Engine/Private/ObjectReferencer.cpp
+++ b/Engine/Source/Runtime/Engine/Private/ObjectReferencer.cpp
@@ -1,6 +1,7 @@
 // Copyright 1998-2014 Epic Games, Inc. All Rights Reserved.
 
 #include "EnginePrivate.h"
+#include "Engine/ObjectReferencer.h"
 
 UObjectReferencer::UObjectReferencer(const FObjectInitializer& ObjectInitializer)
 	: Super(ObjectInitializer)
diff --git a/Engine/Source/Runtime/Engine/Private/OnlineSession.cpp b/Engine/Source/Runtime/Engine/Private/OnlineSession.cpp
--- a/Engine/Source/Runtime/Engine/Private/OnlineSession.cpp
+++ b/Engine/Source/Runtime/Engine/Private/OnlineSession.cpp
@@ -5,7 +5,7 @@
 =============================================================================*/
 
 #include "EnginePrivate.h"
-
+#include "GameFramework/OnlineSession.h"
 
 UOnlineSession::UOnlineSession(const FObjectInitializer& ObjectInitializer)
 	: Super(ObjectInitializer)
diff --git a/Engine/Source/Runtime/Engine/Private/PackageMapClient.cpp b/Engine/Source/Runtime/Engine/Private/PackageMapClient.cpp
--- a/Engine/Source/Runtime/Engine/Private/PackageMapClient.cpp
+++ b/Engine/Source/Runtime/Engine/Private/PackageMapClient.cpp
@@ -7,6 +7,7 @@
 #include "Net/DataReplication.h"
 #include "Engine/ActorChannel.h"
 #include "RepLayout.h"
+#include "Engine/PackageMapClient.h"
 
 // ( OutPacketId == GUID_PACKET_NOT_ACKED ) == NAK'd		(this GUID is not acked, and is not pending either, so sort of waiting)
 // ( OutPacketId == GUID_PACKET_ACKED )		== FULLY ACK'd	(this GUID is fully acked, and we no longer need to send full path)
diff --git a/Engine/Source/Runtime/Engine/Private/Particles/ParticleBeamModules.cpp b/Engine/Source/Runtime/Engine/Private/Particles/ParticleBeamModules.cpp
--- a/Engine/Source/Runtime/Engine/Private/Particles/ParticleBeamModules.cpp
+++ b/Engine/Source/Runtime/Engine/Private/Particles/ParticleBeamModules.cpp
@@ -5,6 +5,7 @@
 =============================================================================*/
 
 #include "EnginePrivate.h"
+#include "Distributions/DistributionVectorConstant.h"
 #include "ParticleDefinitions.h"
 #include "../DistributionHelpers.h"
 #include "Particles/Beam/ParticleModuleBeamBase.h"
diff --git a/Engine/Source/Runtime/Engine/Private/Particles/ParticleComponents.cpp b/Engine/Source/Runtime/Engine/Private/Particles/ParticleComponents.cpp
--- a/Engine/Source/Runtime/Engine/Private/Particles/ParticleComponents.cpp
+++ b/Engine/Source/Runtime/Engine/Private/Particles/ParticleComponents.cpp
@@ -5,6 +5,9 @@
 =============================================================================*/
 
 #include "EnginePrivate.h"
+#include "Distributions/DistributionVectorConstant.h"
+#include "Distributions/DistributionVectorConstantCurve.h"
+#include "Distributions/DistributionVectorUniform.h"
 #include "StaticMeshResources.h"
 #include "ParticleDefinitions.h"
 #include "Particles/EmitterCameraLensEffectBase.h"
diff --git a/Engine/Source/Runtime/Engine/Private/Particles/ParticleModules.cpp b/Engine/Source/Runtime/Engine/Private/Particles/ParticleModules.cpp
--- a/Engine/Source/Runtime/Engine/Private/Particles/ParticleModules.cpp
+++ b/Engine/Source/Runtime/Engine/Private/Particles/ParticleModules.cpp
@@ -5,6 +5,10 @@
 =============================================================================*/
 
 #include "EnginePrivate.h"
+#include "Distributions/DistributionFloatParticleParameter.h"
+#include "Distributions/DistributionVectorParticleParameter.h"
+#include "Distributions/DistributionVectorConstantCurve.h"
+#include "Distributions/DistributionVectorUniformCurve.h"
 #include "FXSystem.h"
 #include "ParticleDefinitions.h"
 #include "../DistributionHelpers.h"
diff --git a/Engine/Source/Runtime/Engine/Private/Particles/ParticleModules_Color.cpp b/Engine/Source/Runtime/Engine/Private/Particles/ParticleModules_Color.cpp
--- a/Engine/Source/Runtime/Engine/Private/Particles/ParticleModules_Color.cpp
+++ b/Engine/Source/Runtime/Engine/Private/Particles/ParticleModules_Color.cpp
@@ -5,6 +5,8 @@
 	Color-related particle module implementations.
 =============================================================================*/
 #include "EnginePrivate.h"
+#include "Distributions/DistributionFloatParticleParameter.h"
+#include "Distributions/DistributionVectorParticleParameter.h"
 #include "ParticleDefinitions.h"
 #include "../DistributionHelpers.h"
 #include "Particles/Color/ParticleModuleColor.h"
@@ -16,6 +18,7 @@
 #include "Particles/ParticleLODLevel.h"
 #include "Particles/ParticleSystem.h"
 #include "Particles/ParticleSystemComponent.h"
+#include "Distributions/DistributionVectorConstantCurve.h"
 
 UParticleModuleColorBase::UParticleModuleColorBase(const FObjectInitializer& ObjectInitializer)
 	: Super(ObjectInitializer)
diff --git a/Engine/Source/Runtime/Engine/Private/Particles/ParticleModules_Location.cpp b/Engine/Source/Runtime/Engine/Private/Particles/ParticleModules_Location.cpp
--- a/Engine/Source/Runtime/Engine/Private/Particles/ParticleModules_Location.cpp
+++ b/Engine/Source/Runtime/Engine/Private/Particles/ParticleModules_Location.cpp
@@ -6,6 +6,7 @@
 =============================================================================*/
 
 #include "EnginePrivate.h"
+#include "Distributions/DistributionVectorConstantCurve.h"
 #include "ParticleDefinitions.h"
 #include "RawIndexBuffer.h"
 #include "../DistributionHelpers.h"
diff --git a/Engine/Source/Runtime/Engine/Private/Particles/ParticleModules_Size.cpp b/Engine/Source/Runtime/Engine/Private/Particles/ParticleModules_Size.cpp
--- a/Engine/Source/Runtime/Engine/Private/Particles/ParticleModules_Size.cpp
+++ b/Engine/Source/Runtime/Engine/Private/Particles/ParticleModules_Size.cpp
@@ -6,6 +6,7 @@
 =============================================================================*/
 
 #include "EnginePrivate.h"
+#include "Distributions/DistributionVectorConstantCurve.h"
 #include "ParticleDefinitions.h"
 #include "../DistributionHelpers.h"
 #include "Particles/Size/ParticleModuleSize_Seeded.h"
diff --git a/Engine/Source/Runtime/Engine/Private/Particles/ParticleModules_Velocity.cpp b/Engine/Source/Runtime/Engine/Private/Particles/ParticleModules_Velocity.cpp
--- a/Engine/Source/Runtime/Engine/Private/Particles/ParticleModules_Velocity.cpp
+++ b/Engine/Source/Runtime/Engine/Private/Particles/ParticleModules_Velocity.cpp
@@ -6,6 +6,9 @@
 =============================================================================*/
 
 #include "EnginePrivate.h"
+#include "Distributions/DistributionVectorConstant.h"
+#include "Distributions/DistributionVectorConstantCurve.h"
+#include "Distributions/DistributionVectorUniform.h"
 #include "ParticleDefinitions.h"
 #include "../DistributionHelpers.h"
 #include "Particles/Lifetime/ParticleModuleLifetimeBase.h"
diff --git a/Engine/Source/Runtime/Engine/Private/PlayerCameraManager.cpp b/Engine/Source/Runtime/Engine/Private/PlayerCameraManager.cpp
--- a/Engine/Source/Runtime/Engine/Private/PlayerCameraManager.cpp
+++ b/Engine/Source/Runtime/Engine/Private/PlayerCameraManager.cpp
@@ -7,6 +7,8 @@
 #include "IHeadMountedDisplay.h"
 #include "Particles/EmitterCameraLensEffectBase.h"
 #include "Camera/CameraActor.h"
+#include "Camera/CameraAnim.h"
+#include "Camera/CameraAnimInst.h"
 
 DEFINE_LOG_CATEGORY_STATIC(LogPlayerCameraManager, Log, All);
 
diff --git a/Engine/Source/Runtime/Engine/Private/RendererSupport.cpp b/Engine/Source/Runtime/Engine/Private/RendererSupport.cpp
--- a/Engine/Source/Runtime/Engine/Private/RendererSupport.cpp
+++ b/Engine/Source/Runtime/Engine/Private/RendererSupport.cpp
@@ -13,6 +13,7 @@
 #include "EngineModule.h"
 #include "RendererInterface.h"
 #include "HotReloadInterface.h"
+#include "ComponentReregisterContext.h"
 
 /** Clears and optionally backs up all references to renderer module classes in other modules, particularly engine. */
 void ClearReferencesToRendererModuleClasses(
diff --git a/Engine/Source/Runtime/Engine/Private/SCS_Node.cpp b/Engine/Source/Runtime/Engine/Private/SCS_Node.cpp
--- a/Engine/Source/Runtime/Engine/Private/SCS_Node.cpp
+++ b/Engine/Source/Runtime/Engine/Private/SCS_Node.cpp
@@ -7,6 +7,7 @@
 #include "Kismet2/BlueprintEditorUtils.h"
 #endif
 #include "DeferRegisterComponents.h"
+#include "Engine/SCS_Node.h"
 
 //////////////////////////////////////////////////////////////////////////
 // USCS_Node
diff --git a/Engine/Source/Runtime/Engine/Private/SaveGameSummary.cpp b/Engine/Source/Runtime/Engine/Private/SaveGameSummary.cpp
--- a/Engine/Source/Runtime/Engine/Private/SaveGameSummary.cpp
+++ b/Engine/Source/Runtime/Engine/Private/SaveGameSummary.cpp
@@ -1,6 +1,7 @@
 // Copyright 1998-2014 Epic Games, Inc. All Rights Reserved.
 
 #include "EnginePrivate.h"
+#include "Engine/SaveGameSummary.h"
 
 UDEPRECATED_SaveGameSummary::UDEPRECATED_SaveGameSummary(const FObjectInitializer& ObjectInitializer)
 	: Super(ObjectInitializer)
diff --git a/Engine/Source/Runtime/Engine/Private/SceneComponent.cpp b/Engine/Source/Runtime/Engine/Private/SceneComponent.cpp
--- a/Engine/Source/Runtime/Engine/Private/SceneComponent.cpp
+++ b/Engine/Source/Runtime/Engine/Private/SceneComponent.cpp
@@ -9,6 +9,7 @@
 #include "MessageLog.h"
 #include "Net/UnrealNetwork.h"
 #include "GameFramework/PhysicsVolume.h"
+#include "ComponentReregisterContext.h"
 
 #define LOCTEXT_NAMESPACE "SceneComponent"
 
diff --git a/Engine/Source/Runtime/Engine/Private/Selection.cpp b/Engine/Source/Runtime/Engine/Private/Selection.cpp
--- a/Engine/Source/Runtime/Engine/Private/Selection.cpp
+++ b/Engine/Source/Runtime/Engine/Private/Selection.cpp
@@ -3,6 +3,7 @@
 
 #include "EnginePrivate.h"
 #include "UObjectAnnotation.h"
+#include "Engine/Selection.h"
 
 DEFINE_LOG_CATEGORY_STATIC(LogSelection, Log, All);
 
diff --git a/Engine/Source/Runtime/Engine/Private/SimpleConstructionScript.cpp b/Engine/Source/Runtime/Engine/Private/SimpleConstructionScript.cpp
--- a/Engine/Source/Runtime/Engine/Private/SimpleConstructionScript.cpp
+++ b/Engine/Source/Runtime/Engine/Private/SimpleConstructionScript.cpp
@@ -1,6 +1,7 @@
 // Copyright 1998-2014 Epic Games, Inc. All Rights Reserved.
 
 #include "EnginePrivate.h"
+#include "Engine/SCS_Node.h"
 #include "BlueprintUtilities.h"
 #if WITH_EDITOR
 #include "Kismet2/CompilerResultsLog.h"
diff --git a/Engine/Source/Runtime/Engine/Private/SkinnedMeshComponent.cpp b/Engine/Source/Runtime/Engine/Private/SkinnedMeshComponent.cpp
--- a/Engine/Source/Runtime/Engine/Private/SkinnedMeshComponent.cpp
+++ b/Engine/Source/Runtime/Engine/Private/SkinnedMeshComponent.cpp
@@ -12,6 +12,7 @@
 #include "AnimationUtils.h"
 #include "AnimTree.h"
 #include "Animation/VertexAnim/MorphTarget.h"
+#include "ComponentReregisterContext.h"
 
 DEFINE_LOG_CATEGORY_STATIC(LogSkinnedMeshComp, Log, All);
 
diff --git a/Engine/Source/Runtime/Engine/Private/SoundNodeDialoguePlayer.cpp b/Engine/Source/Runtime/Engine/Private/SoundNodeDialoguePlayer.cpp
--- a/Engine/Source/Runtime/Engine/Private/SoundNodeDialoguePlayer.cpp
+++ b/Engine/Source/Runtime/Engine/Private/SoundNodeDialoguePlayer.cpp
@@ -1,6 +1,7 @@
 // Copyright 1998-2014 Epic Games, Inc. All Rights Reserved.
 
 #include "EnginePrivate.h"
+#include "Sound/DialogueWave.h"
 #include "Sound/SoundBase.h"
 #include "Sound/SoundNodeDialoguePlayer.h"
 
diff --git a/Engine/Source/Runtime/Engine/Private/StaticMeshComponent.cpp b/Engine/Source/Runtime/Engine/Private/StaticMeshComponent.cpp
--- a/Engine/Source/Runtime/Engine/Private/StaticMeshComponent.cpp
+++ b/Engine/Source/Runtime/Engine/Private/StaticMeshComponent.cpp
@@ -10,6 +10,7 @@
 #include "ComponentInstanceDataCache.h"
 #include "LightMap.h"
 #include "ShadowMap.h"
+#include "ComponentReregisterContext.h"
 
 #define LOCTEXT_NAMESPACE "StaticMeshComponent"
 
diff --git a/Engine/Source/Runtime/Engine/Private/StaticMeshLight.cpp b/Engine/Source/Runtime/Engine/Private/StaticMeshLight.cpp
--- a/Engine/Source/Runtime/Engine/Private/StaticMeshLight.cpp
+++ b/Engine/Source/Runtime/Engine/Private/StaticMeshLight.cpp
@@ -10,6 +10,7 @@
 #include "StaticMeshLight.h"
 #include "LightMap.h"
 #include "ShadowMap.h"
+#include "ComponentReregisterContext.h"
 
 /**
  * Creates a static lighting vertex to represent the given static mesh vertex.
diff --git a/Engine/Source/Runtime/Engine/Private/StaticMeshRender.cpp b/Engine/Source/Runtime/Engine/Private/StaticMeshRender.cpp
--- a/Engine/Source/Runtime/Engine/Private/StaticMeshRender.cpp
+++ b/Engine/Source/Runtime/Engine/Private/StaticMeshRender.cpp
@@ -11,6 +11,7 @@
 #include "LightMap.h"
 #include "ShadowMap.h"
 #include "DistanceFieldAtlas.h"
+#include "ComponentReregisterContext.h"
 
 /** If true, optimized depth-only index buffers are used for shadow rendering. */
 static bool GUseShadowIndexBuffer = true;
diff --git a/Engine/Source/Runtime/Engine/Private/SubtitleManager.cpp b/Engine/Source/Runtime/Engine/Private/SubtitleManager.cpp
--- a/Engine/Source/Runtime/Engine/Private/SubtitleManager.cpp
+++ b/Engine/Source/Runtime/Engine/Private/SubtitleManager.cpp
@@ -4,6 +4,7 @@
 #include "EnginePrivate.h"
 #include "SubtitleManager.h"
 #include "SoundDefinitions.h"
+#include "Engine/Font.h"
 
 DEFINE_LOG_CATEGORY(LogSubtitle);
 
diff --git a/Engine/Source/Runtime/Engine/Private/TextRenderComponent.cpp b/Engine/Source/Runtime/Engine/Private/TextRenderComponent.cpp
--- a/Engine/Source/Runtime/Engine/Private/TextRenderComponent.cpp
+++ b/Engine/Source/Runtime/Engine/Private/TextRenderComponent.cpp
@@ -2,6 +2,7 @@
 
 #include "EnginePrivate.h"
 #include "DynamicMeshBuilder.h"
+#include "Engine/Font.h"
 #include "Engine/TextRenderActor.h"
 #include "LocalVertexFactory.h"
 
diff --git a/Engine/Source/Runtime/Engine/Private/UnrealClient.cpp b/Engine/Source/Runtime/Engine/Private/UnrealClient.cpp
--- a/Engine/Source/Runtime/Engine/Private/UnrealClient.cpp
+++ b/Engine/Source/Runtime/Engine/Private/UnrealClient.cpp
@@ -2,6 +2,7 @@
 
 
 #include "EnginePrivate.h"
+#include "Engine/Font.h"
 #include "Matinee/MatineeActor.h"
 #include "EditorSupportDelegates.h"
 #include "RenderCore.h"
diff --git a/Engine/Source/Runtime/Engine/Private/UnrealEngine.cpp b/Engine/Source/Runtime/Engine/Private/UnrealEngine.cpp
--- a/Engine/Source/Runtime/Engine/Private/UnrealEngine.cpp
+++ b/Engine/Source/Runtime/Engine/Private/UnrealEngine.cpp
@@ -5,7 +5,8 @@
 =============================================================================*/
 
 #include "EnginePrivate.h"
-
+#include "Engine/Font.h"
+#include "Engine/ObjectReferencer.h"
 #include "Net/UnrealNetwork.h"
 #include "Engine/Console.h"
 #include "VisualLog.h"
diff --git a/Engine/Source/Runtime/Engine/Private/UserDefinedEnum.cpp b/Engine/Source/Runtime/Engine/Private/UserDefinedEnum.cpp
--- a/Engine/Source/Runtime/Engine/Private/UserDefinedEnum.cpp
+++ b/Engine/Source/Runtime/Engine/Private/UserDefinedEnum.cpp
@@ -1,7 +1,7 @@
 // Copyright 1998-2014 Epic Games, Inc. All Rights Reserved.
 
 #include "EnginePrivate.h"
-
+#include "Engine/UserDefinedEnum.h"
 #if WITH_EDITOR
 #include "Editor/UnrealEd/Public/Kismet2/EnumEditorUtils.h"
 #endif	// WITH_EDITOR
diff --git a/Engine/Source/Runtime/Engine/Private/UserInterface/Canvas.cpp b/Engine/Source/Runtime/Engine/Private/UserInterface/Canvas.cpp
--- a/Engine/Source/Runtime/Engine/Private/UserInterface/Canvas.cpp
+++ b/Engine/Source/Runtime/Engine/Private/UserInterface/Canvas.cpp
@@ -6,6 +6,7 @@
 
 #include "EnginePrivate.h"
 #include "SlateBasics.h"
+#include "Engine/Font.h"
 #include "TileRendering.h"
 #include "RHIStaticStates.h"
 #include "WordWrapper.h"
diff --git a/Engine/Source/Runtime/Engine/Private/VectorField.cpp b/Engine/Source/Runtime/Engine/Private/VectorField.cpp
--- a/Engine/Source/Runtime/Engine/Private/VectorField.cpp
+++ b/Engine/Source/Runtime/Engine/Private/VectorField.cpp
@@ -18,6 +18,7 @@
 #include "VectorField/VectorFieldAnimated.h"
 #include "VectorField/VectorFieldStatic.h"
 #include "SceneUtils.h"
+#include "ComponentReregisterContext.h"
 
 #define MAX_GLOBAL_VECTOR_FIELDS (16)
 DEFINE_LOG_CATEGORY(LogVectorField)
diff --git a/Engine/Source/Runtime/Engine/Private/World.cpp b/Engine/Source/Runtime/Engine/Private/World.cpp
--- a/Engine/Source/Runtime/Engine/Private/World.cpp
+++ b/Engine/Source/Runtime/Engine/Private/World.cpp
@@ -5,9 +5,11 @@
 =============================================================================*/
 
 #include "EnginePrivate.h"
+#include "ComponentReregisterContext.h"
 #include "Engine/LevelScriptBlueprint.h"
 #include "Engine/LevelScriptActor.h"
 #include "Engine/CullDistanceVolume.h"
+#include "Engine/SaveGameSummary.h"
 #include "GameFramework/DefaultPhysicsVolume.h"
 #include "Engine/Console.h"
 #include "Engine/WorldComposition.h"
