Fix landscape holes on ES2 :)

diff --git a/Engine/Shaders/LandscapeVertexFactory.usf b/Engine/Shaders/LandscapeVertexFactory.usf
--- a/Engine/Shaders/LandscapeVertexFactory.usf
+++ b/Engine/Shaders/LandscapeVertexFactory.usf
@@ -52,6 +52,10 @@ Texture2D XYOffsetmapTexture;
 SamplerState XYOffsetmapTextureSampler;
 #endif
 
+#if ES2_PROFILE || ES3_1_PROFILE
+float3 BlendableLayerMask;
+#endif
+
 struct FVertexFactoryInput
 {
 #if FEATURE_LEVEL >= FEATURE_LEVEL_SM4
@@ -336,8 +340,8 @@ FMaterialPixelParameters GetMaterialPixelParameters(FVertexFactoryInterpolantsVS
 #else
 	Result.LayerWeights.xy = Texture2DSample(NormalmapTexture, NormalmapTextureSampler, Interpolants.WeightMapTexCoord).rg;
 #endif 
- 
-	Result.LayerWeights.z = saturate(1.f - Result.LayerWeights.x - Result.LayerWeights.y);
+
+	Result.LayerWeights.z = saturate(BlendableLayerMask.z - dot(Result.LayerWeights.xy, BlendableLayerMask.xy));
 	Result.LayerWeights.w = 0; // Total weight need to be 1
 #endif
 
diff --git a/Engine/Source/Runtime/Engine/Classes/Landscape/LandscapeComponent.h b/Engine/Source/Runtime/Engine/Classes/Landscape/LandscapeComponent.h
--- a/Engine/Source/Runtime/Engine/Classes/Landscape/LandscapeComponent.h
+++ b/Engine/Source/Runtime/Engine/Classes/Landscape/LandscapeComponent.h
@@ -218,6 +218,9 @@ public:
 	UPROPERTY(Transient, DuplicateTransient)
 	UTexture2D* SelectDataTexture; // Data texture used for selection mask
 
+	/** Pointer to data shared with the render thread, used by the editor tools */
+	struct FLandscapeEditToolRenderData* EditToolRenderData;
+
 	/** Runtime-generated editor data for ES2 emulation */
 	UPROPERTY(Transient, DuplicateTransient)
 	UTexture2D* MobileNormalmapTexture;
@@ -227,10 +230,11 @@ public:
 	UMaterialInterface* MobileMaterialInterface;
 #endif
 
+	/** For ES2 */
+	UPROPERTY()
+	uint8 MobileBlendableLayerMask;
 
 public:
-	/** Pointer to data shared with the render therad, used by the editor tools */
-	struct FLandscapeEditToolRenderData* EditToolRenderData;
 	/** Platform Data where don't support texture sampling in vertex buffer */
 	FLandscapeComponentDerivedData PlatformData;
 
diff --git a/Engine/Source/Runtime/Engine/Private/Landscape/Landscape.cpp b/Engine/Source/Runtime/Engine/Private/Landscape/Landscape.cpp
--- a/Engine/Source/Runtime/Engine/Private/Landscape/Landscape.cpp
+++ b/Engine/Source/Runtime/Engine/Private/Landscape/Landscape.cpp
@@ -91,7 +91,9 @@ ULandscapeComponent::ULandscapeComponent(const FObjectInitializer& ObjectInitial
 
 	Mobility = EComponentMobility::Static;
 
+#if WITH_EDITORONLY_DATA
 	EditToolRenderData = NULL;
+#endif
 
 	LpvBiasMultiplier = 0.0f; // Bias is 0 for landscape, since it's single sided
 }
diff --git a/Engine/Source/Runtime/Engine/Private/Landscape/LandscapeEdit.cpp b/Engine/Source/Runtime/Engine/Private/Landscape/LandscapeEdit.cpp
--- a/Engine/Source/Runtime/Engine/Private/Landscape/LandscapeEdit.cpp
+++ b/Engine/Source/Runtime/Engine/Private/Landscape/LandscapeEdit.cpp
@@ -4757,6 +4757,30 @@ bool ALandscapeProxy::GetSelectedComponents(TArray<UObject*>& SelectedObjects)
 	return false;
 }
 
+struct FMobileLayerAllocation
+{
+	FWeightmapLayerAllocationInfo Allocation;
+
+	friend bool operator<(const FMobileLayerAllocation& lhs, const FMobileLayerAllocation& rhs)
+	{
+		if (!lhs.Allocation.LayerInfo && !rhs.Allocation.LayerInfo) return false; // equally broken :P
+		if (!lhs.Allocation.LayerInfo && rhs.Allocation.LayerInfo) return false; // broken layers sort to the end
+		if (!rhs.Allocation.LayerInfo && lhs.Allocation.LayerInfo) return true;
+
+		if (lhs.Allocation.LayerInfo == ALandscapeProxy::VisibilityLayer) return true; // visibility layer to the front
+		if (rhs.Allocation.LayerInfo == ALandscapeProxy::VisibilityLayer) return false;
+
+		if (lhs.Allocation.LayerInfo->bNoWeightBlend && !rhs.Allocation.LayerInfo->bNoWeightBlend) return false; // non-blended layers sort to the end
+		if (rhs.Allocation.LayerInfo->bNoWeightBlend && !lhs.Allocation.LayerInfo->bNoWeightBlend) return true;
+
+		// TODO: If we want to support cleanly decaying a pc landscape for mobile
+		// we should probably add other sort criteria, e.g. coverage
+		// or e.g. add an "importance" to layerinfos and sort on that
+
+		return false; // equal, preserve order
+	}
+};
+
 UMaterialInstance* ULandscapeComponent::GeneratePlatformPixelData(TArray<UTexture2D*>& InWeightmapTextures, bool bIsCooking)
 {
 	if (IsTemplate() || HeightmapTexture == NULL)
@@ -4764,6 +4788,49 @@ UMaterialInstance* ULandscapeComponent::GeneratePlatformPixelData(TArray<UTextur
 		return MaterialInstance;
 	}
 
+	TArray<FMobileLayerAllocation> MobileLayerAllocations;
+	MobileLayerAllocations.Reserve(WeightmapLayerAllocations.Num());
+	for (const auto& Allocation : WeightmapLayerAllocations)
+	{
+		MobileLayerAllocations.Add(FMobileLayerAllocation{ Allocation });
+	}
+	MobileLayerAllocations.StableSort();
+
+	// in the current mobile shader only 3 layers are supported (the 3rd only as a blended layer)
+	// so make sure we have a blended layer for layer 3 if possible
+	if (MobileLayerAllocations.Num() >= 3 &&
+		MobileLayerAllocations[2].Allocation.LayerInfo && MobileLayerAllocations[2].Allocation.LayerInfo->bNoWeightBlend)
+	{
+		int32 BlendedLayerToMove = INDEX_NONE;
+
+		// First try to swap layer 3 with an earlier blended layer
+		// this will allow both to work
+		for (int32 i = 1; i >= 0; --i)
+		{
+			if (MobileLayerAllocations[i].Allocation.LayerInfo && !MobileLayerAllocations[i].Allocation.LayerInfo->bNoWeightBlend)
+			{
+				BlendedLayerToMove = i;
+				break;
+			}
+		}
+
+		// otherwise swap layer 3 with the first weight-blended layer found
+		// as non-blended layers aren't supported for layer 3 it wasn't going to work anyway, might as well swap it out for one that will work
+		if (BlendedLayerToMove == INDEX_NONE)
+		{
+			// I wish I could specify a start index here, but it doesn't affect the result
+			BlendedLayerToMove = MobileLayerAllocations.IndexOfByPredicate([](const FMobileLayerAllocation& MobileAllocation){ return MobileAllocation.Allocation.LayerInfo && !MobileAllocation.Allocation.LayerInfo->bNoWeightBlend; });
+		}
+
+		if (BlendedLayerToMove != INDEX_NONE)
+		{
+			// Preserve order of all but the blended layer we're moving into slot 3
+			FMobileLayerAllocation TempAllocation = MoveTemp(MobileLayerAllocations[BlendedLayerToMove]);
+			MobileLayerAllocations.RemoveAt(BlendedLayerToMove, 1, false);
+			MobileLayerAllocations.Insert(MoveTemp(TempAllocation), 2);
+		}
+	}
+
 	int32 WeightmapSize = (SubsectionSizeQuads + 1) * NumSubsections;
 	UTexture2D* WeightmapTexture = GetLandscapeProxy()->CreateLandscapeTexture(WeightmapSize, WeightmapSize, TEXTUREGROUP_Terrain_Weightmap, TSF_BGRA8);
 	CreateEmptyTextureMips(WeightmapTexture);
@@ -4773,27 +4840,22 @@ UMaterialInstance* ULandscapeComponent::GeneratePlatformPixelData(TArray<UTextur
 
 		if (InWeightmapTextures.Num() > 0)
 		{
-			// Reordering weight map channels
 			int32 CurrentIdx = 0;
-			int32 FromTextures[3] = { 0, 0, 0 };
-			int32 FromChannels[3] = { 0, 0, 0 };
-			for (int32 i = 0; i < WeightmapLayerAllocations.Num(); ++i)
+			for (const auto& MobileAllocation : MobileLayerAllocations)
 			{
 				// Only for valid Layers
-				if (WeightmapLayerAllocations[i].LayerInfo)
+				if (MobileAllocation.Allocation.LayerInfo)
 				{
-					FromTextures[CurrentIdx] = WeightmapLayerAllocations[i].WeightmapTextureIndex;
-					FromChannels[CurrentIdx] = WeightmapLayerAllocations[i].WeightmapTextureChannel;
+					uint8 TextureIndex = MobileAllocation.Allocation.WeightmapTextureIndex;
+					uint8 TextureChannel = MobileAllocation.Allocation.WeightmapTextureChannel;
+					LandscapeEdit.CopyTextureChannel(WeightmapTexture, CurrentIdx, InWeightmapTextures[TextureIndex], TextureChannel);
 					CurrentIdx++;
-				}
-				if (CurrentIdx >= 3) // ignore 4th channel
-				{
-					break;
+					if (CurrentIdx >= 2) // Only support 2 layers in texture
+					{
+						break;
+					}
 				}
 			}
-
-			LandscapeEdit.CopyTextureChannel(WeightmapTexture, 0, InWeightmapTextures[FromTextures[0]], FromChannels[0]);
-			LandscapeEdit.CopyTextureChannel(WeightmapTexture, 1, InWeightmapTextures[FromTextures[1]], FromChannels[1]);
 		}
 
 		LandscapeEdit.CopyTextureFromHeightmap(WeightmapTexture, 2, this, 2);
@@ -4804,28 +4866,30 @@ UMaterialInstance* ULandscapeComponent::GeneratePlatformPixelData(TArray<UTextur
 	InWeightmapTextures.Empty();
 	InWeightmapTextures.Add(WeightmapTexture);
 
-	FLinearColor Masks[4];
-	Masks[0] = FLinearColor(1.f, 0.f, 0.f, 0.f);
-	Masks[1] = FLinearColor(0.f, 1.f, 0.f, 0.f);
-	Masks[2] = FLinearColor(0.f, 0.f, 1.f, 0.f);
-	Masks[3] = FLinearColor(0.f, 0.f, 0.f, 1.f);
+	FLinearColor Masks[5];
+	Masks[0] = FLinearColor(1, 0, 0, 0);
+	Masks[1] = FLinearColor(0, 1, 0, 0);
+	Masks[2] = FLinearColor(0, 0, 1, 0);
+	Masks[3] = FLinearColor(0, 0, 0, 1);
+	Masks[4] = FLinearColor(0, 0, 0, 0); // mask out layers 4+ altogether
 
 	if (!bIsCooking)
 	{
 		UMaterialInstanceDynamic* MobileMaterialInstance = UMaterialInstanceDynamic::Create(MaterialInstance, GetOutermost());
+
+		MobileBlendableLayerMask = 0;
+
 		// Set the layer mask
 		int32 CurrentIdx = 0;
-		for (int32 AllocIdx = 0; AllocIdx < WeightmapLayerAllocations.Num(); AllocIdx++)
+		for (const auto& MobileAllocation : MobileLayerAllocations)
 		{
-			FWeightmapLayerAllocationInfo& Allocation = WeightmapLayerAllocations[AllocIdx];
-			FName LayerName = Allocation.LayerInfo ? Allocation.LayerInfo->LayerName : NAME_None;
+			const FWeightmapLayerAllocationInfo& Allocation = MobileAllocation.Allocation;
 			if (Allocation.LayerInfo)
 			{
-				MobileMaterialInstance->SetVectorParameterValue(FName(*FString::Printf(TEXT("LayerMask_%s"), *LayerName.ToString())), Masks[FMath::Min(3, CurrentIdx++)]);
-			}
-			else
-			{
-				MobileMaterialInstance->SetVectorParameterValue(FName(*FString::Printf(TEXT("LayerMask_%s"), *LayerName.ToString())), Masks[3]);
+				FName LayerName = Allocation.LayerInfo == ALandscapeProxy::VisibilityLayer ? UMaterialExpressionLandscapeVisibilityMask::ParameterName : Allocation.LayerInfo->LayerName;
+				MobileMaterialInstance->SetVectorParameterValue(FName(*FString::Printf(TEXT("LayerMask_%s"), *LayerName.ToString())), Masks[FMath::Min(4, CurrentIdx)]);
+				MobileBlendableLayerMask |= (!Allocation.LayerInfo->bNoWeightBlend ? (1 << CurrentIdx) : 0);
+				CurrentIdx++;
 			}
 		}
 		return MobileMaterialInstance;
@@ -4837,19 +4901,19 @@ UMaterialInstance* ULandscapeComponent::GeneratePlatformPixelData(TArray<UTextur
 
 		MobileMaterialInstance->SetParentEditorOnly(CombinationMaterialInstance);
 
+		MobileBlendableLayerMask = 0;
+
 		// Set the layer mask
 		int32 CurrentIdx = 0;
-		for (int32 AllocIdx = 0; AllocIdx < WeightmapLayerAllocations.Num(); AllocIdx++)
+		for (const auto& MobileAllocation : MobileLayerAllocations)
 		{
-			FWeightmapLayerAllocationInfo& Allocation = WeightmapLayerAllocations[AllocIdx];
-			FName LayerName = Allocation.LayerInfo ? Allocation.LayerInfo->LayerName : NAME_None;
+			const FWeightmapLayerAllocationInfo& Allocation = MobileAllocation.Allocation;
 			if (Allocation.LayerInfo)
 			{
-				MobileMaterialInstance->SetVectorParameterValueEditorOnly(FName(*FString::Printf(TEXT("LayerMask_%s"), *LayerName.ToString())), Masks[FMath::Min(3, CurrentIdx++)]);
-			}
-			else
-			{
-				MobileMaterialInstance->SetVectorParameterValueEditorOnly(FName(*FString::Printf(TEXT("LayerMask_%s"), *LayerName.ToString())), Masks[3]);
+				FName LayerName = Allocation.LayerInfo == ALandscapeProxy::VisibilityLayer ? UMaterialExpressionLandscapeVisibilityMask::ParameterName : Allocation.LayerInfo->LayerName;
+				MobileMaterialInstance->SetVectorParameterValueEditorOnly(FName(*FString::Printf(TEXT("LayerMask_%s"), *LayerName.ToString())), Masks[FMath::Min(4, CurrentIdx)]);
+				MobileBlendableLayerMask |= (!Allocation.LayerInfo->bNoWeightBlend ? (1 << CurrentIdx) : 0);
+				CurrentIdx++;
 			}
 		}
 
diff --git a/Engine/Source/Runtime/Engine/Private/Landscape/LandscapeRenderMobile.cpp b/Engine/Source/Runtime/Engine/Private/Landscape/LandscapeRenderMobile.cpp
--- a/Engine/Source/Runtime/Engine/Private/Landscape/LandscapeRenderMobile.cpp
+++ b/Engine/Source/Runtime/Engine/Private/Landscape/LandscapeRenderMobile.cpp
@@ -68,7 +68,7 @@ public:
 		const FLandscapeBatchElementParams* BatchElementParams = (const FLandscapeBatchElementParams*)BatchElement.UserData;
 		check(BatchElementParams);
 
-		const FLandscapeComponentSceneProxy* SceneProxy = BatchElementParams->SceneProxy;
+		const FLandscapeComponentSceneProxyMobile* SceneProxy = (const FLandscapeComponentSceneProxyMobile*)BatchElementParams->SceneProxy;
 		SetUniformBufferParameter(RHICmdList, VertexShader->GetVertexShader(),VertexShader->GetUniformBufferParameter<FLandscapeUniformShaderParameters>(),*BatchElementParams->LandscapeUniformShaderParametersResource);
 
 		FVector CameraLocalPos3D = SceneProxy->WorldToLocal.TransformPosition(View.ViewMatrices.ViewOrigin); 
@@ -137,6 +137,55 @@ protected:
 	TShaderUniformBufferParameter<FLandscapeUniformShaderParameters> LandscapeShaderParameters;
 };
 
+/** Shader parameters for use with FLandscapeVertexFactory */
+class FLandscapeVertexFactoryMobilePixelShaderParameters : public FLandscapeVertexFactoryPixelShaderParameters
+{
+public:
+	/**
+	* Bind shader constants by name
+	* @param	ParameterMap - mapping of named shader constants to indices
+	*/
+	virtual void Bind(const FShaderParameterMap& ParameterMap) override
+	{
+		FLandscapeVertexFactoryPixelShaderParameters::Bind(ParameterMap);
+		BlendableLayerMaskParameter.Bind(ParameterMap, TEXT("BlendableLayerMask"));
+	}
+
+	/**
+	* Serialize shader params to an archive
+	* @param	Ar - archive to serialize to
+	*/
+	virtual void Serialize(FArchive& Ar) override
+	{
+		FLandscapeVertexFactoryPixelShaderParameters::Serialize(Ar);
+		Ar << BlendableLayerMaskParameter;
+	}
+	/**
+	* Set any shader data specific to this vertex factory
+	*/
+	virtual void SetMesh(FRHICommandList& RHICmdList, FShader* PixelShader, const FVertexFactory* VertexFactory, const FSceneView& View, const FMeshBatchElement& BatchElement, uint32 DataFlags) const override
+	{
+		SCOPE_CYCLE_COUNTER(STAT_LandscapeVFDrawTime);
+
+		FLandscapeVertexFactoryPixelShaderParameters::SetMesh(RHICmdList, PixelShader, VertexFactory, View, BatchElement, DataFlags);
+
+		const FLandscapeBatchElementParams* BatchElementParams = (const FLandscapeBatchElementParams*)BatchElement.UserData;
+		check(BatchElementParams);
+		const FLandscapeComponentSceneProxyMobile* SceneProxy = (const FLandscapeComponentSceneProxyMobile*)BatchElementParams->SceneProxy;
+
+		if (BlendableLayerMaskParameter.IsBound())
+		{
+			FVector MaskVector;
+			MaskVector[0] = (SceneProxy->BlendableLayerMask & (1 << 0)) ? 1 : 0;
+			MaskVector[1] = (SceneProxy->BlendableLayerMask & (1 << 1)) ? 1 : 0;
+			MaskVector[2] = (SceneProxy->BlendableLayerMask & (1 << 2)) ? 1 : 0;
+			SetShaderValue(RHICmdList, PixelShader->GetPixelShader(), BlendableLayerMaskParameter, MaskVector);
+		}
+	}
+protected:
+	FShaderParameter BlendableLayerMaskParameter;
+};
+
 FVertexFactoryShaderParameters* FLandscapeVertexFactoryMobile::ConstructShaderParameters(EShaderFrequency ShaderFrequency)
 {
 	switch( ShaderFrequency )
@@ -144,7 +193,7 @@ FVertexFactoryShaderParameters* FLandscapeVertexFactoryMobile::ConstructShaderPa
 	case SF_Vertex:
 		return new FLandscapeVertexFactoryMobileVertexShaderParameters();
 	case SF_Pixel:
-		return new FLandscapeVertexFactoryPixelShaderParameters();
+		return new FLandscapeVertexFactoryMobilePixelShaderParameters();
 	default:
 		return NULL;
 	}
@@ -179,7 +228,9 @@ FLandscapeComponentSceneProxyMobile::FLandscapeComponentSceneProxyMobile(ULandsc
 		InComponent->MobileNormalmapTexture = WeightmapTextures[0];
 	}
 #endif
-	NormalmapTexture = WeightmapTextures[0]; // Use Weightmap0 as Normal Texture
+	NormalmapTexture = WeightmapTextures[0]; // Use Weightmap0 as Normal/Weightmap Texture
+
+	BlendableLayerMask = InComponent->MobileBlendableLayerMask;
 }
 
 FLandscapeComponentSceneProxyMobile::~FLandscapeComponentSceneProxyMobile()
diff --git a/Engine/Source/Runtime/Engine/Public/Landscape/LandscapeRenderMobile.h b/Engine/Source/Runtime/Engine/Public/Landscape/LandscapeRenderMobile.h
--- a/Engine/Source/Runtime/Engine/Public/Landscape/LandscapeRenderMobile.h
+++ b/Engine/Source/Runtime/Engine/Public/Landscape/LandscapeRenderMobile.h
@@ -126,6 +126,8 @@ public:
 
 	virtual void CreateRenderThreadResources() override;
 
+	uint8 BlendableLayerMask;
+
 	friend class FLandscapeVertexBufferMobile;
 };
 
