Fixed json serialization for Visual Logs (it was broken for visual log's events) #ue4

diff --git a/Engine/Source/Runtime/Engine/Private/VisualLog.cpp b/Engine/Source/Runtime/Engine/Private/VisualLog.cpp
--- a/Engine/Source/Runtime/Engine/Private/VisualLog.cpp
+++ b/Engine/Source/Runtime/Engine/Private/VisualLog.cpp
@@ -302,13 +302,13 @@ TSharedPtr<FVisualLogEntry> JsonToVisualLogEntry(TSharedPtr<FJsonValue> FromJson
 				Event.UserFriendlyDesc = JsonSampleObject->GetStringField(VisualLogJson::TAG_DESCRIPTION);
 				Event.Counter = (int64)(JsonSampleObject->GetNumberField(VisualLogJson::TAG_COUNTER));
 
-				TArray< TSharedPtr<FJsonValue> > JsonEventTags = JsonEntryObject->GetArrayField(VisualLogJson::TAG_EVENTAGS);
+				TArray< TSharedPtr<FJsonValue> > JsonEventTags = JsonSampleObject->GetArrayField(VisualLogJson::TAG_EVENTAGS);
 				if (JsonEventTags.Num() > 0)
 				{
 					Event.EventTags.Reserve(JsonEventSamples.Num());
 					for (int32 TagIndex = 0; TagIndex < JsonEventTags.Num(); ++TagIndex)
 					{
-						TSharedPtr<FJsonObject> JsonTagObject = JsonEventTags[SampleIndex]->AsObject();
+						TSharedPtr<FJsonObject> JsonTagObject = JsonEventTags[TagIndex]->AsObject();
 						if (JsonTagObject.IsValid())
 						{
 							FString TagName = JsonTagObject->GetStringField(VisualLogJson::TAG_EVENTNAME);
diff --git a/Engine/Source/Runtime/Engine/Public/VisualLog.h b/Engine/Source/Runtime/Engine/Public/VisualLog.h
--- a/Engine/Source/Runtime/Engine/Public/VisualLog.h
+++ b/Engine/Source/Runtime/Engine/Public/VisualLog.h
@@ -58,7 +58,7 @@ namespace VisualLogJson
 	static const FString TAG_COUNTER = TEXT("Counter");
 	static const FString TAG_EVENTNAME = TEXT("EventName");
 	static const FString TAG_EVENTAGS = TEXT("EventTags");
-	static const FString TAG_EVENTSAMPLES = TEXT("EventTags");
+	static const FString TAG_EVENTSAMPLES = TEXT("EventSamples");
 
 	static const FString TAG_HISTOGRAMSAMPLES = TEXT("HistogramSamples");
 	static const FString TAG_HISTOGRAMSAMPLE = TEXT("Sample");
