Fix COMOffset incorrectly set on vehicles

diff --git a/Engine/Source/Runtime/Engine/Classes/Vehicles/WheeledVehicleMovementComponent.h b/Engine/Source/Runtime/Engine/Classes/Vehicles/WheeledVehicleMovementComponent.h
--- a/Engine/Source/Runtime/Engine/Classes/Vehicles/WheeledVehicleMovementComponent.h
+++ b/Engine/Source/Runtime/Engine/Classes/Vehicles/WheeledVehicleMovementComponent.h
@@ -198,13 +198,11 @@ class ENGINE_API UWheeledVehicleMovementComponent : public UPawnMovementComponen
 	UPROPERTY(transient)
 	float DebugDragMagnitude;
 
+	// Used for backwards compat to fixup incorrect COM of vehicles
+
 	/** When vehicle is created we want to compute some helper data like drag area, etc.... Derived classes should use this to properly compute things like engine RPM */
 	virtual void ComputeConstants();
 
-	/** Override center of mass offset, makes tweaking easier [uu] */
-	UPROPERTY(EditAnywhere, Category=VehicleSetup, AdvancedDisplay)
-	FVector COMOffset;
-
 	/** Scales the vehicle's inertia in each direction (forward, right, up) */
 	UPROPERTY(EditAnywhere, Category=VehicleSetup, AdvancedDisplay)
 	FVector InertiaTensorScale;
@@ -426,6 +424,8 @@ protected:
 	/** Clear all interpolated inputs to default values */
 	virtual void ClearInput();
 
+	/** Updates the COMOffset on the actual body instance */
+
 	/** Read current state for simulation */
 	void UpdateState(float DeltaTime);
 
@@ -466,8 +466,8 @@ protected:
 	/** Get the local position of the wheel at rest */
 	virtual FVector GetWheelRestingPosition(const FWheelSetup& WheelSetup);
 
-	/** Get the local COM offset */
-	virtual FVector GetCOMOffset();
+	/** Get the local COM */
+	virtual FVector GetLocalCOM() const;
 
 	/** Get the mesh this vehicle is tied to */
 	class USkinnedMeshComponent* GetMesh();
diff --git a/Engine/Source/Runtime/Engine/Private/SkeletalMeshComponentPhysics.cpp b/Engine/Source/Runtime/Engine/Private/SkeletalMeshComponentPhysics.cpp
--- a/Engine/Source/Runtime/Engine/Private/SkeletalMeshComponentPhysics.cpp
+++ b/Engine/Source/Runtime/Engine/Private/SkeletalMeshComponentPhysics.cpp
@@ -848,6 +848,8 @@ void USkeletalMeshComponent::InitArticulated(FPhysScene* PhysScene)
 			{
 				BodyInst->LockedAxisMode = BodyInstance.LockedAxisMode;
 				BodyInst->CustomLockedAxis = BodyInstance.CustomLockedAxis;
+
+				BodyInst->COMNudge = BodyInstance.COMNudge;
 			}
 			else
 			{
diff --git a/Engine/Source/Runtime/Engine/Private/Vehicles/WheeledVehicleMovementComponent.cpp b/Engine/Source/Runtime/Engine/Private/Vehicles/WheeledVehicleMovementComponent.cpp
--- a/Engine/Source/Runtime/Engine/Private/Vehicles/WheeledVehicleMovementComponent.cpp
+++ b/Engine/Source/Runtime/Engine/Private/Vehicles/WheeledVehicleMovementComponent.cpp
@@ -108,6 +108,7 @@ UWheeledVehicleMovementComponent::UWheeledVehicleMovementComponent(const FObject
 	SteeringInputRate.FallRate = 5.0f;
 
 
+
 #if WITH_VEHICLE
 	// tire load filtering
 	PxVehicleTireLoadFilterData PTireLoadFilterDef;
@@ -160,6 +161,7 @@ void UWheeledVehicleMovementComponent::CreateVehicle()
 			check(UpdatedComponent);
 			check(UpdatedComponent->GetBodyInstance()->GetPxRigidDynamic());
 
+
 			SetupVehicle();
 
 			if ( PVehicle != NULL )
@@ -253,9 +255,8 @@ void UWheeledVehicleMovementComponent::SetupVehicleMass()
 	PVehicleActor->setMassSpaceInertiaTensor( PInertiaTensor );
 	PVehicleActor->setMass( Mass );
 
-	// Adding shapes fudges with the local pose, fix it now
-	const PxVec3 PCOMOffset = U2PVector( GetCOMOffset() );
-	PVehicleActor->setCMassLocalPose( PxTransform( PCOMOffset, PxQuat::createIdentity() ) );
+	const PxVec3 PCOMOffset = U2PVector(GetLocalCOM());
+	PVehicleActor->setCMassLocalPose(PxTransform(PCOMOffset, PxQuat::createIdentity()));
 }
 
 void UWheeledVehicleMovementComponent::SetupWheels( PxVehicleWheelsSimData* PWheelsSimData )
@@ -281,8 +282,8 @@ void UWheeledVehicleMovementComponent::SetupWheels( PxVehicleWheelsSimData* PWhe
 	}
 
 	// Now that we have all the wheel offsets, calculate the sprung masses
-	PxVec3 PCOMOffset = U2PVector(GetCOMOffset());
-	PxVehicleComputeSprungMasses( WheelSetups.Num(), WheelOffsets, PCOMOffset, PVehicleActor->getMass(), 2, SprungMasses );
+	PxVec3 PLocalCOM = U2PVector(GetLocalCOM());
+	PxVehicleComputeSprungMasses(WheelSetups.Num(), WheelOffsets, PLocalCOM, PVehicleActor->getMass(), 2, SprungMasses);
 
 	for ( int32 WheelIdx = 0; WheelIdx < WheelSetups.Num(); ++WheelIdx )
 	{
@@ -320,8 +321,8 @@ void UWheeledVehicleMovementComponent::SetupWheels( PxVehicleWheelsSimData* PWhe
 		const PxVec3 PWheelOffset = WheelOffsets[WheelIdx];
 
 		PxVec3 PSuspTravelDirection = PxVec3( 0.0f, 0.0f, -1.0f );
-		PxVec3 PWheelCentreCMOffset = PWheelOffset - PCOMOffset;
-		PxVec3 PSuspForceAppCMOffset = PxVec3( PWheelOffset.x, PWheelOffset.y, Wheel->SuspensionForceOffset );
+		PxVec3 PWheelCentreCMOffset = PWheelOffset - PLocalCOM;
+		PxVec3 PSuspForceAppCMOffset = PxVec3(PWheelCentreCMOffset.x, PWheelCentreCMOffset.y, Wheel->SuspensionForceOffset);
 		PxVec3 PTireForceAppCMOffset = PSuspForceAppCMOffset;
 
 		// finalize sim data
@@ -501,9 +502,18 @@ FVector UWheeledVehicleMovementComponent::GetWheelRestingPosition( const FWheelS
 	return Offset;
 }
 
-FVector UWheeledVehicleMovementComponent::GetCOMOffset()
+FVector UWheeledVehicleMovementComponent::GetLocalCOM() const
 {
-	return COMOffset;
+	if (const FBodyInstance* BodyInst = UpdatedComponent->GetBodyInstance())
+	{
+		if (PxRigidDynamic* PVehicleActor = BodyInst->GetPxRigidDynamic())
+		{
+			PxTransform PCOMTransform = PVehicleActor->getCMassLocalPose();
+			return P2UVector(PCOMTransform.p);
+		}
+	}
+
+	return FVector::ZeroVector;
 }
 
 USkinnedMeshComponent* UWheeledVehicleMovementComponent::GetMesh()
@@ -1226,6 +1236,10 @@ void UWheeledVehicleMovementComponent::FixupSkeletalMesh()
 {
 	if (USkeletalMeshComponent * Mesh = Cast<USkeletalMeshComponent>(GetMesh()))
 	{
+
+					//in skeletal mesh case we must set the offset on the PrimitiveComponent's BodyInstance, which will later update the actual root body
+					//this is needed for UI
+
 		if (UPhysicsAsset * PhysicsAsset = Mesh->GetPhysicsAsset())
 		{
 			for (int32 WheelIdx = 0; WheelIdx < WheelSetups.Num(); ++WheelIdx)
