UParticleSystemComponent - constructor fix, FXSystem nulled.

diff --git a/Engine/Source/Runtime/Engine/Private/Particles/ParticleComponents.cpp b/Engine/Source/Runtime/Engine/Private/Particles/ParticleComponents.cpp
--- a/Engine/Source/Runtime/Engine/Private/Particles/ParticleComponents.cpp
+++ b/Engine/Source/Runtime/Engine/Private/Particles/ParticleComponents.cpp
@@ -2891,7 +2891,7 @@ bool UParticleSystem::HasGPUEmitter() const
 }
 
 UParticleSystemComponent::UParticleSystemComponent(const FObjectInitializer& ObjectInitializer)
-	: Super(ObjectInitializer)
+	: Super(ObjectInitializer), FXSystem(NULL), ReleaseResourcesFence(NULL)
 {
 	PrimaryComponentTick.bCanEverTick = true;
 	PrimaryComponentTick.TickGroup = TG_DuringPhysics;
