UE4 - Crash fix

diff --git a/Engine/Source/Runtime/Engine/Private/Particles/ParticleModules.cpp b/Engine/Source/Runtime/Engine/Private/Particles/ParticleModules.cpp
--- a/Engine/Source/Runtime/Engine/Private/Particles/ParticleModules.cpp
+++ b/Engine/Source/Runtime/Engine/Private/Particles/ParticleModules.cpp
@@ -280,7 +280,7 @@ void UParticleModule::AutoPopulateInstanceProperties(UParticleSystemComponent* P
 		UObject* Distribution = FRawDistribution::TryGetDistributionObjectFromRawDistributionProperty(*It, (uint8*)this);
 		if (Distribution)
 		{
-		    EParticleSysParamType ParamType	= PSPT_None;
+			EParticleSysParamType ParamType	= PSPT_None;
 			FName ParamName;
 
 			// only handle particle param types
@@ -288,7 +288,7 @@ void UParticleModule::AutoPopulateInstanceProperties(UParticleSystemComponent* P
 			UDistributionVectorParticleParameter* DistVectorParam = Cast<UDistributionVectorParticleParameter>(Distribution);
 			if (DistFloatParam != NULL)
 			{
-			    ParamType = PSPT_Scalar;
+				ParamType = PSPT_Scalar;
 				ParamName = DistFloatParam->ParameterName;
 			}
 			else 
@@ -2382,14 +2382,14 @@ void UParticleModuleAccelerationConstant::CompileModule( FParticleEmitterBuildIn
 void UParticleModuleAccelerationConstant::Spawn(FParticleEmitterInstance* Owner, int32 Offset, float SpawnTime, FBaseParticle* ParticleBase)
 {
 	SPAWN_INIT;
- 	UParticleLODLevel* LODLevel	= Owner->SpriteTemplate->GetCurrentLODLevel(Owner);
- 	check(LODLevel);
- 	if (bAlwaysInWorldSpace && LODLevel->RequiredModule->bUseLocalSpace)
- 	{
- 		FVector LocalAcceleration = Owner->Component->ComponentToWorld.InverseTransformVector(Acceleration);
+	UParticleLODLevel* LODLevel	= Owner->SpriteTemplate->GetCurrentLODLevel(Owner);
+	check(LODLevel);
+	if (bAlwaysInWorldSpace && LODLevel->RequiredModule->bUseLocalSpace)
+	{
+		FVector LocalAcceleration = Owner->Component->ComponentToWorld.InverseTransformVector(Acceleration);
 		Particle.Velocity		+= LocalAcceleration * SpawnTime;
 		Particle.BaseVelocity	+= LocalAcceleration * SpawnTime;
- 	}
+	}
 	else
 	{
 		FVector LocalAcceleration = Acceleration;
@@ -2704,20 +2704,20 @@ void UParticleModuleAcceleration::Spawn(FParticleEmitterInstance* Owner, int32 O
 {
 	SPAWN_INIT;
 	PARTICLE_ELEMENT(FVector, UsedAcceleration);
- 	UsedAcceleration = Acceleration.GetValue(Owner->EmitterTime, Owner->Component);
+	UsedAcceleration = Acceleration.GetValue(Owner->EmitterTime, Owner->Component);
 	if ((bApplyOwnerScale == true) && Owner && Owner->Component)
 	{
 		FVector Scale = Owner->Component->ComponentToWorld.GetScale3D();
 		UsedAcceleration *= Scale;
 	}
- 	UParticleLODLevel* LODLevel	= Owner->SpriteTemplate->GetCurrentLODLevel(Owner);
- 	check(LODLevel);
- 	if (bAlwaysInWorldSpace && LODLevel->RequiredModule->bUseLocalSpace)
- 	{
- 		FVector TempUsedAcceleration = Owner->Component->ComponentToWorld.InverseTransformVector(UsedAcceleration);
+	UParticleLODLevel* LODLevel	= Owner->SpriteTemplate->GetCurrentLODLevel(Owner);
+	check(LODLevel);
+	if (bAlwaysInWorldSpace && LODLevel->RequiredModule->bUseLocalSpace)
+	{
+		FVector TempUsedAcceleration = Owner->Component->ComponentToWorld.InverseTransformVector(UsedAcceleration);
 		Particle.Velocity		+= TempUsedAcceleration * SpawnTime;
 		Particle.BaseVelocity	+= TempUsedAcceleration * SpawnTime;
- 	}
+	}
 	else
 	{
 		if (LODLevel->RequiredModule->bUseLocalSpace)
@@ -3723,17 +3723,16 @@ void UParticleModuleAttractorLine::Update(FParticleEmitterInstance* Owner, int32
 		{
 			// Look up the Range and Strength at that position on the line
 			float AttractorRange = Range.GetValue(fRatio, Owner->Component);
-	        
 			FVector LineToPoint = AdjustedLocation - ProjectedParticle;
-    		float Distance = LineToPoint.Size();
+			float Distance = LineToPoint.Size();
 
 			if ((AttractorRange > 0) && (Distance <= AttractorRange))
 			{
 				// Adjust the strength based on the range ratio
 				float AttractorStrength = Strength.GetValue((AttractorRange - Distance) / AttractorRange, Owner->Component);
 				FVector Direction = LineToPoint^Line;
-    			// Adjust the VELOCITY of the particle based on the attractor... 
-        		Particle.Velocity += Direction * AttractorStrength * DeltaTime;
+				// Adjust the VELOCITY of the particle based on the attractor... 
+				Particle.Velocity += Direction * AttractorStrength * DeltaTime;
 			}
 		}
 	END_UPDATE_LOOP;
@@ -4000,7 +3999,7 @@ void UParticleModuleAttractorParticle::Update(FParticleEmitterInstance* Owner, i
 
 			// Adjust the VELOCITY of the particle based on the attractor... 
 			Dir.Normalize();
-    		Particle.Velocity	+= Dir * AttractorStrength * DeltaTime;
+			Particle.Velocity	+= Dir * AttractorStrength * DeltaTime;
 			Data.SourceVelocity	 = Source->Velocity;
 			if (bAffectBaseVelocity)
 			{
@@ -4156,7 +4155,7 @@ void UParticleModuleAttractorPoint::Update(FParticleEmitterInstance* Owner, int3
 
 			// Adjust the VELOCITY of the particle based on the attractor...
 			Dir = ClampVector(Dir,MinNormalizedDir,MaxNormalizedDir);
-    		Particle.Velocity	+= Dir * AttractorStrength * DeltaTime;
+			Particle.Velocity	+= Dir * AttractorStrength * DeltaTime;
 			if (bAffectBaseVelocity)
 			{
 				Particle.BaseVelocity	+= Dir * AttractorStrength * DeltaTime;
@@ -4584,7 +4583,7 @@ void UParticleModuleTypeDataGpu::Build( FParticleEmitterBuildInfo& EmitterBuildI
 
 	ResourceData.PivotOffset = EmitterBuildInfo.PivotOffset;
 
-    // Store color and scale when using particle parameters.
+	// Store color and scale when using particle parameters.
 	EmitterInfo.DynamicColor = EmitterBuildInfo.DynamicColor;
 	EmitterInfo.DynamicAlpha= EmitterBuildInfo.DynamicAlpha;
 	EmitterInfo.DynamicColorScale = EmitterBuildInfo.DynamicColorScale;
@@ -4608,13 +4607,13 @@ FParticleEmitterInstance* UParticleModuleTypeDataGpu::CreateInstance(UParticleEm
 {
 	check(InComponent);
 	UWorld* World = InComponent->GetWorld();
-	check( World );
+	check(World);
 	UE_LOG(LogParticles,Verbose,
 		TEXT("Create GPU Sprite Emitter @ %fs %s"), World->TimeSeconds,
 		InComponent->Template != NULL ? *InComponent->Template->GetName() : TEXT("NULL"));
 
 	FParticleEmitterInstance* Instance = NULL;
-	if (RHISupportsGPUParticles(World->Scene->GetFeatureLevel()))
+	if (World->Scene && RHISupportsGPUParticles(World->Scene->GetFeatureLevel()))
 	{
 		check( InComponent && InComponent->FXSystem );
 		Instance = InComponent->FXSystem->CreateGPUSpriteEmitterInstance( EmitterInfo );
