Adding support for feature level preview to skeletal meshes.
Fixes TTP 349032

diff --git a/Engine/Source/Runtime/Engine/Private/GPUSkinVertexFactory.cpp b/Engine/Source/Runtime/Engine/Private/GPUSkinVertexFactory.cpp
--- a/Engine/Source/Runtime/Engine/Private/GPUSkinVertexFactory.cpp
+++ b/Engine/Source/Runtime/Engine/Private/GPUSkinVertexFactory.cpp
@@ -112,11 +112,11 @@ TStatId FBoneBufferPool::GetStatId() const
 TConsoleVariableData<int32>* FGPUBaseSkinVertexFactory::ShaderDataType::MaxBonesVar = NULL;
 uint32 FGPUBaseSkinVertexFactory::ShaderDataType::MaxGPUSkinBones = 0;
 
-void FGPUBaseSkinVertexFactory::ShaderDataType::UpdateBoneData()
+void FGPUBaseSkinVertexFactory::ShaderDataType::UpdateBoneData(ERHIFeatureLevel::Type FeatureLevel)
 {
 	uint32 NumBones = BoneMatrices.Num();
 	check(NumBones <= MaxGPUSkinBones);
-	if (GRHIFeatureLevel_DEPRECATED >= ERHIFeatureLevel::SM4)
+	if (FeatureLevel >= ERHIFeatureLevel::SM4)
 	{
 		static FSharedPoolPolicyData PoolPolicy;
 		uint32 NumVectors = NumBones*3;
@@ -313,7 +313,7 @@ template <bool bExtraBoneInfluencesT>
 void TGPUSkinVertexFactory<bExtraBoneInfluencesT>::InitDynamicRHI()
 {
 	FVertexFactory::InitDynamicRHI();
-	ShaderData.UpdateBoneData();
+	ShaderData.UpdateBoneData(GetFeatureLevel());
 }
 
 template <bool bExtraBoneInfluencesT>
@@ -686,11 +686,11 @@ protected:
 /*-----------------------------------------------------------------------------
 	TGPUSkinAPEXClothVertexFactory::ClothShaderType
 -----------------------------------------------------------------------------*/
-void FGPUBaseSkinAPEXClothVertexFactory::ClothShaderType::UpdateClothSimulData(const TArray<FVector4>& InSimulPositions, const TArray<FVector4>& InSimulNormals)
+void FGPUBaseSkinAPEXClothVertexFactory::ClothShaderType::UpdateClothSimulData(const TArray<FVector4>& InSimulPositions, const TArray<FVector4>& InSimulNormals, ERHIFeatureLevel::Type FeatureLevel)
 {
 	uint32 NumSimulVerts = InSimulPositions.Num();
 
-	if (GRHIFeatureLevel_DEPRECATED >= ERHIFeatureLevel::SM4)
+	if (FeatureLevel >= ERHIFeatureLevel::SM4)
 	{
 		FMath::Min(NumSimulVerts, (uint32)MAX_APEXCLOTH_VERTICES_FOR_VB);
 
diff --git a/Engine/Source/Runtime/Engine/Private/SkeletalRender.cpp b/Engine/Source/Runtime/Engine/Private/SkeletalRender.cpp
--- a/Engine/Source/Runtime/Engine/Private/SkeletalRender.cpp
+++ b/Engine/Source/Runtime/Engine/Private/SkeletalRender.cpp
@@ -21,7 +21,7 @@ const float MaxVertexAnimBlendWeight = 5.0f;
 FSkeletalMeshObject
 -----------------------------------------------------------------------------*/
 
-FSkeletalMeshObject::FSkeletalMeshObject(USkinnedMeshComponent* InMeshComponent, FSkeletalMeshResource* InSkeletalMeshResource) 
+FSkeletalMeshObject::FSkeletalMeshObject(USkinnedMeshComponent* InMeshComponent, FSkeletalMeshResource* InSkeletalMeshResource, ERHIFeatureLevel::Type InFeatureLevel)
 :	MinDesiredLODLevel(0)
 ,	MaxDistanceFactor(0.f)
 ,	WorkingMinDesiredLODLevel(0)
@@ -38,6 +38,7 @@ FSkeletalMeshObject::FSkeletalMeshObject(USkinnedMeshComponent* InMeshComponent,
 ,	CustomSortAlternateIndexMode((ECustomSortAlternateIndexMode)InMeshComponent->CustomSortAlternateIndexMode)
 ,	bUsePerBoneMotionBlur(InMeshComponent->bPerBoneMotionBlur)
 ,	StatId(InMeshComponent->SkeletalMesh->GetStatID(true))
+,	FeatureLevel(InFeatureLevel)
 {
 	check(SkeletalMeshResource);
 
diff --git a/Engine/Source/Runtime/Engine/Private/SkeletalRenderCPUSkin.cpp b/Engine/Source/Runtime/Engine/Private/SkeletalRenderCPUSkin.cpp
--- a/Engine/Source/Runtime/Engine/Private/SkeletalRenderCPUSkin.cpp
+++ b/Engine/Source/Runtime/Engine/Private/SkeletalRenderCPUSkin.cpp
@@ -117,8 +117,8 @@ void FFinalSkinVertexBuffer::ReleaseDynamicRHI()
 -----------------------------------------------------------------------------*/
 
 
-FSkeletalMeshObjectCPUSkin::FSkeletalMeshObjectCPUSkin(USkinnedMeshComponent* InMeshComponent, FSkeletalMeshResource* InSkeletalMeshResource) 
-:	FSkeletalMeshObject(InMeshComponent,InSkeletalMeshResource)
+FSkeletalMeshObjectCPUSkin::FSkeletalMeshObjectCPUSkin(USkinnedMeshComponent* InMeshComponent, FSkeletalMeshResource* InSkeletalMeshResource, ERHIFeatureLevel::Type InFeatureLevel)
+	: FSkeletalMeshObject(InMeshComponent, InSkeletalMeshResource, InFeatureLevel)
 ,	DynamicData(NULL)
 ,	CachedVertexLOD(INDEX_NONE)
 ,	bRenderBoneWeight(false)
diff --git a/Engine/Source/Runtime/Engine/Private/SkeletalRenderCPUSkin.h b/Engine/Source/Runtime/Engine/Private/SkeletalRenderCPUSkin.h
--- a/Engine/Source/Runtime/Engine/Private/SkeletalRenderCPUSkin.h
+++ b/Engine/Source/Runtime/Engine/Private/SkeletalRenderCPUSkin.h
@@ -135,7 +135,7 @@ class FSkeletalMeshObjectCPUSkin : public FSkeletalMeshObject
 public:
 
 	/** @param	InSkeletalMeshComponent - skeletal mesh primitive we want to render */
-	FSkeletalMeshObjectCPUSkin(USkinnedMeshComponent* InMeshComponent, FSkeletalMeshResource* InSkeletalMeshResource);
+	FSkeletalMeshObjectCPUSkin(USkinnedMeshComponent* InMeshComponent, FSkeletalMeshResource* InSkeletalMeshResource, ERHIFeatureLevel::Type InFeatureLevel);
 	virtual ~FSkeletalMeshObjectCPUSkin();
 
 	// Begin FSkeletalMeshObject interface
diff --git a/Engine/Source/Runtime/Engine/Private/SkeletalRenderGPUSkin.cpp b/Engine/Source/Runtime/Engine/Private/SkeletalRenderGPUSkin.cpp
--- a/Engine/Source/Runtime/Engine/Private/SkeletalRenderGPUSkin.cpp
+++ b/Engine/Source/Runtime/Engine/Private/SkeletalRenderGPUSkin.cpp
@@ -85,8 +85,8 @@ void FMorphVertexBuffer::ReleaseDynamicRHI()
 FSkeletalMeshObjectGPUSkin
 -----------------------------------------------------------------------------*/
 
-FSkeletalMeshObjectGPUSkin::FSkeletalMeshObjectGPUSkin(USkinnedMeshComponent* InMeshComponent, FSkeletalMeshResource* InSkeletalMeshResource) 
-:	FSkeletalMeshObject(InMeshComponent, InSkeletalMeshResource)
+FSkeletalMeshObjectGPUSkin::FSkeletalMeshObjectGPUSkin(USkinnedMeshComponent* InMeshComponent, FSkeletalMeshResource* InSkeletalMeshResource, ERHIFeatureLevel::Type InFeatureLevel)
+	: FSkeletalMeshObject(InMeshComponent, InSkeletalMeshResource, InFeatureLevel)
 ,	DynamicData(NULL)
 ,	bMorphResourcesInitialized(false)
 {
@@ -113,7 +113,7 @@ void FSkeletalMeshObjectGPUSkin::InitResources()
 	{
 		FSkeletalMeshObjectLOD& SkelLOD = LODs[LODIndex];
 		const FSkelMeshObjectLODInfo& MeshLODInfo = LODInfo[LODIndex];
-		SkelLOD.InitResources(MeshLODInfo);
+		SkelLOD.InitResources(MeshLODInfo, FeatureLevel);
 	}
 }
 
@@ -141,7 +141,7 @@ void FSkeletalMeshObjectGPUSkin::InitMorphResources(bool bInUsePerBoneMotionBlur
 		FSkeletalMeshObjectLOD& SkelLOD = LODs[LODIndex];
 		// init any morph vertex buffers for each LOD
 		const FSkelMeshObjectLODInfo& MeshLODInfo = LODInfo[LODIndex];
-		SkelLOD.InitMorphResources(MeshLODInfo,bInUsePerBoneMotionBlur);
+		SkelLOD.InitMorphResources(MeshLODInfo,bInUsePerBoneMotionBlur, FeatureLevel);
 	}
 	bMorphResourcesInitialized = true;
 }
@@ -233,14 +233,14 @@ void FSkeletalMeshObjectGPUSkin::UpdateDynamicData_RenderThread(FRHICommandListI
 
 	if(DataPresent)
 	{
-		bool bGPUSkinCacheEnabled = GEnableGPUSkinCache && (GRHIFeatureLevel_DEPRECATED >= ERHIFeatureLevel::SM5);
+		bool bGPUSkinCacheEnabled = GEnableGPUSkinCache && (FeatureLevel >= ERHIFeatureLevel::SM5);
 		for( int32 ChunkIdx=0; ChunkIdx < Chunks.Num(); ChunkIdx++ )
 		{
 			const FSkelMeshChunk& Chunk = Chunks[ChunkIdx];
 
 			bool bClothFactory = (DynamicData->ClothSimulUpdateData.Num() > 0) && Chunk.HasApexClothData();
 
-			if (GRHIFeatureLevel_DEPRECATED < ERHIFeatureLevel::SM4)
+			if (FeatureLevel < ERHIFeatureLevel::SM4)
 			{
 				bClothFactory = false;
 			}
@@ -282,7 +282,7 @@ void FSkeletalMeshObjectGPUSkin::UpdateDynamicData_RenderThread(FRHICommandListI
 			}
 
 			// Create a uniform buffer from the bone transforms.
-			ShaderData.UpdateBoneData();
+			ShaderData.UpdateBoneData(FeatureLevel);
 
 			// Try to use the GPU skinning cache if possible
 			if (bGPUSkinCacheEnabled && ChunkIdx < MAX_GPUSKINCACHE_CHUNKS_PER_LOD && !bClothFactory && Chunk.MaxBoneInfluences > 0 && DynamicData->NumWeightedActiveVertexAnims <= 0)
@@ -303,7 +303,7 @@ void FSkeletalMeshObjectGPUSkin::UpdateDynamicData_RenderThread(FRHICommandListI
 				int16 ActorIdx = Chunk.CorrespondClothAssetIndex;
 				if( DynamicData->ClothSimulUpdateData.IsValidIndex(ActorIdx) )
 				{
-					ClothShaderData.UpdateClothSimulData( DynamicData->ClothSimulUpdateData[ActorIdx].ClothSimulPositions, DynamicData->ClothSimulUpdateData[ActorIdx].ClothSimulNormals);
+					ClothShaderData.UpdateClothSimulData( DynamicData->ClothSimulUpdateData[ActorIdx].ClothSimulPositions, DynamicData->ClothSimulUpdateData[ActorIdx].ClothSimulNormals, FeatureLevel);
 				}
 			}
 #endif // WITH_APEX_CLOTHING
@@ -613,10 +613,11 @@ TDynamicUpdateVertexFactoryData<VertexFactoryType>,VertexUpdateData,VertexUpdate
 template <class VertexFactoryTypeBase, class VertexFactoryType>
 static void CreateVertexFactory(TIndirectArray<VertexFactoryTypeBase>& VertexFactories,
 						 const FSkeletalMeshObjectGPUSkin::FVertexFactoryBuffers& InVertexBuffers,
-						 TArray<FBoneSkinning>& InBoneMatrices
+						 TArray<FBoneSkinning>& InBoneMatrices,
+						 ERHIFeatureLevel::Type FeatureLevel
 						 )
 {
-	auto* VertexFactory = new VertexFactoryType(InBoneMatrices);
+	auto* VertexFactory = new VertexFactoryType(InBoneMatrices, FeatureLevel);
 	VertexFactories.Add(VertexFactory);
 
 	// Setup the update data for enqueue
@@ -649,11 +650,12 @@ TDynamicUpdateVertexFactoryData<VertexFactoryType>,VertexUpdateData,VertexUpdate
 template <class VertexFactoryTypeBase, class VertexFactoryType>
 static void CreateVertexFactoryMorph(TIndirectArray<VertexFactoryTypeBase>& VertexFactories,
 						 const FSkeletalMeshObjectGPUSkin::FVertexFactoryBuffers& InVertexBuffers,
-						 TArray<FBoneSkinning>& InBoneMatrices
+						 TArray<FBoneSkinning>& InBoneMatrices,
+						 ERHIFeatureLevel::Type FeatureLevel
 						 )
 
 {
-	auto* VertexFactory = new VertexFactoryType(InBoneMatrices);
+	auto* VertexFactory = new VertexFactoryType(InBoneMatrices, FeatureLevel);
 	VertexFactories.Add(VertexFactory);
 						
 	// Setup the update data for enqueue
@@ -687,11 +689,12 @@ TDynamicUpdateVertexFactoryData<VertexFactoryType>,VertexUpdateData,VertexUpdate
 template <class VertexFactoryTypeBase, class VertexFactoryType>
 static void CreateVertexFactoryCloth(TArray<VertexFactoryTypeBase*>& VertexFactories,
 						 const FSkeletalMeshObjectGPUSkin::FVertexFactoryBuffers& InVertexBuffers,
-						 TArray<FBoneSkinning>& InBoneMatrices
+						 TArray<FBoneSkinning>& InBoneMatrices,
+						 ERHIFeatureLevel::Type FeatureLevel
 						 )
 
 {
-	VertexFactoryType* VertexFactory = new VertexFactoryType(InBoneMatrices);
+	VertexFactoryType* VertexFactory = new VertexFactoryType(InBoneMatrices, FeatureLevel);
 	VertexFactories.Add(VertexFactory);
 						
 	// Setup the update data for enqueue
@@ -747,7 +750,8 @@ void FSkeletalMeshObjectGPUSkin::FVertexFactoryData::InitPerChunkBoneMatrices(co
  */
 void FSkeletalMeshObjectGPUSkin::FVertexFactoryData::InitVertexFactories(
 	const FVertexFactoryBuffers& VertexBuffers, 
-	const TArray<FSkelMeshChunk>& Chunks)
+	const TArray<FSkelMeshChunk>& Chunks, 
+	ERHIFeatureLevel::Type FeatureLevel)
 {
 	// one array of matrices for each chunk (shared across vertex factory types)
 	InitPerChunkBoneMatrices(Chunks);
@@ -761,13 +765,13 @@ void FSkeletalMeshObjectGPUSkin::FVertexFactoryData::InitVertexFactories(
 		{
 			if (VertexBuffers.VertexBufferGPUSkin->HasExtraBoneInfluences())
 			{
-				CreateVertexFactory< FGPUBaseSkinVertexFactory, TGPUSkinVertexFactory<true> >(VertexFactories,VertexBuffers,PerChunkBoneMatricesArray[FactoryIdx]);
-				CreateVertexFactory< FGPUBaseSkinVertexFactory, FGPUSkinPassthroughVertexFactory >(PassthroughVertexFactories,VertexBuffers,PerChunkBoneMatricesArray[FactoryIdx]);
+				CreateVertexFactory< FGPUBaseSkinVertexFactory, TGPUSkinVertexFactory<true> >(VertexFactories, VertexBuffers, PerChunkBoneMatricesArray[FactoryIdx], FeatureLevel);
+				CreateVertexFactory< FGPUBaseSkinVertexFactory, FGPUSkinPassthroughVertexFactory >(PassthroughVertexFactories, VertexBuffers, PerChunkBoneMatricesArray[FactoryIdx], FeatureLevel);
 			}
 			else
 			{
-				CreateVertexFactory< FGPUBaseSkinVertexFactory, TGPUSkinVertexFactory<false> >(VertexFactories,VertexBuffers,PerChunkBoneMatricesArray[FactoryIdx]);
-				CreateVertexFactory< FGPUBaseSkinVertexFactory, FGPUSkinPassthroughVertexFactory >(PassthroughVertexFactories,VertexBuffers,PerChunkBoneMatricesArray[FactoryIdx]);
+				CreateVertexFactory< FGPUBaseSkinVertexFactory, TGPUSkinVertexFactory<false> >(VertexFactories, VertexBuffers, PerChunkBoneMatricesArray[FactoryIdx], FeatureLevel);
+				CreateVertexFactory< FGPUBaseSkinVertexFactory, FGPUSkinPassthroughVertexFactory >(PassthroughVertexFactories, VertexBuffers, PerChunkBoneMatricesArray[FactoryIdx], FeatureLevel);
 			}
 		}
 	}
@@ -793,7 +797,8 @@ void FSkeletalMeshObjectGPUSkin::FVertexFactoryData::ReleaseVertexFactories()
 void FSkeletalMeshObjectGPUSkin::FVertexFactoryData::InitMorphVertexFactories(
 	const FVertexFactoryBuffers& VertexBuffers, 
 	const TArray<FSkelMeshChunk>& Chunks,
-	bool bInUsePerBoneMotionBlur)
+	bool bInUsePerBoneMotionBlur,
+	ERHIFeatureLevel::Type InFeatureLevel)
 {
 	// one array of matrices for each chunk (shared across vertex factory types)
 	InitPerChunkBoneMatrices(Chunks);
@@ -803,11 +808,11 @@ void FSkeletalMeshObjectGPUSkin::FVertexFactoryData::InitMorphVertexFactories(
 	{
 		if (VertexBuffers.VertexBufferGPUSkin->HasExtraBoneInfluences())
 		{
-			CreateVertexFactoryMorph<FGPUBaseSkinVertexFactory, TGPUSkinMorphVertexFactory<true> >(MorphVertexFactories,VertexBuffers,PerChunkBoneMatricesArray[FactoryIdx]);
+			CreateVertexFactoryMorph<FGPUBaseSkinVertexFactory, TGPUSkinMorphVertexFactory<true> >(MorphVertexFactories,VertexBuffers,PerChunkBoneMatricesArray[FactoryIdx], InFeatureLevel);
 		}
 		else
 		{
-			CreateVertexFactoryMorph<FGPUBaseSkinVertexFactory, TGPUSkinMorphVertexFactory<false> >(MorphVertexFactories,VertexBuffers,PerChunkBoneMatricesArray[FactoryIdx]);
+			CreateVertexFactoryMorph<FGPUBaseSkinVertexFactory, TGPUSkinMorphVertexFactory<false> >(MorphVertexFactories, VertexBuffers, PerChunkBoneMatricesArray[FactoryIdx], InFeatureLevel);
 		}
 	}
 }
@@ -827,7 +832,8 @@ void FSkeletalMeshObjectGPUSkin::FVertexFactoryData::ReleaseMorphVertexFactories
 
 void FSkeletalMeshObjectGPUSkin::FVertexFactoryData::InitAPEXClothVertexFactories(
 	const FVertexFactoryBuffers& VertexBuffers, 
-	const TArray<FSkelMeshChunk>& Chunks)
+	const TArray<FSkelMeshChunk>& Chunks,
+	ERHIFeatureLevel::Type InFeatureLevel)
 {
 	// one array of matrices for each chunk (shared across vertex factory types)
 	InitPerChunkBoneMatrices(Chunks);
@@ -836,15 +842,15 @@ void FSkeletalMeshObjectGPUSkin::FVertexFactoryData::InitAPEXClothVertexFactorie
 	ClothVertexFactories.Empty(Chunks.Num());
 	for( int32 FactoryIdx=0; FactoryIdx < Chunks.Num(); FactoryIdx++ )
 	{
-		if (Chunks[FactoryIdx].HasApexClothData() && GRHIFeatureLevel_DEPRECATED >= ERHIFeatureLevel::SM4)
+		if (Chunks[FactoryIdx].HasApexClothData() && InFeatureLevel >= ERHIFeatureLevel::SM4)
 		{
 			if (VertexBuffers.VertexBufferGPUSkin->HasExtraBoneInfluences())
 			{
-				CreateVertexFactoryCloth<FGPUBaseSkinAPEXClothVertexFactory, TGPUSkinAPEXClothVertexFactory<true> >(ClothVertexFactories,VertexBuffers,PerChunkBoneMatricesArray[FactoryIdx]);
+				CreateVertexFactoryCloth<FGPUBaseSkinAPEXClothVertexFactory, TGPUSkinAPEXClothVertexFactory<true> >(ClothVertexFactories, VertexBuffers, PerChunkBoneMatricesArray[FactoryIdx], InFeatureLevel);
 			}
 			else
 			{
-				CreateVertexFactoryCloth<FGPUBaseSkinAPEXClothVertexFactory, TGPUSkinAPEXClothVertexFactory<false> >(ClothVertexFactories,VertexBuffers,PerChunkBoneMatricesArray[FactoryIdx]);
+				CreateVertexFactoryCloth<FGPUBaseSkinAPEXClothVertexFactory, TGPUSkinAPEXClothVertexFactory<false> >(ClothVertexFactories, VertexBuffers, PerChunkBoneMatricesArray[FactoryIdx], InFeatureLevel);
 			}
 		}
 		else
@@ -875,7 +881,7 @@ void FSkeletalMeshObjectGPUSkin::FVertexFactoryData::ReleaseAPEXClothVertexFacto
  * @param MeshLODInfo - information about the state of the bone influence swapping
  * @param Chunks - relevant chunk information (either original or from swapped influence)
  */
-void FSkeletalMeshObjectGPUSkin::FSkeletalMeshObjectLOD::InitResources(const FSkelMeshObjectLODInfo& MeshLODInfo)
+void FSkeletalMeshObjectGPUSkin::FSkeletalMeshObjectLOD::InitResources(const FSkelMeshObjectLODInfo& MeshLODInfo, ERHIFeatureLevel::Type FeatureLevel)
 {
 	check(SkelMeshResource);
 	check(SkelMeshResource->LODModels.IsValidIndex(LODIndex));
@@ -888,10 +894,10 @@ void FSkeletalMeshObjectGPUSkin::FSkeletalMeshObjectLOD::InitResources(const FSk
 	GetVertexBuffers(VertexBuffers,LODModel,MeshLODInfo);
 
 	// init gpu skin factories
-	GPUSkinVertexFactories.InitVertexFactories(VertexBuffers,LODModel.Chunks);
+	GPUSkinVertexFactories.InitVertexFactories(VertexBuffers,LODModel.Chunks, FeatureLevel);
 	if ( LODModel.HasApexClothData() )
 	{
-		GPUSkinVertexFactories.InitAPEXClothVertexFactories(VertexBuffers,LODModel.Chunks);
+		GPUSkinVertexFactories.InitAPEXClothVertexFactories(VertexBuffers,LODModel.Chunks, FeatureLevel);
 	}
 }
 
@@ -907,7 +913,7 @@ void FSkeletalMeshObjectGPUSkin::FSkeletalMeshObjectLOD::ReleaseResources()
 	GPUSkinVertexFactories.ReleaseAPEXClothVertexFactories();
 }
 
-void FSkeletalMeshObjectGPUSkin::FSkeletalMeshObjectLOD::InitMorphResources(const FSkelMeshObjectLODInfo& MeshLODInfo, bool bInUsePerBoneMotionBlur)
+void FSkeletalMeshObjectGPUSkin::FSkeletalMeshObjectLOD::InitMorphResources(const FSkelMeshObjectLODInfo& MeshLODInfo, bool bInUsePerBoneMotionBlur, ERHIFeatureLevel::Type FeatureLevel)
 {
 	check(SkelMeshResource);
 	check(SkelMeshResource->LODModels.IsValidIndex(LODIndex));
@@ -922,7 +928,7 @@ void FSkeletalMeshObjectGPUSkin::FSkeletalMeshObjectLOD::InitMorphResources(cons
 	FVertexFactoryBuffers VertexBuffers;
 	GetVertexBuffers(VertexBuffers,LODModel,MeshLODInfo);
 	// init morph skin factories
-	GPUSkinVertexFactories.InitMorphVertexFactories(VertexBuffers,LODModel.Chunks,bInUsePerBoneMotionBlur);
+	GPUSkinVertexFactories.InitMorphVertexFactories(VertexBuffers, LODModel.Chunks, bInUsePerBoneMotionBlur, FeatureLevel);
 }
 
 /** 
diff --git a/Engine/Source/Runtime/Engine/Private/SkeletalRenderGPUSkin.h b/Engine/Source/Runtime/Engine/Private/SkeletalRenderGPUSkin.h
--- a/Engine/Source/Runtime/Engine/Private/SkeletalRenderGPUSkin.h
+++ b/Engine/Source/Runtime/Engine/Private/SkeletalRenderGPUSkin.h
@@ -241,7 +241,7 @@ class FSkeletalMeshObjectGPUSkin : public FSkeletalMeshObject
 {
 public:
 	/** @param	InSkeletalMeshComponent - skeletal mesh primitive we want to render */
-	FSkeletalMeshObjectGPUSkin(USkinnedMeshComponent* InMeshComponent, FSkeletalMeshResource* InSkeletalMeshResource);
+	FSkeletalMeshObjectGPUSkin(USkinnedMeshComponent* InMeshComponent, FSkeletalMeshResource* InSkeletalMeshResource, ERHIFeatureLevel::Type InFeatureLevel);
 	virtual ~FSkeletalMeshObjectGPUSkin();
 
 	// Begin FSkeletalMeshObject interface
@@ -344,7 +344,7 @@ private:
 		 * @param VertexBuffers - available vertex buffers to reference in vertex factory streams
 		 * @param Chunks - relevant chunk information (either original or from swapped influence)
 		 */
-		void InitVertexFactories(const FVertexFactoryBuffers& VertexBuffers, const TArray<FSkelMeshChunk>& Chunks);
+		void InitVertexFactories(const FVertexFactoryBuffers& VertexBuffers, const TArray<FSkelMeshChunk>& Chunks, ERHIFeatureLevel::Type FeatureLevel);
 		/** 
 		 * Release default vertex factory resources for this LOD 
 		 */
@@ -355,7 +355,7 @@ private:
 		 * @param VertexBuffers - available vertex buffers to reference in vertex factory streams
 		 * @param Chunks - relevant chunk information (either original or from swapped influence)
 		 */
-		void InitMorphVertexFactories(const FVertexFactoryBuffers& VertexBuffers, const TArray<FSkelMeshChunk>& Chunks, bool bInUsePerBoneMotionBlur);
+		void InitMorphVertexFactories(const FVertexFactoryBuffers& VertexBuffers, const TArray<FSkelMeshChunk>& Chunks, bool bInUsePerBoneMotionBlur, ERHIFeatureLevel::Type InFeatureLevel);
 		/** 
 		 * Release morph vertex factory resources for this LOD 
 		 */
@@ -366,7 +366,7 @@ private:
 		 * @param VertexBuffers - available vertex buffers to reference in vertex factory streams
 		 * @param Chunks - relevant chunk information (either original or from swapped influence)
 		 */
-		void InitAPEXClothVertexFactories(const FVertexFactoryBuffers& VertexBuffers, const TArray<FSkelMeshChunk>& Chunks);
+		void InitAPEXClothVertexFactories(const FVertexFactoryBuffers& VertexBuffers, const TArray<FSkelMeshChunk>& Chunks, ERHIFeatureLevel::Type InFeatureLevel);
 		/** 
 		 * Release morph vertex factory resources for this LOD 
 		 */
@@ -421,7 +421,7 @@ private:
 		 * @param bUseLocalVertexFactory - use non-gpu skinned factory when rendering in ref pose
 		 * @param MeshLODInfo - information about the state of the bone influence swapping
 		 */
-		void InitResources(const FSkelMeshObjectLODInfo& MeshLODInfo);		
+		void InitResources(const FSkelMeshObjectLODInfo& MeshLODInfo, ERHIFeatureLevel::Type FeatureLevel);		
 
 		/** 
 		 * Release rendering resources for this LOD 
@@ -433,7 +433,7 @@ private:
 		 * @param MeshLODInfo - information about the state of the bone influence swapping
 		 * @param Chunks - relevant chunk information (either original or from swapped influence)
 		 */
-		void InitMorphResources(const FSkelMeshObjectLODInfo& MeshLODInfo, bool bInUsePerBoneMotionBlur);
+		void InitMorphResources(const FSkelMeshObjectLODInfo& MeshLODInfo, bool bInUsePerBoneMotionBlur, ERHIFeatureLevel::Type FeatureLevel);
 
 		/** 
 		 * Release rendering resources for the morph stream of this LOD
diff --git a/Engine/Source/Runtime/Engine/Private/SkinnedMeshComponent.cpp b/Engine/Source/Runtime/Engine/Private/SkinnedMeshComponent.cpp
--- a/Engine/Source/Runtime/Engine/Private/SkinnedMeshComponent.cpp
+++ b/Engine/Source/Runtime/Engine/Private/SkinnedMeshComponent.cpp
@@ -62,11 +62,11 @@ void USkinnedMeshComponent::CreateRenderState_Concurrent()
 			const bool bIsCPUSkinned = SkelMeshResource->RequiresCPUSkinning(SceneFeatureLevel) || (GIsEditor && ShouldCPUSkin());
 			if(bIsCPUSkinned)
 			{
-				MeshObject = ::new FSkeletalMeshObjectCPUSkin(this,SkelMeshResource);
+				MeshObject = ::new FSkeletalMeshObjectCPUSkin(this, SkelMeshResource, SceneFeatureLevel);
 			}
 			else
 			{
-				MeshObject = ::new FSkeletalMeshObjectGPUSkin(this,SkelMeshResource);
+				MeshObject = ::new FSkeletalMeshObjectGPUSkin(this, SkelMeshResource, SceneFeatureLevel);
 			}
 
 			//Allow the editor a chance to manipulate it before its added to the scene
@@ -748,6 +748,7 @@ void USkinnedMeshComponent::SetSkeletalMesh(USkeletalMesh* InSkelMesh)
 		FComponentReregisterContext ReregisterContext(this);
 		check(MeshObject == NULL);
 		SkeletalMesh = InSkelMesh;
+		InvalidateCachedBounds();
 	}
 	
 	// TODO: (LH) find better way to call this 
diff --git a/Engine/Source/Runtime/Engine/Public/GPUSkinVertexFactory.h b/Engine/Source/Runtime/Engine/Public/GPUSkinVertexFactory.h
--- a/Engine/Source/Runtime/Engine/Public/GPUSkinVertexFactory.h
+++ b/Engine/Source/Runtime/Engine/Public/GPUSkinVertexFactory.h
@@ -355,7 +355,7 @@ public:
 			return true;
 		}
 
-		void UpdateBoneData();
+		void UpdateBoneData(ERHIFeatureLevel::Type FeatureLevel);
 
 		void ReleaseBoneData()
 		{
@@ -399,8 +399,8 @@ public:
 	 *
 	 * @param	InBoneMatrices	Reference to shared bone matrices array.
 	 */
-	FGPUBaseSkinVertexFactory(TArray<FBoneSkinning>& InBoneMatrices)
-	:	ShaderData( InBoneMatrices )
+	FGPUBaseSkinVertexFactory(TArray<FBoneSkinning>& InBoneMatrices, ERHIFeatureLevel::Type InFeatureLevel)
+	:	FVertexFactory(InFeatureLevel), ShaderData( InBoneMatrices )
 	{}
 
 	/** accessor */
@@ -490,8 +490,8 @@ public:
 	 *
 	 * @param	InBoneMatrices	Reference to shared bone matrices array.
 	 */
-	TGPUSkinVertexFactory(TArray<FBoneSkinning>& InBoneMatrices)
-	:	FGPUBaseSkinVertexFactory( InBoneMatrices )
+	TGPUSkinVertexFactory(TArray<FBoneSkinning>& InBoneMatrices, ERHIFeatureLevel::Type InFeatureLevel)
+		: FGPUBaseSkinVertexFactory(InBoneMatrices, InFeatureLevel)
 	{}
 
 	virtual bool UsesExtraBoneInfluences() const
@@ -548,8 +548,8 @@ public:
 	 *
 	 * @param	InBoneMatrices	Reference to shared bone matrices array.
 	 */
-	FGPUSkinPassthroughVertexFactory(TArray<FBoneSkinning>& InBoneMatrices)
-		: TGPUSkinVertexFactory<false>(InBoneMatrices)
+	FGPUSkinPassthroughVertexFactory(TArray<FBoneSkinning>& InBoneMatrices, ERHIFeatureLevel::Type InFeatureLevel)
+		: TGPUSkinVertexFactory<false>(InBoneMatrices, InFeatureLevel)
 	{}
 
 	static void ModifyCompilationEnvironment(EShaderPlatform Platform, const class FMaterial* Material, FShaderCompilerEnvironment& OutEnvironment);
@@ -581,8 +581,8 @@ public:
 	 *
 	 * @param	InBoneMatrices	Reference to shared bone matrices array.
 	 */
-	TGPUSkinMorphVertexFactory(TArray<FBoneSkinning>& InBoneMatrices)
-	: TGPUSkinVertexFactory<bExtraBoneInfluencesT>(InBoneMatrices)
+	TGPUSkinMorphVertexFactory(TArray<FBoneSkinning>& InBoneMatrices, ERHIFeatureLevel::Type InFeatureLevel)
+	: TGPUSkinVertexFactory<bExtraBoneInfluencesT>(InBoneMatrices, InFeatureLevel)
 	{}
 
 	static void ModifyCompilationEnvironment(EShaderPlatform Platform, const class FMaterial* Material, FShaderCompilerEnvironment& OutEnvironment);
@@ -647,7 +647,7 @@ public:
 		}
 
 
-		void UpdateClothSimulData(const TArray<FVector4>& InSimulPositions, const TArray<FVector4>& InSimulNormals);
+		void UpdateClothSimulData(const TArray<FVector4>& InSimulPositions, const TArray<FVector4>& InSimulNormals, ERHIFeatureLevel::Type FeatureLevel);
 
 		void ReleaseClothSimulData()
 		{
@@ -735,8 +735,8 @@ public:
 	 *
 	 * @param	InBoneMatrices	Reference to shared bone matrices array.
 	 */
-	TGPUSkinAPEXClothVertexFactory(TArray<FBoneSkinning>& InBoneMatrices)
-	: TGPUSkinVertexFactory<bExtraBoneInfluencesT>(InBoneMatrices)
+	TGPUSkinAPEXClothVertexFactory(TArray<FBoneSkinning>& InBoneMatrices, ERHIFeatureLevel::Type InFeatureLevel)
+		: TGPUSkinVertexFactory<bExtraBoneInfluencesT>(InBoneMatrices, InFeatureLevel)
 	{}
 
 	static void ModifyCompilationEnvironment(EShaderPlatform Platform, const class FMaterial* Material, FShaderCompilerEnvironment& OutEnvironment);
diff --git a/Engine/Source/Runtime/Engine/Public/SkeletalRenderPublic.h b/Engine/Source/Runtime/Engine/Public/SkeletalRenderPublic.h
--- a/Engine/Source/Runtime/Engine/Public/SkeletalRenderPublic.h
+++ b/Engine/Source/Runtime/Engine/Public/SkeletalRenderPublic.h
@@ -19,7 +19,7 @@ public:
 		MAX_GPUSKINCACHE_CHUNKS_PER_LOD = 16,
 	};
 
-	FSkeletalMeshObject(USkinnedMeshComponent* InMeshComponent, FSkeletalMeshResource* InSkeletalMeshResource);
+	FSkeletalMeshObject(USkinnedMeshComponent* InMeshComponent, FSkeletalMeshResource* InSkeletalMeshResource, ERHIFeatureLevel::Type FeatureLevel);
 	virtual ~FSkeletalMeshObject() {}
 
 	/** 
@@ -232,4 +232,7 @@ protected:
 
 	/** Used for dynamic stats */
 	TStatId StatId;
+
+	/** Feature level to render for. */
+	ERHIFeatureLevel::Type FeatureLevel;
 };
\ No newline at end of file
