FXSystem pointer is nulled in constructors.

diff --git a/Engine/Source/Runtime/Engine/Private/World.cpp b/Engine/Source/Runtime/Engine/Private/World.cpp
--- a/Engine/Source/Runtime/Engine/Private/World.cpp
+++ b/Engine/Source/Runtime/Engine/Private/World.cpp
@@ -84,6 +84,7 @@ UWorld::UWorld( const FObjectInitializer& ObjectInitializer )
 ,	TickTaskLevel(FTickTaskManagerInterface::Get().AllocateTickTaskLevel())
 ,   bIsBuilt(false)
 ,	NextTravelType(TRAVEL_Relative)
+,	FXSystem(NULL)
 {
 	TimerManager = new FTimerManager();
 #if WITH_EDITOR
@@ -99,6 +100,7 @@ UWorld::UWorld( const FObjectInitializer& ObjectInitializer,const FURL& InURL )
 ,	TickTaskLevel(FTickTaskManagerInterface::Get().AllocateTickTaskLevel())
 ,   bIsBuilt(false)
 ,	NextTravelType(TRAVEL_Relative)
+,	FXSystem(NULL)
 {
 	SetFlags( RF_Transactional );
 	TimerManager = new FTimerManager();
diff --git a/Engine/Source/Runtime/Renderer/Private/RendererScene.cpp b/Engine/Source/Runtime/Renderer/Private/RendererScene.cpp
--- a/Engine/Source/Runtime/Renderer/Private/RendererScene.cpp
+++ b/Engine/Source/Runtime/Renderer/Private/RendererScene.cpp
@@ -2062,6 +2062,7 @@ class FNULLSceneInterface : public FSceneInterface
 public:
 	FNULLSceneInterface( UWorld* InWorld )
 		:	World( InWorld )
+		,	FXSystem( NULL )
 	{
 		World->Scene = this;
 	}
