Replace RigidDynamic with RigidBody in cases where applicable.

TTP#348792

diff --git a/Engine/Source/Runtime/Engine/Classes/PhysicsEngine/BodyInstance.h b/Engine/Source/Runtime/Engine/Classes/PhysicsEngine/BodyInstance.h
--- a/Engine/Source/Runtime/Engine/Classes/PhysicsEngine/BodyInstance.h
+++ b/Engine/Source/Runtime/Engine/Classes/PhysicsEngine/BodyInstance.h
@@ -18,6 +18,7 @@ namespace physx
 {
 	class PxRigidActor;
 	class PxAggregate;
+	class PxRigidBody;
 	class PxRigidDynamic;
 	class PxGeometry;
 	class PxShape;
@@ -395,6 +396,8 @@ public:
 	physx::PxRigidActor* GetPxRigidActor(int32 SceneType = -1) const;
 	/** Return the PxRigidDynamic if it exists in one of the scenes (NULL otherwise).  Currently a PxRigidDynamic can exist in only one of the two scenes. */
 	physx::PxRigidDynamic* GetPxRigidDynamic() const;
+	/** Return the PxRigidBody if it exists in one of the scenes (NULL otherwise).  Currently a PxRigidBody can exist in only one of the two scenes. */
+	physx::PxRigidBody* GetPxRigidBody() const;
 
 	/** 
 	 *	Utility to get all the shapes from a FBodyInstance 
diff --git a/Engine/Source/Runtime/Engine/Private/DestructibleComponent.cpp b/Engine/Source/Runtime/Engine/Private/DestructibleComponent.cpp
--- a/Engine/Source/Runtime/Engine/Private/DestructibleComponent.cpp
+++ b/Engine/Source/Runtime/Engine/Private/DestructibleComponent.cpp
@@ -513,7 +513,7 @@ void UDestructibleComponent::AddRadialImpulse(FVector Origin, float Radius, floa
 					LockedScene->lockWrite();
 					LockedScene->lockRead();
 				}
-				AddRadialImpulseToPxRigidDynamic(*PActor, Origin, Radius, Strength, Falloff, bVelChange);
+				AddRadialImpulseToPxRigidBody(*PActor, Origin, Radius, Strength, Falloff, bVelChange);
 			}
 		}
 
@@ -556,7 +556,7 @@ void UDestructibleComponent::AddRadialForce(FVector Origin, float Radius, float
 					LockedScene->lockRead();
 				}
 
-				AddRadialForceToPxRigidDynamic(*PActor, Origin, Radius, Strength, Falloff);
+				AddRadialForceToPxRigidBody(*PActor, Origin, Radius, Strength, Falloff);
 			}
 
 			if (LockedScene)
diff --git a/Engine/Source/Runtime/Engine/Private/PhysicsEngine/BodyInstance.cpp b/Engine/Source/Runtime/Engine/Private/PhysicsEngine/BodyInstance.cpp
--- a/Engine/Source/Runtime/Engine/Private/PhysicsEngine/BodyInstance.cpp
+++ b/Engine/Source/Runtime/Engine/Private/PhysicsEngine/BodyInstance.cpp
@@ -1379,7 +1379,7 @@ void FBodyInstance::InitBody(UBodySetup* Setup, const FTransform& Transform, UPr
 		CreateDOFLock();
 
 		// wakeUp and putToSleep will issue warnings on kinematic actors
-		if (IsRigidDynamicNonKinematic(PNewDynamic))
+		if (IsRigidBodyNonKinematic(PNewDynamic))
 		{
 		    // Sleep/wake up as appropriate
 		    if (bShouldStartAwake)
@@ -2300,7 +2300,7 @@ void FBodyInstance::SetBodyTransform(const FTransform& NewTransform, bool bTelep
 		// Do nothing if already in correct place
 		{
 			SCOPED_SCENE_READ_LOCK(RigidActor->getScene());
-			if (PRigidDynamic && !IsRigidDynamicNonKinematic(PRigidDynamic) && !bTeleport)
+			if (PRigidDynamic && !IsRigidBodyNonKinematic(PRigidDynamic) && !bTeleport)
 			{
 				const PxScene* PScene = PRigidDynamic->getScene();
 				FPhysScene* PhysScene = FPhysxUserData::Get<FPhysScene>(PScene->userData);
@@ -2336,7 +2336,7 @@ void FBodyInstance::SetBodyTransform(const FTransform& NewTransform, bool bTelep
 		if (PRigidDynamic)
 		{
 			// If kinematic and not teleporting, set kinematic target
-			if (!IsRigidDynamicNonKinematic(PRigidDynamic) && !bTeleport)
+			if (!IsRigidBodyNonKinematic(PRigidDynamic) && !bTeleport)
 			{
 				const PxScene* PScene = PRigidDynamic->getScene();
 				FPhysScene* PhysScene = FPhysxUserData::Get<FPhysScene>(PScene->userData);
@@ -2386,10 +2386,10 @@ FVector FBodyInstance::GetUnrealWorldVelocity() const
 	FVector LinVel(EForceInit::ForceInitToZero);
 
 #if WITH_PHYSX
-	if (PxRigidDynamic* PRigidDynamic = GetPxRigidDynamic())
+	if (PxRigidBody* PRigidBody = GetPxRigidBody())
 	{
-		SCOPED_SCENE_READ_LOCK(PRigidDynamic->getScene());
-		LinVel = P2UVector(PRigidDynamic->getLinearVelocity());
+		SCOPED_SCENE_READ_LOCK(PRigidBody->getScene());
+		LinVel = P2UVector(PRigidBody->getLinearVelocity());
 	}
 #endif // WITH_PHYSX
 
@@ -2409,11 +2409,11 @@ FVector FBodyInstance::GetUnrealWorldAngularVelocity() const
 	FVector AngVel(EForceInit::ForceInitToZero);
 
 #if WITH_PHYSX
-	PxRigidDynamic* PRigidDynamic = GetPxRigidDynamic();
-	if(PRigidDynamic != NULL)
+	PxRigidBody* PRigidBody = GetPxRigidBody();
+	if (PRigidBody != NULL)
 	{
-		SCOPED_SCENE_READ_LOCK(PRigidDynamic->getScene());
-		AngVel = FMath::RadiansToDegrees( P2UVector(PRigidDynamic->getAngularVelocity()) );
+		SCOPED_SCENE_READ_LOCK(PRigidBody->getScene());
+		AngVel = FMath::RadiansToDegrees(P2UVector(PRigidBody->getAngularVelocity()));
 	}
 #endif // WITH_PHYSX
 
@@ -2432,11 +2432,11 @@ FVector FBodyInstance::GetUnrealWorldVelocityAtPoint(const FVector& Point) const
 	FVector LinVel(EForceInit::ForceInitToZero);
 
 #if WITH_PHYSX
-	if (PxRigidDynamic* PRigidDynamic = GetPxRigidDynamic())
+	if (PxRigidBody* PRigidBody = GetPxRigidBody())
 	{
-		SCOPED_SCENE_READ_LOCK(PRigidDynamic->getScene());
+		SCOPED_SCENE_READ_LOCK(PRigidBody->getScene());
 		PxVec3 PPoint = U2PVector(Point);
-		LinVel = P2UVector( PxRigidBodyExt::getVelocityAtPos(*PRigidDynamic, PPoint) );
+		LinVel = P2UVector(PxRigidBodyExt::getVelocityAtPos(*PRigidBody, PPoint));
 	}
 #endif // WITH_PHYSX
 
@@ -2454,11 +2454,11 @@ FVector FBodyInstance::GetUnrealWorldVelocityAtPoint(const FVector& Point) const
 FVector FBodyInstance::GetCOMPosition() const
 {
 #if WITH_PHYSX
-	if (PxRigidDynamic* PRigidDynamic = GetPxRigidDynamic())
+	if (PxRigidBody* PRigidBody = GetPxRigidBody())
 	{
-		SCOPED_SCENE_READ_LOCK(PRigidDynamic->getScene());
-		PxTransform PLocalCOM = PRigidDynamic->getCMassLocalPose();
-		PxVec3 PWorldCOM = PRigidDynamic->getGlobalPose().transform(PLocalCOM.p);
+		SCOPED_SCENE_READ_LOCK(PRigidBody->getScene());
+		PxTransform PLocalCOM = PRigidBody->getCMassLocalPose();
+		PxVec3 PWorldCOM = PRigidBody->getGlobalPose().transform(PLocalCOM.p);
 		return P2UVector(PWorldCOM);
 	}
 	else
@@ -2482,10 +2482,10 @@ float FBodyInstance::GetBodyMass() const
 	float Retval = 0.f;
 
 #if WITH_PHYSX
-	if (PxRigidDynamic* PRigidDynamic = GetPxRigidDynamic())
+	if (PxRigidBody* PRigidBody = GetPxRigidBody())
 	{
-		SCOPED_SCENE_READ_LOCK(PRigidDynamic->getScene());
-		Retval = PRigidDynamic->getMass();
+		SCOPED_SCENE_READ_LOCK(PRigidBody->getScene());
+		Retval = PRigidBody->getMass();
 	}
 #endif // WITH_PHYSX
 
@@ -2504,12 +2504,12 @@ FBox FBodyInstance::GetBodyBounds() const
 	FBox Bounds;
 
 #if WITH_PHYSX
-	PxRigidDynamic* PRigidDynamic = GetPxRigidDynamic();
-	if(PRigidDynamic)
+	PxRigidBody* PRigidBody = GetPxRigidBody();
+	if (PRigidBody)
 	{
-		SCOPED_SCENE_READ_LOCK(PRigidDynamic->getScene());
+		SCOPED_SCENE_READ_LOCK(PRigidBody->getScene());
 		
-		PxBounds3 PBounds = PRigidDynamic->getWorldBounds();
+		PxBounds3 PBounds = PRigidBody->getWorldBounds();
 
 		Bounds.Min = P2UVector(PBounds.minimum);
 		Bounds.Max = P2UVector(PBounds.maximum);
@@ -2595,6 +2595,24 @@ PxRigidDynamic* FBodyInstance::GetPxRigidDynamic() const
 		return NULL;
 	}
 }
+
+PxRigidBody* FBodyInstance::GetPxRigidBody() const
+{
+	// The logic below works because dynamic actors are non-NULL in only one scene.
+	// If this assumption changes, the logic needs to change.
+	if (RigidActorSync != NULL)
+	{
+		return RigidActorSync->isRigidBody();
+	}
+	else if (RigidActorAsync != NULL)
+	{
+		return RigidActorAsync->isRigidBody();
+	}
+	else
+	{
+		return NULL;
+	}
+}
 #endif // WITH_PHYSX
 
 
@@ -2716,11 +2734,11 @@ TArray<UPhysicalMaterial*> FBodyInstance::GetComplexPhysicalMaterials() const
 
 #if WITH_PHYSX
 /** Util for finding the number of 'collision sim' shapes on this PxRigidActor */
-int32 GetNumSimShapes(PxRigidDynamic* PRigidDynamic)
+int32 GetNumSimShapes(PxRigidBody* PRigidBody)
 {
 	TArray<PxShape*> PShapes;
-	PShapes.AddZeroed(PRigidDynamic->getNbShapes());
-	int32 NumShapes = PRigidDynamic->getShapes(PShapes.GetData(), PShapes.Num());
+	PShapes.AddZeroed(PRigidBody->getNbShapes());
+	int32 NumShapes = PRigidBody->getShapes(PShapes.GetData(), PShapes.Num());
 
 	int32 NumSimShapes = 0;
 
@@ -2759,16 +2777,15 @@ void FBodyInstance::UpdateMassProperties()
 	check(PhysMat);
 
 #if WITH_PHYSX
-	PxRigidDynamic* PRigidDynamic = GetPxRigidDynamic();
-	if ((PRigidDynamic != NULL) && (GetNumSimShapes(PRigidDynamic) > 0))
+	PxRigidBody* PRigidBody = GetPxRigidBody();
+	if ((PRigidBody != NULL) && (GetNumSimShapes(PRigidBody) > 0))
 	{
-
 		// physical material - nothing can weigh less than hydrogen (0.09 kg/m^3)
 		float DensityKGPerCubicUU = FMath::Max(KgPerM3ToKgPerCm3(0.09f), gPerCm3ToKgPerCm3(PhysMat->Density));
-		PxRigidBodyExt::updateMassAndInertia(*PRigidDynamic, DensityKGPerCubicUU);
+		PxRigidBodyExt::updateMassAndInertia(*PRigidBody, DensityKGPerCubicUU);
 
 		//grab OldMass so we can apply new mass while maintaining inertia tensor
-		float OldMass = PRigidDynamic->getMass();
+		float OldMass = PRigidBody->getMass();
 		float NewMass = 0.f;
 
 		if (bOverrideMass == false)
@@ -2787,18 +2804,18 @@ void FBodyInstance::UpdateMassProperties()
 		check (NewMass > 0.f);
 
 		float MassRatio = NewMass/OldMass;
-		PxVec3 InertiaTensor = PRigidDynamic->getMassSpaceInertiaTensor();
+		PxVec3 InertiaTensor = PRigidBody->getMassSpaceInertiaTensor();
 
-		PRigidDynamic->setMassSpaceInertiaTensor(InertiaTensor * MassRatio);
-		PRigidDynamic->setMass(NewMass);
+		PRigidBody->setMassSpaceInertiaTensor(InertiaTensor * MassRatio);
+		PRigidBody->setMass(NewMass);
 
 		// Apply the COMNudge
 		if (!COMNudge.IsZero())
 		{
 			PxVec3 PCOMNudge = U2PVector(COMNudge);
-			PxTransform PCOMTransform = PRigidDynamic->getCMassLocalPose();
+			PxTransform PCOMTransform = PRigidBody->getCMassLocalPose();
 			PCOMTransform.p += PCOMNudge;
-			PRigidDynamic->setCMassLocalPose(PCOMTransform);
+			PRigidBody->setCMassLocalPose(PCOMTransform);
 		}
 	}
 #endif
@@ -2881,7 +2898,7 @@ void FBodyInstance::WakeInstance()
 {
 #if WITH_PHYSX
 	PxRigidDynamic* PRigidDynamic = GetPxRigidDynamic();
-	if (IsRigidDynamicNonKinematic(PRigidDynamic))
+	if (IsRigidBodyNonKinematic(PRigidDynamic))
 	{
 		SCOPED_SCENE_WRITE_LOCK(PRigidDynamic->getScene());
 		PRigidDynamic->wakeUp();
@@ -2900,7 +2917,7 @@ void FBodyInstance::PutInstanceToSleep()
 {
 #if WITH_PHYSX
 	PxRigidDynamic* PRigidDynamic = GetPxRigidDynamic();
-	if (IsRigidDynamicNonKinematic(PRigidDynamic))
+	if (IsRigidBodyNonKinematic(PRigidDynamic))
 	{
 		PRigidDynamic->putToSleep();
 	}
@@ -2917,17 +2934,17 @@ void FBodyInstance::PutInstanceToSleep()
 void FBodyInstance::SetLinearVelocity(const FVector& NewVel, bool bAddToCurrent)
 {
 #if WITH_PHYSX
-	if (PxRigidDynamic* PRigidDynamic = GetPxRigidDynamic())
+	if (PxRigidBody* PRigidBody = GetPxRigidBody())
 	{
 		PxVec3 PNewVel = U2PVector(NewVel);
 
 		if (bAddToCurrent)
 		{
-			const PxVec3 POldVel = PRigidDynamic->getLinearVelocity();
+			const PxVec3 POldVel = PRigidBody->getLinearVelocity();
 			PNewVel += POldVel;
 		}
 
-		PRigidDynamic->setLinearVelocity(PNewVel);
+		PRigidBody->setLinearVelocity(PNewVel);
 	}
 #endif // WITH_PHYSX
 
@@ -2951,17 +2968,17 @@ void FBodyInstance::SetLinearVelocity(const FVector& NewVel, bool bAddToCurrent)
 void FBodyInstance::SetAngularVelocity(const FVector& NewAngVel, bool bAddToCurrent)
 {
 #if WITH_PHYSX
-	if (PxRigidDynamic* PRigidDynamic = GetPxRigidDynamic())
+	if (PxRigidBody* PRigidBody = GetPxRigidBody())
 	{
 		PxVec3 PNewAngVel = U2PVector( FMath::DegreesToRadians(NewAngVel) );
 
 		if (bAddToCurrent)
 		{
-			const PxVec3 POldAngVel = PRigidDynamic->getAngularVelocity();
+			const PxVec3 POldAngVel = PRigidBody->getAngularVelocity();
 			PNewAngVel += POldAngVel;
 		}
 
-		PRigidDynamic->setAngularVelocity(PNewAngVel);
+		PRigidBody->setAngularVelocity(PNewAngVel);
 	}
 #endif //WITH_PHYSX
 
@@ -3019,9 +3036,9 @@ void FBodyInstance::SetMaxDepenetrationVelocity(float MaxVelocity)
 	//0 from both instance and project settings means 
 	float UseMaxVelocity = MaxVelocity == 0.f ? PX_MAX_F32 : MaxVelocity;
 
-	if (PxRigidDynamic* PRigidDynamic = GetPxRigidDynamic())
+	if (PxRigidBody* PRigidBody = GetPxRigidBody())
 	{
-		PRigidDynamic->setMaxDepenetrationVelocity(UseMaxVelocity);
+		PRigidBody->setMaxDepenetrationVelocity(UseMaxVelocity);
 	}
 #endif
 
@@ -3031,10 +3048,10 @@ void FBodyInstance::SetMaxDepenetrationVelocity(float MaxVelocity)
 void FBodyInstance::AddForce(const FVector& Force, bool bAllowSubstepping)
 {
 #if WITH_PHYSX
-	PxRigidDynamic* PRigidDynamic = GetPxRigidDynamic();
-	if (IsRigidDynamicNonKinematic(PRigidDynamic))
+	PxRigidBody* PRigidBody = GetPxRigidBody();
+	if (IsRigidBodyNonKinematic(PRigidBody))
 	{
-		const PxScene* PScene = PRigidDynamic->getScene();
+		const PxScene* PScene = PRigidBody->getScene();
 		FPhysScene* PhysScene = FPhysxUserData::Get<FPhysScene>(PScene->userData);
 		PhysScene->AddForce(this, Force, bAllowSubstepping);
 		
@@ -3056,10 +3073,10 @@ void FBodyInstance::AddForce(const FVector& Force, bool bAllowSubstepping)
 void FBodyInstance::AddForceAtPosition(const FVector& Force, const FVector& Position, bool bAllowSubstepping)
 {
 #if WITH_PHYSX
-	PxRigidDynamic* PRigidDynamic = GetPxRigidDynamic();
-	if (IsRigidDynamicNonKinematic(PRigidDynamic))
+	PxRigidBody* PRigidBody = GetPxRigidBody();
+	if (IsRigidBodyNonKinematic(PRigidBody))
 	{
-		const PxScene* PScene = PRigidDynamic->getScene();
+		const PxScene* PScene = PRigidBody->getScene();
 		FPhysScene* PhysScene = FPhysxUserData::Get<FPhysScene>(PScene->userData);
 		PhysScene->AddForceAtPosition(this, Force, Position, bAllowSubstepping);
 	}
@@ -3081,10 +3098,10 @@ void FBodyInstance::AddForceAtPosition(const FVector& Force, const FVector& Posi
 void FBodyInstance::AddTorque(const FVector& Torque, bool bAllowSubstepping)
 {
 #if WITH_PHYSX
-	PxRigidDynamic* PRigidDynamic = GetPxRigidDynamic();
-	if (IsRigidDynamicNonKinematic(PRigidDynamic))
+	PxRigidBody* PRigidBody = GetPxRigidBody();
+	if (IsRigidBodyNonKinematic(PRigidBody))
 	{
-		const PxScene* PScene = PRigidDynamic->getScene();
+		const PxScene* PScene = PRigidBody->getScene();
 		FPhysScene* PhysScene = FPhysxUserData::Get<FPhysScene>(PScene->userData);
 		PhysScene->AddTorque(this, Torque, bAllowSubstepping);
 	}
@@ -3105,11 +3122,11 @@ void FBodyInstance::AddTorque(const FVector& Torque, bool bAllowSubstepping)
 void FBodyInstance::AddImpulse(const FVector& Impulse, bool bVelChange)
 {
 #if WITH_PHYSX
-	PxRigidDynamic* PRigidDynamic = GetPxRigidDynamic();
-	if (IsRigidDynamicNonKinematic(PRigidDynamic))
+	PxRigidBody* PRigidBody = GetPxRigidBody();
+	if (IsRigidBodyNonKinematic(PRigidBody))
 	{
 		PxForceMode::Enum Mode = bVelChange ? PxForceMode::eVELOCITY_CHANGE : PxForceMode::eIMPULSE;
-		PRigidDynamic->addForce(U2PVector(Impulse), Mode, true);
+		PRigidBody->addForce(U2PVector(Impulse), Mode, true);
 	}
 #endif // WITH_PHYSX
 
@@ -3128,11 +3145,11 @@ void FBodyInstance::AddImpulse(const FVector& Impulse, bool bVelChange)
 void FBodyInstance::AddImpulseAtPosition(const FVector& Impulse, const FVector& Position)
 {
 #if WITH_PHYSX
-	PxRigidDynamic* PRigidDynamic = GetPxRigidDynamic();
-	if (IsRigidDynamicNonKinematic(PRigidDynamic))
+	PxRigidBody* PRigidBody = GetPxRigidBody();
+	if (IsRigidBodyNonKinematic(PRigidBody))
 	{
 		PxForceMode::Enum Mode = PxForceMode::eIMPULSE; // does not support eVELOCITY_CHANGE
-		PxRigidBodyExt::addForceAtPos(*PRigidDynamic, U2PVector(Impulse), U2PVector(Position), Mode, true);
+		PxRigidBodyExt::addForceAtPos(*PRigidBody, U2PVector(Impulse), U2PVector(Position), Mode, true);
 	}
 #endif // WITH_PHYSX
 
@@ -3162,7 +3179,7 @@ void FBodyInstance::SetEnableGravity(bool bInGravityEnabled)
 		bEnableGravity = bInGravityEnabled;
 
 #if WITH_PHYSX
-		if (PxRigidDynamic* PActor = GetPxRigidDynamic())
+		if (PxRigidBody* PActor = GetPxRigidBody())
 		{
 			PActor->setActorFlag( PxActorFlag::eDISABLE_GRAVITY, !bEnableGravity );
 		}
@@ -3228,10 +3245,10 @@ void AddRadialImpulseToBox2DBodyInstance(class b2Body* BodyInstancePtr, const FV
 void FBodyInstance::AddRadialImpulseToBody(const FVector& Origin, float Radius, float Strength, uint8 Falloff, bool bVelChange)
 {
 #if WITH_PHYSX
-	PxRigidDynamic* PRigidDynamic = GetPxRigidDynamic();
-	if(IsRigidDynamicNonKinematic(PRigidDynamic))
+	PxRigidBody* PRigidBody = GetPxRigidBody();
+	if (IsRigidBodyNonKinematic(PRigidBody))
 	{
-		AddRadialImpulseToPxRigidDynamic(*PRigidDynamic, Origin, Radius, Strength, Falloff, bVelChange);
+		AddRadialImpulseToPxRigidBody(*PRigidBody, Origin, Radius, Strength, Falloff, bVelChange);
 	}
 #endif // WITH_PHYSX
 
@@ -3246,10 +3263,10 @@ void FBodyInstance::AddRadialImpulseToBody(const FVector& Origin, float Radius,
 void FBodyInstance::AddRadialForceToBody(const FVector& Origin, float Radius, float Strength, uint8 Falloff)
 {
 #if WITH_PHYSX
-	PxRigidDynamic* PRigidDynamic = GetPxRigidDynamic();
-	if(IsRigidDynamicNonKinematic(PRigidDynamic))
+	PxRigidBody* PRigidBody = GetPxRigidBody();
+	if (IsRigidBodyNonKinematic(PRigidBody))
 	{
-		AddRadialForceToPxRigidDynamic(*PRigidDynamic, Origin, Radius, Strength, Falloff);
+		AddRadialForceToPxRigidBody(*PRigidBody, Origin, Radius, Strength, Falloff);
 	}
 #endif // WITH_PHYSX
 
diff --git a/Engine/Source/Runtime/Engine/Private/PhysicsEngine/ConstraintInstance.cpp b/Engine/Source/Runtime/Engine/Private/PhysicsEngine/ConstraintInstance.cpp
--- a/Engine/Source/Runtime/Engine/Private/PhysicsEngine/ConstraintInstance.cpp
+++ b/Engine/Source/Runtime/Engine/Private/PhysicsEngine/ConstraintInstance.cpp
@@ -240,7 +240,7 @@ void FConstraintInstance::InitConstraint(USceneComponent* Owner, FBodyInstance*
 
 	// Do not create joint unless you have two actors
 	// Do not create joint unless one of the actors is dynamic
-	if( (!PActor1 || !PActor1->isRigidDynamic()) && (!PActor2 || !PActor2->isRigidDynamic() ))
+	if ((!PActor1 || !PActor1->isRigidBody()) && (!PActor2 || !PActor2->isRigidBody()))
 	{
 		return;
 	}
@@ -248,13 +248,13 @@ void FConstraintInstance::InitConstraint(USceneComponent* Owner, FBodyInstance*
 	// Need to worry about the case where one is static and one is dynamic, and make sure the static scene is used which matches the dynamic scene
 	if(PActor1 != NULL && PActor2 != NULL)
 	{
-		if(PActor1->isRigidStatic() && PActor2->isRigidDynamic())
+		if (PActor1->isRigidStatic() && PActor2->isRigidBody())
 		{
 			const uint32 SceneType = Body2->RigidActorSync != NULL ? PST_Sync : PST_Async;
 			PActor1 = Body1->GetPxRigidActor(SceneType);
 		}
 		else
-		if(PActor2->isRigidStatic() && PActor1->isRigidDynamic())
+		if (PActor2->isRigidStatic() && PActor1->isRigidBody())
 		{
 			const uint32 SceneType = Body1->RigidActorSync != NULL ? PST_Sync : PST_Async;
 			PActor2 = Body2->GetPxRigidActor(SceneType);
@@ -266,15 +266,15 @@ void FConstraintInstance::InitConstraint(USceneComponent* Owner, FBodyInstance*
 		float TotalMass = 0;
 		int NumDynamic = 0;
 
-		if (PActor1 && PActor1->isRigidDynamic())
+		if (PActor1 && PActor1->isRigidBody())
 		{
-			TotalMass += PActor1->isRigidDynamic()->getMass();
+			TotalMass += PActor1->isRigidBody()->getMass();
 			++NumDynamic;
 		}
 
-		if (PActor2 && PActor2->isRigidDynamic())
+		if (PActor2 && PActor2->isRigidBody())
 		{
-			TotalMass += PActor2->isRigidDynamic()->getMass();
+			TotalMass += PActor2->isRigidBody()->getMass();
 			++NumDynamic;
 		}
 
@@ -445,12 +445,12 @@ void FConstraintInstance::InitConstraint(USceneComponent* Owner, FBodyInstance*
 	// creation of joints wakes up rigid bodies, so we put them to sleep again if both were initially asleep
 	if(bActor1Asleep && bActor2Asleep)
 	{
-		if ( PActor1 != NULL && IsRigidDynamicNonKinematic(PActor1->isRigidDynamic()))
+		if ( PActor1 != NULL && IsRigidBodyNonKinematic(PActor1->isRigidDynamic()))
 		{
 			PActor1->isRigidDynamic()->putToSleep();
 		}
 
-		if ( PActor2 != NULL && IsRigidDynamicNonKinematic(PActor2->isRigidDynamic()))
+		if (PActor2 != NULL && IsRigidBodyNonKinematic(PActor2->isRigidDynamic()))
 		{
 			PActor2->isRigidDynamic()->putToSleep();
 		}
diff --git a/Engine/Source/Runtime/Engine/Private/PhysicsEngine/PhysScene.cpp b/Engine/Source/Runtime/Engine/Private/PhysicsEngine/PhysScene.cpp
--- a/Engine/Source/Runtime/Engine/Private/PhysicsEngine/PhysScene.cpp
+++ b/Engine/Source/Runtime/Engine/Private/PhysicsEngine/PhysScene.cpp
@@ -35,7 +35,7 @@ FORCEINLINE EPhysicsSceneType SceneType(const FBodyInstance* BodyInstance)
 {
 #if WITH_PHYSX
 	//This is a helper function for dynamic actors - static actors are in both scenes
-	check(BodyInstance->GetPxRigidDynamic());
+	check(BodyInstance->GetPxRigidBody());
 	return UPhysicsSettings::Get()->bEnableAsyncScene && BodyInstance->UseAsyncScene() ? PST_Async : PST_Sync;
 #endif
 
@@ -230,7 +230,7 @@ void FPhysScene::AddForce(FBodyInstance* BodyInstance, const FVector& Force, boo
 {
 #if WITH_PHYSX
 
-	if (PxRigidDynamic * PRigidDynamic = BodyInstance->GetPxRigidDynamic())
+	if (PxRigidBody * PRigidBody = BodyInstance->GetPxRigidBody())
 	{
 #if WITH_SUBSTEPPING
 		uint32 BodySceneType = SceneType(BodyInstance);
@@ -242,8 +242,8 @@ void FPhysScene::AddForce(FBodyInstance* BodyInstance, const FVector& Force, boo
 		else
 #endif
 		{
-			SCOPED_SCENE_WRITE_LOCK(PRigidDynamic->getScene());
-			PRigidDynamic->addForce(U2PVector(Force), PxForceMode::eFORCE, true);
+			SCOPED_SCENE_WRITE_LOCK(PRigidBody->getScene());
+			PRigidBody->addForce(U2PVector(Force), PxForceMode::eFORCE, true);
 		}
 	}
 #endif
@@ -253,7 +253,7 @@ void FPhysScene::AddForceAtPosition(FBodyInstance* BodyInstance, const FVector&
 {
 #if WITH_PHYSX
 
-	if (PxRigidDynamic * PRigidDynamic = BodyInstance->GetPxRigidDynamic())
+	if (PxRigidBody * PRigidBody = BodyInstance->GetPxRigidBody())
 	{
 #if WITH_SUBSTEPPING
 		uint32 BodySceneType = SceneType(BodyInstance);
@@ -265,8 +265,8 @@ void FPhysScene::AddForceAtPosition(FBodyInstance* BodyInstance, const FVector&
 		else
 #endif
 		{
-			SCOPED_SCENE_WRITE_LOCK(PRigidDynamic->getScene());
-			PxRigidBodyExt::addForceAtPos(*PRigidDynamic, U2PVector(Force), U2PVector(Position), PxForceMode::eFORCE, true);
+			SCOPED_SCENE_WRITE_LOCK(PRigidBody->getScene());
+			PxRigidBodyExt::addForceAtPos(*PRigidBody, U2PVector(Force), U2PVector(Position), PxForceMode::eFORCE, true);
 		}
 	}
 #endif
@@ -276,7 +276,7 @@ void FPhysScene::AddTorque(FBodyInstance* BodyInstance, const FVector& Torque, b
 {
 #if WITH_PHYSX
 
-	if (PxRigidDynamic * PRigidDynamic = BodyInstance->GetPxRigidDynamic())
+	if (PxRigidBody * PRigidBody = BodyInstance->GetPxRigidBody())
 	{
 #if WITH_SUBSTEPPING
 		uint32 BodySceneType = SceneType(BodyInstance);
@@ -288,8 +288,8 @@ void FPhysScene::AddTorque(FBodyInstance* BodyInstance, const FVector& Torque, b
 		else
 #endif
 		{
-			SCOPED_SCENE_WRITE_LOCK(PRigidDynamic->getScene());
-			PRigidDynamic->addTorque(U2PVector(Torque), PxForceMode::eFORCE, true);
+			SCOPED_SCENE_WRITE_LOCK(PRigidBody->getScene());
+			PRigidBody->addTorque(U2PVector(Torque), PxForceMode::eFORCE, true);
 		}
 	}
 #endif
@@ -316,7 +316,7 @@ void FPhysScene::RemoveBodyFromActiveTransforms(class PxActor * PActor, uint32 S
 void FPhysScene::TermBody(FBodyInstance* BodyInstance)
 {
 #if WITH_SUBSTEPPING
-	if (PxRigidDynamic * PRigidDynamic = BodyInstance->GetPxRigidDynamic())
+	if (PxRigidBody * PRigidBody = BodyInstance->GetPxRigidBody())
 	{
 		FPhysSubstepTask * PhysSubStepper = PhysSubSteppers[SceneType(BodyInstance)];
 		PhysSubStepper->RemoveBodyInstance(BodyInstance);
diff --git a/Engine/Source/Runtime/Engine/Private/PhysicsEngine/PhysSubstepTasks.cpp b/Engine/Source/Runtime/Engine/Private/PhysicsEngine/PhysSubstepTasks.cpp
--- a/Engine/Source/Runtime/Engine/Private/PhysicsEngine/PhysSubstepTasks.cpp
+++ b/Engine/Source/Runtime/Engine/Private/PhysicsEngine/PhysSubstepTasks.cpp
@@ -56,7 +56,7 @@ void FPhysSubstepTask::SetKinematicTarget(FBodyInstance* Body, const FTransform&
 	PxRigidDynamic * PRigidDynamic = Body->GetPxRigidDynamic();
 	SCOPED_SCENE_READ_LOCK(PRigidDynamic->getScene());
 	//We only interpolate kinematic actors
-	if (!IsRigidDynamicNonKinematic(PRigidDynamic))
+	if (!IsRigidBodyNonKinematic(PRigidDynamic))
 	{
 		FKinematicTarget KinmaticTarget(Body, TM);
 		FPhysTarget & TargetState = PhysTargetBuffers[External].FindOrAdd(Body);
@@ -88,10 +88,10 @@ void FPhysSubstepTask::AddForce(FBodyInstance* Body, const FVector& Force)
 {
 #if WITH_PHYSX
 	check(Body);
-	PxRigidDynamic * PRigidDynamic = Body->GetPxRigidDynamic();
-	SCOPED_SCENE_READ_LOCK(PRigidDynamic->getScene());
+	PxRigidBody* PRigidBody = Body->GetPxRigidBody();
+	SCOPED_SCENE_READ_LOCK(PRigidBody->getScene());
 	//We should only apply forces on non kinematic actors
-	if (IsRigidDynamicNonKinematic(PRigidDynamic))
+	if (IsRigidBodyNonKinematic(PRigidBody))
 	{
 
 		FForceTarget ForceTarget;
@@ -109,10 +109,10 @@ void FPhysSubstepTask::AddForceAtPosition(FBodyInstance* Body, const FVector& Fo
 #if WITH_PHYSX
 	check(Body);
 
-	PxRigidDynamic * PRigidDynamic = Body->GetPxRigidDynamic();
-	SCOPED_SCENE_READ_LOCK(PRigidDynamic->getScene());
+	PxRigidBody* PRigidBody = Body->GetPxRigidBody();
+	SCOPED_SCENE_READ_LOCK(PRigidBody->getScene());
 	//We should only apply forces on non kinematic actors
-	if (IsRigidDynamicNonKinematic(PRigidDynamic))
+	if (IsRigidBodyNonKinematic(PRigidBody))
 	{
 		FForceTarget ForceTarget;
 		ForceTarget.bPosition = true;
@@ -129,10 +129,10 @@ void FPhysSubstepTask::AddTorque(FBodyInstance* Body, const FVector& Torque)
 #if WITH_PHYSX
 	check(Body);
 
-	PxRigidDynamic * PRigidDynamic = Body->GetPxRigidDynamic();
-	SCOPED_SCENE_READ_LOCK(PRigidDynamic->getScene());
+	PxRigidBody* PRigidBody = Body->GetPxRigidBody();
+	SCOPED_SCENE_READ_LOCK(PRigidBody->getScene());
 	//We should only apply torque on non kinematic actors
-	if (IsRigidDynamicNonKinematic(PRigidDynamic))
+	if (IsRigidBodyNonKinematic(PRigidBody))
 	{
 		FTorqueTarget TorqueTarget;
 		TorqueTarget.Torque = Torque;
@@ -148,7 +148,7 @@ void FPhysSubstepTask::ApplyForces(const FPhysTarget& PhysTarget, FBodyInstance*
 {
 #if WITH_PHYSX
 	/** Apply Forces */
-	PxRigidDynamic * PRigidDynamic = BodyInstance->GetPxRigidDynamic();
+	PxRigidBody* PRigidBody = BodyInstance->GetPxRigidBody();
 
 	for (int32 i = 0; i < PhysTarget.Forces.Num(); ++i)
 	{
@@ -156,11 +156,11 @@ void FPhysSubstepTask::ApplyForces(const FPhysTarget& PhysTarget, FBodyInstance*
 
 		if (ForceTarget.bPosition)
 		{
-			PxRigidBodyExt::addForceAtPos(*PRigidDynamic, U2PVector(ForceTarget.Force), U2PVector(ForceTarget.Position), PxForceMode::eFORCE, true);
+			PxRigidBodyExt::addForceAtPos(*PRigidBody, U2PVector(ForceTarget.Force), U2PVector(ForceTarget.Position), PxForceMode::eFORCE, true);
 		}
 		else
 		{
-			PRigidDynamic->addForce(U2PVector(ForceTarget.Force), PxForceMode::eFORCE, true);
+			PRigidBody->addForce(U2PVector(ForceTarget.Force), PxForceMode::eFORCE, true);
 		}
 	}
 #endif
@@ -171,12 +171,12 @@ void FPhysSubstepTask::ApplyTorques(const FPhysTarget& PhysTarget, FBodyInstance
 {
 #if WITH_PHYSX
 	/** Apply Torques */
-	PxRigidDynamic * PRigidDynamic = BodyInstance->GetPxRigidDynamic();
+	PxRigidBody* PRigidBody = BodyInstance->GetPxRigidBody();
 
 	for (int32 i = 0; i < PhysTarget.Torques.Num(); ++i)
 	{
 		const FTorqueTarget& TorqueTarget = PhysTarget.Torques[i];
-		PRigidDynamic->addTorque(U2PVector(TorqueTarget.Torque), PxForceMode::eFORCE, true);
+		PRigidBody->addTorque(U2PVector(TorqueTarget.Torque), PxForceMode::eFORCE, true);
 	}
 #endif
 }
@@ -192,7 +192,7 @@ void FPhysSubstepTask::InterpolateKinematicActor(const FPhysTarget& PhysTarget,
 	if (PhysTarget.bKinematicTarget)
 	{
 		//It's possible that the actor is no longer kinematic and is now simulating. In that case do nothing
-		if (!IsRigidDynamicNonKinematic(PRigidDynamic))
+		if (!IsRigidBodyNonKinematic(PRigidDynamic))
 		{
 			const FKinematicTarget& KinematicTarget = PhysTarget.KinematicTarget;
 			const FTransform& TargetTM = KinematicTarget.TargetTM;
@@ -227,16 +227,16 @@ void FPhysSubstepTask::SubstepInterpolation(float InAlpha)
 	for (PhysTargetMap::TIterator Itr = Targets.CreateIterator(); Itr; ++Itr)
 	{
 		FPhysTarget & PhysTarget = Itr.Value();
-		FBodyInstance* BodyInstance = Itr.Key();
-		PxRigidDynamic * PRigidDynamic = BodyInstance->GetPxRigidDynamic();
+		FBodyInstance * BodyInstance = Itr.Key();
+		PxRigidBody* PRigidBody = BodyInstance->GetPxRigidBody();
 
-		if (PRigidDynamic == NULL)
+		if (PRigidBody == NULL)
 		{
 			continue;
 		}
 
 		//We should only be iterating over actors that belong to this scene
-		check(PRigidDynamic->getScene() == PScene);
+		check(PRigidBody->getScene() == PScene);
 
 		ApplyForces(PhysTarget, BodyInstance);
 		ApplyTorques(PhysTarget, BodyInstance);
diff --git a/Engine/Source/Runtime/Engine/Private/PhysicsEngine/PhysXSupport.cpp b/Engine/Source/Runtime/Engine/Private/PhysicsEngine/PhysXSupport.cpp
--- a/Engine/Source/Runtime/Engine/Private/PhysicsEngine/PhysXSupport.cpp
+++ b/Engine/Source/Runtime/Engine/Private/PhysicsEngine/PhysXSupport.cpp
@@ -211,13 +211,13 @@ PxScene* GetPhysXSceneFromIndex(int32 InSceneIndex)
 #endif	// #if WITH_APEX
 
 
-void AddRadialImpulseToPxRigidDynamic(PxRigidDynamic& PRigidDynamic, const FVector& Origin, float Radius, float Strength, uint8 Falloff, bool bVelChange)
+void AddRadialImpulseToPxRigidBody(PxRigidBody& PRigidBody, const FVector& Origin, float Radius, float Strength, uint8 Falloff, bool bVelChange)
 {
 #if WITH_PHYSX
-	if (!(PRigidDynamic.getRigidDynamicFlags() & PxRigidDynamicFlag::eKINEMATIC))
+	if (!(PRigidBody.getRigidBodyFlags() & PxRigidBodyFlag::eKINEMATIC))
 	{
-		float Mass = PRigidDynamic.getMass();
-		PxTransform PCOMTransform = PRigidDynamic.getGlobalPose().transform(PRigidDynamic.getCMassLocalPose());
+		float Mass = PRigidBody.getMass();
+		PxTransform PCOMTransform = PRigidBody.getGlobalPose().transform(PRigidBody.getCMassLocalPose());
 		PxVec3 PCOMPos = PCOMTransform.p; // center of mass in world space
 		PxVec3 POrigin = U2PVector(Origin); // origin of radial impulse, in world space
 		PxVec3 PDelta = PCOMPos - POrigin; // vector from origin to COM
@@ -242,18 +242,18 @@ void AddRadialImpulseToPxRigidDynamic(PxRigidDynamic& PRigidDynamic, const FVect
 		PxVec3 PImpulse = PDelta * ImpulseMag;
 
 		PxForceMode::Enum Mode = bVelChange ? PxForceMode::eVELOCITY_CHANGE : PxForceMode::eIMPULSE;
-		PRigidDynamic.addForce(PImpulse, Mode);
+		PRigidBody.addForce(PImpulse, Mode);
 	}
 #endif // WITH_PHYSX
 }
 
-void AddRadialForceToPxRigidDynamic(PxRigidDynamic& PRigidDynamic, const FVector& Origin, float Radius, float Strength, uint8 Falloff)
+void AddRadialForceToPxRigidBody(PxRigidBody& PRigidBody, const FVector& Origin, float Radius, float Strength, uint8 Falloff)
 {
 #if WITH_PHYSX
-	if (!(PRigidDynamic.getRigidDynamicFlags() & PxRigidDynamicFlag::eKINEMATIC))
+	if (!(PRigidBody.getRigidBodyFlags() & PxRigidBodyFlag::eKINEMATIC))
 	{
-		float Mass = PRigidDynamic.getMass();
-		PxTransform PCOMTransform = PRigidDynamic.getGlobalPose().transform(PRigidDynamic.getCMassLocalPose());
+		float Mass = PRigidBody.getMass();
+		PxTransform PCOMTransform = PRigidBody.getGlobalPose().transform(PRigidBody.getCMassLocalPose());
 		PxVec3 PCOMPos = PCOMTransform.p; // center of mass in world space
 		PxVec3 POrigin = U2PVector(Origin); // origin of radial impulse, in world space
 		PxVec3 PDelta = PCOMPos - POrigin; // vector from
@@ -277,16 +277,16 @@ void AddRadialForceToPxRigidDynamic(PxRigidDynamic& PRigidDynamic, const FVector
 
 		// Apply force
 		PxVec3 PImpulse = PDelta * ForceMag;
-		PRigidDynamic.addForce(PImpulse, PxForceMode::eFORCE);
+		PRigidBody.addForce(PImpulse, PxForceMode::eFORCE);
 	}
 #endif // WITH_PHYSX
 }
 
-bool IsRigidDynamicNonKinematic(PxRigidDynamic* PRigidDynamic)
+bool IsRigidBodyNonKinematic(PxRigidBody* PRigidBody)
 {
-	if(PRigidDynamic != NULL)
+	if (PRigidBody != NULL)
 	{
-		return !(PRigidDynamic->getRigidDynamicFlags() & PxRigidDynamicFlag::eKINEMATIC);
+		return !(PRigidBody->getRigidBodyFlags() & PxRigidBodyFlag::eKINEMATIC);
 	}
 
 	return false;
diff --git a/Engine/Source/Runtime/Engine/Private/PhysicsEngine/PhysXSupport.h b/Engine/Source/Runtime/Engine/Private/PhysicsEngine/PhysXSupport.h
--- a/Engine/Source/Runtime/Engine/Private/PhysicsEngine/PhysXSupport.h
+++ b/Engine/Source/Runtime/Engine/Private/PhysicsEngine/PhysXSupport.h
@@ -160,13 +160,13 @@ NxApexScene* GetApexSceneFromIndex(int32 InSceneIndex);
 void DeferredPhysResourceCleanup();
 
 /** Calculates correct impulse at the body's center of mass and adds the impulse to the body. */
-ENGINE_API void AddRadialImpulseToPxRigidDynamic(PxRigidDynamic& PRigidDynamic, const FVector& Origin, float Radius, float Strength, uint8 Falloff, bool bVelChange);
+ENGINE_API void AddRadialImpulseToPxRigidBody(PxRigidBody& PRigidBody, const FVector& Origin, float Radius, float Strength, uint8 Falloff, bool bVelChange);
 
 /** Calculates correct force at the body's center of mass and adds force to the body. */
-ENGINE_API void AddRadialForceToPxRigidDynamic(PxRigidDynamic& PRigidDynamic, const FVector& Origin, float Radius, float Strength, uint8 Falloff);
+ENGINE_API void AddRadialForceToPxRigidBody(PxRigidBody& PRigidBody, const FVector& Origin, float Radius, float Strength, uint8 Falloff);
 
 /** Util to see if a PxRigidActor is non-kinematic */
-bool IsRigidDynamicNonKinematic(PxRigidDynamic* PRigidDynamic);
+bool IsRigidBodyNonKinematic(PxRigidBody* PRigidBody);
 
 
 /////// GLOBAL POINTERS
