Allow ReceiveTick and LatentActions to work while paused for player controller

diff --git a/Engine/Source/Runtime/Engine/Private/PlayerController.cpp b/Engine/Source/Runtime/Engine/Private/PlayerController.cpp
--- a/Engine/Source/Runtime/Engine/Private/PlayerController.cpp
+++ b/Engine/Source/Runtime/Engine/Private/PlayerController.cpp
@@ -3639,6 +3639,9 @@ void APlayerController::SetPlayer( UPlayer* InPlayer )
 
 void APlayerController::TickPlayerInput(const float DeltaSeconds, const bool bGamePaused)
 {
+	check(PlayerInput);
+	PlayerInput->Tick(DeltaSeconds);
+
 	if (ULocalPlayer* LocalPlayer = Cast<ULocalPlayer>(Player))
 	{
 		if (bEnableMouseOverEvents)
@@ -3679,22 +3682,27 @@ void APlayerController::TickPlayerInput(const float DeltaSeconds, const bool bGa
 		}
 	}
 
-	if (PlayerInput != NULL)
-	{
-		ProcessPlayerInput(DeltaSeconds, bGamePaused);
-		ProcessForceFeedback(DeltaSeconds, bGamePaused);
-	}
+	ProcessPlayerInput(DeltaSeconds, bGamePaused);
+	ProcessForceFeedback(DeltaSeconds, bGamePaused);
 }
 
 void APlayerController::TickActor( float DeltaSeconds, ELevelTick TickType, FActorTickFunction& ThisTickFunction )
 {
 	if (TickType == LEVELTICK_PauseTick && !ShouldPerformFullTickWhenPaused())
 	{
-		TickPlayerInput(DeltaSeconds, true);
+		if (PlayerInput)
+		{
+			TickPlayerInput(DeltaSeconds, true);
+		}
 
 		// Clear axis inputs from previous frame.
 		RotationInput = FRotator::ZeroRotator;
 
+		if (!IsPendingKill())
+		{
+			Tick(DeltaSeconds);	// perform any tick functions unique to an actor subclass
+		}
+
 		return; //root of tick hierarchy
 	}
 
@@ -3745,7 +3753,6 @@ void APlayerController::TickActor( float DeltaSeconds, ELevelTick TickType, FAct
 
 		if (PlayerInput)
 		{
-			PlayerInput->Tick(DeltaSeconds);
 			PlayerTick(DeltaSeconds);
 		}
 
