#ue4 - MovementComponent::ComputeSlideVector() now uses FVector::VectorPlaneProject() for clarity.

diff --git a/Engine/Source/Runtime/Engine/Private/MovementComponent.cpp b/Engine/Source/Runtime/Engine/Private/MovementComponent.cpp
--- a/Engine/Source/Runtime/Engine/Private/MovementComponent.cpp
+++ b/Engine/Source/Runtime/Engine/Private/MovementComponent.cpp
@@ -411,7 +411,7 @@ bool UMovementComponent::ResolvePenetration(const FVector& ProposedAdjustment, c
 
 FVector UMovementComponent::ComputeSlideVector(const FVector& Delta, const float Time, const FVector& Normal, const FHitResult& Hit) const
 {
-	return (Delta - Normal * (Delta | Normal)) * Time;
+	return FVector::VectorPlaneProject(Delta, Normal) * Time;
 }
 
 
