Proper fix for trails with realtime update off

diff --git a/Engine/Source/Runtime/Engine/Private/Particles/ParticleSystemRender.cpp b/Engine/Source/Runtime/Engine/Private/Particles/ParticleSystemRender.cpp
--- a/Engine/Source/Runtime/Engine/Private/Particles/ParticleSystemRender.cpp
+++ b/Engine/Source/Runtime/Engine/Private/Particles/ParticleSystemRender.cpp
@@ -671,120 +671,43 @@ void FDynamicSpriteEmitterDataBase::RenderDebug(const FParticleSystemSceneProxy*
 	}
 }
 
-/**
- *	Set up an buffer for async filling
- *
- *	@param	Proxy					The primitive scene proxy for the emitter.
- *	@param	InBufferIndex			Index of this buffer
- *	@param	InView					View for this buffer
- *	@param	InVertexCount			Count of verts for this buffer
- *	@param	InVertexSize			Stride of these verts, only used for verification
- *	@param	InDynamicParameterVertexStride	Stride of the dynamic parameter
- */
-void FDynamicSpriteEmitterDataBase::BuildViewFillData(FParticleSystemSceneProxy* Proxy, int32 InBufferIndex,const FSceneView *InView,int32 InVertexCount,int32 InVertexSize,int32 InDynamicParameterVertexStride)
+void FDynamicSpriteEmitterDataBase::BuildViewFillData(
+	const FParticleSystemSceneProxy* Proxy, 
+	const FSceneView* InView, 
+	int32 InVertexCount, 
+	int32 InVertexSize, 
+	int32 InDynamicParameterVertexStride, 
+	FGlobalDynamicVertexBuffer::FAllocation& DynamicVertexAllocation,
+	FGlobalDynamicIndexBuffer::FAllocation& DynamicIndexAllocation,
+	FGlobalDynamicVertexBuffer::FAllocation* DynamicParameterAllocation,
+	FAsyncBufferFillData& Data) const
 {
-	bool bSuccess = true;
-
-	FAsyncBufferFillData Data;
-
 	Data.LocalToWorld = Proxy->GetLocalToWorld();
 	Data.WorldToLocal = Proxy->GetWorldToLocal();
 	Data.View = InView;
 	check(Data.VertexSize == 0 || Data.VertexSize == InVertexSize);
 
-	check(InBufferIndex == InstanceDataAllocations.Num())
-	FGlobalDynamicVertexBuffer::FAllocation* DynamicVertexAllocation = new(InstanceDataAllocations) FGlobalDynamicVertexBuffer::FAllocation();
-	*DynamicVertexAllocation = FGlobalDynamicVertexBuffer::Get().Allocate( InVertexCount * InVertexSize );
-	Data.VertexData = DynamicVertexAllocation->Buffer;
+	DynamicVertexAllocation = FGlobalDynamicVertexBuffer::Get().Allocate( InVertexCount * InVertexSize );
+	Data.VertexData = DynamicVertexAllocation.Buffer;
 	Data.VertexCount = InVertexCount;
 	Data.VertexSize = InVertexSize;
 
-	bSuccess &= DynamicVertexAllocation->IsValid();
-
 	int32 NumIndices, IndexStride;
 	GetIndexAllocInfo(NumIndices, IndexStride);
 	check((uint32)NumIndices <= 65535);
 	check(IndexStride > 0);
 
-	FGlobalDynamicIndexBuffer::FAllocation* DynamicIndexAllocation = new(IndexDataAllocations) FGlobalDynamicIndexBuffer::FAllocation();
-	*DynamicIndexAllocation = FGlobalDynamicIndexBuffer::Get().Allocate( NumIndices, IndexStride );
-	Data.IndexData = DynamicIndexAllocation->Buffer;
+	DynamicIndexAllocation = FGlobalDynamicIndexBuffer::Get().Allocate( NumIndices, IndexStride );
+	Data.IndexData = DynamicIndexAllocation.Buffer;
 	Data.IndexCount = NumIndices;
 
-	bSuccess &= DynamicIndexAllocation->IsValid();
-
 	Data.DynamicParameterData = NULL;
+
 	if( bUsesDynamicParameter )
 	{
 		check( InDynamicParameterVertexStride > 0 );
-		check(InBufferIndex == DynamicParameterDataAllocations.Num())
-		FGlobalDynamicVertexBuffer::FAllocation* DynamicParameterAllocation = new(DynamicParameterDataAllocations) FGlobalDynamicVertexBuffer::FAllocation();
 		*DynamicParameterAllocation = FGlobalDynamicVertexBuffer::Get().Allocate( InVertexCount * InDynamicParameterVertexStride );
 		Data.DynamicParameterData = DynamicParameterAllocation->Buffer;
-
-		bSuccess &= DynamicParameterAllocation->IsValid();
-	}
-
-	//Create the data but only fill it if all the allocations succeeded.
-	if (InBufferIndex >= AsyncBufferFillTasks.Num())
-	{
-		new (AsyncBufferFillTasks) FAsyncBufferFillData();
-	}
-	check(InBufferIndex < AsyncBufferFillTasks.Num()); // please add the views in order
-
-	if( bSuccess )
-	{
-		AsyncBufferFillTasks[InBufferIndex] = Data;
-	}
-}
-
-/**
- *	Set up all buffers for async filling
- *
- *	@param	Proxy							The primitive scene proxy for the emitter.
- *	@param	ViewFamily						View family to process
- *	@param	VisibilityMap					Visibility map for the sub-views
- *	@param	bOnlyOneView					If true, then we don't need per-view buffers
- *	@param	InVertexCount					Count of verts for this buffer
- *	@param	InVertexSize					Stride of these verts, only used for verification
- *  @param  InDynamicParameterVertexSize	Stride of the dynamic parameter
- */
-void FDynamicSpriteEmitterDataBase::BuildViewFillDataAndSubmit(FParticleSystemSceneProxy* Proxy, const FSceneViewFamily* ViewFamily,const uint32 VisibilityMap,bool bOnlyOneView,int32 InVertexCount,int32 InVertexSize,int32 InDynamicParameterVertexSize)
-{
-	EnsureAsyncTaskComplete();
-	int32 NumUsedViews = 0;
-	InstanceDataAllocations.Empty();
-	IndexDataAllocations.Empty();
-	DynamicParameterDataAllocations.Empty();
-	for (int32 ViewIndex = 0; ViewIndex < ViewFamily->Views.Num(); ViewIndex++)
-	{
-		if (VisibilityMap & (1<<ViewIndex))
-		{
-			const FSceneView* View = ViewFamily->Views[ViewIndex];
-			BuildViewFillData(Proxy,NumUsedViews++,View,InVertexCount,InVertexSize, InDynamicParameterVertexSize);
-			if (bOnlyOneView)
-			{
-				break;
-			}
-		}
-	}
-	if (AsyncBufferFillTasks.Num() > NumUsedViews)
-	{
-		AsyncBufferFillTasks.RemoveAt(NumUsedViews,AsyncBufferFillTasks.Num() - NumUsedViews);
-	}
-	if (NumUsedViews)
-	{
-#if 0 // async particle fill is broken in UE4
-		if (!GIsEditor)
-		{
-			check(!AsyncTask.GetReference()); // this should not be still outstanding
-			AsyncTask = TGraphTask<FAsyncParticleFill>::CreateTask(NULL, ENamedThreads::RenderThread).ConstructAndDispatchWhenReady(this);
-		}
-		else
-#endif
-		{
-			DoBufferFill();
-		}
 	}
 }
 
@@ -3040,8 +2963,8 @@ void FDynamicBeam2EmitterData::GetDynamicMeshElementsEmitter(const FParticleSyst
 	FGlobalDynamicVertexBuffer::FAllocation DynamicParameterAllocation;
 	FAsyncBufferFillData Data;
 
-	//BuildViewFillData(Proxy, View, Source.VertexCount, sizeof(FParticleBeamTrailVertex), 0, DynamicVertexAllocation, DynamicIndexAllocation, &DynamicParameterAllocation, Data);
-	//DoBufferFill(Data);
+	BuildViewFillData(Proxy, View, Source.VertexCount, sizeof(FParticleBeamTrailVertex), 0, DynamicVertexAllocation, DynamicIndexAllocation, &DynamicParameterAllocation, Data);
+	DoBufferFill(Data);
 
 	if (Source.bUseLocalSpace == false)
 	{
@@ -3265,9 +3188,25 @@ void FDynamicBeam2EmitterData::PreRenderView(FParticleSystemSceneProxy* Proxy, c
 	// Only need to do this once per-view
 	if (LastFramePreRendered < FrameNumber)
 	{
-		bool bOnlyOneView = !GIsEditor && ((GEngine && GEngine->GameViewport && (GEngine->GameViewport->GetCurrentSplitscreenConfiguration() == ESplitScreenType::None) && !GEngine->IsStereoscopic3D()) ? true : false);
+		InstanceDataAllocations.Empty();
+		IndexDataAllocations.Empty();
+		DynamicParameterDataAllocations.Empty();
 
-		BuildViewFillDataAndSubmit(Proxy,ViewFamily,VisibilityMap,bOnlyOneView,Source.VertexCount,sizeof(FParticleBeamTrailVertex),0);
+		for (int32 ViewIndex = 0; ViewIndex < ViewFamily->Views.Num(); ViewIndex++)
+		{
+			if (VisibilityMap & (1 << ViewIndex))
+			{
+				const FSceneView* View = ViewFamily->Views[ViewIndex];
+
+				FGlobalDynamicVertexBuffer::FAllocation* DynamicVertexAllocation = new(InstanceDataAllocations) FGlobalDynamicVertexBuffer::FAllocation();
+				FGlobalDynamicIndexBuffer::FAllocation* DynamicIndexAllocation = new(IndexDataAllocations) FGlobalDynamicIndexBuffer::FAllocation();
+				FGlobalDynamicVertexBuffer::FAllocation* DynamicParameterAllocation = bUsesDynamicParameter ? new(DynamicParameterDataAllocations) FGlobalDynamicVertexBuffer::FAllocation() : NULL;
+				FAsyncBufferFillData* Data = new (AsyncBufferFillTasks) FAsyncBufferFillData();
+
+				BuildViewFillData(Proxy, View, Source.VertexCount, sizeof(FParticleBeamTrailVertex), 0, *DynamicVertexAllocation, *DynamicIndexAllocation, DynamicParameterAllocation, *Data);
+				DoBufferFill(*Data);
+			}
+		}
 
 		// Set the frame tracker
 		LastFramePreRendered = FrameNumber;
@@ -5907,8 +5846,8 @@ void FDynamicTrailsEmitterData::GetDynamicMeshElementsEmitter(const FParticleSys
 	const int32 VertexStride = GetDynamicVertexStride(ViewFamily.GetFeatureLevel());
 	const int32 DynamicParameterVertexStride = bUsesDynamicParameter ? GetDynamicParameterVertexStride() : 0;
 
-	//BuildViewFillData(Proxy, View, SourcePointer->VertexCount, VertexStride, DynamicParameterVertexStride, DynamicVertexAllocation, DynamicIndexAllocation, &DynamicParameterAllocation, Data);
-	//DoBufferFill(Data);
+	BuildViewFillData(Proxy, View, SourcePointer->VertexCount, VertexStride, DynamicParameterVertexStride, DynamicVertexAllocation, DynamicIndexAllocation, &DynamicParameterAllocation, Data);
+	DoBufferFill(Data);
 
 	if (SourcePointer->bUseLocalSpace == false)
 	{
@@ -6160,15 +6099,25 @@ void FDynamicTrailsEmitterData::PreRenderView(FParticleSystemSceneProxy* Proxy,
 			DynamicParameterVertexStride = GetDynamicParameterVertexStride();
 		}
 
-		bool bOnlyOneView = ShouldUsePrerenderView() || ((GEngine && GEngine->GameViewport && (GEngine->GameViewport->GetCurrentSplitscreenConfiguration() == ESplitScreenType::None)) ? true : false);
+		InstanceDataAllocations.Empty();
+		IndexDataAllocations.Empty();
+		DynamicParameterDataAllocations.Empty();
 
-		// Render both views when in stereo
-		if (GEngine->HMDDevice.IsValid() && GEngine->IsStereoscopic3D())
+		for (int32 ViewIndex = 0; ViewIndex < ViewFamily->Views.Num(); ViewIndex++)
 		{
-			bOnlyOneView = false;
-		}
+			if (VisibilityMap & (1 << ViewIndex))
+			{
+				const FSceneView* View = ViewFamily->Views[ViewIndex];
 
-		BuildViewFillDataAndSubmit(Proxy,ViewFamily,VisibilityMap,bOnlyOneView,SourcePointer->VertexCount,VertexStride, DynamicParameterVertexStride);
+				FGlobalDynamicVertexBuffer::FAllocation* DynamicVertexAllocation = new(InstanceDataAllocations) FGlobalDynamicVertexBuffer::FAllocation();
+				FGlobalDynamicIndexBuffer::FAllocation* DynamicIndexAllocation = new(IndexDataAllocations) FGlobalDynamicIndexBuffer::FAllocation();
+				FGlobalDynamicVertexBuffer::FAllocation* DynamicParameterAllocation = bUsesDynamicParameter ? new(DynamicParameterDataAllocations) FGlobalDynamicVertexBuffer::FAllocation() : NULL;
+				FAsyncBufferFillData* Data = new (AsyncBufferFillTasks) FAsyncBufferFillData();
+
+				BuildViewFillData(Proxy, View, SourcePointer->VertexCount, VertexStride, DynamicParameterVertexStride, *DynamicVertexAllocation, *DynamicIndexAllocation, DynamicParameterAllocation, *Data);
+				DoBufferFill(*Data);
+			}
+		}
 
 		// Set the frame tracker
 		LastFramePreRendered = FrameNumber;
diff --git a/Engine/Source/Runtime/Engine/Public/ParticleHelper.h b/Engine/Source/Runtime/Engine/Public/ParticleHelper.h
--- a/Engine/Source/Runtime/Engine/Public/ParticleHelper.h
+++ b/Engine/Source/Runtime/Engine/Public/ParticleHelper.h
@@ -1077,42 +1077,6 @@ struct FAsyncBufferFillData
 	}
 };
 
-/*-----------------------------------------------------------------------------
-	Async Fill Task, simple wrapper to forward the request to a FDynamicSpriteEmitterDataBase
------------------------------------------------------------------------------*/
-
-struct FAsyncParticleFill
-{
-	/** Emitter to forward to   */
-	struct FDynamicSpriteEmitterDataBase* Parent;
-
-	/** Constructor, just sets up the parent pointer  
-	  * @param InParent emitter to forward the eventual async call to
-	*/
-	FAsyncParticleFill(struct FDynamicSpriteEmitterDataBase* InParent)
-		: Parent(InParent)
-	{
-	}
-
-	/** Work function, just forwards the request to the parent  */
-	void DoTask(ENamedThreads::Type CurrentThread, const FGraphEventRef& MyCompletionGraphEvent);
-
-	FORCEINLINE TStatId GetStatId() const
-	{
-		return GET_STATID( STAT_ParticleAsyncTime );
-	}
-
-	static ENamedThreads::Type GetDesiredThread()
-	{
-		return ENamedThreads::AnyThread;
-	}
-
-	static ESubsequentsMode::Type GetSubsequentsMode() 
-	{ 
-		return ESubsequentsMode::TrackSubsequents; 
-	}
-};
-
 // TAsyncBufferFillTasks - handy typedef for an inline array of buffer fill tasks
 typedef TArray<FAsyncBufferFillData, TInlineAllocator<2> > TAsyncBufferFillTasks;
 
@@ -1484,7 +1448,6 @@ struct FDynamicSpriteEmitterDataBase : public FDynamicEmitterDataBase
 {
 	FDynamicSpriteEmitterDataBase(const UParticleModuleRequired* RequiredModule) : 
 		FDynamicEmitterDataBase(RequiredModule),
-		AsyncTask(NULL),
 		bUsesDynamicParameter( false )
 	{
 		MaterialResource[0] = NULL;
@@ -1493,7 +1456,6 @@ struct FDynamicSpriteEmitterDataBase : public FDynamicEmitterDataBase
 
 	virtual ~FDynamicSpriteEmitterDataBase()
 	{
-		EnsureAsyncTaskComplete();
 	}
 
 	/**
@@ -1570,62 +1532,31 @@ struct FDynamicSpriteEmitterDataBase : public FDynamicEmitterDataBase
 	 */
 	virtual void RenderDebug(const FParticleSystemSceneProxy* Proxy, FPrimitiveDrawInterface* PDI, const FSceneView* View, bool bCrosses) const;
 
-	/**
-	 *	Fill index and vertex buffers. Often called from a different thread
-	 *
-	 */
-	void DoBufferFill()
-	{
-		for (int32 TaskIndex = 0; TaskIndex < AsyncBufferFillTasks.Num(); TaskIndex++) 
-		{
-			DoBufferFill(AsyncBufferFillTasks[TaskIndex]);
-		}
-	}
-	/**
-	 *	Fill index and vertex buffers. Often called from a different thread
-	 *
-	 *	@param	Me			buffer pair to compute
-	 */
-	virtual void DoBufferFill(FAsyncBufferFillData& Me)const
+	virtual void DoBufferFill(FAsyncBufferFillData& Me) const
 	{
-		// this must be overridden, but in some cases a destructor call will leave this a no-op
-		// because the vtable has been reset to the base class 
-		// checkf(0, TEXT("DoBufferFill MUST be overridden"));
+		// Must be overridden if called
+		check(0);
 	}
+
 	/**
 	 *	Set up an buffer for async filling
 	 *
 	 *	@param	Proxy					The primitive scene proxy for the emitter.
-	 *	@param	InBufferIndex			Index of this buffer
 	 *	@param	InView					View for this buffer
 	 *	@param	InVertexCount			Count of verts for this buffer
 	 *	@param	InVertexSize			Stride of these verts, only used for verification
 	 *	@param	InDynamicParameterVertexStride	Stride of the dynamic parameter
 	 */
-	void BuildViewFillData(FParticleSystemSceneProxy* Proxy, int32 InBufferIndex,const FSceneView *InView,int32 InVertexCount,int32 InVertexSize,int32 InDynamicParameterVertexSize);
-
-	/**
-	 *	Set up all buffers for async filling
-	 *
-	 *	@param	Proxy							The primitive scene proxy for the emitter.
-	 *	@param	ViewFamily						View family to process
-	 *	@param	VisibilityMap					Visibility map for the sub-views
-	 *	@param	bOnlyOneView					If true, then we don't need per-view buffers
-	 *	@param	InVertexCount					Count of verts for this buffer
-	 *	@param	InVertexSize					Stride of these verts, only used for verification
-	 *	@param	InDynamicParameterVertexStride	Stride of the dynamic parameter
-	 */
-	void BuildViewFillDataAndSubmit(FParticleSystemSceneProxy* Proxy, const FSceneViewFamily* ViewFamily,const uint32 VisibilityMap,bool bOnlyOneView,int32 InVertexCount,int32 InVertexSize, int32 InDynamicParameterVertexSize);
-
-	void EnsureAsyncTaskComplete()
-	{
-		if (AsyncTask.GetReference())
-		{
-			SCOPE_CYCLE_COUNTER(STAT_ParticleAsyncWaitTime);
-			FTaskGraphInterface::Get().WaitUntilTaskCompletes(AsyncTask, ENamedThreads::RenderThread_Local);
-			AsyncTask = NULL;
-		}
-	}
+	void BuildViewFillData(
+		const FParticleSystemSceneProxy* Proxy, 
+		const FSceneView *InView, 
+		int32 InVertexCount, 
+		int32 InVertexSize, 
+		int32 InDynamicParameterVertexSize, 
+		FGlobalDynamicVertexBuffer::FAllocation& DynamicVertexAllocation,
+		FGlobalDynamicIndexBuffer::FAllocation& DynamicIndexAllocation,
+		FGlobalDynamicVertexBuffer::FAllocation* DynamicParameterAllocation,
+		FAsyncBufferFillData& Data) const;
 
 	/**
 	 *	Called to verify that a buffer is ready to use, blocks to wait and can sometimes execute the buffer fill on the current thread
@@ -1651,9 +1582,6 @@ struct FDynamicSpriteEmitterDataBase : public FDynamicEmitterDataBase
 	/** Async task is queued for execution */
 	bool									bAsyncTaskOutstanding;
 
-	/** Async task that is queued in the hi priority pool */
-	FGraphEventRef							AsyncTask;
-
 	/** Array of buffers for filling by async task */
 	TAsyncBufferFillTasks					AsyncBufferFillTasks;
 
