#docs - Clarify FTickFunction::TickGroup.

diff --git a/Engine/Source/Runtime/Engine/Classes/Engine/EngineBaseTypes.h b/Engine/Source/Runtime/Engine/Classes/Engine/EngineBaseTypes.h
--- a/Engine/Source/Runtime/Engine/Classes/Engine/EngineBaseTypes.h
+++ b/Engine/Source/Runtime/Engine/Classes/Engine/EngineBaseTypes.h
@@ -131,7 +131,13 @@ struct ENGINE_API FTickFunction
 public:
 	// The following UPROPERTYs are for configuration and inherited from the CDO/archetype/blueprint etc
 
-	/** Defines the minimum tick group for this tick function; given prerequisites, it can be delayed. **/
+	/**
+	 * Defines the minimum tick group for this tick function. These groups determine the relative order of when objects tick during a frame update.
+	 * Given prerequisites, the tick may be delayed.
+	 *
+	 * @see ETickingGroup 
+	 * @see FTickFunction::AddPrerequisite()
+	 */
 	UPROPERTY()
 	TEnumAsByte<enum ETickingGroup> TickGroup;
 
@@ -159,7 +165,7 @@ private:
 	uint32 bRegistered:1;
 
 	/** 
-	 * If false, this tick will not fire, nor will any other tick that has this tick function as an EnableParent 
+	 * If false, this tick will not fire, nor will any other tick that has this tick function as an EnableParent. 
 	 * CAUTION: Do not set this directly
 	 **/
 	uint32 bTickEnabled:1;
