PR517: Fix typo in a comment in the ParticleSystemComponent class.

diff --git a/Engine/Source/Runtime/Engine/Classes/Particles/ParticleSystemComponent.h b/Engine/Source/Runtime/Engine/Classes/Particles/ParticleSystemComponent.h
--- a/Engine/Source/Runtime/Engine/Classes/Particles/ParticleSystemComponent.h
+++ b/Engine/Source/Runtime/Engine/Classes/Particles/ParticleSystemComponent.h
@@ -245,7 +245,7 @@ struct FParticleEventKismetData : public FParticleEventData
 };
 
 /** 
- * A particle emmitter.
+ * A particle emitter.
  */
 UCLASS(ClassGroup=Rendering, hidecategories=Object, hidecategories=Physics, hidecategories=Collision, showcategories=Trigger, editinlinenew, meta=(BlueprintSpawnableComponent))
 class ENGINE_API UParticleSystemComponent : public UPrimitiveComponent
