#docs - Added comments/docs for FNavAgentProperties and FMovementProperties.

diff --git a/Engine/Source/Runtime/Engine/Classes/AI/Navigation/NavigationTypes.h b/Engine/Source/Runtime/Engine/Classes/AI/Navigation/NavigationTypes.h
--- a/Engine/Source/Runtime/Engine/Classes/AI/Navigation/NavigationTypes.h
+++ b/Engine/Source/Runtime/Engine/Classes/AI/Navigation/NavigationTypes.h
@@ -259,24 +259,29 @@ typedef TSharedRef<struct FNavigationPath, ESPMode::ThreadSafe> FNavPathSharedRe
 typedef TSharedPtr<struct FNavigationPath, ESPMode::ThreadSafe> FNavPathSharedPtr;
 typedef TWeakPtr<struct FNavigationPath, ESPMode::ThreadSafe> FNavPathWeakPtr;
 
+/** Movement capabilities, determining available movement options for Pawns and used by AI for reachability tests. */
 USTRUCT()
 struct FMovementProperties
 {
 	GENERATED_USTRUCT_BODY()
 
+	/** If true, this Pawn is capable of crouching. */
 	UPROPERTY(EditAnywhere, BlueprintReadWrite, Category=MovementProperties)
-	uint32 bCanCrouch:1;    // if true, this pawn is capable of crouching
+	uint32 bCanCrouch:1;
 
-	// movement capabilities - used by AI for reachability tests
+	/** If true, this Pawn is capable of jumping. */
 	UPROPERTY(EditAnywhere, BlueprintReadWrite, Category=MovementProperties)
 	uint32 bCanJump:1;
 
+	/** If true, this Pawn is capable of walking or moving on the ground. */
 	UPROPERTY(EditAnywhere, BlueprintReadWrite, Category=MovementProperties)
 	uint32 bCanWalk:1;
 
+	/** If true, this Pawn is capable of swimming or moving through fluid volumes. */
 	UPROPERTY(EditAnywhere, BlueprintReadWrite, Category=MovementProperties)
 	uint32 bCanSwim:1;
 
+	/** If true, this Pawn is capable of flying. */
 	UPROPERTY(EditAnywhere, BlueprintReadWrite, Category=MovementProperties)
 	uint32 bCanFly:1;
 
@@ -290,18 +295,21 @@ struct FMovementProperties
 	}
 };
 
+/** Properties of representation of an 'agent' (or Pawn) used by AI navigation/pathfinding. */
 USTRUCT()
 struct FNavAgentProperties : public FMovementProperties
 {
 	GENERATED_USTRUCT_BODY()
 
+	/** Radius of the capsule used for navigation/pathfinding. */
 	UPROPERTY(EditAnywhere, BlueprintReadWrite, Category = MovementProperties)
 	float AgentRadius;
 
+	/** Total height of the capsule used for navigation/pathfinding. */
 	UPROPERTY(EditAnywhere, BlueprintReadWrite, Category = MovementProperties)
 	float AgentHeight;
 
-	/** step height to use, or -1 for default value from navdata's config */
+	/** Step height to use, or -1 for default value from navdata's config. */
 	UPROPERTY(EditAnywhere, BlueprintReadWrite, Category = MovementProperties)
 	float AgentStepHeight;
 
