#ue4 - Fix AHUD::GetCoordinateOffset() returning uninitialized FVector2D when there is no local player.

diff --git a/Engine/Source/Runtime/Engine/Private/HUD.cpp b/Engine/Source/Runtime/Engine/Private/HUD.cpp
--- a/Engine/Source/Runtime/Engine/Private/HUD.cpp
+++ b/Engine/Source/Runtime/Engine/Private/HUD.cpp
@@ -74,7 +74,7 @@ void AHUD::NotifyBindPostProcessEffects()
 
 FVector2D AHUD::GetCoordinateOffset() const
 {
-	FVector2D Offset;
+	FVector2D Offset(0.f, 0.f);
 
 	ULocalPlayer* LocalPlayer = Cast<ULocalPlayer>(GetOwningPlayerController()->Player);
 
