#ue4 - Cleaned up CameraComponent and SpringArmComponent to use new deprecation enable/disable macros instead of all the compiler-specific pragmas in-line.

diff --git a/Engine/Source/Runtime/Engine/Private/Camera/CameraComponent.cpp b/Engine/Source/Runtime/Engine/Private/Camera/CameraComponent.cpp
--- a/Engine/Source/Runtime/Engine/Private/Camera/CameraComponent.cpp
+++ b/Engine/Source/Runtime/Engine/Private/Camera/CameraComponent.cpp
@@ -203,24 +203,9 @@ void UCameraComponent::CheckForErrors()
 
 void SetDeprecatedControllerViewRotation(UCameraComponent& Component, bool bValue)
 {
-	// BEGIN_IGNORE_DEPRECATION_WARNINGS
-	#ifdef __clang__
-		#pragma clang diagnostic push
-		#pragma clang diagnostic ignored "-Wdeprecated-declarations"
-	#else
-		#pragma warning(push)
-		#pragma warning(disable:4995)
-		#pragma warning(disable:4996)
-	#endif
-
+	PRAGMA_DISABLE_DEPRECATION_WARNINGS
 	Component.bUseControllerViewRotation = bValue;
-
-	// END_IGNORE_DEPRECATION_WARNINGS
-	#ifdef __clang__
-		#pragma clang diagnostic pop
-	#else
-		#pragma warning(pop)
-	#endif
+	PRAGMA_ENABLE_DEPRECATION_WARNINGS
 }
 
 
diff --git a/Engine/Source/Runtime/Engine/Private/GameFramework/SpringArmComponent.cpp b/Engine/Source/Runtime/Engine/Private/GameFramework/SpringArmComponent.cpp
--- a/Engine/Source/Runtime/Engine/Private/GameFramework/SpringArmComponent.cpp
+++ b/Engine/Source/Runtime/Engine/Private/GameFramework/SpringArmComponent.cpp
@@ -194,22 +194,7 @@ void USpringArmComponent::QuerySupportedSockets(TArray<FComponentSocketDescripti
 
 void SetDeprecatedControllerViewRotation(USpringArmComponent& Component, bool bValue)
 {
-	// BEGIN_IGNORE_DEPRECATION_WARNINGS
-	#ifdef __clang__
-		#pragma clang diagnostic push
-		#pragma clang diagnostic ignored "-Wdeprecated-declarations"
-	#else
-		#pragma warning(push)
-		#pragma warning(disable:4995)
-		#pragma warning(disable:4996)
-	#endif
-
+	PRAGMA_DISABLE_DEPRECATION_WARNINGS
 	Component.bUseControllerViewRotation = bValue;
-
-	// END_IGNORE_DEPRECATION_WARNINGS
-	#ifdef __clang__
-		#pragma clang diagnostic pop
-	#else
-		#pragma warning(pop)
-	#endif
+	PRAGMA_ENABLE_DEPRECATION_WARNINGS
 }
