#ue4 - Cleaned up CharacterMovementComponent to use new deprecation enable/disable macros instead of all the compiler-specific pragmas in-line.

diff --git a/Engine/Source/Runtime/Engine/Private/CharacterMovementComponent.cpp b/Engine/Source/Runtime/Engine/Private/CharacterMovementComponent.cpp
--- a/Engine/Source/Runtime/Engine/Private/CharacterMovementComponent.cpp
+++ b/Engine/Source/Runtime/Engine/Private/CharacterMovementComponent.cpp
@@ -2442,44 +2442,18 @@ float UCharacterMovementComponent::GetMaxJumpHeight() const
 float UCharacterMovementComponent::GetModifiedMaxAcceleration() const
 {
 	// Allow calling old deprecated function to maintain old behavior until it is removed.
-#ifdef __clang__
-	#pragma clang diagnostic push
-	#pragma clang diagnostic ignored "-Wdeprecated-declarations"
-#else
-	#pragma warning(push)
-	#pragma warning(disable:4995)
-	#pragma warning(disable:4996)
-#endif
-
+	PRAGMA_DISABLE_DEPRECATION_WARNINGS
 	return CharacterOwner ? MaxAcceleration * GetMaxSpeedModifier() : 0.f;
-
-#ifdef __clang__
-	#pragma clang diagnostic pop
-#else
-	#pragma warning(pop)
-#endif
+	PRAGMA_ENABLE_DEPRECATION_WARNINGS
 }
 
 // TODO: deprecated, remove.
 float UCharacterMovementComponent::K2_GetModifiedMaxAcceleration() const
 {
 	// Allow calling old deprecated function to maintain old behavior until it is removed.
-#ifdef __clang__
-	#pragma clang diagnostic push
-	#pragma clang diagnostic ignored "-Wdeprecated-declarations"
-#else
-	#pragma warning(push)
-	#pragma warning(disable:4995)
-	#pragma warning(disable:4996)
-#endif
-
+	PRAGMA_DISABLE_DEPRECATION_WARNINGS
 	return GetModifiedMaxAcceleration();
-
-#ifdef __clang__
-	#pragma clang diagnostic pop
-#else
-	#pragma warning(pop)
-#endif
+	PRAGMA_ENABLE_DEPRECATION_WARNINGS
 }
 
 
