Removed shader log spam

diff --git a/Engine/Source/Runtime/Engine/Private/Materials/MaterialShared.cpp b/Engine/Source/Runtime/Engine/Private/Materials/MaterialShared.cpp
--- a/Engine/Source/Runtime/Engine/Private/Materials/MaterialShared.cpp
+++ b/Engine/Source/Runtime/Engine/Private/Materials/MaterialShared.cpp
@@ -1384,7 +1384,6 @@ bool FMaterial::BeginCompileShaderMap(
 		}
 		else
 		{
-			UE_LOG(LogMaterial, Display, TEXT("Adding new outstanding shader map compilation id %d."), NewShaderMap->GetCompilingId() );
 			OutstandingCompileShaderMapIds.Add( NewShaderMap->GetCompilingId() );
 			// Async compile, use NULL so that rendering will fall back to the default material.
 			OutShaderMap = NULL;
diff --git a/Engine/Source/Runtime/Engine/Private/ShaderCompiler/ShaderCompiler.cpp b/Engine/Source/Runtime/Engine/Private/ShaderCompiler/ShaderCompiler.cpp
--- a/Engine/Source/Runtime/Engine/Private/ShaderCompiler/ShaderCompiler.cpp
+++ b/Engine/Source/Runtime/Engine/Private/ShaderCompiler/ShaderCompiler.cpp
@@ -1413,8 +1413,6 @@ void FShaderCompilingManager::ProcessCompiledShaderMaps(
 				Material->RemoveOutstandingCompileId( ShaderMap->CompilingId );
 			}
 
-			UE_LOG(LogShaderCompilers, Display, TEXT("Received shader map compiling id %d"), ShaderMap->CompilingId );
-
 			if (bShaderMapComplete)
 			{
 				ShaderMap->bCompiledSuccessfully = bSuccess;
