Fix for TTP 349644 (CRITICAL: Regression: CRASH: Occurs when launching a game with -NullRHI)

- Don't call DrawDebug on the world's FXSystem if it hasn't been created

diff --git a/Engine/Source/Runtime/Engine/Private/GameViewportClient.cpp b/Engine/Source/Runtime/Engine/Private/GameViewportClient.cpp
--- a/Engine/Source/Runtime/Engine/Private/GameViewportClient.cpp
+++ b/Engine/Source/Runtime/Engine/Private/GameViewportClient.cpp
@@ -954,7 +954,10 @@ void UGameViewportClient::Draw(FViewport* InViewport, FCanvas* SceneCanvas)
 	}
 
 	// Draw FX debug information.
-	GetWorld()->FXSystem->DrawDebug(SceneCanvas);
+	if (GetWorld()->FXSystem)
+	{
+		GetWorld()->FXSystem->DrawDebug(SceneCanvas);
+	}
 
 	// Render the UI.
 	{
