Landscape module from Engine split out.

diff --git a/Engine/Config/BaseEngine.ini b/Engine/Config/BaseEngine.ini
--- a/Engine/Config/BaseEngine.ini
+++ b/Engine/Config/BaseEngine.ini
@@ -453,9 +453,30 @@ bUseStreamingPause=false
 +ActiveClassRedirects=(OldClassName="AIDebugComponent",NewClassName="GameplayDebuggingComponent")
 
 ;Landscape redirects
-+ActiveClassRedirects=(OldClassName="MaterialExpressionTerrainLayerCoords",NewClassName="MaterialExpressionLandscapeLayerCoords")
-+ActiveClassRedirects=(OldClassName="MaterialExpressionTerrainLayerSwitch",NewClassName="MaterialExpressionLandscapeLayerSwitch")
-+ActiveClassRedirects=(OldClassName="MaterialExpressionTerrainLayerWeight",NewClassName="MaterialExpressionLandscapeLayerWeight")
++ActiveClassRedirects=(OldClassName="MaterialExpressionTerrainLayerCoords",NewClassName="/Script/Landscape.MaterialExpressionLandscapeLayerCoords")
++ActiveClassRedirects=(OldClassName="MaterialExpressionTerrainLayerSwitch",NewClassName="/Script/Landscape.MaterialExpressionLandscapeLayerSwitch")
++ActiveClassRedirects=(OldClassName="MaterialExpressionTerrainLayerWeight",NewClassName="/Script/Landscape.MaterialExpressionLandscapeLayerWeight")
++ActiveClassRedirects=(OldClassName="Landscape",NewClassName="/Script/Landscape.Landscape")
++ActiveClassRedirects=(OldClassName="LandscapeGizmoActiveActor",NewClassName="/Script/Landscape.LandscapeGizmoActiveActor")
++ActiveClassRedirects=(OldClassName="LandscapeGizmoActor",NewClassName="/Script/Landscape.LandscapeGizmoActor")
++ActiveClassRedirects=(OldClassName="LandscapeProxy",NewClassName="/Script/Landscape.LandscapeProxy")
++ActiveClassRedirects=(OldClassName="ControlPointMeshComponent",NewClassName="/Script/Landscape.ControlPointMeshComponent")
++ActiveClassRedirects=(OldClassName="LandscapeComponent",NewClassName="/Script/Landscape.LandscapeComponent")
++ActiveClassRedirects=(OldClassName="LandscapeGizmoRenderComponent",NewClassName="/Script/Landscape.LandscapeGizmoRenderComponent")
++ActiveClassRedirects=(OldClassName="LandscapeHeightfieldCollisionComponent",NewClassName="/Script/Landscape.LandscapeHeightfieldCollisionComponent")
++ActiveClassRedirects=(OldClassName="LandscapeInfo",NewClassName="/Script/Landscape.LandscapeInfo")
++ActiveClassRedirects=(OldClassName="LandscapeInfoMap",NewClassName="/Script/Landscape.LandscapeInfoMap")
++ActiveClassRedirects=(OldClassName="LandscapeLayerInfoObject",NewClassName="/Script/Landscape.LandscapeLayerInfoObject")
++ActiveClassRedirects=(OldClassName="LandscapeMaterialInstanceConstant",NewClassName="/Script/Landscape.LandscapeMaterialInstanceConstant")
++ActiveClassRedirects=(OldClassName="LandscapeMeshCollisionComponent",NewClassName="/Script/Landscape.LandscapeMeshCollisionComponent")
++ActiveClassRedirects=(OldClassName="LandscapeSplineControlPoint",NewClassName="/Script/Landscape.LandscapeSplineControlPoint")
++ActiveClassRedirects=(OldClassName="LandscapeSplineSegment",NewClassName="/Script/Landscape.LandscapeSplineSegment")
++ActiveClassRedirects=(OldClassName="LandscapeSplinesComponent",NewClassName="/Script/Landscape.LandscapeSplinesComponent")
++ActiveClassRedirects=(OldClassName="MaterialExpressionLandscapeLayerBlend",NewClassName="/Script/Landscape.MaterialExpressionLandscapeLayerBlend")
++ActiveClassRedirects=(OldClassName="MaterialExpressionLandscapeLayerCoords",NewClassName="/Script/Landscape.MaterialExpressionLandscapeLayerCoords")
++ActiveClassRedirects=(OldClassName="MaterialExpressionLandscapeLayerSwitch",NewClassName="/Script/Landscape.MaterialExpressionLandscapeLayerSwitch")
++ActiveClassRedirects=(OldClassName="MaterialExpressionLandscapeLayerWeight",NewClassName="/Script/Landscape.MaterialExpressionLandscapeLayerWeight")
++ActiveClassRedirects=(OldClassName="MaterialExpressionLandscapeVisibilityMask",NewClassName="/Script/Landscape.MaterialExpressionLandscapeVisibilityMask")
 
 ;Shortened animation compression algorithm class names
 +ActiveClassRedirects=(OldClassName="AnimationCompressionAlgorithm",NewClassName="AnimCompress")
diff --git a/Engine/Source/Developer/AssetTools/AssetTools.Build.cs b/Engine/Source/Developer/AssetTools/AssetTools.Build.cs
--- a/Engine/Source/Developer/AssetTools/AssetTools.Build.cs
+++ b/Engine/Source/Developer/AssetTools/AssetTools.Build.cs
@@ -21,7 +21,8 @@ public class AssetTools : ModuleRules
 				"SourceControl",
 				"TextureEditor",
 				"UnrealEd",
-				"Kismet"
+				"Kismet",
+				"Landscape",
 			}
 		);
 
@@ -40,7 +41,8 @@ public class AssetTools : ModuleRules
 				"FontEditor",
 				"SoundCueEditor",
 				"SoundClassEditor",
-				"SourceControl"
+				"SourceControl",
+				"Landscape",
 			}
 		);
 
diff --git a/Engine/Source/Developer/AssetTools/Private/AssetTypeActions/AssetTypeActions_LandscapeLayer.h b/Engine/Source/Developer/AssetTools/Private/AssetTypeActions/AssetTypeActions_LandscapeLayer.h
--- a/Engine/Source/Developer/AssetTools/Private/AssetTypeActions/AssetTypeActions_LandscapeLayer.h
+++ b/Engine/Source/Developer/AssetTools/Private/AssetTypeActions/AssetTypeActions_LandscapeLayer.h
@@ -2,7 +2,7 @@
 
 #pragma once
 
-#include "Landscape/LandscapeLayerInfoObject.h"
+#include "LandscapeLayerInfoObject.h"
 
 class FAssetTypeActions_LandscapeLayer : public FAssetTypeActions_Base
 {
diff --git a/Engine/Source/Developer/MeshUtilities/MeshUtilities.Build.cs b/Engine/Source/Developer/MeshUtilities/MeshUtilities.Build.cs
--- a/Engine/Source/Developer/MeshUtilities/MeshUtilities.Build.cs
+++ b/Engine/Source/Developer/MeshUtilities/MeshUtilities.Build.cs
@@ -17,9 +17,16 @@ public class MeshUtilities : ModuleRules
 				"Slate",      // For FSlateTextureAtlas
 				"SlateCore",
 				"UnrealEd",
+				"Landscape",
 			}
 		);
 
+		PrivateIncludePathModuleNames.AddRange(
+				new string[] {
+					"Landscape",
+				}
+		);
+
 		AddThirdPartyPrivateStaticDependencies(Target, "nvTriStrip");
 		AddThirdPartyPrivateStaticDependencies(Target, "ForsythTriOptimizer");
 		AddThirdPartyPrivateStaticDependencies(Target, "MeshSimplifier");      
diff --git a/Engine/Source/Developer/MeshUtilities/Private/MeshUtilities.cpp b/Engine/Source/Developer/MeshUtilities/Private/MeshUtilities.cpp
--- a/Engine/Source/Developer/MeshUtilities/Private/MeshUtilities.cpp
+++ b/Engine/Source/Developer/MeshUtilities/Private/MeshUtilities.cpp
@@ -3,14 +3,14 @@
 #include "MeshUtilitiesPrivate.h"
 #include "StaticMeshResources.h"
 #include "SkeletalMeshTypes.h"
-#include "Landscape/LandscapeRender.h"
+#include "LandscapeRender.h"
 #include "MeshBuild.h"
 #include "TessellationRendering.h"
 #include "NvTriStrip.h"
 #include "forsythtriangleorderoptimizer.h"
 #include "ThirdParty/nvtesslib/inc/nvtess.h"
 #include "SkeletalMeshTools.h"
-#include "Landscape/LandscapeDataAccess.h"
+#include "LandscapeDataAccess.h"
 #include "ImageUtils.h"
 #include "MaterialExportUtils.h"
 #include "Textures/TextureAtlas.h"
diff --git a/Engine/Source/Editor/DetailCustomizations/DetailCustomizations.Build.cs b/Engine/Source/Editor/DetailCustomizations/DetailCustomizations.Build.cs
--- a/Engine/Source/Editor/DetailCustomizations/DetailCustomizations.Build.cs
+++ b/Engine/Source/Editor/DetailCustomizations/DetailCustomizations.Build.cs
@@ -15,6 +15,7 @@ public class DetailCustomizations : ModuleRules
 				"CoreUObject",
 				"DesktopWidgets",
 				"Engine",
+				"Landscape",
                 "InputCore",
 				"Slate",
 				"SlateCore",
@@ -48,6 +49,7 @@ public class DetailCustomizations : ModuleRules
 		PrivateIncludePathModuleNames.AddRange(
 			new string[] {
 				"PropertyEditor",
+				"Landscape",
 				"LandscapeEditor",
 			}
 		);
diff --git a/Engine/Source/Editor/DetailCustomizations/Private/ActorMaterialCategory.cpp b/Engine/Source/Editor/DetailCustomizations/Private/ActorMaterialCategory.cpp
--- a/Engine/Source/Editor/DetailCustomizations/Private/ActorMaterialCategory.cpp
+++ b/Engine/Source/Editor/DetailCustomizations/Private/ActorMaterialCategory.cpp
@@ -5,8 +5,8 @@
 #include "AssetThumbnail.h"
 #include "ActorEditorUtils.h"
 
-#include "Landscape/LandscapeProxy.h"
-#include "Landscape/LandscapeComponent.h"
+#include "LandscapeProxy.h"
+#include "LandscapeComponent.h"
 
 #define LOCTEXT_NAMESPACE "SMaterialList"
 
diff --git a/Engine/Source/Editor/FoliageEdit/FoliageEdit.build.cs b/Engine/Source/Editor/FoliageEdit/FoliageEdit.build.cs
--- a/Engine/Source/Editor/FoliageEdit/FoliageEdit.build.cs
+++ b/Engine/Source/Editor/FoliageEdit/FoliageEdit.build.cs
@@ -18,6 +18,7 @@ public class FoliageEdit : ModuleRules
                 "EditorStyle",
 				"RenderCore",
 				"LevelEditor",
+				"Landscape",
 			}
 		);
 	}
diff --git a/Engine/Source/Editor/FoliageEdit/Private/FoliageEdMode.cpp b/Engine/Source/Editor/FoliageEdit/Private/FoliageEdMode.cpp
--- a/Engine/Source/Editor/FoliageEdit/Private/FoliageEdMode.cpp
+++ b/Engine/Source/Editor/FoliageEdit/Private/FoliageEdMode.cpp
@@ -20,9 +20,9 @@
 // Classes
 #include "Foliage/InstancedFoliageActor.h"
 #include "Foliage/FoliageType.h"
-#include "Landscape/LandscapeComponent.h"
-#include "Landscape/LandscapeHeightfieldCollisionComponent.h"
-#include "Landscape/LandscapeInfo.h"
+#include "LandscapeComponent.h"
+#include "LandscapeHeightfieldCollisionComponent.h"
+#include "LandscapeInfo.h"
 #include "Components/SplineMeshComponent.h"
 
 #define FOLIAGE_SNAP_TRACE (10000.f)
diff --git a/Engine/Source/Editor/LandscapeEditor/Classes/LandscapeEditorObject.h b/Engine/Source/Editor/LandscapeEditor/Classes/LandscapeEditorObject.h
--- a/Engine/Source/Editor/LandscapeEditor/Classes/LandscapeEditorObject.h
+++ b/Engine/Source/Editor/LandscapeEditor/Classes/LandscapeEditorObject.h
@@ -2,7 +2,7 @@
 
 #pragma once
 
-#include "Landscape/LandscapeProxy.h" // for ELandscapeLayerPaintingRestriction
+#include "LandscapeProxy.h" // for ELandscapeLayerPaintingRestriction
 #include "LandscapeEditorObject.generated.h"
 
 UENUM()
diff --git a/Engine/Source/Editor/LandscapeEditor/LandscapeEditor.Build.cs b/Engine/Source/Editor/LandscapeEditor/LandscapeEditor.Build.cs
--- a/Engine/Source/Editor/LandscapeEditor/LandscapeEditor.Build.cs
+++ b/Engine/Source/Editor/LandscapeEditor/LandscapeEditor.Build.cs
@@ -14,6 +14,7 @@ public class LandscapeEditor : ModuleRules
 				"SlateCore",
                 "EditorStyle",
 				"Engine",
+				"Landscape",
 				"RenderCore",
                 "InputCore",
 				"UnrealEd",
diff --git a/Engine/Source/Editor/LandscapeEditor/Private/LandscapeEdMode.cpp b/Engine/Source/Editor/LandscapeEditor/Private/LandscapeEdMode.cpp
--- a/Engine/Source/Editor/LandscapeEditor/Private/LandscapeEdMode.cpp
+++ b/Engine/Source/Editor/LandscapeEditor/Private/LandscapeEdMode.cpp
@@ -6,10 +6,10 @@
 #include "ObjectTools.h"
 #include "LandscapeEdMode.h"
 #include "ScopedTransaction.h"
-#include "Landscape/LandscapeEdit.h"
-#include "Landscape/LandscapeRender.h"
-#include "Landscape/LandscapeDataAccess.h"
-#include "Landscape/LandscapeSplineProxies.h"
+#include "LandscapeEdit.h"
+#include "LandscapeRender.h"
+#include "LandscapeDataAccess.h"
+#include "LandscapeSplineProxies.h"
 #include "LandscapeEditorModule.h"
 #include "Editor/LevelEditor/Public/LevelEditor.h"
 #include "Editor/PropertyEditor/Public/PropertyEditorModule.h"
@@ -25,11 +25,11 @@
 #include "SLandscapeEditor.h"
 
 // Classes
-#include "Landscape/Landscape.h"
-#include "Landscape/LandscapeLayerInfoObject.h"
-#include "Landscape/LandscapeHeightfieldCollisionComponent.h"
-#include "Landscape/LandscapeMaterialInstanceConstant.h"
-#include "Landscape/LandscapeSplinesComponent.h"
+#include "Landscape.h"
+#include "LandscapeLayerInfoObject.h"
+#include "LandscapeHeightfieldCollisionComponent.h"
+#include "LandscapeMaterialInstanceConstant.h"
+#include "LandscapeSplinesComponent.h"
 #include "Foliage/InstancedFoliageActor.h"
 #include "ComponentReregisterContext.h"
 
@@ -1630,7 +1630,7 @@ int32 FEdModeLandscape::UpdateLandscapeList()
 	if (GetWorld())
 	{
 		int32 Index = 0;
-		for (auto It = GetWorld()->LandscapeInfoMap.CreateIterator(); It; ++It)
+		for (auto It = GetLandscapeInfoMap(GetWorld()).Map.CreateIterator(); It; ++It)
 		{
 			ULandscapeInfo* LandscapeInfo = It.Value();
 			if (LandscapeInfo)
@@ -2579,7 +2579,7 @@ ALandscape* FEdModeLandscape::ChangeComponentSetting(int32 NumComponentsX, int32
 							{
 								OldFoliageActor->MoveInstancesToNewComponent(OldCollisionComponent, NewCollisionComponent);
 								AInstancedFoliageActor* NewFoliageActor = AInstancedFoliageActor::GetInstancedFoliageActorForLevel(NewCollisionComponent->GetTypedOuter<ULevel>());
-								NewFoliageActor->SnapInstancesForLandscape(NewCollisionComponent, FBox(FVector(-WORLD_MAX), FVector(WORLD_MAX)));
+								NewCollisionComponent->SnapFoliageInstances(*NewFoliageActor, FBox(FVector(-WORLD_MAX), FVector(WORLD_MAX)));
 							}
 						}
 					}
diff --git a/Engine/Source/Editor/LandscapeEditor/Private/LandscapeEdModeBrushes.cpp b/Engine/Source/Editor/LandscapeEditor/Private/LandscapeEdModeBrushes.cpp
--- a/Engine/Source/Editor/LandscapeEditor/Private/LandscapeEdModeBrushes.cpp
+++ b/Engine/Source/Editor/LandscapeEditor/Private/LandscapeEdModeBrushes.cpp
@@ -5,10 +5,10 @@
 #include "ObjectTools.h"
 #include "LandscapeEdMode.h"
 #include "ScopedTransaction.h"
-#include "Landscape/LandscapeEdit.h"
-#include "Landscape/LandscapeRender.h"
-#include "Landscape/Landscape.h"
-#include "Landscape/LandscapeHeightfieldCollisionComponent.h"
+#include "LandscapeEdit.h"
+#include "LandscapeRender.h"
+#include "Landscape.h"
+#include "LandscapeHeightfieldCollisionComponent.h"
 
 #include "LevelUtils.h"
 
diff --git a/Engine/Source/Editor/LandscapeEditor/Private/LandscapeEdModeComponentTools.cpp b/Engine/Source/Editor/LandscapeEditor/Private/LandscapeEdModeComponentTools.cpp
--- a/Engine/Source/Editor/LandscapeEditor/Private/LandscapeEdModeComponentTools.cpp
+++ b/Engine/Source/Editor/LandscapeEditor/Private/LandscapeEdModeComponentTools.cpp
@@ -4,12 +4,12 @@
 #include "ObjectTools.h"
 #include "LandscapeEdMode.h"
 #include "ScopedTransaction.h"
-#include "Landscape/Landscape.h"
-#include "Landscape/LandscapeEdit.h"
-#include "Landscape/LandscapeLayerInfoObject.h"
-#include "Landscape/LandscapeRender.h"
-#include "Landscape/LandscapeDataAccess.h"
-#include "Landscape/LandscapeSplineProxies.h"
+#include "Landscape.h"
+#include "LandscapeEdit.h"
+#include "LandscapeLayerInfoObject.h"
+#include "LandscapeRender.h"
+#include "LandscapeDataAccess.h"
+#include "LandscapeSplineProxies.h"
 #include "LandscapeEditorModule.h"
 #include "Editor/PropertyEditor/Public/PropertyEditorModule.h"
 #include "LandscapeEdModeTools.h"
diff --git a/Engine/Source/Editor/LandscapeEditor/Private/LandscapeEdModeErosionTools.cpp b/Engine/Source/Editor/LandscapeEditor/Private/LandscapeEdModeErosionTools.cpp
--- a/Engine/Source/Editor/LandscapeEditor/Private/LandscapeEdModeErosionTools.cpp
+++ b/Engine/Source/Editor/LandscapeEditor/Private/LandscapeEdModeErosionTools.cpp
@@ -4,15 +4,15 @@
 #include "ObjectTools.h"
 #include "LandscapeEdMode.h"
 #include "ScopedTransaction.h"
-#include "Landscape/LandscapeEdit.h"
-#include "Landscape/LandscapeRender.h"
-#include "Landscape/LandscapeDataAccess.h"
-#include "Landscape/LandscapeSplineProxies.h"
+#include "LandscapeEdit.h"
+#include "LandscapeRender.h"
+#include "LandscapeDataAccess.h"
+#include "LandscapeSplineProxies.h"
 #include "LandscapeEditorModule.h"
 #include "Editor/PropertyEditor/Public/PropertyEditorModule.h"
 #include "LandscapeEdModeTools.h"
-#include "Landscape/Landscape.h"
-#include "Landscape/LandscapeLayerInfoObject.h"
+#include "Landscape.h"
+#include "LandscapeLayerInfoObject.h"
 
 //
 // FLandscapeToolErosionBase
diff --git a/Engine/Source/Editor/LandscapeEditor/Private/LandscapeEdModePaintTools.cpp b/Engine/Source/Editor/LandscapeEditor/Private/LandscapeEdModePaintTools.cpp
--- a/Engine/Source/Editor/LandscapeEditor/Private/LandscapeEdModePaintTools.cpp
+++ b/Engine/Source/Editor/LandscapeEditor/Private/LandscapeEdModePaintTools.cpp
@@ -4,10 +4,10 @@
 #include "ObjectTools.h"
 #include "LandscapeEdMode.h"
 #include "ScopedTransaction.h"
-#include "Landscape/LandscapeEdit.h"
-#include "Landscape/LandscapeRender.h"
-#include "Landscape/LandscapeDataAccess.h"
-#include "Landscape/LandscapeSplineProxies.h"
+#include "LandscapeEdit.h"
+#include "LandscapeRender.h"
+#include "LandscapeDataAccess.h"
+#include "LandscapeSplineProxies.h"
 #include "LandscapeEditorModule.h"
 #include "Editor/PropertyEditor/Public/PropertyEditorModule.h"
 #include "LandscapeEdModeTools.h"
diff --git a/Engine/Source/Editor/LandscapeEditor/Private/LandscapeEdModeRampTool.cpp b/Engine/Source/Editor/LandscapeEditor/Private/LandscapeEdModeRampTool.cpp
--- a/Engine/Source/Editor/LandscapeEditor/Private/LandscapeEdModeRampTool.cpp
+++ b/Engine/Source/Editor/LandscapeEditor/Private/LandscapeEdModeRampTool.cpp
@@ -3,10 +3,10 @@
 #include "LandscapeEditorPrivatePCH.h"
 #include "LandscapeEdMode.h"
 #include "ScopedTransaction.h"
-#include "Landscape/LandscapeEdit.h"
-#include "Landscape/LandscapeRender.h"
-#include "Landscape/LandscapeDataAccess.h"
-#include "Landscape/LandscapeHeightfieldCollisionComponent.h"
+#include "LandscapeEdit.h"
+#include "LandscapeRender.h"
+#include "LandscapeDataAccess.h"
+#include "LandscapeHeightfieldCollisionComponent.h"
 #include "Raster.h"
 
 #define LOCTEXT_NAMESPACE "Landscape"
diff --git a/Engine/Source/Editor/LandscapeEditor/Private/LandscapeEdModeSplineTools.cpp b/Engine/Source/Editor/LandscapeEditor/Private/LandscapeEdModeSplineTools.cpp
--- a/Engine/Source/Editor/LandscapeEditor/Private/LandscapeEdModeSplineTools.cpp
+++ b/Engine/Source/Editor/LandscapeEditor/Private/LandscapeEdModeSplineTools.cpp
@@ -4,19 +4,19 @@
 #include "ObjectTools.h"
 #include "LandscapeEdMode.h"
 #include "ScopedTransaction.h"
-#include "Landscape/LandscapeEdit.h"
-#include "Landscape/LandscapeRender.h"
-#include "Landscape/LandscapeDataAccess.h"
-#include "Landscape/LandscapeSplineProxies.h"
+#include "LandscapeEdit.h"
+#include "LandscapeRender.h"
+#include "LandscapeDataAccess.h"
+#include "LandscapeSplineProxies.h"
 #include "LandscapeEditorModule.h"
 #include "Editor/PropertyEditor/Public/PropertyEditorModule.h"
 #include "LandscapeEdModeTools.h"
 #include "Components/SplineMeshComponent.h"
 #include "LandscapeSplineImportExport.h"
-#include "Landscape/LandscapeSplinesComponent.h"
-#include "Landscape/LandscapeSplineControlPoint.h"
-#include "Landscape/LandscapeSplineSegment.h"
-#include "Landscape/ControlPointMeshComponent.h"
+#include "LandscapeSplinesComponent.h"
+#include "LandscapeSplineControlPoint.h"
+#include "LandscapeSplineSegment.h"
+#include "ControlPointMeshComponent.h"
 
 
 #define LOCTEXT_NAMESPACE "Landscape"
diff --git a/Engine/Source/Editor/LandscapeEditor/Private/LandscapeEdModeTools.h b/Engine/Source/Editor/LandscapeEditor/Private/LandscapeEdModeTools.h
--- a/Engine/Source/Editor/LandscapeEditor/Private/LandscapeEdModeTools.h
+++ b/Engine/Source/Editor/LandscapeEditor/Private/LandscapeEdModeTools.h
@@ -2,8 +2,8 @@
 
 #pragma once
 
-#include "Landscape/Landscape.h"
-#include "Landscape/LandscapeHeightfieldCollisionComponent.h"
+#include "Landscape.h"
+#include "LandscapeHeightfieldCollisionComponent.h"
 #include "Foliage/InstancedFoliageActor.h"
 
 //
@@ -625,7 +625,7 @@ struct FHeightmapAccessor
 					AInstancedFoliageActor* IFA = AInstancedFoliageActor::GetInstancedFoliageActorForLevel(CollisionComponent->GetComponentLevel());
 					if (IFA)
 					{
-						IFA->SnapInstancesForLandscape(CollisionComponent, PreUpdateLocalBoxes[Index].TransformBy(LandscapeInfo->GetLandscapeProxy()->LandscapeActorToWorld().ToMatrixWithScale()).ExpandBy(1.f));
+						CollisionComponent->SnapFoliageInstances(*IFA, PreUpdateLocalBoxes[Index].TransformBy(LandscapeInfo->GetLandscapeProxy()->LandscapeActorToWorld().ToMatrixWithScale()).ExpandBy(1.f));
 					}
 				}
 			}
@@ -800,7 +800,7 @@ struct FXYOffsetmapAccessor
 				{
 					ULandscapeHeightfieldCollisionComponent* CollisionComponent = CollisionComponents[Index];
 					AInstancedFoliageActor* IFA = AInstancedFoliageActor::GetInstancedFoliageActorForLevel(CollisionComponent->GetComponentLevel());
-					IFA->SnapInstancesForLandscape(CollisionComponent, PreUpdateLocalBoxes[Index].TransformBy(LandscapeInfo->GetLandscapeProxy()->LandscapeActorToWorld().ToMatrixWithScale()).ExpandBy(1.f));
+					CollisionComponent->SnapFoliageInstances(*IFA, PreUpdateLocalBoxes[Index].TransformBy(LandscapeInfo->GetLandscapeProxy()->LandscapeActorToWorld().ToMatrixWithScale()).ExpandBy(1.f));
 				}
 			}
 			else
diff --git a/Engine/Source/Editor/LandscapeEditor/Private/LandscapeEdModeXYOffset.cpp b/Engine/Source/Editor/LandscapeEditor/Private/LandscapeEdModeXYOffset.cpp
--- a/Engine/Source/Editor/LandscapeEditor/Private/LandscapeEdModeXYOffset.cpp
+++ b/Engine/Source/Editor/LandscapeEditor/Private/LandscapeEdModeXYOffset.cpp
@@ -4,14 +4,14 @@
 #include "ObjectTools.h"
 #include "LandscapeEdMode.h"
 #include "ScopedTransaction.h"
-#include "Landscape/LandscapeEdit.h"
-#include "Landscape/LandscapeRender.h"
-#include "Landscape/LandscapeDataAccess.h"
-#include "Landscape/LandscapeSplineProxies.h"
+#include "LandscapeEdit.h"
+#include "LandscapeRender.h"
+#include "LandscapeDataAccess.h"
+#include "LandscapeSplineProxies.h"
 #include "LandscapeEditorModule.h"
 #include "Editor/PropertyEditor/Public/PropertyEditorModule.h"
 #include "LandscapeEdModeTools.h"
-#include "Landscape/Landscape.h"
+#include "Landscape.h"
 
 namespace
 {
diff --git a/Engine/Source/Editor/LandscapeEditor/Private/LandscapeEditorDetailCustomization_CopyPaste.cpp b/Engine/Source/Editor/LandscapeEditor/Private/LandscapeEditorDetailCustomization_CopyPaste.cpp
--- a/Engine/Source/Editor/LandscapeEditor/Private/LandscapeEditorDetailCustomization_CopyPaste.cpp
+++ b/Engine/Source/Editor/LandscapeEditor/Private/LandscapeEditorDetailCustomization_CopyPaste.cpp
@@ -18,7 +18,7 @@
 #include "DesktopPlatformModule.h"
 #include "MainFrame.h"
 
-#include "Landscape/LandscapeLayerInfoObject.h"
+#include "LandscapeLayerInfoObject.h"
 #include "SNumericEntryBox.h"
 
 #define LOCTEXT_NAMESPACE "LandscapeEditor.CopyPaste"
diff --git a/Engine/Source/Editor/LandscapeEditor/Private/LandscapeEditorDetailCustomization_NewLandscape.cpp b/Engine/Source/Editor/LandscapeEditor/Private/LandscapeEditorDetailCustomization_NewLandscape.cpp
--- a/Engine/Source/Editor/LandscapeEditor/Private/LandscapeEditorDetailCustomization_NewLandscape.cpp
+++ b/Engine/Source/Editor/LandscapeEditor/Private/LandscapeEditorDetailCustomization_NewLandscape.cpp
@@ -27,8 +27,8 @@
 #include "AssetRegistryModule.h"
 #include "ImageWrapper.h"
 
-#include "Landscape/Landscape.h"
-#include "Landscape/LandscapeLayerInfoObject.h"
+#include "Landscape.h"
+#include "LandscapeLayerInfoObject.h"
 #include "TutorialMetaData.h"
 #include "SNumericEntryBox.h"
 
diff --git a/Engine/Source/Editor/LandscapeEditor/Private/LandscapeEditorDetailCustomization_TargetLayers.cpp b/Engine/Source/Editor/LandscapeEditor/Private/LandscapeEditorDetailCustomization_TargetLayers.cpp
--- a/Engine/Source/Editor/LandscapeEditor/Private/LandscapeEditorDetailCustomization_TargetLayers.cpp
+++ b/Engine/Source/Editor/LandscapeEditor/Private/LandscapeEditorDetailCustomization_TargetLayers.cpp
@@ -22,9 +22,9 @@
 #include "DesktopPlatformModule.h"
 #include "MainFrame.h"
 #include "AssetRegistryModule.h"
-#include "LandscapeRender.h"
 
-#include "Landscape/LandscapeLayerInfoObject.h"
+#include "LandscapeRender.h"
+#include "LandscapeLayerInfoObject.h"
 
 #define LOCTEXT_NAMESPACE "LandscapeEditor.TargetLayers"
 
diff --git a/Engine/Source/Editor/LandscapeEditor/Private/LandscapeEditorModule.cpp b/Engine/Source/Editor/LandscapeEditor/Private/LandscapeEditorModule.cpp
--- a/Engine/Source/Editor/LandscapeEditor/Private/LandscapeEditorModule.cpp
+++ b/Engine/Source/Editor/LandscapeEditor/Private/LandscapeEditorModule.cpp
@@ -12,9 +12,9 @@
 #include "PropertyEditorDelegates.h"
 
 #include "Editor/LevelEditor/Public/LevelEditor.h"
-#include "LandscapeRender.h"
 
-#include "Landscape/Landscape.h"
+#include "Landscape.h"
+#include "LandscapeRender.h"
 
 #define LOCTEXT_NAMESPACE "LandscapeEditor"
 
diff --git a/Engine/Source/Editor/LandscapeEditor/Private/LandscapeEditorObject.cpp b/Engine/Source/Editor/LandscapeEditor/Private/LandscapeEditorObject.cpp
--- a/Engine/Source/Editor/LandscapeEditor/Private/LandscapeEditorObject.cpp
+++ b/Engine/Source/Editor/LandscapeEditor/Private/LandscapeEditorObject.cpp
@@ -5,7 +5,7 @@
 #include "LandscapeRender.h"
 #include "ImageWrapper.h"
 
-#include "Landscape/LandscapeMaterialInstanceConstant.h"
+#include "LandscapeMaterialInstanceConstant.h"
 
 ULandscapeEditorObject::ULandscapeEditorObject(const FObjectInitializer& ObjectInitializer)
 	: Super(ObjectInitializer)
diff --git a/Engine/Source/Editor/LandscapeEditor/Private/LandscapeSplineImportExport.cpp b/Engine/Source/Editor/LandscapeEditor/Private/LandscapeSplineImportExport.cpp
--- a/Engine/Source/Editor/LandscapeEditor/Private/LandscapeSplineImportExport.cpp
+++ b/Engine/Source/Editor/LandscapeEditor/Private/LandscapeSplineImportExport.cpp
@@ -2,8 +2,8 @@
 
 #include "LandscapeEditorPrivatePCH.h"
 #include "LandscapeSplineImportExport.h"
-#include "Landscape/LandscapeSplineControlPoint.h"
-#include "Landscape/LandscapeSplineSegment.h"
+#include "LandscapeSplineControlPoint.h"
+#include "LandscapeSplineSegment.h"
 
 #define LOCTEXT_NAMESPACE "Landscape"
 
diff --git a/Engine/Source/Editor/LandscapeEditor/Private/Tests/LandscapeAutomationTests.cpp b/Engine/Source/Editor/LandscapeEditor/Private/Tests/LandscapeAutomationTests.cpp
--- a/Engine/Source/Editor/LandscapeEditor/Private/Tests/LandscapeAutomationTests.cpp
+++ b/Engine/Source/Editor/LandscapeEditor/Private/Tests/LandscapeAutomationTests.cpp
@@ -8,7 +8,7 @@
 #include "LandscapeEdMode.h"
 #include "../LandscapeEditorDetailCustomization_NewLandscape.h"
 #include "ScopedTransaction.h"
-#include "Landscape/Landscape.h"
+#include "Landscape.h"
 
 #define LOCTEXT_NAMESPACE "LandscapeEditor.NewLandscape"
 
diff --git a/Engine/Source/Editor/LandscapeEditor/Public/LandscapeEdMode.h b/Engine/Source/Editor/LandscapeEditor/Public/LandscapeEdMode.h
--- a/Engine/Source/Editor/LandscapeEditor/Public/LandscapeEdMode.h
+++ b/Engine/Source/Editor/LandscapeEditor/Public/LandscapeEdMode.h
@@ -13,10 +13,10 @@ class FLandscapeToolSplines;
 struct FHeightmapToolTarget;
 template<typename TargetType> class FLandscapeToolCopyPaste;
 
-#include "Landscape/LandscapeInfo.h"
-#include "Landscape/LandscapeLayerInfoObject.h"
-#include "Landscape/LandscapeProxy.h"
-#include "Landscape/LandscapeGizmoActiveActor.h"
+#include "LandscapeInfo.h"
+#include "LandscapeLayerInfoObject.h"
+#include "LandscapeProxy.h"
+#include "LandscapeGizmoActiveActor.h"
 
 // FLandscapeToolMousePosition - Struct to store mouse positions since the last time we applied the brush
 struct FLandscapeToolMousePosition
diff --git a/Engine/Source/Editor/StatsViewer/Private/StatsPages/PrimitiveStatsPage.cpp b/Engine/Source/Editor/StatsViewer/Private/StatsPages/PrimitiveStatsPage.cpp
--- a/Engine/Source/Editor/StatsViewer/Private/StatsPages/PrimitiveStatsPage.cpp
+++ b/Engine/Source/Editor/StatsViewer/Private/StatsPages/PrimitiveStatsPage.cpp
@@ -4,8 +4,8 @@
 #include "StaticMeshResources.h"
 #include "PrimitiveStatsPage.h"
 #include "PrimitiveStats.h"
-#include "Landscape/LandscapeProxy.h"
-#include "Landscape/LandscapeComponent.h"
+#include "LandscapeProxy.h"
+#include "LandscapeComponent.h"
 #include "LightMap.h"
 #include "ShadowMap.h"
 
diff --git a/Engine/Source/Editor/StatsViewer/StatsViewer.Build.cs b/Engine/Source/Editor/StatsViewer/StatsViewer.Build.cs
--- a/Engine/Source/Editor/StatsViewer/StatsViewer.Build.cs
+++ b/Engine/Source/Editor/StatsViewer/StatsViewer.Build.cs
@@ -17,13 +17,15 @@ public class StatsViewer : ModuleRules
 				"SlateCore",
                 "EditorStyle",
 				"RHI",
-				"UnrealEd"
+				"UnrealEd",
+				"Landscape"
 			}
 		);
 
         PrivateIncludePathModuleNames.AddRange(
 			new string[] {
-                "PropertyEditor"
+                "PropertyEditor",
+				"Landscape"
 			}
 		);
 
diff --git a/Engine/Source/Editor/UnrealEd/Private/AssetSelection.cpp b/Engine/Source/Editor/UnrealEd/Private/AssetSelection.cpp
--- a/Engine/Source/Editor/UnrealEd/Private/AssetSelection.cpp
+++ b/Engine/Source/Editor/UnrealEd/Private/AssetSelection.cpp
@@ -18,7 +18,7 @@
 #include "SnappingUtils.h"
 #include "ActorEditorUtils.h"
 #include "LevelEditorViewport.h"
-#include "Landscape/Landscape.h"
+#include "Landscape.h"
 
 #include "Particles/Emitter.h"
 #include "Editor/ActorPositioning.h"
diff --git a/Engine/Source/Editor/UnrealEd/Private/Editor/ActorPositioning.cpp b/Engine/Source/Editor/UnrealEd/Private/Editor/ActorPositioning.cpp
--- a/Engine/Source/Editor/UnrealEd/Private/Editor/ActorPositioning.cpp
+++ b/Engine/Source/Editor/UnrealEd/Private/Editor/ActorPositioning.cpp
@@ -3,8 +3,8 @@
 #include "UnrealEd.h"
 #include "ActorPositioning.h"
 #include "SnappingUtils.h"
-#include "Landscape/LandscapeHeightfieldCollisionComponent.h"
-#include "Landscape/LandscapeComponent.h"
+#include "LandscapeHeightfieldCollisionComponent.h"
+#include "LandscapeComponent.h"
 
 FActorPositionTraceResult FActorPositioning::TraceWorldForPositionWithDefault(const FViewportCursorLocation& Cursor, const FSceneView& View, const TArray<AActor*>* IgnoreActors)
 {
diff --git a/Engine/Source/Editor/UnrealEd/Private/EditorExporters.cpp b/Engine/Source/Editor/UnrealEd/Private/EditorExporters.cpp
--- a/Engine/Source/Editor/UnrealEd/Private/EditorExporters.cpp
+++ b/Engine/Source/Editor/UnrealEd/Private/EditorExporters.cpp
@@ -8,7 +8,6 @@
 #include "StaticMeshResources.h"
 #include "Materials/MaterialExpressionConstant.h"
 #include "SoundDefinitions.h"
-#include "Landscape/LandscapeDataAccess.h"
 #include "Kismet2/DebuggerCommands.h"
 #include "MaterialCompiler.h"
 #include "SurfaceIterators.h"
@@ -18,9 +17,10 @@
 #include "MaterialExportUtils.h"
 #include "ImageUtils.h"
 #include "Foliage/InstancedFoliageActor.h"
-#include "Landscape/Landscape.h"
-#include "Landscape/LandscapeComponent.h"
-#include "Landscape/LandscapeInfo.h"
+#include "Landscape.h"
+#include "LandscapeComponent.h"
+#include "LandscapeInfo.h"
+#include "LandscapeDataAccess.h"
 #include "Foliage/FoliageType.h"
 #include "UnrealExporter.h"
 #include "EngineModule.h"
diff --git a/Engine/Source/Editor/UnrealEd/Private/EditorServer.cpp b/Engine/Source/Editor/UnrealEd/Private/EditorServer.cpp
--- a/Engine/Source/Editor/UnrealEd/Private/EditorServer.cpp
+++ b/Engine/Source/Editor/UnrealEd/Private/EditorServer.cpp
@@ -22,7 +22,7 @@
 #include "../Private/GeomFitUtils.h"
 #include "Editor/GeometryMode/Public/GeometryEdMode.h"
 #include "Editor/GeometryMode/Public/EditorGeometry.h"
-#include "Landscape/LandscapeProxy.h"
+#include "LandscapeProxy.h"
 #include "Lightmass/PrecomputedVisibilityOverrideVolume.h"
 #include "Animation/AnimSet.h"
 #include "Matinee/InterpTrackAnimControl.h"
diff --git a/Engine/Source/Editor/UnrealEd/Private/Fbx/FbxMainExport.cpp b/Engine/Source/Editor/UnrealEd/Private/Fbx/FbxMainExport.cpp
--- a/Engine/Source/Editor/UnrealEd/Private/Fbx/FbxMainExport.cpp
+++ b/Engine/Source/Editor/UnrealEd/Private/Fbx/FbxMainExport.cpp
@@ -49,9 +49,9 @@
 
 #include "StaticMeshResources.h"
 #include "LandscapeDataAccess.h"
-#include "Landscape/LandscapeProxy.h"
-#include "Landscape/LandscapeInfo.h"
-#include "Landscape/LandscapeComponent.h"
+#include "LandscapeProxy.h"
+#include "LandscapeInfo.h"
+#include "LandscapeComponent.h"
 #include "Components/SplineMeshComponent.h"
 
 #include "FbxExporter.h"
diff --git a/Engine/Source/Editor/UnrealEd/Private/Kismet2/BlueprintEditorUtils.cpp b/Engine/Source/Editor/UnrealEd/Private/Kismet2/BlueprintEditorUtils.cpp
--- a/Engine/Source/Editor/UnrealEd/Private/Kismet2/BlueprintEditorUtils.cpp
+++ b/Engine/Source/Editor/UnrealEd/Private/Kismet2/BlueprintEditorUtils.cpp
@@ -1355,6 +1355,18 @@ void FBlueprintEditorUtils::PropagateParentBlueprintDefaults(UClass* ClassToProp
 	}
 }
 
+void DestroyGeneratedClass(TSubclassOf<UObject> &ClassToDestroy)
+{
+	if (ClassToDestroy != nullptr)
+	{
+		// We just need to rename the object, so the name won't collide.
+		ClassToDestroy->Rename(TEXT("None"), nullptr, REN_DontCreateRedirectors);
+
+		// The rest will do GC after we lose reference to it.
+		ClassToDestroy = nullptr;
+	}
+}
+
 void FBlueprintEditorUtils::PostDuplicateBlueprint(UBlueprint* Blueprint)
 {
 	// Only recompile after duplication if this isn't PIE
@@ -1378,9 +1390,13 @@ void FBlueprintEditorUtils::PostDuplicateBlueprint(UBlueprint* Blueprint)
 			TArray<UTimelineTemplate*> Timelines = Blueprint->Timelines;
 			Blueprint->Timelines.Empty();
 
-			// Null out the existing class references, the compile will create new ones
-			Blueprint->GeneratedClass = NULL;
-			Blueprint->SkeletonGeneratedClass = NULL;
+			// Delete the existing class references, the compile will create new ones
+			DestroyGeneratedClass(Blueprint->GeneratedClass);
+			DestroyGeneratedClass(Blueprint->SkeletonGeneratedClass);
+				
+
+			Blueprint->GeneratedClass = nullptr;
+			Blueprint->SkeletonGeneratedClass = nullptr;
 
 			// Make sure the new blueprint has a shiny new class
 			IKismetCompilerInterface& Compiler = FModuleManager::LoadModuleChecked<IKismetCompilerInterface>(KISMET_COMPILER_MODULENAME);
diff --git a/Engine/Source/Editor/UnrealEd/Private/Lightmass/Lightmass.cpp b/Engine/Source/Editor/UnrealEd/Private/Lightmass/Lightmass.cpp
--- a/Engine/Source/Editor/UnrealEd/Private/Lightmass/Lightmass.cpp
+++ b/Engine/Source/Editor/UnrealEd/Private/Lightmass/Lightmass.cpp
@@ -7,8 +7,8 @@
 #include "UnrealEd.h"
 #include "PrecomputedLightVolume.h"
 #include "Runtime/Engine/Public/StaticMeshResources.h"
-#include "Runtime/Engine/Public/Landscape/LandscapeRender.h"
-#include "Runtime/Engine/Public/Landscape/LandscapeLight.h"
+#include "LandscapeRender.h"
+#include "LandscapeLight.h"
 #include "Runtime/Engine/Classes/Matinee/MatineeActor.h"
 #include "Runtime/Engine/Classes/Matinee/InterpGroup.h"
 #include "Runtime/Engine/Classes/Matinee/InterpGroupInst.h"
@@ -54,8 +54,8 @@ void FSwarmDebugOptions::Touch()
 #include "Lightmass.h"
 #include "StaticMeshLight.h"
 #include "ModelLight.h"
-#include "Runtime/Engine/Public/Landscape/LandscapeLight.h"
-#include "Runtime/Engine/Public/Landscape/LandscapeDataAccess.h"
+#include "LandscapeLight.h"
+#include "LandscapeDataAccess.h"
 #include "Editor/StatsViewer/Public/StatsViewerModule.h"
 #include "Editor/StatsViewer/Classes/LightingBuildInfo.h"
 #include "MessageLog.h"
diff --git a/Engine/Source/Editor/UnrealEd/Private/StaticLightingSystem/StaticLightingExport.cpp b/Engine/Source/Editor/UnrealEd/Private/StaticLightingSystem/StaticLightingExport.cpp
--- a/Engine/Source/Editor/UnrealEd/Private/StaticLightingSystem/StaticLightingExport.cpp
+++ b/Engine/Source/Editor/UnrealEd/Private/StaticLightingSystem/StaticLightingExport.cpp
@@ -9,8 +9,9 @@
 #include "StaticMeshLight.h"
 #include "ModelLight.h"
 #include "StaticMeshResources.h"
-#include "Runtime/Engine/Public/Landscape/LandscapeLight.h"
-#include "Runtime/Engine/Classes/Landscape/LandscapeComponent.h"
+
+#include "LandscapeLight.h"
+#include "LandscapeComponent.h"
 
 #include "Lightmass/Lightmass.h"
 
diff --git a/Engine/Source/Editor/UnrealEd/Private/UnrealEdSrv.cpp b/Engine/Source/Editor/UnrealEd/Private/UnrealEdSrv.cpp
--- a/Engine/Source/Editor/UnrealEd/Private/UnrealEdSrv.cpp
+++ b/Engine/Source/Editor/UnrealEd/Private/UnrealEdSrv.cpp
@@ -27,12 +27,13 @@
 #include "Editor/ActorPositioning.h"
 #include "Matinee/InterpData.h"
 #include "Animation/SkeletalMeshActor.h"
-#include "Landscape/LandscapeInfo.h"
-#include "Landscape/LandscapeLayerInfoObject.h"
-#include "Landscape/LandscapeProxy.h"
-#include "Landscape/LandscapeGizmoActiveActor.h"
-#include "Landscape/LandscapeComponent.h"
-#include "Landscape/LandscapeHeightfieldCollisionComponent.h"
+#include "Landscape.h"
+#include "LandscapeInfo.h"
+#include "LandscapeLayerInfoObject.h"
+#include "LandscapeProxy.h"
+#include "LandscapeGizmoActiveActor.h"
+#include "LandscapeComponent.h"
+#include "LandscapeHeightfieldCollisionComponent.h"
 #include "ComponentReregisterContext.h"
 #include "JsonInternationalizationArchiveSerializer.h"
 #include "JsonInternationalizationManifestSerializer.h"
@@ -526,7 +527,9 @@ bool UUnrealEdEngine::HandleUpdateLandscapeEditorDataCommand( const TCHAR* Str,
 		}
 
 		// Fixed up for Landscape fix match case
-		for (auto It = InWorld->LandscapeInfoMap.CreateIterator(); It; ++It)
+		auto& LandscapeInfoMap = GetLandscapeInfoMap(InWorld);
+
+		for (auto It = LandscapeInfoMap.Map.CreateIterator(); It; ++It)
 		{
 			ULandscapeInfo* LandscapeInfo = It.Value();
 			if (LandscapeInfo)
@@ -579,7 +582,7 @@ bool UUnrealEdEngine::HandleUpdateLandscapeEditorDataCommand( const TCHAR* Str,
 		}
 
 		// Fix proxies relative transformations to LandscapeActor
-		for (auto It = InWorld->LandscapeInfoMap.CreateIterator(); It; ++It)
+		for (auto It = LandscapeInfoMap.Map.CreateIterator(); It; ++It)
 		{
 			ULandscapeInfo* Info = It.Value();
 			Info->FixupProxiesWeightmaps();
@@ -596,7 +599,7 @@ bool UUnrealEdEngine::HandleUpdateLandscapeMICCommand( const TCHAR* Str, FOutput
 
 	if (World && World->GetWorldSettings())
 	{
-		for (auto It = World->LandscapeInfoMap.CreateIterator(); It; ++It)
+		for (auto It = GetLandscapeInfoMap(World).Map.CreateIterator(); It; ++It)
 		{
 			ULandscapeInfo* Info = It.Value();
 			for (auto LandscapeComponentIt = Info->XYtoComponentMap.CreateIterator(); LandscapeComponentIt; ++LandscapeComponentIt )
@@ -618,7 +621,7 @@ bool UUnrealEdEngine::HandleRecreateLandscapeCollisionCommand(const TCHAR* Str,
 {
 	if (!PlayWorld && InWorld && InWorld->GetWorldSettings())
 	{
-		for (auto It = InWorld->LandscapeInfoMap.CreateIterator(); It; ++It)
+		for (auto It = GetLandscapeInfoMap(InWorld).Map.CreateIterator(); It; ++It)
 		{
 			ULandscapeInfo* Info = It.Value();
 			Info->RecreateCollisionComponents();
diff --git a/Engine/Source/Editor/UnrealEd/UnrealEd.Build.cs b/Engine/Source/Editor/UnrealEd/UnrealEd.Build.cs
--- a/Engine/Source/Editor/UnrealEd/UnrealEd.Build.cs
+++ b/Engine/Source/Editor/UnrealEd/UnrealEd.Build.cs
@@ -49,7 +49,8 @@ public class UnrealEd : ModuleRules
 				"ReferenceViewer",
                 "IntroTutorials",
                 "SuperSearch",
-				"OutputLog"
+				"OutputLog",
+				"Landscape"
 			}
 		);
 
@@ -84,7 +85,6 @@ public class UnrealEd : ModuleRules
 				"OnlineBlueprintSupport",
 				"DesktopPlatform",
 				"DirectoryWatcher",
-				"EditorSettingsViewer",
 				"EditorStyle",
 				"EngineSettings",
 				"InputCore",
@@ -115,6 +115,7 @@ public class UnrealEd : ModuleRules
 				"Kismet",
                 "InternationalizationSettings",
                 "JsonUtilities",
+				"Landscape",
 			}
 		);
 
diff --git a/Engine/Source/Editor/WorldBrowser/Private/Tiles/STiledLandscapeImportDlg.cpp b/Engine/Source/Editor/WorldBrowser/Private/Tiles/STiledLandscapeImportDlg.cpp
--- a/Engine/Source/Editor/WorldBrowser/Private/Tiles/STiledLandscapeImportDlg.cpp
+++ b/Engine/Source/Editor/WorldBrowser/Private/Tiles/STiledLandscapeImportDlg.cpp
@@ -6,7 +6,7 @@
 #include "IDesktopPlatform.h"
 #include "DesktopPlatformModule.h"
 #include "ContentBrowserModule.h"
-#include "Landscape/LandscapeProxy.h"
+#include "LandscapeProxy.h"
 #include "SNumericEntryBox.h"
 
 #define LOCTEXT_NAMESPACE "WorldBrowser"
diff --git a/Engine/Source/Editor/WorldBrowser/WorldBrowser.Build.cs b/Engine/Source/Editor/WorldBrowser/WorldBrowser.Build.cs
--- a/Engine/Source/Editor/WorldBrowser/WorldBrowser.Build.cs
+++ b/Engine/Source/Editor/WorldBrowser/WorldBrowser.Build.cs
@@ -13,6 +13,7 @@ public class WorldBrowser : ModuleRules
 				"AssetTools",
                 "MeshUtilities",
                 "ContentBrowser",
+				"Landscape",
 			}
         );
      
@@ -24,6 +25,7 @@ public class WorldBrowser : ModuleRules
                 "RenderCore",
                 "InputCore",
                 "Engine",
+				"Landscape",
                 "Slate",
 				"SlateCore",
                 "EditorStyle",
diff --git a/Engine/Source/Programs/UnrealBuildTool/Configuration/UEBuildModule.cs b/Engine/Source/Programs/UnrealBuildTool/Configuration/UEBuildModule.cs
--- a/Engine/Source/Programs/UnrealBuildTool/Configuration/UEBuildModule.cs
+++ b/Engine/Source/Programs/UnrealBuildTool/Configuration/UEBuildModule.cs
@@ -963,6 +963,9 @@ namespace UnrealBuildTool
 		/** If true and unity builds are enabled, this module will build without unity. */
 		public bool bFasterWithoutUnity = false;
 
+		/** If true then the engine will call it's StartupModule at engine initialization automatically. */
+		public bool bIsAutoStartupModule = false;
+
 		/** Overrides BuildConfiguration.MinFilesUsingPrecompiledHeader if non-zero. */
 		public int MinFilesUsingPrecompiledHeaderOverride = 0;
 
@@ -1054,6 +1057,7 @@ namespace UnrealBuildTool
 			bool InUseRTTI,
 			bool InEnableBufferSecurityChecks,
 			bool InFasterWithoutUnity,
+			bool InIsAutoStartupModule,
 			int InMinFilesUsingPrecompiledHeaderOverride,
 			bool InEnableExceptions,
 			bool bInBuildSourceFiles
@@ -1104,6 +1108,7 @@ namespace UnrealBuildTool
 			bUseRTTI                               = InUseRTTI;
 			bEnableBufferSecurityChecks 		   = InEnableBufferSecurityChecks;
 			bFasterWithoutUnity                    = InFasterWithoutUnity;
+			bIsAutoStartupModule                   = InIsAutoStartupModule;
 			MinFilesUsingPrecompiledHeaderOverride = InMinFilesUsingPrecompiledHeaderOverride;
 			bEnableExceptions                      = InEnableExceptions;
 		}
@@ -2133,6 +2138,7 @@ namespace UnrealBuildTool
 			bool InUseRTTI,
 			bool InEnableBufferSecurityChecks,
 			bool InFasterWithoutUnity,
+			bool InIsAutoStartupModule,
 			int InMinFilesUsingPrecompiledHeaderOverride,
 			bool InEnableExceptions,
 			bool bInBuildSourceFiles
@@ -2142,7 +2148,7 @@ namespace UnrealBuildTool
 			InPublicIncludePathModuleNames,InPublicDependencyModuleNames,InPublicDelayLoadDLLs,InPublicAdditionalLibraries,InPublicFrameworks,InPublicWeakFrameworks,InPublicAdditionalFrameworks,InPublicAdditionalShadowFiles,InPublicAdditionalBundleResources,
 			InPrivateIncludePaths,InPrivateIncludePathModuleNames,InPrivateDependencyModuleNames,
             InCircularlyReferencedDependentModules, InDynamicallyLoadedModuleNames, InPlatformSpecificDynamicallyLoadedModuleNames, InOptimizeCode,
-			InAllowSharedPCH, InSharedPCHHeaderFile, InUseRTTI, InEnableBufferSecurityChecks, InFasterWithoutUnity, InMinFilesUsingPrecompiledHeaderOverride,
+			InAllowSharedPCH, InSharedPCHHeaderFile, InUseRTTI, InEnableBufferSecurityChecks, InFasterWithoutUnity, InIsAutoStartupModule, InMinFilesUsingPrecompiledHeaderOverride,
 			InEnableExceptions, bInBuildSourceFiles)
 		{
 			PrivateAssemblyReferences = HashSetFromOptionalEnumerableStringParameter(InPrivateAssemblyReferences);
diff --git a/Engine/Source/Programs/UnrealBuildTool/Configuration/UEBuildTarget.cs b/Engine/Source/Programs/UnrealBuildTool/Configuration/UEBuildTarget.cs
--- a/Engine/Source/Programs/UnrealBuildTool/Configuration/UEBuildTarget.cs
+++ b/Engine/Source/Programs/UnrealBuildTool/Configuration/UEBuildTarget.cs
@@ -1508,6 +1508,8 @@ namespace UnrealBuildTool
 				CreateLinkerFixupsCPPFile();
 			}
 
+			CreateAutoStartupModuleListGetter();
+
 			// Build the target's binaries.
 			foreach (var Binary in AppBinaries)
 			{
@@ -1797,6 +1799,71 @@ namespace UnrealBuildTool
 			return FilteredBinaries;
 		}
 
+		/// <summary>
+		/// Gets auto-startup module list.
+		/// </summary>
+		/// <returns>Auto-startup module name list.</returns>
+		private IEnumerable<string> GetAutoStartupModuleList()
+		{
+			var Output = new List<string>();
+
+			foreach(var Module in AppBinaries[0].GetAllDependencyModules(false, true))
+			{
+				if (Module is UEBuildModuleCPP && (Module as UEBuildModuleCPP).bIsAutoStartupModule)
+				{
+					Output.Add(Module.Name);
+				}
+			}
+
+			return Output;
+		}
+
+		/// <summary>
+		/// Creates a getter module that outputs auto-startup module list.
+		/// </summary>
+		private void CreateAutoStartupModuleListGetter()
+		{
+			var ModuleName = "AutoStartupModuleListGetter";
+
+			string SourceFilename = Path.Combine(GlobalCompileEnvironment.Config.OutputDirectory, ModuleName + ".cpp");
+			string HeaderFilename = Path.Combine(GlobalCompileEnvironment.Config.OutputDirectory, ModuleName + ".h");
+
+			if (!File.Exists(HeaderFilename))
+			{
+				File.WriteAllText(HeaderFilename, "");
+			}
+
+			var SourceLines = new List<string>();
+
+			SourceLines.Add(string.Format("#include \"{0}\"", ModuleName + ".h"));
+			SourceLines.Add("");
+			SourceLines.Add("#include \"Core.h\"");
+			SourceLines.Add("#include \"Array.h\"");
+			SourceLines.Add("#include \"UnrealString.h\"");
+			SourceLines.Add("");
+			SourceLines.Add("void GetAutoStartupModuleList(TArray<FString>& Out)");
+			SourceLines.Add("{");
+
+			foreach(var AutoStartupModuleName in GetAutoStartupModuleList())
+			{
+				SourceLines.Add(string.Format("\tOut.Add(TEXT(\"{0}\"));", AutoStartupModuleName));
+			}
+
+			SourceLines.Add("}");
+			SourceLines.Add("");
+
+			if (!File.Exists(SourceFilename) || File.ReadAllText(SourceFilename).Trim() != string.Join("\r\n", SourceLines).Trim())
+			{
+				File.WriteAllLines(SourceFilename, SourceLines);
+			}
+
+			BindArtificialModuleToBinary(
+				CreateArtificialModule(
+					ModuleName, GlobalCompileEnvironment.Config.OutputDirectory,
+					new FileItem[] { FileItem.GetItemByPath(SourceFilename) }, new string[] { "Core" }
+				), AppBinaries[0]);
+		}
+
 		/// <summary>
 		/// All non-program monolithic binaries implicitly depend on all static plugin libraries so they are always linked appropriately
 		/// In order to do this, we create a new module here with a cpp file we emit that invokes an empty function in each library.
@@ -1901,54 +1968,77 @@ namespace UnrealBuildTool
 
 				// Create the CPP module
 				var FakeModuleDirectory = Path.GetDirectoryName( LinkerFixupCPPFilename );
-				UEBuildModuleCPP NewModule = new UEBuildModuleCPP(
-					InTarget: this,
-					InName: LinkerFixupsName,
-					InType: UEBuildModuleType.Game,
-					InModuleDirectory:FakeModuleDirectory,
-					InOutputDirectory: GlobalCompileEnvironment.Config.OutputDirectory,
-					InIsRedistributableOverride: null,
-					InIntelliSenseGatherer: null,
-					InSourceFiles: SourceFiles,
-					InPublicIncludePaths: new List<string>(),
-					InPublicSystemIncludePaths: new List<string>(),
-					InDefinitions: new List<string>(),
-					InPublicIncludePathModuleNames: new List<string>(),
-					InPublicDependencyModuleNames: new List<string>(),
-					InPublicDelayLoadDLLs: new List<string>(),
-					InPublicAdditionalLibraries: new List<string>(),
-					InPublicFrameworks: new List<string>(),
-					InPublicWeakFrameworks: new List<string>(),
-					InPublicAdditionalShadowFiles: new List<string>(),
-					InPublicAdditionalBundleResources: new List<UEBuildBundleResource>(),
-					InPublicAdditionalFrameworks: new List<UEBuildFramework>(),
-					InPrivateIncludePaths: new List<string>(),
-					InPrivateIncludePathModuleNames: new List<string>(),
-					InPrivateDependencyModuleNames: PrivateDependencyModuleNames,
-					InCircularlyReferencedDependentModules: new List<string>(),
-					InDynamicallyLoadedModuleNames: new List<string>(),
-                    InPlatformSpecificDynamicallyLoadedModuleNames: new List<string>(),
-                    InOptimizeCode: ModuleRules.CodeOptimization.Default,
-					InAllowSharedPCH: false,
-					InSharedPCHHeaderFile: "",
-					InUseRTTI: false,
-					InEnableBufferSecurityChecks: true,
-					InFasterWithoutUnity: true,
-					InMinFilesUsingPrecompiledHeaderOverride: 0,
-					InEnableExceptions: false,
-					bInBuildSourceFiles: true
-					);
+				var NewModule = CreateArtificialModule(LinkerFixupsName, FakeModuleDirectory, SourceFiles, PrivateDependencyModuleNames);
 
 				// Now bind this new module to the executable binary so it will link the plugin libs correctly
-				NewModule.Binary = ExecutableBinary;
-				NewModule.bIncludedInTarget = true;
+				BindArtificialModuleToBinary(NewModule, ExecutableBinary);
+			}
+		}
+
+		/// <summary>
+		/// Binds artificial module to given binary.
+		/// </summary>
+		/// <param name="Module">Module to bind.</param>
+		/// <param name="Binary">Binary to bind.</param>
+		private void BindArtificialModuleToBinary(UEBuildModuleCPP Module, UEBuildBinary Binary)
+		{
+			Module.Binary = Binary;
+			Module.bIncludedInTarget = true;
 
-				// Process dependencies for this new module
-				NewModule.CachePCHUsageForModuleSourceFiles(NewModule.CreateModuleCompileEnvironment(GlobalCompileEnvironment));
+			// Process dependencies for this new module
+			Module.CachePCHUsageForModuleSourceFiles(Module.CreateModuleCompileEnvironment(GlobalCompileEnvironment));
 
-				// Add module to binary
-				ExecutableBinary.AddModule(NewModule.Name);
-			}
+			// Add module to binary
+			Binary.AddModule(Module.Name);
+		}
+
+		/// <summary>
+		/// Creates artificial module.
+		/// </summary>
+		/// <param name="Name">Name of the module.</param>
+		/// <param name="Directory">Directory of the module.</param>
+		/// <param name="SourceFiles">Source files.</param>
+		/// <param name="PrivateDependencyModuleNames">Private dependency list.</param>
+		/// <returns>Created module.</returns>
+		private UEBuildModuleCPP CreateArtificialModule(string Name, string Directory, IEnumerable<FileItem> SourceFiles, IEnumerable<string> PrivateDependencyModuleNames)
+		{
+			return new UEBuildModuleCPP(
+				InTarget: this,
+				InName: Name,
+				InType: UEBuildModuleType.Game,
+				InModuleDirectory: Directory,
+				InOutputDirectory: GlobalCompileEnvironment.Config.OutputDirectory,
+				InIsRedistributableOverride: null,
+				InIntelliSenseGatherer: null,
+				InSourceFiles: SourceFiles.ToList(),
+				InPublicIncludePaths: new List<string>(),
+				InPublicSystemIncludePaths: new List<string>(),
+				InDefinitions: new List<string>(),
+				InPublicIncludePathModuleNames: new List<string>(),
+				InPublicDependencyModuleNames: new List<string>(),
+				InPublicDelayLoadDLLs: new List<string>(),
+				InPublicAdditionalLibraries: new List<string>(),
+				InPublicFrameworks: new List<string>(),
+				InPublicWeakFrameworks: new List<string>(),
+				InPublicAdditionalShadowFiles: new List<string>(),
+				InPublicAdditionalBundleResources: new List<UEBuildBundleResource>(),
+				InPublicAdditionalFrameworks: new List<UEBuildFramework>(),
+				InPrivateIncludePaths: new List<string>(),
+				InPrivateIncludePathModuleNames: new List<string>(),
+				InPrivateDependencyModuleNames: PrivateDependencyModuleNames.ToList(),
+				InCircularlyReferencedDependentModules: new List<string>(),
+				InDynamicallyLoadedModuleNames: new List<string>(),
+				InPlatformSpecificDynamicallyLoadedModuleNames: new List<string>(),
+				InOptimizeCode: ModuleRules.CodeOptimization.Default,
+				InAllowSharedPCH: false,
+				InSharedPCHHeaderFile: "",
+				InUseRTTI: false,
+				InEnableBufferSecurityChecks: true,
+				InFasterWithoutUnity: true,
+				InIsAutoStartupModule: false,
+				InMinFilesUsingPrecompiledHeaderOverride: 0,
+				InEnableExceptions: false,
+				bInBuildSourceFiles: true);
 		}
 
 
@@ -2881,6 +2971,7 @@ namespace UnrealBuildTool
 							InUseRTTI: RulesObject.bUseRTTI,
 							InEnableBufferSecurityChecks: RulesObject.bEnableBufferSecurityChecks,
 							InFasterWithoutUnity: RulesObject.bFasterWithoutUnity,
+							InIsAutoStartupModule: RulesObject.bIsAutoStartupModule,
 							InMinFilesUsingPrecompiledHeaderOverride: RulesObject.MinFilesUsingPrecompiledHeaderOverride,
 							InEnableExceptions: RulesObject.bEnableExceptions,
 							bInBuildSourceFiles: bBuildSourceFiles
@@ -2921,6 +3012,7 @@ namespace UnrealBuildTool
 							InUseRTTI: RulesObject.bUseRTTI,
 							InEnableBufferSecurityChecks: RulesObject.bEnableBufferSecurityChecks,
 							InFasterWithoutUnity: RulesObject.bFasterWithoutUnity,
+							InIsAutoStartupModule: RulesObject.bIsAutoStartupModule,
 							InMinFilesUsingPrecompiledHeaderOverride: RulesObject.MinFilesUsingPrecompiledHeaderOverride,
 							InEnableExceptions: RulesObject.bEnableExceptions,
 							bInBuildSourceFiles : bBuildSourceFiles
diff --git a/Engine/Source/Programs/UnrealBuildTool/System/RulesCompiler.cs b/Engine/Source/Programs/UnrealBuildTool/System/RulesCompiler.cs
--- a/Engine/Source/Programs/UnrealBuildTool/System/RulesCompiler.cs
+++ b/Engine/Source/Programs/UnrealBuildTool/System/RulesCompiler.cs
@@ -155,6 +155,9 @@ namespace UnrealBuildTool
 		/** If true and unity builds are enabled, this module will build without unity. */
 		public bool bFasterWithoutUnity = false;
 
+		/** If true then the engine will call it's StartupModule at engine initialization automatically. */
+		public bool bIsAutoStartupModule = false;
+
 		/** Overrides BuildConfiguration.MinFilesUsingPrecompiledHeader if non-zero. */
 		public int MinFilesUsingPrecompiledHeaderOverride = 0;
 
diff --git a/Engine/Source/Runtime/CoreUObject/Private/Serialization/DuplicateDataWriter.cpp b/Engine/Source/Runtime/CoreUObject/Private/Serialization/DuplicateDataWriter.cpp
--- a/Engine/Source/Runtime/CoreUObject/Private/Serialization/DuplicateDataWriter.cpp
+++ b/Engine/Source/Runtime/CoreUObject/Private/Serialization/DuplicateDataWriter.cpp
@@ -30,6 +30,9 @@ FDuplicateDataWriter::FDuplicateDataWriter( FUObjectAnnotationSparse<FDuplicated
 	ArAllowLazyLoading	= false;
 	ArPortFlags |= PPF_Duplicate | InPortFlags;
 
+	SourcePackage = SourceObject->GetOutermost();
+	DestPackage = DestObject->GetOutermost();
+
 	AddDuplicate(SourceObject,DestObject);
 }
 
@@ -121,14 +124,43 @@ UObject* FDuplicateDataWriter::GetDuplicatedObject( UObject* Object )
 		}
 		else
 		{
+			UObject* NewEmptyDuplicate = nullptr;
+
 			// Check to see if the object's outer is being duplicated.
 			UObject* DupOuter = GetDuplicatedObject(Object->GetOuter());
 			if(DupOuter != NULL)
 			{
-				// The object's outer is being duplicated, create a duplicate of this object.
-				UObject* NewEmptyDuplicate = StaticConstructObject(Object->GetClass(), DupOuter, Object->GetFName(), ApplyFlags|Object->GetMaskedFlags(FlagMask),
+				// The object's outer is being duplicated, create a duplicate
+				// of this object.
+				NewEmptyDuplicate = StaticConstructObject(
+					Object->GetClass(),
+					DupOuter,
+					Object->GetFName(),
+					ApplyFlags|Object->GetMaskedFlags(FlagMask),
 					Object->GetArchetype(), true, InstanceGraph);
 
+			}
+			else if (
+				!Object->HasAnyFlags(RF_ClassDefaultObject) // We shouldn't duplicate CDOs
+				&& (ArPortFlags & PPF_DuplicateForPIE) == 0 // We can also skip this if duplicating for PIE -- this is mainly performance optimization
+				&& SourcePackage != DestPackage
+				&& Object->GetOuter() == SourcePackage
+				&& DestPackage != GetTransientPackage())
+			{
+				// If not then maybe it's a package changing duplication and
+				// the object's outer is the source package itself. Then we
+				// need to duplicate that object also with changed package.
+
+				NewEmptyDuplicate = StaticConstructObject(
+					Object->GetClass(),
+					DestPackage,
+					Object->GetFName(),
+					ApplyFlags | Object->GetMaskedFlags(FlagMask),
+					nullptr, true, InstanceGraph);
+			}
+
+			if (NewEmptyDuplicate != nullptr)
+			{
 				Result = AddDuplicate(Object, NewEmptyDuplicate);
 			}
 		}
diff --git a/Engine/Source/Runtime/CoreUObject/Public/Serialization/ArchiveUObject.h b/Engine/Source/Runtime/CoreUObject/Public/Serialization/ArchiveUObject.h
--- a/Engine/Source/Runtime/CoreUObject/Public/Serialization/ArchiveUObject.h
+++ b/Engine/Source/Runtime/CoreUObject/Public/Serialization/ArchiveUObject.h
@@ -6,6 +6,8 @@
 
 #pragma once
 
+struct FObjectInstancingGraph;
+
 /**
  * Archive for counting memory usage.
  */
@@ -922,6 +924,18 @@ private:
 	EObjectFlags							FlagMask;
 	EObjectFlags							ApplyFlags;
 
+	/**
+	 * A pointer to source package, i.e. to the package the object being
+	 * duplicated is comming from.
+	 */
+	UPackage								*SourcePackage;
+
+	/**
+	 * A pointer to destination package, i.e. to the package the new object
+	 * should be created in.
+	 */
+	UPackage								*DestPackage;
+
 	/**
 	 * This is used to prevent object & component instancing resulting from the calls to StaticConstructObject(); instancing subobjects and components is pointless,
 	 * since we do that manually and replace the current value with our manually created object anyway.
@@ -1004,7 +1018,15 @@ public:
 	 * @param	InApplyFlags			the flags that should always be set on the duplicated objects (regardless of whether they're set on the source)
 	 * @param	InInstanceGraph			the instancing graph to use when creating the duplicate objects.
 	 */
-	FDuplicateDataWriter( class FUObjectAnnotationSparse<FDuplicatedObject,false>& InDuplicatedObjects,TArray<uint8>& InObjectData,UObject* SourceObject,UObject* DestObject,EObjectFlags InFlagMask,EObjectFlags InApplyMask,struct FObjectInstancingGraph* InInstanceGraph, uint32 InPortFlags );
+	FDuplicateDataWriter(
+		FUObjectAnnotationSparse<FDuplicatedObject, false>& InDuplicatedObjects,
+		TArray<uint8>& InObjectData,
+		UObject* SourceObject,
+		UObject* DestObject,
+		EObjectFlags InFlagMask,
+		EObjectFlags InApplyMask,
+		FObjectInstancingGraph* InInstanceGraph,
+		uint32 InPortFlags);
 };
 
 /*----------------------------------------------------------------------------
diff --git a/Engine/Source/Runtime/Engine/Classes/Components/InstancedStaticMeshComponent.h b/Engine/Source/Runtime/Engine/Classes/Components/InstancedStaticMeshComponent.h
--- a/Engine/Source/Runtime/Engine/Classes/Components/InstancedStaticMeshComponent.h
+++ b/Engine/Source/Runtime/Engine/Classes/Components/InstancedStaticMeshComponent.h
@@ -1,6 +1,9 @@
 // Copyright 1998-2014 Epic Games, Inc. All Rights Reserved.
+
 #pragma once
 
+#include "Components/StaticMeshComponent.h"
+
 #include "InstancedStaticMeshComponent.generated.h"
 
 USTRUCT()
diff --git a/Engine/Source/Runtime/Engine/Classes/Components/SplineMeshComponent.h b/Engine/Source/Runtime/Engine/Classes/Components/SplineMeshComponent.h
--- a/Engine/Source/Runtime/Engine/Classes/Components/SplineMeshComponent.h
+++ b/Engine/Source/Runtime/Engine/Classes/Components/SplineMeshComponent.h
@@ -1,8 +1,11 @@
 // Copyright 1998-2014 Epic Games, Inc. All Rights Reserved.
 
 #pragma once
-#include "SplineMeshComponent.generated.h"
 
+#include "Components/StaticMeshComponent.h"
+#include "Interfaces/Interface_CollisionDataProvider.h"
+
+#include "SplineMeshComponent.generated.h"
 
 UENUM()
 namespace ESplineMeshAxis
diff --git a/Engine/Source/Runtime/Engine/Classes/Components/StaticMeshComponent.h b/Engine/Source/Runtime/Engine/Classes/Components/StaticMeshComponent.h
--- a/Engine/Source/Runtime/Engine/Classes/Components/StaticMeshComponent.h
+++ b/Engine/Source/Runtime/Engine/Classes/Components/StaticMeshComponent.h
@@ -10,6 +10,8 @@
 
 class FColorVertexBuffer;
 class UStaticMesh;
+class FStaticMeshStaticLightingMesh;
+class ULightComponent;
 
 /** Cached vertex information at the time the mesh was painted. */
 USTRUCT()
@@ -399,7 +401,7 @@ public:
 	void ReleaseResources();
 
 	/** Allocates an implementation of FStaticLightingMesh that will handle static lighting for this component */
-	virtual class FStaticMeshStaticLightingMesh* AllocateStaticLightingMesh(int32 LODIndex, const TArray<ULightComponent*>& InRelevantLights);
+	virtual FStaticMeshStaticLightingMesh* AllocateStaticLightingMesh(int32 LODIndex, const TArray<ULightComponent*>& InRelevantLights);
 
 	/** Add or remove elements to have the size in the specified range. Reconstructs elements if MaxSize<MinSize */
 	void SetLODDataCount( const uint32 MinSize, const uint32 MaxSize );
diff --git a/Engine/Source/Runtime/Engine/Classes/Engine/Texture2D.h b/Engine/Source/Runtime/Engine/Classes/Engine/Texture2D.h
--- a/Engine/Source/Runtime/Engine/Classes/Engine/Texture2D.h
+++ b/Engine/Source/Runtime/Engine/Classes/Engine/Texture2D.h
@@ -231,7 +231,7 @@ public:
 
 private:
 	/** The minimum number of mips that must be resident in memory (cannot be streamed). */
-	static int32 GMinTextureResidentMipCount;
+	static ENGINE_API int32 GMinTextureResidentMipCount;
 
 public:
 	/** Returns the minimum number of mips that must be resident in memory (cannot be streamed). */
@@ -405,12 +405,12 @@ public:
 	 *	@param SrcData - the source data
 	 *  @param bFreeData - if true, the SrcData and Regions pointers will be freed after the update.
 	 */
-	void UpdateTextureRegions( int32 MipIndex, uint32 NumRegions, FUpdateTextureRegion2D* Regions, uint32 SrcPitch, uint32 SrcBpp, uint8* SrcData, bool bFreeData );
+	ENGINE_API void UpdateTextureRegions(int32 MipIndex, uint32 NumRegions, FUpdateTextureRegion2D* Regions, uint32 SrcPitch, uint32 SrcBpp, uint8* SrcData, bool bFreeData);
 
 	/**
 	 * Temporarily disable streaming so we update subregions of this texture without streaming clobbering it. 
 	 */
-	void TemporarilyDisableStreaming();
+	ENGINE_API void TemporarilyDisableStreaming();
 
 	/** Called after an editor or undo operation is formed on texture
 	*/
diff --git a/Engine/Source/Runtime/Engine/Classes/Engine/World.h b/Engine/Source/Runtime/Engine/Classes/Engine/World.h
--- a/Engine/Source/Runtime/Engine/Classes/Engine/World.h
+++ b/Engine/Source/Runtime/Engine/Classes/Engine/World.h
@@ -532,6 +532,14 @@ class ENGINE_API UWorld : public UObject, public FNetworkNotify
 	UPROPERTY(Transient)
 	TArray<UObject*>							ExtraReferencedObjects;
 
+	/**
+	 * External modules can have additional data associated with this UWorld.
+	 * This is a list of per module world data objects. These aren't
+	 * loaded/saved by default.
+	 */
+	UPROPERTY(Transient)
+	TArray<UObject*>							PerModuleDataObjects;
+
 	/** Level collection. ULevels are referenced by FName (Package name) to avoid serialized references. Also contains offsets in world units */
 	UPROPERTY(Transient)
 	TArray<class ULevelStreaming*>				StreamingLevels;
@@ -908,11 +916,6 @@ public:
 	/** @todo document */
 	FName CommittedPersistentLevelName;
 
-#if WITH_EDITORONLY_DATA
-	/** Map of LandscapeInfos for all loaded levels, valid in the editor only */
-	TMap< FGuid, class ULandscapeInfo* > LandscapeInfoMap;
-#endif // WITH_EDITORONLY_DATA
-
 	/**
 	 * This is a int on the level which is set when a light that needs to have lighting rebuilt
 	 * is moved.  This is then checked in CheckMap for errors to let you know that this level should
@@ -2446,9 +2449,6 @@ public:
 	/** Gets all LightMaps and ShadowMaps associated with this world. Specify the level or leave null for persistent */
 	void GetLightMapsAndShadowMaps(ULevel* Level, TArray<UTexture2D*>& OutLightMapsAndShadowMaps);
 
-	/** Gets all textures and materials used by all landscape components in the specified level */
-	void GetLandscapeTexturesAndMaterials(ULevel* Level, TArray<UObject*>& OutTexturesAndMaterials);
-
 public:
 	static FString ConvertToPIEPackageName(const FString& PackageName, int32 PIEInstanceID);
 	static FString StripPIEPrefixFromPackageName(const FString& PackageName, const FString& Prefix);
@@ -2471,6 +2471,7 @@ public:
 	DECLARE_MULTICAST_DELEGATE_TwoParams(FWorldInitializationEvent, UWorld* /*World*/, const UWorld::InitializationValues /*IVS*/);
 	DECLARE_MULTICAST_DELEGATE_ThreeParams(FWorldCleanupEvent, UWorld* /*World*/, bool /*bSessionEnded*/, bool /*bCleanupResources*/);
 	DECLARE_MULTICAST_DELEGATE_OneParam(FWorldEvent, UWorld* /*World*/);
+	DECLARE_MULTICAST_DELEGATE_TwoParams(FWorldPostDuplicateEvent, UWorld* /*World*/, bool /*bDuplicateForPIE*/);
 	// Delegate type for level change events
 	DECLARE_MULTICAST_DELEGATE_TwoParams(FOnLevelChanged, ULevel*, UWorld*);
 	
@@ -2480,6 +2481,9 @@ public:
 	// Callback for world initialization (post)
 	static FWorldInitializationEvent OnPostWorldInitialization;
 
+	// Post duplication event.
+	static FWorldPostDuplicateEvent OnPostDuplicate;
+
 	// Callback for world cleanup
 	static FWorldCleanupEvent OnWorldCleanup;
 
diff --git a/Engine/Source/Runtime/Engine/Classes/Foliage/InstancedFoliageActor.h b/Engine/Source/Runtime/Engine/Classes/Foliage/InstancedFoliageActor.h
--- a/Engine/Source/Runtime/Engine/Classes/Foliage/InstancedFoliageActor.h
+++ b/Engine/Source/Runtime/Engine/Classes/Foliage/InstancedFoliageActor.h
@@ -2,13 +2,15 @@
 
 
 #pragma once
+
 #include "GameFramework/Actor.h"
 #include "Templates/UniqueObj.h"
+#include "Components/InstancedStaticMeshComponent.h"
+
 #include "InstancedFoliageActor.generated.h"
 
 // Forward declarations
 class UFoliageType;
-class ULandscapeHeightfieldCollisionComponent;
 struct FFoliageInstancePlacementInfo;
 struct FFoliageMeshInfo;
 
@@ -43,9 +45,6 @@ public:
 	// Called in response to FEditorDelegates::MapChange.Broadcast(MapChangeEventFlags::MapRebuild ) so BSP rebuilds migrate foliage from obsolete to new components.
 	ENGINE_API void MapRebuild();
 
-	// Called from editor code to manage instances for components
-	ENGINE_API void SnapInstancesForLandscape(ULandscapeHeightfieldCollisionComponent* InComponent, const FBox& InInstanceBox);
-
 	// Moves instances based on the specified component to the current streaming level
 	ENGINE_API void MoveInstancesForComponentToCurrentLevel(UActorComponent* InComponent);
 
diff --git a/Engine/Source/Runtime/Engine/Classes/Materials/Material.h b/Engine/Source/Runtime/Engine/Classes/Materials/Material.h
--- a/Engine/Source/Runtime/Engine/Classes/Materials/Material.h
+++ b/Engine/Source/Runtime/Engine/Classes/Materials/Material.h
@@ -843,9 +843,6 @@ public:
 private:
 	void BackwardsCompatibilityInputConversion();
 
-	/** Regenerate expression guids for legacy terrain layer nodes. */
-	void FixupTerrainLayerWeightNodes();
-
 	/** Handles setting up an annotation for this object if a flag has changed value */
 	void MarkUsageFlagDirty(EMaterialUsage Usage, bool CurrentValue, bool NewValue);
 
@@ -891,7 +888,21 @@ public:
 	 * @return	Returns a array of parameter names used in this material for the specified expression type.
 	 */
 	template<typename ExpressionType>
-	void GetAllParameterNames(TArray<FName> &OutParameterNames, TArray<FGuid> &OutParameterIds);
+	void GetAllParameterNames(TArray<FName> &OutParameterNames, TArray<FGuid> &OutParameterIds)
+	{
+		for (int32 ExpressionIndex = 0; ExpressionIndex < Expressions.Num(); ExpressionIndex++)
+		{
+			ExpressionType* ParameterExpression =
+				Cast<ExpressionType>(Expressions[ExpressionIndex]);
+
+			if (ParameterExpression)
+			{
+				ParameterExpression->GetAllParameterNames(OutParameterNames, OutParameterIds);
+			}
+		}
+
+		check(OutParameterNames.Num() == OutParameterIds.Num());
+	}
 	
 	ENGINE_API void GetAllVectorParameterNames(TArray<FName> &OutParameterNames, TArray<FGuid> &OutParameterIds);
 	ENGINE_API void GetAllScalarParameterNames(TArray<FName> &OutParameterNames, TArray<FGuid> &OutParameterIds);
@@ -899,7 +910,6 @@ public:
 	ENGINE_API void GetAllFontParameterNames(TArray<FName> &OutParameterNames, TArray<FGuid> &OutParameterIds);
 	void GetAllStaticSwitchParameterNames(TArray<FName> &OutParameterNames, TArray<FGuid> &OutParameterIds);
 	void GetAllStaticComponentMaskParameterNames(TArray<FName> &OutParameterNames, TArray<FGuid> &OutParameterIds);
-	void GetAllTerrainLayerWeightParameterNames(TArray<FName> &OutParameterNames, TArray<FGuid> &OutParameterIds);
 
 	/** Returns the material's decal blend mode, calculated from the DecalBlendMode property and what inputs are connected. */
 	uint32 GetDecalBlendMode() const { return DecalBlendMode; }
diff --git a/Engine/Source/Runtime/Engine/Classes/Materials/MaterialExpression.h b/Engine/Source/Runtime/Engine/Classes/Materials/MaterialExpression.h
--- a/Engine/Source/Runtime/Engine/Classes/Materials/MaterialExpression.h
+++ b/Engine/Source/Runtime/Engine/Classes/Materials/MaterialExpression.h
@@ -1,10 +1,13 @@
 // Copyright 1998-2014 Epic Games, Inc. All Rights Reserved.
 
-
 #pragma once
+
+#include "MaterialExpressionIO.h"
+
 #include "MaterialExpression.generated.h"
 
 class FMaterialExpressionKey;
+class UEdGraphNode;
 
 //@warning: FExpressionInput is mirrored in MaterialShared.h and manually "subclassed" in Material.h (FMaterialInput)
 #if !CPP      //noexport struct
@@ -92,7 +95,7 @@ class ENGINE_API UMaterialExpression : public UObject
 	int32 MaterialExpressionEditorY;
 
 	/** Expression's Graph representation */
-	class UEdGraphNode*	GraphNode;
+	UEdGraphNode*	GraphNode;
 
 	/** Text of last error for this expression */
 	FString LastErrorText;
@@ -357,6 +360,13 @@ class ENGINE_API UMaterialExpression : public UObject
 	/** Checks whether any inputs to this expression create a loop */
 	bool ContainsInputLoop();
 
+	/**
+	 * Tells if this material expression type is a parameter expression.
+	 *
+	 * @returns True if the type of this is a parameter expression. False otherwise.
+	 */
+	bool IsParameterExpression() const;
+
 protected:
 	/**
 	 * Checks whether any inputs to this expression create a loop by recursively
diff --git a/Engine/Source/Runtime/Engine/Classes/Materials/MaterialInstance.h b/Engine/Source/Runtime/Engine/Classes/Materials/MaterialInstance.h
--- a/Engine/Source/Runtime/Engine/Classes/Materials/MaterialInstance.h
+++ b/Engine/Source/Runtime/Engine/Classes/Materials/MaterialInstance.h
@@ -135,6 +135,37 @@ class UMaterialInstance : public UMaterialInterface
 	UPROPERTY(EditAnywhere, BlueprintReadOnly, Category=MaterialInstance, AssetRegistrySearchable)
 	class UMaterialInterface* Parent;
 
+	/**
+	 * Delegate for custom static paramters getter.
+	 *
+	 * @param OutStaticParameterSet Parameter set to append.
+	 * @param Material Material instance to collect parameters.
+	 */
+	DECLARE_MULTICAST_DELEGATE_TwoParams(FCustomStaticParametersGetterDelegate, FStaticParameterSet&, UMaterialInstance*);
+
+	/**
+	 * Delegate for custom static parameters updater.
+	 *
+	 * @param StaticParameterSet Parameter set to update.
+	 * @param Material Material to update.
+	 *
+	 * @returns True if any parameter been updated. False otherwise.
+	 */
+	DECLARE_DELEGATE_RetVal_TwoParams(bool, FCustomParameterSetUpdaterDelegate, FStaticParameterSet&, UMaterial*);
+
+	// Custom static parameters getter delegate.
+	ENGINE_API static FCustomStaticParametersGetterDelegate CustomStaticParametersGetters;
+
+	// An array of custom paramter set updaters.
+	ENGINE_API static TArray<FCustomParameterSetUpdaterDelegate> CustomParameterSetUpdaters;
+
+	/**
+	 * Gets static parameter set for this material.
+	 *
+	 * @returns Static parameter set.
+	 */
+	ENGINE_API const FStaticParameterSet& GetStaticParameters() const;
+
 	/**
 	 * Indicates whether the instance has static permutation resources (which are required when static parameters are present) 
 	 * Read directly from the rendering thread, can only be modified with the use of a FMaterialUpdateContext.
@@ -209,39 +240,39 @@ private:
 
 public:
 	// Begin interface IBlendableInterface
-	virtual void OverrideBlendableSettings(class FSceneView& View, float Weight) const;
+	virtual ENGINE_API void OverrideBlendableSettings(class FSceneView& View, float Weight) const;
 	// End interface IBlendableInterface
 
 	// Begin UMaterialInterface interface.
-	virtual UMaterial* GetMaterial() override;
-	virtual const UMaterial* GetMaterial() const override;
-	virtual const UMaterial* GetMaterial_Concurrent(TMicRecursionGuard& RecursionGuard) const override;
-	virtual FMaterialResource* GetMaterialResource(ERHIFeatureLevel::Type InFeatureLevel, EMaterialQualityLevel::Type QualityLevel = EMaterialQualityLevel::Num) override;
-	virtual const FMaterialResource* GetMaterialResource(ERHIFeatureLevel::Type InFeatureLevel, EMaterialQualityLevel::Type QualityLevel = EMaterialQualityLevel::Num) const override;
-	virtual bool GetFontParameterValue(FName ParameterName,class UFont*& OutFontValue, int32& OutFontPage) const override;
-	virtual bool GetScalarParameterValue(FName ParameterName,float& OutValue) const override;
-	virtual bool GetTextureParameterValue(FName ParameterName,class UTexture*& OutValue) const override;
-	virtual bool GetVectorParameterValue(FName ParameterName,FLinearColor& OutValue) const override;
-	virtual void GetUsedTextures(TArray<UTexture*>& OutTextures, EMaterialQualityLevel::Type QualityLevel, bool bAllQualityLevels, ERHIFeatureLevel::Type FeatureLevel, bool bAllFeatureLevels) const override;
-	virtual void OverrideTexture(const UTexture* InTextureToOverride, UTexture* OverrideTexture, ERHIFeatureLevel::Type InFeatureLevel) override;
-	virtual bool CheckMaterialUsage(const EMaterialUsage Usage, const bool bSkipPrim = false) override;
-	virtual bool CheckMaterialUsage_Concurrent(const EMaterialUsage Usage, const bool bSkipPrim = false) const override;
-	virtual bool GetStaticSwitchParameterValue(FName ParameterName,bool &OutValue,FGuid &OutExpressionGuid) override;
-	virtual bool GetStaticComponentMaskParameterValue(FName ParameterName, bool &R, bool &G, bool &B, bool &A,FGuid &OutExpressionGuid) override;
-	virtual bool GetTerrainLayerWeightParameterValue(FName ParameterName, int32& OutWeightmapIndex, FGuid &OutExpressionGuid) override;
-	virtual bool IsDependent(UMaterialInterface* TestDependency) override;
-	virtual FMaterialRenderProxy* GetRenderProxy(bool Selected, bool bHovered=false) const override;
-	virtual UPhysicalMaterial* GetPhysicalMaterial() const override;
-	virtual bool UpdateLightmassTextureTracking() override;
-	virtual bool GetCastShadowAsMasked() const override;
-	virtual float GetEmissiveBoost() const override;
-	virtual float GetDiffuseBoost() const override;
-	virtual float GetExportResolutionScale() const override;
-	virtual float GetDistanceFieldPenumbraScale() const override;
-	virtual bool GetTexturesInPropertyChain(EMaterialProperty InProperty, TArray<UTexture*>& OutTextures, 
+	virtual ENGINE_API UMaterial* GetMaterial() override;
+	virtual ENGINE_API const UMaterial* GetMaterial() const override;
+	virtual ENGINE_API const UMaterial* GetMaterial_Concurrent(TMicRecursionGuard& RecursionGuard) const override;
+	virtual ENGINE_API FMaterialResource* GetMaterialResource(ERHIFeatureLevel::Type InFeatureLevel, EMaterialQualityLevel::Type QualityLevel = EMaterialQualityLevel::Num) override;
+	virtual ENGINE_API const FMaterialResource* GetMaterialResource(ERHIFeatureLevel::Type InFeatureLevel, EMaterialQualityLevel::Type QualityLevel = EMaterialQualityLevel::Num) const override;
+	virtual ENGINE_API bool GetFontParameterValue(FName ParameterName, class UFont*& OutFontValue, int32& OutFontPage) const override;
+	virtual ENGINE_API bool GetScalarParameterValue(FName ParameterName, float& OutValue) const override;
+	virtual ENGINE_API bool GetTextureParameterValue(FName ParameterName, class UTexture*& OutValue) const override;
+	virtual ENGINE_API bool GetVectorParameterValue(FName ParameterName, FLinearColor& OutValue) const override;
+	virtual ENGINE_API void GetUsedTextures(TArray<UTexture*>& OutTextures, EMaterialQualityLevel::Type QualityLevel, bool bAllQualityLevels, ERHIFeatureLevel::Type FeatureLevel, bool bAllFeatureLevels) const override;
+	virtual ENGINE_API void OverrideTexture(const UTexture* InTextureToOverride, UTexture* OverrideTexture, ERHIFeatureLevel::Type InFeatureLevel) override;
+	virtual ENGINE_API bool CheckMaterialUsage(const EMaterialUsage Usage, const bool bSkipPrim = false) override;
+	virtual ENGINE_API bool CheckMaterialUsage_Concurrent(const EMaterialUsage Usage, const bool bSkipPrim = false) const override;
+	virtual ENGINE_API bool GetStaticSwitchParameterValue(FName ParameterName, bool &OutValue, FGuid &OutExpressionGuid) override;
+	virtual ENGINE_API bool GetStaticComponentMaskParameterValue(FName ParameterName, bool &R, bool &G, bool &B, bool &A, FGuid &OutExpressionGuid) override;
+	virtual ENGINE_API bool GetTerrainLayerWeightParameterValue(FName ParameterName, int32& OutWeightmapIndex, FGuid &OutExpressionGuid) override;
+	virtual ENGINE_API bool IsDependent(UMaterialInterface* TestDependency) override;
+	virtual ENGINE_API FMaterialRenderProxy* GetRenderProxy(bool Selected, bool bHovered = false) const override;
+	virtual ENGINE_API UPhysicalMaterial* GetPhysicalMaterial() const override;
+	virtual ENGINE_API bool UpdateLightmassTextureTracking() override;
+	virtual ENGINE_API bool GetCastShadowAsMasked() const override;
+	virtual ENGINE_API float GetEmissiveBoost() const override;
+	virtual ENGINE_API float GetDiffuseBoost() const override;
+	virtual ENGINE_API float GetExportResolutionScale() const override;
+	virtual ENGINE_API float GetDistanceFieldPenumbraScale() const override;
+	virtual ENGINE_API bool GetTexturesInPropertyChain(EMaterialProperty InProperty, TArray<UTexture*>& OutTextures,
 		TArray<FName>* OutTextureParamNames, class FStaticParameterSet* InStaticParameterSet) override;
-	virtual void RecacheUniformExpressions() const override;
-	virtual bool GetRefractionSettings(float& OutBiasValue) const override;
+	virtual ENGINE_API void RecacheUniformExpressions() const override;
+	virtual ENGINE_API bool GetRefractionSettings(float& OutBiasValue) const override;
 	ENGINE_API virtual void ForceRecompileForRendering() override;
 	
 	ENGINE_API virtual float GetOpacityMaskClipValue_Internal() const;
@@ -269,20 +300,20 @@ public:
 	// End UMaterialInterface interface.
 
 	// Begin UObject interface.
-	virtual SIZE_T GetResourceSize(EResourceSizeMode::Type Mode) override;
-	virtual void PostInitProperties() override;	
-	virtual void BeginCacheForCookedPlatformData(  const ITargetPlatform *TargetPlatform ) override;
-	virtual bool IsCachedCookedPlatformDataLoaded(  const ITargetPlatform* TargetPlatform ) override;
-	virtual void ClearCachedCookedPlatformData( const ITargetPlatform *TargetPlatform ) override;
-	virtual void ClearAllCachedCookedPlatformData() override;
-	virtual void Serialize(FArchive& Ar) override;
-	virtual void PostLoad() override;
-	virtual void BeginDestroy() override;
-	virtual bool IsReadyForFinishDestroy() override;
-	virtual void FinishDestroy() override;
+	virtual ENGINE_API SIZE_T GetResourceSize(EResourceSizeMode::Type Mode) override;
+	virtual ENGINE_API void PostInitProperties() override;	
+	virtual ENGINE_API void BeginCacheForCookedPlatformData(const ITargetPlatform *TargetPlatform) override;
+	virtual ENGINE_API bool IsCachedCookedPlatformDataLoaded(const ITargetPlatform* TargetPlatform) override;
+	virtual ENGINE_API void ClearCachedCookedPlatformData(const ITargetPlatform *TargetPlatform) override;
+	virtual ENGINE_API void ClearAllCachedCookedPlatformData() override;
+	virtual ENGINE_API void Serialize(FArchive& Ar) override;
+	virtual ENGINE_API void PostLoad() override;
+	virtual ENGINE_API void BeginDestroy() override;
+	virtual ENGINE_API bool IsReadyForFinishDestroy() override;
+	virtual ENGINE_API void FinishDestroy() override;
 	ENGINE_API static void AddReferencedObjects(UObject* InThis, FReferenceCollector& Collector);
 #if WITH_EDITOR
-	virtual void PostEditChangeProperty(FPropertyChangedEvent& PropertyChangedEvent) override;
+	virtual ENGINE_API void PostEditChangeProperty(FPropertyChangedEvent& PropertyChangedEvent) override;
 
 	/**
 	 * Sets new static parameter overrides on the instance and recompiles the static permutation resources if needed (can be forced with bForceRecompile).
@@ -322,6 +353,8 @@ public:
 	 */
 	ENGINE_API void GetStaticParameterValues(FStaticParameterSet& OutStaticParameters);
 
+	void GetAllStaticSwitchParameterValues(TArray<FStaticTerrainLayerWeightParameter> &TerrainLayerWeightParameters, UMaterial* ParentMaterial);
+
 	void GetBasePropertyOverridesHash(FSHAHash& OutHash)const;
 
 	// For all materials instances, UMaterialInstance::CacheResourceShadersForRendering
@@ -402,5 +435,61 @@ protected:
 	friend class FMaterialUpdateContext;
 };
 
+/**
+ * This function takes a array of parameter structs and attempts to establish a reference to the expression object each parameter represents.
+ * If a reference exists, the function checks to see if the parameter has been renamed.
+ *
+ * @param Parameters		Array of parameters to operate on.
+ * @param ParentMaterial	Parent material to search in for expressions.
+ *
+ * @return Returns whether or not any of the parameters was changed.
+ */
+template <typename ParameterType, typename ExpressionType>
+bool UpdateParameterSet(TArray<ParameterType> &Parameters, UMaterial* ParentMaterial)
+{
+	bool bChanged = false;
 
+	// Loop through all of the parameters and try to either establish a reference to the 
+	// expression the parameter represents, or check to see if the parameter's name has changed.
+	for (int32 ParameterIdx = 0; ParameterIdx < Parameters.Num(); ParameterIdx++)
+	{
+		bool bTryToFindByName = true;
+
+		ParameterType &Parameter = Parameters[ParameterIdx];
+
+		if (Parameter.ExpressionGUID.IsValid())
+		{
+			ExpressionType* Expression = ParentMaterial->FindExpressionByGUID<ExpressionType>(Parameter.ExpressionGUID);
+
+			// Check to see if the parameter name was changed.
+			if (Expression)
+			{
+				bTryToFindByName = false;
+
+				if (Parameter.ParameterName != Expression->ParameterName)
+				{
+					Parameter.ParameterName = Expression->ParameterName;
+					bChanged = true;
+				}
+			}
+		}
+
+		// No reference to the material expression exists, so try to find one in the material expression's array if we are in the editor.
+		if (bTryToFindByName && GIsEditor && !FApp::IsGame())
+		{
+			for (int32 ExpressionIndex = 0; ExpressionIndex < ParentMaterial->Expressions.Num(); ExpressionIndex++)
+			{
+				ExpressionType* ParameterExpression = Cast<ExpressionType>(ParentMaterial->Expressions[ExpressionIndex]);
+
+				if (ParameterExpression && ParameterExpression->ParameterName == Parameter.ParameterName)
+				{
+					Parameter.ExpressionGUID = ParameterExpression->ExpressionGUID;
+					bChanged = true;
+					break;
+				}
+			}
+		}
+	}
 
+	return bChanged;
+}
diff --git a/Engine/Source/Runtime/Engine/Classes/Materials/MaterialInstanceConstant.h b/Engine/Source/Runtime/Engine/Classes/Materials/MaterialInstanceConstant.h
--- a/Engine/Source/Runtime/Engine/Classes/Materials/MaterialInstanceConstant.h
+++ b/Engine/Source/Runtime/Engine/Classes/Materials/MaterialInstanceConstant.h
@@ -1,6 +1,9 @@
 // Copyright 1998-2014 Epic Games, Inc. All Rights Reserved.
 
 #pragma once
+
+#include "Materials/MaterialInstance.h"
+
 #include "MaterialInstanceConstant.generated.h"
 
 /**
diff --git a/Engine/Source/Runtime/Engine/Private/ActorEditor.cpp b/Engine/Source/Runtime/Engine/Private/ActorEditor.cpp
--- a/Engine/Source/Runtime/Engine/Private/ActorEditor.cpp
+++ b/Engine/Source/Runtime/Engine/Private/ActorEditor.cpp
@@ -8,8 +8,6 @@
 #include "LevelUtils.h"
 #include "MapErrors.h"
 #include "Foliage/InstancedFoliageActor.h"
-#include "Landscape/LandscapeHeightfieldCollisionComponent.h"
-#include "Landscape/LandscapeComponent.h"
 #include "Engine/LevelBounds.h"
 
 #if WITH_EDITOR
diff --git a/Engine/Source/Runtime/Engine/Private/Collision/PhysicsFiltering.h b/Engine/Source/Runtime/Engine/Private/Collision/PhysicsFiltering.h
--- a/Engine/Source/Runtime/Engine/Private/Collision/PhysicsFiltering.h
+++ b/Engine/Source/Runtime/Engine/Private/Collision/PhysicsFiltering.h
@@ -23,7 +23,7 @@ enum EPhysXFilterDataFlags
 
 struct FPhysicsFilterBuilder
 {
-	FPhysicsFilterBuilder(TEnumAsByte<enum ECollisionChannel> InObjectType, const struct FCollisionResponseContainer& ResponseToChannels);
+	ENGINE_API FPhysicsFilterBuilder(TEnumAsByte<enum ECollisionChannel> InObjectType, const struct FCollisionResponseContainer& ResponseToChannels);
 
 	inline void ConditionalSetFlags(EPhysXFilterDataFlags Flag, bool bEnabled)
 	{
diff --git a/Engine/Source/Runtime/Engine/Private/InstancedFoliage.cpp b/Engine/Source/Runtime/Engine/Private/InstancedFoliage.cpp
--- a/Engine/Source/Runtime/Engine/Private/InstancedFoliage.cpp
+++ b/Engine/Source/Runtime/Engine/Private/InstancedFoliage.cpp
@@ -11,7 +11,6 @@
 #include "MapErrors.h"
 #include "Foliage/InstancedFoliageActor.h"
 #include "Foliage/FoliageType_InstancedStaticMesh.h"
-#include "Landscape/LandscapeHeightfieldCollisionComponent.h"
 
 #define LOCTEXT_NAMESPACE "InstancedFoliage"
 
@@ -756,122 +755,6 @@ AInstancedFoliageActor* AInstancedFoliageActor::GetInstancedFoliageActorForLevel
 	return nullptr;
 }
 
-void AInstancedFoliageActor::SnapInstancesForLandscape(ULandscapeHeightfieldCollisionComponent* InLandscapeComponent, const FBox& InInstanceBox)
-{
-	for (auto& MeshPair : FoliageMeshes)
-	{
-		// Find the per-mesh info matching the mesh.
-		UFoliageType* Settings = MeshPair.Key;
-		FFoliageMeshInfo& MeshInfo = *MeshPair.Value;
-
-		FFoliageComponentHashInfo* ComponentHashInfo = MeshInfo.ComponentHash.Find(InLandscapeComponent);
-		if (ComponentHashInfo)
-		{
-			float TraceExtentSize = InLandscapeComponent->Bounds.SphereRadius * 2.f + 10.f; // extend a little
-			FVector TraceVector = InLandscapeComponent->GetOwner()->GetRootComponent()->ComponentToWorld.GetUnitAxis(EAxis::Z) * TraceExtentSize;
-
-			bool bFirst = true;
-			TArray<int32> InstancesToRemove;
-			for (int32 InstanceIndex : ComponentHashInfo->Instances)
-			{
-				FFoliageInstance& Instance = MeshInfo.Instances[InstanceIndex];
-
-				// Test location should remove any Z offset
-				FVector TestLocation = FMath::Abs(Instance.ZOffset) > KINDA_SMALL_NUMBER
-					? (FVector)Instance.GetInstanceWorldTransform().TransformPosition(FVector(0, 0, -Instance.ZOffset))
-					: Instance.Location;
-
-				if (InInstanceBox.IsInside(TestLocation))
-				{
-					if (bFirst)
-					{
-						bFirst = false;
-						Modify();
-					}
-
-					FVector Start = TestLocation + TraceVector;
-					FVector End = TestLocation - TraceVector;
-
-					static FName TraceTag = FName(TEXT("FoliageSnapToLandscape"));
-					TArray<FHitResult> Results;
-					UWorld* World = InLandscapeComponent->GetWorld();
-					check(World);
-					// Editor specific landscape heightfield uses ECC_Visibility collision channel
-					World->LineTraceMulti(Results, Start, End, FCollisionQueryParams(TraceTag, true), FCollisionObjectQueryParams(ECollisionChannel::ECC_Visibility));
-
-					bool bFoundHit = false;
-					for (const FHitResult& Hit : Results)
-					{
-						if (Hit.Component == InLandscapeComponent)
-						{
-							bFoundHit = true;
-							if ((TestLocation - Hit.Location).SizeSquared() > KINDA_SMALL_NUMBER)
-							{
-								// Remove instance location from the hash. Do not need to update ComponentHash as we re-add below.
-								MeshInfo.InstanceHash->RemoveInstance(Instance.Location, InstanceIndex);
-
-								// Update the instance editor data
-								Instance.Location = Hit.Location;
-
-								if (Instance.Flags & FOLIAGE_AlignToNormal)
-								{
-									// Remove previous alignment and align to new normal.
-									Instance.Rotation = Instance.PreAlignRotation;
-									Instance.AlignToNormal(Hit.Normal, Settings->AlignMaxAngle);
-								}
-
-								// Reapply the Z offset in local space
-								if (FMath::Abs(Instance.ZOffset) > KINDA_SMALL_NUMBER)
-								{
-									Instance.Location = Instance.GetInstanceWorldTransform().TransformPosition(FVector(0, 0, Instance.ZOffset));
-								}
-
-								// Todo: add do validation with other parameters such as max/min height etc.
-
-								// Update this instances' transform in the UInstancedStaticMeshComponent
-								if (MeshInfo.InstanceClusters.IsValidIndex(Instance.ClusterIndex))
-								{
-									FFoliageInstanceCluster& Cluster = MeshInfo.InstanceClusters[Instance.ClusterIndex];
-
-									Cluster.ClusterComponent->Modify();
-
-									int32 ClusterInstanceDataIndex = Cluster.InstanceIndices.Find(InstanceIndex);
-
-									check(ClusterInstanceDataIndex != INDEX_NONE);
-									check(Cluster.ClusterComponent->PerInstanceSMData.IsValidIndex(ClusterInstanceDataIndex));
-									Cluster.ClusterComponent->MarkRenderStateDirty();
-
-									// Update bounds
-									FTransform InstanceToWorld = Instance.GetInstanceWorldTransform();
-									Cluster.Bounds = Cluster.Bounds + Cluster.ClusterComponent->StaticMesh->GetBounds().TransformBy(InstanceToWorld);
-
-									// Update transform in InstancedStaticMeshComponent
-									FTransform InstanceToLocal = InstanceToWorld.GetRelativeTransform(Cluster.ClusterComponent->ComponentToWorld);
-									Cluster.ClusterComponent->PerInstanceSMData[ClusterInstanceDataIndex].Transform = InstanceToLocal.ToMatrixWithScale();
-									Cluster.ClusterComponent->InvalidateLightingCache();
-								}
-
-								// Re-add the new instance location to the hash
-								MeshInfo.InstanceHash->InsertInstance(Instance.Location, InstanceIndex);
-							}
-							break;
-						}
-					}
-
-					if (!bFoundHit)
-					{
-						// Couldn't find new spot - remove instance
-						InstancesToRemove.Add(InstanceIndex);
-					}
-				}
-			}
-
-			// Remove any unused instances
-			MeshInfo.RemoveInstances(this, InstancesToRemove);
-		}
-	}
-}
-
 void AInstancedFoliageActor::MoveInstancesForMovedComponent(UActorComponent* InComponent)
 {
 	bool bUpdatedInstances = false;
diff --git a/Engine/Source/Runtime/Engine/Private/Materials/Material.cpp b/Engine/Source/Runtime/Engine/Private/Materials/Material.cpp
--- a/Engine/Source/Runtime/Engine/Private/Materials/Material.cpp
+++ b/Engine/Source/Runtime/Engine/Private/Materials/Material.cpp
@@ -10,11 +10,6 @@
 #include "Materials/MaterialExpressionComment.h"
 #include "Materials/MaterialExpressionDynamicParameter.h"
 #include "Materials/MaterialExpressionFontSampleParameter.h"
-#include "Materials/MaterialExpressionLandscapeLayerBlend.h"
-#include "Materials/MaterialExpressionLandscapeLayerSwitch.h"
-#include "Materials/MaterialExpressionLandscapeLayerSample.h"
-#include "Materials/MaterialExpressionLandscapeLayerWeight.h"
-#include "Materials/MaterialExpressionLandscapeVisibilityMask.h"
 #include "Materials/MaterialExpressionMaterialFunctionCall.h"
 #include "Materials/MaterialExpressionMultiply.h"
 #include "Materials/MaterialExpressionQualitySwitch.h"
@@ -985,29 +980,6 @@ bool UMaterial::SetMaterialUsage(bool &bNeedsRecompile, EMaterialUsage Usage, co
 	return true;
 }
 
-/**
-* @param	OutParameterNames		Storage array for the parameter names we are returning.
-* @param	OutParameterIds			Storage array for the parameter id's we are returning.
-*
-* @return	Returns a array of vector parameter names used in this material.
-*/
-template<typename ExpressionType>
-void UMaterial::GetAllParameterNames(TArray<FName> &OutParameterNames, TArray<FGuid> &OutParameterIds)
-{
-	for(int32 ExpressionIndex = 0;ExpressionIndex < Expressions.Num();ExpressionIndex++)
-	{
-		ExpressionType* ParameterExpression =
-			Cast<ExpressionType>(Expressions[ExpressionIndex]);
-
-		if(ParameterExpression)
-		{
-			ParameterExpression->GetAllParameterNames(OutParameterNames, OutParameterIds);
-		}
-	}
-
-	check(OutParameterNames.Num() == OutParameterIds.Num());
-}
-
 void UMaterial::GetAllVectorParameterNames(TArray<FName> &OutParameterNames, TArray<FGuid> &OutParameterIds)
 {
 	OutParameterNames.Empty();
@@ -1048,17 +1020,6 @@ void UMaterial::GetAllStaticComponentMaskParameterNames(TArray<FName> &OutParame
 	GetAllParameterNames<UMaterialExpressionStaticComponentMaskParameter>(OutParameterNames, OutParameterIds);
 }
 
-void UMaterial::GetAllTerrainLayerWeightParameterNames(TArray<FName> &OutParameterNames, TArray<FGuid> &OutParameterIds)
-{
-	OutParameterNames.Empty();
-	OutParameterIds.Empty();
-	GetAllParameterNames<UMaterialExpressionLandscapeLayerWeight>(OutParameterNames, OutParameterIds);
-	GetAllParameterNames<UMaterialExpressionLandscapeLayerSample>(OutParameterNames, OutParameterIds);
-	GetAllParameterNames<UMaterialExpressionLandscapeLayerSwitch>(OutParameterNames, OutParameterIds);
-	GetAllParameterNames<UMaterialExpressionLandscapeLayerBlend>(OutParameterNames, OutParameterIds);
-	GetAllParameterNames<UMaterialExpressionLandscapeVisibilityMask>(OutParameterNames, OutParameterIds);
-}
-
 extern FPostProcessMaterialNode* IteratePostProcessMaterialNodes(const FFinalPostProcessSettings& Dest, const UMaterial* Material, FBlendableEntry*& Iterator);
 
 void UMaterialInterface::OverrideBlendableSettings(class FSceneView& View, float Weight) const
@@ -1660,25 +1621,6 @@ const FMaterialResource* UMaterial::GetMaterialResource(ERHIFeatureLevel::Type I
 	return MaterialResources[QualityLevel][InFeatureLevel];
 }
 
-void UMaterial::FixupTerrainLayerWeightNodes()
-{
-	for (int32 ExpressionIndex = 0; ExpressionIndex < Expressions.Num(); ExpressionIndex++)
-	{
-		UMaterialExpressionLandscapeLayerWeight* LayerWeight = Cast<UMaterialExpressionLandscapeLayerWeight>(Expressions[ExpressionIndex]);
-		UMaterialExpressionLandscapeLayerSwitch* LayerSwitch = Cast<UMaterialExpressionLandscapeLayerSwitch>(Expressions[ExpressionIndex]);
-
-		// Regenerate parameter guids since the old ones were not generated consistently
-		if (LayerWeight)
-		{
-			LayerWeight->UpdateParameterGuid(true, true);
-		}
-		else if (LayerSwitch)
-		{
-			LayerSwitch->UpdateParameterGuid(true, true);
-		}
-	}
-}
-
 void UMaterial::Serialize(FArchive& Ar)
 {
 	Super::Serialize(Ar);
@@ -1940,11 +1882,6 @@ void UMaterial::PostLoad()
 		}
 	}
 
-	if (GetLinkerUE4Version() < VER_UE4_FIXUP_TERRAIN_LAYER_NODES)
-	{
-		FixupTerrainLayerWeightNodes();
-	}
-
 	// Fixup for legacy materials which didn't recreate the lighting guid properly on duplication
 	if (GetLinker() && GetLinker()->UE4Ver() < VER_UE4_BUMPED_MATERIAL_EXPORT_GUIDS)
 	{
