Landscape module from Engine split out.

diff --git a/Engine/Source/Runtime/Engine/Private/Materials/MaterialExpressions.cpp b/Engine/Source/Runtime/Engine/Private/Materials/MaterialExpressions.cpp
--- a/Engine/Source/Runtime/Engine/Private/Materials/MaterialExpressions.cpp
+++ b/Engine/Source/Runtime/Engine/Private/Materials/MaterialExpressions.cpp
@@ -52,12 +52,6 @@
 #include "Materials/MaterialExpressionFunctionOutput.h"
 #include "Materials/MaterialExpressionGIReplace.h"
 #include "Materials/MaterialExpressionIf.h"
-#include "Materials/MaterialExpressionLandscapeLayerBlend.h"
-#include "Materials/MaterialExpressionLandscapeLayerCoords.h"
-#include "Materials/MaterialExpressionLandscapeLayerSwitch.h"
-#include "Materials/MaterialExpressionLandscapeLayerWeight.h"
-#include "Materials/MaterialExpressionLandscapeLayerSample.h"
-#include "Materials/MaterialExpressionLandscapeVisibilityMask.h"
 #include "Materials/MaterialExpressionLightmapUVs.h"
 #include "Materials/MaterialExpressionLightmassReplace.h"
 #include "Materials/MaterialExpressionLightVector.h"
@@ -176,11 +170,7 @@ bool IsAllowedExpressionType(UClass* Class, bool bMaterialFunction)
 			&& !Class->IsChildOf(UMaterialExpressionTextureSampleParameter::StaticClass())
 			&& !Class->IsChildOf(UMaterialExpressionDynamicParameter::StaticClass())
 			&& !Class->IsChildOf(UMaterialExpressionFontSampleParameter::StaticClass())
-			&& !Class->IsChildOf(UMaterialExpressionLandscapeLayerWeight::StaticClass())
-			&& !Class->IsChildOf(UMaterialExpressionLandscapeLayerSample::StaticClass())
-			&& !Class->IsChildOf(UMaterialExpressionLandscapeLayerSwitch::StaticClass())
-			&& !Class->IsChildOf(UMaterialExpressionLandscapeLayerBlend::StaticClass())
-			&& !Class->IsChildOf(UMaterialExpressionLandscapeVisibilityMask::StaticClass());
+			&& !Class->GetDefaultObject<UMaterialExpression>()->IsParameterExpression();
 	}
 	else
 	{
@@ -993,6 +983,11 @@ bool UMaterialExpression::ContainsInputLoop()
 	return ContainsInputLoopInternal(ExpressionStack);
 }
 
+bool UMaterialExpression::IsParameterExpression() const
+{
+	return bIsParameterExpression;
+}
+
 bool UMaterialExpression::ContainsInputLoopInternal(TArray<FMaterialExpressionKey>& ExpressionStack)
 {
 	const TArray<FExpressionInput*> Inputs = GetInputs();
@@ -8400,664 +8395,6 @@ void UMaterialExpressionAntialiasedTextureMask::SetDefaultTexture()
 	Texture = LoadObject<UTexture2D>(NULL, TEXT("/Engine/EngineResources/DefaultTexture.DefaultTexture"), NULL, LOAD_None, NULL);
 }
 
-///////////////////////////////////////////////////////////////////////////////
-// UMaterialExpressionLandscapeLayerSample
-///////////////////////////////////////////////////////////////////////////////
-UMaterialExpressionLandscapeLayerSample::UMaterialExpressionLandscapeLayerSample(const FObjectInitializer& ObjectInitializer)
-	: Super(ObjectInitializer)
-{
-	// Structure to hold one-time initialization
-	struct FConstructorStatics
-	{
-		FString NAME_Landscape;
-		FConstructorStatics()
-			: NAME_Landscape(LOCTEXT("Landscape", "Landscape").ToString())
-		{
-		}
-	};
-	static FConstructorStatics ConstructorStatics;
-
-	bIsParameterExpression = true;
-	MenuCategories.Add(ConstructorStatics.NAME_Landscape);
-}
-
-int32 UMaterialExpressionLandscapeLayerSample::Compile(class FMaterialCompiler* Compiler, int32 OutputIndex, int32 MultiplexIndex)
-{
-	const int32 WeightCode = Compiler->StaticTerrainLayerWeight(ParameterName, Compiler->Constant(PreviewWeight));
-
-	return WeightCode;
-}
-
-UTexture* UMaterialExpressionLandscapeLayerSample::GetReferencedTexture()
-{
-	return GEngine->WeightMapPlaceholderTexture;
-}
-
-void UMaterialExpressionLandscapeLayerSample::GetCaption(TArray<FString>& OutCaptions) const
-{
-	OutCaptions.Add(FString::Printf(TEXT("Sample '%s'"), *ParameterName.ToString()));
-}
-
-void UMaterialExpressionLandscapeLayerSample::GetAllParameterNames(TArray<FName> &OutParameterNames, TArray<FGuid> &OutParameterIds)
-{
-	if (!OutParameterNames.Contains(ParameterName))
-	{
-		OutParameterNames.Add(ParameterName);
-		OutParameterIds.Add(ExpressionGUID);
-	}
-}
-
-///////////////////////////////////////////////////////////////////////////////
-// UMaterialExpressionLandscapeLayerWeight
-///////////////////////////////////////////////////////////////////////////////
-UMaterialExpressionLandscapeLayerWeight::UMaterialExpressionLandscapeLayerWeight(const FObjectInitializer& ObjectInitializer)
-	: Super(ObjectInitializer)
-{
-	// Structure to hold one-time initialization
-	struct FConstructorStatics
-	{
-		FString NAME_Landscape;
-		FConstructorStatics()
-			: NAME_Landscape(LOCTEXT( "Landscape", "Landscape" ).ToString())
-		{
-		}
-	};
-	static FConstructorStatics ConstructorStatics;
-
-	bIsParameterExpression = true;
-	MenuCategories.Add(ConstructorStatics.NAME_Landscape);
-	PreviewWeight = 0.0f;
-	ConstBase = FVector(0.f, 0.f, 0.f);
-}
-
-bool UMaterialExpressionLandscapeLayerWeight::IsResultMaterialAttributes(int32 OutputIndex)
-{
-	if (ContainsInputLoop())
-	{
-		// If there is a loop anywhere in this expression's inputs then we can't risk checking them
-		return false;
-	}
-	bool bLayerIsMaterialAttributes = Layer.Expression != NULL && Layer.Expression->IsResultMaterialAttributes(Layer.OutputIndex);
-	bool bBaseIsMaterialAttributes = Base.Expression != NULL && Base.Expression->IsResultMaterialAttributes(Base.OutputIndex);
-	return bLayerIsMaterialAttributes || bBaseIsMaterialAttributes;
-}
-
-int32 UMaterialExpressionLandscapeLayerWeight::Compile(class FMaterialCompiler* Compiler, int32 OutputIndex, int32 MultiplexIndex)
-{
-	const int32 BaseCode = Base.Expression ? Base.Compile(Compiler, MultiplexIndex) : Compiler->Constant3(ConstBase.X, ConstBase.Y, ConstBase.Z);
-	const int32 WeightCode = Compiler->StaticTerrainLayerWeight(ParameterName, Compiler->Constant(PreviewWeight));
-
-	int32 ReturnCode = INDEX_NONE;
-	if(WeightCode == INDEX_NONE)
-	{
-		ReturnCode = BaseCode;
-	}
-	else
-	{
-		const int32 LayerCode = Layer.Compile(Compiler, MultiplexIndex);
-		ReturnCode = Compiler->Add(BaseCode, Compiler->Mul(LayerCode, WeightCode));
-	}
-
-	if (ReturnCode != INDEX_NONE && //If we've already failed for some other reason don't bother with this check. It could have been the reentrant check causing this to loop infinitely!
-		Layer.Expression != NULL && Base.Expression != NULL &&
-		Layer.Expression->IsResultMaterialAttributes(Layer.OutputIndex) != Base.Expression->IsResultMaterialAttributes(Base.OutputIndex))
-	{
-		Compiler->Error(TEXT("Cannot mix MaterialAttributes and non MaterialAttributes nodes"));
-	}
-
-	return ReturnCode;
-}
-
-UTexture* UMaterialExpressionLandscapeLayerWeight::GetReferencedTexture()
-{
-	return GEngine->WeightMapPlaceholderTexture;
-}
-
-void UMaterialExpressionLandscapeLayerWeight::GetCaption(TArray<FString>& OutCaptions) const
-{
-	OutCaptions.Add(FString::Printf(TEXT("Layer '%s'"), *ParameterName.ToString()));
-}
-
-void UMaterialExpressionLandscapeLayerWeight::GetAllParameterNames(TArray<FName> &OutParameterNames, TArray<FGuid> &OutParameterIds)
-{
-	int32 CurrentSize = OutParameterNames.Num();
-	OutParameterNames.AddUnique(ParameterName);
-
-	if(CurrentSize != OutParameterNames.Num())
-	{
-		OutParameterIds.Add(ExpressionGUID);
-	}
-}
-
-///////////////////////////////////////////////////////////////////////////////
-// UMaterialExpressionLandscapeLayerSwitch
-///////////////////////////////////////////////////////////////////////////////
-UMaterialExpressionLandscapeLayerSwitch::UMaterialExpressionLandscapeLayerSwitch(const FObjectInitializer& ObjectInitializer)
-	: Super(ObjectInitializer)
-{
-	// Structure to hold one-time initialization
-	struct FConstructorStatics
-	{
-		FString NAME_Landscape;
-		FConstructorStatics()
-			: NAME_Landscape(LOCTEXT( "Landscape", "Landscape" ).ToString())
-		{
-		}
-	};
-	static FConstructorStatics ConstructorStatics;
-
-	bIsParameterExpression = true;
-	MenuCategories.Add(ConstructorStatics.NAME_Landscape);
-	PreviewUsed = true;
-	bCollapsed = false;
-}
-
-bool UMaterialExpressionLandscapeLayerSwitch::IsResultMaterialAttributes(int32 OutputIndex)
-{
-	if (ContainsInputLoop())
-	{
-		// If there is a loop anywhere in this expression's inputs then we can't risk checking them
-		return false;
-	}
-	bool bLayerUsedIsMaterialAttributes = LayerUsed.Expression != NULL && LayerUsed.Expression->IsResultMaterialAttributes(LayerUsed.OutputIndex);
-	bool bLayerNotUsedIsMaterialAttributes = LayerNotUsed.Expression != NULL && LayerNotUsed.Expression->IsResultMaterialAttributes(LayerNotUsed.OutputIndex);
-	return bLayerUsedIsMaterialAttributes || bLayerNotUsedIsMaterialAttributes;
-}
-
-int32 UMaterialExpressionLandscapeLayerSwitch::Compile(class FMaterialCompiler* Compiler, int32 OutputIndex, int32 MultiplexIndex)
-{
-	const int32 WeightCode = Compiler->StaticTerrainLayerWeight(
-		ParameterName,
-		PreviewUsed ? Compiler->Constant(1.0f) : INDEX_NONE
-		);
-
-	int32 ReturnCode = INDEX_NONE;
-	if (WeightCode != INDEX_NONE)
-	{
-		ReturnCode = LayerUsed.Compile(Compiler, MultiplexIndex);
-	}
-	else
-	{
-		ReturnCode = LayerNotUsed.Compile(Compiler, MultiplexIndex);
-	}	
-	
-	if (ReturnCode != INDEX_NONE && //If we've already failed for some other reason don't bother with this check. It could have been the reentrant check causing this to loop infinitely!
-		LayerUsed.Expression != NULL && LayerNotUsed.Expression != NULL &&
-		LayerUsed.Expression->IsResultMaterialAttributes(LayerUsed.OutputIndex) != LayerNotUsed.Expression->IsResultMaterialAttributes(LayerNotUsed.OutputIndex))
-	{
-		Compiler->Error(TEXT("Cannot mix MaterialAttributes and non MaterialAttributes nodes"));
-	}
-
-	return ReturnCode;
-}
-
-UTexture* UMaterialExpressionLandscapeLayerSwitch::GetReferencedTexture()
-{
-	return GEngine->WeightMapPlaceholderTexture;
-}
-
-void UMaterialExpressionLandscapeLayerSwitch::GetCaption(TArray<FString>& OutCaptions) const
-{
-	OutCaptions.Add(TEXT("Layer Switch")); 
-	OutCaptions.Add(FString::Printf(TEXT("'%s'"), *ParameterName.ToString()));
-}
-
-void UMaterialExpressionLandscapeLayerSwitch::Serialize(FArchive& Ar)
-{
-	Super::Serialize(Ar);
-
-	if (Ar.UE4Ver() < VER_UE4_FIX_TERRAIN_LAYER_SWITCH_ORDER)
-	{
-		Swap(LayerUsed, LayerNotUsed);
-	}
-}
-
-void UMaterialExpressionLandscapeLayerSwitch::GetAllParameterNames(TArray<FName> &OutParameterNames, TArray<FGuid> &OutParameterIds)
-{
-	int32 CurrentSize = OutParameterNames.Num();
-	OutParameterNames.AddUnique(ParameterName);
-
-	if(CurrentSize != OutParameterNames.Num())
-	{
-		OutParameterIds.Add(ExpressionGUID);
-	}
-}
-
-///////////////////////////////////////////////////////////////////////////////
-// UMaterialExpressionLandscapeLayerBlend
-///////////////////////////////////////////////////////////////////////////////
-UMaterialExpressionLandscapeLayerBlend::UMaterialExpressionLandscapeLayerBlend(const FObjectInitializer& ObjectInitializer)
-	: Super(ObjectInitializer)
-{
-	// Structure to hold one-time initialization
-	struct FConstructorStatics
-	{
-		FString NAME_Landscape;
-		FConstructorStatics()
-			: NAME_Landscape(LOCTEXT( "Landscape", "Landscape" ).ToString())
-		{
-		}
-	};
-	static FConstructorStatics ConstructorStatics;
-
-	bIsParameterExpression = true;
-
-	MenuCategories.Add(ConstructorStatics.NAME_Landscape);
-}
-
-void UMaterialExpressionLandscapeLayerBlend::Serialize(FArchive& Ar)
-{
-	Super::Serialize(Ar);
-
-	if (Ar.IsLoading() && Ar.UE4Ver() < VER_UE4_ADD_LB_WEIGHTBLEND)
-	{
-		// convert any LB_AlphaBlend entries to LB_WeightBlend
-		for (FLayerBlendInput& LayerInput : Layers)
-		{
-			if (LayerInput.BlendType == LB_AlphaBlend)
-			{
-				LayerInput.BlendType = LB_WeightBlend;
-			}
-		}
-	}
-}
-
-const TArray<FExpressionInput*> UMaterialExpressionLandscapeLayerBlend::GetInputs()
-{
-	TArray<FExpressionInput*> Result;
-	for( int32 LayerIdx=0;LayerIdx<Layers.Num();LayerIdx++ )
-	{	
-		Result.Add( &Layers[LayerIdx].LayerInput );
-		if( Layers[LayerIdx].BlendType == LB_HeightBlend )
-		{
-			Result.Add( &Layers[LayerIdx].HeightInput );
-		}
-	}
-	return Result;
-}
-
-FExpressionInput* UMaterialExpressionLandscapeLayerBlend::GetInput(int32 InputIndex)
-{
-	int32 Idx=0;
-	for( int32 LayerIdx=0;LayerIdx<Layers.Num();LayerIdx++ )
-	{	
-		if( InputIndex == Idx++ ) 
-		{
-			return &Layers[LayerIdx].LayerInput;
-		}
-		if( Layers[LayerIdx].BlendType == LB_HeightBlend )
-		{
-			if( InputIndex == Idx++ ) 
-			{
-				return &Layers[LayerIdx].HeightInput;
-			}
-		}
-	}
-	return NULL;
-}
-
-FString UMaterialExpressionLandscapeLayerBlend::GetInputName(int32 InputIndex) const
-{
-	int32 Idx=0;
-	for( int32 LayerIdx=0;LayerIdx<Layers.Num();LayerIdx++ )
-	{	
-		if( InputIndex == Idx++ ) 
-		{
-			return FString::Printf(TEXT("Layer %s"), *Layers[LayerIdx].LayerName.ToString());
-		}
-		if( Layers[LayerIdx].BlendType == LB_HeightBlend )
-		{
-			if( InputIndex == Idx++ ) 
-			{
-				return FString::Printf(TEXT("Height %s"), *Layers[LayerIdx].LayerName.ToString());
-			}
-		}
-	}
-	return TEXT("");
-}
-
-bool UMaterialExpressionLandscapeLayerBlend::IsResultMaterialAttributes(int32 OutputIndex)
-{
-	if (ContainsInputLoop())
-	{
-		// If there is a loop anywhere in this expression's inputs then we can't risk checking them
-		return false;
-	}
-	for (int32 LayerIdx = 0; LayerIdx < Layers.Num(); LayerIdx++)
-	{
-		if (Layers[LayerIdx].LayerInput.Expression && Layers[LayerIdx].LayerInput.Expression->IsResultMaterialAttributes(Layers[LayerIdx].LayerInput.OutputIndex))
-		{
-			return true;
-		}
-	}
-	return false;
-}
-
-int32 UMaterialExpressionLandscapeLayerBlend::Compile(class FMaterialCompiler* Compiler, int32 OutputIndex, int32 MultiplexIndex)
-{
-	// For renormalization
-	bool bNeedsRenormalize = false;
-	int32 WeightSumCode = Compiler->Constant(0);
-
-	// Temporary store for each layer's weight
-	TArray<int32> WeightCodes;
-	WeightCodes.Empty(Layers.Num());
-	
-	for( int32 LayerIdx=0;LayerIdx<Layers.Num();LayerIdx++ )
-	{
-		WeightCodes.Add(INDEX_NONE);
-
-		FLayerBlendInput& Layer = Layers[LayerIdx];
-
-		// LB_AlphaBlend layers are blended last
-		if (Layer.BlendType != LB_AlphaBlend)
-		{
-			// Height input
-			const int32 HeightCode = Layer.HeightInput.Expression ? Layer.HeightInput.Compile(Compiler, MultiplexIndex) : Compiler->Constant(Layer.ConstHeightInput);
-	
-			const int32 WeightCode = Compiler->StaticTerrainLayerWeight(Layer.LayerName,Layer.PreviewWeight > 0.0f ? Compiler->Constant(Layer.PreviewWeight) : INDEX_NONE);
-			if( WeightCode != INDEX_NONE )
-			{
-				switch( Layer.BlendType )
-				{
-				case LB_WeightBlend:
-					{
-						// Store the weight plus accumulate the sum of all weights so far
-						WeightCodes[LayerIdx] = WeightCode;
-						WeightSumCode = Compiler->Add(WeightSumCode, WeightCode);
-					}
-					break;
-				case LB_HeightBlend:
-					{
-						bNeedsRenormalize = true;
-	
-						// Modify weight with height
-						int32 ModifiedWeightCode = Compiler->Clamp(
-							Compiler->Add(Compiler->Lerp(Compiler->Constant(-1.f), Compiler->Constant(1.f), WeightCode), HeightCode),
-							Compiler->Constant(0.f), Compiler->Constant(1.f) );
-	
-						// Store the final weight plus accumulate the sum of all weights so far
-						WeightCodes[LayerIdx] = ModifiedWeightCode;
-						WeightSumCode = Compiler->Add(WeightSumCode, ModifiedWeightCode);
-					}
-					break;
-				}
-			}
-		}
-	}
-
-	int32 InvWeightSumCode = Compiler->Div(Compiler->Constant(1.f),WeightSumCode);
-
-	int32 OutputCode = Compiler->Constant(0);
-
-	for( int32 LayerIdx=0;LayerIdx<Layers.Num();LayerIdx++ )
-	{
-		FLayerBlendInput& Layer = Layers[LayerIdx];
-
-		if( WeightCodes[LayerIdx] != INDEX_NONE )
-		{
-			// Layer input
-			int32 LayerCode = Layer.LayerInput.Expression ? Layer.LayerInput.Compile(Compiler, MultiplexIndex) : Compiler->Constant3(Layer.ConstLayerInput.X, Layer.ConstLayerInput.Y, Layer.ConstLayerInput.Z);
-
-			if( bNeedsRenormalize )
-			{
-				// Renormalize the weights as our height modification has made them non-uniform
-				OutputCode = Compiler->Add( OutputCode, Compiler->Mul(LayerCode, Compiler->Mul(InvWeightSumCode,WeightCodes[LayerIdx])) );
-			}
-			else
-			{
-				// No renormalization is necessary, so just add the weights
-				OutputCode = Compiler->Add( OutputCode, Compiler->Mul(LayerCode,WeightCodes[LayerIdx]) );
-			}
-		}
-	}
-
-	// Blend in LB_AlphaBlend layers
-	for (FLayerBlendInput& Layer : Layers)
-	{
-		if (Layer.BlendType == LB_AlphaBlend)
-		{
-			const int32 WeightCode = Compiler->StaticTerrainLayerWeight(Layer.LayerName, Layer.PreviewWeight > 0.0f ? Compiler->Constant(Layer.PreviewWeight) : INDEX_NONE);
-			if (WeightCode != INDEX_NONE)
-			{
-				const int32 LayerCode = Layer.LayerInput.Expression ? Layer.LayerInput.Compile(Compiler, MultiplexIndex) : Compiler->Constant3(Layer.ConstLayerInput.X, Layer.ConstLayerInput.Y, Layer.ConstLayerInput.Z);
-				// Blend in the layer using the alpha value
-				OutputCode = Compiler->Lerp(OutputCode, LayerCode, WeightCode);
-			}
-		}
-	}
-	
-	if (OutputCode != INDEX_NONE)
-	{
-		// We've definitely passed the reentrant check here so we're good to call IsResultMaterialAttributes().
-		bool bFoundExpression = false;
-		bool bIsResultMaterialAttributes = false;
-		for (int32 LayerIdx = 0; LayerIdx < Layers.Num(); LayerIdx++)
-		{
-			if (Layers[LayerIdx].HeightInput.Expression)
-			{
-				bool bHeightIsMaterialAttributes = Layers[LayerIdx].HeightInput.Expression->IsResultMaterialAttributes(Layers[LayerIdx].HeightInput.OutputIndex);
-				if (bHeightIsMaterialAttributes)
-				{
-					Compiler->Errorf(TEXT("Height input (%s) does not accept MaterialAttributes"), Layers[LayerIdx].LayerName);
-				}
-			}
-			if (Layers[LayerIdx].LayerInput.Expression)
-			{
-				bool bLayerIsMaterialAttributes = Layers[LayerIdx].LayerInput.Expression->IsResultMaterialAttributes(Layers[LayerIdx].LayerInput.OutputIndex);
-				if (!bFoundExpression)
-				{
-					bFoundExpression = true;
-					bIsResultMaterialAttributes = bLayerIsMaterialAttributes;
-				}
-				else if (bIsResultMaterialAttributes != bLayerIsMaterialAttributes)
-				{
-					Compiler->Error(TEXT("Cannot mix MaterialAttributes and non MaterialAttributes nodes"));
-					break;
-				}
-			}
-		}
-	}
-
-	return OutputCode;
-}
-
-UTexture* UMaterialExpressionLandscapeLayerBlend::GetReferencedTexture()
-{
-	return GEngine->WeightMapPlaceholderTexture;
-}
-
-void UMaterialExpressionLandscapeLayerBlend::GetCaption(TArray<FString>& OutCaptions) const
-{
-	OutCaptions.Add(FString(TEXT("Layer Blend")));
-}
-
-#if WITH_EDITOR
-
-void UMaterialExpressionLandscapeLayerBlend::PostEditChangeProperty(FPropertyChangedEvent& PropertyChangedEvent)
-{
-	Super::PostEditChangeProperty(PropertyChangedEvent);
-
-	// Clear out any height expressions for layers not using height blending
-	for( int32 LayerIdx=0;LayerIdx<Layers.Num();LayerIdx++ )
-	{	
-		if( Layers[LayerIdx].BlendType != LB_HeightBlend )
-		{
-			Layers[LayerIdx].HeightInput.Expression = NULL;
-		}
-	}
-
-	if (PropertyChangedEvent.MemberProperty)
-	{
-		const FName PropertyName = PropertyChangedEvent.MemberProperty->GetFName();
-		if (PropertyName == GET_MEMBER_NAME_CHECKED(UMaterialExpressionLandscapeLayerBlend, Layers))
-		{
-			if (GraphNode)
-			{
-				GraphNode->ReconstructNode();
-			}
-		}
-	}
-}
-#endif
-
-void UMaterialExpressionLandscapeLayerBlend::GetAllParameterNames(TArray<FName> &OutParameterNames, TArray<FGuid> &OutParameterIds)
-{
-	for( int32 LayerIdx=0;LayerIdx<Layers.Num();LayerIdx++ )
-	{
-		FLayerBlendInput& Layer = Layers[LayerIdx];
-
-		int32 CurrentSize = OutParameterNames.Num();
-		OutParameterNames.AddUnique(Layer.LayerName);
-
-		if(CurrentSize != OutParameterNames.Num())
-		{
-			OutParameterIds.Add(ExpressionGUID);
-		}
-	}
-}
-
-///////////////////////////////////////////////////////////////////////////////
-// UMaterialExpressionLandscapeVisibilityMask
-///////////////////////////////////////////////////////////////////////////////
-
-FName UMaterialExpressionLandscapeVisibilityMask::ParameterName = FName("__LANDSCAPE_VISIBILITY__");
-
-UMaterialExpressionLandscapeVisibilityMask::UMaterialExpressionLandscapeVisibilityMask(const FObjectInitializer& ObjectInitializer)
-: Super(ObjectInitializer)
-{
-	// Structure to hold one-time initialization
-	struct FConstructorStatics
-	{
-		FString NAME_Landscape;
-		FConstructorStatics()
-		: NAME_Landscape(LOCTEXT( "Landscape", "Landscape" ).ToString())
-		{
-		}
-	};
-	static FConstructorStatics ConstructorStatics;
-
-	bIsParameterExpression = true;
-	MenuCategories.Add(ConstructorStatics.NAME_Landscape);
-}
-
-int32 UMaterialExpressionLandscapeVisibilityMask::Compile(class FMaterialCompiler* Compiler, int32 OutputIndex, int32 MultiplexIndex)
-{
-	return Compiler->Sub(Compiler->Constant(1.f), Compiler->StaticTerrainLayerWeight(ParameterName, Compiler->Constant(0.f)));
-}
-
-UTexture* UMaterialExpressionLandscapeVisibilityMask::GetReferencedTexture()
-{
-	return GEngine->WeightMapPlaceholderTexture;
-}
-
-void UMaterialExpressionLandscapeVisibilityMask::GetAllParameterNames(TArray<FName> &OutParameterNames, TArray<FGuid> &OutParameterIds)
-{
-	int32 CurrentSize = OutParameterNames.Num();
-	OutParameterNames.AddUnique(ParameterName);
-
-	if(CurrentSize != OutParameterNames.Num())
-	{
-		OutParameterIds.Add(ExpressionGUID);
-	}
-}
-
-void UMaterialExpressionLandscapeVisibilityMask::GetCaption(TArray<FString>& OutCaptions) const
-{
-	OutCaptions.Add(FString(TEXT("Landscape Visibility Mask")));
-}
-
-///////////////////////////////////////////////////////////////////////////////
-// UMaterialExpressionLandscapeLayerCoords
-///////////////////////////////////////////////////////////////////////////////
-UMaterialExpressionLandscapeLayerCoords::UMaterialExpressionLandscapeLayerCoords(const FObjectInitializer& ObjectInitializer)
-	: Super(ObjectInitializer)
-{
-	// Structure to hold one-time initialization
-	struct FConstructorStatics
-	{
-		FString NAME_Landscape;
-		FConstructorStatics()
-			: NAME_Landscape(LOCTEXT( "Landscape", "Landscape" ).ToString())
-		{
-		}
-	};
-	static FConstructorStatics ConstructorStatics;
-
-	MenuCategories.Add(ConstructorStatics.NAME_Landscape);
-	bCollapsed = false;
-}
-
-int32 UMaterialExpressionLandscapeLayerCoords::Compile(class FMaterialCompiler* Compiler, int32 OutputIndex, int32 MultiplexIndex)
-{
-	switch (CustomUVType)
-	{
-	case LCCT_CustomUV0:
-		return Compiler->TextureCoordinate(0, false, false);
-	case LCCT_CustomUV1:
-		return Compiler->TextureCoordinate(1, false, false);
-	case LCCT_CustomUV2:
-		return Compiler->TextureCoordinate(2, false, false);
-	case LCCT_WeightMapUV:
-		return Compiler->TextureCoordinate(3, false, false);
-	default:
-		break;
-	}
-
-	int32 BaseUV;
-
-	switch(MappingType)
-	{
-	case TCMT_Auto:
-	case TCMT_XY: BaseUV = Compiler->TextureCoordinate(0, false, false); break;
-	case TCMT_XZ: BaseUV = Compiler->TextureCoordinate(1, false, false); break;
-	case TCMT_YZ: BaseUV = Compiler->TextureCoordinate(2, false, false); break;
-	default: UE_LOG(LogMaterial, Fatal,TEXT("Invalid mapping type %u"),(uint8)MappingType); return INDEX_NONE;
-	};
-
-	float Scale = (MappingScale == 0.0f) ? 1.0f : 1.0f / MappingScale;
-	int32 RealScale = Compiler->Constant(Scale);
-
-	float Cos = FMath::Cos(MappingRotation * PI / 180.0);
-	float Sin = FMath::Sin(MappingRotation * PI / 180.0);
-
-	int32 ResultUV = INDEX_NONE;
-	int32 TransformedUV = Compiler->Add(
-		Compiler->Mul( RealScale,
-		Compiler->AppendVector(
-		Compiler->Dot(BaseUV, Compiler->Constant2(+Cos,+Sin)),
-		Compiler->Dot(BaseUV, Compiler->Constant2(-Sin,+Cos)) )
-		),
-		Compiler->Constant2(MappingPanU, MappingPanV)
-		);
-
-	auto FeatureLevel = Compiler->GetFeatureLevel();
-	if (FeatureLevel != ERHIFeatureLevel::ES2 && FeatureLevel != ERHIFeatureLevel::ES3_1) // No need to localize UV
-	{
-		ResultUV = TransformedUV;
-	}
-	else
-	{
-		int32 Offset = Compiler->TextureCoordinateOffset();
-		int32 TransformedOffset = 
-			Compiler->Floor(
-			Compiler->Mul(RealScale,
-			Compiler->AppendVector(
-			Compiler->Dot(Offset, Compiler->Constant2(+Cos,+Sin)),
-			Compiler->Dot(Offset, Compiler->Constant2(-Sin,+Cos)) )
-			));
-
-		ResultUV = Compiler->Sub(TransformedUV, TransformedOffset);
-	}
-
-	return ResultUV;
-}
-
-void UMaterialExpressionLandscapeLayerCoords::GetCaption(TArray<FString>& OutCaptions) const
-{
-	OutCaptions.Add(FString(TEXT("LandscapeCoords")));
-}
-
 UMaterialExpressionDepthFade::UMaterialExpressionDepthFade(const FObjectInitializer& ObjectInitializer)
 	: Super(ObjectInitializer)
 {
diff --git a/Engine/Source/Runtime/Engine/Private/Materials/MaterialInstance.cpp b/Engine/Source/Runtime/Engine/Private/Materials/MaterialInstance.cpp
--- a/Engine/Source/Runtime/Engine/Private/Materials/MaterialInstance.cpp
+++ b/Engine/Source/Runtime/Engine/Private/Materials/MaterialInstance.cpp
@@ -7,7 +7,6 @@
 #include "Materials/MaterialExpressionFontSampleParameter.h"
 #include "Materials/MaterialExpressionStaticBoolParameter.h"
 #include "Materials/MaterialExpressionStaticComponentMaskParameter.h"
-#include "Materials/MaterialExpressionLandscapeLayerWeight.h"
 #include "Materials/MaterialInstanceBasePropertyOverrides.h"
 #include "MaterialInstanceSupport.h"
 #include "MaterialShaderType.h"
@@ -58,65 +57,6 @@ void RecacheMaterialInstanceUniformExpressions(const UMaterialInterface* ParentM
 	}
 }
 
-/**
- * This function takes a array of parameter structs and attempts to establish a reference to the expression object each parameter represents.
- * If a reference exists, the function checks to see if the parameter has been renamed.
- *
- * @param Parameters		Array of parameters to operate on.
- * @param ParentMaterial	Parent material to search in for expressions.
- *
- * @return Returns whether or not any of the parameters was changed.
- */
-template <typename ParameterType, typename ExpressionType>
-bool UpdateParameterSet(TArray<ParameterType> &Parameters, UMaterial* ParentMaterial)
-{
-	bool bChanged = false;
-
-	// Loop through all of the parameters and try to either establish a reference to the 
-	// expression the parameter represents, or check to see if the parameter's name has changed.
-	for(int32 ParameterIdx=0; ParameterIdx<Parameters.Num(); ParameterIdx++)
-	{
-		bool bTryToFindByName = true;
-
-		ParameterType &Parameter = Parameters[ParameterIdx];
-
-		if(Parameter.ExpressionGUID.IsValid())
-		{
-			ExpressionType* Expression = ParentMaterial->FindExpressionByGUID<ExpressionType>(Parameter.ExpressionGUID);
-
-			// Check to see if the parameter name was changed.
-			if(Expression)
-			{
-				bTryToFindByName = false;
-
-				if(Parameter.ParameterName != Expression->ParameterName)
-				{
-					Parameter.ParameterName = Expression->ParameterName;
-					bChanged = true;
-				}
-			}
-		}
-
-		// No reference to the material expression exists, so try to find one in the material expression's array if we are in the editor.
-		if(bTryToFindByName && GIsEditor && !FApp::IsGame())
-		{
-			for(int32 ExpressionIndex = 0;ExpressionIndex < ParentMaterial->Expressions.Num();ExpressionIndex++)
-			{
-				ExpressionType* ParameterExpression = Cast<ExpressionType>(ParentMaterial->Expressions[ExpressionIndex]);
-
-				if(ParameterExpression && ParameterExpression->ParameterName == Parameter.ParameterName)
-				{
-					Parameter.ExpressionGUID = ParameterExpression->ExpressionGUID;
-					bChanged = true;
-					break;
-				}
-			}
-		}
-	}
-
-	return bChanged;
-}
-
 FFontParameterValue::ValueType FFontParameterValue::GetValue(const FFontParameterValue& Parameter)
 {
 	ValueType Value = NULL;
@@ -415,7 +355,11 @@ bool UMaterialInstance::UpdateParameters()
 			// Static component mask parameters
 			bDirty = UpdateParameterSet<FStaticComponentMaskParameter, UMaterialExpressionStaticComponentMaskParameter>(StaticParameters.StaticComponentMaskParameters, ParentMaterial) || bDirty;
 
-			bDirty = UpdateParameterSet<FStaticTerrainLayerWeightParameter, UMaterialExpressionLandscapeLayerWeight>(StaticParameters.TerrainLayerWeightParameters, ParentMaterial) || bDirty;
+			// Custom parameters
+			for (const auto& CustomParameterSetUpdater : CustomParameterSetUpdaters)
+			{
+				bDirty |= CustomParameterSetUpdater.Execute(StaticParameters, ParentMaterial);
+			}
 		}
 	}
 	return bDirty;
@@ -1179,42 +1123,10 @@ void UMaterialInstance::GetStaticParameterValues(FStaticParameterSet& OutStaticP
 				}
 			}
 		}
-
-		// TerrainLayerWeight Parameters
-		ParentMaterial->GetAllTerrainLayerWeightParameterNames(ParameterNames, Guids);
-		OutStaticParameters.TerrainLayerWeightParameters.AddZeroed(ParameterNames.Num());
-		for(int32 ParameterIdx=0; ParameterIdx<ParameterNames.Num(); ParameterIdx++)
-		{
-			FStaticTerrainLayerWeightParameter& ParentParameter = OutStaticParameters.TerrainLayerWeightParameters[ParameterIdx];
-			FName ParameterName = ParameterNames[ParameterIdx];
-			FGuid ExpressionId = Guids[ParameterIdx];
-			int32 WeightmapIndex = INDEX_NONE;
-
-			ParentParameter.bOverride = false;
-			ParentParameter.ParameterName = ParameterName;
-			//get the settings from the parent in the MIC chain
-			if(Parent->GetTerrainLayerWeightParameterValue(ParameterName, WeightmapIndex, ExpressionId))
-			{
-				ParentParameter.WeightmapIndex = WeightmapIndex;
-			}
-			ParentParameter.ExpressionGUID = ExpressionId;
-
-			// if the SourceInstance is overriding this parameter, use its settings
-			for(int32 WeightParamIdx = 0; WeightParamIdx < StaticParameters.TerrainLayerWeightParameters.Num(); WeightParamIdx++)
-			{
-				const FStaticTerrainLayerWeightParameter &TerrainLayerWeightParam = StaticParameters.TerrainLayerWeightParameters[WeightParamIdx];
-
-				if(ParameterName == TerrainLayerWeightParam.ParameterName)
-				{
-					ParentParameter.bOverride = TerrainLayerWeightParam.bOverride;
-					if (TerrainLayerWeightParam.bOverride)
-					{
-						ParentParameter.WeightmapIndex = TerrainLayerWeightParam.WeightmapIndex;
-					}
-				}
-			}
-		}
 	}
+
+	// Custom parameters.
+	CustomStaticParametersGetters.Broadcast(OutStaticParameters, this);
 }
 
 void UMaterialInstance::ForceRecompileForRendering()
@@ -2558,3 +2470,11 @@ int32 UMaterialInstance::CompilePropertyEx( class FMaterialCompiler* Compiler, E
 {
 	return Parent ? Parent->CompilePropertyEx(Compiler, Property) : INDEX_NONE;
 }
+
+const FStaticParameterSet& UMaterialInstance::GetStaticParameters() const
+{
+	return StaticParameters;
+}
+
+UMaterialInstance::FCustomStaticParametersGetterDelegate UMaterialInstance::CustomStaticParametersGetters;
+TArray<UMaterialInstance::FCustomParameterSetUpdaterDelegate> UMaterialInstance::CustomParameterSetUpdaters;
diff --git a/Engine/Source/Runtime/Engine/Private/PhysicsEngine/PhysDerivedData.cpp b/Engine/Source/Runtime/Engine/Private/PhysicsEngine/PhysDerivedData.cpp
--- a/Engine/Source/Runtime/Engine/Private/PhysicsEngine/PhysDerivedData.cpp
+++ b/Engine/Source/Runtime/Engine/Private/PhysicsEngine/PhysDerivedData.cpp
@@ -2,10 +2,9 @@
 #include "EnginePrivate.h"
 #include "PhysDerivedData.h"
 #include "TargetPlatform.h"
+#include "IPhysXFormatModule.h"
 
 #if WITH_PHYSX && (WITH_RUNTIME_PHYSICS_COOKING || WITH_EDITOR)
-#include "IPhysXFormatModule.h"
-#include "Landscape/LandscapeMeshCollisionComponent.h"
 
 FDerivedDataPhysXCooker::FDerivedDataPhysXCooker( FName InFormat, UBodySetup* InBodySetup )
 	: BodySetup( InBodySetup )
@@ -26,20 +25,6 @@ FDerivedDataPhysXCooker::FDerivedDataPhysXCooker( FName InFormat, UBodySetup* In
 	InitCooker();
 }
 
-// This constructor only used by ULandscapeMeshCollisionComponent, which always only build TriMesh, not Convex...
-FDerivedDataPhysXCooker::FDerivedDataPhysXCooker( FName InFormat, ULandscapeMeshCollisionComponent* InMeshCollision, bool bMirrored )
-	: BodySetup( NULL )
-	, CollisionDataProvider( InMeshCollision )
-	, Format( InFormat )
-	, bGenerateNormalMesh( !bMirrored )
-	, bGenerateMirroredMesh( bMirrored )
-	, Cooker( NULL )
-{
-	check( InMeshCollision != NULL );
-	DataGuid = InMeshCollision->MeshGuid;
-	InitCooker();
-}
-
 void FDerivedDataPhysXCooker::InitCooker()
 {
 #if WITH_EDITOR
diff --git a/Engine/Source/Runtime/Engine/Private/PhysicsEngine/PhysDerivedData.h b/Engine/Source/Runtime/Engine/Private/PhysicsEngine/PhysDerivedData.h
--- a/Engine/Source/Runtime/Engine/Private/PhysicsEngine/PhysDerivedData.h
+++ b/Engine/Source/Runtime/Engine/Private/PhysicsEngine/PhysDerivedData.h
@@ -10,9 +10,6 @@
 #include "PhysXSupport.h"
 #include "IPhysXFormat.h"
 
-// Forward declarations
-class ULandscapeMeshCollisionComponent;
-
 //////////////////////////////////////////////////////////////////////////
 // PhysX Cooker
 class FDerivedDataPhysXCooker : public FDerivedDataPluginInterface
@@ -30,8 +27,6 @@ private:
 
 public:
 	FDerivedDataPhysXCooker( FName InFormat, UBodySetup* InBodySetup );
-	// This constructor only used by ULandscapeMeshCollisionComponent, which always only build TriMesh, not Convex...
-	FDerivedDataPhysXCooker( FName InFormat, ULandscapeMeshCollisionComponent* InMeshCollision, bool bMirrored );
 
 	virtual const TCHAR* GetPluginName() const override
 	{
diff --git a/Engine/Source/Runtime/Engine/Private/PhysicsEngine/PhysXSupport.h b/Engine/Source/Runtime/Engine/Private/PhysicsEngine/PhysXSupport.h
--- a/Engine/Source/Runtime/Engine/Private/PhysicsEngine/PhysXSupport.h
+++ b/Engine/Source/Runtime/Engine/Private/PhysicsEngine/PhysXSupport.h
@@ -9,6 +9,7 @@
 #if WITH_PHYSX
 
 #include "PhysXIncludes.h"
+#include "EngineLogs.h"
 
 // Whether or not to use the PhysX scene lock
 #ifndef USE_SCENE_LOCK
@@ -199,10 +200,10 @@ extern int32					GNumPhysXConvexMeshes;
 extern TArray<PxConvexMesh*>	GPhysXPendingKillConvex;
 
 /** Array of PxTriangleMesh objects which are awaiting cleaning up. */
-extern TArray<PxTriangleMesh*>	GPhysXPendingKillTriMesh;
+extern ENGINE_API TArray<PxTriangleMesh*>	GPhysXPendingKillTriMesh;
 
 /** Array of PxHeightField objects which are awaiting cleaning up. */
-extern TArray<PxHeightField*>	GPhysXPendingKillHeightfield;
+extern ENGINE_API TArray<PxHeightField*>	GPhysXPendingKillHeightfield;
 
 /** Array of PxMaterial objects which are awaiting cleaning up. */
 extern TArray<PxMaterial*>		GPhysXPendingKillMaterial;
diff --git a/Engine/Source/Runtime/Engine/Private/UnrealEngine.cpp b/Engine/Source/Runtime/Engine/Private/UnrealEngine.cpp
--- a/Engine/Source/Runtime/Engine/Private/UnrealEngine.cpp
+++ b/Engine/Source/Runtime/Engine/Private/UnrealEngine.cpp
@@ -7007,17 +7007,6 @@ DEFINE_STAT(STAT_PlatformMessageTime);
 DEFINE_STAT(STAT_FrameSyncTime);
 DEFINE_STAT(STAT_DeferredTickTime);
 
-/** Landscape stats */
-
-DEFINE_STAT(STAT_LandscapeDynamicDrawTime);
-DEFINE_STAT(STAT_LandscapeStaticDrawLODTime);
-DEFINE_STAT(STAT_LandscapeVFDrawTime);
-DEFINE_STAT(STAT_LandscapeComponents);
-DEFINE_STAT(STAT_LandscapeDrawCalls);
-DEFINE_STAT(STAT_LandscapeTriangles);
-DEFINE_STAT(STAT_LandscapeVertexMem);
-DEFINE_STAT(STAT_LandscapeComponentMem);
-
 /** Input stat */
 DEFINE_STAT(STAT_InputTime);
 DEFINE_STAT(STAT_InputLatencyTime);
diff --git a/Engine/Source/Runtime/Engine/Private/World.cpp b/Engine/Source/Runtime/Engine/Private/World.cpp
--- a/Engine/Source/Runtime/Engine/Private/World.cpp
+++ b/Engine/Source/Runtime/Engine/Private/World.cpp
@@ -47,12 +47,10 @@
 
 #include "MallocProfiler.h"
 
-#include "Landscape/LandscapeInfo.h"
 #include "Particles/ParticleEventManager.h"
 
 #include "EngineModule.h"
 #include "ContentStreaming.h"
-#include "Runtime/Engine/Classes/Landscape/LandscapeComponent.h"
 #include "RendererInterface.h"
 #include "DataChannel.h"
 #include "ShaderCompiler.h"
@@ -73,6 +71,7 @@ TMap<FName, EWorldType::Type> UWorld::WorldTypePreLoadMap;
 
 FWorldDelegates::FWorldInitializationEvent FWorldDelegates::OnPreWorldInitialization;
 FWorldDelegates::FWorldInitializationEvent FWorldDelegates::OnPostWorldInitialization;
+FWorldDelegates::FWorldPostDuplicateEvent FWorldDelegates::OnPostDuplicate;
 FWorldDelegates::FWorldCleanupEvent FWorldDelegates::OnWorldCleanup;
 FWorldDelegates::FWorldEvent FWorldDelegates::OnPreWorldFinishDestroy;
 FWorldDelegates::FOnLevelChanged FWorldDelegates::LevelAddedToWorld;
@@ -139,10 +138,6 @@ void UWorld::Serialize( FArchive& Ar )
 		Ar << GameState;
 		Ar << AuthorityGameMode;
 		Ar << NetworkManager;
-#if WITH_EDITORONLY_DATA
-		Ar << LandscapeInfoMap;
-		[REDACTED_EMAIL]: How to handle w/out having a cooker/pakcager?
-#endif	//#if WITH_EDITORONLY_DATA
 
 		Ar << NavigationSystem;
 		Ar << AvoidanceManager;
@@ -205,11 +200,6 @@ void UWorld::AddReferencedObjects(UObject* InThis, FReferenceCollector& Collecto
 		Collector.AddReferencedObject( This->NavigationSystem, This );
 		Collector.AddReferencedObject( This->AvoidanceManager, This );
 
-#if WITH_EDITORONLY_DATA
-		Collector.AddReferencedObjects(This->LandscapeInfoMap, This);
-		[REDACTED_EMAIL]: How to handle w/out having a cooker/packager?
-#endif	//#if WITH_EDITORONLY_DATA
-
 		for( int32 Index = 0; Index < This->ExtraReferencedObjects.Num(); Index++ )
 		{
 			Collector.AddReferencedObject( This->ExtraReferencedObjects[Index], This );
@@ -218,6 +208,10 @@ void UWorld::AddReferencedObjects(UObject* InThis, FReferenceCollector& Collecto
 		{
 			Collector.AddReferencedObject( This->StreamingLevels[Index], This );
 		}
+		for (auto* Object : This->PerModuleDataObjects)
+		{
+			Collector.AddReferencedObject(Object, This);
+		}
 	}
 #endif
 
@@ -240,16 +234,17 @@ bool UWorld::Rename(const TCHAR* InName, UObject* NewOuter, ERenameFlags Flags)
 		}
 	}
 
-	// Also rename all textures and materials used by landscape components
-	TArray<UObject*> LandscapeTexturesAndMaterials;
-	GetLandscapeTexturesAndMaterials(PersistentLevel, LandscapeTexturesAndMaterials);
+	// UWorld rename is also renaming package, so move all direct children
+	// of the old package.
+
 	UPackage* PersistentLevelPackage = PersistentLevel->GetOutermost();
-	for (auto* OldTexOrMat : LandscapeTexturesAndMaterials)
+	if (PersistentLevelPackage != NewOuter)
 	{
-		if (OldTexOrMat && OldTexOrMat->GetOuter() == PersistentLevelPackage)
+		TArray<UObject*> ReferencedObjects;
+		GetObjectsWithOuter(PersistentLevelPackage, ReferencedObjects, false);
+		for (auto* ReferencedObject : ReferencedObjects)
 		{
-			// The names for these objects are not important, just generate a new name to avoid collisions
-			if ( !OldTexOrMat->Rename(nullptr, NewOuter, Flags) )
+			if (ReferencedObject != this && !ReferencedObject->Rename(nullptr, NewOuter, Flags))
 			{
 				return false;
 			}
@@ -336,86 +331,9 @@ void UWorld::PostDuplicate(bool bDuplicateForPIE)
 
 		// Make sure PKG_ContainsMap is set
 		MyPackage->ThisContainsMap();
-
-#if WITH_EDITOR
-		TArray<UObject*> ObjectsToFixReferences;
-		TMap<UObject*, UObject*> ReplacementMap;
-
-		// Add the world to the list of objects in which to fix up references.
-		ObjectsToFixReferences.Add(this);
-
-		// Gather the textures
-		TArray<UTexture2D*> LightMapsAndShadowMaps;
-		GetLightMapsAndShadowMaps(PersistentLevel, LightMapsAndShadowMaps);
-
-		// Duplicate the textures, if any
-		for (auto* Tex : LightMapsAndShadowMaps)
-		{
-			if (Tex && Tex->GetOutermost() != MyPackage)
-			{
-				UObject* NewTex = StaticDuplicateObject(Tex, MyPackage, *Tex->GetName());
-				ReplacementMap.Add(Tex, NewTex);
-			}
-		}
-
-		// Also duplicate all textures and materials used by landscape components
-		TArray<UObject*> LandscapeTexturesAndMaterials;
-		GetLandscapeTexturesAndMaterials(PersistentLevel, LandscapeTexturesAndMaterials);
-		for (auto* OldTexOrMat : LandscapeTexturesAndMaterials)
-		{
-			if (OldTexOrMat && OldTexOrMat->GetOuter() != MyPackage)
-			{
-				// The names for these objects are not important, just generate a new name to avoid collisions
-				UObject* NewTextureOrMaterial = StaticDuplicateObject(OldTexOrMat, MyPackage, nullptr);
-				ReplacementMap.Add(OldTexOrMat, NewTextureOrMaterial);
-
-				// Materials hold references to the textures being moved, so they will need references fixed up as well
-				if ( OldTexOrMat->IsA(UMaterialInterface::StaticClass()) )
-				{
-					ObjectsToFixReferences.Add(NewTextureOrMaterial);
-				}
-			}
-		}
-
-		// Duplicate the level script blueprint generated classes as well
-		const bool bDontCreate = true;
-		UBlueprint* LevelScriptBlueprint = PersistentLevel->GetLevelScriptBlueprint(bDontCreate);
-		if (LevelScriptBlueprint)
-		{
-			UObject* OldGeneratedClass = LevelScriptBlueprint->GeneratedClass;
-			if (OldGeneratedClass)
-			{
-				UObject* NewGeneratedClass = StaticDuplicateObject(OldGeneratedClass, MyPackage, *OldGeneratedClass->GetName());
-				ReplacementMap.Add(OldGeneratedClass, NewGeneratedClass);
-
-				// The class may have referenced a lightmap or landscape resource that is also being duplicated. Add it to the list of objects that need references fixed up.
-				ObjectsToFixReferences.Add(NewGeneratedClass);
-			}
-
-			UObject* OldSkeletonClass = LevelScriptBlueprint->SkeletonGeneratedClass;
-			if (OldSkeletonClass)
-			{
-				UObject* NewSkeletonClass = StaticDuplicateObject(OldSkeletonClass, MyPackage, *OldSkeletonClass->GetName());
-				ReplacementMap.Add(OldSkeletonClass, NewSkeletonClass);
-
-				// The class may have referenced a lightmap or landscape resource that is also being duplicated. Add it to the list of objects that need references fixed up.
-				ObjectsToFixReferences.Add(NewSkeletonClass);
-			}
-		}
-
-		// Now replace references from the old textures/classes to the new ones, if any were duplicated
-		if (ReplacementMap.Num() > 0)
-		{
-			const bool bNullPrivateRefs = false;
-			const bool bIgnoreOuterRef = true;
-			const bool bIgnoreArchetypeRef = false;
-			for ( auto* Obj : ObjectsToFixReferences )
-			{
-				FArchiveReplaceObjectRef<UObject> ReplaceAr(Obj, ReplacementMap, bNullPrivateRefs, bIgnoreOuterRef, bIgnoreArchetypeRef);
-			}
-		}
-#endif // WITH_EDITOR
 	}
+
+	FWorldDelegates::OnPostDuplicate.Broadcast(this, bDuplicateForPIE);
 }
 
 void UWorld::FinishDestroy()
@@ -5312,27 +5230,6 @@ void UWorld::CreateFXSystem()
 	}
 }
 
-void UWorld::GetLandscapeTexturesAndMaterials(ULevel* Level, TArray<UObject*>& OutTexturesAndMaterials)
-{
-	UObject* SearchLevel = Level;
-	if (!SearchLevel)
-	{
-		SearchLevel = PersistentLevel;
-	}
-
-	TArray<UObject*> ObjectsInLevel;
-	const bool bIncludeNestedObjects = true;
-	GetObjectsWithOuter(SearchLevel, ObjectsInLevel, bIncludeNestedObjects);
-	for (auto* ObjInLevel : ObjectsInLevel)
-	{
-		ULandscapeComponent* LandscapeComponent = Cast<ULandscapeComponent>(ObjInLevel);
-		if (LandscapeComponent)
-		{
-			LandscapeComponent->GetGeneratedTexturesAndMaterialInstances(OutTexturesAndMaterials);
-		}
-	}
-}
-
 #if WITH_EDITOR
 void UWorld::ChangeFeatureLevel(ERHIFeatureLevel::Type InFeatureLevel)
 {
diff --git a/Engine/Source/Runtime/Engine/Public/AI/NavigationSystemHelpers.h b/Engine/Source/Runtime/Engine/Public/AI/NavigationSystemHelpers.h
--- a/Engine/Source/Runtime/Engine/Public/AI/NavigationSystemHelpers.h
+++ b/Engine/Source/Runtime/Engine/Public/AI/NavigationSystemHelpers.h
@@ -2,6 +2,7 @@
 
 #pragma once
 
+#include "AI/NavigationModifier.h"
 #include "AI/Navigation/NavLinkDefinition.h"
 
 #if WITH_PHYSX
diff --git a/Engine/Source/Runtime/Engine/Public/ComponentReregisterContext.h b/Engine/Source/Runtime/Engine/Public/ComponentReregisterContext.h
--- a/Engine/Source/Runtime/Engine/Public/ComponentReregisterContext.h
+++ b/Engine/Source/Runtime/Engine/Public/ComponentReregisterContext.h
@@ -3,6 +3,7 @@
 #pragma once
 
 #include "ContentStreaming.h"
+#include "AI/Navigation/NavigationSystem.h"
 
 ENGINE_API DECLARE_LOG_CATEGORY_EXTERN(LogActorComponent, Log, All);
 
diff --git a/Engine/Source/Runtime/Engine/Public/ContentStreaming.h b/Engine/Source/Runtime/Engine/Public/ContentStreaming.h
--- a/Engine/Source/Runtime/Engine/Public/ContentStreaming.h
+++ b/Engine/Source/Runtime/Engine/Public/ContentStreaming.h
@@ -6,6 +6,8 @@
 
 #pragma once
 
+#include "Audio.h"
+
 /*-----------------------------------------------------------------------------
 	Stats.
 -----------------------------------------------------------------------------*/
@@ -63,6 +65,7 @@ class UActorComponent;
 class UPrimitiveComponent;
 class AActor;
 class UTexture2D;
+class FSoundSource;
 template<typename T>
 class FAsyncTask;
 struct FStreamingManagerTexture;
diff --git a/Engine/Source/Runtime/Engine/Public/EngineLogs.h b/Engine/Source/Runtime/Engine/Public/EngineLogs.h
--- a/Engine/Source/Runtime/Engine/Public/EngineLogs.h
+++ b/Engine/Source/Runtime/Engine/Public/EngineLogs.h
@@ -23,5 +23,4 @@ ENGINE_API DECLARE_LOG_CATEGORY_EXTERN(LogNetDormancy, Warning, All);
 ENGINE_API DECLARE_LOG_CATEGORY_EXTERN(LogSkeletalControl, Warning, All);
 ENGINE_API DECLARE_LOG_CATEGORY_EXTERN(LogSubtitle, Log, All);
 ENGINE_API DECLARE_LOG_CATEGORY_EXTERN(LogTexture, Log, All);
-ENGINE_API DECLARE_LOG_CATEGORY_EXTERN(LogLandscape, Warning, All);
 ENGINE_API DECLARE_LOG_CATEGORY_EXTERN(LogPlayerManagement, Error, All);
\ No newline at end of file
diff --git a/Engine/Source/Runtime/Engine/Public/EngineStats.h b/Engine/Source/Runtime/Engine/Public/EngineStats.h
--- a/Engine/Source/Runtime/Engine/Public/EngineStats.h
+++ b/Engine/Source/Runtime/Engine/Public/EngineStats.h
@@ -164,18 +164,6 @@ DECLARE_CYCLE_STAT_EXTERN(TEXT("Get Batched Element Time"),STAT_Canvas_GetBatchE
 DECLARE_CYCLE_STAT_EXTERN(TEXT("Add Material Tile Time"),STAT_Canvas_AddTileRenderTime,STATGROUP_Canvas, );
 DECLARE_DWORD_COUNTER_STAT_EXTERN(TEXT("Num Batches Created"),STAT_Canvas_NumBatchesCreated,STATGROUP_Canvas, );	
 
-/**
- * Landscape stats
- */
-DECLARE_CYCLE_STAT_EXTERN(TEXT("Landscape Dynamic Draw Time"),STAT_LandscapeDynamicDrawTime,STATGROUP_Landscape, );
-DECLARE_CYCLE_STAT_EXTERN(TEXT("Landscape Static Draw LOD Time"),STAT_LandscapeStaticDrawLODTime,STATGROUP_Landscape, );
-DECLARE_CYCLE_STAT_EXTERN(TEXT("LandscapeVF Draw Time"),STAT_LandscapeVFDrawTime,STATGROUP_Landscape, );
-DECLARE_DWORD_COUNTER_STAT_EXTERN(TEXT("Landscape Triangles"),STAT_LandscapeTriangles,STATGROUP_Landscape, );
-DECLARE_DWORD_COUNTER_STAT_EXTERN(TEXT("Landscape Render Passes"),STAT_LandscapeComponents,STATGROUP_Landscape, );
-DECLARE_DWORD_COUNTER_STAT_EXTERN(TEXT("Landscape DrawCalls"),STAT_LandscapeDrawCalls,STATGROUP_Landscape, );
-DECLARE_MEMORY_STAT_EXTERN(TEXT("Landscape Vertex Mem"),STAT_LandscapeVertexMem,STATGROUP_Landscape, );
-DECLARE_DWORD_COUNTER_STAT_EXTERN(TEXT("Component Mem"),STAT_LandscapeComponentMem,STATGROUP_Landscape, );
-
 /**
  * Network stats counters
  */
diff --git a/Engine/Source/Runtime/Engine/Public/LightMap.h b/Engine/Source/Runtime/Engine/Public/LightMap.h
--- a/Engine/Source/Runtime/Engine/Public/LightMap.h
+++ b/Engine/Source/Runtime/Engine/Public/LightMap.h
@@ -121,7 +121,7 @@ private:
 
 
 /** Lightmap reference serializer */
-extern FArchive& operator<<(FArchive& Ar,FLightMap*& R);
+extern ENGINE_API FArchive& operator<<(FArchive& Ar, FLightMap*& R);
 
 /** 
  * Incident lighting for a single sample, as produced by a lighting build. 
@@ -209,7 +209,7 @@ public:
 	 * @param	BasisIndex - The basis index.
 	 * @return	true if the specified basis has a valid lightmap texture, otherwise false
 	 */
-	bool IsValid(uint32 BasisIndex) const;
+	ENGINE_API bool IsValid(uint32 BasisIndex) const;
 
 	const FVector2D& GetCoordinateScale() const { return CoordinateScale; }
 	const FVector2D& GetCoordinateBias() const { return CoordinateBias; }
@@ -244,7 +244,7 @@ public:
 	 * @param	InPaddingType - the method for padding the lightmap.
 	 * @param	LightmapFlags - flags that determine how the lightmap is stored (e.g. streamed or not)
 	 */
-	static class FLightMap2D* AllocateLightMap(UObject* LightMapOuter, struct FQuantizedLightmapData*& SourceQuantizedData, 
+	static ENGINE_API class FLightMap2D* AllocateLightMap(UObject* LightMapOuter, struct FQuantizedLightmapData*& SourceQuantizedData,
 		const FBoxSphereBounds& Bounds, ELightMapPaddingType InPaddingType, ELightMapFlags InLightmapFlags );
 
 	/**
@@ -490,7 +490,7 @@ struct FQuantizedLightmapData
 		bHasSkyShadowing(false)
 	{}
 
-	bool HasNonZeroData() const;
+	ENGINE_API bool HasNonZeroData() const;
 };
 
 /**
diff --git a/Engine/Source/Runtime/Engine/Public/MaterialExpressionIO.h b/Engine/Source/Runtime/Engine/Public/MaterialExpressionIO.h
--- a/Engine/Source/Runtime/Engine/Public/MaterialExpressionIO.h
+++ b/Engine/Source/Runtime/Engine/Public/MaterialExpressionIO.h
@@ -40,7 +40,7 @@ struct FExpressionInput
 	{
 	}
 
-	int32 Compile(class FMaterialCompiler* Compiler, int32 MultiplexIndex=INDEX_NONE);
+	ENGINE_API int32 Compile(class FMaterialCompiler* Compiler, int32 MultiplexIndex=INDEX_NONE);
 
 	/**
 	 * Tests if the input has a material expression connected to it
diff --git a/Engine/Source/Runtime/Engine/Public/PhysicsPublic.h b/Engine/Source/Runtime/Engine/Public/PhysicsPublic.h
--- a/Engine/Source/Runtime/Engine/Public/PhysicsPublic.h
+++ b/Engine/Source/Runtime/Engine/Public/PhysicsPublic.h
@@ -309,7 +309,7 @@ public:
 
 #if WITH_PHYSX
 	/** Utility for looking up the PxScene of the given EPhysicsSceneType associated with this FPhysScene.  SceneType must be in the range [0,PST_MAX). */
-	physx::PxScene*					GetPhysXScene(uint32 SceneType);
+	ENGINE_API physx::PxScene*					GetPhysXScene(uint32 SceneType);
 
 #if WITH_VEHICLE
 	/** Get the vehicle manager */
diff --git a/Engine/Source/Runtime/Engine/Public/SceneManagement.h b/Engine/Source/Runtime/Engine/Public/SceneManagement.h
--- a/Engine/Source/Runtime/Engine/Public/SceneManagement.h
+++ b/Engine/Source/Runtime/Engine/Public/SceneManagement.h
@@ -1803,7 +1803,8 @@ extern ENGINE_API void ApplyViewModeOverrides(
 extern ENGINE_API void DrawUVs(FViewport* InViewport, FCanvas* InCanvas, int32 InTextYPos, const int32 LODLevel, int32 UVChannel, TArray<FVector2D> SelectedEdgeTexCoords, class FStaticMeshRenderData* StaticMeshRenderData, class FStaticLODModel* SkeletalMeshRenderData );
 
 /** Returns true if the Material and Vertex Factory combination require adjacency information. */
-bool RequiresAdjacencyInformation( class UMaterialInterface* Material, const class FVertexFactoryType* VertexFactoryType, ERHIFeatureLevel::Type InFeatureLevel );
+ENGINE_API bool RequiresAdjacencyInformation(class UMaterialInterface* Material, const class FVertexFactoryType* VertexFactoryType, ERHIFeatureLevel::Type InFeatureLevel);
+
 /**
  * Computes the screen size of a given sphere bounds in the given view
  * @param Origin - Origin of the bounds in world space
diff --git a/Engine/Source/Runtime/Engine/Public/ShadowMap.h b/Engine/Source/Runtime/Engine/Public/ShadowMap.h
--- a/Engine/Source/Runtime/Engine/Public/ShadowMap.h
+++ b/Engine/Source/Runtime/Engine/Public/ShadowMap.h
@@ -285,7 +285,7 @@ public:
 	 */
 	static void EncodeSingleTexture(struct FShadowMapPendingTexture& PendingTexture, UShadowMapTexture2D* Texture, TArray< TArray<FFourDistanceFieldSamples> >& MipData);
 
-	static FShadowMap2D* AllocateShadowMap(
+	static ENGINE_API FShadowMap2D* AllocateShadowMap(
 		UObject* LightMapOuter, 
 		const TMap<ULightComponent*,FShadowMapData2D*>& ShadowMapData,
 		const FBoxSphereBounds& Bounds, 
@@ -345,4 +345,4 @@ private:
 
 
 /** Shadowmap reference serializer */
-extern FArchive& operator<<(FArchive& Ar,FShadowMap*& R);
+extern ENGINE_API FArchive& operator<<(FArchive& Ar, FShadowMap*& R);
diff --git a/Engine/Source/Runtime/Engine/Public/StaticLighting.h b/Engine/Source/Runtime/Engine/Public/StaticLighting.h
--- a/Engine/Source/Runtime/Engine/Public/StaticLighting.h
+++ b/Engine/Source/Runtime/Engine/Public/StaticLighting.h
@@ -6,6 +6,11 @@
 
 #pragma once
 
+#include "Vector.h"
+#include "Vector2D.h"
+#include "Components.h"
+#include "Components/PrimitiveComponent.h"
+#include "Components/LightComponent.h"
 
 // Forward declarations.
 class FStaticLightingTextureMapping;
@@ -150,7 +155,7 @@ public:
 	TArray<TRefCountPtr<FStaticLightingMesh> > OtherMeshLODs;
 
 	/** Initialization constructor. */
-	FStaticLightingMesh(
+	ENGINE_API FStaticLightingMesh(
 		int32 InNumTriangles,
 		int32 InNumShadingTriangles,
 		int32 InNumVertices,
@@ -378,7 +383,7 @@ public:
 	const bool bBilinearFilter;
 
 	/** Initialization constructor. */
-	FStaticLightingTextureMapping(FStaticLightingMesh* InMesh,UObject* InOwner,int32 InSizeX,int32 InSizeY,int32 InLightmapTextureCoordinateIndex,bool bInBilinearFilter=true);
+	ENGINE_API FStaticLightingTextureMapping(FStaticLightingMesh* InMesh, UObject* InOwner, int32 InSizeX, int32 InSizeY, int32 InLightmapTextureCoordinateIndex, bool bInBilinearFilter = true);
 
 	/**
 	 * Called when the static lighting has been computed to apply it to the mapping's owner.
diff --git a/Engine/Source/Runtime/Engine/Public/StaticMeshResources.h b/Engine/Source/Runtime/Engine/Public/StaticMeshResources.h
--- a/Engine/Source/Runtime/Engine/Public/StaticMeshResources.h
+++ b/Engine/Source/Runtime/Engine/Public/StaticMeshResources.h
@@ -6,6 +6,8 @@
 
 #pragma once
 
+#include "PhysicsEngine/BodySetup.h"
+#include "Engine/StaticMesh.h"
 #include "RawIndexBuffer.h"
 #include "TextureLayout3d.h"
 #include "LocalVertexFactory.h"
diff --git a/Engine/Source/Runtime/Engine/Public/TessellationRendering.h b/Engine/Source/Runtime/Engine/Public/TessellationRendering.h
--- a/Engine/Source/Runtime/Engine/Public/TessellationRendering.h
+++ b/Engine/Source/Runtime/Engine/Public/TessellationRendering.h
@@ -8,6 +8,6 @@
 #define __TESSELLATIONRENDERING_H__
 
 /** Returns true if the Material and Vertex Factory combination require adjacency information. */
-extern bool RequiresAdjacencyInformation(UMaterialInterface* Material, const FVertexFactoryType* VertexFactoryType, ERHIFeatureLevel::Type InFeatureLevel);
+extern ENGINE_API bool RequiresAdjacencyInformation(UMaterialInterface* Material, const FVertexFactoryType* VertexFactoryType, ERHIFeatureLevel::Type InFeatureLevel);
 
 #endif
diff --git a/Engine/Source/Runtime/Engine/Public/UnrealClient.h b/Engine/Source/Runtime/Engine/Public/UnrealClient.h
--- a/Engine/Source/Runtime/Engine/Public/UnrealClient.h
+++ b/Engine/Source/Runtime/Engine/Public/UnrealClient.h
@@ -7,7 +7,9 @@
 #pragma once
 
 #include "RenderResource.h"
-#include "HitProxies.h"		
+#include "HitProxies.h"
+#include "InputCoreTypes.h"
+#include "EngineBaseTypes.h"
 
 class FCanvas;
 class FViewportClient;
diff --git a/Engine/Source/Runtime/Engine/Public/UnrealEngine.h b/Engine/Source/Runtime/Engine/Public/UnrealEngine.h
--- a/Engine/Source/Runtime/Engine/Public/UnrealEngine.h
+++ b/Engine/Source/Runtime/Engine/Public/UnrealEngine.h
@@ -6,6 +6,8 @@
 
 #pragma once
 
+#include "Engine/Engine.h"
+
 //
 //	FLocalPlayerIterator - Iterates over local players in the game.
 //	There are no advantages to using this over GEngine->GetLocalPlayerIterator(GetWorld());
diff --git a/Engine/Source/Runtime/Landscape/Classes/Landscape.h b/Engine/Source/Runtime/Landscape/Classes/Landscape.h
new file mode 100644
--- /dev/null
+++ b/Engine/Source/Runtime/Landscape/Classes/Landscape.h
@@ -0,0 +1,93 @@
+// Copyright 1998-2014 Epic Games, Inc. All Rights Reserved.
+
+#pragma once
+
+#include "LandscapeProxy.h"
+#include "LandscapeInfoMap.h"
+#include "Engine/World.h"
+
+#include "Landscape.generated.h"
+
+UENUM()
+enum ELandscapeSetupErrors
+{
+	LSE_None,
+	// No Landscape Info available
+	LSE_NoLandscapeInfo,
+	// There was already component with same X,Y
+	LSE_CollsionXY,
+	// No Layer Info, need to add proper layers
+	LSE_NoLayerInfo,
+	LSE_MAX,
+};
+
+UCLASS(Placeable, hidecategories=LandscapeProxy, showcategories=(Display, Movement, Collision, Lighting, LOD, Input), MinimalAPI)
+class ALandscape : public ALandscapeProxy
+{
+	GENERATED_UCLASS_BODY()
+
+	// Make a key for XYtoComponentMap
+	static FIntPoint MakeKey( int32 X, int32 Y ) { return FIntPoint(X, Y); }
+	static void UnpackKey( FIntPoint Key, int32& OutX, int32& OutY ) { OutX = Key.X; OutY = Key.Y; }
+
+	// Begin AActor Interface
+#if WITH_EDITOR
+	virtual void CheckForErrors() override;
+#endif
+	// End AActor Interface
+
+	// Begin ALandscapeProxy Interface
+	LANDSCAPE_API virtual ALandscape* GetLandscapeActor() override;
+#if WITH_EDITOR
+	virtual UMaterialInterface* GetLandscapeMaterial() const override;
+	virtual UMaterialInterface* GetLandscapeHoleMaterial() const override;
+	// End ALandscapeProxy Interface
+
+	LANDSCAPE_API bool HasAllComponent(); // determine all component is in this actor
+	
+	// Include Components with overlapped vertices
+	LANDSCAPE_API static void CalcComponentIndicesOverlap(const int32 X1, const int32 Y1, const int32 X2, const int32 Y2, const int32 ComponentSizeQuads, 
+		int32& ComponentIndexX1, int32& ComponentIndexY1, int32& ComponentIndexX2, int32& ComponentIndexY2);
+
+	// Exclude Components with overlapped vertices
+	static void CalcComponentIndices(const int32 X1, const int32 Y1, const int32 X2, const int32 Y2, const int32 ComponentSizeQuads, 
+		int32& ComponentIndexX1, int32& ComponentIndexY1, int32& ComponentIndexX2, int32& ComponentIndexY2);
+
+	static void SplitHeightmap(ULandscapeComponent* Comp, bool bMoveToCurrentLevel = false);
+	
+	// Begin UObject interface.
+	virtual void PreSave() override;
+	virtual void PostEditChangeProperty(FPropertyChangedEvent& PropertyChangedEvent) override;
+	virtual void PostEditMove(bool bFinished) override;
+#endif
+	virtual void PostLoad() override;
+	// End UObject Interface
+};
+
+LANDSCAPE_API DECLARE_LOG_CATEGORY_EXTERN(LogLandscape, Warning, All);
+
+/**
+ * Landscape stats
+ */
+
+DECLARE_CYCLE_STAT_EXTERN(TEXT("Landscape Dynamic Draw Time"), STAT_LandscapeDynamicDrawTime, STATGROUP_Landscape, );
+DECLARE_CYCLE_STAT_EXTERN(TEXT("Landscape Static Draw LOD Time"), STAT_LandscapeStaticDrawLODTime, STATGROUP_Landscape, );
+DECLARE_CYCLE_STAT_EXTERN(TEXT("LandscapeVF Draw Time"), STAT_LandscapeVFDrawTime, STATGROUP_Landscape, );
+DECLARE_DWORD_COUNTER_STAT_EXTERN(TEXT("Landscape Triangles"), STAT_LandscapeTriangles, STATGROUP_Landscape, );
+DECLARE_DWORD_COUNTER_STAT_EXTERN(TEXT("Landscape Render Passes"), STAT_LandscapeComponents, STATGROUP_Landscape, );
+DECLARE_DWORD_COUNTER_STAT_EXTERN(TEXT("Landscape DrawCalls"), STAT_LandscapeDrawCalls, STATGROUP_Landscape, );
+DECLARE_MEMORY_STAT_EXTERN(TEXT("Landscape Vertex Mem"), STAT_LandscapeVertexMem, STATGROUP_Landscape, );
+DECLARE_DWORD_COUNTER_STAT_EXTERN(TEXT("Component Mem"), STAT_LandscapeComponentMem, STATGROUP_Landscape, );
+
+#if WITH_EDITORONLY_DATA
+
+/**
+ * Gets landscape-specific data for given world.
+ *
+ * @param World A pointer to world that return data should be associated with.
+ *
+ * @returns Landscape-specific data associated with given world.
+ */
+LANDSCAPE_API ULandscapeInfoMap& GetLandscapeInfoMap(UWorld* World);
+
+#endif // WITH_EDITORONLY_DATA
\ No newline at end of file
diff --git a/Engine/Source/Runtime/Landscape/Classes/LandscapeComponent.h b/Engine/Source/Runtime/Landscape/Classes/LandscapeComponent.h
new file mode 100644
--- /dev/null
+++ b/Engine/Source/Runtime/Landscape/Classes/LandscapeComponent.h
@@ -0,0 +1,498 @@
+// Copyright 1998-2014 Epic Games, Inc. All Rights Reserved.
+
+#pragma once
+
+#include "SceneTypes.h"
+#include "StaticLighting.h"
+#include "Components/PrimitiveComponent.h"
+
+#include "LandscapeComponent.generated.h"
+
+class ULandscapeLayerInfoObject;
+class ULandscapeInfo;
+class ALandscapeProxy;
+
+class FLandscapeComponentDerivedData
+{
+	/** The compressed Landscape component data for mobile rendering. Serialized to disk. 
+	    On device, freed once it has been decompressed. */
+	TArray<uint8> CompressedLandscapeData;
+
+public:
+	/** Returns true if there is any valid platform data */
+	bool HasValidPlatformData() const
+	{
+		return CompressedLandscapeData.Num() != 0;
+	}
+
+	/** Initializes the compressed data from an uncompressed source. */
+	void InitializeFromUncompressedData(const TArray<uint8>& UncompressedData);
+
+	/** Decompresses and returns the data. Also frees the compressed data from memory when running with cooked data */
+	void GetUncompressedData(TArray<uint8>& OutUncompressedData);
+
+	/** Constructs a key string for the DDC that uniquely identifies a the Landscape component's derived data. */
+	static FString GetDDCKeyString(const FGuid& StateId);
+
+	/** Loads the platform data from DDC */
+	bool LoadFromDDC(const FGuid& StateId);
+
+	/** Saves the compressed platform data to the DDC */
+	void SaveToDDC(const FGuid& StateId);
+
+	/* Serializer */
+	friend FArchive& operator<<(FArchive& Ar, FLandscapeComponentDerivedData& Data);
+};
+
+/* Used to uniquely reference a landscape vertex in a component, and generate a key suitable for a TMap. */
+struct FLandscapeVertexRef
+{
+	FLandscapeVertexRef(int16 InX, int16 InY, int8 InSubX, int8 InSubY)
+	: X(InX)
+	, Y(InY)
+	, SubX(InSubX)
+	, SubY(InSubY)
+	{}
+	int16 X;
+	int16 Y;
+	int8 SubX;
+	int8 SubY;
+
+	uint64 MakeKey() const
+	{
+		return (uint64)X << 32 | (uint64)Y << 16 | (uint64)SubX << 8 | (uint64)SubY;
+	}
+};
+
+/** Stores information about which weightmap texture and channel each layer is stored */
+USTRUCT()
+struct FWeightmapLayerAllocationInfo
+{
+	GENERATED_USTRUCT_BODY()
+
+	UPROPERTY()
+	FName LayerName_DEPRECATED;
+
+	UPROPERTY()
+	ULandscapeLayerInfoObject* LayerInfo;
+
+	UPROPERTY()
+	uint8 WeightmapTextureIndex;
+
+	UPROPERTY()
+	uint8 WeightmapTextureChannel;
+
+
+	FWeightmapLayerAllocationInfo()
+		: WeightmapTextureIndex(0)
+		, WeightmapTextureChannel(0)
+	{
+	}
+
+
+	FWeightmapLayerAllocationInfo(ULandscapeLayerInfoObject* InLayerInfo)
+		:	LayerInfo(InLayerInfo)
+		,	WeightmapTextureIndex(255)	// Indicates an invalid allocation
+		,	WeightmapTextureChannel(255)
+	{
+	}
+	
+	FName GetLayerName() const;
+};
+
+UCLASS(hidecategories=(Display, Attachment, Physics, Debug, Collision, Movement, Rendering, PrimitiveComponent, Object, Transform), showcategories=("Rendering|Material"), MinimalAPI)
+class ULandscapeComponent : public UPrimitiveComponent
+{
+	GENERATED_UCLASS_BODY()
+
+	UPROPERTY(VisibleAnywhere, BlueprintReadOnly, Category=LandscapeComponent)
+	int32 SectionBaseX;
+
+	UPROPERTY(VisibleAnywhere, BlueprintReadOnly, Category=LandscapeComponent)
+	int32 SectionBaseY;
+
+	/** Total number of quads for this component */
+	UPROPERTY()
+	int32 ComponentSizeQuads;
+
+	/** Number of quads for a subsection of the component. SubsectionSizeQuads+1 must be a power of two. */
+	UPROPERTY()
+	int32 SubsectionSizeQuads;
+
+	/** Number of subsections in X or Y axis */
+	UPROPERTY()
+	int32 NumSubsections;
+
+	UPROPERTY(EditAnywhere, BlueprintReadWrite, Category=LandscapeComponent)
+	class UMaterialInterface* OverrideMaterial;
+
+	UPROPERTY(EditAnywhere, BlueprintReadWrite, Category=LandscapeComponent)
+	class UMaterialInterface* OverrideHoleMaterial;
+
+	UPROPERTY(TextExportTransient)
+	class UMaterialInstanceConstant* MaterialInstance;
+
+	/** List of layers, and the weightmap and channel they are stored */
+	UPROPERTY()
+	TArray<struct FWeightmapLayerAllocationInfo> WeightmapLayerAllocations;
+
+	/** Weightmap texture reference */
+	UPROPERTY(TextExportTransient)
+	TArray<class UTexture2D*> WeightmapTextures;
+
+	/** XYOffsetmap texture reference */
+	UPROPERTY(TextExportTransient)
+	class UTexture2D* XYOffsetmapTexture;
+
+	/** UV offset to component's weightmap data from component local coordinates*/
+	UPROPERTY()
+	FVector4 WeightmapScaleBias;
+
+	/** U or V offset into the weightmap for the first subsection, in texture UV space */
+	UPROPERTY()
+	float WeightmapSubsectionOffset;
+
+	/** UV offset to Heightmap data from component local coordinates */
+	UPROPERTY()
+	FVector4 HeightmapScaleBias;
+
+	/** Heightmap texture reference */
+	UPROPERTY(TextExportTransient)
+	class UTexture2D* HeightmapTexture;
+
+	/** Cached bounds, created at heightmap update time */
+	UPROPERTY()
+	FBoxSphereBounds CachedBoxSphereBounds_DEPRECATED;
+
+	/** Cached local-space bounding box, created at heightmap update time */
+	UPROPERTY()
+	FBox CachedLocalBox;
+
+	/** Reference to associated collision component */
+	UPROPERTY()
+	TLazyObjectPtr<class ULandscapeHeightfieldCollisionComponent> CollisionComponent;
+
+private:
+#if WITH_EDITORONLY_DATA
+	/** Unique ID for this component, used for caching during distributed lighting */
+	UPROPERTY()
+	FGuid LightingGuid;
+
+#endif // WITH_EDITORONLY_DATA
+public:
+	/**	INTERNAL: Array of lights that don't apply to the terrain component.		*/
+	UPROPERTY()
+	TArray<FGuid> IrrelevantLights;
+
+	/** Reference to the texture lightmap resource. */
+	FLightMapRef LightMap;
+
+	FShadowMapRef ShadowMap;
+
+	/** Heightfield mipmap used to generate collision */
+	UPROPERTY(EditAnywhere, Category=LandscapeComponent)
+	int32 CollisionMipLevel;
+
+	/** StaticLightingResolution overriding per component, default value 0 means no overriding */
+	UPROPERTY(EditAnywhere, Category=LandscapeComponent)
+	float StaticLightingResolution;
+
+	/** Forced LOD level to use when rendering */
+	UPROPERTY(EditAnywhere, BlueprintReadOnly, Category=LandscapeComponent)
+	int32 ForcedLOD;
+
+	/** Neighbor LOD data to use when rendering, 255 is unspecified */
+	UPROPERTY()
+	uint8 NeighborLOD[8];
+
+	/** LOD level Bias to use when rendering */
+	UPROPERTY(EditAnywhere, BlueprintReadOnly, Category=LandscapeComponent)
+	int32 LODBias;
+
+	/** Neighbor LOD bias to use when rendering, 128 is 0 bias, 0 is -128 bias, 255 is 127 bias */
+	UPROPERTY()
+	uint8 NeighborLODBias[8];
+
+	UPROPERTY()
+	FGuid StateId;
+
+#if WITH_EDITORONLY_DATA
+	/** LOD level Bias to use when lighting buidling via lightmass, -1 Means automatic LOD calculation based on ForcedLOD + LODBias */
+	UPROPERTY(EditAnywhere, Category=LandscapeComponent)
+	int32 LightingLODBias;
+
+	UPROPERTY(Transient, DuplicateTransient)
+	UTexture2D* SelectDataTexture; // Data texture used for selection mask
+
+	/** Pointer to data shared with the render thread, used by the editor tools */
+	struct FLandscapeEditToolRenderData* EditToolRenderData;
+
+	/** Runtime-generated editor data for ES2 emulation */
+	UPROPERTY(Transient, DuplicateTransient)
+	UTexture2D* MobileNormalmapTexture;
+
+	/** Runtime-generated editor data for ES2 emulation */
+	UPROPERTY(Transient, DuplicateTransient)
+	UMaterialInterface* MobileMaterialInterface;
+#endif
+
+	/** For ES2 */
+	UPROPERTY()
+	uint8 MobileBlendableLayerMask;
+
+public:
+	/** Platform Data where don't support texture sampling in vertex buffer */
+	FLandscapeComponentDerivedData PlatformData;
+
+	virtual ~ULandscapeComponent();
+
+	// Begin UObject interface.	
+	virtual void PostInitProperties() override;	
+	virtual void Serialize(FArchive& Ar) override;
+	virtual void BeginDestroy() override;
+	static void AddReferencedObjects(UObject* InThis, FReferenceCollector& Collector);
+	virtual void PostDuplicate(bool bDuplicateForPIE) override;
+#if WITH_EDITOR
+	virtual void PostLoad() override;
+	virtual void PostEditUndo() override;
+	virtual void PreEditChange(UProperty* PropertyThatWillChange) override;
+	virtual void PostEditChangeProperty(FPropertyChangedEvent& PropertyChangedEvent) override;
+	// End UObject interface
+
+	/** Fix up component layers, weightmaps
+	 */
+	LANDSCAPE_API void FixupWeightmaps();
+	
+	// Begin UPrimitiveComponent interface.
+	virtual bool GetLightMapResolution( int32& Width, int32& Height ) const override;
+	virtual void GetLightAndShadowMapMemoryUsage( int32& LightMapMemoryUsage, int32& ShadowMapMemoryUsage ) const override;
+	virtual void GetStaticLightingInfo(FStaticLightingPrimitiveInfo& OutPrimitiveInfo,const TArray<ULightComponent*>& InRelevantLights,const FLightingBuildOptions& Options) override;
+#endif
+	virtual void GetUsedMaterials( TArray<UMaterialInterface*>& OutMaterials ) const override;
+	virtual FPrimitiveSceneProxy* CreateSceneProxy() override;
+	virtual ELightMapInteractionType GetStaticLightingType() const override { return LMIT_Texture;	}
+	virtual void GetStreamingTextureInfo(TArray<FStreamingTexturePrimitiveInfo>& OutStreamingTextures) const override;
+
+#if WITH_EDITOR
+	virtual int32 GetNumMaterials() const override;
+	virtual class UMaterialInterface* GetMaterial(int32 ElementIndex) const override;
+	virtual void SetMaterial(int32 ElementIndex, class UMaterialInterface* Material) override;
+#endif
+	// End UPrimitiveComponent interface.
+
+	// Begin USceneComponent interface.
+	virtual void DestroyComponent() override;
+	virtual FBoxSphereBounds CalcBounds(const FTransform& LocalToWorld) const override;
+	// End USceneComponent interface.
+
+	// Begin UActorComponent interface.
+	virtual void OnRegister() override;
+	virtual void OnUnregister() override;
+	// End UActorComponent interface.
+
+
+#if WITH_EDITOR
+	/** @todo document */
+	LANDSCAPE_API ULandscapeInfo* GetLandscapeInfo(bool bSpawnNewActor = true) const;
+
+	/** @todo document */
+	void DeleteLayer(ULandscapeLayerInfoObject* LayerInfo, struct FLandscapeEditDataInterface* LandscapeEdit);
+
+	void ReplaceLayer(ULandscapeLayerInfoObject* FromLayerInfo, ULandscapeLayerInfoObject* ToLayerInfo, struct FLandscapeEditDataInterface* LandscapeEdit);
+	
+	void GeneratePlatformVertexData();
+	UMaterialInstance* GeneratePlatformPixelData(TArray<class UTexture2D*>& WeightmapTextures, bool bIsCooking);
+#endif
+
+	/** @todo document */
+	virtual void InvalidateLightingCacheDetailed(bool bInvalidateBuildEnqueuedLighting, bool bTranslationOnly);
+
+
+	/** Get the landscape actor associated with this component. */
+	class ALandscape* GetLandscapeActor() const;
+
+	/** Get the level in which the owning actor resides */
+	class ULevel* GetLevel() const;
+
+	/** Returns all generated textures and material instances used by this component. */
+	LANDSCAPE_API void GetGeneratedTexturesAndMaterialInstances(TArray<UObject*>& OutTexturesAndMaterials) const;
+
+	/** @todo document */
+	LANDSCAPE_API ALandscapeProxy* GetLandscapeProxy() const;
+
+	/** @return Component section base as FIntPoint */
+	LANDSCAPE_API FIntPoint GetSectionBase() const; 
+
+	/** @param InSectionBase new section base for a component */
+	LANDSCAPE_API void SetSectionBase(FIntPoint InSectionBase);
+
+	/** @todo document */
+	TMap< UTexture2D*,struct FLandscapeWeightmapUsage >& GetWeightmapUsageMap();
+
+	/** @todo document */
+	const FGuid& GetLightingGuid() const
+	{
+#if WITH_EDITORONLY_DATA
+		return LightingGuid;
+#else
+		static const FGuid NullGuid( 0, 0, 0, 0 );
+		return NullGuid;
+#endif // WITH_EDITORONLY_DATA
+	}
+
+	/** @todo document */
+	void SetLightingGuid()
+	{
+#if WITH_EDITORONLY_DATA
+		LightingGuid = FGuid::NewGuid();
+#endif // WITH_EDITORONLY_DATA
+	}
+
+#if WITH_EDITOR
+
+
+	/** Initialize the landscape component */
+	LANDSCAPE_API void Init(int32 InBaseX,int32 InBaseY,int32 InComponentSizeQuads, int32 InNumSubsections,int32 InSubsectionSizeQuads);
+
+	/**
+	 * Recalculate cached bounds using height values.
+	 */
+	LANDSCAPE_API void UpdateCachedBounds();
+
+	/**
+	 * Update the MaterialInstance parameters to match the layer and weightmaps for this component
+	 * Creates the MaterialInstance if it doesn't exist.
+	 */
+	LANDSCAPE_API void UpdateMaterialInstances();
+
+	/** Helper function for UpdateMaterialInstance to get Material without set parameters */
+	UMaterialInstanceConstant* GetCombinationMaterial(bool bMobile = false);
+	/**
+	 * Generate mipmaps for height and tangent data.
+	 * @param HeightmapTextureMipData - array of pointers to the locked mip data.
+	 *           This should only include the mips that are generated directly from this component's data
+	 *           ie where each subsection has at least 2 vertices.
+	* @param ComponentX1 - region of texture to update in component space, MAX_int32 meant end of X component in ALandscape::Import()
+	* @param ComponentY1 - region of texture to update in component space, MAX_int32 meant end of Y component in ALandscape::Import()
+	* @param ComponentX2 (optional) - region of texture to update in component space
+	* @param ComponentY2 (optional) - region of texture to update in component space
+	* @param TextureDataInfo - FLandscapeTextureDataInfo pointer, to notify of the mip data region updated.
+	 */
+	void GenerateHeightmapMips(TArray<FColor*>& HeightmapTextureMipData, int32 ComponentX1=0, int32 ComponentY1=0, int32 ComponentX2=MAX_int32, int32 ComponentY2=MAX_int32,struct FLandscapeTextureDataInfo* TextureDataInfo=NULL);
+
+	/**
+	 * Generate empty mipmaps for weightmap
+	 */
+	LANDSCAPE_API static void CreateEmptyTextureMips(UTexture2D* Texture, bool bClear = false);
+
+	/**
+	 * Generate mipmaps for weightmap
+	 * Assumes all weightmaps are unique to this component.
+	 * @param WeightmapTextureBaseMipData: array of pointers to each of the weightmaps' base locked mip data.
+	 */
+	template<typename DataType>
+
+	/** @todo document */
+	static void GenerateMipsTempl(int32 InNumSubsections, int32 InSubsectionSizeQuads, UTexture2D* WeightmapTexture, DataType* BaseMipData);
+
+	/** @todo document */
+	static void GenerateWeightmapMips(int32 InNumSubsections, int32 InSubsectionSizeQuads, UTexture2D* WeightmapTexture, FColor* BaseMipData);
+
+	/**
+	 * Update mipmaps for existing weightmap texture
+	 */
+	template<typename DataType>
+
+	/** @todo document */
+	static void UpdateMipsTempl(int32 InNumSubsections, int32 InSubsectionSizeQuads, UTexture2D* WeightmapTexture, TArray<DataType*>& WeightmapTextureMipData, int32 ComponentX1=0, int32 ComponentY1=0, int32 ComponentX2=MAX_int32, int32 ComponentY2=MAX_int32, struct FLandscapeTextureDataInfo* TextureDataInfo=NULL);
+
+	/** @todo document */
+	LANDSCAPE_API static void UpdateWeightmapMips(int32 InNumSubsections, int32 InSubsectionSizeQuads, UTexture2D* WeightmapTexture, TArray<FColor*>& WeightmapTextureMipData, int32 ComponentX1=0, int32 ComponentY1=0, int32 ComponentX2=MAX_int32, int32 ComponentY2=MAX_int32, struct FLandscapeTextureDataInfo* TextureDataInfo=NULL);
+
+	/** @todo document */
+	static void UpdateDataMips(int32 InNumSubsections, int32 InSubsectionSizeQuads, UTexture2D* Texture, TArray<uint8*>& TextureMipData, int32 ComponentX1=0, int32 ComponentY1=0, int32 ComponentX2=MAX_int32, int32 ComponentY2=MAX_int32, struct FLandscapeTextureDataInfo* TextureDataInfo=NULL);
+
+	/**
+	 * Create or updatescollision component height data
+	 * @param HeightmapTextureMipData: heightmap data
+	 * @param ComponentX1, ComponentY1, ComponentX2, ComponentY2: region to update
+	 * @param Whether to update bounds from render component.
+	 */
+	void UpdateCollisionHeightData(const FColor* HeightmapTextureMipData, int32 ComponentX1=0, int32 ComponentY1=0, int32 ComponentX2=MAX_int32, int32 ComponentY2=MAX_int32, bool bUpdateBounds=false, const FColor* XYOffsetTextureMipData = NULL, bool bRebuild=false);
+
+	/**
+	 * Update collision component dominant layer data
+	 * @param WeightmapTextureMipData: weightmap data
+	 * @param ComponentX1, ComponentY1, ComponentX2, ComponentY2: region to update
+	 * @param Whether to update bounds from render component.
+	 */
+	void UpdateCollisionLayerData(TArray<FColor*>& WeightmapTextureMipData, int32 ComponentX1=0, int32 ComponentY1=0, int32 ComponentX2=MAX_int32, int32 ComponentY2=MAX_int32);
+
+	/**
+	 * Update collision component dominant layer data for the whole component, locking and unlocking the weightmap textures.
+	 */
+	LANDSCAPE_API void UpdateCollisionLayerData();
+
+	/**
+	 * Create weightmaps for this component for the layers specified in the WeightmapLayerAllocations array
+	 */
+	void ReallocateWeightmaps(struct FLandscapeEditDataInterface* DataInterface=NULL);
+
+	/** Returns the actor's LandscapeMaterial, or the Component's OverrideLandscapeMaterial if set */
+	LANDSCAPE_API UMaterialInterface* GetLandscapeMaterial() const;
+
+	/** Returns the actor's LandscapeHoleMaterial, or the Component's OverrideLandscapeHoleMaterial if set */
+	LANDSCAPE_API UMaterialInterface* GetLandscapeHoleMaterial() const;
+
+	/** Returns true if this component has visibility painted */
+	LANDSCAPE_API bool ComponentHasVisibilityPainted() const;
+
+	/**
+	 * Generate a key for this component's layer allocations to use with MaterialInstanceConstantMap.
+	 */
+	FString GetLayerAllocationKey(bool bMobile = false) const;
+
+	/** @todo document */
+	void GetLayerDebugColorKey(int32& R, int32& G, int32& B) const;
+
+	/** @todo document */
+	void RemoveInvalidWeightmaps();
+
+	/** @todo document */
+	virtual void ExportCustomProperties(FOutputDevice& Out, uint32 Indent);
+
+	/** @todo document */
+	virtual void ImportCustomProperties(const TCHAR* SourceText, FFeedbackContext* Warn);
+
+	/** @todo document */
+	LANDSCAPE_API void InitHeightmapData(TArray<FColor>& Heights, bool bUpdateCollision);
+
+	/** @todo document */
+	LANDSCAPE_API void InitWeightmapData(TArray<class ULandscapeLayerInfoObject*>& LayerInfos, TArray<TArray<uint8> >& Weights);
+
+	/** @todo document */
+	LANDSCAPE_API float GetLayerWeightAtLocation( const FVector& InLocation, ULandscapeLayerInfoObject* LayerInfo, TArray<uint8>* LayerCache=NULL );
+
+	/** Extends passed region with this component section size */
+	LANDSCAPE_API void GetComponentExtent(int32& MinX, int32& MinY, int32& MaxX, int32& MaxY) const;
+
+	/** Updates navigation properties to match landscape's master switch */
+	void UpdateNavigationRelevance();
+#endif
+
+	friend class FLandscapeComponentSceneProxy;
+	friend struct FLandscapeComponentDataInterface;
+
+	void SetLOD(bool bForced, int32 InLODValue);
+
+protected:
+
+	/** Whether the component type supports static lighting. */
+	virtual bool SupportsStaticLighting() const override
+	{
+		return true;
+	}
+};
+
+
+
diff --git a/Engine/Source/Runtime/Landscape/Classes/LandscapeGizmoActor.h b/Engine/Source/Runtime/Landscape/Classes/LandscapeGizmoActor.h
new file mode 100644
--- /dev/null
+++ b/Engine/Source/Runtime/Landscape/Classes/LandscapeGizmoActor.h
@@ -0,0 +1,58 @@
+// Copyright 1998-2014 Epic Games, Inc. All Rights Reserved.
+
+#pragma once
+
+#include "Components/BillboardComponent.h"
+
+#include "LandscapeGizmoActor.generated.h"
+
+UCLASS(notplaceable, MinimalAPI, NotBlueprintable)
+class ALandscapeGizmoActor : public AActor
+{
+	GENERATED_UCLASS_BODY()
+
+#if WITH_EDITORONLY_DATA
+	UPROPERTY(EditAnywhere, Category=Gizmo)
+	float Width;
+
+	UPROPERTY(EditAnywhere, Category=Gizmo)
+	float Height;
+
+	UPROPERTY(EditAnywhere, Category=Gizmo)
+	float LengthZ;
+
+	UPROPERTY(EditAnywhere, Category=Gizmo)
+	float MarginZ;
+
+	UPROPERTY(EditAnywhere, Category=Gizmo)
+	float MinRelativeZ;
+
+	UPROPERTY(EditAnywhere, Category=Gizmo)
+	float RelativeScaleZ;
+
+	UPROPERTY(EditAnywhere, transient, Category=Gizmo)
+	class ULandscapeInfo* TargetLandscapeInfo;
+
+	UPROPERTY()
+	TSubobjectPtr<UBillboardComponent> SpriteComponent;
+#endif // WITH_EDITORONLY_DATA
+
+
+#if WITH_EDITOR
+	virtual void Duplicate(ALandscapeGizmoActor* Gizmo); 
+	//virtual void EditorApplyTranslation(const FVector& DeltaTranslation, bool bAltDown, bool bShiftDown, bool bCtrlDown);
+
+	bool EditorCanAttachTo(const AActor* InParent, FText& OutReason) const override
+	{
+		return false;
+	}
+#endif
+
+	/** 
+	 * Indicates whether this actor should participate in level bounds calculations
+	 */
+	virtual bool IsLevelBoundsRelevant() const override { return false; }
+};
+
+
+
diff --git a/Engine/Source/Runtime/Landscape/Classes/LandscapeGizmoRenderComponent.h b/Engine/Source/Runtime/Landscape/Classes/LandscapeGizmoRenderComponent.h
new file mode 100644
--- /dev/null
+++ b/Engine/Source/Runtime/Landscape/Classes/LandscapeGizmoRenderComponent.h
@@ -0,0 +1,23 @@
+// Copyright 1998-2014 Epic Games, Inc. All Rights Reserved.
+
+
+#pragma once
+#include "LandscapeGizmoRenderComponent.generated.h"
+
+UCLASS(hidecategories=Object)
+class ULandscapeGizmoRenderComponent : public UPrimitiveComponent
+{
+	GENERATED_UCLASS_BODY()
+
+
+	// Begin UPrimitiveComponent Interface
+	virtual FPrimitiveSceneProxy* CreateSceneProxy() override;
+	// End UPrimitiveComponent Interface
+
+	// Begin USceneComponent interface.
+	virtual FBoxSphereBounds CalcBounds(const FTransform& LocalToWorld) const override;
+	// End USceneComponent interface.
+};
+
+
+
diff --git a/Engine/Source/Runtime/Landscape/Classes/LandscapeHeightfieldCollisionComponent.h b/Engine/Source/Runtime/Landscape/Classes/LandscapeHeightfieldCollisionComponent.h
new file mode 100644
--- /dev/null
+++ b/Engine/Source/Runtime/Landscape/Classes/LandscapeHeightfieldCollisionComponent.h
@@ -0,0 +1,196 @@
+// Copyright 1998-2014 Epic Games, Inc. All Rights Reserved.
+
+
+#pragma once
+
+class ULandscapeLayerInfoObject;
+
+#include "Foliage/InstancedFoliageActor.h"
+
+#include "LandscapeHeightfieldCollisionComponent.generated.h"
+
+#if WITH_PHYSX
+namespace physx
+{
+	class PxMaterial;
+	class PxHeightField;
+}
+#endif // WITH_PHYSX
+
+UCLASS(MinimalAPI)
+class ULandscapeHeightfieldCollisionComponent : public UPrimitiveComponent
+{
+	GENERATED_UCLASS_BODY()
+
+	/** DEPRECATED List of layers painted on this component. Matches the WeightmapLayerAllocations array in the LandscapeComponent. */
+	UPROPERTY()
+	TArray<FName> ComponentLayers_DEPRECATED;
+
+	/** List of layers painted on this component. Matches the WeightmapLayerAllocations array in the LandscapeComponent. */
+	UPROPERTY()
+	TArray<ULandscapeLayerInfoObject*> ComponentLayerInfos;
+
+	/** Offset of component in landscape quads */
+	UPROPERTY()
+	int32 SectionBaseX;
+
+	UPROPERTY()
+	int32 SectionBaseY;
+
+	/** Size of component in collision quads */
+	UPROPERTY()
+	int32 CollisionSizeQuads;
+
+	/** Collision scale: (ComponentSizeQuads) / (CollisionSizeQuads) */
+	UPROPERTY()
+	float CollisionScale;
+
+	/** The flags for each collision quad. See ECollisionQuadFlags. */
+	UPROPERTY()
+	TArray<uint8> CollisionQuadFlags;
+
+	/** Guid used to share PhysX heightfield objects in the editor */
+	UPROPERTY()
+	FGuid HeightfieldGuid;
+
+	/** Cached bounds, created at heightmap update time */
+	UPROPERTY()
+	FBoxSphereBounds CachedBoxSphereBounds_DEPRECATED;
+
+	/** Cached local-space bounding box, created at heightmap update time */
+	UPROPERTY()
+	FBox CachedLocalBox;
+
+	/** Reference to render component */
+	UPROPERTY()
+	TLazyObjectPtr<class ULandscapeComponent> RenderComponent;
+
+	struct FPhysXHeightfieldRef : public FRefCountedObject
+	{
+		FGuid Guid;
+
+#if WITH_PHYSX
+		/** List of PxMaterials used on this landscape */
+		TArray<physx::PxMaterial*>	UsedPhysicalMaterialArray;
+		class physx::PxHeightField*			RBHeightfield;
+#if WITH_EDITOR
+		class physx::PxHeightField*			RBHeightfieldEd; // Used only by landscape editor, does not have holes in it
+#endif	//WITH_EDITOR
+#endif	//WITH_PHYSX
+
+		/** tors **/
+		FPhysXHeightfieldRef() 
+#if WITH_PHYSX
+			:	RBHeightfield(NULL)
+#if WITH_EDITOR
+			,	RBHeightfieldEd(NULL)
+#endif	//WITH_EDITOR
+#endif	//WITH_PHYSX
+		{}
+		FPhysXHeightfieldRef(FGuid& InGuid)
+			:	Guid(InGuid)
+#if WITH_PHYSX
+			,	RBHeightfield(NULL)
+#if WITH_EDITOR
+			,	RBHeightfieldEd(NULL)
+#endif	//WITH_EDITOR
+#endif	//WITH_PHYSX
+		{}
+		virtual ~FPhysXHeightfieldRef();
+	};
+	
+	/** The collision height values. Stripped from cooked content */
+	FWordBulkData								CollisionHeightData;
+
+	/** Indices into the ComponentLayers array for the per-vertex dominant layer. Stripped from cooked content */
+	FByteBulkData								DominantLayerData;
+
+	/** 
+	 *	Cooked HeightField data. Serialized only with cooked content 
+	 *	Stored as array instead of BulkData to take advantage of precaching during async loading
+	 */
+	TArray<uint8>								CookedCollisionData;
+
+	/** This is a list of physical materials that is actually used by a cooked HeightField. Serialized only with cooked content */
+	UPROPERTY(transient)
+	TArray<UPhysicalMaterial*>					CookedPhysicalMaterials;
+	
+	/** Physics engine version of heightfield data. */
+	TRefCountPtr<struct FPhysXHeightfieldRef>	HeightfieldRef;
+
+	enum ECollisionQuadFlags
+	{
+		QF_PhysicalMaterialMask = 63,	// Mask value for the physical material index, stored in the lower 6 bits.
+		QF_EdgeTurned = 64,				// This quad's diagonal has been turned.
+		QF_NoCollision = 128,			// This quad has no collision.
+	};
+
+	// Begin UActorComponent interface.
+	virtual void CreatePhysicsState() override;
+	virtual void ApplyWorldOffset(const FVector& InOffset, bool bWorldShift) override;
+	// End UActorComponent interface.
+
+	// Begin USceneComponent interface.
+	virtual void DestroyComponent() override;
+	virtual FBoxSphereBounds CalcBounds(const FTransform &BoundTransform) const override;
+
+	virtual ECollisionEnabled::Type GetCollisionEnabled() const override;
+	virtual ECollisionResponse GetCollisionResponseToChannel(ECollisionChannel Channel) const override;
+	virtual ECollisionChannel GetCollisionObjectType() const override;
+	virtual const FCollisionResponseContainer& GetCollisionResponseToChannels() const override;
+	// End USceneComponent interface.
+
+	// Begin UPrimitiveComponent interface
+	virtual bool DoCustomNavigableGeometryExport(struct FNavigableGeometryExport* GeomExport) const override;
+	//End UPrimitiveComponent interface
+
+	// Begin UObject Interface.
+	virtual void Serialize(FArchive& Ar) override;
+	virtual void BeginDestroy() override;
+	virtual void PostLoad() override;
+#if WITH_EDITOR
+	virtual void ExportCustomProperties(FOutputDevice& Out, uint32 Indent) override;
+	virtual void ImportCustomProperties(const TCHAR* SourceText, FFeedbackContext* Warn) override;
+	virtual void PostEditImport() override;
+	virtual void PostEditUndo() override;
+	// End UObject Interface.
+
+	// Update Collision object for add LandscapeComponent tool
+	LANDSCAPE_API void UpdateAddCollisions();
+
+	// @todo document
+	class ULandscapeInfo* GetLandscapeInfo(bool bSpawnNewActor = true) const;
+
+	/** 
+	 * Cooks raw height data into collision object binary stream
+	 */
+	virtual bool CookCollsionData(const FName& Format, bool bUseOnlyDefMaterial, TArray<uint8>& OutCookedData, TArray<UPhysicalMaterial*>& OutMaterails) const;
+
+	/** Modify a sub-region of the PhysX heightfield. Note that this does not update the physical material */
+	void UpdateHeightfieldRegion(int32 ComponentX1, int32 ComponentY1, int32 ComponentX2, int32 ComponentY2);
+
+#endif
+	/** Creates collision object from a cooked collision data */
+	virtual void CreateCollisionObject();
+
+	/** Return the landscape actor associated with this component. */
+	class ALandscape* GetLandscapeActor() const;
+	LANDSCAPE_API class ALandscapeProxy* GetLandscapeProxy() const;
+
+	/** @return Component section base as FIntPoint */
+	LANDSCAPE_API FIntPoint GetSectionBase() const; 
+
+	/** @param InSectionBase new section base for a component */
+	LANDSCAPE_API void SetSectionBase(FIntPoint InSectionBase);
+
+	/** Recreate heightfield and restart physics */
+	LANDSCAPE_API virtual void RecreateCollision(bool bUpdateAddCollision = true);
+
+#if WITH_EDITORONLY_DATA
+	// Called from editor code to manage foliage instances on landscape.
+	LANDSCAPE_API void SnapFoliageInstances(AInstancedFoliageActor &IFA, const FBox& InInstanceBox);
+#endif
+};
+
+
+
diff --git a/Engine/Source/Runtime/Landscape/Classes/LandscapeInfo.h b/Engine/Source/Runtime/Landscape/Classes/LandscapeInfo.h
new file mode 100644
--- /dev/null
+++ b/Engine/Source/Runtime/Landscape/Classes/LandscapeInfo.h
@@ -0,0 +1,236 @@
+// Copyright 1998-2014 Epic Games, Inc. All Rights Reserved.
+
+#pragma once
+#include "LandscapeInfo.generated.h"
+
+/** Structure storing Collision for LandscapeComponent Add */
+USTRUCT()
+struct FLandscapeAddCollision
+{
+	GENERATED_USTRUCT_BODY()
+
+#if WITH_EDITORONLY_DATA
+	UPROPERTY()
+	FVector Corners[4];
+
+#endif // WITH_EDITORONLY_DATA
+
+
+		FLandscapeAddCollision()
+		{
+			#if WITH_EDITORONLY_DATA
+			Corners[0] = Corners[1] = Corners[2] = Corners[3] = FVector::ZeroVector;
+			#endif // WITH_EDITORONLY_DATA
+		}
+	
+};
+
+USTRUCT()
+struct FLandscapeInfoLayerSettings
+{
+	GENERATED_USTRUCT_BODY()
+
+	UPROPERTY()
+	class ULandscapeLayerInfoObject* LayerInfoObj;
+
+	UPROPERTY()
+	FName LayerName;
+
+#if WITH_EDITORONLY_DATA
+	UPROPERTY(transient)
+	class UMaterialInstanceConstant* ThumbnailMIC;
+
+	UPROPERTY()
+	class ALandscapeProxy* Owner;
+
+	UPROPERTY(transient)
+	int32 DebugColorChannel;
+
+	UPROPERTY(transient)
+	uint32 bValid:1;
+#endif // WITH_EDITORONLY_DATA
+
+	FLandscapeInfoLayerSettings()
+		: LayerInfoObj(NULL)
+		, LayerName(NAME_None)
+#if WITH_EDITORONLY_DATA
+		, ThumbnailMIC(NULL)
+		, Owner(NULL)
+		, DebugColorChannel(0)
+		, bValid(false)
+#endif // WITH_EDITORONLY_DATA
+	{
+	}
+
+	LANDSCAPE_API FLandscapeInfoLayerSettings(ULandscapeLayerInfoObject* InLayerInfo, class ALandscapeProxy* InProxy);
+
+	FLandscapeInfoLayerSettings(FName InPlaceholderLayerName, class ALandscapeProxy* InProxy)
+		: LayerInfoObj(NULL)
+		, LayerName(InPlaceholderLayerName)
+#if WITH_EDITORONLY_DATA
+		, ThumbnailMIC(NULL)
+		, Owner(InProxy)
+		, DebugColorChannel(0)
+		, bValid(false)
+#endif
+	{
+	}
+
+	LANDSCAPE_API FName GetLayerName() const;
+
+#if WITH_EDITORONLY_DATA
+	LANDSCAPE_API struct FLandscapeEditorLayerSettings& GetEditorSettings() const;
+#endif
+};
+
+UCLASS()
+class ULandscapeInfo : public UObject
+{
+	GENERATED_UCLASS_BODY()
+
+	UPROPERTY()
+	TLazyObjectPtr<class ALandscape> LandscapeActor;
+
+	UPROPERTY()
+	FGuid LandscapeGuid;
+
+	UPROPERTY()
+	int32 ComponentSizeQuads;
+	
+	UPROPERTY()
+	int32 SubsectionSizeQuads;
+
+	UPROPERTY()
+	int32 ComponentNumSubsections;
+	
+	UPROPERTY()
+	FVector DrawScale;
+	
+#if WITH_EDITORONLY_DATA
+	UPROPERTY(transient)
+	uint32 bIsValid:1;
+
+	UPROPERTY()
+	TArray<struct FLandscapeInfoLayerSettings> Layers;
+#endif // WITH_EDITORONLY_DATA
+
+public:
+	struct FLandscapeDataInterface* DataInterface;
+
+	/** Map of the offsets (in component space) to the component. Valid in editor only. */
+	TMap<FIntPoint,class ULandscapeComponent*> XYtoComponentMap;
+
+	/** Map of the offsets to the newly added collision components. Only available near valid LandscapeComponents. Valid in editor only. */
+	TMap<FIntPoint, FLandscapeAddCollision> XYtoAddCollisionMap;
+
+	TSet<class ALandscapeProxy*> Proxies;
+
+private:
+	TSet<class ULandscapeComponent*> SelectedComponents;
+	TSet<class ULandscapeComponent*> SelectedRegionComponents;
+
+public:
+
+	/** True if we're currently editing this landscape */
+	bool bCurrentlyEditing;
+
+	TMap<FIntPoint,float> SelectedRegion;
+
+	// Begin UObject Interface.
+	virtual void Serialize(FArchive& Ar) override;
+	virtual void BeginDestroy() override;
+#if WITH_EDITOR
+	virtual void PostEditUndo() override;
+	// End UObject Interface
+
+	// @todo document 
+	// all below.
+	LANDSCAPE_API struct FLandscapeDataInterface* GetDataInterface();
+
+	LANDSCAPE_API void GetComponentsInRegion(int32 X1, int32 Y1, int32 X2, int32 Y2, TSet<ULandscapeComponent*>& OutComponents);
+	LANDSCAPE_API bool GetLandscapeExtent(int32& MinX, int32& MinY, int32& MaxX, int32& MaxY);
+	LANDSCAPE_API void Export(const TArray<ULandscapeLayerInfoObject*>& LayerInfos, const TArray<FString>& Filenames);
+	LANDSCAPE_API void ExportHeightmap(const FString& Filename);
+	LANDSCAPE_API void ExportLayer(ULandscapeLayerInfoObject* LayerInfo, const FString& Filename);
+	LANDSCAPE_API bool ApplySplines(bool bOnlySelected);
+	bool ApplySplinesInternal(bool bOnlySelected, ALandscapeProxy* Landscape);
+
+	LANDSCAPE_API bool GetSelectedExtent(int32& MinX, int32& MinY, int32& MaxX, int32& MaxY);
+	FVector GetLandscapeCenterPos(float& LengthZ, int32 MinX = MAX_int32, int32 MinY = MAX_int32, int32 MaxX = MIN_int32, int32 MaxY = MIN_int32);
+	LANDSCAPE_API bool IsValidPosition(int32 X, int32 Y);
+	LANDSCAPE_API void DeleteLayer(ULandscapeLayerInfoObject* LayerInfo);
+	LANDSCAPE_API void ReplaceLayer(ULandscapeLayerInfoObject* FromLayerInfo, ULandscapeLayerInfoObject* ToLayerInfo);
+
+	LANDSCAPE_API void UpdateDebugColorMaterial();
+
+	LANDSCAPE_API TSet<class ULandscapeComponent*> GetSelectedComponents() const;
+	LANDSCAPE_API TSet<class ULandscapeComponent*> GetSelectedRegionComponents() const;
+	LANDSCAPE_API void UpdateSelectedComponents(TSet<class ULandscapeComponent*>& NewComponents, bool bIsComponentwise = true);
+	LANDSCAPE_API void SortSelectedComponents();
+	LANDSCAPE_API void ClearSelectedRegion(bool bIsComponentwise = true);
+
+	LANDSCAPE_API void UpdateAllAddCollisions();
+	void UpdateAddCollision(FIntPoint LandscapeKey);
+
+	LANDSCAPE_API FLandscapeEditorLayerSettings& GetLayerEditorSettings(ULandscapeLayerInfoObject* LayerInfo) const;
+	LANDSCAPE_API void CreateLayerEditorSettingsFor(ULandscapeLayerInfoObject* LayerInfo);
+
+	LANDSCAPE_API ULandscapeLayerInfoObject* GetLayerInfoByName(FName LayerName, class ALandscapeProxy* Owner = NULL) const;
+	LANDSCAPE_API int32 GetLayerInfoIndex(FName LayerName, class ALandscapeProxy* Owner = NULL) const;
+	LANDSCAPE_API int32 GetLayerInfoIndex(ULandscapeLayerInfoObject* LayerInfo, class ALandscapeProxy* Owner = NULL) const;
+	LANDSCAPE_API bool UpdateLayerInfoMap(class ALandscapeProxy* Proxy = NULL, bool bInvalidate = false);
+
+	/** 
+	 *  Returns landscape which is spawned in the current level that was previously added to this landscape info object
+	 *  @param	bRegistered		Whether to consider only registered(visible) landscapes
+	 *	@return					Landscape or landscape proxy found in the current level 
+	 */
+	LANDSCAPE_API ALandscapeProxy* GetCurrentLevelLandscapeProxy(bool bRegistered) const;
+	
+	/** 
+	 *	returns shared landscape or landscape proxy, mostly for transformations
+	 *	@todo: should be removed
+	 */
+	LANDSCAPE_API ALandscapeProxy* GetLandscapeProxy() const;
+
+	/** Associates passed actor with this info object
+ 	 *  @param	Proxy		Landscape actor to register
+	 *  @param  bMapCheck	Whether to warn about landscape errors
+	 */
+	LANDSCAPE_API void RegisterActor(ALandscapeProxy* Proxy, bool bMapCheck = false);
+	
+	/** Deassociates passed actor with this info object*/
+	LANDSCAPE_API void UnregisterActor(ALandscapeProxy* Proxy);
+
+	/** Associates passed landscape component with this info object
+	 *  @param	Component	Landscape component to register
+	 *  @param  bMapCheck	Whether to warn about landscape errors
+	 */
+	LANDSCAPE_API void RegisterActorComponent(ULandscapeComponent* Component, bool bMapCheck = false);
+	
+	/** Deassociates passed landscape component with this info object*/
+	LANDSCAPE_API void UnregisterActorComponent(ULandscapeComponent* Component);
+		
+	/** Resets all actors, proxies, components registrations */
+	LANDSCAPE_API void Reset();
+
+	/** Recreate all LandscapeInfo objects in given world
+	 *  @param  bMapCheck	Whether to warn about landscape errors
+	 */
+	LANDSCAPE_API static void RecreateLandscapeInfo(UWorld* InWorld, bool bMapCheck);
+
+	/** 
+	 *  Fixes up proxies relative position to landscape actor
+	 *  basically makes sure that each LandscapeProxy RootComponent transform reflects LandscapeSectionOffset value
+	 *  requires LandscapeActor to be loaded
+	 *  Does not work in World composition mode!
+	 */
+	LANDSCAPE_API void FixupProxiesTransform();
+	
+	/** Fix up component layers, weightmaps for all registered proxies
+	 */
+	LANDSCAPE_API void FixupProxiesWeightmaps();
+
+	LANDSCAPE_API void RecreateCollisionComponents();
+#endif
+};
diff --git a/Engine/Source/Runtime/Landscape/Classes/LandscapeLayerInfoObject.h b/Engine/Source/Runtime/Landscape/Classes/LandscapeLayerInfoObject.h
new file mode 100644
--- /dev/null
+++ b/Engine/Source/Runtime/Landscape/Classes/LandscapeLayerInfoObject.h
@@ -0,0 +1,35 @@
+// Copyright 1998-2014 Epic Games, Inc. All Rights Reserved.
+
+#pragma once
+
+#include "LandscapeLayerInfoObject.generated.h"
+
+UCLASS(MinimalAPI)
+class ULandscapeLayerInfoObject : public UObject
+{
+	GENERATED_UCLASS_BODY()
+
+	UPROPERTY(VisibleAnywhere, Category=LandscapeLayerInfoObject, AssetRegistrySearchable)
+	FName LayerName;
+
+	UPROPERTY(EditAnywhere, Category=LandscapeLayerInfoObject)
+	class UPhysicalMaterial* PhysMaterial;
+
+	UPROPERTY(EditAnywhere, Category=LandscapeLayerInfoObject)
+	float Hardness;
+
+#if WITH_EDITORONLY_DATA
+	UPROPERTY(VisibleAnywhere, Category=LandscapeLayerInfoObject)
+	uint32 bNoWeightBlend:1;
+#endif // WITH_EDITORONLY_DATA
+
+#if WITH_EDITOR
+	// Begin UObject Interface
+	virtual void PostEditChangeProperty(FPropertyChangedEvent& PropertyChangedEvent) override;
+	virtual void PostLoad() override;
+	// End UObject Interface
+#endif
+};
+
+
+
diff --git a/Engine/Source/Runtime/Landscape/Classes/LandscapeMaterialInstanceConstant.h b/Engine/Source/Runtime/Landscape/Classes/LandscapeMaterialInstanceConstant.h
new file mode 100644
--- /dev/null
+++ b/Engine/Source/Runtime/Landscape/Classes/LandscapeMaterialInstanceConstant.h
@@ -0,0 +1,22 @@
+// Copyright 1998-2014 Epic Games, Inc. All Rights Reserved.
+
+#pragma once
+#include "LandscapeMaterialInstanceConstant.generated.h"
+
+UCLASS(MinimalAPI)
+class ULandscapeMaterialInstanceConstant : public UMaterialInstanceConstant
+{
+	GENERATED_UCLASS_BODY()
+
+	UPROPERTY()
+	uint32 bIsLayerThumbnail:1;
+
+	UPROPERTY()
+	int32 DataWeightmapIndex;
+
+	UPROPERTY()
+	int32 DataWeightmapSize;
+};
+
+
+
