First phase of deprecating TSubobjectPtr<> replacing TSubobjectPtr properties in the Engine with normal pointers and adding subobject accessors for derived classes/client code.

diff --git a/Engine/Source/Runtime/Engine/Private/CharacterMovementComponent.cpp b/Engine/Source/Runtime/Engine/Private/CharacterMovementComponent.cpp
--- a/Engine/Source/Runtime/Engine/Private/CharacterMovementComponent.cpp
+++ b/Engine/Source/Runtime/Engine/Private/CharacterMovementComponent.cpp
@@ -415,7 +415,7 @@ FVector UCharacterMovementComponent::GetPawnCapsuleExtent(const EShrinkCapsuleEx
 	check(CharacterOwner);
 
 	float Radius, HalfHeight;
-	CharacterOwner->CapsuleComponent->GetScaledCapsuleSize(Radius, HalfHeight);
+	CharacterOwner->GetCapsuleComponent()->GetScaledCapsuleSize(Radius, HalfHeight);
 	FVector CapsuleExtent(Radius, Radius, HalfHeight);
 
 	float RadiusEpsilon = 0.f;
@@ -482,7 +482,7 @@ FVector UCharacterMovementComponent::GetImpartedMovementBaseVelocity() const
 			
 			if (bImpartBaseAngularVelocity)
 			{
-				const FVector CharacterBasePosition = (UpdatedComponent->GetComponentLocation() - FVector(0.f, 0.f, CharacterOwner->CapsuleComponent->GetScaledCapsuleHalfHeight()));
+				const FVector CharacterBasePosition = (UpdatedComponent->GetComponentLocation() - FVector(0.f, 0.f, CharacterOwner->GetCapsuleComponent()->GetScaledCapsuleHalfHeight()));
 				const FVector BaseTangentialVel = MovementBaseUtility::GetMovementBaseTangentialVelocity(MovementBase, CharacterOwner->GetBasedMovement().BoneName, CharacterBasePosition);
 				BaseVelocity += BaseTangentialVel;
 			}
@@ -757,7 +757,7 @@ void UCharacterMovementComponent::PerformAirControlForPathFollowing(FVector Dire
 					Acceleration *= -1.f;
 				}
 
-				if ( Dist2D < 1.5f*CharacterOwner->CapsuleComponent->GetScaledCapsuleRadius() )
+				if ( Dist2D < 1.5f*CharacterOwner->GetCapsuleComponent()->GetScaledCapsuleRadius() )
 				{
 					Velocity.X = 0.f;
 					Velocity.Y = 0.f;
@@ -891,8 +891,8 @@ void UCharacterMovementComponent::AdjustProxyCapsuleSize()
 		}
 
 		float Radius, HalfHeight;
-		CharacterOwner->CapsuleComponent->GetUnscaledCapsuleSize(Radius, HalfHeight);
-		const float ComponentScale = CharacterOwner->CapsuleComponent->GetShapeScale();
+		CharacterOwner->GetCapsuleComponent()->GetUnscaledCapsuleSize(Radius, HalfHeight);
+		const float ComponentScale = CharacterOwner->GetCapsuleComponent()->GetShapeScale();
 
 		if (ComponentScale <= KINDA_SMALL_NUMBER)
 		{
@@ -911,8 +911,8 @@ void UCharacterMovementComponent::AdjustProxyCapsuleSize()
 		UE_LOG(LogCharacterMovement, Verbose, TEXT("Shrinking capsule for %s from (r=%.3f, h=%.3f) to (r=%.3f, h=%.3f)"), *CharacterOwner->GetName(),
 			Radius * ComponentScale, HalfHeight * ComponentScale, NewRadius * ComponentScale, NewHalfHeight * ComponentScale);
 
-		CharacterOwner->CapsuleComponent->SetCapsuleSize(NewRadius, NewHalfHeight, true);
-		CharacterOwner->CapsuleComponent->UpdateBounds();
+		CharacterOwner->GetCapsuleComponent()->SetCapsuleSize(NewRadius, NewHalfHeight, true);
+		CharacterOwner->GetCapsuleComponent()->UpdateBounds();
 	}	
 }
 
@@ -928,7 +928,7 @@ void UCharacterMovementComponent::SimulatedTick(float DeltaSeconds)
 		UE_LOG(LogRootMotion, Verbose, TEXT("UCharacterMovementComponent::SimulatedTick"));
 
 		// Tick animations before physics.
-		if( CharacterOwner->Mesh )
+		if( CharacterOwner->GetMesh() )
 		{
 			TickCharacterPose(DeltaSeconds);
 
@@ -952,7 +952,7 @@ void UCharacterMovementComponent::SimulatedTick(float DeltaSeconds)
 			{
 				const FRotator NewRotation = CharacterOwner->GetActorRotation();
 				const FVector NewLocation = CharacterOwner->GetActorLocation();
-				DrawDebugCoordinateSystem(GetWorld(), CharacterOwner->Mesh->GetComponentLocation() + FVector(0,0,1), NewRotation, 50.f, false);
+				DrawDebugCoordinateSystem(GetWorld(), CharacterOwner->GetMesh()->GetComponentLocation() + FVector(0,0,1), NewRotation, 50.f, false);
 				DrawDebugLine(GetWorld(), OldLocation, NewLocation, FColor::Red, true, 10.f);
 
 				UE_LOG(LogRootMotion, Log,  TEXT("UCharacterMovementComponent::SimulatedTick DeltaMovement Translation: %s, Rotation: %s, MovementBase: %s"),
@@ -1002,10 +1002,10 @@ void UCharacterMovementComponent::SimulatedTick(float DeltaSeconds)
 
 void UCharacterMovementComponent::SimulateRootMotion(float DeltaSeconds, const FTransform& LocalRootMotionTransform)
 {
-	if( CharacterOwner && CharacterOwner->Mesh && (DeltaSeconds > 0.f) )
+	if( CharacterOwner && CharacterOwner->GetMesh() && (DeltaSeconds > 0.f) )
 	{
 		// Convert Local Space Root Motion to world space. Do it right before used by physics to make sure we use up to date transforms, as translation is relative to rotation.
-		const FTransform WorldSpaceRootMotionTransform = CharacterOwner->Mesh->ConvertLocalRootMotionToWorld(LocalRootMotionTransform);
+		const FTransform WorldSpaceRootMotionTransform = CharacterOwner->GetMesh()->ConvertLocalRootMotionToWorld(LocalRootMotionTransform);
 
 		// Compute root motion velocity to be used by physics
 		const FVector RootMotionVelocity = WorldSpaceRootMotionTransform.GetTranslation() / DeltaSeconds;
@@ -1294,7 +1294,7 @@ void UCharacterMovementComponent::UpdateBasedMovement(float DeltaSeconds)
 
 			// We need to offset the base of the character here, not its origin, so offset by half height
 			float HalfHeight, Radius;
-			CharacterOwner->CapsuleComponent->GetScaledCapsuleSize(Radius, HalfHeight);
+			CharacterOwner->GetCapsuleComponent()->GetScaledCapsuleSize(Radius, HalfHeight);
 
 			FVector const BaseOffset(0.0f, 0.0f, HalfHeight);
 			FVector const LocalBasePos = OldLocalToWorld.InverseTransformPosition(CharacterOwner->GetActorLocation() - BaseOffset);
@@ -1319,9 +1319,9 @@ void UCharacterMovementComponent::UpdateBasedMovement(float DeltaSeconds)
 			}
 		}
 
-		if (MovementBase->IsSimulatingPhysics() && CharacterOwner->Mesh)
+		if (MovementBase->IsSimulatingPhysics() && CharacterOwner->GetMesh())
 		{
-			CharacterOwner->Mesh->ApplyDeltaToAllPhysicsTransforms(DeltaPosition, DeltaQuat);
+			CharacterOwner->GetMesh()->ApplyDeltaToAllPhysicsTransforms(DeltaPosition, DeltaQuat);
 		}
 	}
 }
@@ -1415,7 +1415,7 @@ void UCharacterMovementComponent::PerformMovement(float DeltaSeconds)
 		HandlePendingLaunch();
 
 		// If using RootMotion, tick animations before running physics.
-		if( !CharacterOwner->bClientUpdating && CharacterOwner->IsPlayingRootMotion() && CharacterOwner->Mesh )
+		if( !CharacterOwner->bClientUpdating && CharacterOwner->IsPlayingRootMotion() && CharacterOwner->GetMesh() )
 		{
 			TickCharacterPose(DeltaSeconds);
 
@@ -1435,7 +1435,7 @@ void UCharacterMovementComponent::PerformMovement(float DeltaSeconds)
 		// if we're about to use root motion, convert it to world space first.
 		if( HasRootMotion() )
 		{
-			USkeletalMeshComponent * SkelMeshComp = CharacterOwner->Mesh;
+			USkeletalMeshComponent * SkelMeshComp = CharacterOwner->GetMesh();
 			if( SkelMeshComp )
 			{
 				// Convert Local Space Root Motion to world space. Do it right before used by physics to make sure we use up to date transforms, as translation is relative to rotation.
@@ -1496,7 +1496,7 @@ void UCharacterMovementComponent::PerformMovement(float DeltaSeconds)
 				const FRotator ResultingRotation = CharacterOwner->GetActorRotation();
 
 				// Show current position
-				DrawDebugCoordinateSystem(GetWorld(), CharacterOwner->Mesh->GetComponentLocation() + FVector(0,0,1), ResultingRotation, 50.f, false);
+				DrawDebugCoordinateSystem(GetWorld(), CharacterOwner->GetMesh()->GetComponentLocation() + FVector(0,0,1), ResultingRotation, 50.f, false);
 
 				// Show resulting delta move.
 				DrawDebugLine(GetWorld(), OldLocation, ResultingLocation, FColor::Red, true, 10.f);
@@ -1601,7 +1601,7 @@ void UCharacterMovementComponent::Crouch(bool bClientSimulation)
 	}
 
 	// Do not perform if collision is already at desired size.
-	if( CharacterOwner->CapsuleComponent->GetUnscaledCapsuleHalfHeight() == CrouchedHalfHeight )
+	if( CharacterOwner->GetCapsuleComponent()->GetUnscaledCapsuleHalfHeight() == CrouchedHalfHeight )
 	{
 		return;
 	}
@@ -1615,14 +1615,14 @@ void UCharacterMovementComponent::Crouch(bool bClientSimulation)
 	{
 		// restore collision size before crouching
 		ACharacter* DefaultCharacter = CharacterOwner->GetClass()->GetDefaultObject<ACharacter>();
-		CharacterOwner->CapsuleComponent->SetCapsuleSize(DefaultCharacter->CapsuleComponent->GetUnscaledCapsuleRadius(), DefaultCharacter->CapsuleComponent->GetUnscaledCapsuleHalfHeight());
+		CharacterOwner->GetCapsuleComponent()->SetCapsuleSize(DefaultCharacter->GetCapsuleComponent()->GetUnscaledCapsuleRadius(), DefaultCharacter->GetCapsuleComponent()->GetUnscaledCapsuleHalfHeight());
 		bShrinkProxyCapsule = true;
 	}
 
 	// Change collision size to crouching dimensions
-	const float ComponentScale = CharacterOwner->CapsuleComponent->GetShapeScale();
-	const float OldUnscaledHalfHeight = CharacterOwner->CapsuleComponent->GetUnscaledCapsuleHalfHeight();
-	CharacterOwner->CapsuleComponent->SetCapsuleSize(CharacterOwner->CapsuleComponent->GetUnscaledCapsuleRadius(), CrouchedHalfHeight);
+	const float ComponentScale = CharacterOwner->GetCapsuleComponent()->GetShapeScale();
+	const float OldUnscaledHalfHeight = CharacterOwner->GetCapsuleComponent()->GetUnscaledCapsuleHalfHeight();
+	CharacterOwner->GetCapsuleComponent()->SetCapsuleSize(CharacterOwner->GetCapsuleComponent()->GetUnscaledCapsuleRadius(), CrouchedHalfHeight);
 	float HalfHeightAdjust = (OldUnscaledHalfHeight - CrouchedHalfHeight);
 	float ScaledHalfHeightAdjust = HalfHeightAdjust * ComponentScale;
 
@@ -1641,7 +1641,7 @@ void UCharacterMovementComponent::Crouch(bool bClientSimulation)
 			// If encroached, cancel
 			if( bEncroached )
 			{
-				CharacterOwner->CapsuleComponent->SetCapsuleSize(CharacterOwner->CapsuleComponent->GetUnscaledCapsuleRadius(), OldUnscaledHalfHeight);
+				CharacterOwner->GetCapsuleComponent()->SetCapsuleSize(CharacterOwner->GetCapsuleComponent()->GetUnscaledCapsuleRadius(), OldUnscaledHalfHeight);
 				return;
 			}
 		}
@@ -1659,7 +1659,7 @@ void UCharacterMovementComponent::Crouch(bool bClientSimulation)
 
 	// OnStartCrouch takes the change from the Default size, not the current one (though they are usually the same).
 	ACharacter* DefaultCharacter = CharacterOwner->GetClass()->GetDefaultObject<ACharacter>();
-	HalfHeightAdjust = (DefaultCharacter->CapsuleComponent->GetUnscaledCapsuleHalfHeight() - CrouchedHalfHeight);
+	HalfHeightAdjust = (DefaultCharacter->GetCapsuleComponent()->GetUnscaledCapsuleHalfHeight() - CrouchedHalfHeight);
 	ScaledHalfHeightAdjust = HalfHeightAdjust * ComponentScale;
 
 	AdjustProxyCapsuleSize();
@@ -1677,24 +1677,24 @@ void UCharacterMovementComponent::UnCrouch(bool bClientSimulation)
 	ACharacter* DefaultCharacter = CharacterOwner->GetClass()->GetDefaultObject<ACharacter>();
 
 	// Do not perform if collision is already at desired size.
-	if( CharacterOwner->CapsuleComponent->GetUnscaledCapsuleHalfHeight() == DefaultCharacter->CapsuleComponent->GetUnscaledCapsuleHalfHeight() )
+	if( CharacterOwner->GetCapsuleComponent()->GetUnscaledCapsuleHalfHeight() == DefaultCharacter->GetCapsuleComponent()->GetUnscaledCapsuleHalfHeight() )
 	{
 		return;
 	}
 
-	const float CurrentCrouchedHalfHeight = CharacterOwner->CapsuleComponent->GetScaledCapsuleHalfHeight();
+	const float CurrentCrouchedHalfHeight = CharacterOwner->GetCapsuleComponent()->GetScaledCapsuleHalfHeight();
 
-	const float ComponentScale = CharacterOwner->CapsuleComponent->GetShapeScale();
-	const float OldUnscaledHalfHeight = CharacterOwner->CapsuleComponent->GetUnscaledCapsuleHalfHeight();
-	const float HalfHeightAdjust = DefaultCharacter->CapsuleComponent->GetUnscaledCapsuleHalfHeight() - OldUnscaledHalfHeight;
+	const float ComponentScale = CharacterOwner->GetCapsuleComponent()->GetShapeScale();
+	const float OldUnscaledHalfHeight = CharacterOwner->GetCapsuleComponent()->GetUnscaledCapsuleHalfHeight();
+	const float HalfHeightAdjust = DefaultCharacter->GetCapsuleComponent()->GetUnscaledCapsuleHalfHeight() - OldUnscaledHalfHeight;
 	const float ScaledHalfHeightAdjust = HalfHeightAdjust * ComponentScale;
 	const FVector PawnLocation = CharacterOwner->GetActorLocation();
 
 	// Grow to uncrouched size.
-	check(CharacterOwner->CapsuleComponent);
+	check(CharacterOwner->GetCapsuleComponent());
 	bool bUpdateOverlaps = false;
-	CharacterOwner->CapsuleComponent->SetCapsuleSize(DefaultCharacter->CapsuleComponent->GetUnscaledCapsuleRadius(), DefaultCharacter->CapsuleComponent->GetUnscaledCapsuleHalfHeight(), bUpdateOverlaps);
-	CharacterOwner->CapsuleComponent->UpdateBounds(); // Force an update of the bounds with the new dimensions
+	CharacterOwner->GetCapsuleComponent()->SetCapsuleSize(DefaultCharacter->GetCapsuleComponent()->GetUnscaledCapsuleRadius(), DefaultCharacter->GetCapsuleComponent()->GetUnscaledCapsuleHalfHeight(), bUpdateOverlaps);
+	CharacterOwner->GetCapsuleComponent()->UpdateBounds(); // Force an update of the bounds with the new dimensions
 
 	if( !bClientSimulation )
 	{
@@ -1720,7 +1720,7 @@ void UCharacterMovementComponent::UnCrouch(bool bClientSimulation)
 				{
 					// Shrink to a short capsule, sweep down to base to find where that would hit something, and then try to stand up from there.
 					float PawnRadius, PawnHalfHeight;
-					CharacterOwner->CapsuleComponent->GetScaledCapsuleSize(PawnRadius, PawnHalfHeight);
+					CharacterOwner->GetCapsuleComponent()->GetScaledCapsuleSize(PawnRadius, PawnHalfHeight);
 					const float ShrinkHalfHeight = PawnHalfHeight - PawnRadius;
 					const float TraceDist = PawnHalfHeight - ShrinkHalfHeight;
 					const FVector Down = FVector(0.f, 0.f, -TraceDist);
@@ -1778,8 +1778,8 @@ void UCharacterMovementComponent::UnCrouch(bool bClientSimulation)
 		// If still encroached then abort.
 		if (bEncroached)
 		{
-			CharacterOwner->CapsuleComponent->SetCapsuleSize(CharacterOwner->CapsuleComponent->GetUnscaledCapsuleRadius(), OldUnscaledHalfHeight, false);
-			CharacterOwner->CapsuleComponent->UpdateBounds(); // Update bounds again back to old value
+			CharacterOwner->GetCapsuleComponent()->SetCapsuleSize(CharacterOwner->GetCapsuleComponent()->GetUnscaledCapsuleRadius(), OldUnscaledHalfHeight, false);
+			CharacterOwner->GetCapsuleComponent()->UpdateBounds(); // Update bounds again back to old value
 			return;
 		}
 
@@ -1792,7 +1792,7 @@ void UCharacterMovementComponent::UnCrouch(bool bClientSimulation)
 
 	// now call SetCapsuleSize() to cause touch/untouch events
 	bUpdateOverlaps = true;
-	CharacterOwner->CapsuleComponent->SetCapsuleSize(DefaultCharacter->CapsuleComponent->GetUnscaledCapsuleRadius(), DefaultCharacter->CapsuleComponent->GetUnscaledCapsuleHalfHeight(), bUpdateOverlaps);
+	CharacterOwner->GetCapsuleComponent()->SetCapsuleSize(DefaultCharacter->GetCapsuleComponent()->GetUnscaledCapsuleRadius(), DefaultCharacter->GetCapsuleComponent()->GetUnscaledCapsuleHalfHeight(), bUpdateOverlaps);
 
 	AdjustProxyCapsuleSize();
 	CharacterOwner->OnEndCrouch( HalfHeightAdjust, ScaledHalfHeightAdjust );
@@ -2044,7 +2044,7 @@ float UCharacterMovementComponent::ImmersionDepth()
 		}
 		else
 		{
-			UBrushComponent* VolumeBrushComp = GetPhysicsVolume()->BrushComponent;
+			UBrushComponent* VolumeBrushComp = GetPhysicsVolume()->GetBrushComponent();
 			FHitResult Hit(1.f);
 			if ( VolumeBrushComp )
 			{
@@ -2298,7 +2298,7 @@ void UCharacterMovementComponent::CalcAvoidanceVelocity(float DeltaTime)
 #endif
 
 	//Adjust velocity only if we're in "Walking" mode. We should also check if we're dazed, being knocked around, maybe off-navmesh, etc.
-	UCapsuleComponent *OurCapsule = GetCharacterOwner()->CapsuleComponent.Get();
+	UCapsuleComponent *OurCapsule = GetCharacterOwner()->GetCapsuleComponent();
 	if (!Velocity.IsZero() && MovementMode == MOVE_Walking && OurCapsule)
 	{
 		//See if we're doing a locked avoidance move already, and if so, skip the testing and just do the move.
@@ -2387,7 +2387,7 @@ void UCharacterMovementComponent::UpdateDefaultAvoidance()
 	UAvoidanceManager* AvoidanceManager = GetWorld()->GetAvoidanceManager();
 	if (AvoidanceManager && !bWasAvoidanceUpdated)
 	{
-		if (UCapsuleComponent *OurCapsule = GetCharacterOwner()->CapsuleComponent.Get())
+		if (UCapsuleComponent *OurCapsule = GetCharacterOwner()->GetCapsuleComponent())
 		{
 			AvoidanceManager->UpdateRVO(AvoidanceUID, GetActorFeetLocation(),
 				OurCapsule->GetScaledCapsuleRadius(), OurCapsule->GetScaledCapsuleHalfHeight() * 2.0f,
@@ -3873,7 +3873,7 @@ bool UCharacterMovementComponent::CheckWaterJump(FVector CheckPoint, FVector& Wa
 	CheckPoint.Z = 0.f;
 	FVector CheckNorm = CheckPoint.SafeNormal();
 	float PawnCapsuleRadius, PawnCapsuleHalfHeight;
-	CharacterOwner->CapsuleComponent->GetScaledCapsuleSize(PawnCapsuleRadius, PawnCapsuleHalfHeight);
+	CharacterOwner->GetCapsuleComponent()->GetScaledCapsuleSize(PawnCapsuleRadius, PawnCapsuleHalfHeight);
 	CheckPoint = CharacterOwner->GetActorLocation() + 1.2f * PawnCapsuleRadius * CheckNorm;
 	FVector Extent(PawnCapsuleRadius, PawnCapsuleRadius, PawnCapsuleHalfHeight);
 	FHitResult HitInfo(1.f);
@@ -4080,7 +4080,7 @@ void UCharacterMovementComponent::ComputeFloorDist(const FVector& CapsuleLocatio
 	}
 
 	float PawnRadius, PawnHalfHeight;
-	CharacterOwner->CapsuleComponent->GetScaledCapsuleSize(PawnRadius, PawnHalfHeight);
+	CharacterOwner->GetCapsuleComponent()->GetScaledCapsuleSize(PawnRadius, PawnHalfHeight);
 
 	bool bSkipSweep = false;
 	if (DownwardSweepResult != NULL && DownwardSweepResult->IsValidBlockingHit())
@@ -4243,7 +4243,7 @@ void UCharacterMovementComponent::FindFloor(const FVector& CapsuleLocation, FFin
 		if ( bAlwaysCheckFloor || !bZeroDelta || bForceNextFloorCheck || bJustTeleported )
 		{
 			MutableThis->bForceNextFloorCheck = false;
-			ComputeFloorDist(CapsuleLocation, FloorLineTraceDist, FloorSweepTraceDist, OutFloorResult, CharacterOwner->CapsuleComponent->GetScaledCapsuleRadius(), DownwardSweepResult);
+			ComputeFloorDist(CapsuleLocation, FloorLineTraceDist, FloorSweepTraceDist, OutFloorResult, CharacterOwner->GetCapsuleComponent()->GetScaledCapsuleRadius(), DownwardSweepResult);
 		}
 		else
 		{
@@ -4272,7 +4272,7 @@ void UCharacterMovementComponent::FindFloor(const FVector& CapsuleLocation, FFin
 			else
 			{
 				MutableThis->bForceNextFloorCheck = false;
-				ComputeFloorDist(CapsuleLocation, FloorLineTraceDist, FloorSweepTraceDist, OutFloorResult, CharacterOwner->CapsuleComponent->GetScaledCapsuleRadius(), DownwardSweepResult);
+				ComputeFloorDist(CapsuleLocation, FloorLineTraceDist, FloorSweepTraceDist, OutFloorResult, CharacterOwner->GetCapsuleComponent()->GetScaledCapsuleRadius(), DownwardSweepResult);
 			}
 		}
 	}
@@ -4372,7 +4372,7 @@ bool UCharacterMovementComponent::IsValidLandingSpot(const FVector& CapsuleLocat
 		}
 
 		float PawnRadius, PawnHalfHeight;
-		CharacterOwner->CapsuleComponent->GetScaledCapsuleSize(PawnRadius, PawnHalfHeight);
+		CharacterOwner->GetCapsuleComponent()->GetScaledCapsuleSize(PawnRadius, PawnHalfHeight);
 
 		// Reject hits that are above our lower hemisphere (can happen when sliding down a vertical surface).
 		const float LowerHemisphereZ = Hit.Location.Z - PawnHalfHeight + PawnRadius;
@@ -4407,7 +4407,7 @@ bool UCharacterMovementComponent::ShouldCheckForValidLandingSpot(float DeltaTime
 	if (Hit.Normal.Z > KINDA_SMALL_NUMBER && !Hit.Normal.Equals(Hit.ImpactNormal))
 	{
 		const FVector PawnLocation = UpdatedComponent->GetComponentLocation();
-		if (IsWithinEdgeTolerance(PawnLocation, Hit.ImpactPoint, CharacterOwner->CapsuleComponent->GetScaledCapsuleRadius()))
+		if (IsWithinEdgeTolerance(PawnLocation, Hit.ImpactPoint, CharacterOwner->GetCapsuleComponent()->GetScaledCapsuleRadius()))
 		{						
 			return true;
 		}
@@ -4428,7 +4428,7 @@ float UCharacterMovementComponent::GetValidPerchRadius() const
 {
 	if (CharacterOwner)
 	{
-		const float PawnRadius = CharacterOwner->CapsuleComponent->GetScaledCapsuleRadius();
+		const float PawnRadius = CharacterOwner->GetCapsuleComponent()->GetScaledCapsuleRadius();
 		return FMath::Clamp(PawnRadius - GetPerchRadiusThreshold(), 0.1f, PawnRadius);
 	}
 	return 0.f;
@@ -4472,7 +4472,7 @@ bool UCharacterMovementComponent::ComputePerchResult(const float TestRadius, con
 
 	// Sweep further than actual requested distance, because a reduced capsule radius means we could miss some hits that the normal radius would contact.
 	float PawnRadius, PawnHalfHeight;
-	CharacterOwner->CapsuleComponent->GetScaledCapsuleSize(PawnRadius, PawnHalfHeight);
+	CharacterOwner->GetCapsuleComponent()->GetScaledCapsuleSize(PawnRadius, PawnHalfHeight);
 
 	const float InHitAboveBase = FMath::Max(0.f, InHit.ImpactPoint.Z - (InHit.Location.Z - PawnHalfHeight));
 	const float PerchLineDist = FMath::Max(0.f, InMaxFloorDist - InHitAboveBase);
@@ -4545,7 +4545,7 @@ bool UCharacterMovementComponent::StepUp(const FVector& GravDir, const FVector&
 
 	const FVector OldLocation = UpdatedComponent->GetComponentLocation();
 	float PawnRadius, PawnHalfHeight;
-	CharacterOwner->CapsuleComponent->GetScaledCapsuleSize(PawnRadius, PawnHalfHeight);
+	CharacterOwner->GetCapsuleComponent()->GetScaledCapsuleSize(PawnRadius, PawnHalfHeight);
 
 	// Don't bother stepping up if top of capsule is hitting something.
 	const float InitialImpactZ = InHit.ImpactPoint.Z;
@@ -4988,7 +4988,7 @@ void UCharacterMovementComponent::SmoothClientPosition(float DeltaSeconds)
 	if (ClientData && ClientData->bSmoothNetUpdates)
 	{
 		// smooth interpolation of mesh translation to avoid popping of other client pawns, unless driving or ragdoll or low tick rate
-		if ((DeltaSeconds < ClientData->SmoothNetUpdateTime) && CharacterOwner->Mesh && !CharacterOwner->Mesh->IsSimulatingPhysics())
+		if ((DeltaSeconds < ClientData->SmoothNetUpdateTime) && CharacterOwner->GetMesh() && !CharacterOwner->GetMesh()->IsSimulatingPhysics())
 		{
 			ClientData->MeshTranslationOffset = (ClientData->MeshTranslationOffset * (1.f - DeltaSeconds / ClientData->SmoothNetUpdateTime));
 		}
@@ -5003,10 +5003,10 @@ void UCharacterMovementComponent::SmoothClientPosition(float DeltaSeconds)
 			ClientData->MeshTranslationOffset.Z = 0;
 		}
 
-		if (CharacterOwner->Mesh)
+		if (CharacterOwner->GetMesh())
 		{
 			const FVector NewRelTranslation = CharacterOwner->ActorToWorld().InverseTransformVectorNoScale(ClientData->MeshTranslationOffset + CharacterOwner->GetBaseTranslationOffset());
-			CharacterOwner->Mesh->SetRelativeLocation(NewRelTranslation);
+			CharacterOwner->GetMesh()->SetRelativeLocation(NewRelTranslation);
 		}
 	}
 }
@@ -5725,7 +5725,7 @@ void UCharacterMovementComponent::MoveAutonomous
 	PerformMovement(DeltaTime);
 
 	// If not playing root motion, tick animations after physics. We do this here to keep events, notifies, states and transitions in sync with client updates.
-	if( !CharacterOwner->bClientUpdating && !CharacterOwner->IsPlayingRootMotion() && CharacterOwner->Mesh )
+	if( !CharacterOwner->bClientUpdating && !CharacterOwner->IsPlayingRootMotion() && CharacterOwner->GetMesh() )
 	{
 		TickCharacterPose(DeltaTime);
 		// TODO: SaveBaseLocation() in case tick moves us?
@@ -6282,14 +6282,14 @@ void UCharacterMovementComponent::RegisterComponentTickFunctions(bool bRegister)
 
 void UCharacterMovementComponent::TickCharacterPose(float DeltaTime)
 {
-	check(CharacterOwner && CharacterOwner->Mesh);
+	check(CharacterOwner && CharacterOwner->GetMesh());
 
-	CharacterOwner->Mesh->TickPose(DeltaTime);
+	CharacterOwner->GetMesh()->TickPose(DeltaTime);
 
 	// Grab root motion now that we have ticked the pose
 	if (CharacterOwner->IsPlayingRootMotion())
 	{
-		if (UAnimInstance * AnimInstance = CharacterOwner->Mesh->GetAnimInstance())
+		if (UAnimInstance * AnimInstance = CharacterOwner->GetMesh()->GetAnimInstance())
 		{
 			RootMotionParams.Accumulate(AnimInstance->ConsumeExtractedRootMotion());
 		}
@@ -6521,9 +6521,9 @@ void FSavedMove_Character::SetMoveFor(ACharacter* Character, float InDeltaTime,
 
 	bPressedJump = Character->bPressedJump;
 	JumpKeyHoldTime = Character->JumpKeyHoldTime;
-	bWantsToCrouch = Character->CharacterMovement->bWantsToCrouch;
-	bForceMaxAccel = Character->CharacterMovement->bForceMaxAccel;
-	MovementMode = Character->CharacterMovement->PackNetworkMovementMode();
+	bWantsToCrouch = Character->GetCharacterMovement()->bWantsToCrouch;
+	bForceMaxAccel = Character->GetCharacterMovement()->bForceMaxAccel;
+	MovementMode = Character->GetCharacterMovement()->PackNetworkMovementMode();
 
 	TimeStamp = ClientData.CurrentTimeStamp;
 }
@@ -6532,11 +6532,11 @@ void FSavedMove_Character::SetInitialPosition(ACharacter* Character)
 {
 	StartLocation = Character->GetActorLocation();
 	StartRotation = Character->GetActorRotation();
-	StartVelocity = Character->CharacterMovement->Velocity;
+	StartVelocity = Character->GetCharacterMovement()->Velocity;
 	UPrimitiveComponent* const MovementBase = Character->GetMovementBase();
 	StartBase = MovementBase;
 	StartBaseRotation = FQuat::Identity;
-	StartFloor = Character->CharacterMovement->CurrentFloor;
+	StartFloor = Character->GetCharacterMovement()->CurrentFloor;
 	CustomTimeDilation = Character->CustomTimeDilation;
 	StartBoneName = Character->GetBasedMovement().BoneName;
 
@@ -6548,7 +6548,7 @@ void FSavedMove_Character::SetInitialPosition(ACharacter* Character)
 	}
 
 	StartControlRotation = Character->GetControlRotation().Clamp();
-	Character->CapsuleComponent->GetScaledCapsuleSize(StartCapsuleRadius, StartCapsuleHalfHeight);
+	Character->GetCapsuleComponent()->GetScaledCapsuleSize(StartCapsuleRadius, StartCapsuleHalfHeight);
 }
 
 void FSavedMove_Character::PostUpdate(ACharacter* Character, FSavedMove_Character::EPostUpdateMode PostUpdateMode)
@@ -6682,10 +6682,10 @@ void FSavedMove_Character::PrepMoveFor(ACharacter* Character)
 				RootMotionMontageInstance->SetPosition(RootMotionTrackPosition);
 			}
 		}
-		Character->CharacterMovement->RootMotionParams = RootMotionMovement;
+		Character->GetCharacterMovement()->RootMotionParams = RootMotionMovement;
 	}
 
-	Character->CharacterMovement->bForceMaxAccel = bForceMaxAccel;
+	Character->GetCharacterMovement()->bForceMaxAccel = bForceMaxAccel;
 	Character->JumpKeyHoldTime = JumpKeyHoldTime;
 }
 
diff --git a/Engine/Source/Runtime/Engine/Private/CheatManager.cpp b/Engine/Source/Runtime/Engine/Private/CheatManager.cpp
--- a/Engine/Source/Runtime/Engine/Private/CheatManager.cpp
+++ b/Engine/Source/Runtime/Engine/Private/CheatManager.cpp
@@ -92,11 +92,11 @@ void UCheatManager::ChangeSize( float F )
 	if (Character)
 	{
 		ACharacter* DefaultCharacter = Character->GetClass()->GetDefaultObject<ACharacter>();
-		Character->CapsuleComponent->SetCapsuleSize( DefaultCharacter->CapsuleComponent->GetUnscaledCapsuleRadius() * F, DefaultCharacter->CapsuleComponent->GetUnscaledCapsuleHalfHeight() * F );
+		Character->GetCapsuleComponent()->SetCapsuleSize(DefaultCharacter->GetCapsuleComponent()->GetUnscaledCapsuleRadius() * F, DefaultCharacter->GetCapsuleComponent()->GetUnscaledCapsuleHalfHeight() * F);
 
-		if (Character->Mesh.IsValid())
+		if (Character->GetMesh())
 		{
-			Character->Mesh->SetRelativeScale3D( FVector(F) );
+			Character->GetMesh()->SetRelativeScale3D(FVector(F));
 		}
 		Character->TeleportTo(Character->GetActorLocation(), Character->GetActorRotation());
 	}
@@ -857,7 +857,7 @@ void UCheatManager::TestCollisionDistance()
 			if (Volume->GetClass()->GetDefaultObject() != Volume)
 			{
 				FVector ClosestPoint(0,0,0);
-				float Distance = Volume->BrushComponent->GetDistanceToCollision(ViewLoc, ClosestPoint);
+				float Distance = Volume->GetBrushComponent()->GetDistanceToCollision(ViewLoc, ClosestPoint);
 				float NormalizedDistance = FMath::Clamp<float>(Distance, 0.f, 1000.f)/1000.f;
 				FColor DrawColor(255*NormalizedDistance, 255*(1-NormalizedDistance), 0);
 				DrawDebugLine(GetWorld(), ViewLoc, ClosestPoint, DrawColor, true);
diff --git a/Engine/Source/Runtime/Engine/Private/CullDistanceVolume.cpp b/Engine/Source/Runtime/Engine/Private/CullDistanceVolume.cpp
--- a/Engine/Source/Runtime/Engine/Private/CullDistanceVolume.cpp
+++ b/Engine/Source/Runtime/Engine/Private/CullDistanceVolume.cpp
@@ -6,9 +6,9 @@
 ACullDistanceVolume::ACullDistanceVolume(const FObjectInitializer& ObjectInitializer)
 	: Super(ObjectInitializer)
 {
-	BrushComponent->BodyInstance.bEnableCollision_DEPRECATED = false;
-	BrushComponent->SetCollisionProfileName(UCollisionProfile::NoCollision_ProfileName);
-	BrushComponent->bAlwaysCreatePhysicsState = true;
+	GetBrushComponent()->BodyInstance.bEnableCollision_DEPRECATED = false;
+	GetBrushComponent()->SetCollisionProfileName(UCollisionProfile::NoCollision_ProfileName);
+	GetBrushComponent()->bAlwaysCreatePhysicsState = true;
 
 	CullDistances.Add(FCullDistanceSizePair(0,0));
 	CullDistances.Add(FCullDistanceSizePair(10000,0));
@@ -73,7 +73,7 @@ bool ACullDistanceVolume::CanBeAffectedByVolumes( UPrimitiveComponent* Primitive
 void ACullDistanceVolume::GetPrimitiveMaxDrawDistances(TMap<UPrimitiveComponent*,float>& OutCullDistances)
 {
 	// Nothing to do if there is no brush component or no cull distances are set
-	if( BrushComponent && CullDistances.Num() > 0 && bEnabled )
+	if (GetBrushComponent() && CullDistances.Num() > 0 && bEnabled)
 	{
 		for (auto It(OutCullDistances.CreateIterator()); It; ++It)
 		{
diff --git a/Engine/Source/Runtime/Engine/Private/DecalActor.cpp b/Engine/Source/Runtime/Engine/Private/DecalActor.cpp
--- a/Engine/Source/Runtime/Engine/Private/DecalActor.cpp
+++ b/Engine/Source/Runtime/Engine/Private/DecalActor.cpp
@@ -91,7 +91,7 @@ void ADecalActor::PostEditMove(bool bFinished)
 {
 	Super::PostEditMove(bFinished);
 
-	if (Decal.IsValid())
+	if (Decal)
 	{
 		Decal->RecreateRenderState_Concurrent();
 	}
@@ -102,7 +102,7 @@ void ADecalActor::PostEditChangeProperty(FPropertyChangedEvent& PropertyChangedE
 	// AActor::PostEditChange will ForceUpdateComponents()
 	Super::PostEditChangeProperty(PropertyChangedEvent);
 
-	if (Decal.IsValid())
+	if (Decal)
 	{
 		Decal->RecreateRenderState_Concurrent();
 	}
@@ -118,7 +118,7 @@ void ADecalActor::EditorApplyScale(const FVector& DeltaScale, const FVector* Piv
 
 void ADecalActor::SetDecalMaterial(class UMaterialInterface* NewDecalMaterial)
 {
-	if (Decal.IsValid())
+	if (Decal)
 	{
 		Decal->SetDecalMaterial(NewDecalMaterial);
 	}
@@ -126,10 +126,10 @@ void ADecalActor::SetDecalMaterial(class UMaterialInterface* NewDecalMaterial)
 
 class UMaterialInterface* ADecalActor::GetDecalMaterial() const
 {
-	return Decal.IsValid() ? Decal->GetDecalMaterial() : NULL;
+	return Decal ? Decal->GetDecalMaterial() : NULL;
 }
 
 class UMaterialInstanceDynamic* ADecalActor::CreateDynamicMaterialInstance()
 {
-	return Decal.IsValid() ? Decal->CreateDynamicMaterialInstance() : NULL;
+	return Decal ? Decal->CreateDynamicMaterialInstance() : NULL;
 }
diff --git a/Engine/Source/Runtime/Engine/Private/DestructibleActor.cpp b/Engine/Source/Runtime/Engine/Private/DestructibleActor.cpp
--- a/Engine/Source/Runtime/Engine/Private/DestructibleActor.cpp
+++ b/Engine/Source/Runtime/Engine/Private/DestructibleActor.cpp
@@ -20,7 +20,7 @@ ADestructibleActor::ADestructibleActor(const FObjectInitializer& ObjectInitializ
 #if WITH_EDITOR
 bool ADestructibleActor::GetReferencedContentObjects( TArray<UObject*>& Objects ) const
 {
-	if (DestructibleComponent.IsValid() && DestructibleComponent->SkeletalMesh)
+	if (DestructibleComponent && DestructibleComponent->SkeletalMesh)
 	{
 		Objects.Add(DestructibleComponent->SkeletalMesh);
 	}
diff --git a/Engine/Source/Runtime/Engine/Private/DocumentationActor.cpp b/Engine/Source/Runtime/Engine/Private/DocumentationActor.cpp
--- a/Engine/Source/Runtime/Engine/Private/DocumentationActor.cpp
+++ b/Engine/Source/Runtime/Engine/Private/DocumentationActor.cpp
@@ -10,7 +10,7 @@
 ADocumentationActor::ADocumentationActor(const FObjectInitializer& ObjectInitializer)
 	: Super(ObjectInitializer)
 {
-	TSubobjectPtr<USceneComponent> SceneComponent = ObjectInitializer.CreateDefaultSubobject<USceneComponent>(this, TEXT("SceneComp"));
+	USceneComponent* SceneComponent = ObjectInitializer.CreateDefaultSubobject<USceneComponent>(this, TEXT("SceneComp"));
 	RootComponent = SceneComponent;	
 
 #if WITH_EDITORONLY_DATA
diff --git a/Engine/Source/Runtime/Engine/Private/FoliageComponent.cpp b/Engine/Source/Runtime/Engine/Private/FoliageComponent.cpp
--- a/Engine/Source/Runtime/Engine/Private/FoliageComponent.cpp
+++ b/Engine/Source/Runtime/Engine/Private/FoliageComponent.cpp
@@ -116,10 +116,10 @@ void AInteractiveFoliageActor::CapsuleTouched(AActor* Other, UPrimitiveComponent
 
 void AInteractiveFoliageActor::SetupCollisionCylinder()
 {
-	if (StaticMeshComponent->StaticMesh)
+	if (GetStaticMeshComponent()->StaticMesh)
 	{
-		const FBoxSphereBounds MeshBounds = StaticMeshComponent->StaticMesh->GetBounds();
-		const FVector Scale3D = StaticMeshComponent->RelativeScale3D;
+		const FBoxSphereBounds MeshBounds = GetStaticMeshComponent()->StaticMesh->GetBounds();
+		const FVector Scale3D = GetStaticMeshComponent()->RelativeScale3D;
 		// Set the cylinder's radius based off of the static mesh's bounds radius
 		// CollisionRadius is in world space so apply the actor's scale
 		CapsuleComponent->SetCapsuleSize(MeshBounds.SphereRadius * .7f * FMath::Max(Scale3D.X, Scale3D.Y), MeshBounds.BoxExtent.Z * Scale3D.Z);
@@ -132,7 +132,7 @@ void AInteractiveFoliageActor::SetupCollisionCylinder()
 
 void AInteractiveFoliageActor::Tick(float DeltaSeconds)
 {
-	UInteractiveFoliageComponent* const FoliageComponent = CastChecked<UInteractiveFoliageComponent>(StaticMeshComponent);
+	UInteractiveFoliageComponent* const FoliageComponent = CastChecked<UInteractiveFoliageComponent>(GetStaticMeshComponent());
 	// Can only push updates to the scene proxy if we are being ticked while registered
 	// The proxy will be NULL on dedicated server
 	if (FoliageComponent->IsRegistered() && FoliageComponent->FoliageSceneProxy)
diff --git a/Engine/Source/Runtime/Engine/Private/HeightFogComponent.cpp b/Engine/Source/Runtime/Engine/Private/HeightFogComponent.cpp
--- a/Engine/Source/Runtime/Engine/Private/HeightFogComponent.cpp
+++ b/Engine/Source/Runtime/Engine/Private/HeightFogComponent.cpp
@@ -166,7 +166,7 @@ AExponentialHeightFog::AExponentialHeightFog(const FObjectInitializer& ObjectIni
 	bHidden = false;
 
 #if WITH_EDITORONLY_DATA
-	if (!IsRunningCommandlet() && (SpriteComponent != NULL))
+	if (!IsRunningCommandlet() && (GetSpriteComponent() != NULL))
 	{
 		// Structure to hold one-time initialization
 		struct FConstructorStatics
@@ -183,11 +183,11 @@ AExponentialHeightFog::AExponentialHeightFog(const FObjectInitializer& ObjectIni
 		};
 		static FConstructorStatics ConstructorStatics;
 
-		SpriteComponent->Sprite = ConstructorStatics.FogTextureObject.Get();
-		SpriteComponent->RelativeScale3D = FVector(0.5f, 0.5f, 0.5f);
-		SpriteComponent->SpriteInfo.Category = ConstructorStatics.ID_Fog;
-		SpriteComponent->SpriteInfo.DisplayName = ConstructorStatics.NAME_Fog;
-		SpriteComponent->AttachParent = Component;
+		GetSpriteComponent()->Sprite = ConstructorStatics.FogTextureObject.Get();
+		GetSpriteComponent()->RelativeScale3D = FVector(0.5f, 0.5f, 0.5f);
+		GetSpriteComponent()->SpriteInfo.Category = ConstructorStatics.ID_Fog;
+		GetSpriteComponent()->SpriteInfo.DisplayName = ConstructorStatics.NAME_Fog;
+		GetSpriteComponent()->AttachParent = Component;
 	}
 #endif // WITH_EDITORONLY_DATA
 }
diff --git a/Engine/Source/Runtime/Engine/Private/InstancedFoliage.cpp b/Engine/Source/Runtime/Engine/Private/InstancedFoliage.cpp
--- a/Engine/Source/Runtime/Engine/Private/InstancedFoliage.cpp
+++ b/Engine/Source/Runtime/Engine/Private/InstancedFoliage.cpp
@@ -707,7 +707,7 @@ void FFoliageMeshInfo::SelectInstances(AInstancedFoliageActor* InIFA, bool bSele
 AInstancedFoliageActor::AInstancedFoliageActor(const FObjectInitializer& ObjectInitializer)
 	: Super(ObjectInitializer)
 {
-	TSubobjectPtr<USceneComponent> SceneComponent = ObjectInitializer.CreateDefaultSubobject<USceneComponent>(this, TEXT("RootComponent0"));
+	USceneComponent* SceneComponent = ObjectInitializer.CreateDefaultSubobject<USceneComponent>(this, TEXT("RootComponent0"));
 	RootComponent = SceneComponent;
 	RootComponent->Mobility = EComponentMobility::Movable;
 	
diff --git a/Engine/Source/Runtime/Engine/Private/InteractiveFoliageActor.cpp b/Engine/Source/Runtime/Engine/Private/InteractiveFoliageActor.cpp
--- a/Engine/Source/Runtime/Engine/Private/InteractiveFoliageActor.cpp
+++ b/Engine/Source/Runtime/Engine/Private/InteractiveFoliageActor.cpp
@@ -10,7 +10,7 @@ AInteractiveFoliageActor::AInteractiveFoliageActor(const FObjectInitializer& Obj
 		.SetDefaultSubobjectClass<UInteractiveFoliageComponent>("StaticMeshComponent0"))
 {
 
-	UInteractiveFoliageComponent* FoliageMeshComponent = CastChecked<UInteractiveFoliageComponent>(StaticMeshComponent);
+	UInteractiveFoliageComponent* FoliageMeshComponent = CastChecked<UInteractiveFoliageComponent>(GetStaticMeshComponent());
 	FoliageMeshComponent->BodyInstance.bEnableCollision_DEPRECATED = false;
 	FoliageMeshComponent->SetCollisionProfileName(UCollisionProfile::NoCollision_ProfileName);
 	FoliageMeshComponent->Mobility = EComponentMobility::Static;
diff --git a/Engine/Source/Runtime/Engine/Private/Interpolation.cpp b/Engine/Source/Runtime/Engine/Private/Interpolation.cpp
--- a/Engine/Source/Runtime/Engine/Private/Interpolation.cpp
+++ b/Engine/Source/Runtime/Engine/Private/Interpolation.cpp
@@ -320,7 +320,7 @@ bool AMatineeActor::IgnoreActorSelection()
 AMatineeActor::AMatineeActor(const FObjectInitializer& ObjectInitializer)
 	: Super(ObjectInitializer)
 {
-	TSubobjectPtr<USceneComponent> SceneComponent = ObjectInitializer.CreateDefaultSubobject<USceneComponent>(this, TEXT("SceneComp"));
+	USceneComponent* SceneComponent = ObjectInitializer.CreateDefaultSubobject<USceneComponent>(this, TEXT("SceneComp"));
 	RootComponent = SceneComponent;
 
 #if WITH_EDITORONLY_DATA
@@ -5116,19 +5116,19 @@ void UInterpTrackToggle::UpdateTrack(float NewPosition, UInterpTrackInst* TrInst
 				// We have found the key to the left of the position
 				if (ToggleKey.ToggleAction == ETTA_On)
 				{
-					EmitterActor->ParticleSystemComponent->ActivateSystem(bActivateWithJustAttachedFlag);
+					EmitterActor->GetParticleSystemComponent()->ActivateSystem(bActivateWithJustAttachedFlag);
 				}
 				else
 				if (ToggleKey.ToggleAction == ETTA_Trigger)
 				{
 					if (ToggleKey.Time >= ToggleInst->LastUpdatePosition)
 					{
-						EmitterActor->ParticleSystemComponent->SetActive(true, bActivateWithJustAttachedFlag);
+						EmitterActor->GetParticleSystemComponent()->SetActive(true, bActivateWithJustAttachedFlag);
 					}
 				}
 				else
 				{
-					EmitterActor->ParticleSystemComponent->DeactivateSystem();
+					EmitterActor->GetParticleSystemComponent()->DeactivateSystem();
 				}
 				break;
 			}
@@ -5251,7 +5251,7 @@ void UInterpTrackToggle::UpdateTrack(float NewPosition, UInterpTrackInst* TrInst
 						if( bPlayTriggersWhenJumping || !bJump )
 						{
 							// Use ActivateSystem as multiple triggers should fire it multiple times.
-							EmitterActor->ParticleSystemComponent->ActivateSystem(bActivateWithJustAttachedFlag);
+							EmitterActor->GetParticleSystemComponent()->ActivateSystem(bActivateWithJustAttachedFlag);
 							// don't set bCurrentlyActive (assume it's a one shot effect which the client will perform through its own matinee simulation)
 						}
 					}
@@ -5288,7 +5288,7 @@ void UInterpTrackToggle::UpdateTrack(float NewPosition, UInterpTrackInst* TrInst
 						bShouldActivate = !bShouldActivate;
 					}
 
-					EmitterActor->ParticleSystemComponent->SetActive( bShouldActivate, bActivateWithJustAttachedFlag );
+					EmitterActor->GetParticleSystemComponent()->SetActive( bShouldActivate, bActivateWithJustAttachedFlag );
 					EmitterActor->bCurrentlyActive = bShouldActivate;
 					if (!MatineeActor->bClientSideOnly)
 					{
@@ -5307,7 +5307,7 @@ void UInterpTrackToggle::UpdateTrack(float NewPosition, UInterpTrackInst* TrInst
 							bShouldActivate = !bShouldActivate;
 						}
 
-						LightActor->LightComponent->SetVisibility( bShouldActivate );
+						LightActor->GetLightComponent()->SetVisibility( bShouldActivate );
 					}
 				}
 				else
@@ -8149,7 +8149,7 @@ void UInterpTrackFloatParticleParam::UpdateTrack(float NewPosition, UInterpTrack
 	}
 
 	float NewFloatValue = FloatTrack.Eval(NewPosition, 0.f);
-	Emitter->ParticleSystemComponent->SetFloatParameter( ParamName, NewFloatValue );
+	Emitter->GetParticleSystemComponent()->SetFloatParameter( ParamName, NewFloatValue );
 }
 
 /*-----------------------------------------------------------------------------
@@ -8170,7 +8170,7 @@ void UInterpTrackInstFloatParticleParam::SaveActorState(UInterpTrack* Track)
 		return;
 	}
 
-	bool bFoundParam = Emitter->ParticleSystemComponent->GetFloatParameter( ParamTrack->ParamName, ResetFloat );
+	bool bFoundParam = Emitter->GetParticleSystemComponent()->GetFloatParameter( ParamTrack->ParamName, ResetFloat );
 	if(!bFoundParam)
 	{
 		ResetFloat = 0.f;
@@ -8187,7 +8187,7 @@ void UInterpTrackInstFloatParticleParam::RestoreActorState(UInterpTrack* Track)
 		return;
 	}
 
-	Emitter->ParticleSystemComponent->SetFloatParameter( ParamTrack->ParamName, ResetFloat );
+	Emitter->GetParticleSystemComponent()->SetFloatParameter( ParamTrack->ParamName, ResetFloat );
 }
 
 /*------------------------------------------------------------------------------
@@ -8774,7 +8774,7 @@ static void GetListOfStaticActors(FString& OutString, const TArray<AActor*>& Act
 AMaterialInstanceActor::AMaterialInstanceActor(const FObjectInitializer& ObjectInitializer)
 	: Super(ObjectInitializer)
 {
-	TSubobjectPtr<USceneComponent> SceneComponent = ObjectInitializer.CreateDefaultSubobject<USceneComponent>(this, TEXT("SceneComp"));
+	USceneComponent* SceneComponent = ObjectInitializer.CreateDefaultSubobject<USceneComponent>(this, TEXT("SceneComp"));
 	RootComponent = SceneComponent;
 
 #if WITH_EDITORONLY_DATA
@@ -8863,7 +8863,7 @@ void UInterpTrackInstToggle::SaveActorState(UInterpTrack* Track)
 	}
 	else if( LightActor != NULL )
 	{
-		bSavedActiveState = LightActor->LightComponent->bVisible;
+		bSavedActiveState = LightActor->GetLightComponent()->bVisible;
 	}
 }
 
@@ -8879,7 +8879,7 @@ void UInterpTrackInstToggle::RestoreActorState(UInterpTrack* Track)
 	if( EmitterActor )
 	{
 		// Use SetActive to only activate a non-active system...
-		EmitterActor->ParticleSystemComponent->SetActive(bSavedActiveState);
+		EmitterActor->GetParticleSystemComponent()->SetActive(bSavedActiveState);
 		EmitterActor->bCurrentlyActive = bSavedActiveState;
 		EmitterActor->ForceNetRelevant();
 	}
@@ -8888,7 +8888,7 @@ void UInterpTrackInstToggle::RestoreActorState(UInterpTrack* Track)
 		// We'll only allow *toggleable* lights to be toggled like this!  Static lights are ignored.
 		if( LightActor->IsToggleable() )
 		{
-			LightActor->LightComponent->SetVisibility( bSavedActiveState );
+			LightActor->GetLightComponent()->SetVisibility( bSavedActiveState );
 		}
 	}
 }
@@ -9389,7 +9389,7 @@ void UInterpTrackParticleReplay::UpdateTrack( float NewPosition, UInterpTrackIns
 
 	// Particle replay tracks are expecting to be dealing with emitter actors
 	AEmitter* EmitterActor = Cast< AEmitter >( Actor );
-	if( EmitterActor != NULL && EmitterActor->ParticleSystemComponent.IsValid() )
+	if( EmitterActor != NULL && EmitterActor->GetParticleSystemComponent() )
 	{
 		if( ( NewPosition > ParticleReplayInst->LastUpdatePosition ) && !bJump )
 		{
@@ -9405,7 +9405,7 @@ void UInterpTrackParticleReplay::UpdateTrack( float NewPosition, UInterpTrackIns
 					{
 						// Do we already have data for this clip?
 						UParticleSystemReplay* ExistingClipReplay =
-							EmitterActor->ParticleSystemComponent->FindReplayClipForIDNumber( ParticleReplayKey.ClipIDNumber );
+							EmitterActor->GetParticleSystemComponent()->FindReplayClipForIDNumber( ParticleReplayKey.ClipIDNumber );
 						if( ExistingClipReplay != NULL )
 						{
 							// Clear the existing clip's frame data.  We're re-recording the clip now!
@@ -9413,23 +9413,23 @@ void UInterpTrackParticleReplay::UpdateTrack( float NewPosition, UInterpTrackIns
 						}
 
 						// Start capturing!
-						EmitterActor->ParticleSystemComponent->ReplayState = PRS_Capturing;
-						EmitterActor->ParticleSystemComponent->ReplayClipIDNumber = ParticleReplayKey.ClipIDNumber;
-						EmitterActor->ParticleSystemComponent->ReplayFrameIndex = 0;
+						EmitterActor->GetParticleSystemComponent()->ReplayState = PRS_Capturing;
+						EmitterActor->GetParticleSystemComponent()->ReplayClipIDNumber = ParticleReplayKey.ClipIDNumber;
+						EmitterActor->GetParticleSystemComponent()->ReplayFrameIndex = 0;
 
 						// Make sure we're alive and kicking
-						EmitterActor->ParticleSystemComponent->SetActive( true );
+						EmitterActor->GetParticleSystemComponent()->SetActive( true );
 					}
 					else
 #endif // WITH_EDITORONLY_DATA
 					{
 						// Start playback!
-						EmitterActor->ParticleSystemComponent->ReplayState = PRS_Replaying;
-						EmitterActor->ParticleSystemComponent->ReplayClipIDNumber = ParticleReplayKey.ClipIDNumber;
-						EmitterActor->ParticleSystemComponent->ReplayFrameIndex = 0;
+						EmitterActor->GetParticleSystemComponent()->ReplayState = PRS_Replaying;
+						EmitterActor->GetParticleSystemComponent()->ReplayClipIDNumber = ParticleReplayKey.ClipIDNumber;
+						EmitterActor->GetParticleSystemComponent()->ReplayFrameIndex = 0;
 
 						// Make sure we're alive and kicking
-						EmitterActor->ParticleSystemComponent->SetActive( true );
+						EmitterActor->GetParticleSystemComponent()->SetActive( true );
 					}
 				}
 
@@ -9442,13 +9442,13 @@ void UInterpTrackParticleReplay::UpdateTrack( float NewPosition, UInterpTrackIns
 #endif // WITH_EDITORONLY_DATA
 					{
 						// Done playing back replay sequence, so turn off the particle system
-						EmitterActor->ParticleSystemComponent->SetActive( false );
+						EmitterActor->GetParticleSystemComponent()->SetActive( false );
 
 						// Stop playback/capture!  We'll still keep the particle system in 'replay mode' while
 						// the replay track is bound so that the system doesn't start simulating/rendering
-						EmitterActor->ParticleSystemComponent->ReplayState = PRS_Replaying;
-						EmitterActor->ParticleSystemComponent->ReplayClipIDNumber = INDEX_NONE;
-						EmitterActor->ParticleSystemComponent->ReplayFrameIndex = INDEX_NONE;
+						EmitterActor->GetParticleSystemComponent()->ReplayState = PRS_Replaying;
+						EmitterActor->GetParticleSystemComponent()->ReplayClipIDNumber = INDEX_NONE;
+						EmitterActor->GetParticleSystemComponent()->ReplayFrameIndex = INDEX_NONE;
 					}
 				}
 			}
@@ -9462,7 +9462,7 @@ void UInterpTrackParticleReplay::UpdateTrack( float NewPosition, UInterpTrackIns
 			if( bIsCapturingReplay )
 			{
 				// Scrubbing while capturing will stop the capture
-				EmitterActor->ParticleSystemComponent->ReplayState = PRS_Disabled;
+				EmitterActor->GetParticleSystemComponent()->ReplayState = PRS_Disabled;
 			}
 			else
 			{
@@ -9493,7 +9493,7 @@ void UInterpTrackParticleReplay::UpdateTrack( float NewPosition, UInterpTrackIns
 
 					// Check to see if we have a clip
 					UParticleSystemReplay* ParticleSystemReplay =
-						EmitterActor->ParticleSystemComponent->FindReplayClipForIDNumber( CurrentReplayStartKey.ClipIDNumber );
+						EmitterActor->GetParticleSystemComponent()->FindReplayClipForIDNumber( CurrentReplayStartKey.ClipIDNumber );
 					if( ParticleSystemReplay != NULL )
 					{
 						if( ReplayFrameIndex < ParticleSystemReplay->Frames.Num() )
@@ -9502,12 +9502,12 @@ void UInterpTrackParticleReplay::UpdateTrack( float NewPosition, UInterpTrackIns
 							bIsReplayingSingleFrame = true;
 
 							// Make sure replay mode is turned on
-							EmitterActor->ParticleSystemComponent->ReplayState = PRS_Replaying;
-							EmitterActor->ParticleSystemComponent->ReplayClipIDNumber = CurrentReplayStartKey.ClipIDNumber;
-							EmitterActor->ParticleSystemComponent->ReplayFrameIndex = ReplayFrameIndex;
+							EmitterActor->GetParticleSystemComponent()->ReplayState = PRS_Replaying;
+							EmitterActor->GetParticleSystemComponent()->ReplayClipIDNumber = CurrentReplayStartKey.ClipIDNumber;
+							EmitterActor->GetParticleSystemComponent()->ReplayFrameIndex = ReplayFrameIndex;
 
 							// Make sure we're alive and kicking
-							EmitterActor->ParticleSystemComponent->SetActive( true );
+							EmitterActor->GetParticleSystemComponent()->SetActive( true );
 						}
 					}
 				}
@@ -9516,13 +9516,13 @@ void UInterpTrackParticleReplay::UpdateTrack( float NewPosition, UInterpTrackIns
 				{
 					// Stop playback!  We'll still keep the particle system in 'replay mode' while
 					// the replay track is bound so that the system doesn't start simulating/rendering
-					EmitterActor->ParticleSystemComponent->ReplayState = PRS_Replaying;
-					EmitterActor->ParticleSystemComponent->ReplayClipIDNumber = INDEX_NONE;
-					EmitterActor->ParticleSystemComponent->ReplayFrameIndex = INDEX_NONE;
+					EmitterActor->GetParticleSystemComponent()->ReplayState = PRS_Replaying;
+					EmitterActor->GetParticleSystemComponent()->ReplayClipIDNumber = INDEX_NONE;
+					EmitterActor->GetParticleSystemComponent()->ReplayFrameIndex = INDEX_NONE;
 
 					// We're not currently capturing and we're not in the middle of a replay frame,
 					// so turn off the particle system
-					EmitterActor->ParticleSystemComponent->SetActive( false );
+					EmitterActor->GetParticleSystemComponent()->SetActive( false );
 				}
 			}
 		}
@@ -9530,16 +9530,16 @@ void UInterpTrackParticleReplay::UpdateTrack( float NewPosition, UInterpTrackIns
 #endif // WITH_EDITORONLY_DATA
 		{
 			// Okay, we're not scrubbing, but are we replaying a particle system?
-			if( EmitterActor->ParticleSystemComponent->ReplayState == PRS_Replaying )
+			if( EmitterActor->GetParticleSystemComponent()->ReplayState == PRS_Replaying )
 			{
 				// Advance to next frame (or reverse to the previous frame)
 				if( MatineeActor->bReversePlayback )
 				{
-					--EmitterActor->ParticleSystemComponent->ReplayFrameIndex;
+					--EmitterActor->GetParticleSystemComponent()->ReplayFrameIndex;
 				}
 				else
 				{
-					++EmitterActor->ParticleSystemComponent->ReplayFrameIndex;
+					++EmitterActor->GetParticleSystemComponent()->ReplayFrameIndex;
 				}
 			}
 		}
@@ -9610,15 +9610,15 @@ void UInterpTrackInstParticleReplay::RestoreActorState(UInterpTrack* Track)
 	{
 		// Particle replay tracks are expecting to be dealing with emitter actors
 		AEmitter* EmitterActor = Cast< AEmitter >( Actor );
-		if( EmitterActor != NULL && EmitterActor->ParticleSystemComponent.IsValid() )
+		if( EmitterActor != NULL && EmitterActor->GetParticleSystemComponent() )
 		{
 			// Make sure we don't leave the particle system in 'capture mode'
 			
 			// Stop playback/capture!  We'll still keep the particle system in 'replay mode' while
 			// the replay track is bound so that the system doesn't start simulating/rendering
-			EmitterActor->ParticleSystemComponent->ReplayState = PRS_Disabled;
-			EmitterActor->ParticleSystemComponent->ReplayClipIDNumber = 0;
-			EmitterActor->ParticleSystemComponent->ReplayFrameIndex = 0;
+			EmitterActor->GetParticleSystemComponent()->ReplayState = PRS_Disabled;
+			EmitterActor->GetParticleSystemComponent()->ReplayClipIDNumber = 0;
+			EmitterActor->GetParticleSystemComponent()->ReplayFrameIndex = 0;
 		}
 	}
 }
diff --git a/Engine/Source/Runtime/Engine/Private/KismetSystemLibrary.cpp b/Engine/Source/Runtime/Engine/Private/KismetSystemLibrary.cpp
--- a/Engine/Source/Runtime/Engine/Private/KismetSystemLibrary.cpp
+++ b/Engine/Source/Runtime/Engine/Private/KismetSystemLibrary.cpp
@@ -2436,7 +2436,7 @@ void UKismetSystemLibrary::DrawDebugCamera(const ACameraActor* CameraActor, FLin
 	{
 		FVector CamLoc = CameraActor->GetActorLocation();
 		FRotator CamRot = CameraActor->GetActorRotation();
-		::DrawDebugCamera(CameraActor->GetWorld(), CameraActor->GetActorLocation(), CameraActor->GetActorRotation(), CameraActor->CameraComponent->FieldOfView, 1.0f, CameraColor, false, Duration, SDPG_World);
+		::DrawDebugCamera(CameraActor->GetWorld(), CameraActor->GetActorLocation(), CameraActor->GetActorRotation(), CameraActor->GetCameraComponent()->FieldOfView, 1.0f, CameraColor, false, Duration, SDPG_World);
 	}
 }
 
diff --git a/Engine/Source/Runtime/Engine/Private/Level.cpp b/Engine/Source/Runtime/Engine/Private/Level.cpp
--- a/Engine/Source/Runtime/Engine/Private/Level.cpp
+++ b/Engine/Source/Runtime/Engine/Private/Level.cpp
@@ -1544,7 +1544,7 @@ ABrush* ULevel::GetDefaultBrush() const
 		// If the second actor is not a brush then it certainly cannot be the builder brush.
 		if (DefaultBrush != nullptr)
 		{
-			checkf(DefaultBrush->BrushComponent, *GetPathName());
+			checkf(DefaultBrush->GetBrushComponent(), *GetPathName());
 			checkf(DefaultBrush->Brush != nullptr, *GetPathName());
 		}
 	}
diff --git a/Engine/Source/Runtime/Engine/Private/LevelBounds.cpp b/Engine/Source/Runtime/Engine/Private/LevelBounds.cpp
--- a/Engine/Source/Runtime/Engine/Private/LevelBounds.cpp
+++ b/Engine/Source/Runtime/Engine/Private/LevelBounds.cpp
@@ -9,7 +9,7 @@ static const FVector DefaultLevelSize = FVector(1000.f);
 ALevelBounds::ALevelBounds(const FObjectInitializer& ObjectInitializer)
 	: Super(ObjectInitializer)
 {
-	TSubobjectPtr<UBoxComponent> BoxComponent = ObjectInitializer.CreateDefaultSubobject<UBoxComponent>(this, TEXT("BoxComponent0"));
+	UBoxComponent* BoxComponent = ObjectInitializer.CreateDefaultSubobject<UBoxComponent>(this, TEXT("BoxComponent0"));
 	RootComponent = BoxComponent;
 	RootComponent->Mobility = EComponentMobility::Static;
 	RootComponent->RelativeScale3D = DefaultLevelSize;
diff --git a/Engine/Source/Runtime/Engine/Private/LevelStreaming.cpp b/Engine/Source/Runtime/Engine/Private/LevelStreaming.cpp
--- a/Engine/Source/Runtime/Engine/Private/LevelStreaming.cpp
+++ b/Engine/Source/Runtime/Engine/Private/LevelStreaming.cpp
@@ -794,9 +794,9 @@ FBox ULevelStreaming::GetStreamingVolumeBounds()
 	for(int32 VolIdx=0; VolIdx<EditorStreamingVolumes.Num(); VolIdx++)
 	{
 		ALevelStreamingVolume* StreamingVol = EditorStreamingVolumes[VolIdx];
-		if(StreamingVol && StreamingVol->BrushComponent)
+		if(StreamingVol && StreamingVol->GetBrushComponent())
 		{
-			Bounds += StreamingVol->BrushComponent->BrushBodySetup->AggGeom.CalcAABB(StreamingVol->BrushComponent->ComponentToWorld);
+			Bounds += StreamingVol->GetBrushComponent()->BrushBodySetup->AggGeom.CalcAABB(StreamingVol->GetBrushComponent()->ComponentToWorld);
 		}
 	}
 
diff --git a/Engine/Source/Runtime/Engine/Private/LevelStreamingVolume.cpp b/Engine/Source/Runtime/Engine/Private/LevelStreamingVolume.cpp
--- a/Engine/Source/Runtime/Engine/Private/LevelStreamingVolume.cpp
+++ b/Engine/Source/Runtime/Engine/Private/LevelStreamingVolume.cpp
@@ -6,10 +6,10 @@
 ALevelStreamingVolume::ALevelStreamingVolume(const FObjectInitializer& ObjectInitializer)
 	: Super(ObjectInitializer)
 {
-	BrushComponent->BodyInstance.bEnableCollision_DEPRECATED = false;
+	GetBrushComponent()->BodyInstance.bEnableCollision_DEPRECATED = false;
 
-	BrushComponent->SetCollisionProfileName(UCollisionProfile::NoCollision_ProfileName);
-	BrushComponent->bAlwaysCreatePhysicsState = true;
+	GetBrushComponent()->SetCollisionProfileName(UCollisionProfile::NoCollision_ProfileName);
+	GetBrushComponent()->bAlwaysCreatePhysicsState = true;
 
 	bColored = true;
 	BrushColor.R = 255;
@@ -51,4 +51,4 @@ void ALevelStreamingVolume::UpdateStreamingLevelsRefs()
 	}
 }
 
-#endif// WITH_EDITOR
\ No newline at end of file
+#endif// WITH_EDITOR
diff --git a/Engine/Source/Runtime/Engine/Private/Light.cpp b/Engine/Source/Runtime/Engine/Private/Light.cpp
--- a/Engine/Source/Runtime/Engine/Private/Light.cpp
+++ b/Engine/Source/Runtime/Engine/Private/Light.cpp
@@ -155,7 +155,7 @@ void ALight::SetAffectTranslucentLighting(bool bNewValue)
 APointLight::APointLight(const FObjectInitializer& ObjectInitializer)
 	: Super(ObjectInitializer.SetDefaultSubobjectClass<UPointLightComponent>(TEXT("LightComponent0")))
 {
-	PointLightComponent = CastChecked<UPointLightComponent>(LightComponent);
+	PointLightComponent = CastChecked<UPointLightComponent>(GetLightComponent());
 	PointLightComponent->Mobility = EComponentMobility::Stationary;
 
 	RootComponent = PointLightComponent;
@@ -165,9 +165,9 @@ void APointLight::PostLoad()
 {
 	Super::PostLoad();
 
-	if (LightComponent->Mobility == EComponentMobility::Static)
+	if (GetLightComponent()->Mobility == EComponentMobility::Static)
 	{
-		LightComponent->LightFunctionMaterial = NULL;
+		GetLightComponent()->LightFunctionMaterial = NULL;
 	}
 }
 
@@ -182,19 +182,19 @@ void APointLight::LoadedFromAnotherClass(const FName& OldClassName)
 		static FName PointLightMovable_NAME(TEXT("PointLightMovable"));
 		static FName PointLightStationary_NAME(TEXT("PointLightStationary"));
 
-		check(LightComponent != NULL);
+		check(GetLightComponent() != NULL);
 
 		if(OldClassName == PointLightStatic_NAME)
 		{
-			LightComponent->Mobility = EComponentMobility::Static;
+			GetLightComponent()->Mobility = EComponentMobility::Static;
 		}
 		else if(OldClassName == PointLightMovable_NAME)
 		{
-			LightComponent->Mobility = EComponentMobility::Movable;
+			GetLightComponent()->Mobility = EComponentMobility::Movable;
 		}
 		else if(OldClassName == PointLightStationary_NAME)
 		{
-			LightComponent->Mobility = EComponentMobility::Stationary;
+			GetLightComponent()->Mobility = EComponentMobility::Stationary;
 		}
 	}
 }
@@ -216,7 +216,7 @@ ADirectionalLight::ADirectionalLight(const FObjectInitializer& ObjectInitializer
 	};
 	static FConstructorStatics ConstructorStatics;
 
-	UDirectionalLightComponent* DirectionalLightComponent = CastChecked<UDirectionalLightComponent>(LightComponent);
+	UDirectionalLightComponent* DirectionalLightComponent = CastChecked<UDirectionalLightComponent>(GetLightComponent());
 	DirectionalLightComponent->Mobility = EComponentMobility::Stationary;
 	DirectionalLightComponent->RelativeRotation = FRotator(-46.0f, 0.0f, 0.0f);
 	// Make directional light icons big since they tend to be important
@@ -246,9 +246,9 @@ void ADirectionalLight::PostLoad()
 {
 	Super::PostLoad();
 
-	if (LightComponent->Mobility == EComponentMobility::Static)
+	if (GetLightComponent()->Mobility == EComponentMobility::Static)
 	{
-		LightComponent->LightFunctionMaterial = NULL;
+		GetLightComponent()->LightFunctionMaterial = NULL;
 	}
 #if WITH_EDITORONLY_DATA
 	if(ArrowComponent != nullptr)
@@ -269,7 +269,7 @@ void ADirectionalLight::LoadedFromAnotherClass(const FName& OldClassName)
 		static FName DirectionalLightMovable_NAME(TEXT("DirectionalLightMovable"));
 		static FName DirectionalLightStationary_NAME(TEXT("DirectionalLightStationary"));
 
-		UDirectionalLightComponent* DirLightComp = CastChecked<UDirectionalLightComponent>(LightComponent);
+		UDirectionalLightComponent* DirLightComp = CastChecked<UDirectionalLightComponent>(GetLightComponent());
 
 		if(OldClassName == DirectionalLightStatic_NAME)
 		{
@@ -331,7 +331,7 @@ AGeneratedMeshAreaLight::AGeneratedMeshAreaLight(const FObjectInitializer& Objec
 	bEditable = false;
 #endif // WITH_EDITORONLY_DATA
 
-	LightComponent->CastStaticShadows = false;
+	GetLightComponent()->CastStaticShadows = false;
 }
 
 // EOF
diff --git a/Engine/Source/Runtime/Engine/Private/LightmassCharacterIndirectDetailVolume.cpp b/Engine/Source/Runtime/Engine/Private/LightmassCharacterIndirectDetailVolume.cpp
--- a/Engine/Source/Runtime/Engine/Private/LightmassCharacterIndirectDetailVolume.cpp
+++ b/Engine/Source/Runtime/Engine/Private/LightmassCharacterIndirectDetailVolume.cpp
@@ -7,9 +7,9 @@
 ALightmassCharacterIndirectDetailVolume::ALightmassCharacterIndirectDetailVolume(const FObjectInitializer& ObjectInitializer)
 	: Super(ObjectInitializer)
 {
-	BrushComponent->BodyInstance.bEnableCollision_DEPRECATED = false;
+	GetBrushComponent()->BodyInstance.bEnableCollision_DEPRECATED = false;
 
-	BrushComponent->SetCollisionProfileName(UCollisionProfile::NoCollision_ProfileName);
+	GetBrushComponent()->SetCollisionProfileName(UCollisionProfile::NoCollision_ProfileName);
 
 	bColored = true;
 	BrushColor.R = 155;
diff --git a/Engine/Source/Runtime/Engine/Private/LightmassImportanceVolume.cpp b/Engine/Source/Runtime/Engine/Private/LightmassImportanceVolume.cpp
--- a/Engine/Source/Runtime/Engine/Private/LightmassImportanceVolume.cpp
+++ b/Engine/Source/Runtime/Engine/Private/LightmassImportanceVolume.cpp
@@ -7,9 +7,9 @@
 ALightmassImportanceVolume::ALightmassImportanceVolume(const FObjectInitializer& ObjectInitializer)
 	: Super(ObjectInitializer)
 {
-	BrushComponent->BodyInstance.bEnableCollision_DEPRECATED = false;
+	GetBrushComponent()->BodyInstance.bEnableCollision_DEPRECATED = false;
 
-	BrushComponent->SetCollisionProfileName(UCollisionProfile::NoCollision_ProfileName);
+	GetBrushComponent()->SetCollisionProfileName(UCollisionProfile::NoCollision_ProfileName);
 
 	bColored = true;
 	BrushColor.R = 255;
diff --git a/Engine/Source/Runtime/Engine/Private/Model.cpp b/Engine/Source/Runtime/Engine/Private/Model.cpp
--- a/Engine/Source/Runtime/Engine/Private/Model.cpp
+++ b/Engine/Source/Runtime/Engine/Private/Model.cpp
@@ -516,7 +516,7 @@ UModel::UModel( const class FObjectInitializer& ObjectInitializer,ABrush* Owner,
 	EmptyModel( 1, 1 );
 	if( Owner )
 	{
-		check(Owner->BrushComponent);
+		check(Owner->GetBrushComponent());
 		Owner->Brush = this;
 #if WITH_EDITOR
 		Owner->InitPosRotScale();
@@ -767,4 +767,4 @@ void UModel::ReleaseVertices()
 		VertexBuffer->Vertices.SetAllowCPUAccess(false);
 		VertexBuffer->Vertices.Discard();
 	});
-}
\ No newline at end of file
+}
diff --git a/Engine/Source/Runtime/Engine/Private/NavMeshBoundsVolume.cpp b/Engine/Source/Runtime/Engine/Private/NavMeshBoundsVolume.cpp
--- a/Engine/Source/Runtime/Engine/Private/NavMeshBoundsVolume.cpp
+++ b/Engine/Source/Runtime/Engine/Private/NavMeshBoundsVolume.cpp
@@ -6,10 +6,10 @@
 ANavMeshBoundsVolume::ANavMeshBoundsVolume(const FObjectInitializer& ObjectInitializer)
 	: Super(ObjectInitializer)
 {
-	BrushComponent->BodyInstance.bEnableCollision_DEPRECATED = false;
+	GetBrushComponent()->BodyInstance.bEnableCollision_DEPRECATED = false;
 
-	BrushComponent->SetCollisionProfileName(UCollisionProfile::NoCollision_ProfileName);
-	BrushComponent->Mobility = EComponentMobility::Static;
+	GetBrushComponent()->SetCollisionProfileName(UCollisionProfile::NoCollision_ProfileName);
+	GetBrushComponent()->Mobility = EComponentMobility::Static;
 
 	BrushColor = FColor(200, 200, 200, 255);
 
diff --git a/Engine/Source/Runtime/Engine/Private/NavigationObjectBase.cpp b/Engine/Source/Runtime/Engine/Private/NavigationObjectBase.cpp
--- a/Engine/Source/Runtime/Engine/Private/NavigationObjectBase.cpp
+++ b/Engine/Source/Runtime/Engine/Private/NavigationObjectBase.cpp
@@ -113,7 +113,7 @@ void ANavigationObjectBase::PostEditUndo()
 
 class APhysicsVolume* ANavigationObjectBase::GetNavPhysicsVolume()
 {
-	if ( CapsuleComponent.IsValid() )
+	if ( CapsuleComponent )
 	{
 		return CapsuleComponent->GetPhysicsVolume();
 	}	
@@ -164,11 +164,11 @@ void ANavigationObjectBase::FindBase()
 // 			}
 // 		}
 
-		if (GoodSprite.IsValid())
+		if (GoodSprite)
 		{
 			GoodSprite->SetVisibility(true);
 		}
-		if (BadSprite.IsValid())
+		if (BadSprite)
 		{
 			BadSprite->SetVisibility(false);
 		}
diff --git a/Engine/Source/Runtime/Engine/Private/Note.cpp b/Engine/Source/Runtime/Engine/Private/Note.cpp
--- a/Engine/Source/Runtime/Engine/Private/Note.cpp
+++ b/Engine/Source/Runtime/Engine/Private/Note.cpp
@@ -5,7 +5,7 @@
 ANote::ANote(const FObjectInitializer& ObjectInitializer)
 	: Super(ObjectInitializer)
 {
-	TSubobjectPtr<USceneComponent> SceneComponent = ObjectInitializer.CreateDefaultSubobject<USceneComponent>(this, TEXT("SceneComp"));
+	USceneComponent* SceneComponent = ObjectInitializer.CreateDefaultSubobject<USceneComponent>(this, TEXT("SceneComp"));
 	RootComponent = SceneComponent;
 	RootComponent->Mobility = EComponentMobility::Static;
 
diff --git a/Engine/Source/Runtime/Engine/Private/Particles/ParticleComponents.cpp b/Engine/Source/Runtime/Engine/Private/Particles/ParticleComponents.cpp
--- a/Engine/Source/Runtime/Engine/Private/Particles/ParticleComponents.cpp
+++ b/Engine/Source/Runtime/Engine/Private/Particles/ParticleComponents.cpp
@@ -278,7 +278,7 @@ void AEmitter::CheckForErrors()
 	UObject* Outer = GetOuter();
 	if( Cast<ULevel>( Outer ) )
 	{
-		if ( !ParticleSystemComponent.IsValid() )
+		if ( !ParticleSystemComponent )
 		{
 			FMessageLog("MapCheck").Warning()
 				->AddToken(FUObjectToken::Create(this))
@@ -294,7 +294,7 @@ FString AEmitter::GetDetailedInfoInternal() const
 {
 	FString Result;  
 
-	if( ParticleSystemComponent.IsValid() )
+	if( ParticleSystemComponent )
 	{
 		Result = ParticleSystemComponent->GetDetailedInfoInternal();
 	}
@@ -354,7 +354,7 @@ void AEmitter::PostInitializeComponents()
 	}
 
 	// Set Notification Delegate
-	if (ParticleSystemComponent.IsValid())
+	if (ParticleSystemComponent)
 	{
 		ParticleSystemComponent->OnSystemFinished.BindDynamic( this, &AEmitter::OnParticleSystemFinished );
 		bCurrentlyActive = ParticleSystemComponent->bAutoActivate;
@@ -382,7 +382,7 @@ void AEmitter::OnParticleSystemFinished(UParticleSystemComponent* FinishedCompon
 
 void AEmitter::Activate()
 {
-	if(ParticleSystemComponent.IsValid())
+	if(ParticleSystemComponent)
 	{
 		ParticleSystemComponent->ActivateSystem(false);
 	}
@@ -391,7 +391,7 @@ void AEmitter::Activate()
 
 void AEmitter::Deactivate()
 {
-	if(ParticleSystemComponent.IsValid())
+	if(ParticleSystemComponent)
 	{
 		ParticleSystemComponent->DeactivateSystem();
 	}
@@ -400,7 +400,7 @@ void AEmitter::Deactivate()
 
 void AEmitter::ToggleActive()
 {
-	if(ParticleSystemComponent.IsValid())
+	if(ParticleSystemComponent)
 	{
 		ParticleSystemComponent->ToggleActive();
 		bCurrentlyActive = ParticleSystemComponent->bIsActive;
@@ -409,7 +409,7 @@ void AEmitter::ToggleActive()
 
 bool AEmitter::IsActive() const
 {
-	if(ParticleSystemComponent.IsValid())
+	if(ParticleSystemComponent)
 	{
 		// @todo: I'm not updating bCurrentlyActive flag here. 
 		// Technically I don't have to, and it seems bCurrentlyActive flag can be easily broken if we modify Component directly.
@@ -421,7 +421,7 @@ bool AEmitter::IsActive() const
 
 void AEmitter::SetFloatParameter(FName ParameterName, float Param)
 {
-	if (ParticleSystemComponent.IsValid())
+	if (ParticleSystemComponent)
 	{
 		ParticleSystemComponent->SetFloatParameter(ParameterName, Param);
 	}
@@ -429,7 +429,7 @@ void AEmitter::SetFloatParameter(FName ParameterName, float Param)
 
 void AEmitter::SetVectorParameter(FName ParameterName, FVector Param)
 {
-	if (ParticleSystemComponent.IsValid())
+	if (ParticleSystemComponent)
 	{
 		ParticleSystemComponent->SetVectorParameter(ParameterName, Param);
 	}
@@ -437,7 +437,7 @@ void AEmitter::SetVectorParameter(FName ParameterName, FVector Param)
 
 void AEmitter::SetColorParameter(FName ParameterName, FLinearColor Param)
 {
-	if (ParticleSystemComponent.IsValid())
+	if (ParticleSystemComponent)
 	{
 		ParticleSystemComponent->SetColorParameter(ParameterName, Param);
 	}
@@ -445,7 +445,7 @@ void AEmitter::SetColorParameter(FName ParameterName, FLinearColor Param)
 
 void AEmitter::SetActorParameter(FName ParameterName, AActor* Param)
 {
-	if (ParticleSystemComponent.IsValid())
+	if (ParticleSystemComponent)
 	{
 		ParticleSystemComponent->SetActorParameter(ParameterName, Param);
 	}
@@ -453,7 +453,7 @@ void AEmitter::SetActorParameter(FName ParameterName, AActor* Param)
 
 void AEmitter::SetMaterialParameter(FName ParameterName, UMaterialInterface* Param)
 {
-	if (ParticleSystemComponent.IsValid())
+	if (ParticleSystemComponent)
 	{
 		ParticleSystemComponent->SetMaterialParameter(ParameterName, Param);
 	}
@@ -6009,8 +6009,8 @@ AEmitterCameraLensEffectBase::AEmitterCameraLensEffectBase(const FObjectInitiali
 	BaseFOV = 80.0f;
 	bDestroyOnSystemFinish = true;
 
-	ParticleSystemComponent->bOnlyOwnerSee = true;
-	ParticleSystemComponent->SecondsBeforeInactive = 0.0f;
+	GetParticleSystemComponent()->bOnlyOwnerSee = true;
+	GetParticleSystemComponent()->SecondsBeforeInactive = 0.0f;
 }
 
 void AEmitterCameraLensEffectBase::UpdateLocation(const FVector& CamLoc, const FRotator& CamRot, float CamFOVDeg)
@@ -6050,7 +6050,7 @@ void AEmitterCameraLensEffectBase::NotifyRetriggered() {}
 
 void AEmitterCameraLensEffectBase::PostInitializeComponents()
 {
-	ParticleSystemComponent->SetDepthPriorityGroup(SDPG_Foreground);
+	GetParticleSystemComponent()->SetDepthPriorityGroup(SDPG_Foreground);
 	Super::PostInitializeComponents();
 	ActivateLensEffect();
 }
@@ -6078,4 +6078,4 @@ void AEmitterCameraLensEffectBase::ActivateLensEffect()
 	}
 }
 
-#undef LOCTEXT_NAMESPACE
\ No newline at end of file
+#undef LOCTEXT_NAMESPACE
diff --git a/Engine/Source/Runtime/Engine/Private/Particles/ParticleModules_Location.cpp b/Engine/Source/Runtime/Engine/Private/Particles/ParticleModules_Location.cpp
--- a/Engine/Source/Runtime/Engine/Private/Particles/ParticleModules_Location.cpp
+++ b/Engine/Source/Runtime/Engine/Private/Particles/ParticleModules_Location.cpp
@@ -1909,7 +1909,7 @@ USkeletalMeshComponent* UParticleModuleLocationBoneSocket::GetSkeletalMeshCompon
 		ASkeletalMeshActor* SkelMeshActor = Cast<ASkeletalMeshActor>(Actor);
 		if (SkelMeshActor != NULL)
 		{
-			return SkelMeshActor->SkeletalMeshComponent;
+			return SkelMeshActor->GetSkeletalMeshComponent();
 		}
 		else if (Actor)
 		{
@@ -2356,12 +2356,12 @@ void UParticleModuleLocationSkelVertSurface::UpdateBoneIndicesList(FParticleEmit
 
 		if ( SkeletalMeshActor != NULL )
 		{
-			if ( SkeletalMeshActor->SkeletalMeshComponent.IsValid() && (SkeletalMeshActor->SkeletalMeshComponent->SkeletalMesh != NULL) )
+			if (SkeletalMeshActor->GetSkeletalMeshComponent() && (SkeletalMeshActor->GetSkeletalMeshComponent()->SkeletalMesh != NULL))
 			{
 				int32 InsertionIndex = 0;
 				for (int32 FindBoneIdx = 0; FindBoneIdx < ValidAssociatedBones.Num(); FindBoneIdx++)
 				{
-					const int32 BoneIdx = SkeletalMeshActor->SkeletalMeshComponent->SkeletalMesh->RefSkeleton.FindBoneIndex(ValidAssociatedBones[FindBoneIdx]);
+					const int32 BoneIdx = SkeletalMeshActor->GetSkeletalMeshComponent()->SkeletalMesh->RefSkeleton.FindBoneIndex(ValidAssociatedBones[FindBoneIdx]);
 					if (BoneIdx != INDEX_NONE && ValidAssociatedBones.Num() > InsertionIndex)
 					{
 						InstancePayload->ValidAssociatedBoneIndices[InsertionIndex++] = BoneIdx;
@@ -2537,7 +2537,7 @@ USkeletalMeshComponent* UParticleModuleLocationSkelVertSurface::GetSkeletalMeshC
 		ASkeletalMeshActor* SkelMeshActor = Cast<ASkeletalMeshActor>(Actor);
 		if (SkelMeshActor != NULL)
 		{
-			return SkelMeshActor->SkeletalMeshComponent;
+			return SkelMeshActor->GetSkeletalMeshComponent();
 		}
 		else
 		{
diff --git a/Engine/Source/Runtime/Engine/Private/PhysicsEngine/PhysicsConstraintActor.cpp b/Engine/Source/Runtime/Engine/Private/PhysicsEngine/PhysicsConstraintActor.cpp
--- a/Engine/Source/Runtime/Engine/Private/PhysicsEngine/PhysicsConstraintActor.cpp
+++ b/Engine/Source/Runtime/Engine/Private/PhysicsEngine/PhysicsConstraintActor.cpp
@@ -29,7 +29,7 @@ void APhysicsConstraintActor::PostLoad()
 	Super::PostLoad();
 
 	// Copy 'actors to constrain' into component
-	if (GetLinkerUE4Version() < VER_UE4_ALL_PROPS_TO_CONSTRAINTINSTANCE && ConstraintComp.Get() != NULL)
+	if (GetLinkerUE4Version() < VER_UE4_ALL_PROPS_TO_CONSTRAINTINSTANCE && ConstraintComp != NULL)
 	{
 		ConstraintComp->ConstraintActor1 = ConstraintActor1_DEPRECATED;
 		ConstraintComp->ConstraintActor2 = ConstraintActor2_DEPRECATED;
diff --git a/Engine/Source/Runtime/Engine/Private/PhysicsEngine/RadialForceComponent.cpp b/Engine/Source/Runtime/Engine/Private/PhysicsEngine/RadialForceComponent.cpp
--- a/Engine/Source/Runtime/Engine/Private/PhysicsEngine/RadialForceComponent.cpp
+++ b/Engine/Source/Runtime/Engine/Private/PhysicsEngine/RadialForceComponent.cpp
@@ -227,7 +227,7 @@ void ARadialForceActor::EditorApplyScale(const FVector& DeltaScale, const FVecto
 	FVector ModifiedScale = DeltaScale * ( AActor::bUsePercentageBasedScaling ? 500.0f : 5.0f );
 
 	const float Multiplier = ( ModifiedScale.X > 0.0f || ModifiedScale.Y > 0.0f || ModifiedScale.Z > 0.0f ) ? 1.0f : -1.0f;
-	if(ForceComponent.IsValid())
+	if(ForceComponent)
 	{
 		ForceComponent->Radius += Multiplier * ModifiedScale.Size();
 		ForceComponent->Radius = FMath::Max( 0.f, ForceComponent->Radius );
@@ -237,7 +237,7 @@ void ARadialForceActor::EditorApplyScale(const FVector& DeltaScale, const FVecto
 
 void ARadialForceActor::FireImpulse()
 {
-	if(ForceComponent.IsValid())
+	if(ForceComponent)
 	{
 		ForceComponent->FireImpulse();
 	}
@@ -245,7 +245,7 @@ void ARadialForceActor::FireImpulse()
 
 void ARadialForceActor::EnableForce()
 {
-	if(ForceComponent.IsValid())
+	if(ForceComponent)
 	{
 		ForceComponent->Activate();
 	}
@@ -253,7 +253,7 @@ void ARadialForceActor::EnableForce()
 
 void ARadialForceActor::DisableForce()
 {
-	if(ForceComponent.IsValid())
+	if(ForceComponent)
 	{
 		ForceComponent->Deactivate();
 	}
@@ -261,7 +261,7 @@ void ARadialForceActor::DisableForce()
 
 void ARadialForceActor::ToggleForce()
 {
-	if(ForceComponent.IsValid())
+	if(ForceComponent)
 	{
 		ForceComponent->ToggleActive();
 	}
diff --git a/Engine/Source/Runtime/Engine/Private/PhysicsVolume.cpp b/Engine/Source/Runtime/Engine/Private/PhysicsVolume.cpp
--- a/Engine/Source/Runtime/Engine/Private/PhysicsVolume.cpp
+++ b/Engine/Source/Runtime/Engine/Private/PhysicsVolume.cpp
@@ -6,9 +6,9 @@
 APhysicsVolume::APhysicsVolume(const FObjectInitializer& ObjectInitializer)
 	: Super(ObjectInitializer)
 {
-	BrushComponent->BodyInstance.bEnableCollision_DEPRECATED = true;
+	GetBrushComponent()->BodyInstance.bEnableCollision_DEPRECATED = true;
 	static FName CollisionProfileName(TEXT("OverlapAllDynamic"));
-	BrushComponent->SetCollisionProfileName(CollisionProfileName);
+	GetBrushComponent()->SetCollisionProfileName(CollisionProfileName);
 
 	FluidFriction = UPhysicsSettings::Get()->DefaultFluidFriction;
 	TerminalVelocity = UPhysicsSettings::Get()->DefaultTerminalVelocity;
@@ -28,7 +28,7 @@ void APhysicsVolume::LoadedFromAnotherClass(const FName& OldClassName)
 
 		if(OldClassName == DynamicPhysicsVolume_NAME)
 		{
-			BrushComponent->Mobility = EComponentMobility::Movable;
+			GetBrushComponent()->Mobility = EComponentMobility::Movable;
 		}
 	}
 }
diff --git a/Engine/Source/Runtime/Engine/Private/PlayerCameraManager.cpp b/Engine/Source/Runtime/Engine/Private/PlayerCameraManager.cpp
--- a/Engine/Source/Runtime/Engine/Private/PlayerCameraManager.cpp
+++ b/Engine/Source/Runtime/Engine/Private/PlayerCameraManager.cpp
@@ -329,13 +329,13 @@ void APlayerCameraManager::ApplyAnimToCamera(ACameraActor const* AnimatedCamActo
 	// fov
 	const float FOVMin = 5.f;
 	const float FOVMax = 170.f;
-	InOutPOV.FOV += (AnimatedCamActor->CameraComponent->FieldOfView - AnimInst->InitialFOV) * Scale;
+	InOutPOV.FOV += (AnimatedCamActor->GetCameraComponent()->FieldOfView - AnimInst->InitialFOV) * Scale;
 	InOutPOV.FOV = FMath::Clamp<float>(InOutPOV.FOV, FOVMin, FOVMax);
 
 	// postprocess
-	if (AnimatedCamActor->CameraComponent.Get()->PostProcessBlendWeight > 0.f)
+	if (AnimatedCamActor->GetCameraComponent()->PostProcessBlendWeight > 0.f)
 	{
-		AddCachedPPBlend(AnimatedCamActor->CameraComponent.Get()->PostProcessSettings, AnimatedCamActor->CameraComponent.Get()->PostProcessBlendWeight);
+		AddCachedPPBlend(AnimatedCamActor->GetCameraComponent()->PostProcessSettings, AnimatedCamActor->GetCameraComponent()->PostProcessBlendWeight);
 	}
 }
 
@@ -442,10 +442,10 @@ void APlayerCameraManager::InitTempCameraActor(ACameraActor* CamActor, UCameraAn
 			ACameraActor const* const DefaultCamActor = GetDefault<ACameraActor>();
 			if (DefaultCamActor)
 			{
-				CamActor->CameraComponent->AspectRatio = DefaultCamActor->CameraComponent->AspectRatio;
-				CamActor->CameraComponent->FieldOfView = AnimInstToInitFor->CamAnim->BaseFOV;
-				CamActor->CameraComponent->PostProcessSettings = AnimInstToInitFor->CamAnim->BasePostProcessSettings;
-				CamActor->CameraComponent->PostProcessBlendWeight = AnimInstToInitFor->CamAnim->BasePostProcessBlendWeight;
+				CamActor->GetCameraComponent()->AspectRatio = DefaultCamActor->GetCameraComponent()->AspectRatio;
+				CamActor->GetCameraComponent()->FieldOfView = AnimInstToInitFor->CamAnim->BaseFOV;
+				CamActor->GetCameraComponent()->PostProcessSettings = AnimInstToInitFor->CamAnim->BasePostProcessSettings;
+				CamActor->GetCameraComponent()->PostProcessBlendWeight = AnimInstToInitFor->CamAnim->BasePostProcessBlendWeight;
 			}
 		}
 	}
@@ -487,7 +487,7 @@ void APlayerCameraManager::UpdateViewTarget(FTViewTarget& OutVT, float DeltaTime
 	if (ACameraActor* CamActor = Cast<ACameraActor>(OutVT.Target))
 	{
 		// Viewing through a camera actor.
-		CamActor->CameraComponent->GetCameraView(DeltaTime, OutVT.POV);
+		CamActor->GetCameraComponent()->GetCameraView(DeltaTime, OutVT.POV);
 	}
 	else
 	{
diff --git a/Engine/Source/Runtime/Engine/Private/PlayerStart.cpp b/Engine/Source/Runtime/Engine/Private/PlayerStart.cpp
--- a/Engine/Source/Runtime/Engine/Private/PlayerStart.cpp
+++ b/Engine/Source/Runtime/Engine/Private/PlayerStart.cpp
@@ -5,7 +5,7 @@
 APlayerStart::APlayerStart(const FObjectInitializer& ObjectInitializer)
 	: Super(ObjectInitializer)
 {
-	CapsuleComponent->InitCapsuleSize(40.0f, 92.0f);
+	GetCapsuleComponent()->InitCapsuleSize(40.0f, 92.0f);
 
 #if WITH_EDITORONLY_DATA
 	ArrowComponent = ObjectInitializer.CreateEditorOnlyDefaultSubobject<UArrowComponent>(this, TEXT("Arrow"));
@@ -31,16 +31,16 @@ APlayerStart::APlayerStart(const FObjectInitializer& ObjectInitializer)
 		};
 		static FConstructorStatics ConstructorStatics;
 
-		if (GoodSprite)
+		if (GetGoodSprite())
 		{
-			GoodSprite->Sprite = ConstructorStatics.PlayerStartTextureObject.Get();
-			GoodSprite->RelativeScale3D = FVector(0.5f, 0.5f, 0.5f);
-			GoodSprite->SpriteInfo.Category = ConstructorStatics.ID_PlayerStart;
-			GoodSprite->SpriteInfo.DisplayName = ConstructorStatics.NAME_PlayerStart;
+			GetGoodSprite()->Sprite = ConstructorStatics.PlayerStartTextureObject.Get();
+			GetGoodSprite()->RelativeScale3D = FVector(0.5f, 0.5f, 0.5f);
+			GetGoodSprite()->SpriteInfo.Category = ConstructorStatics.ID_PlayerStart;
+			GetGoodSprite()->SpriteInfo.DisplayName = ConstructorStatics.NAME_PlayerStart;
 		}
-		if (BadSprite)
+		if (GetBadSprite())
 		{
-			BadSprite->SetVisibility(false);
+			GetBadSprite()->SetVisibility(false);
 		}
 
 		if (ArrowComponent)
@@ -51,7 +51,7 @@ APlayerStart::APlayerStart(const FObjectInitializer& ObjectInitializer)
 			ArrowComponent->bTreatAsASprite = true;
 			ArrowComponent->SpriteInfo.Category = ConstructorStatics.ID_Navigation;
 			ArrowComponent->SpriteInfo.DisplayName = ConstructorStatics.NAME_Navigation;
-			ArrowComponent->AttachParent = CapsuleComponent;
+			ArrowComponent->AttachParent = GetCapsuleComponent();
 			ArrowComponent->bIsScreenSizeScaled = true;
 		}
 	}
diff --git a/Engine/Source/Runtime/Engine/Private/PostProcessVolume.cpp b/Engine/Source/Runtime/Engine/Private/PostProcessVolume.cpp
--- a/Engine/Source/Runtime/Engine/Private/PostProcessVolume.cpp
+++ b/Engine/Source/Runtime/Engine/Private/PostProcessVolume.cpp
@@ -6,11 +6,11 @@
 APostProcessVolume::APostProcessVolume(const FObjectInitializer& ObjectInitializer)
 	: Super(ObjectInitializer)
 {
-	BrushComponent->BodyInstance.bEnableCollision_DEPRECATED = false;
-	BrushComponent->SetCollisionProfileName(UCollisionProfile::NoCollision_ProfileName);
+	GetBrushComponent()->BodyInstance.bEnableCollision_DEPRECATED = false;
+	GetBrushComponent()->SetCollisionProfileName(UCollisionProfile::NoCollision_ProfileName);
 	// post process volume needs physics data for trace
-	BrushComponent->bAlwaysCreatePhysicsState = true;
-	BrushComponent->Mobility = EComponentMobility::Movable;
+	GetBrushComponent()->bAlwaysCreatePhysicsState = true;
+	GetBrushComponent()->Mobility = EComponentMobility::Movable;
 	
 	bEnabled = true;
 	BlendRadius = 100.0f;
@@ -75,4 +75,4 @@ bool APostProcessVolume::CanEditChange(const UProperty* InProperty) const
 	return Super::CanEditChange(InProperty);
 }
 
-#endif // WITH_EDITOR
\ No newline at end of file
+#endif // WITH_EDITOR
diff --git a/Engine/Source/Runtime/Engine/Private/PrecomputedVisibilityVolume.cpp b/Engine/Source/Runtime/Engine/Private/PrecomputedVisibilityVolume.cpp
--- a/Engine/Source/Runtime/Engine/Private/PrecomputedVisibilityVolume.cpp
+++ b/Engine/Source/Runtime/Engine/Private/PrecomputedVisibilityVolume.cpp
@@ -6,8 +6,8 @@
 APrecomputedVisibilityVolume::APrecomputedVisibilityVolume(const FObjectInitializer& ObjectInitializer)
 	: Super(ObjectInitializer)
 {
-	BrushComponent->BodyInstance.bEnableCollision_DEPRECATED = false;
-	BrushComponent->SetCollisionProfileName(UCollisionProfile::NoCollision_ProfileName);
+	GetBrushComponent()->BodyInstance.bEnableCollision_DEPRECATED = false;
+	GetBrushComponent()->SetCollisionProfileName(UCollisionProfile::NoCollision_ProfileName);
 	bColored = true;
 	BrushColor.R = 25;
 	BrushColor.G = 255;
diff --git a/Engine/Source/Runtime/Engine/Private/PrecomputedVisiblityOverrideVolume.cpp b/Engine/Source/Runtime/Engine/Private/PrecomputedVisiblityOverrideVolume.cpp
--- a/Engine/Source/Runtime/Engine/Private/PrecomputedVisiblityOverrideVolume.cpp
+++ b/Engine/Source/Runtime/Engine/Private/PrecomputedVisiblityOverrideVolume.cpp
@@ -6,8 +6,8 @@
 APrecomputedVisibilityOverrideVolume::APrecomputedVisibilityOverrideVolume(const FObjectInitializer& ObjectInitializer)
 	: Super(ObjectInitializer)
 {
-	BrushComponent->BodyInstance.bEnableCollision_DEPRECATED = false;
-	BrushComponent->SetCollisionProfileName(UCollisionProfile::NoCollision_ProfileName);
+	GetBrushComponent()->BodyInstance.bEnableCollision_DEPRECATED = false;
+	GetBrushComponent()->SetCollisionProfileName(UCollisionProfile::NoCollision_ProfileName);
 
 	bColored = true;
 	BrushColor.R = 25;
diff --git a/Engine/Source/Runtime/Engine/Private/ReflectionCaptureComponent.cpp b/Engine/Source/Runtime/Engine/Private/ReflectionCaptureComponent.cpp
--- a/Engine/Source/Runtime/Engine/Private/ReflectionCaptureComponent.cpp
+++ b/Engine/Source/Runtime/Engine/Private/ReflectionCaptureComponent.cpp
@@ -88,17 +88,17 @@ void AReflectionCapture::PostEditMove(bool bFinished)
 ASphereReflectionCapture::ASphereReflectionCapture(const FObjectInitializer& ObjectInitializer)
 	: Super(ObjectInitializer.SetDefaultSubobjectClass<USphereReflectionCaptureComponent>(TEXT("NewReflectionComponent")))
 {
-	USphereReflectionCaptureComponent* SphereComponent = CastChecked<USphereReflectionCaptureComponent>(CaptureComponent);
+	USphereReflectionCaptureComponent* SphereComponent = CastChecked<USphereReflectionCaptureComponent>(GetCaptureComponent());
 	RootComponent = SphereComponent;
 #if WITH_EDITORONLY_DATA
-	if (SpriteComponent)
+	if (GetSpriteComponent())
 	{
-		SpriteComponent->AttachParent = SphereComponent;
+		GetSpriteComponent()->AttachParent = SphereComponent;
 	}
 #endif	//WITH_EDITORONLY_DATA
 
-	TSubobjectPtr<UDrawSphereComponent> DrawInfluenceRadius = ObjectInitializer.CreateDefaultSubobject<UDrawSphereComponent>(this, TEXT("DrawRadius0"));
-	DrawInfluenceRadius->AttachParent = CaptureComponent;
+	UDrawSphereComponent* DrawInfluenceRadius = ObjectInitializer.CreateDefaultSubobject<UDrawSphereComponent>(this, TEXT("DrawRadius0"));
+	DrawInfluenceRadius->AttachParent = GetCaptureComponent();
 	DrawInfluenceRadius->bDrawOnlyIfSelected = true;
 	DrawInfluenceRadius->bUseEditorCompositing = true;
 	DrawInfluenceRadius->BodyInstance.bEnableCollision_DEPRECATED = false;
@@ -106,7 +106,7 @@ ASphereReflectionCapture::ASphereReflectionCapture(const FObjectInitializer& Obj
 	SphereComponent->PreviewInfluenceRadius = DrawInfluenceRadius;
 
 	DrawCaptureRadius = ObjectInitializer.CreateDefaultSubobject<UDrawSphereComponent>(this, TEXT("DrawRadius1"));
-	DrawCaptureRadius->AttachParent = CaptureComponent;
+	DrawCaptureRadius->AttachParent = GetCaptureComponent();
 	DrawCaptureRadius->bDrawOnlyIfSelected = true;
 	DrawCaptureRadius->bUseEditorCompositing = true;
 	DrawCaptureRadius->BodyInstance.bEnableCollision_DEPRECATED = false;
@@ -117,21 +117,21 @@ ASphereReflectionCapture::ASphereReflectionCapture(const FObjectInitializer& Obj
 #if WITH_EDITOR
 void ASphereReflectionCapture::EditorApplyScale(const FVector& DeltaScale, const FVector* PivotLocation, bool bAltDown, bool bShiftDown, bool bCtrlDown)
 {
-	USphereReflectionCaptureComponent* SphereComponent = Cast<USphereReflectionCaptureComponent>(CaptureComponent);
+	USphereReflectionCaptureComponent* SphereComponent = Cast<USphereReflectionCaptureComponent>(GetCaptureComponent());
 	check(SphereComponent);
 	const FVector ModifiedScale = DeltaScale * ( AActor::bUsePercentageBasedScaling ? 5000.0f : 50.0f );
 	FMath::ApplyScaleToFloat(SphereComponent->InfluenceRadius, ModifiedScale);
-	CaptureComponent->SetCaptureIsDirty();
+	GetCaptureComponent()->SetCaptureIsDirty();
 	PostEditChange();
 }
 
 void APlaneReflectionCapture::EditorApplyScale(const FVector& DeltaScale, const FVector* PivotLocation, bool bAltDown, bool bShiftDown, bool bCtrlDown)
 {
-	UPlaneReflectionCaptureComponent* PlaneComponent = Cast<UPlaneReflectionCaptureComponent>(CaptureComponent);
+	UPlaneReflectionCaptureComponent* PlaneComponent = Cast<UPlaneReflectionCaptureComponent>(GetCaptureComponent());
 	check(PlaneComponent);
 	const FVector ModifiedScale = DeltaScale * ( AActor::bUsePercentageBasedScaling ? 5000.0f : 50.0f );
 	FMath::ApplyScaleToFloat(PlaneComponent->InfluenceRadiusScale, ModifiedScale);
-	CaptureComponent->SetCaptureIsDirty();
+	GetCaptureComponent()->SetCaptureIsDirty();
 	PostEditChange();
 }
 #endif
@@ -139,17 +139,17 @@ void APlaneReflectionCapture::EditorApplyScale(const FVector& DeltaScale, const
 ABoxReflectionCapture::ABoxReflectionCapture(const FObjectInitializer& ObjectInitializer)
 	: Super(ObjectInitializer.SetDefaultSubobjectClass<UBoxReflectionCaptureComponent>(TEXT("NewReflectionComponent")))
 {
-	UBoxReflectionCaptureComponent* BoxComponent = CastChecked<UBoxReflectionCaptureComponent>(CaptureComponent);
+	UBoxReflectionCaptureComponent* BoxComponent = CastChecked<UBoxReflectionCaptureComponent>(GetCaptureComponent());
 	BoxComponent->RelativeScale3D = FVector(1000, 1000, 400);
 	RootComponent = BoxComponent;
 #if WITH_EDITORONLY_DATA
-	if (SpriteComponent)
+	if (GetSpriteComponent())
 	{
-		SpriteComponent->AttachParent = BoxComponent;
+		GetSpriteComponent()->AttachParent = BoxComponent;
 	}
 #endif	//WITH_EDITORONLY_DATA
-	TSubobjectPtr<UBoxComponent> DrawInfluenceBox = ObjectInitializer.CreateDefaultSubobject<UBoxComponent>(this, TEXT("DrawBox0"));
-	DrawInfluenceBox->AttachParent = CaptureComponent;
+	UBoxComponent* DrawInfluenceBox = ObjectInitializer.CreateDefaultSubobject<UBoxComponent>(this, TEXT("DrawBox0"));
+	DrawInfluenceBox->AttachParent = GetCaptureComponent();
 	DrawInfluenceBox->bDrawOnlyIfSelected = true;
 	DrawInfluenceBox->bUseEditorCompositing = true;
 	DrawInfluenceBox->BodyInstance.bEnableCollision_DEPRECATED = false;
@@ -157,8 +157,8 @@ ABoxReflectionCapture::ABoxReflectionCapture(const FObjectInitializer& ObjectIni
 	DrawInfluenceBox->InitBoxExtent(FVector(1, 1, 1));
 	BoxComponent->PreviewInfluenceBox = DrawInfluenceBox;
 
-	TSubobjectPtr<UBoxComponent> DrawCaptureBox = ObjectInitializer.CreateDefaultSubobject<UBoxComponent>(this, TEXT("DrawBox1"));
-	DrawCaptureBox->AttachParent = CaptureComponent;
+	UBoxComponent* DrawCaptureBox = ObjectInitializer.CreateDefaultSubobject<UBoxComponent>(this, TEXT("DrawBox1"));
+	DrawCaptureBox->AttachParent = GetCaptureComponent();
 	DrawCaptureBox->bDrawOnlyIfSelected = true;
 	DrawCaptureBox->bUseEditorCompositing = true;
 	DrawCaptureBox->BodyInstance.bEnableCollision_DEPRECATED = false;
@@ -171,17 +171,17 @@ ABoxReflectionCapture::ABoxReflectionCapture(const FObjectInitializer& ObjectIni
 APlaneReflectionCapture::APlaneReflectionCapture(const FObjectInitializer& ObjectInitializer)
 	: Super(ObjectInitializer.SetDefaultSubobjectClass<UPlaneReflectionCaptureComponent>(TEXT("NewReflectionComponent")))
 {
-	UPlaneReflectionCaptureComponent* PlaneComponent = CastChecked<UPlaneReflectionCaptureComponent>(CaptureComponent);
+	UPlaneReflectionCaptureComponent* PlaneComponent = CastChecked<UPlaneReflectionCaptureComponent>(GetCaptureComponent());
 	PlaneComponent->RelativeScale3D = FVector(1, 1000, 1000);
 	RootComponent = PlaneComponent;
 #if WITH_EDITORONLY_DATA
-	if (SpriteComponent)
+	if (GetSpriteComponent())
 	{
-		SpriteComponent->AttachParent = PlaneComponent;
+		GetSpriteComponent()->AttachParent = PlaneComponent;
 	}
 #endif	//#if WITH_EDITORONLY_DATA
-	TSubobjectPtr<UDrawSphereComponent> DrawInfluenceRadius = ObjectInitializer.CreateDefaultSubobject<UDrawSphereComponent>(this, TEXT("DrawRadius0"));
-	DrawInfluenceRadius->AttachParent = CaptureComponent;
+	UDrawSphereComponent* DrawInfluenceRadius = ObjectInitializer.CreateDefaultSubobject<UDrawSphereComponent>(this, TEXT("DrawRadius0"));
+	DrawInfluenceRadius->AttachParent = GetCaptureComponent();
 	DrawInfluenceRadius->bDrawOnlyIfSelected = true;
 	DrawInfluenceRadius->bAbsoluteScale = true;
 	DrawInfluenceRadius->bUseEditorCompositing = true;
@@ -189,8 +189,8 @@ APlaneReflectionCapture::APlaneReflectionCapture(const FObjectInitializer& Objec
 	DrawInfluenceRadius->SetCollisionProfileName(UCollisionProfile::NoCollision_ProfileName);
 	PlaneComponent->PreviewInfluenceRadius = DrawInfluenceRadius;
 
-	TSubobjectPtr<UBoxComponent> DrawCaptureBox = ObjectInitializer.CreateDefaultSubobject<UBoxComponent>(this, TEXT("DrawBox1"));
-	DrawCaptureBox->AttachParent = CaptureComponent;
+	UBoxComponent* DrawCaptureBox = ObjectInitializer.CreateDefaultSubobject<UBoxComponent>(this, TEXT("DrawBox1"));
+	DrawCaptureBox->AttachParent = GetCaptureComponent();
 	DrawCaptureBox->bDrawOnlyIfSelected = true;
 	DrawCaptureBox->bUseEditorCompositing = true;
 	DrawCaptureBox->BodyInstance.bEnableCollision_DEPRECATED = false;
diff --git a/Engine/Source/Runtime/Engine/Private/SceneCaptureComponent.cpp b/Engine/Source/Runtime/Engine/Private/SceneCaptureComponent.cpp
--- a/Engine/Source/Runtime/Engine/Private/SceneCaptureComponent.cpp
+++ b/Engine/Source/Runtime/Engine/Private/SceneCaptureComponent.cpp
@@ -28,10 +28,10 @@ ASceneCapture2D::ASceneCapture2D(const FObjectInitializer& ObjectInitializer)
 	DrawFrustum = ObjectInitializer.CreateDefaultSubobject<UDrawFrustumComponent>(this, TEXT("DrawFrust0"));
 	DrawFrustum->AlwaysLoadOnClient = false;
 	DrawFrustum->AlwaysLoadOnServer = false;
-	DrawFrustum->AttachParent = MeshComp;
+	DrawFrustum->AttachParent = GetMeshComp();
 
 	CaptureComponent2D = ObjectInitializer.CreateDefaultSubobject<USceneCaptureComponent2D>(this, TEXT("NewSceneCaptureComponent2D"));
-	CaptureComponent2D->AttachParent = MeshComp;
+	CaptureComponent2D->AttachParent = GetMeshComp();
 }
 
 void ASceneCapture2D::OnInterpToggle(bool bEnable)
@@ -60,14 +60,14 @@ void ASceneCapture2D::PostActorCreated()
 
 	// no need load the editor mesh when there is no editor
 #if WITH_EDITOR
-	if(MeshComp)
+	if(GetMeshComp())
 	{
 		if (!IsRunningCommandlet())
 		{
-			if( !MeshComp->StaticMesh)
+			if( !GetMeshComp()->StaticMesh)
 			{
 				UStaticMesh* CamMesh = LoadObject<UStaticMesh>(NULL, TEXT("/Engine/EditorMeshes/MatineeCam_SM.MatineeCam_SM"), NULL, LOAD_None, NULL);
-				MeshComp->SetStaticMesh(CamMesh);
+				GetMeshComp()->SetStaticMesh(CamMesh);
 			}
 		}
 	}
@@ -84,10 +84,10 @@ ASceneCaptureCube::ASceneCaptureCube(const FObjectInitializer& ObjectInitializer
 	DrawFrustum = ObjectInitializer.CreateDefaultSubobject<UDrawFrustumComponent>(this, TEXT("DrawFrust0"));
 	DrawFrustum->AlwaysLoadOnClient = false;
 	DrawFrustum->AlwaysLoadOnServer = false;
-	DrawFrustum->AttachParent = MeshComp;
+	DrawFrustum->AttachParent = GetMeshComp();
 
 	CaptureComponentCube = ObjectInitializer.CreateDefaultSubobject<USceneCaptureComponentCube>(this, TEXT("NewSceneCaptureComponentCube"));
-	CaptureComponentCube->AttachParent = MeshComp;
+	CaptureComponentCube->AttachParent = GetMeshComp();
 }
 
 void ASceneCaptureCube::OnInterpToggle(bool bEnable)
@@ -115,14 +115,14 @@ void ASceneCaptureCube::PostActorCreated()
 
 	// no need load the editor mesh when there is no editor
 #if WITH_EDITOR
-	if(MeshComp)
+	if(GetMeshComp())
 	{
 		if (!IsRunningCommandlet())
 		{
-			if( !MeshComp->StaticMesh)
+			if( !GetMeshComp()->StaticMesh)
 			{
 				UStaticMesh* CamMesh = LoadObject<UStaticMesh>(NULL, TEXT("/Engine/EditorMeshes/MatineeCam_SM.MatineeCam_SM"), NULL, LOAD_None, NULL);
-				MeshComp->SetStaticMesh(CamMesh);
+				GetMeshComp()->SetStaticMesh(CamMesh);
 			}
 		}
 	}
diff --git a/Engine/Source/Runtime/Engine/Private/SkyLightComponent.cpp b/Engine/Source/Runtime/Engine/Private/SkyLightComponent.cpp
--- a/Engine/Source/Runtime/Engine/Private/SkyLightComponent.cpp
+++ b/Engine/Source/Runtime/Engine/Private/SkyLightComponent.cpp
@@ -496,12 +496,12 @@ ASkyLight::ASkyLight(const FObjectInitializer& ObjectInitializer)
 	};
 	static FConstructorStatics ConstructorStatics;
 
-		if (SpriteComponent)
+		if (GetSpriteComponent())
 		{
-			SpriteComponent->Sprite = ConstructorStatics.SkyLightTextureObject.Get();
-			SpriteComponent->SpriteInfo.Category = ConstructorStatics.ID_Sky;
-			SpriteComponent->SpriteInfo.DisplayName = ConstructorStatics.NAME_Sky;
-			SpriteComponent->AttachParent = LightComponent;
+			GetSpriteComponent()->Sprite = ConstructorStatics.SkyLightTextureObject.Get();
+			GetSpriteComponent()->SpriteInfo.Category = ConstructorStatics.ID_Sky;
+			GetSpriteComponent()->SpriteInfo.DisplayName = ConstructorStatics.NAME_Sky;
+			GetSpriteComponent()->AttachParent = LightComponent;
 		}
 	}
 #endif // WITH_EDITORONLY_DATA
diff --git a/Engine/Source/Runtime/Engine/Private/SpectatorPawn.cpp b/Engine/Source/Runtime/Engine/Private/SpectatorPawn.cpp
--- a/Engine/Source/Runtime/Engine/Private/SpectatorPawn.cpp
+++ b/Engine/Source/Runtime/Engine/Private/SpectatorPawn.cpp
@@ -19,7 +19,7 @@ ASpectatorPawn::ASpectatorPawn(const FObjectInitializer& ObjectInitializer)
 	bCollideWhenPlacing = false;
 
 	static FName CollisionProfileName(TEXT("Spectator"));
-	CollisionComponent->SetCollisionProfileName(CollisionProfileName);
+	GetCollisionComponent()->SetCollisionProfileName(CollisionProfileName);
 }
 
 
diff --git a/Engine/Source/Runtime/Engine/Private/SpotLight.cpp b/Engine/Source/Runtime/Engine/Private/SpotLight.cpp
--- a/Engine/Source/Runtime/Engine/Private/SpotLight.cpp
+++ b/Engine/Source/Runtime/Engine/Private/SpotLight.cpp
@@ -18,7 +18,7 @@ ASpotLight::ASpotLight(const FObjectInitializer& ObjectInitializer)
 	};
 	static FConstructorStatics ConstructorStatics;
 
-	SpotLightComponent = CastChecked<USpotLightComponent>(LightComponent);
+	SpotLightComponent = CastChecked<USpotLightComponent>(GetLightComponent());
 	SpotLightComponent->Mobility = EComponentMobility::Stationary;
 	SpotLightComponent->RelativeRotation = FRotator(-90, 0, 0);
 	SpotLightComponent->UpdateComponentToWorld();
@@ -44,9 +44,9 @@ void ASpotLight::PostLoad()
 {
 	Super::PostLoad();
 
-	if (LightComponent->Mobility == EComponentMobility::Static)
+	if (GetLightComponent()->Mobility == EComponentMobility::Static)
 	{
-		LightComponent->LightFunctionMaterial = NULL;
+		GetLightComponent()->LightFunctionMaterial = NULL;
 	}
 
 #if WITH_EDITORONLY_DATA
@@ -68,19 +68,19 @@ void ASpotLight::LoadedFromAnotherClass(const FName& OldClassName)
 		static FName SpotLightMovable_NAME(TEXT("SpotLightMovable"));
 		static FName SpotLightStationary_NAME(TEXT("SpotLightStationary"));
 
-		check(LightComponent != NULL);
+		check(GetLightComponent() != NULL);
 
 		if(OldClassName == SpotLightStatic_NAME)
 		{
-			LightComponent->Mobility = EComponentMobility::Static;
+			GetLightComponent()->Mobility = EComponentMobility::Static;
 		}
 		else if(OldClassName == SpotLightMovable_NAME)
 		{
-			LightComponent->Mobility = EComponentMobility::Movable;
+			GetLightComponent()->Mobility = EComponentMobility::Movable;
 		}
 		else if(OldClassName == SpotLightStationary_NAME)
 		{
-			LightComponent->Mobility = EComponentMobility::Stationary;
+			GetLightComponent()->Mobility = EComponentMobility::Stationary;
 		}
 	}
 }
diff --git a/Engine/Source/Runtime/Engine/Private/StaticMeshActor.cpp b/Engine/Source/Runtime/Engine/Private/StaticMeshActor.cpp
--- a/Engine/Source/Runtime/Engine/Private/StaticMeshActor.cpp
+++ b/Engine/Source/Runtime/Engine/Private/StaticMeshActor.cpp
@@ -124,7 +124,7 @@ void AStaticMeshActor::CheckForErrors()
 
 	FMessageLog MapCheck("MapCheck");
 
-	if( !StaticMeshComponent.IsValid() )
+	if( !StaticMeshComponent )
 	{
 		MapCheck.Warning()
 			->AddToken(FUObjectToken::Create(this))
@@ -151,7 +151,7 @@ void AStaticMeshActor::CheckForErrors()
 			if( CurTouchingActor != NULL )
 			{
 				AStaticMeshActor *A = Cast<AStaticMeshActor>( CurTouchingActor );
-				if ( A && (A != this) && (A->GetActorLocation() - GetActorLocation()).IsNearlyZero() && A->StaticMeshComponent.IsValid()
+				if ( A && (A != this) && (A->GetActorLocation() - GetActorLocation()).IsNearlyZero() && A->StaticMeshComponent
 					&& (A->StaticMeshComponent->StaticMesh == StaticMeshComponent->StaticMesh) && (A->GetActorRotation() == GetActorRotation()) 
 					&&  (A->StaticMeshComponent->RelativeScale3D == StaticMeshComponent->RelativeScale3D) )
 				{
