Added output logging upon successfully capturing a screenshot.
#ttp 349482 - FEATURE REQUEST: EDITOR: High Resolution Screenshot: Pathway and success piped to output log
#branch UE4
#proj Runtime.Engine

diff --git a/Engine/Source/Runtime/Engine/Private/UnrealClient.cpp b/Engine/Source/Runtime/Engine/Private/UnrealClient.cpp
--- a/Engine/Source/Runtime/Engine/Private/UnrealClient.cpp
+++ b/Engine/Source/Runtime/Engine/Private/UnrealClient.cpp
@@ -792,7 +792,8 @@ void FViewport::HighResScreenshot()
 		}, HyperLinkText);
 		Info.HyperlinkText = FText::FromString(HyperLinkText);
 		
-		FSlateNotificationManager::Get().AddNotification(Info); 
+		FSlateNotificationManager::Get().AddNotification(Info);
+		UE_LOG(LogClient, Log, TEXT("%s %s"), *Message.ToString(), *HyperLinkText);
 	}
 }
 
