polish UI

diff --git a/Engine/Source/Runtime/Engine/Private/Materials/Material.cpp b/Engine/Source/Runtime/Engine/Private/Materials/Material.cpp
--- a/Engine/Source/Runtime/Engine/Private/Materials/Material.cpp
+++ b/Engine/Source/Runtime/Engine/Private/Materials/Material.cpp
@@ -2157,17 +2157,29 @@ bool UMaterial::CanEditChange(const UProperty* InProperty) const
 			// some settings don't make sense for postprocess materials
 
 			if (PropertyName == GET_MEMBER_NAME_STRING_CHECKED(UMaterial, bTangentSpaceNormal) ||
-				PropertyName == GET_MEMBER_NAME_STRING_CHECKED(UMaterial, TwoSided) ||
 				PropertyName == GET_MEMBER_NAME_STRING_CHECKED(UMaterial, bDisableDepthTest) ||
-				PropertyName == GET_MEMBER_NAME_STRING_CHECKED(UMaterial, bUseMaterialAttributes) ||
-				PropertyName == GET_MEMBER_NAME_STRING_CHECKED(UMaterial, bFullyRough) ||
-				PropertyName == GET_MEMBER_NAME_STRING_CHECKED(UMaterial, bUseLightmapDirectionality)
+				PropertyName == GET_MEMBER_NAME_STRING_CHECKED(UMaterial, bUseMaterialAttributes)
 				)
 			{
 				return false;
 			}
 		}
 
+		if (PropertyName == GET_MEMBER_NAME_STRING_CHECKED(UMaterial, bFullyRough) ||
+			PropertyName == GET_MEMBER_NAME_STRING_CHECKED(UMaterial, TwoSided) ||
+			PropertyName == GET_MEMBER_NAME_STRING_CHECKED(UMaterial, bUseLightmapDirectionality) ||
+			PropertyName == GET_MEMBER_NAME_STRING_CHECKED(UMaterial, D3D11TessellationMode))
+		{
+			return MaterialDomain == MD_Surface;
+		}
+
+		if (PropertyName == GET_MEMBER_NAME_STRING_CHECKED(UMaterial, bEnableCrackFreeDisplacement) ||
+			PropertyName == GET_MEMBER_NAME_STRING_CHECKED(UMaterial, bEnableAdaptiveTessellation)
+			)
+		{
+			return MaterialDomain == MD_Surface && D3D11TessellationMode != MTM_NoTessellation;
+		}
+
 		if (PropertyName == GET_MEMBER_NAME_STRING_CHECKED(UMaterial, BlendableLocation) ||
 			PropertyName == GET_MEMBER_NAME_STRING_CHECKED(UMaterial, BlendablePriority)
 			)
