AO static fraction exposed to project settings for better performance when not using baked lighting

diff --git a/Engine/Source/Runtime/Engine/Classes/Engine/RendererSettings.h b/Engine/Source/Runtime/Engine/Classes/Engine/RendererSettings.h
--- a/Engine/Source/Runtime/Engine/Classes/Engine/RendererSettings.h
+++ b/Engine/Source/Runtime/Engine/Classes/Engine/RendererSettings.h
@@ -201,6 +201,11 @@ class ENGINE_API URendererSettings
 		ToolTip = "Whether the default for AmbientOcclusion is enabled or not (postprocess volume/camera/game setting can still override and enable or disable it independently)"))
 	uint32 bDefaultFeatureAmbientOcclusion : 1;
 
+	UPROPERTY(config, EditAnywhere, Category = DefaultPostprocessingSettings, meta = (
+		ConsoleVariable = "r.DefaultFeature.AmbientOcclusionStaticFraction", DisplayName = "Ambient Occlusion Static Fraction (AO for baked lighting)",
+		ToolTip = "Whether the default for AmbientOcclusionStaticFraction is enabled or not (only useful for baked lighting and if AO is on, allows to have SSAO affect baked lighting as well, costs performance, postprocess volume/camera/game setting can still override and enable or disable it independently)"))
+		uint32 bDefaultFeatureAmbientOcclusionStaticFraction : 1;
+
 	UPROPERTY(config, EditAnywhere, Category = DefaultPostprocessingSettings, meta = (
 		ConsoleVariable = "r.DefaultFeature.AutoExposure", DisplayName = "Auto Exposure",
 		ToolTip = "Whether the default for AutoExposure is enabled or not (postprocess volume/camera/game setting can still override and enable or disable it independently)"))
diff --git a/Engine/Source/Runtime/Engine/Classes/Engine/Scene.h b/Engine/Source/Runtime/Engine/Classes/Engine/Scene.h
--- a/Engine/Source/Runtime/Engine/Classes/Engine/Scene.h
+++ b/Engine/Source/Runtime/Engine/Classes/Engine/Scene.h
@@ -845,6 +845,7 @@ struct FPostProcessSettings
 		GrainJitter = 0.0f;
 		// next value might get overwritten by r.DefaultFeature.AmbientOcclusion
 		AmbientOcclusionIntensity = .5f;
+		// next value might get overwritten by r.DefaultFeature.AmbientOcclusionStaticFraction
 		AmbientOcclusionStaticFraction = 1.0f;
 		AmbientOcclusionRadius = 40.0f;
 		AmbientOcclusionDistance = 80.0f;
diff --git a/Engine/Source/Runtime/Engine/Private/SceneView.cpp b/Engine/Source/Runtime/Engine/Private/SceneView.cpp
--- a/Engine/Source/Runtime/Engine/Private/SceneView.cpp
+++ b/Engine/Source/Runtime/Engine/Private/SceneView.cpp
@@ -76,6 +76,13 @@ static TAutoConsoleVariable<int32> CVarDefaultAmbientOcclusion(
 	TEXT(" 0: off, sets AmbientOcclusionIntensity to 0\n")
 	TEXT(" 1: on (default)"));
 
+static TAutoConsoleVariable<int32> CVarDefaultAmbientOcclusionStaticFraction(
+	TEXT("r.DefaultFeature.AmbientOcclusionStaticFraction"),
+	1,
+	TEXT("Engine default (project setting) for AmbientOcclusion is (postprocess volume/camera/game setting still can override)\n")
+	TEXT(" 0: off, sets AmbientOcclusionStaticFraction to 0\n")
+	TEXT(" 1: on (default, costs extra pass, only useful if there is some baked lighting)"));
+
 static TAutoConsoleVariable<int32> CVarDefaultAutoExposure(
 	TEXT("r.DefaultFeature.AutoExposure"),
 	1,
@@ -909,6 +916,15 @@ void FSceneView::StartFinalPostprocessSettings(FVector InViewLocation)
 				FinalPostProcessSettings.AntiAliasingMethod = (EAntiAliasingMethod)Value;
 			}
 		}
+
+		{
+			int32 Value = CVarDefaultAmbientOcclusionStaticFraction.GetValueOnGameThread();
+
+			if(!Value)
+			{
+				FinalPostProcessSettings.AmbientOcclusionStaticFraction = 0.0f;
+			}
+		}
 	}
 
 	if(State)
