Bunch of fixes for Visual Logger (ttps: #349696, 349692, 349679, 349636, 349624) #ue4

diff --git a/Engine/Source/Developer/GameplayDebugger/Private/GameplayDebuggingComponent.cpp b/Engine/Source/Developer/GameplayDebugger/Private/GameplayDebuggingComponent.cpp
--- a/Engine/Source/Developer/GameplayDebugger/Private/GameplayDebuggingComponent.cpp
+++ b/Engine/Source/Developer/GameplayDebugger/Private/GameplayDebuggingComponent.cpp
@@ -664,7 +664,10 @@ void UGameplayDebuggingComponent::CollectEQSData()
 			CurrentLocalData = &EQSLocalData[EQSLocalData.Num()-1];
 		}
 
-		UEnvQueryDebugHelpers::QueryToDebugData(CachedQueryInstance.Get(), *CurrentLocalData);
+		if (CachedQueryInstance.IsValid())
+		{
+			UEnvQueryDebugHelpers::QueryToDebugData(*CachedQueryInstance, *CurrentLocalData);
+		}
 		CurrentLocalData->Timestamp = AllQueries[Index].Timestamp;
 	}
 
diff --git a/Engine/Source/Developer/LogVisualizer/Private/LogVisualizer.cpp b/Engine/Source/Developer/LogVisualizer/Private/LogVisualizer.cpp
--- a/Engine/Source/Developer/LogVisualizer/Private/LogVisualizer.cpp
+++ b/Engine/Source/Developer/LogVisualizer/Private/LogVisualizer.cpp
@@ -115,7 +115,7 @@ void FLogVisualizer::PullDataFromVisualLog(const FVisualLog& VisualLog)
 	}
 }
 
-void FLogVisualizer::OnNewLog(const AActor* Actor, TSharedPtr<FActorsVisLog> Log)
+void FLogVisualizer::OnNewLog(const UObject* Object, TSharedPtr<FActorsVisLog> Log)
 {
 	Logs.Add(Log);
 	LogAddedEvent.Broadcast();
diff --git a/Engine/Source/Developer/LogVisualizer/Private/LogVisualizer.h b/Engine/Source/Developer/LogVisualizer/Private/LogVisualizer.h
--- a/Engine/Source/Developer/LogVisualizer/Private/LogVisualizer.h
+++ b/Engine/Source/Developer/LogVisualizer/Private/LogVisualizer.h
@@ -48,7 +48,7 @@ public:
 	virtual class AActor* GetHelperActor(class UWorld* InWorld);
 
 protected:
-	void OnNewLog(const class AActor* Actor, TSharedPtr<FActorsVisLog> Log);
+	void OnNewLog(const class UObject* Object, TSharedPtr<FActorsVisLog> Log);
 	
 public:
 	/** All hosted logs */
diff --git a/Engine/Source/Developer/LogVisualizer/Private/SLogBar.cpp b/Engine/Source/Developer/LogVisualizer/Private/SLogBar.cpp
--- a/Engine/Source/Developer/LogVisualizer/Private/SLogBar.cpp
+++ b/Engine/Source/Developer/LogVisualizer/Private/SLogBar.cpp
@@ -254,11 +254,6 @@ void SLogBar::SelectEntry(const FGeometry& MyGeometry, const float ClickX)
 					BestDistance = Distance;
 					BestIndex = EntryIndex;
 				}
-				else if (Distance > LastDistance)
-				{
-					break;
-				}
-				LastDistance = Distance;
 			}
 		}
 	}
diff --git a/Engine/Source/Developer/LogVisualizer/Private/SLogVisualizer.cpp b/Engine/Source/Developer/LogVisualizer/Private/SLogVisualizer.cpp
--- a/Engine/Source/Developer/LogVisualizer/Private/SLogVisualizer.cpp
+++ b/Engine/Source/Developer/LogVisualizer/Private/SLogVisualizer.cpp
@@ -863,7 +863,10 @@ void SLogVisualizer::UpdateVisibleEntriesCache(const TSharedPtr<FActorsVisLog>&
 		// if there is no LogFilter widget - show all
 		OutEntriesCached[CachedLogIndex].CachedEntries = Log->Entries;
 	}
-	LogsList[CurrentIndex]->LastEndTimestamp = LogsList[CurrentIndex]->EndTimestamp;
+	if (LogsList.IsValidIndex(CurrentIndex))
+	{
+		LogsList[CurrentIndex]->LastEndTimestamp = LogsList[CurrentIndex]->EndTimestamp;
+	}
 }
 
 void SLogVisualizer::GetVisibleEntries(const TSharedPtr<FActorsVisLog>& Log, TArray<TSharedPtr<FVisualLogEntry> >& OutEntries)
diff --git a/Engine/Source/Runtime/AIModule/Classes/EnvironmentQuery/EnvQueryDebugHelpers.h b/Engine/Source/Runtime/AIModule/Classes/EnvironmentQuery/EnvQueryDebugHelpers.h
--- a/Engine/Source/Runtime/AIModule/Classes/EnvironmentQuery/EnvQueryDebugHelpers.h
+++ b/Engine/Source/Runtime/AIModule/Classes/EnvironmentQuery/EnvQueryDebugHelpers.h
@@ -138,18 +138,9 @@ FArchive& operator<<(FArchive& Ar, EQSDebug::FQueryData& Data)
 #endif //USE_EQS_DEBUGGER || ENABLE_VISUAL_LOG
 
 #if ENABLE_VISUAL_LOG && USE_EQS_DEBUGGER
-
-#define UE_VLOG_EQS(Query, CategoryName, Verbosity) \
-{ \
-	if (Query != NULL) \
-	{ \
-		UEnvQueryDebugHelpers::LogQuery(Query, CategoryName.GetCategoryName(), ELogVerbosity::Verbosity, FString::Printf(TEXT("Executed EQS: \n - Name: '%s' (id=%d, option=%d),\n - All Items: %d,\n - ValidItems: %d"), *Query->QueryName, Query->QueryID, Query->OptionIndex, Query->ItemDetails.Num(), Query->NumValidItems)); \
-	} \
-}
-
+#	define UE_VLOG_EQS(Query, Category, Verbosity)  UEnvQueryDebugHelpers::LogQuery(Query, Category, ELogVerbosity::Verbosity);
 #else
-#define UE_VLOG_EQS(Query, CategoryName, Verbosity)
-
+#	define UE_VLOG_EQS(Query, CategoryName, Verbosity)
 #endif //ENABLE_VISUAL_LOG && USE_EQS_DEBUGGER
 
 UCLASS(Abstract, CustomConstructor)
@@ -159,49 +150,36 @@ class AIMODULE_API UEnvQueryDebugHelpers : public UObject
 
 	UEnvQueryDebugHelpers(const FObjectInitializer& ObjectInitializer) : Super(ObjectInitializer) {}
 #if USE_EQS_DEBUGGER
-	static void QueryToDebugData(struct FEnvQueryInstance* Query, EQSDebug::FQueryData& EQSLocalData);
-	static void QueryToBlobArray(struct FEnvQueryInstance* Query, TArray<uint8>& BlobArray, bool bUseCompression = false);
+	static void QueryToDebugData(struct FEnvQueryInstance& Query, EQSDebug::FQueryData& EQSLocalData);
+	static void QueryToBlobArray(struct FEnvQueryInstance& Query, TArray<uint8>& BlobArray, bool bUseCompression = false);
 	static void BlobArrayToDebugData(const TArray<uint8>& BlobArray, EQSDebug::FQueryData& EQSLocalData, bool bUseCompression = false);
-	static void LogQuery(struct FEnvQueryInstance* Query, const FName& CategoryName, ELogVerbosity::Type, const FString& AdditionalLogInfo);
+#endif
+
+#if ENABLE_VISUAL_LOG && USE_EQS_DEBUGGER
+	static void LogQuery(struct FEnvQueryInstance& Query, const struct FLogCategoryBase& Category, ELogVerbosity::Type);
 #endif
 };
 
-#if USE_EQS_DEBUGGER
-FORCEINLINE void UEnvQueryDebugHelpers::LogQuery(struct FEnvQueryInstance* Query, const FName& CategoryName, ELogVerbosity::Type Type, const FString& AdditionalLogInfo)
+#if ENABLE_VISUAL_LOG && USE_EQS_DEBUGGER
+FORCEINLINE void UEnvQueryDebugHelpers::LogQuery(struct FEnvQueryInstance& Query, const struct FLogCategoryBase& Category, ELogVerbosity::Type Type)
 {
-#if ENABLE_VISUAL_LOG
-	if ((GEngine && GEngine->bDisableAILogging) || Query == NULL || FVisualLogger::Get().IsRecording() == false || (FVisualLogger::Get().IsBlockedForAllCategories() && FVisualLogger::Get().GetWhiteList().Find(CategoryName) == INDEX_NONE))
-	{
-		return;
-	}
-
-	const UObject* LogOwner = FVisualLogger::Get().FindRedirection(Query->Owner.Get());
-	if (LogOwner == NULL || LogOwner->HasAnyFlags(RF_ClassDefaultObject))
-	{
-		return;
-	}
-
-	UWorld* World = GEngine->GetWorldFromContextObject(Query->Owner.Get(), false);
-	if (!World)
+	UWorld *World = NULL;
+	FVisualLogEntry *CurrentEntry = NULL;
+	if (CheckVisualLogInputInternal(Query.Owner.Get(), Category, Type, &World, &CurrentEntry) == false)
 	{
 		return;
 	}
 
+	const int32 UniqueId = FVisualLogger::Get().GetUniqueId(World->TimeSeconds);
 	TArray<uint8> BlobArray;
 	UEnvQueryDebugHelpers::QueryToBlobArray(Query, BlobArray);
+	FString AdditionalLogInfo = FString::Printf(TEXT("Executed EQS: \n - Name: '%s' (id=%d, option=%d),\n - All Items: %d,\n - ValidItems: %d"), *Query.QueryName, Query.QueryID, Query.OptionIndex, Query.ItemDetails.Num(), Query.NumValidItems);
+	FVisualLogEntry::FLogLine Line(Category.GetCategoryName(), Type, AdditionalLogInfo, Query.QueryID);
+	Line.TagName = *EVisLogTags::TAG_EQS;
+	Line.UniqueId = UniqueId;
 
-	FVisualLogEntry* EntryToWrite = FVisualLogger::Get().GetEntryToWrite(LogOwner, World->TimeSeconds);
-	const int32 UniqueId = FVisualLogger::Get().GetUniqueId(World->TimeSeconds);
-	if (EntryToWrite)
-	{
-		FVisualLogEntry::FLogLine Line(CategoryName, Type, AdditionalLogInfo, Query->QueryID);
-		Line.TagName = *EVisLogTags::TAG_EQS;
-		Line.UniqueId = UniqueId;
-		EntryToWrite->LogLines.Add(Line);
-
-		EntryToWrite->AddDataBlock(EVisLogTags::TAG_EQS, BlobArray, CategoryName).UniqueId = UniqueId;
-	}
-#endif
+	CurrentEntry->LogLines.Add(Line);
+	CurrentEntry->AddDataBlock(EVisLogTags::TAG_EQS, BlobArray, Category.GetCategoryName()).UniqueId = UniqueId;
 }
 #endif
 
diff --git a/Engine/Source/Runtime/AIModule/Private/EnvironmentQuery/EnvQueryDebugHelpers.cpp b/Engine/Source/Runtime/AIModule/Private/EnvironmentQuery/EnvQueryDebugHelpers.cpp
--- a/Engine/Source/Runtime/AIModule/Private/EnvironmentQuery/EnvQueryDebugHelpers.cpp
+++ b/Engine/Source/Runtime/AIModule/Private/EnvironmentQuery/EnvQueryDebugHelpers.cpp
@@ -7,7 +7,7 @@
 #include "EnvironmentQuery/EQSQueryResultSourceInterface.h"
 
 #if USE_EQS_DEBUGGER
-void UEnvQueryDebugHelpers::QueryToBlobArray(struct FEnvQueryInstance* Query, TArray<uint8>& BlobArray, bool bUseCompression)
+void UEnvQueryDebugHelpers::QueryToBlobArray(struct FEnvQueryInstance& Query, TArray<uint8>& BlobArray, bool bUseCompression)
 {
 	EQSDebug::FQueryData EQSLocalData;
 	UEnvQueryDebugHelpers::QueryToDebugData(Query, EQSLocalData);
@@ -38,42 +38,42 @@ void UEnvQueryDebugHelpers::QueryToBlobArray(struct FEnvQueryInstance* Query, TA
 	}
 }
 
-void UEnvQueryDebugHelpers::QueryToDebugData(struct FEnvQueryInstance* Query, EQSDebug::FQueryData& EQSLocalData)
+void UEnvQueryDebugHelpers::QueryToDebugData(struct FEnvQueryInstance& Query, EQSDebug::FQueryData& EQSLocalData)
 {
 	// step 1: data for rendering component
 	EQSLocalData.Reset();
 
-	FEQSSceneProxy::CollectEQSData(Query, Query, EQSLocalData.SolidSpheres, EQSLocalData.Texts, true, EQSLocalData.RenderDebugHelpers);
+	FEQSSceneProxy::CollectEQSData(&Query, &Query, EQSLocalData.SolidSpheres, EQSLocalData.Texts, true, EQSLocalData.RenderDebugHelpers);
 
 	// step 2: detailed scoring data for HUD
 	const int32 MaxDetailedItems = 10;
 	const int32 FirstItemIndex = 0;
 
-	const int32 NumTests = Query->ItemDetails.IsValidIndex(0) ? Query->ItemDetails[0].TestResults.Num() : 0;
-	const int32 NumItems = FMath::Min(MaxDetailedItems, Query->NumValidItems);
+	const int32 NumTests = Query.ItemDetails.IsValidIndex(0) ? Query.ItemDetails[0].TestResults.Num() : 0;
+	const int32 NumItems = FMath::Min(MaxDetailedItems, Query.NumValidItems);
 
-	EQSLocalData.Name = Query->QueryName;
-	EQSLocalData.Id = Query->QueryID;
+	EQSLocalData.Name = Query.QueryName;
+	EQSLocalData.Id = Query.QueryID;
 	EQSLocalData.Items.Reset(NumItems);
 	EQSLocalData.Tests.Reset(NumTests);
-	EQSLocalData.NumValidItems = Query->NumValidItems;
+	EQSLocalData.NumValidItems = Query.NumValidItems;
 
-	UEnvQueryItemType* ItemCDO = Query->ItemType.GetDefaultObject();
+	UEnvQueryItemType* ItemCDO = Query.ItemType.GetDefaultObject();
 	for (int32 ItemIdx = 0; ItemIdx < NumItems; ItemIdx++)
 	{
 		EQSDebug::FItemData ItemInfo;
 		ItemInfo.ItemIdx = ItemIdx + FirstItemIndex;
-		ItemInfo.TotalScore = Query->Items[ItemInfo.ItemIdx].Score;
+		ItemInfo.TotalScore = Query.Items[ItemInfo.ItemIdx].Score;
 
-		const uint8* ItemData = Query->RawData.GetData() + Query->Items[ItemInfo.ItemIdx].DataOffset;
+		const uint8* ItemData = Query.RawData.GetData() + Query.Items[ItemInfo.ItemIdx].DataOffset;
 		ItemInfo.Desc = FString::Printf(TEXT("[%d] %s"), ItemInfo.ItemIdx, *ItemCDO->GetDescription(ItemData));
 
 		ItemInfo.TestValues.Reserve(NumTests);
 		ItemInfo.TestScores.Reserve(NumTests);
 		for (int32 TestIdx = 0; TestIdx < NumTests; TestIdx++)
 		{
-			const float ScoreN = Query->ItemDetails[ItemInfo.ItemIdx].TestResults[TestIdx];
-			const float ScoreW = Query->ItemDetails[ItemInfo.ItemIdx].TestWeightedScores[TestIdx];
+			const float ScoreN = Query.ItemDetails[ItemInfo.ItemIdx].TestResults[TestIdx];
+			const float ScoreW = Query.ItemDetails[ItemInfo.ItemIdx].TestWeightedScores[TestIdx];
 
 			ItemInfo.TestValues.Add(ScoreN);
 			ItemInfo.TestScores.Add(ScoreW);
@@ -82,24 +82,24 @@ void UEnvQueryDebugHelpers::QueryToDebugData(struct FEnvQueryInstance* Query, EQ
 		EQSLocalData.Items.Add(ItemInfo);
 	}
 
-	const int32 NumAllTests = Query->Options[Query->OptionIndex].Tests.Num();
+	const int32 NumAllTests = Query.Options[Query.OptionIndex].Tests.Num();
 	for (int32 TestIdx = 0; TestIdx < NumAllTests; TestIdx++)
 	{
 		EQSDebug::FTestData TestInfo;
 
-		UEnvQueryTest* TestOb = Query->Options[Query->OptionIndex].Tests[TestIdx];
+		UEnvQueryTest* TestOb = Query.Options[Query.OptionIndex].Tests[TestIdx];
 		TestInfo.ShortName = TestOb->GetDescriptionTitle();
 		TestInfo.Detailed = TestOb->GetDescriptionDetails().ToString().Replace(TEXT("\n"), TEXT(", "));
 
 		EQSLocalData.Tests.Add(TestInfo);
 	}
 
-	EQSLocalData.UsedOption = Query->OptionIndex;
-	const int32 NumOptions = Query->Options.Num();
+	EQSLocalData.UsedOption = Query.OptionIndex;
+	const int32 NumOptions = Query.Options.Num();
 	for (int32 OptionIndex = 0; OptionIndex < NumOptions; ++OptionIndex)
 	{
-		const FString UserName = Query->Options[OptionIndex].Generator->OptionName;
-		EQSLocalData.Options.Add(UserName.Len() > 0 ? UserName : Query->Options[OptionIndex].Generator->GetName());
+		const FString UserName = Query.Options[OptionIndex].Generator->OptionName;
+		EQSLocalData.Options.Add(UserName.Len() > 0 ? UserName : Query.Options[OptionIndex].Generator->GetName());
 	}
 }
 
diff --git a/Engine/Source/Runtime/AIModule/Private/EnvironmentQuery/EnvQueryManager.cpp b/Engine/Source/Runtime/AIModule/Private/EnvironmentQuery/EnvQueryManager.cpp
--- a/Engine/Source/Runtime/AIModule/Private/EnvironmentQuery/EnvQueryManager.cpp
+++ b/Engine/Source/Runtime/AIModule/Private/EnvironmentQuery/EnvQueryManager.cpp
@@ -187,7 +187,7 @@ TSharedPtr<FEnvQueryResult> UEnvQueryManager::RunInstantQuery(const FEnvQueryReq
 		QueryInstance->ExecuteOneStep((double)FLT_MAX);
 	}
 
-	UE_VLOG_EQS(QueryInstance.Get(), LogEQS, All);
+	UE_VLOG_EQS(*QueryInstance.Get(), LogEQS, All);
 
 #if USE_EQS_DEBUGGER
 	EQSDebugger.StoreQuery(QueryInstance);
@@ -260,7 +260,7 @@ void UEnvQueryManager::Tick(float DeltaTime)
 			
 			if (QueryInstance->Status != EEnvQueryStatus::Processing)
 			{
-				UE_VLOG_EQS(QueryInstance.Get(), LogEQS, All);
+				UE_VLOG_EQS(*QueryInstance.Get(), LogEQS, All);
 
 #if USE_EQS_DEBUGGER
 				EQSDebugger.StoreQuery(QueryInstance);
diff --git a/Engine/Source/Runtime/Core/Public/HAL/OutputDevices.h b/Engine/Source/Runtime/Core/Public/HAL/OutputDevices.h
--- a/Engine/Source/Runtime/Core/Public/HAL/OutputDevices.h
+++ b/Engine/Source/Runtime/Core/Public/HAL/OutputDevices.h
@@ -129,7 +129,7 @@ struct CORE_API FLogCategoryBase
 	}
 	/** Called just after a logging statement being allow to print. Checks a few things and maybe breaks into the debugger. **/
 	void PostTrigger(ELogVerbosity::Type VerbosityLevel);
-	FORCEINLINE FName GetCategoryName()
+	FORCEINLINE FName GetCategoryName() const
 	{
 		return CategoryFName;
 	}
diff --git a/Engine/Source/Runtime/Engine/Private/VisualLog.cpp b/Engine/Source/Runtime/Engine/Private/VisualLog.cpp
--- a/Engine/Source/Runtime/Engine/Private/VisualLog.cpp
+++ b/Engine/Source/Runtime/Engine/Private/VisualLog.cpp
@@ -359,24 +359,17 @@ TSharedPtr<FVisualLogEntry> JsonToVisualLogEntry(TSharedPtr<FJsonValue> FromJson
 //----------------------------------------------------------------------//
 // FActorsVisLog
 //----------------------------------------------------------------------//
+FActorsVisLog::FActorsVisLog(FName InName, TArray<TWeakObjectPtr<UObject> >* Children)
+	: Name(InName)
+{
+	Entries.Reserve(VisLogInitialSize);
+}
+
 FActorsVisLog::FActorsVisLog(const class UObject* Object, TArray<TWeakObjectPtr<UObject> >* Children)
 	: Name(Object->GetFName())
 	, FullName(Object->GetFullName())
 {
-	const class AActor* AsActor = Cast<AActor>(Object);
-
 	Entries.Reserve(VisLogInitialSize);
-	if (!AsActor)
-	{
-		UWorld* World = GEngine->GetWorldFromContextObject(Object);
-		check(World);
-
-		Entries.Add(MakeShareable(new FVisualLogEntry(World->GetTimeSeconds(), FVector(), Object, Children)));
-	}
-	else
-	{
-		Entries.Add(MakeShareable(new FVisualLogEntry(AsActor, Children)));
-	}
 }
 
 FActorsVisLog::FActorsVisLog(TSharedPtr<FJsonValue> FromJson)
@@ -471,10 +464,20 @@ void FVisualLog::SetFileName(const FString& InFileName)
 	BaseFileName = InFileName;
 }
 
-void FVisualLog::Serialize(const class UObject* LogOwner, const FVisualLogEntry& LogEntry)
+void FVisualLog::Serialize(const class UObject* LogOwner, FName OwnerName, const FVisualLogEntry& LogEntry)
 {
-	TSharedPtr<FActorsVisLog> Log = GetLog(LogOwner);
+	const bool bContainsLogOwner = LogsMap.Contains(LogOwner);
+	if (!bContainsLogOwner)
+	{
+		LogsMap.Add(LogOwner, MakeShareable(new FActorsVisLog(OwnerName, NULL)));
+	}
+
+	TSharedPtr<FActorsVisLog> Log = LogsMap[LogOwner];
 	Log->Entries.Add(MakeShareable(new FVisualLogEntry(LogEntry)));
+	if (!bContainsLogOwner)
+	{
+		OnNewLogCreated.ExecuteIfBound(LogOwner, Log);
+	}
 }
 
 FString FVisualLog::GetLogFileFullName() const
diff --git a/Engine/Source/Runtime/Engine/Private/VisualLogger/VisualLogger.cpp b/Engine/Source/Runtime/Engine/Private/VisualLogger/VisualLogger.cpp
--- a/Engine/Source/Runtime/Engine/Private/VisualLogger/VisualLogger.cpp
+++ b/Engine/Source/Runtime/Engine/Private/VisualLogger/VisualLogger.cpp
@@ -8,7 +8,7 @@
 
 DEFINE_LOG_CATEGORY(LogVisual);
 
-TMap<class UObject*, TArray<const class UObject*> > FVisualLogger::RedirectionMap;
+TMap<class UObject*, TArray<TWeakObjectPtr<const class UObject> > > FVisualLogger::RedirectionMap;
 
 FVisualLogger::FVisualLogger()
 {
@@ -66,27 +66,35 @@ int32 FVisualLogger::GetUniqueId(float Timestamp)
 
 void FVisualLogger::Redirect(class UObject* FromObject, class UObject* ToObject)
 {
+	if (FromObject == ToObject)
+	{
+		return;
+	}
+
 	UObject* OldRedirection = FindRedirection(FromObject);
 	UObject* NewRedierection = FindRedirection(ToObject);
 
-	UE_CVLOG(OldRedirection != NULL, OldRedirection, LogVisual, Log, TEXT("Redirected '%s' to '%s'"), *FromObject->GetName(), *NewRedierection->GetName());
-
-	TArray<const class UObject*>& OldArray = RedirectionMap.FindOrAdd(OldRedirection);
-	OldArray.RemoveSingleSwap(FromObject);
+	auto OldArray = RedirectionMap.Find(OldRedirection);
+	if (OldArray)
+	{
+		OldArray->RemoveSingleSwap(FromObject);
+	}
 
-	TArray<const class UObject*>& NewArray = RedirectionMap.FindOrAdd(NewRedierection);
-	NewArray.AddUnique(FromObject);
+	RedirectionMap.FindOrAdd(NewRedierection).AddUnique(FromObject);
 
-	UE_CVLOG(NewRedierection != NULL, NewRedierection, LogVisual, Log, TEXT("Redirected '%s' to '%s'"), *FromObject->GetName(), *NewRedierection->GetName());
+	UE_CVLOG(FromObject != NULL, FromObject, LogVisual, Log, TEXT("Redirected '%s' to '%s'"), *FromObject->GetName(), *NewRedierection->GetName());
 }
 
 class UObject* FVisualLogger::FindRedirection(const class UObject* Object)
 {
-	for (auto& Redirection : RedirectionMap)
+	if (RedirectionMap.Contains(Object) == false)
 	{
-		if (Redirection.Value.Find(Object) != INDEX_NONE)
+		for (auto& Redirection : RedirectionMap)
 		{
-			return FindRedirection(Redirection.Key);
+			if (Redirection.Value.Find(Object) != INDEX_NONE)
+			{
+				return FindRedirection(Redirection.Key);
+			}
 		}
 	}
 
diff --git a/Engine/Source/Runtime/Engine/Private/VisualLogger/VisualLoggerAutomationTests.cpp b/Engine/Source/Runtime/Engine/Private/VisualLogger/VisualLoggerAutomationTests.cpp
--- a/Engine/Source/Runtime/Engine/Private/VisualLogger/VisualLoggerAutomationTests.cpp
+++ b/Engine/Source/Runtime/Engine/Private/VisualLogger/VisualLoggerAutomationTests.cpp
@@ -22,7 +22,7 @@ public:
 	virtual void StartRecordingToFile(float TImeStamp) override;
 	virtual void StopRecordingToFile(float TImeStamp) override;
 	virtual void SetFileName(const FString& InFileName) override;
-	virtual void Serialize(const class UObject* LogOwner, const FVisualLogEntry& LogEntry) override;
+	virtual void Serialize(const class UObject* LogOwner, FName OwnerName, const FVisualLogEntry& LogEntry) override;
 
 	class UObject* LastObject;
 	FVisualLogEntry LastEntry;
@@ -54,7 +54,7 @@ void FVisualLoggerTestDevice::SetFileName(const FString& InFileName)
 
 }
 
-void FVisualLoggerTestDevice::Serialize(const class UObject* LogOwner, const FVisualLogEntry& LogEntry)
+void FVisualLoggerTestDevice::Serialize(const class UObject* LogOwner, FName OwnerName, const FVisualLogEntry& LogEntry)
 {
 	LastObject = const_cast<class UObject*>(LogOwner);
 	LastEntry = LogEntry;
diff --git a/Engine/Source/Runtime/Engine/Public/VisualLog.h b/Engine/Source/Runtime/Engine/Public/VisualLog.h
--- a/Engine/Source/Runtime/Engine/Public/VisualLog.h
+++ b/Engine/Source/Runtime/Engine/Public/VisualLog.h
@@ -79,6 +79,7 @@ struct ENGINE_API FActorsVisLog
 	FString FullName;
 	TArray<TSharedPtr<FVisualLogEntry> > Entries;
 
+	FActorsVisLog(FName Name, TArray<TWeakObjectPtr<UObject> >* Children);
 	FActorsVisLog(const class UObject* Object, TArray<TWeakObjectPtr<UObject> >* Children);
 	FActorsVisLog(TSharedPtr<FJsonValue> FromJson);
 
@@ -94,7 +95,7 @@ public:
 	typedef TMap<const class UObject*, TSharedPtr<FActorsVisLog> > FLogsMap;
 	typedef TMap<const class AActor*, TArray<TWeakObjectPtr<UObject> > > FLogRedirectsMap;
 
-	DECLARE_DELEGATE_TwoParams(FOnNewLogCreatedDelegate, const AActor*, TSharedPtr<FActorsVisLog>);
+	DECLARE_DELEGATE_TwoParams(FOnNewLogCreatedDelegate, const UObject*, TSharedPtr<FActorsVisLog>);
 
 	FVisualLog();
 	~FVisualLog();
@@ -108,7 +109,7 @@ public:
 	virtual void Cleanup(bool bReleaseMemory) override;
 	virtual void StartRecordingToFile(float TImeStamp) override;
 	virtual void StopRecordingToFile(float TImeStamp) override;
-	virtual void Serialize(const class UObject* LogOwner, const FVisualLogEntry& LogEntry) override;
+	virtual void Serialize(const class UObject* LogOwner, FName OwnerName, const FVisualLogEntry& LogEntry) override;
 	virtual void SetFileName(const FString& InFileName) override;
 
 	void Redirect(class UObject* Source, const class AActor* NewRedirection);
@@ -150,12 +151,6 @@ protected:
 	FORCEINLINE TSharedPtr<FActorsVisLog> GetLog(const class UObject* Object)
 	{
 		TSharedPtr<FActorsVisLog>* Log = LogsMap.Find(Object);
-		if (Log == NULL)
-		{
-			const class AActor* Actor = GetVisualLogRedirection(Object);
-			Log = &(LogsMap.Add(Object, MakeShareable(new FActorsVisLog(Object, RedirectsMap.Find(Actor)))));
-			OnNewLogCreated.ExecuteIfBound(Actor, *Log);
-		}
 		return *Log;
 	}
 			
diff --git a/Engine/Source/Runtime/Engine/Public/VisualLogger/VisualLogger.h b/Engine/Source/Runtime/Engine/Public/VisualLogger/VisualLogger.h
--- a/Engine/Source/Runtime/Engine/Public/VisualLogger/VisualLogger.h
+++ b/Engine/Source/Runtime/Engine/Public/VisualLogger/VisualLogger.h
@@ -155,10 +155,12 @@ public:
 	FORCEINLINE void BlockAllCategories(bool InBlock) { bBlockedAllCategories = InBlock; }
 
 	/** checks if all categories are blocked */
-	FORCEINLINE bool IsBlockedForAllCategories() { return !!bBlockedAllCategories; }
+	FORCEINLINE bool IsBlockedForAllCategories() const { return !!bBlockedAllCategories; }
 
 	/** Returns white list for modifications */
-	TArray<class FName>& GetWhiteList() { return CategoriesWhiteList; }
+	FORCEINLINE TArray<class FName>& GetWhiteList() { return CategoriesWhiteList; }
+
+	FORCEINLINE bool IsWhiteListed(const FName& Name) const { return CategoriesWhiteList.Find(Name) != INDEX_NONE; }
 
 	FORCEINLINE void AddCategortyToWhiteList(class FName Category) { CategoriesWhiteList.AddUnique(Category); }
 
@@ -211,7 +213,7 @@ protected:
 	/** Array of output devices to redirect to */
 	TArray<FVisualLogDevice*> OutputDevices;
 	// Map for inter-objects redirections
-	static TMap<class UObject*, TArray<const class UObject*> > RedirectionMap;
+	static TMap<class UObject*, TArray<TWeakObjectPtr<const class UObject> > > RedirectionMap;
 	// white list of categories to bypass blocking
 	TArray<class FName>	CategoriesWhiteList;
 	// Visual Logger extensions map
@@ -220,6 +222,10 @@ protected:
 	TMap<float, int32> LastUniqueIds;
 	// Current entry with all data
 	TMap<const class UObject*, FVisualLogEntry>	CurrentEntryPerObject;
+	// Map to contain names for Objects (they can be destroyed after while)
+	TMap<const class UObject*, FName> ObjectToNameMap;
+	// Map to contain information about pointers in game
+	TMap<const class UObject*, TWeakObjectPtr<const class UObject> > ObjectToPointerMap;
 	// if set all categories are blocked from logging
 	int32 bBlockedAllCategories : 1;
 	// if set we are recording and collecting all vlog data
@@ -243,7 +249,7 @@ public:
 	virtual void Cleanup(bool bReleaseMemory = false) = 0;
 	virtual void StartRecordingToFile(float TImeStamp) = 0;
 	virtual void StopRecordingToFile(float TImeStamp) = 0;
-	virtual void Serialize(const class UObject* LogOwner, const FVisualLogEntry& LogEntry) = 0;
+	virtual void Serialize(const class UObject* LogOwner, FName OwnerName, const FVisualLogEntry& LogEntry) = 0;
 	virtual void SetFileName(const FString& InFileName) = 0;
 	virtual bool HasFlags(int32 InFlags) { return !!(InFlags & (VisualLogger::CanSaveToFile | VisualLogger::StoreLogsLocally)); }
 };
diff --git a/Engine/Source/Runtime/Engine/Public/VisualLogger/VisualLogger.inl b/Engine/Source/Runtime/Engine/Public/VisualLogger/VisualLogger.inl
--- a/Engine/Source/Runtime/Engine/Public/VisualLogger/VisualLogger.inl
+++ b/Engine/Source/Runtime/Engine/Public/VisualLogger/VisualLogger.inl
@@ -2,40 +2,52 @@
 
 #include "Runtime/Core/Private/Misc/VarargsHelper.h"
 
-#define CHECK_VLOG_DATA() \
-	float CurrentTime = 0; \
-	const FName CategoryName = const_cast<struct FLogCategoryBase&>(Category).GetCategoryName(); \
-	if (GEngine && GEngine->bDisableAILogging || FVisualLogger::Get().bIsRecording == false || !Object || Object->HasAnyFlags(RF_ClassDefaultObject) || (FVisualLogger::Get().IsBlockedForAllCategories() && FVisualLogger::Get().CategoriesWhiteList.Find(CategoryName) == INDEX_NONE)) \
-	{ \
-		return; \
-	} \
-	UObject* LogOwner = FindRedirection(Object); \
-	if (LogOwner == NULL) \
-	{ \
-		return; \
-	} \
-	UWorld* World = GEngine->GetWorldFromContextObject(Object); \
-	if(ensure(World) == false) \
-	{ \
-		return; \
+FORCEINLINE
+bool CheckVisualLogInputInternal(const class UObject* Object, const struct FLogCategoryBase& Category, ELogVerbosity::Type Verbosity, UWorld **World, FVisualLogEntry **CurrentEntry)
+{
+	FVisualLogger& VisualLogger = FVisualLogger::Get();
+	if (!Object || (GEngine && GEngine->bDisableAILogging) || VisualLogger.IsRecording() == false || Object->HasAnyFlags(RF_ClassDefaultObject))
+	{
+		return false; 
+	}
+	const FName CategoryName = Category.GetCategoryName(); 
+	if (VisualLogger.IsBlockedForAllCategories() && VisualLogger.IsWhiteListed(CategoryName) == false)
+	{
+		return false; 
+	}
+
+	*World = GEngine->GetWorldFromContextObject(Object, false); 
+	if (ensure(*World != NULL) == false)
+	{
+		return false; 
+	}
+
+	*CurrentEntry = VisualLogger.GetEntryToWrite(Object, (*World)->TimeSeconds);
+	if (ensure(CurrentEntry != NULL) == false)
+	{
+		return false; 
 	}
 
-FORCEINLINE_DEBUGGABLE
+	return true;
+}
+
+FORCEINLINE
 FVisualLogEntry* FVisualLogger::GetEntryToWrite(const class UObject* Object, float TimeStamp, VisualLogger::ECreateIfNeeded ShouldCreate)
 {
 	FVisualLogEntry* CurrentEntry = NULL;
+	UObject * LogOwner = FVisualLogger::FindRedirection(Object);
 
 	bool InitializeNewEntry = false;
-	if (CurrentEntryPerObject.Contains(Object))
+	if (CurrentEntryPerObject.Contains(LogOwner))
 	{
-		CurrentEntry = &CurrentEntryPerObject[Object];
+		CurrentEntry = &CurrentEntryPerObject[LogOwner];
 		if (TimeStamp > CurrentEntry->TimeStamp && ShouldCreate == VisualLogger::Create)
 		{
 			if (CurrentEntry->TimeStamp >= 0) //-1 means not initialized entry information
 			{
 				for (auto* Device : OutputDevices)
 				{
-					Device->Serialize(Object, *CurrentEntry);
+					Device->Serialize(LogOwner, ObjectToNameMap[LogOwner], *CurrentEntry);
 				}
 			}
 			InitializeNewEntry = true;
@@ -43,7 +55,9 @@ FVisualLogEntry* FVisualLogger::GetEntryToWrite(const class UObject* Object, flo
 	}
 	else
 	{
-		CurrentEntry = &CurrentEntryPerObject.Add(Object);
+		CurrentEntry = &CurrentEntryPerObject.Add(LogOwner);
+		ObjectToNameMap.Add(LogOwner, LogOwner->GetFName());
+		ObjectToPointerMap.Add(LogOwner, LogOwner);
 		InitializeNewEntry = true;
 	}
 
@@ -51,29 +65,44 @@ FVisualLogEntry* FVisualLogger::GetEntryToWrite(const class UObject* Object, flo
 	{
 		CurrentEntry->Reset();
 		CurrentEntry->TimeStamp = TimeStamp;
-		const class AActor* ObjectAsActor = Cast<class AActor>(Object);
-		if (ObjectAsActor)
-		{
-			CurrentEntry->Location = ObjectAsActor->GetActorLocation();
-			ObjectAsActor->GrabDebugSnapshot(CurrentEntry);
-		}
-		if (RedirectionMap.Contains(Object))
+
+		if (RedirectionMap.Contains(LogOwner))
 		{
-			for (auto Child : RedirectionMap[Object])
+			if (ObjectToPointerMap.Contains(LogOwner) && ObjectToPointerMap[LogOwner].IsValid())
+			{
+				const class AActor* LogOwnerAsActor = Cast<class AActor>(LogOwner);
+				if (LogOwnerAsActor)
+				{
+					LogOwnerAsActor->GrabDebugSnapshot(CurrentEntry);
+				}
+			}
+			for (auto Child : RedirectionMap[LogOwner])
 			{
-				const class AActor* ChildAsActor = Cast<class AActor>(Child);
-				if (ChildAsActor)
+				if (Child.IsValid())
 				{
-					ChildAsActor->GrabDebugSnapshot(CurrentEntry);
+					const class AActor* ChildAsActor = Cast<class AActor>(Child.Get());
+					if (ChildAsActor)
+					{
+						ChildAsActor->GrabDebugSnapshot(CurrentEntry);
+					}
 				}
 			}
 		}
+		else
+		{
+			const class AActor* ObjectAsActor = Cast<class AActor>(Object);
+			if (ObjectAsActor)
+			{
+				CurrentEntry->Location = ObjectAsActor->GetActorLocation();
+				ObjectAsActor->GrabDebugSnapshot(CurrentEntry);
+			}
+		}
 	}
 
 	return CurrentEntry;
 }
 
-FORCEINLINE_DEBUGGABLE
+FORCEINLINE
 void FVisualLogger::Flush()
 {
 	for (auto &CurrentEntry : CurrentEntryPerObject)
@@ -82,7 +111,7 @@ void FVisualLogger::Flush()
 		{
 			for (auto* Device : OutputDevices)
 			{
-				Device->Serialize(CurrentEntry.Key, CurrentEntry.Value);
+				Device->Serialize(CurrentEntry.Key, ObjectToNameMap[CurrentEntry.Key], CurrentEntry.Value);
 			}
 			CurrentEntry.Value.Reset();
 		}
@@ -93,104 +122,120 @@ void FVisualLogger::Flush()
 FORCEINLINE
 VARARG_BODY(void, FVisualLogger::CategorizedLogf, const TCHAR*, VARARG_EXTRA(const class UObject* Object) VARARG_EXTRA(const struct FLogCategoryBase& Category) VARARG_EXTRA(ELogVerbosity::Type Verbosity) VARARG_EXTRA(int32 UniqueLogId))
 {
-	CHECK_VLOG_DATA();
-
-	FVisualLogEntry* CurrentEntry = FVisualLogger::Get().GetEntryToWrite(LogOwner, World->TimeSeconds);
-	ensure(CurrentEntry);
+	UWorld *World = NULL;
+	FVisualLogEntry *CurrentEntry = NULL;
+	if (CheckVisualLogInputInternal(Object, Category, Verbosity, &World, &CurrentEntry) == false)
+	{
+		return;
+	}
 
 	GROWABLE_LOGF(
-		CurrentEntry->AddText(Buffer, CategoryName);
+		CurrentEntry->AddText(Buffer, Category.GetCategoryName());
 	);
 }
 
 FORCEINLINE
 VARARG_BODY(void, FVisualLogger::GeometryShapeLogf, const TCHAR*, VARARG_EXTRA(const class UObject* Object) VARARG_EXTRA(const struct FLogCategoryBase& Category) VARARG_EXTRA(ELogVerbosity::Type Verbosity) VARARG_EXTRA(int32 UniqueLogId) VARARG_EXTRA(const FVector& Start) VARARG_EXTRA(const FVector& End) VARARG_EXTRA(const FColor& Color))
 {
-	CHECK_VLOG_DATA();
-
-	FVisualLogEntry* CurrentEntry = FVisualLogger::Get().GetEntryToWrite(LogOwner, World->TimeSeconds);
-	ensure(CurrentEntry);
+	UWorld *World = NULL;
+	FVisualLogEntry *CurrentEntry = NULL;
+	if (CheckVisualLogInputInternal(Object, Category, Verbosity, &World, &CurrentEntry) == false)
+	{
+		return;
+	}
 
 	GROWABLE_LOGF(
-		CurrentEntry->AddElement(Start, End, CategoryName, Color, Buffer);
+		CurrentEntry->AddElement(Start, End, Category.GetCategoryName(), Color, Buffer);
 	);
 }
 
 FORCEINLINE
 VARARG_BODY(void, FVisualLogger::GeometryShapeLogf, const TCHAR*, VARARG_EXTRA(const class UObject* Object) VARARG_EXTRA(const struct FLogCategoryBase& Category) VARARG_EXTRA(ELogVerbosity::Type Verbosity) VARARG_EXTRA(int32 UniqueLogId) VARARG_EXTRA(const FVector& Location) VARARG_EXTRA(float Radius) VARARG_EXTRA(const FColor& Color))
 {
-	CHECK_VLOG_DATA();
-
-	FVisualLogEntry* CurrentEntry = FVisualLogger::Get().GetEntryToWrite(LogOwner, World->TimeSeconds);
-	ensure(CurrentEntry);
+	UWorld *World = NULL;
+	FVisualLogEntry *CurrentEntry = NULL;
+	if (CheckVisualLogInputInternal(Object, Category, Verbosity, &World, &CurrentEntry) == false)
+	{
+		return;
+	}
 
 	GROWABLE_LOGF(
-		CurrentEntry->AddElement(Location, CategoryName, Color, Buffer, Radius);
+		CurrentEntry->AddElement(Location, Category.GetCategoryName(), Color, Buffer, Radius);
 	);
 }
 
 FORCEINLINE
 VARARG_BODY(void, FVisualLogger::GeometryShapeLogf, const TCHAR*, VARARG_EXTRA(const class UObject* Object) VARARG_EXTRA(const struct FLogCategoryBase& Category) VARARG_EXTRA(ELogVerbosity::Type Verbosity) VARARG_EXTRA(int32 UniqueLogId) VARARG_EXTRA(const FBox& Box) VARARG_EXTRA(const FColor& Color))
 {
-	CHECK_VLOG_DATA();
-
-	FVisualLogEntry* CurrentEntry = FVisualLogger::Get().GetEntryToWrite(LogOwner, World->TimeSeconds);
-	ensure(CurrentEntry);
+	UWorld *World = NULL;
+	FVisualLogEntry *CurrentEntry = NULL;
+	if (CheckVisualLogInputInternal(Object, Category, Verbosity, &World, &CurrentEntry) == false)
+	{
+		return;
+	}
 
 	GROWABLE_LOGF(
-		CurrentEntry->AddElement(Box, CategoryName, Color, Buffer);
+		CurrentEntry->AddElement(Box, Category.GetCategoryName(), Color, Buffer);
 	);
 }
 
 FORCEINLINE
 VARARG_BODY(void, FVisualLogger::GeometryShapeLogf, const TCHAR*, VARARG_EXTRA(const class UObject* Object) VARARG_EXTRA(const struct FLogCategoryBase& Category) VARARG_EXTRA(ELogVerbosity::Type Verbosity) VARARG_EXTRA(int32 UniqueLogId) VARARG_EXTRA(const FVector& Orgin) VARARG_EXTRA(const FVector& Direction) VARARG_EXTRA(const float Length) VARARG_EXTRA(const float Angle)  VARARG_EXTRA(const FColor& Color))
 {
-	CHECK_VLOG_DATA();
-
-	FVisualLogEntry* CurrentEntry = FVisualLogger::Get().GetEntryToWrite(LogOwner, World->TimeSeconds);
-	ensure(CurrentEntry);
+	UWorld *World = NULL;
+	FVisualLogEntry *CurrentEntry = NULL;
+	if (CheckVisualLogInputInternal(Object, Category, Verbosity, &World, &CurrentEntry) == false)
+	{
+		return;
+	}
 
 	GROWABLE_LOGF(
-		CurrentEntry->AddElement(Orgin, Direction, Length, Angle, Angle, CategoryName, Color, Buffer);
+		CurrentEntry->AddElement(Orgin, Direction, Length, Angle, Angle, Category.GetCategoryName(), Color, Buffer);
 	);
 }
 
 FORCEINLINE
 VARARG_BODY(void, FVisualLogger::GeometryShapeLogf, const TCHAR*, VARARG_EXTRA(const class UObject* Object) VARARG_EXTRA(const struct FLogCategoryBase& Category) VARARG_EXTRA(ELogVerbosity::Type Verbosity) VARARG_EXTRA(int32 UniqueLogId) VARARG_EXTRA(const FVector& Start) VARARG_EXTRA(const FVector& End) VARARG_EXTRA(const float Radius) VARARG_EXTRA(const FColor& Color))
 {
-	CHECK_VLOG_DATA();
-
-	FVisualLogEntry* CurrentEntry = FVisualLogger::Get().GetEntryToWrite(LogOwner, World->TimeSeconds);
-	ensure(CurrentEntry);
+	UWorld *World = NULL;
+	FVisualLogEntry *CurrentEntry = NULL;
+	if (CheckVisualLogInputInternal(Object, Category, Verbosity, &World, &CurrentEntry) == false)
+	{
+		return;
+	}
 
 	GROWABLE_LOGF(
-		CurrentEntry->AddElement(Start, End, Radius, CategoryName, Color, Buffer);
+		CurrentEntry->AddElement(Start, End, Radius, Category.GetCategoryName(), Color, Buffer);
 	);
 }
 
 FORCEINLINE
 VARARG_BODY(void, FVisualLogger::GeometryShapeLogf, const TCHAR*, VARARG_EXTRA(const class UObject* Object) VARARG_EXTRA(const struct FLogCategoryBase& Category) VARARG_EXTRA(ELogVerbosity::Type Verbosity) VARARG_EXTRA(int32 UniqueLogId) VARARG_EXTRA(const FVector& Center) VARARG_EXTRA(float HalfHeight) VARARG_EXTRA(float Radius) VARARG_EXTRA(const FQuat& Rotation) VARARG_EXTRA(const FColor& Color))
 {
-	CHECK_VLOG_DATA();
-
-	FVisualLogEntry* CurrentEntry = FVisualLogger::Get().GetEntryToWrite(LogOwner, World->TimeSeconds);
-	ensure(CurrentEntry);
+	UWorld *World = NULL;
+	FVisualLogEntry *CurrentEntry = NULL;
+	if (CheckVisualLogInputInternal(Object, Category, Verbosity, &World, &CurrentEntry) == false)
+	{
+		return;
+	}
 
 	GROWABLE_LOGF(
-		CurrentEntry->AddCapsule(Center, HalfHeight, Radius, Rotation, CategoryName, Color, Buffer);
+		CurrentEntry->AddCapsule(Center, HalfHeight, Radius, Rotation, Category.GetCategoryName(), Color, Buffer);
 	);
 }
 
 FORCEINLINE
 VARARG_BODY(void, FVisualLogger::HistogramDataLogf, const TCHAR*, VARARG_EXTRA(const class UObject* Object) VARARG_EXTRA(const struct FLogCategoryBase& Category) VARARG_EXTRA(ELogVerbosity::Type Verbosity) VARARG_EXTRA(int32 UniqueLogId) VARARG_EXTRA(FName GraphName) VARARG_EXTRA(FName DataName) VARARG_EXTRA(const FVector2D& Data) VARARG_EXTRA(const FColor& Color))
 {
-	CHECK_VLOG_DATA();
-
-	FVisualLogEntry* CurrentEntry = FVisualLogger::Get().GetEntryToWrite(LogOwner, World->TimeSeconds);
-	ensure(CurrentEntry);
+	UWorld *World = NULL;
+	FVisualLogEntry *CurrentEntry = NULL;
+	if (CheckVisualLogInputInternal(Object, Category, Verbosity, &World, &CurrentEntry) == false)
+	{
+		return;
+	}
 
 	GROWABLE_LOGF(
-		CurrentEntry->AddHistogramData(Data, CategoryName, GraphName, DataName);
+		CurrentEntry->AddHistogramData(Data, Category.GetCategoryName(), GraphName, DataName);
 	);
 }
 
@@ -229,45 +274,34 @@ void FVisualLogger::EventLog(const class UObject* Object, const FName EventTag1,
 	EventLog(Object, EventTag1, Event2);
 }
 
-FORCEINLINE_DEBUGGABLE_ACTUAL
+FORCEINLINE
 void FVisualLogger::EventLog(const class UObject* Object, const FName EventTag1, const FVisualLogEventBase& Event, const FName EventTag2, const FName EventTag3, const FName EventTag4, const FName EventTag5, const FName EventTag6)
 {
-	const FName CategoryName = *Event.GetName();
-
-	float CurrentTime = 0;
-	if ((GEngine && GEngine->bDisableAILogging) || FVisualLogger::Get().bIsRecording == false || !Object || Object->HasAnyFlags(RF_ClassDefaultObject) || (FVisualLogger::Get().IsBlockedForAllCategories() && FVisualLogger::Get().CategoriesWhiteList.Find(CategoryName) == INDEX_NONE))
+	UWorld *World = NULL;
+	FVisualLogEntry *CurrentEntry = NULL;
+	const FLogCategory<ELogVerbosity::Log, ELogVerbosity::Log> Category(*Event.GetName());
+	if (CheckVisualLogInputInternal(Object, Category, ELogVerbosity::Log, &World, &CurrentEntry) == false)
 	{
 		return;
 	}
-	UObject* LogOwner = FindRedirection(Object);
-	UWorld* World = GEngine->GetWorldFromContextObject(Object);
-	if (World == NULL) 
-	{ 
-		ensure(World); 
-		return; 
-	} 
-
-	FVisualLogEntry* CurrentEntry = FVisualLogger::Get().GetEntryToWrite(LogOwner, World->TimeSeconds);
-	if (ensure(CurrentEntry))
-	{
-		int32 Index = CurrentEntry->Events.Find(FVisualLogEntry::FLogEvent(Event));
-		if (Index != INDEX_NONE)
-		{
-			CurrentEntry->Events[Index].Counter++;
-		}
-		else
-		{
-			Index = CurrentEntry->AddEvent(Event);
-		}
 
-		CurrentEntry->Events[Index].EventTags.FindOrAdd(EventTag1)++;
-		CurrentEntry->Events[Index].EventTags.FindOrAdd(EventTag2)++;
-		CurrentEntry->Events[Index].EventTags.FindOrAdd(EventTag3)++;
-		CurrentEntry->Events[Index].EventTags.FindOrAdd(EventTag4)++;
-		CurrentEntry->Events[Index].EventTags.FindOrAdd(EventTag5)++;
-		CurrentEntry->Events[Index].EventTags.FindOrAdd(EventTag6)++;
-		CurrentEntry->Events[Index].EventTags.Remove(NAME_None);
+	int32 Index = CurrentEntry->Events.Find(FVisualLogEntry::FLogEvent(Event));
+	if (Index != INDEX_NONE)
+	{
+		CurrentEntry->Events[Index].Counter++;
 	}
+	else
+	{
+		Index = CurrentEntry->AddEvent(Event);
+	}
+
+	CurrentEntry->Events[Index].EventTags.FindOrAdd(EventTag1)++;
+	CurrentEntry->Events[Index].EventTags.FindOrAdd(EventTag2)++;
+	CurrentEntry->Events[Index].EventTags.FindOrAdd(EventTag3)++;
+	CurrentEntry->Events[Index].EventTags.FindOrAdd(EventTag4)++;
+	CurrentEntry->Events[Index].EventTags.FindOrAdd(EventTag5)++;
+	CurrentEntry->Events[Index].EventTags.FindOrAdd(EventTag6)++;
+	CurrentEntry->Events[Index].EventTags.Remove(NAME_None);
 }
 
 #undef  CHECK_VLOG_DATA
