Eliminate ULevelBase
Improve ULevel documentation

diff --git a/Engine/Source/Runtime/Engine/Classes/Engine/Level.h b/Engine/Source/Runtime/Engine/Classes/Engine/Level.h
--- a/Engine/Source/Runtime/Engine/Classes/Engine/Level.h
+++ b/Engine/Source/Runtime/Engine/Classes/Engine/Level.h
@@ -1,11 +1,6 @@
 // Copyright 1998-2014 Epic Games, Inc. All Rights Reserved.
 
-//
-// The level object.  Contains the level's actor list, Bsp information, and brush list.
-//
-
 #pragma once
-#include "Engine/LevelBase.h"
 #include "Engine/World.h"
 #include "Level.generated.h"
 
@@ -66,50 +61,7 @@ struct ENGINE_API FDynamicTextureInstance : public FStreamableTextureInstance
 	friend FArchive& operator<<( FArchive& Ar, FDynamicTextureInstance& TextureInstance );
 };
 
-
-/** 
- *	Used for storing cached data for simplified static mesh collision at a particular scaling.
- *	One entry, relating to a particular static mesh cached at a particular scaling, giving an index into the CachedPhysSMDataStore. 
- */
-USTRUCT()
-struct ENGINE_API FCachedPhysSMData
-{
-	GENERATED_USTRUCT_BODY()
-	/** Scale of mesh that the data was cached for. */
-	FVector				Scale3D;
-
-	/** Index into CachedPhysSMDataStore that this cached data is stored at. */
-	int32					CachedDataIndex;
-
-	/** Serializer function. */
-	friend FArchive& operator<<( FArchive& Ar, FCachedPhysSMData& D )
-	{
-		Ar << D.Scale3D << D.CachedDataIndex;
-		return Ar;
-	}
-};
-
-/** 
- *	Used for storing cached data for per-tri static mesh collision at a particular scaling. 
- */
-USTRUCT()
-struct ENGINE_API FCachedPerTriPhysSMData
-{
-	GENERATED_USTRUCT_BODY()
-	/** Scale of mesh that the data was cached for. */
-	FVector			Scale3D;
-
-	/** Index into array Cached data for this mesh at this scale. */
-	int32				CachedDataIndex;
-
-	/** Serializer function. */
-	friend FArchive& operator<<( FArchive& Ar, FCachedPerTriPhysSMData& D )
-	{
-		Ar << D.Scale3D << D.CachedDataIndex;
-		return Ar;
-	}
-};
-
+/** Struct that holds on to information about Actors that wish to be auto enabled for input before the player controller has been created */
 struct FPendingAutoReceiveInputActor
 {
 	TWeakObjectPtr<AActor> Actor;
@@ -276,10 +228,31 @@ private:
 	friend class FLightmassProcessor;
 };
 
-UCLASS(customConstructor,MinimalAPI)
-class ULevel : public ULevelBase, public IInterface_AssetUserData
+//
+// The level object.  Contains the level's actor list, BSP information, and brush list.
+//
+
+
+/**
+ * A Level contains actors, a blueprint representing the level logic, BSP information, and brushes.
+ * Every Level has a World as its Outer and can be used as the PersistentLevel, however, 
+ * when a Level has been streamed in the OwningWorld represents the World that it is a part of.
+ * 
+ * @see https://docs.unrealengine.com/latest/INT/Engine/Levels
+ * @see UActor
+ */
+UCLASS(MinimalAPI)
+class ULevel : public UObject, public IInterface_AssetUserData
 {
-	GENERATED_UCLASS_BODY()
+	GENERATED_BODY()
+
+public:
+
+	/** URL associated with this level. */
+	FURL					URL;
+
+	/** Array of all actors in this level, used by FActorIteratorBase and derived classes */
+	TTransArray<AActor*> Actors;
 
 	/** Set before calling LoadPackage for a streaming level to ensure that OwningWorld is correct on the Level */
 	ENGINE_API static TMap<FName, UWorld*> StreamedLevelsOwningWorld;
@@ -331,8 +304,6 @@ class ULevel : public ULevelBase, public IInterface_AssetUserData
 	UPROPERTY()
 	TArray<FVector> StaticNavigableGeometry;
 
-	bool bTextureStreamingBuilt;
-
 	/** Static information used by texture streaming code, generated during PreSave									*/
 	TMap<UTexture2D*,TArray<FStreamableTextureInstance> >	TextureToInstancesMap;
 
@@ -369,6 +340,9 @@ class ULevel : public ULevelBase, public IInterface_AssetUserData
 	/** Whether the geometry needs to be rebuilt for correct lighting */
 	uint32										bGeometryDirtyForLighting:1;
 
+	/** Has texture streaming been built */
+	uint32										bTextureStreamingBuilt:1;
+
 	/** Whether the level is currently visible/ associated with the world */
 	UPROPERTY(transient)
 	uint32										bIsVisible:1;
@@ -448,6 +422,7 @@ protected:
 
 private:
 
+	// Actors awaiting input to be enabled once the appropriate PlayerController has been created
 	TArray<FPendingAutoReceiveInputActor> PendingAutoReceiveInputActors;
 
 public:
@@ -517,7 +492,7 @@ public:
 	void InvalidateModelGeometry();
 
 #if WITH_EDITOR
-	/** Called to creete ModelComponents for BSP rendering */
+	/** Called to create ModelComponents for BSP rendering */
 	void CreateModelComponents();
 #endif // WITH_EDITOR
 
diff --git a/Engine/Source/Runtime/Engine/Classes/Engine/LevelBase.h b/Engine/Source/Runtime/Engine/Classes/Engine/LevelBase.h
deleted file mode 100644
diff --git a/Engine/Source/Runtime/Engine/Private/EnginePrivate.h b/Engine/Source/Runtime/Engine/Private/EnginePrivate.h
--- a/Engine/Source/Runtime/Engine/Private/EnginePrivate.h
+++ b/Engine/Source/Runtime/Engine/Private/EnginePrivate.h
@@ -183,7 +183,6 @@
 #include "Engine/InterpCurveEdSetup.h"
 #include "Engine/IntSerialization.h"
 #include "Layers/Layer.h"
-#include "Engine/LevelBase.h"
 #include "Engine/Level.h"
 #include "Engine/LevelStreaming.h"
 #include "Engine/LevelStreamingAlwaysLoaded.h"
diff --git a/Engine/Source/Runtime/Engine/Private/Level.cpp b/Engine/Source/Runtime/Engine/Private/Level.cpp
--- a/Engine/Source/Runtime/Engine/Private/Level.cpp
+++ b/Engine/Source/Runtime/Engine/Private/Level.cpp
@@ -30,39 +30,6 @@ Level.cpp: Level-related functions
 
 DEFINE_LOG_CATEGORY(LogLevel);
 
-/*-----------------------------------------------------------------------------
-ULevelBase implementation.
------------------------------------------------------------------------------*/
-
-ULevelBase::ULevelBase( const FObjectInitializer& ObjectInitializer,const FURL& InURL )
-	: UObject(ObjectInitializer)
-	, URL( InURL )
-	, Actors( this )
-{
-
-}
-
-void ULevelBase::AddReferencedObjects(UObject* InThis, FReferenceCollector& Collector)
-{	
-	ULevelBase* This = CastChecked<ULevelBase>(InThis);
-	// Let GC know that we're referencing some AActor objects
-	for (auto& Actor : This->Actors)
-	{
-		Collector.AddReferencedObject(Actor, This);
-	}
-	UObject* ActorsOwner = This->Actors.GetOwner();
-	Collector.AddReferencedObject(ActorsOwner, This);
-	Super::AddReferencedObjects(This, Collector);
-}
-
-void ULevelBase::Serialize( FArchive& Ar )
-{
-	Super::Serialize(Ar);
-	Ar << Actors;
-	Ar << URL;
-}
-
-
 /*-----------------------------------------------------------------------------
 ULevel implementation.
 -----------------------------------------------------------------------------*/
@@ -254,17 +221,20 @@ FArchive& operator<<( FArchive& Ar, FPrecomputedVolumeDistanceField& D )
 TMap<FName, UWorld*> ULevel::StreamedLevelsOwningWorld;
 
 ULevel::ULevel( const FObjectInitializer& ObjectInitializer )
-	:	ULevelBase( ObjectInitializer )
-	,	OwningWorld(NULL)
-	,	TickTaskLevel(FTickTaskManagerInterface::Get().AllocateTickTaskLevel())
+	:   UObject( ObjectInitializer )
+	,   Actors(this)
+	,   OwningWorld(NULL)
+	,   TickTaskLevel(FTickTaskManagerInterface::Get().AllocateTickTaskLevel())
 {
 	PrecomputedLightVolume = new FPrecomputedLightVolume();
 }
 
 ULevel::ULevel( const FObjectInitializer& ObjectInitializer,const FURL& InURL )
-	:	ULevelBase( ObjectInitializer, InURL )
-	,	OwningWorld(NULL)
-	,	TickTaskLevel(FTickTaskManagerInterface::Get().AllocateTickTaskLevel())
+	:   UObject( ObjectInitializer )
+	,   URL(InURL)
+	,   Actors(this)
+	,   OwningWorld(NULL)
+	,   TickTaskLevel(FTickTaskManagerInterface::Get().AllocateTickTaskLevel())
 {
 	PrecomputedLightVolume = new FPrecomputedLightVolume();
 }
@@ -306,6 +276,14 @@ void ULevel::AddReferencedObjects(UObject* InThis, FReferenceCollector& Collecto
 		Collector.AddReferencedObject( Texture2D, This );
 	}
 
+	// Let GC know that we're referencing some AActor objects
+	for (auto& Actor : This->Actors)
+	{
+		Collector.AddReferencedObject(Actor, This);
+	}
+	UObject* ActorsOwner = This->Actors.GetOwner();
+	Collector.AddReferencedObject(ActorsOwner, This);
+
 	Super::AddReferencedObjects( This, Collector );
 }
 
@@ -339,6 +317,9 @@ void ULevel::Serialize( FArchive& Ar )
 {
 	Super::Serialize( Ar );
 
+	Ar << Actors;
+	Ar << URL;
+
 	Ar << Model;
 
 	Ar << ModelComponents;
diff --git a/Engine/Source/Runtime/Engine/Public/Engine.h b/Engine/Source/Runtime/Engine/Public/Engine.h
--- a/Engine/Source/Runtime/Engine/Public/Engine.h
+++ b/Engine/Source/Runtime/Engine/Public/Engine.h
@@ -237,7 +237,6 @@
 #include "Engine/InterpCurveEdSetup.h"
 #include "Engine/IntSerialization.h"
 #include "Layers/Layer.h"
-#include "Engine/LevelBase.h"
 #include "Engine/Level.h"
 #include "Engine/LevelStreaming.h"
 #include "Engine/LevelStreamingAlwaysLoaded.h"
