Moved PImplRecastNavMesh.h to public folder #UE4

diff --git a/Engine/Source/Runtime/Engine/Private/AI/Navigation/PImplRecastNavMesh.h b/Engine/Source/Runtime/Engine/Private/AI/Navigation/PImplRecastNavMesh.h
deleted file mode 100644
diff --git a/Engine/Source/Runtime/Engine/Public/AI/Navigation/PImplRecastNavMesh.h b/Engine/Source/Runtime/Engine/Public/AI/Navigation/PImplRecastNavMesh.h
new file mode 100644
--- /dev/null
+++ b/Engine/Source/Runtime/Engine/Public/AI/Navigation/PImplRecastNavMesh.h
@@ -0,0 +1,238 @@
+// Copyright 1998-2014 Epic Games, Inc. All Rights Reserved.
+
+//
+// Private implementation for communication with Recast library
+// 
+// All functions should be called through RecastNavMesh actor to make them thread safe!
+//
+
+#pragma once 
+
+#if WITH_RECAST
+
+#include "AI/Navigation/NavFilters/NavigationQueryFilter.h"
+#include "Detour/DetourNavMeshQuery.h"
+
+#define RECAST_VERY_SMALL_AGENT_RADIUS 0.0f
+
+class ENGINE_API FRecastQueryFilter : public INavigationQueryFilterInterface, public dtQueryFilter
+{
+public:
+	FRecastQueryFilter(bool bIsVirtual = true);
+	virtual ~FRecastQueryFilter(){}
+
+	virtual void Reset() override;
+
+	virtual void SetAreaCost(uint8 AreaType, float Cost) override;
+	virtual void SetFixedAreaEnteringCost(uint8 AreaType, float Cost) override;
+	virtual void SetExcludedArea(uint8 AreaType) override;
+	virtual void SetAllAreaCosts(const float* CostArray, const int32 Count) override;
+	virtual void GetAllAreaCosts(float* CostArray, float* FixedCostArray, const int32 Count) const override;
+	virtual void SetBacktrackingEnabled(const bool bBacktracking) override;
+	virtual bool IsBacktrackingEnabled() const override;
+	virtual bool IsEqual(const INavigationQueryFilterInterface* Other) const override;
+	virtual void SetIncludeFlags(uint16 Flags) override;
+	virtual uint16 GetIncludeFlags() const override;
+	virtual void SetExcludeFlags(uint16 Flags) override;
+	virtual uint16 GetExcludeFlags() const override;
+	virtual INavigationQueryFilterInterface* CreateCopy() const override;
+
+	const dtQueryFilter* GetAsDetourQueryFilter() const { return this; }
+
+	/** note that it results in loosing all area cost setup. Call it before setting anything else */
+	void SetIsVirtual(bool bIsVirtual);
+};
+
+struct ENGINE_API FRecastSpeciaLinkFilter : public dtQuerySpecialLinkFilter
+{
+	FRecastSpeciaLinkFilter(UNavigationSystem* NavSystem, const UObject* Owner) : NavSys(NavSystem), SearchOwner(Owner) {}
+	virtual bool isLinkAllowed(const int32 UserId) const override;
+
+	UNavigationSystem* NavSys;
+	const UObject* SearchOwner;
+};
+
+/** Engine Private! - Private Implementation details of ARecastNavMesh */
+class ENGINE_API FPImplRecastNavMesh
+{
+public:
+
+	/** Constructor */
+	FPImplRecastNavMesh(class ARecastNavMesh* Owner);
+
+	/** Dtor */
+	~FPImplRecastNavMesh();
+
+	/**
+	 * Serialization.
+	 * @param Ar - The archive with which to serialize.
+	 * @returns true if serialization was successful.
+	 */
+	void Serialize(FArchive& Ar);
+
+	/** Debug rendering. */
+	void GetDebugGeometry(FRecastDebugGeometry& OutGeometry, int32 TileIndex = INDEX_NONE) const;
+	void GetDebugTileBounds(FBox& OuterBox, int32& NumTilesX, int32& NumTilesY) const;
+
+	/** Returns bounding box for the whole navmesh. */
+	FBox GetNavMeshBounds() const;
+
+	/** Returns bounding box for a given navmesh tile. */
+	FBox GetNavMeshTileBounds(int32 TileIndex) const;
+
+	/** Retrieves XY and layer coordinates of tile specified by index */
+	void GetNavMeshTileXY(int32 TileIndex, int32& OutX, int32& OutY, int32& OutLayer) const;
+
+	/** Retrieves XY coordinates of tile specified by position */
+	void GetNavMeshTileXY(const FVector& Point, int32& OutX, int32& OutY) const;
+
+	/** Retrieves all tile indices at matching XY coordinates */
+	void GetNavMeshTilesAt(int32 TileX, int32 TileY, TArray<int32>& Indices) const;
+
+	/** Retrieves number of tiles in this navmesh */
+	FORCEINLINE int32 GetNavMeshTilesCount() const { return DetourNavMesh->getMaxTiles(); }
+
+	/** Supported queries */
+
+	// @TODONAV
+	/** Generates path from the given query. Synchronous. */
+	ENavigationQueryResult::Type FindPath(const FVector& StartLoc, const FVector& EndLoc, FNavMeshPath& Path, const FNavigationQueryFilter& Filter, const UObject* Owner) const;
+
+	/** Check if path exists */
+	ENavigationQueryResult::Type TestPath(const FVector& StartLoc, const FVector& EndLoc, const FNavigationQueryFilter& Filter, const UObject* Owner, int32* NumVisitedNodes = 0) const;
+
+	/** Check if path exists using cluster graph */
+	ENavigationQueryResult::Type TestClusterPath(const FVector& StartLoc, const FVector& EndLoc, int32* NumVisitedNodes = 0) const;
+
+	/** Checks if the whole segment is in navmesh */
+	void Raycast2D(const FVector& StartLoc, const FVector& EndLoc, const FNavigationQueryFilter& InQueryFilter, const UObject* Owner, ARecastNavMesh::FRaycastResult& RaycastResult) const;
+	void Raycast2D(NavNodeRef StartNode, const FVector& StartLoc, const FVector& EndLoc, const FNavigationQueryFilter& InQueryFilter, const UObject* Owner, ARecastNavMesh::FRaycastResult& RaycastResult) const;
+
+	/** Generates path from given query and collect data for every step of A* algorithm */
+	int32 DebugPathfinding(const FVector& StartLoc, const FVector& EndLoc, const FNavigationQueryFilter& Filter, const UObject* Owner, TArray<FRecastDebugPathfindingStep>& Steps);
+
+	/** Returns a random location on the navmesh. */
+	FNavLocation GetRandomPoint(const FNavigationQueryFilter& Filter, const UObject* Owner) const;
+
+	/** Returns a random location on the navmesh within Radius from Origin. 
+	 *	@return false if no valid navigable location available in specified area */
+	bool GetRandomPointInRadius(const FVector& Origin, float Radius, FNavLocation& OutLocation, const FNavigationQueryFilter& Filter, const UObject* Owner) const;
+
+	/** Returns a random location on the navmesh within cluster */
+	bool GetRandomPointInCluster(NavNodeRef ClusterRef, FNavLocation& OutLocation) const;
+
+	bool ProjectPointToNavMesh(const FVector& Point, FNavLocation& Result, const FVector& Extent, const FNavigationQueryFilter& Filter, const UObject* Owner) const;
+	
+	/** Project single point and grab all vertical intersections */
+	bool ProjectPointMulti(const FVector& Point, TArray<FNavLocation>& OutLocations, const FVector& Extent,
+		float MinZ, float MaxZ, const FNavigationQueryFilter& Filter, const UObject* Owner) const;
+
+	/** Returns nearest navmesh polygon to Loc, or INVALID_NAVMESHREF if Loc is not on the navmesh. */
+	NavNodeRef FindNearestPoly(FVector const& Loc, FVector const& Extent, const FNavigationQueryFilter& Filter, const UObject* Owner) const;
+
+	/** Retrieves all polys within given pathing distance from StartLocation.
+	 *	@NOTE query is not using string-pulled path distance (for performance reasons),
+	 *		it measured distance between middles of portal edges, do you might want to 
+	 *		add an extra margin to PathingDistance */
+	bool GetPolysWithinPathingDistance(FVector const& StartLoc, const float PathingDistance, const FNavigationQueryFilter& Filter, const UObject* Owner, TArray<NavNodeRef>& FoundPolys) const;
+
+	//@todo document
+	void GetEdgesForPathCorridor(const TArray<NavNodeRef>* PathCorridor, TArray<FNavigationPortalEdge>* PathCorridorEdges) const;
+
+	/** finds stringpulled path from given corridor */
+	bool FindStraightPath(const FVector& StartLoc, const FVector& EndLoc, const TArray<NavNodeRef>& PathCorridor, TArray<FNavPathPoint>& PathPoints, TArray<uint32>* CustomLinks = NULL) const;
+
+	/** Filters nav polys in PolyRefs with Filter */
+	bool FilterPolys(TArray<NavNodeRef>& PolyRefs, const class FRecastQueryFilter* Filter, const UObject* Owner) const;
+
+	/** Get all polys from tile */
+	bool GetPolysInTile(int32 TileIndex, TArray<FNavPoly>& Polys) const;
+
+	/** Updates area on polygons creating point-to-point connection with given UserId */
+	void UpdateNavigationLinkArea(int32 UserId, uint8 AreaType, uint16 PolyFlags) const;
+	/** Updates area on polygons creating segment-to-segment connection with given UserId */
+	void UpdateSegmentLinkArea(int32 UserId, uint8 AreaType, uint16 PolyFlags) const;
+
+	/** Retrieves center of the specified polygon. Returns false on error. */
+	bool GetPolyCenter(NavNodeRef PolyID, FVector& OutCenter) const;
+	/** Retrieves the vertices for the specified polygon. Returns false on error. */
+	bool GetPolyVerts(NavNodeRef PolyID, TArray<FVector>& OutVerts) const;
+	/** Retrieves the flags for the specified polygon. Returns false on error. */
+	bool GetPolyData(NavNodeRef PolyID, uint16& Flags, uint8& AreaType) const;
+	/** Retrieves area ID for the specified polygon. */
+	uint32 GetPolyAreaID(NavNodeRef PolyID) const;
+	/** Finds closest point constrained to given poly */
+	bool GetClosestPointOnPoly(NavNodeRef PolyID, const FVector& TestPt, FVector& PointOnPoly) const;
+	/** Decode poly ID into tile index and poly index */
+	bool GetPolyTileIndex(NavNodeRef PolyID, uint32& PolyIndex, uint32& TileIndex) const;
+	/** Retrieves user ID for given offmesh link poly */
+	uint32 GetLinkUserId(NavNodeRef LinkPolyID) const;
+	/** Retrieves start and end point of offmesh link */
+	bool GetLinkEndPoints(NavNodeRef LinkPolyID, FVector& PointA, FVector& PointB) const;
+
+	/** Retrieves bounds of cluster. Returns false on error. */
+	bool GetClusterBounds(NavNodeRef ClusterRef, FBox& OutBounds) const;
+
+	uint32 GetTileIndexFromPolyRef(const NavNodeRef PolyRef) const { return DetourNavMesh != NULL ? DetourNavMesh->decodePolyIdTile(PolyRef) : uint32(INDEX_NONE); }
+	NavNodeRef GetClusterRefFromPolyRef(const NavNodeRef PolyRef) const;
+
+	static uint16 GetFilterForbiddenFlags(const FRecastQueryFilter* Filter);
+	static void SetFilterForbiddenFlags(FRecastQueryFilter* Filter, uint16 ForbiddenFlags);
+
+public:
+	class dtNavMesh const* GetRecastMesh() const { return DetourNavMesh; };
+	class dtNavMesh* GetRecastMesh() { return DetourNavMesh; };
+
+	bool GetOwnsNavMeshData() const { return bOwnsNavMeshData; }
+
+	/** Assigns recast generated navmesh to this instance.
+	 *	@param bOwnData if true from now on this FPImplRecastNavMesh instance will be responsible for this piece 
+	 *		of memory
+	 */
+	void SetRecastMesh(class dtNavMesh* NavMesh, bool bOwnData = false);
+
+	float GetTotalDataSize() const;
+
+	/** Called on world origin changes */
+	void ApplyWorldOffset(const FVector& InOffset, bool bWorldShift);
+
+	/** calculated cost of given segment if traversed on specified poly. Function measures distance between specified points
+	 *	and returns cost of traversing this distance on given poly.
+	 *	@note no check if segment is on poly is performed. */
+	float CalcSegmentCostOnPoly(NavNodeRef PolyID, const dtQueryFilter* Filter, const FVector& StartLoc, const FVector& EndLoc) const;
+
+	class ARecastNavMesh* NavMeshOwner;
+
+	/** if true instance is responsible for deallocation of recast navmesh */
+	bool bOwnsNavMeshData;
+
+	/** Recast's runtime navmesh data that we can query against */
+	class dtNavMesh* DetourNavMesh;
+
+	/** query used for searching data on game thread */
+	mutable dtNavMeshQuery SharedNavQuery;
+
+	/** Helper function to serialize a single Recast tile. */
+	static void SerializeRecastMeshTile(FArchive& Ar, unsigned char*& TileData, int32& TileDataSize);
+
+	/** Initialize data for pathfinding */
+	bool InitPathfinding(const FVector& UnrealStart, const FVector& UnrealEnd, 
+		const dtNavMeshQuery& Query, const dtQueryFilter* Filter,
+		FVector& RecastStart, dtPolyRef& StartPoly,
+		FVector& RecastEnd, dtPolyRef& EndPoly) const;
+
+	/** Marks path flags, perform string pulling if needed */
+	void PostProcessPath(dtStatus PathfindResult, FNavMeshPath& Path,
+		const dtNavMeshQuery& Query, const dtQueryFilter* Filter,
+		NavNodeRef StartNode, NavNodeRef EndNode,
+		const FVector& UnrealStart, const FVector& UnrealEnd,
+		const FVector& RecastStart, FVector& RecastEnd,
+		dtQueryResult& PathResult) const;
+
+	void GetDebugPolyEdges(const struct dtMeshTile* Tile, bool bInternalEdges, bool bNavMeshEdges, TArray<FVector>& InternalEdgeVerts, TArray<FVector>& NavMeshEdgeVerts) const;
+
+	/** workhorse function finding portal edges between corridor polys */
+	void GetEdgesForPathCorridorImpl(const TArray<NavNodeRef>* PathCorridor, TArray<FNavigationPortalEdge>* PathCorridorEdges, const dtNavMeshQuery& NavQuery) const;
+};
+
+#endif	// WITH_RECAST
