Moved RecastHelpers.h to the public folder #UE4

diff --git a/Engine/Source/Runtime/Engine/Private/AI/Navigation/RecastHelpers.h b/Engine/Source/Runtime/Engine/Private/AI/Navigation/RecastHelpers.h
deleted file mode 100644
diff --git a/Engine/Source/Runtime/Engine/Public/AI/Navigation/RecastHelpers.h b/Engine/Source/Runtime/Engine/Public/AI/Navigation/RecastHelpers.h
new file mode 100644
--- /dev/null
+++ b/Engine/Source/Runtime/Engine/Public/AI/Navigation/RecastHelpers.h
@@ -0,0 +1,21 @@
+// Copyright 1998-2014 Epic Games, Inc. All Rights Reserved.
+
+/**
+ * Coord system utilities
+ *
+ * Translates between Unreal and Recast coords.
+ * Unreal: x, y, z
+ * Recast: -x, z, -y
+ */
+
+#pragma once
+
+extern ENGINE_API FVector Unreal2RecastPoint(const float* UnrealPoint);
+extern ENGINE_API FVector Unreal2RecastPoint(const FVector& UnrealPoint);
+extern ENGINE_API FBox Unreal2RecastBox(const FBox& UnrealBox);
+extern ENGINE_API FMatrix Unreal2RecastMatrix();
+
+extern ENGINE_API FVector Recast2UnrealPoint(const float* RecastPoint);
+extern ENGINE_API FVector Recast2UnrealPoint(const FVector& RecastPoint);
+extern ENGINE_API FBox Recast2UnrealBox(const float* RecastMin, const float* RecastMax);
+extern ENGINE_API FBox Recast2UnrealBox(const FBox& RecastBox);
