UE4 - Crash fix during reflection captures with showmaterialdrawevents.

diff --git a/Engine/Source/Runtime/Engine/Private/PrimitiveDrawingUtils.cpp b/Engine/Source/Runtime/Engine/Private/PrimitiveDrawingUtils.cpp
--- a/Engine/Source/Runtime/Engine/Private/PrimitiveDrawingUtils.cpp
+++ b/Engine/Source/Runtime/Engine/Private/PrimitiveDrawingUtils.cpp
@@ -15,7 +15,7 @@ void EmitMeshDrawEvents_Inner(FRHICommandList& RHICmdList, const FPrimitiveScene
 {
 	// Only show material name at the top level
 	// Note: this is the parent's material name, not the material instance
-	SCOPED_DRAW_EVENTF(RHICmdList, MaterialEvent, DEC_SCENE_ITEMS, *Mesh.MaterialRenderProxy->GetMaterial(PrimitiveSceneProxy->GetScene().GetFeatureLevel())->GetFriendlyName());
+	SCOPED_DRAW_EVENTF(RHICmdList, MaterialEvent, DEC_SCENE_ITEMS, *Mesh.MaterialRenderProxy->GetMaterial(PrimitiveSceneProxy ? PrimitiveSceneProxy->GetScene().GetFeatureLevel() : GRHIFeatureLevel_DEPRECATED)->GetFriendlyName());
 	if (PrimitiveSceneProxy)
 	{
 		// Show Actor, level and resource name inside the material name
@@ -696,13 +696,13 @@ void DrawDisc(class FPrimitiveDrawInterface* PDI,const FVector& Base,const FVect
 /*
 x-axis is from point 0 to point 2
 y-axis is from point 0 to point 1
-        6
-        /\
-       /  \
-      /    \
-     4_2  3_5
-       |  |
-       0__1
+		6
+		/\
+	   /  \
+	  /    \
+	 4_2  3_5
+	   |  |
+	   0__1
 */
 
 
