#ue4 - CharacterMovement StepUp() function aborts if forward movement results in no net change in position.

- Step downs are also limited to step up height when not walking.
- Fixes flying Characters moving down when impacting vertical surfaces, even if movement direction was slightly upward.
- Also reduces the worst-case number of sweeps in character movement (previously occurred when pushing into corners).

diff --git a/Engine/Source/Runtime/Engine/Private/CharacterMovementComponent.cpp b/Engine/Source/Runtime/Engine/Private/CharacterMovementComponent.cpp
--- a/Engine/Source/Runtime/Engine/Private/CharacterMovementComponent.cpp
+++ b/Engine/Source/Runtime/Engine/Private/CharacterMovementComponent.cpp
@@ -4531,7 +4531,7 @@ bool UCharacterMovementComponent::CanStepUp(const FHitResult& Hit) const
 }
 
 
-bool UCharacterMovementComponent::StepUp(const FVector& GravDir, const FVector& Delta, const FHitResult &InHit, FStepDownResult* OutStepDownResult)
+bool UCharacterMovementComponent::StepUp(const FVector& InGravDir, const FVector& Delta, const FHitResult &InHit, FStepDownResult* OutStepDownResult)
 {
 	if (!CanStepUp(InHit))
 	{
@@ -4560,7 +4560,14 @@ bool UCharacterMovementComponent::StepUp(const FVector& GravDir, const FVector&
 		return false;
 	}
 
-	float StepTravelHeight = MaxStepHeight;
+	const FVector GravDir = InGravDir.SafeNormal();
+	if (GravDir.IsZero())
+	{
+		return false;
+	}
+
+	float StepTravelUpHeight = MaxStepHeight;
+	float StepTravelDownHeight = StepTravelUpHeight;
 	const float StepSideZ = -1.f * (InHit.ImpactNormal | GravDir);
 	float PawnInitialFloorBaseZ = OldLocation.Z - PawnHalfHeight;
 	float PawnFloorPointZ = PawnInitialFloorBaseZ;
@@ -4570,7 +4577,8 @@ bool UCharacterMovementComponent::StepUp(const FVector& GravDir, const FVector&
 		// Since we float a variable amount off the floor, we need to enforce max step height off the actual point of impact with the floor.
 		const float FloorDist = FMath::Max(0.f, CurrentFloor.FloorDist);
 		PawnInitialFloorBaseZ -= FloorDist;
-		StepTravelHeight = FMath::Max(StepTravelHeight - FloorDist, 0.f);
+		StepTravelUpHeight = FMath::Max(StepTravelUpHeight - FloorDist, 0.f);
+		StepTravelDownHeight = (MaxStepHeight + MAX_FLOOR_DIST*2.f);
 
 		const bool bHitVerticalFace = !IsWithinEdgeTolerance(InHit.Location, InHit.ImpactPoint, PawnRadius);
 		if (!CurrentFloor.bLineTrace && !bHitVerticalFace)
@@ -4590,7 +4598,7 @@ bool UCharacterMovementComponent::StepUp(const FVector& GravDir, const FVector&
 	// step up - treat as vertical wall
 	FHitResult SweepUpHit(1.f);
 	const FRotator PawnRotation = CharacterOwner->GetActorRotation();
-	SafeMoveUpdatedComponent( -GravDir * StepTravelHeight, PawnRotation, true, SweepUpHit);
+	SafeMoveUpdatedComponent(-GravDir * StepTravelUpHeight, PawnRotation, true, SweepUpHit);
 
 	// step fwd
 	FHitResult Hit(1.f);
@@ -4622,8 +4630,17 @@ bool UCharacterMovementComponent::StepUp(const FVector& GravDir, const FVector&
 		}
 
 		// adjust and try again
-		SlideAlongSurface(Delta, 1.f - Hit.Time, Hit.Normal, Hit, true);
-		if ( IsFalling() )
+		const float ForwardHitTime = Hit.Time;
+		const float ForwardSlideAmount = SlideAlongSurface(Delta, 1.f - Hit.Time, Hit.Normal, Hit, true);
+		
+		if (IsFalling())
+		{
+			ScopedStepUpMovement.RevertMove();
+			return false;
+		}
+
+		// If both the forward hit and the deflection got us nowhere, there is no point in this step up.
+		if (ForwardHitTime == 0.f && ForwardSlideAmount == 0.f)
 		{
 			ScopedStepUpMovement.RevertMove();
 			return false;
@@ -4631,7 +4648,7 @@ bool UCharacterMovementComponent::StepUp(const FVector& GravDir, const FVector&
 	}
 	
 	// Step down
-	SafeMoveUpdatedComponent(GravDir * (MaxStepHeight + MAX_FLOOR_DIST*2.f), CharacterOwner->GetActorRotation(), true, Hit);
+	SafeMoveUpdatedComponent(GravDir * StepTravelDownHeight, CharacterOwner->GetActorRotation(), true, Hit);
 
 	// If step down was initially penetrating abort the step up
 	if (Hit.bStartPenetrating)
