#ue4 - Removed redundant GetWorld() call in SceneComponent::GetPhysicsVolume()

diff --git a/Engine/Source/Runtime/Engine/Private/SceneComponent.cpp b/Engine/Source/Runtime/Engine/Private/SceneComponent.cpp
--- a/Engine/Source/Runtime/Engine/Private/SceneComponent.cpp
+++ b/Engine/Source/Runtime/Engine/Private/SceneComponent.cpp
@@ -1161,9 +1161,9 @@ APhysicsVolume* USceneComponent::GetPhysicsVolume() const
 	{
 		return PhysicsVolume;
 	}
-	else if (GetWorld())
+	else if (const UWorld* MyWorld = GetWorld())
 	{
-		return GetWorld()->GetDefaultPhysicsVolume();
+		return MyWorld->GetDefaultPhysicsVolume();
 	}
 
 	return NULL;
