#ue4 - Add KismetMathLibrary functions ProjectPointOnToPlane(), ProjectVectorOnToPlane(). Rename ProjectOnTo() to ProjectVectorOnToVector().

diff --git a/Engine/Config/BaseEngine.ini b/Engine/Config/BaseEngine.ini
--- a/Engine/Config/BaseEngine.ini
+++ b/Engine/Config/BaseEngine.ini
@@ -1168,6 +1168,10 @@ bUseStreamingPause=false
 +K2FieldRedirects=(OldFieldName="KismetMathLibrary.NearlyEqual_VectorVector",NewFieldName="EqualEqual_VectorVector")
 +K2FieldRedirects=(OldFieldName="KismetMathLibrary.NearlyEqual_RotatorRotator",NewFieldName="EqualEqual_RotatorRotator")
 
++K2FieldRedirects=(OldFieldName="KismetMathLibrary.ProjectOnTo",NewFieldName="ProjectVectorOnToVector")
++K2ParamRedirects=(NodeName="/Script/BlueprintGraph.K2Node_CallFunction", OldParamName="KismetMathLibrary.ProjectVectorOnToVector.X", NewParamName="V")
++K2ParamRedirects=(NodeName="/Script/BlueprintGraph.K2Node_CallFunction", OldParamName="KismetMathLibrary.ProjectVectorOnToVector.Y", NewParamName="Target")
+
 +K2FieldRedirects=(OldFieldName="PrimitiveComponent.SetRBLinearVelocity",NewFieldName="SetPhysicsLinearVelocity")
 +K2FieldRedirects=(OldFieldName="PrimitiveComponent.GetRBLinearVelocity",NewFieldName="GetPhysicsLinearVelocity")
 +K2FieldRedirects=(OldFieldName="PrimitiveComponent.SetAllRBLinearVelocity",NewFieldName="SetAllPhysicsLinearVelocity")
diff --git a/Engine/Source/Runtime/Engine/Classes/Kismet/KismetMathLibrary.h b/Engine/Source/Runtime/Engine/Classes/Kismet/KismetMathLibrary.h
--- a/Engine/Source/Runtime/Engine/Classes/Kismet/KismetMathLibrary.h
+++ b/Engine/Source/Runtime/Engine/Classes/Kismet/KismetMathLibrary.h
@@ -581,9 +581,37 @@ class UKismetMathLibrary : public UBlueprintFunctionLibrary
 	UFUNCTION(BlueprintPure, Category="Math|Vector")
 	static FVector MirrorVectorByNormal(FVector InVect, FVector InNormal);
 
-	// Projects one vector (X) onto another (Y) and returns the projected vector.
-	UFUNCTION(BlueprintPure, Category="Math|Vector" )
-	static FVector ProjectOnTo(FVector X, FVector Y);
+	/**
+	* Projects one vector (V) onto another (Target) and returns the projected vector.
+	* If Target is nearly zero in length, returns the zero vector.
+	*
+	* @param  V Vector to project.
+	* @param  Target Vector on which we are projecting.
+	* @return V projected on to Target.
+	*/
+	UFUNCTION(BlueprintPure, Category="Math|Vector", meta=(Keywords = "ProjectOnTo"))
+	static FVector ProjectVectorOnToVector(FVector V, FVector Target);
+
+	 /**
+	 * Projects a point onto a plane defined by a point on the plane and a plane normal.
+	 *
+	 * @param  Point Point to project onto the plane.
+	 * @param  PlaneBase A point on the plane.
+	 * @param  PlaneNormal Normal of the plane.
+	 * @return Point projected onto the plane.
+	 */
+	UFUNCTION(BlueprintPure, Category = "Math|Vector", meta=(Keywords = "ProjectOnTo"))
+	static FVector ProjectPointOnToPlane(FVector Point, FVector PlaneBase, FVector PlaneNormal);
+
+	/**
+	* Projects a vector onto a plane defined by a normalized vector (PlaneNormal).
+	*
+	* @param  V Vector to project onto the plane.
+	* @param  PlaneNormal Normal of the plane.
+	* @return Vector projected onto the plane.
+	*/
+	UFUNCTION(BlueprintPure, Category="Math|Vector", meta=(Keywords = "ProjectOnTo"))
+	static FVector ProjectVectorOnToPlane(FVector V, FVector PlaneNormal);
 
 	/** Negate a vector. */
 	UFUNCTION(BlueprintPure, Category="Math|Vector")
diff --git a/Engine/Source/Runtime/Engine/Private/KismetMathLibrary.cpp b/Engine/Source/Runtime/Engine/Private/KismetMathLibrary.cpp
--- a/Engine/Source/Runtime/Engine/Private/KismetMathLibrary.cpp
+++ b/Engine/Source/Runtime/Engine/Private/KismetMathLibrary.cpp
@@ -783,9 +783,27 @@ FRotator UKismetMathLibrary::RandomRotator(bool bRoll)
 	return RRot;
 }
 
-FVector UKismetMathLibrary::ProjectOnTo(FVector X, FVector Y)
+FVector UKismetMathLibrary::ProjectVectorOnToVector(FVector V, FVector Target)
 {
-	return X.ProjectOnTo( Y );
+	if (Target.SizeSquared() > SMALL_NUMBER)
+	{
+		return V.ProjectOnTo(Target);
+	}
+	else
+	{
+		FFrame::KismetExecutionMessage(TEXT("Divide by zero: ProjectVectorOnToVector with zero Target vector"), ELogVerbosity::Warning);
+		return FVector::ZeroVector;
+	}
+}
+
+FVector UKismetMathLibrary::ProjectPointOnToPlane(FVector Point, FVector PlaneBase, FVector PlaneNormal)
+{
+	return FVector::PointPlaneProject(Point, PlaneBase, PlaneNormal);
+}
+
+FVector UKismetMathLibrary::ProjectVectorOnToPlane(FVector V, FVector PlaneNormal)
+{
+	return FVector::VectorPlaneProject(V, PlaneNormal);
 }
 
 FVector UKismetMathLibrary::NegateVector(FVector A)
