#ue4 - Fix MovementComponent::ComputeSlideVector() calculating incorrect deflection vector when the plane constraint is enabled.

- Rather than projecting the final result onto the movement plane, the normal should first be projected onto the plane and normalized.
- This fixes cases where collision with geometry angled outside the movement plane would stop movement, rather than deflect it.

diff --git a/Engine/Source/Runtime/Engine/Private/MovementComponent.cpp b/Engine/Source/Runtime/Engine/Private/MovementComponent.cpp
--- a/Engine/Source/Runtime/Engine/Private/MovementComponent.cpp
+++ b/Engine/Source/Runtime/Engine/Private/MovementComponent.cpp
@@ -266,7 +266,7 @@ FVector UMovementComponent::ConstrainDirectionToPlane(FVector Direction) const
 {
 	if (bConstrainToPlane)
 	{
-		Direction = Direction - (PlaneConstraintNormal * (Direction | PlaneConstraintNormal));
+		Direction = FVector::VectorPlaneProject(Direction, PlaneConstraintNormal);
 	}
 
 	return Direction;
@@ -411,7 +411,15 @@ bool UMovementComponent::ResolvePenetration(const FVector& ProposedAdjustment, c
 
 FVector UMovementComponent::ComputeSlideVector(const FVector& Delta, const float Time, const FVector& Normal, const FHitResult& Hit) const
 {
-	return FVector::VectorPlaneProject(Delta, Normal) * Time;
+	if (!bConstrainToPlane)
+	{
+		return FVector::VectorPlaneProject(Delta, Normal) * Time;
+	}
+	else
+	{
+		const FVector ProjectedNormal = ConstrainDirectionToPlane(Normal).SafeNormal();
+		return FVector::VectorPlaneProject(Delta, ProjectedNormal) * Time;
+	}
 }
 
 
