#ue4 - Better AirControl support, allowing movement towards potentially unwalkable surfaces but still restricting boost so you can't wall climb. This allows players to dodge towards vertical walls more effectively, and also steer laterally down unwalkable slopes rather than taking away all air control in these cases.

- GetAirControl() now returns a FVector rather than float, to allow air control code to act non-uniformly (ie different laterally versus forward) if desired.

- GetFallingLateralAcceleration()/GetAirControl() no longer do the sweep test forward in time to determine if air control should be limited. This is now done separately by FindAirControlImpact(), and the results can be modified by LimitAirControl().

- LimitAirControl() now returns a FVector rather than float. The behavior is now to allow acceleration parallel to the surface in the horizontal plane, or away from the surface altogether.
- LimitAirControl() may be called multiple times during a falling movement update.

- HandleSlopeBoosting() now prevents upward boosts no matter the incoming velocity. We used to only limit it if moving upwards, but you can deflect upwards when moving horizontally or downward, which bypassed that check.
- HandleSlopeBoosting() used to send the player straight down if hitting the side of a step. This has been removed since this is now checked in falling movement.

diff --git a/Engine/Source/Runtime/Engine/Classes/GameFramework/CharacterMovementComponent.h b/Engine/Source/Runtime/Engine/Classes/GameFramework/CharacterMovementComponent.h
--- a/Engine/Source/Runtime/Engine/Classes/GameFramework/CharacterMovementComponent.h
+++ b/Engine/Source/Runtime/Engine/Classes/GameFramework/CharacterMovementComponent.h
@@ -997,7 +997,8 @@ public:
 
 	/**
 	 * Get the lateral acceleration to use during falling movement. The Z component of the result is ignored.
-	 * Default implementation returns current Acceleration value modified by GetAirControl(), with Z component removed.
+	 * Default implementation returns current Acceleration value modified by GetAirControl(), with Z component removed,
+	 * with magnitude clamped to GetMaxAcceleration().
 	 * This function is used internally by PhysFalling().
 	 *
 	 * @param DeltaTime Time step for the current update.
@@ -1006,9 +1007,8 @@ public:
 	virtual FVector GetFallingLateralAcceleration(float DeltaTime);
 	
 	/**
-	 * Get the air control to use during falling movement. This is a multiplier applied to the acceleration used when falling.
-	 * Given an initial air control (TickAirControl), applies the result of BoostAirControl(). If non-zero and FindAirControlImpact()
-	 * returns true, we then apply the result of LimitAirControl().
+	 * Get the air control to use during falling movement.
+	 * Given an initial air control (TickAirControl), applies the result of BoostAirControl().
 	 * This function is used internally by GetFallingLateralAcceleration().
 	 *
 	 * @param DeltaTime			Time step for the current update.
@@ -1017,7 +1017,7 @@ public:
 	 * @return Air control to use during falling movement.
 	 * @see AirControl, BoostAirControl(), LimitAirControl(), GetFallingLateralAcceleration()
 	 */
-	virtual float GetAirControl(float DeltaTime, float TickAirControl, const FVector& FallAcceleration);
+	virtual FVector GetAirControl(float DeltaTime, float TickAirControl, const FVector& FallAcceleration);
 
 protected:
 
@@ -1035,29 +1035,28 @@ protected:
 
 	/**
 	 * Checks if air control will cause the player collision shape to hit something given current conditions.
-	 * This function is used internally by GetAirControl().
+	 * This function is used internally by PhysFalling().
 	 *
 	 * @param DeltaTime			Time step for the current update.
-	 * @param TickAirControl	Current air control value.
 	 * @param FallAcceleration	Acceleration used during movement.
 	 * @param OutHitResult		Result of impact, valid if this function returns true.
 	 * @return True if there is an impact, in which case OutHitResult contains the result of that impact.
 	 * @see GetAirControl()
 	 */
-	virtual bool FindAirControlImpact(float DeltaTime, float TickAirControl, const FVector& FallAcceleration, FHitResult& OutHitResult);
+	virtual bool FindAirControlImpact(float DeltaTime, const FVector& FallAcceleration, FHitResult& OutHitResult);
 
 	/**
 	 * Limits the air control to use during falling movement, given an impact from FindAirControlImpact().
-	 * This function is used internally by GetAirControl().
+	 * This function is used internally by PhysFalling().
 	 *
 	 * @param DeltaTime			Time step for the current update.
-	 * @param TickAirControl	Current air control value.
 	 * @param FallAcceleration	Acceleration used during movement.
-	 * @param HitResult			Result of impact from FindAirControlImpact().
-	 * @return Modified air control to use during falling movement
+	 * @param HitResult			Result of impact.
+	 * @param bCheckForValidLandingSpot If true, will use IsValidLandingSpot() to determine if HitResult is a walkable surface. If false, this check is skipped.
+	 * @return Modified air control acceleration to use during falling movement.
 	 * @see FindAirControlImpact()
 	 */
-	virtual float LimitAirControl(float DeltaTime, float TickAirControl, const FVector& FallAcceleration, const FHitResult& HitResult);
+	virtual FVector LimitAirControl(float DeltaTime, const FVector& FallAcceleration, const FHitResult& HitResult, bool bCheckForValidLandingSpot);
 	
 
 	/** Handle landing against Hit surface over remaingTime and iterations, calling SetPostLandedPhysics() and starting the new movement mode. */
diff --git a/Engine/Source/Runtime/Engine/Private/CharacterMovementComponent.cpp b/Engine/Source/Runtime/Engine/Private/CharacterMovementComponent.cpp
--- a/Engine/Source/Runtime/Engine/Private/CharacterMovementComponent.cpp
+++ b/Engine/Source/Runtime/Engine/Private/CharacterMovementComponent.cpp
@@ -22,6 +22,7 @@ DEFINE_LOG_CATEGORY_STATIC(LogCharacterMovement, Log, All);
 // MAGIC NUMBERS
 const float MAX_STEP_SIDE_Z = 0.08f;	// maximum z value for the normal on the vertical side of steps
 const float SWIMBOBSPEED = -80.f;
+const float VERTICAL_SLOPE_NORMAL_Z = -0.005f; // Slope is vertical or upward if Normal.Z >= Thresh. Accounts for precision problems that sometimes angle normals slightly off horizontal for vertical surface.
 
 const float UCharacterMovementComponent::MIN_TICK_TIME = 0.0002f;
 const float UCharacterMovementComponent::MIN_FLOOR_DIST = 1.9f;
@@ -1955,7 +1956,7 @@ FVector UCharacterMovementComponent::ComputeSlideVector(const FVector& Delta, co
 	FVector Result = Super::ComputeSlideVector(Delta, Time, Normal, Hit);
 
 	// prevent boosting up slopes
-	if (Result.Z > 0.f && IsFalling())
+	if (IsFalling())
 	{
 		Result = HandleSlopeBoosting(Result, Delta, Time, Normal, Hit);
 	}
@@ -1970,28 +1971,27 @@ FVector UCharacterMovementComponent::HandleSlopeBoosting(const FVector& SlideRes
 
 	if (Result.Z > 0.f)
 	{
-		if (Delta.Z < 0.f && (Hit.ImpactNormal.Z < MAX_STEP_SIDE_Z))
-		{
-			// We were moving downward, but a slide was going to send us upward. We want to aim
-			// straight down for the next move to make sure we get the most upward-facing opposing normal.
-			Result = FVector(0.f, 0.f, Delta.Z);
-		}
-		else
+		// Don't move any higher than we originally intended.
+		const float ZLimit = Delta.Z * Time;
+		if (Result.Z - ZLimit > KINDA_SMALL_NUMBER)
 		{
-			// Don't move any higher than we originally intended.
-			const float ZLimit = Delta.Z * Time;
-			if (Result.Z - ZLimit > KINDA_SMALL_NUMBER && ZLimit > KINDA_SMALL_NUMBER)
+			if (ZLimit > 0.f)
 			{
 				// Rescale the entire vector (not just the Z component) otherwise we change the direction and likely head right back into the impact.
 				const float UpPercent = ZLimit / Result.Z;
 				Result *= UpPercent;
-
-				// Make remaining portion of original result horizontal and parallel to impact normal.
-				const FVector RemainderXY = (SlideResult - Result) * FVector(1.f, 1.f, 0.f);
-				const FVector NormalXY = Normal.SafeNormal2D();
-				const FVector Adjust = Super::ComputeSlideVector(RemainderXY, 1.f, NormalXY, Hit);
-				Result += Adjust;
 			}
+			else
+			{
+				// We were heading down but were going to deflect upwards. Just make the deflection horizontal.
+				Result = FVector::ZeroVector;
+			}
+
+			// Make remaining portion of original result horizontal and parallel to impact normal.
+			const FVector RemainderXY = (SlideResult - Result) * FVector(1.f, 1.f, 0.f);
+			const FVector NormalXY = Normal.SafeNormal2D();
+			const FVector Adjust = Super::ComputeSlideVector(RemainderXY, 1.f, NormalXY, Hit);
+			Result += Adjust;
 		}
 	}
 
@@ -2774,35 +2774,23 @@ FVector UCharacterMovementComponent::GetFallingLateralAcceleration(float DeltaTi
 	// bound acceleration, falling object has minimal ability to impact acceleration
 	if (!HasRootMotion() && FallAcceleration.SizeSquared2D() > 0.f)
 	{
-		const float TickAirControl = GetAirControl(DeltaTime, AirControl, FallAcceleration);
-
-		const float MaxAccel = GetMaxAcceleration() * TickAirControl;
-		FallAcceleration = FallAcceleration.ClampMaxSize(MaxAccel);
+		FallAcceleration = GetAirControl(DeltaTime, AirControl, FallAcceleration);
+		FallAcceleration = FallAcceleration.ClampMaxSize(GetMaxAcceleration());
 	}
 
 	return FallAcceleration;
 }
 
 
-float UCharacterMovementComponent::GetAirControl(float DeltaTime, float TickAirControl, const FVector& FallAcceleration)
+FVector UCharacterMovementComponent::GetAirControl(float DeltaTime, float TickAirControl, const FVector& FallAcceleration)
 {
 	// Boost
-	if (TickAirControl > 0.f)
+	if (TickAirControl != 0.f)
 	{
 		TickAirControl = BoostAirControl(DeltaTime, TickAirControl, FallAcceleration);
 	}
 
-	// Limit
-	if (TickAirControl > 0.f)
-	{
-		FHitResult HitResult(1.f);
-		if (FindAirControlImpact(DeltaTime, TickAirControl, FallAcceleration, HitResult))
-		{
-			TickAirControl = LimitAirControl(DeltaTime, TickAirControl, FallAcceleration, HitResult);
-		}
-	}
-
-	return TickAirControl;
+	return TickAirControl * FallAcceleration;
 }
 
 
@@ -2827,6 +2815,7 @@ void UCharacterMovementComponent::PhysFalling(float deltaTime, int32 Iterations)
 
 	FVector FallAcceleration = GetFallingLateralAcceleration(deltaTime);
 	FallAcceleration.Z = 0.f;
+	const bool bHasAirControl = (FallAcceleration.SizeSquared2D() > 0.f);
 
 	float remainingTime = deltaTime;
 	while( (remainingTime >= MIN_TICK_TIME) && (Iterations < MaxSimulationIterations) )
@@ -2840,20 +2829,42 @@ void UCharacterMovementComponent::PhysFalling(float deltaTime, int32 Iterations)
 		bJustTeleported = false;
 
 		FVector OldVelocity = Velocity;
+		FVector VelocityNoAirControl = Velocity;
 
 		// Apply input
 		if (!HasRootMotion())
 		{
-			// Acceleration = FallAcceleration for CalcVelocity, but we restore it after using it.
-			TGuardValue<FVector> RestoreAcceleration(Acceleration, FallAcceleration);
-			const float SavedVelZ = Velocity.Z;
-			Velocity.Z = 0.f;
-			CalcVelocity(timeTick, FallingLateralFriction, false, BrakingDecelerationFalling);
-			Velocity.Z = SavedVelZ;
+			// Compute VelocityNoAirControl
+			if (bHasAirControl)
+			{
+				// Find velocity *without* acceleration.
+				TGuardValue<FVector> RestoreAcceleration(Acceleration, FVector::ZeroVector);
+				TGuardValue<FVector> RestoreVelocity(Velocity, Velocity);
+				Velocity.Z = 0.f;
+				CalcVelocity(timeTick, FallingLateralFriction, false, BrakingDecelerationFalling);
+				VelocityNoAirControl = FVector(Velocity.X, Velocity.Y, OldVelocity.Z);
+			}
+
+			// Compute Velocity
+			{
+				// Acceleration = FallAcceleration for CalcVelocity(), but we restore it after using it.
+				TGuardValue<FVector> RestoreAcceleration(Acceleration, FallAcceleration);
+				Velocity.Z = 0.f;
+				CalcVelocity(timeTick, FallingLateralFriction, false, BrakingDecelerationFalling);
+				Velocity.Z = OldVelocity.Z;
+			}
+
+			// Just copy Velocity to VelocityNoAirControl if they are the same (ie no acceleration).
+			if (!bHasAirControl)
+			{
+				VelocityNoAirControl = Velocity;
+			}
 		}
 
 		// Apply gravity
-		Velocity = NewFallVelocity(Velocity, FVector(0.f,0.f,GetGravityZ()), timeTick);
+		const FVector Gravity(0.f, 0.f, GetGravityZ());
+		Velocity = NewFallVelocity(Velocity, Gravity, timeTick);
+		VelocityNoAirControl = NewFallVelocity(VelocityNoAirControl, Gravity, timeTick);
 
 		if( bNotifyApex && CharacterOwner->Controller && (Velocity.Z <= 0.f) )
 		{
@@ -2862,9 +2873,22 @@ void UCharacterMovementComponent::PhysFalling(float deltaTime, int32 Iterations)
 			NotifyJumpApex();
 		}
 
+		// Test air control to see if we should limit it.
+		const FVector AirControlAccel = (Velocity - VelocityNoAirControl) / timeTick;
+		if (bHasAirControl)
+		{
+			FHitResult TestAirResult(1.f);
+			if (FindAirControlImpact(timeTick, AirControlAccel, TestAirResult))
+			{
+				const FVector LimitedAirAccel = LimitAirControl(timeTick, AirControlAccel, TestAirResult, /*bCheckForValidLandingSpot=*/ true);
+				Velocity = VelocityNoAirControl + (LimitedAirAccel * timeTick);
+			}
+		}
+
+
 		// Move
 		FHitResult Hit(1.f);
-		FVector Adjusted = 0.5f*(OldVelocity + Velocity) * timeTick;  
+		FVector Adjusted = 0.5f*(OldVelocity + Velocity) * timeTick;
 		SafeMoveUpdatedComponent( Adjusted, PawnRotation, true, Hit);
 		
 		if (!HasValidData())
@@ -2881,7 +2905,7 @@ void UCharacterMovementComponent::PhysFalling(float deltaTime, int32 Iterations)
 			StartSwimming(OldLocation, OldVelocity, timeTick, remainingTime, Iterations);
 			return;
 		}
-		else if ( Hit.Time < 1.f )
+		else if ( Hit.bBlockingHit )
 		{
 			if (IsValidLandingSpot(UpdatedComponent->GetComponentLocation(), Hit))
 			{
@@ -2917,6 +2941,15 @@ void UCharacterMovementComponent::PhysFalling(float deltaTime, int32 Iterations)
 					return;
 				}
 
+				// Limit air control based on what we hit.
+				// We moved to the impact point using air control, but may want to deflect from there based on a non-air control acceleration.
+				if (bHasAirControl)
+				{
+					const bool bCheckLandingSpot = false; // we already checked above.
+					const FVector AirControlDeltaV = LimitAirControl(LastMoveTimeSlice, AirControlAccel, Hit, bCheckLandingSpot) * LastMoveTimeSlice;
+					Adjusted = (VelocityNoAirControl + AirControlDeltaV) * LastMoveTimeSlice;
+				}
+
 				const FVector OldHitNormal = Hit.Normal;
 				const FVector OldHitImpactNormal = Hit.ImpactNormal;				
 				FVector Delta = ComputeSlideVector(Adjusted, 1.f - Hit.Time, OldHitNormal, Hit);
@@ -2933,7 +2966,7 @@ void UCharacterMovementComponent::PhysFalling(float deltaTime, int32 Iterations)
 					// Move in deflected direction.
 					SafeMoveUpdatedComponent( Delta, PawnRotation, true, Hit);
 					
-					if (Hit.Time < 1.f)
+					if (Hit.bBlockingHit)
 					{
 						// hit second wall
 						LastMoveTimeSlice = subTimeTickRemaining;
@@ -2954,8 +2987,29 @@ void UCharacterMovementComponent::PhysFalling(float deltaTime, int32 Iterations)
 							return;
 						}
 
-						const FVector PreTwoWallDelta = Delta;
+						// Act as if there was no air control on the last move when computing new deflection.
+						if (bHasAirControl && Hit.Normal.Z >= VERTICAL_SLOPE_NORMAL_Z)
+						{
+							const FVector LastMoveNoAirControl = VelocityNoAirControl * LastMoveTimeSlice;
+							Delta = ComputeSlideVector(LastMoveNoAirControl, 1.f, OldHitNormal, Hit);
+						}
+
+						FVector PreTwoWallDelta = Delta;
 						TwoWallAdjust(Delta, Hit, OldHitNormal);
+
+						// Limit air control, but allow a slide along the second wall.
+						if (bHasAirControl)
+						{
+							const bool bCheckLandingSpot = false; // we already checked above.
+							const FVector AirControlDeltaV = LimitAirControl(subTimeTickRemaining, AirControlAccel, Hit, bCheckLandingSpot) * subTimeTickRemaining;
+
+							// Only allow if not back in to first wall
+							if (FVector::DotProduct(AirControlDeltaV, OldHitNormal) > 0.f)
+							{
+								Delta += (AirControlDeltaV * subTimeTickRemaining);
+							}
+						}
+
 						if (Delta.Z > 0.f)
 						{
 							Delta = HandleSlopeBoosting(Delta, PreTwoWallDelta, 1.f - Hit.Time, Hit.Normal, Hit);
@@ -3017,11 +3071,12 @@ void UCharacterMovementComponent::PhysFalling(float deltaTime, int32 Iterations)
 }
 
 
-bool UCharacterMovementComponent::FindAirControlImpact(float DeltaTime, float TickAirControl, const FVector& FallAcceleration, FHitResult& OutHitResult)
+bool UCharacterMovementComponent::FindAirControlImpact(float DeltaTime, const FVector& FallAcceleration, FHitResult& OutHitResult)
 {
 	// test for slope to avoid using air control to climb walls
-	const float TestWalkTime = FMath::Max(DeltaTime, 0.05f);
-	const FVector TestWalk = ((TickAirControl * GetMaxAcceleration() * FallAcceleration.SafeNormal() + FVector(0.f,0.f,GetGravityZ())) * TestWalkTime + Velocity) * TestWalkTime;
+	const float TestWalkTime = FMath::Max(DeltaTime, 0.01f);
+	const FVector PostGravityVelocity = NewFallVelocity(Velocity, FVector(0.f,0.f,GetGravityZ()), TestWalkTime);
+	const FVector TestWalk = ((FallAcceleration * TestWalkTime) + PostGravityVelocity) * TestWalkTime;
 	if (!TestWalk.IsZero())
 	{
 		static const FName FallingTraceParamsTag = FName(TEXT("PhysFalling"));
@@ -3041,17 +3096,30 @@ bool UCharacterMovementComponent::FindAirControlImpact(float DeltaTime, float Ti
 }
 
 
-float UCharacterMovementComponent::LimitAirControl(float DeltaTime, float TickAirControl, const FVector& FallAcceleration, const FHitResult& HitResult)
+FVector UCharacterMovementComponent::LimitAirControl(float DeltaTime, const FVector& FallAcceleration, const FHitResult& HitResult, bool bCheckForValidLandingSpot)
 {
-	if (HitResult.bBlockingHit)
+	FVector Result(FallAcceleration);
+
+	if (HitResult.IsValidBlockingHit() && HitResult.Normal.Z >= VERTICAL_SLOPE_NORMAL_Z)
 	{
-		if (!IsValidLandingSpot(HitResult.Location, HitResult))
+		if (!bCheckForValidLandingSpot || !IsValidLandingSpot(HitResult.Location, HitResult))
 		{
-			TickAirControl = 0.f;
+			// If acceleration is into the wall, limit contribution.
+			if (FVector::DotProduct(FallAcceleration, HitResult.Normal) < 0.f)
+			{
+				// Allow movement parallel to the wall, but not into it because that may push us up.
+				const FVector Normal2D = HitResult.Normal.SafeNormal2D();
+				Result = FVector::VectorPlaneProject(FallAcceleration, Normal2D);
+			}
 		}
 	}
+	else if (HitResult.bStartPenetrating)
+	{
+		// Allow movement out of penetration.
+		return (FVector::DotProduct(Result, HitResult.Normal) > 0.f ? Result : FVector::ZeroVector);
+	}
 
-	return TickAirControl;
+	return Result;
 }
 
 bool UCharacterMovementComponent::CheckLedgeDirection(const FVector& OldLocation, const FVector& SideStep, const FVector& GravDir)
