Fix retargeting issue when missing bones. (TTP 349830)

diff --git a/Engine/Source/Runtime/Engine/Private/Animation/AnimationRuntime.cpp b/Engine/Source/Runtime/Engine/Private/Animation/AnimationRuntime.cpp
--- a/Engine/Source/Runtime/Engine/Private/Animation/AnimationRuntime.cpp
+++ b/Engine/Source/Runtime/Engine/Private/Animation/AnimationRuntime.cpp
@@ -1249,7 +1249,7 @@ void FAnimationRuntime::FillUpSpaceBasesRetargetBasePose(const USkeleton* Skelet
 		// initialize to identity since some of them don't have tracks
 		for(int Index=0; Index <SpaceBaseRefPose.Num(); ++Index)
 		{
-			int32 ParentIndex = Skeleton->GetReferenceSkeleton().GetParentIndex(Index);
+			int32 ParentIndex = PreviewMesh->RefSkeleton.GetParentIndex(Index);
 			if(ParentIndex != INDEX_NONE)
 			{
 				SpaceBaseRefPose[Index] = ReferencePose[Index] * SpaceBaseRefPose[ParentIndex];
