Add debug data for retargeting

diff --git a/Engine/Source/Runtime/Engine/Private/Animation/AnimSequence.cpp b/Engine/Source/Runtime/Engine/Private/Animation/AnimSequence.cpp
--- a/Engine/Source/Runtime/Engine/Private/Animation/AnimSequence.cpp
+++ b/Engine/Source/Runtime/Engine/Private/Animation/AnimSequence.cpp
@@ -1862,6 +1862,7 @@ void UAnimSequence::RemapTracksToNewSkeleton( USkeleton* NewSkeleton, bool bConv
 				// saves the translation conversion datao
 				OldToNewTranslationRatio.AddUninitialized(NumNodes);
 
+				const TArray<FNode> & Nodes = Rig->GetNodes();
 				// calculate the relative transform to new skeleton
 				// so that we can apply the delta in component space
 				for (int32 NodeIndex = 0; NodeIndex < NumNodes; ++NodeIndex)
@@ -1876,6 +1877,10 @@ void UAnimSequence::RemapTracksToNewSkeleton( USkeleton* NewSkeleton, bool bConv
 					float OldTranslationSize = OldTranslation.Size();
 					float NewTranslationSize = NewTranslation.Size();
 					OldToNewTranslationRatio[NodeIndex] = (FMath::IsNearlyZero(OldTranslationSize)) ? 1.f/*do not touch new translation size*/ : NewTranslationSize / OldTranslationSize;
+
+					UE_LOG(LogAnimation, Verbose, TEXT("Retargeting (%s : %d) : OldtoNewTranslationRatio (%0.2f), Relative Transform (%s)"), *Nodes[NodeIndex].Name.ToString(), NodeIndex, 
+						OldToNewTranslationRatio[NodeIndex], *RelativeToNewSpaceBases[NodeIndex].ToString());
+					UE_LOG(LogAnimation, Verbose, TEXT("\tOldSpaceBase(%s), NewSpaceBase(%s)"), *OldSpaceBases[NodeIndex].ToString(), *NewSpaceBases[NodeIndex].ToString());
 				}
 			}
 
