#ue4 - Don't recalculate Character velocity when falling if there was a teleport (usually adjustment out of penetration).

diff --git a/Engine/Source/Runtime/Engine/Private/CharacterMovementComponent.cpp b/Engine/Source/Runtime/Engine/Private/CharacterMovementComponent.cpp
--- a/Engine/Source/Runtime/Engine/Private/CharacterMovementComponent.cpp
+++ b/Engine/Source/Runtime/Engine/Private/CharacterMovementComponent.cpp
@@ -2955,7 +2955,7 @@ void UCharacterMovementComponent::PhysFalling(float deltaTime, int32 Iterations)
 				FVector Delta = ComputeSlideVector(Adjusted, 1.f - Hit.Time, OldHitNormal, Hit);
 
 				// Compute velocity after deflection (only gravity component for RootMotion)
-				if (subTimeTickRemaining > KINDA_SMALL_NUMBER)
+				if (subTimeTickRemaining > KINDA_SMALL_NUMBER && !bJustTeleported)
 				{
 					const FVector NewVelocity = (Delta / subTimeTickRemaining);
 					Velocity = HasRootMotion() ? FVector(Velocity.X, Velocity.Y, NewVelocity.Z) : NewVelocity;
@@ -3016,7 +3016,7 @@ void UCharacterMovementComponent::PhysFalling(float deltaTime, int32 Iterations)
 						}
 
 						// Compute velocity after deflection (only gravity component for RootMotion)
-						if (subTimeTickRemaining > KINDA_SMALL_NUMBER)
+						if (subTimeTickRemaining > KINDA_SMALL_NUMBER && !bJustTeleported)
 						{
 							const FVector NewVelocity = (Delta / subTimeTickRemaining);
 							Velocity = HasRootMotion() ? FVector(Velocity.X, Velocity.Y, NewVelocity.Z) : NewVelocity;
