#ue4 - Avoid heap allocation in single sweep and raycast queries when collision analyzer isn't enabled.

diff --git a/Engine/Source/Runtime/Engine/Private/Collision/PhysXCollision.cpp b/Engine/Source/Runtime/Engine/Private/Collision/PhysXCollision.cpp
--- a/Engine/Source/Runtime/Engine/Private/Collision/PhysXCollision.cpp
+++ b/Engine/Source/Runtime/Engine/Private/Collision/PhysXCollision.cpp
@@ -914,20 +914,28 @@ bool RaycastSingle(const UWorld* World, struct FHitResult& OutHit, const FVector
 	}
 
 
-#if ENABLE_COLLISION_ANALYZER || !(UE_BUILD_SHIPPING || UE_BUILD_TEST)
-	TArray<FHitResult> Hits;
-	if (bHaveBlockingHit)
-	{
-		Hits.Add(OutHit);
-	}
 #if !(UE_BUILD_SHIPPING || UE_BUILD_TEST)
 	if ((World->DebugDrawTraceTag != NAME_None) && (World->DebugDrawTraceTag == Params.TraceTag))
 	{
+		TArray<FHitResult> Hits;
+		if (bHaveBlockingHit)
+		{
+			Hits.Add(OutHit);
+		}
 		DrawLineTraces(World, Start, End, Hits, DebugLineLifetime);	
 	}
 #endif //!(UE_BUILD_SHIPPING || UE_BUILD_TEST)
 
-	CAPTURERAYCAST(World, Start, End, ECAQueryType::Single, TraceChannel, Params, ResponseParams, ObjectParams, Hits);
+#if ENABLE_COLLISION_ANALYZER
+	if (GCollisionAnalyzerIsRecording)
+	{
+		TArray<FHitResult> Hits;
+		if (bHaveBlockingHit)
+		{
+			Hits.Add(OutHit);
+		}
+		CAPTURERAYCAST(World, Start, End, ECAQueryType::Single, TraceChannel, Params, ResponseParams, ObjectParams, Hits);
+	}
 #endif
 
 	return bHaveBlockingHit;
@@ -1250,22 +1258,30 @@ bool GeomSweepSingle(const UWorld* World, const struct FCollisionShape& Collisio
 	}
 
 
-#if ENABLE_COLLISION_ANALYZER || !(UE_BUILD_SHIPPING || UE_BUILD_TEST)
-	TArray<FHitResult> Hits;
-	if (bHaveBlockingHit)
-	{
-		Hits.Add(OutHit);
-	}
-
 #if !(UE_BUILD_SHIPPING || UE_BUILD_TEST)
 	if((World->DebugDrawTraceTag != NAME_None) && (World->DebugDrawTraceTag == Params.TraceTag))
 	{
+		TArray<FHitResult> Hits;
+		if (bHaveBlockingHit)
+		{
+			Hits.Add(OutHit);
+		}
 		DrawGeomSweeps(World, Start, End, PGeom, PGeomRot, Hits, DebugLineLifetime);
 	}
 #endif //!(UE_BUILD_SHIPPING || UE_BUILD_TEST)
 
-	CAPTUREGEOMSWEEP(World, Start, End, PGeomRot, ECAQueryType::Single, PGeom, TraceChannel, Params, ResponseParams, ObjectParams, Hits);
+#if ENABLE_COLLISION_ANALYZER
+	if (GCollisionAnalyzerIsRecording)
+	{
+		TArray<FHitResult> Hits;
+		if (bHaveBlockingHit)
+		{
+			Hits.Add(OutHit);
+		}
+		CAPTUREGEOMSWEEP(World, Start, End, PGeomRot, ECAQueryType::Single, PGeom, TraceChannel, Params, ResponseParams, ObjectParams, Hits);
+	}
 #endif
+
 #endif // WITH_PHYSX
 
 	//@TODO: BOX2D: Implement GeomSweepSingle
