#ue4 - Improve depenetration vector when overlapping trimesh and heightfields.

- Previously we tried using the supplied impact normal when face index was valid, but this meant we could be inside an edge and try to move along the normal, rather than away from the edge. Now we just use MTD with an inflated test to get the depentration vector.

diff --git a/Engine/Source/Runtime/Engine/Private/Collision/CollisionConversions.cpp b/Engine/Source/Runtime/Engine/Private/Collision/CollisionConversions.cpp
--- a/Engine/Source/Runtime/Engine/Private/Collision/CollisionConversions.cpp
+++ b/Engine/Source/Runtime/Engine/Private/Collision/CollisionConversions.cpp
@@ -632,11 +632,7 @@ static bool ConvertOverlappedShapeToImpactHit(const PxLocationHit& PHit, const F
 
 		if (GeometryType == PxGeometryType::eTRIANGLEMESH || GeometryType == PxGeometryType::eHEIGHTFIELD)
 		{
-			if (FindGeomOpposingNormal(PHit, TraceDir, OutResult.ImpactNormal))
-			{
-				// Success
-			}
-			else if (ComputeInflatedMTD(MtdInflation, PHit, OutResult, QueryTM, Geom, PShapeWorldPose))
+			if (ComputeInflatedMTD(MtdInflation, PHit, OutResult, QueryTM, Geom, PShapeWorldPose))
 			{
 				// Success
 			}
