#ue4 - Improve Pawn physics penetration handling when overlapping more than one object/surface.

diff --git a/Engine/Source/Runtime/Engine/Private/MovementComponent.cpp b/Engine/Source/Runtime/Engine/Private/MovementComponent.cpp
--- a/Engine/Source/Runtime/Engine/Private/MovementComponent.cpp
+++ b/Engine/Source/Runtime/Engine/Private/MovementComponent.cpp
@@ -387,9 +387,24 @@ bool UMovementComponent::ResolvePenetration(const FVector& ProposedAdjustment, c
 			TGuardValue<EMoveComponentFlags> ScopedFlagRestore(MoveComponentFlags, EMoveComponentFlags(MoveComponentFlags & (~MOVECOMP_NeverIgnoreBlockingOverlaps)));
 
 			// Try sweeping as far as possible...
-			bool bMoved = MoveUpdatedComponent(Adjustment, NewRotation, true);
+			FHitResult SweepOutHit(1.f);
+			bool bMoved = MoveUpdatedComponent(Adjustment, NewRotation, true, &SweepOutHit);
 			UE_LOG(LogMovement, Verbose, TEXT("ResolvePenetration: sweep %s by %s (success = %d)"), *ActorOwner->GetName(), *Adjustment.ToString(), bMoved);
 			
+			// Still stuck?
+			if (!bMoved && SweepOutHit.bStartPenetrating)
+			{
+				// Combine two MTD results to get a new direction that gets out of multiple surfaces.
+				const FVector SecondMTD = GetPenetrationAdjustment(SweepOutHit);
+				const FVector CombinedMTD = Adjustment + SecondMTD;
+				if (SecondMTD != Adjustment && !CombinedMTD.IsZero())
+				{
+					bMoved = MoveUpdatedComponent(CombinedMTD, NewRotation, true);
+					UE_LOG(LogMovement, Verbose, TEXT("ResolvePenetration: sweep %s by %s (MTD combo success = %d)"), *ActorOwner->GetName(), *CombinedMTD.ToString(), bMoved);
+				}
+			}
+
+			// Still stuck?
 			if (!bMoved)
 			{
 				// Try moving the proposed adjustment plus the attempted move direction. This can sometimes get out of penetrations with multiple objects
