Fix case where destroying bodyinstance doesn't properly unweld from parent

diff --git a/Engine/Source/Runtime/Engine/Private/PhysicsEngine/BodyInstance.cpp b/Engine/Source/Runtime/Engine/Private/PhysicsEngine/BodyInstance.cpp
--- a/Engine/Source/Runtime/Engine/Private/PhysicsEngine/BodyInstance.cpp
+++ b/Engine/Source/Runtime/Engine/Private/PhysicsEngine/BodyInstance.cpp
@@ -1498,9 +1498,10 @@ void TermBodyHelper(int32& SceneIndex, PxRigidActor*& PRigidActor, FBodyInstance
  */
 void FBodyInstance::TermBody()
 {
-
 	if (UPrimitiveComponent* OwnerComponentInst = OwnerComponent.Get())
 	{
+		OwnerComponentInst->UnWeldFromParent();
+
 		for (USceneComponent * Child : OwnerComponentInst->AttachChildren)
 		{
 			if (UPrimitiveComponent * PrimChild = Cast<UPrimitiveComponent>(Child))
