Fixed debug labels of ANavigationTestingActor drawing on scree even if actor's instance is out of the frustum #UE4

diff --git a/Engine/Source/Runtime/Engine/Private/AI/Navigation/NavTestRenderingComponent.cpp b/Engine/Source/Runtime/Engine/Private/AI/Navigation/NavTestRenderingComponent.cpp
--- a/Engine/Source/Runtime/Engine/Private/AI/Navigation/NavTestRenderingComponent.cpp
+++ b/Engine/Source/Runtime/Engine/Private/AI/Navigation/NavTestRenderingComponent.cpp
@@ -408,6 +408,11 @@ public:
 #endif
 	}
 
+	FORCEINLINE bool LocationInView(const FVector& Location, const FSceneView* View)
+	{
+		return View->ViewFrustum.IntersectBox(Location, FVector::ZeroVector);
+	}
+
 	void DrawDebugLabels(UCanvas* Canvas, APlayerController*)
 	{
 		if (NavTestActor == NULL)
@@ -417,6 +422,7 @@ public:
 
 		const FColor OldDrawColor = Canvas->DrawColor;
 		Canvas->SetDrawColor(FColor::White);
+		const FSceneView* View = Canvas->SceneView;
 
 #if WITH_EDITORONLY_DATA && WITH_RECAST
 		if (NodeDebug.Num())
@@ -426,20 +432,23 @@ public:
 			{
 				const FNodeDebugData& NodeData = *It;
 
-				FColor MyColor = NodeData.bClosedSet ? FColor(64,64,64) : FColor::White;
-				if (!bShowBestPath && It.GetId() == BestNodeId)
+				if (LocationInView(NodeData.Position, View))
 				{
-					MyColor = FColor::Red;
-				}
-				if (bShowDiff)
-				{
-					MyColor.A = NodeData.bModified ? NavMeshRenderAlpha_Modifed : NavMeshRenderAlpha_NonModified;
-				}
+					FColor MyColor = NodeData.bClosedSet ? FColor(64, 64, 64) : FColor::White;
+					if (!bShowBestPath && It.GetId() == BestNodeId)
+					{
+						MyColor = FColor::Red;
+					}
+					if (bShowDiff)
+					{
+						MyColor.A = NodeData.bModified ? NavMeshRenderAlpha_Modifed : NavMeshRenderAlpha_NonModified;
+					}
 
-				Canvas->SetDrawColor(MyColor);
+					Canvas->SetDrawColor(MyColor);
 
-				const FVector ScreenLoc = Canvas->Project(NodeData.Position) + FVector( NavTestActor->TextCanvasOffset, 0.f );
-				Canvas->DrawText(RenderFont, NodeData.Desc, ScreenLoc.X, ScreenLoc.Y);			
+					const FVector ScreenLoc = Canvas->Project(NodeData.Position) + FVector(NavTestActor->TextCanvasOffset, 0.f);
+					Canvas->DrawText(RenderFont, NodeData.Desc, ScreenLoc.X, ScreenLoc.Y);
+				}
 			}
 		}
 		else
@@ -447,9 +456,12 @@ public:
 #endif
 			for (int32 PointIndex = 0; PointIndex < PathPoints.Num(); ++PointIndex)
 			{
-				const FVector PathPointLoc = Canvas->Project(PathPoints[PointIndex]);
-				UFont* RenderFont = GEngine->GetSmallFont();
-				Canvas->DrawText(RenderFont, PathPointFlags[PointIndex], PathPointLoc.X, PathPointLoc.Y);
+				if (LocationInView(PathPoints[PointIndex], View))
+				{
+					const FVector PathPointLoc = Canvas->Project(PathPoints[PointIndex]);
+					UFont* RenderFont = GEngine->GetSmallFont();
+					Canvas->DrawText(RenderFont, PathPointFlags[PointIndex], PathPointLoc.X, PathPointLoc.Y);
+				}
 			}
 
 #if WITH_EDITORONLY_DATA && WITH_RECAST
