Fixed FPImplRecastNavMesh::Raycast2D not handling start location being not on navmesh #UE4

diff --git a/Engine/Source/Runtime/Engine/Private/AI/Navigation/PImplRecastNavMesh.cpp b/Engine/Source/Runtime/Engine/Private/AI/Navigation/PImplRecastNavMesh.cpp
--- a/Engine/Source/Runtime/Engine/Private/AI/Navigation/PImplRecastNavMesh.cpp
+++ b/Engine/Source/Runtime/Engine/Private/AI/Navigation/PImplRecastNavMesh.cpp
@@ -673,6 +673,12 @@ void FPImplRecastNavMesh::Raycast2D(const FVector& StartLoc, const FVector& EndL
 			UE_VLOG(NavMeshOwner, LogNavigation, Log, TEXT("FPImplRecastNavMesh::Raycast2D failed"));
 		}
 	}
+	else
+	{
+		// start location is not on navmesh, treat it as a blocked raycast
+		RaycastResult.HitTime = 0.f;
+		RaycastResult.HitNormal = (StartLoc - EndLoc).SafeNormal();
+	}
 }
 
 void FPImplRecastNavMesh::Raycast2D(NavNodeRef StartNode, const FVector& StartLoc, const FVector& EndLoc, const FNavigationQueryFilter& InQueryFilter, const UObject* Owner, ARecastNavMesh::FRaycastResult& RaycastResult) const
@@ -706,6 +712,12 @@ void FPImplRecastNavMesh::Raycast2D(NavNodeRef StartNode, const FVector& StartLo
 			UE_VLOG(NavMeshOwner, LogNavigation, Log, TEXT("FPImplRecastNavMesh::Raycast2D failed"));
 		}
 	}
+	else
+	{
+		// start location is not on navmesh, treat it as a blocked raycast
+		RaycastResult.HitTime = 0.f;
+		RaycastResult.HitNormal = (StartLoc - EndLoc).SafeNormal();
+	}
 }
 
 // @TODONAV
