Fixed world composition initialization while traveling to during PIE

diff --git a/Engine/Source/Runtime/Engine/Classes/Engine/World.h b/Engine/Source/Runtime/Engine/Classes/Engine/World.h
--- a/Engine/Source/Runtime/Engine/Classes/Engine/World.h
+++ b/Engine/Source/Runtime/Engine/Classes/Engine/World.h
@@ -2450,6 +2450,7 @@ public:
 	void GetLightMapsAndShadowMaps(ULevel* Level, TArray<UTexture2D*>& OutLightMapsAndShadowMaps);
 
 public:
+	void RenameToPIEWorld(int32 PIEInstanceID);
 	static FString ConvertToPIEPackageName(const FString& PackageName, int32 PIEInstanceID);
 	static FString StripPIEPrefixFromPackageName(const FString& PackageName, const FString& Prefix);
 	static FString BuildPIEPackagePrefix(int32 PIEInstanceID);
diff --git a/Engine/Source/Runtime/Engine/Classes/Engine/WorldComposition.h b/Engine/Source/Runtime/Engine/Classes/Engine/WorldComposition.h
--- a/Engine/Source/Runtime/Engine/Classes/Engine/WorldComposition.h
+++ b/Engine/Source/Runtime/Engine/Classes/Engine/WorldComposition.h
@@ -110,6 +110,9 @@ class ENGINE_API UWorldComposition : public UObject
 	/** Scans world root folder for relevant packages and initializes world composition structures */
 	void Rescan();
 
+	/**  */
+	void ReinitializeForPIE();
+
 #if WITH_EDITOR
 	/** @returns FWorldTileInfo associated with specified package */
 	FWorldTileInfo GetTileInfo(const FName& InPackageName) const;
diff --git a/Engine/Source/Runtime/Engine/Private/UnrealEngine.cpp b/Engine/Source/Runtime/Engine/Private/UnrealEngine.cpp
--- a/Engine/Source/Runtime/Engine/Private/UnrealEngine.cpp
+++ b/Engine/Source/Runtime/Engine/Private/UnrealEngine.cpp
@@ -8819,18 +8819,7 @@ bool UEngine::LoadMap( FWorldContext& WorldContext, FURL URL, class UPendingNetG
 			// Otherwise we are probably loading new map while in PIE, so we need to rename world package and all streaming levels
 			else if (Pending == NULL)
 			{
-#if WITH_EDITOR
-				WorldPackage->PIEInstanceID = WorldContext.PIEInstance;
-#endif				
-				const FString PIEPackageName = *UWorld::ConvertToPIEPackageName(WorldPackage->GetName(), WorldContext.PIEInstance);
-			
-				WorldPackage->Rename(*PIEPackageName);
-				for (int32 StreamingLevelIdx = 0; StreamingLevelIdx < NewWorld->StreamingLevels.Num(); StreamingLevelIdx++)
-				{
-					NewWorld->StreamingLevels[StreamingLevelIdx]->RenameForPIE(WorldContext.PIEInstance);
-				}
-				
-				NewWorld->StreamingLevelsPrefix = UWorld::BuildPIEPackagePrefix(WorldContext.PIEInstance);
+				NewWorld->RenameToPIEWorld(WorldContext.PIEInstance);
 			}
 		}
 	}
diff --git a/Engine/Source/Runtime/Engine/Private/World.cpp b/Engine/Source/Runtime/Engine/Private/World.cpp
--- a/Engine/Source/Runtime/Engine/Private/World.cpp
+++ b/Engine/Source/Runtime/Engine/Private/World.cpp
@@ -1977,6 +1977,29 @@ int32 FLevelStreamingGCHelper::GetNumLevelsPendingPurge()
 	return LevelsPendingUnload.Num();
 }
 
+void UWorld::RenameToPIEWorld(int32 PIEInstanceID)
+{
+	UPackage* WorldPackage = GetOutermost();
+
+#if WITH_EDITOR
+	WorldPackage->PIEInstanceID = PIEInstanceID;
+#endif
+	const FString PIEPackageName = *UWorld::ConvertToPIEPackageName(WorldPackage->GetName(), PIEInstanceID);
+	WorldPackage->Rename(*PIEPackageName);
+	
+	for (ULevelStreaming* LevelStreaming : StreamingLevels)
+	{
+		LevelStreaming->RenameForPIE(PIEInstanceID);
+	}
+				
+	StreamingLevelsPrefix = UWorld::BuildPIEPackagePrefix(PIEInstanceID);
+
+	if (WorldComposition)
+	{
+		WorldComposition->ReinitializeForPIE();
+	}
+}
+
 FString UWorld::ConvertToPIEPackageName(const FString& PackageName, int32 PIEInstanceID)
 {
 	const FString PackageAssetName = FPackageName::GetLongPackageAssetName(PackageName);
diff --git a/Engine/Source/Runtime/Engine/Private/WorldComposition.cpp b/Engine/Source/Runtime/Engine/Private/WorldComposition.cpp
--- a/Engine/Source/Runtime/Engine/Private/WorldComposition.cpp
+++ b/Engine/Source/Runtime/Engine/Private/WorldComposition.cpp
@@ -185,16 +185,17 @@ void UWorldComposition::Rescan()
 		
 	Reset();	
 
+	UWorld* OwningWorld = GetWorld();
+	
 	FString RootPackageName = GetOutermost()->GetName();
+	RootPackageName = UWorld::StripPIEPrefixFromPackageName(RootPackageName, OwningWorld->StreamingLevelsPrefix);
 	if (!FPackageName::DoesPackageExist(RootPackageName))
 	{
 		return;	
 	}
 	
 	WorldRoot = FPaths::GetPath(RootPackageName) + TEXT("/");
-
-	UWorld* OwningWorld = GetWorld();
-		
+			
 	// Gather tiles packages from a specified folder
 	FWorldTilesGatherer Gatherer;
 	FString WorldRootFilename = FPackageName::LongPackageNameToFilename(WorldRoot);
@@ -262,6 +263,14 @@ void UWorldComposition::Rescan()
 	CaclulateTilesAbsolutePositions();
 }
 
+void UWorldComposition::ReinitializeForPIE()
+{
+	Rescan();
+	FixupForPIE(GetOutermost()->PIEInstanceID);
+	GetWorld()->StreamingLevels.Empty();
+	GetWorld()->StreamingLevels.Append(TilesStreaming);
+}
+
 ULevelStreaming* UWorldComposition::CreateStreamingLevel(const FWorldCompositionTile& InTile) const
 {
 	UWorld* OwningWorld = GetWorld();
