CIS fix

diff --git a/Engine/Source/Runtime/Engine/Private/WorldComposition.cpp b/Engine/Source/Runtime/Engine/Private/WorldComposition.cpp
--- a/Engine/Source/Runtime/Engine/Private/WorldComposition.cpp
+++ b/Engine/Source/Runtime/Engine/Private/WorldComposition.cpp
@@ -266,7 +266,9 @@ void UWorldComposition::Rescan()
 void UWorldComposition::ReinitializeForPIE()
 {
 	Rescan();
+#if WITH_EDITOR
 	FixupForPIE(GetOutermost()->PIEInstanceID);
+#endif// WITH_EDITOR
 	GetWorld()->StreamingLevels.Empty();
 	GetWorld()->StreamingLevels.Append(TilesStreaming);
 }
