More fixes to world composition traveling during PIE

diff --git a/Engine/Source/Runtime/Engine/Private/World.cpp b/Engine/Source/Runtime/Engine/Private/World.cpp
--- a/Engine/Source/Runtime/Engine/Private/World.cpp
+++ b/Engine/Source/Runtime/Engine/Private/World.cpp
@@ -1986,18 +1986,18 @@ void UWorld::RenameToPIEWorld(int32 PIEInstanceID)
 #endif
 	const FString PIEPackageName = *UWorld::ConvertToPIEPackageName(WorldPackage->GetName(), PIEInstanceID);
 	WorldPackage->Rename(*PIEPackageName);
-	
-	for (ULevelStreaming* LevelStreaming : StreamingLevels)
-	{
-		LevelStreaming->RenameForPIE(PIEInstanceID);
-	}
-				
-	StreamingLevelsPrefix = UWorld::BuildPIEPackagePrefix(PIEInstanceID);
 
+	StreamingLevelsPrefix = UWorld::BuildPIEPackagePrefix(PIEInstanceID);
+	
 	if (WorldComposition)
 	{
 		WorldComposition->ReinitializeForPIE();
 	}
+	
+	for (ULevelStreaming* LevelStreaming : StreamingLevels)
+	{
+		LevelStreaming->RenameForPIE(PIEInstanceID);
+	}
 }
 
 FString UWorld::ConvertToPIEPackageName(const FString& PackageName, int32 PIEInstanceID)
