better UI

diff --git a/Engine/Source/Runtime/Engine/Classes/Engine/Scene.h b/Engine/Source/Runtime/Engine/Classes/Engine/Scene.h
--- a/Engine/Source/Runtime/Engine/Classes/Engine/Scene.h
+++ b/Engine/Source/Runtime/Engine/Classes/Engine/Scene.h
@@ -638,7 +638,7 @@ struct FPostProcessSettings
 	float AmbientOcclusionDistance;
 
 	/** >0, in unreal units, bigger values means even distant surfaces affect the ambient occlusion */
-	UPROPERTY(interp, BlueprintReadWrite, Category=AmbientOcclusion, AdvancedDisplay, meta=(ClampMin = "0.1", UIMax = "4.0", editcondition = "bOverride_AmbientOcclusionPower", DisplayName = "Power"))
+	UPROPERTY(interp, BlueprintReadWrite, Category=AmbientOcclusion, AdvancedDisplay, meta=(ClampMin = "0.1", UIMax = "8.0", editcondition = "bOverride_AmbientOcclusionPower", DisplayName = "Power"))
 	float AmbientOcclusionPower;
 
 	/** >0, in unreal units, default (3.0) works well for flat surfaces but can reduce details */
