Fix for default APawn::AIControllerClass value not getting applied to derived classes #UE4

diff --git a/Engine/Source/Runtime/Engine/Private/Pawn.cpp b/Engine/Source/Runtime/Engine/Private/Pawn.cpp
--- a/Engine/Source/Runtime/Engine/Private/Pawn.cpp
+++ b/Engine/Source/Runtime/Engine/Private/Pawn.cpp
@@ -34,6 +34,10 @@ APawn::APawn(const FObjectInitializer& ObjectInitializer)
 	{
 		AIControllerClass = LoadClass<AController>(NULL, *((UEngine*)(UEngine::StaticClass()->GetDefaultObject()))->AIControllerClassName.ToString(), NULL, LOAD_None, NULL);
 	}
+	else
+	{
+		AIControllerClass = ((APawn*)APawn::StaticClass()->GetDefaultObject())->AIControllerClass;
+	}
 	bCanBeDamaged = true;
 	
 	SetRemoteRoleForBackwardsCompat(ROLE_SimulatedProxy);
