Bunch of Niagara work: effects editor, underlying infrastructure changes to the component, scene proxy, effect renderers, NiagaraSimulation, and some code cleanup; next up, proper serialization and more code cleanup.

diff --git a/Engine/Source/Editor/NiagaraEditor/NiagaraEditor.Build.cs b/Engine/Source/Editor/NiagaraEditor/NiagaraEditor.Build.cs
--- a/Engine/Source/Editor/NiagaraEditor/NiagaraEditor.Build.cs
+++ b/Engine/Source/Editor/NiagaraEditor/NiagaraEditor.Build.cs
@@ -8,13 +8,14 @@ public class NiagaraEditor : ModuleRules
 	{
 		PrivateDependencyModuleNames.AddRange(
 			new string[] {
+				"Engine", 
 				"Core", 
 				"CoreUObject", 
-				"Engine", 
                 "InputCore",
 				"RenderCore",
 				"Slate", 
 				"SlateCore",
+				"Kismet",
                 "EditorStyle",
 				"UnrealEd", 
 				"GraphEditor", 
@@ -22,7 +23,19 @@ public class NiagaraEditor : ModuleRules
 			}
 		);
 
-		PublicIncludePathModuleNames.Add("LevelEditor");
+        PublicDependencyModuleNames.AddRange(
+            new string[] {
+				"Engine"
+            }
+        );
+
+        PublicIncludePathModuleNames.AddRange(
+            new string[] {
+				"Engine", 
+				"Messaging", 
+				"GraphEditor", 
+				"LevelEditor"}
+                );
 
 		DynamicallyLoadedModuleNames.Add("WorkspaceMenuStructure");
 	}
diff --git a/Engine/Source/Editor/NiagaraEditor/Private/NiagaraEffectEditor.cpp b/Engine/Source/Editor/NiagaraEditor/Private/NiagaraEffectEditor.cpp
--- a/Engine/Source/Editor/NiagaraEditor/Private/NiagaraEffectEditor.cpp
+++ b/Engine/Source/Editor/NiagaraEditor/Private/NiagaraEffectEditor.cpp
@@ -10,6 +10,7 @@
 
 #define LOCTEXT_NAMESPACE "NiagaraEffectEditor"
 
+
 const FName FNiagaraEffectEditor::UpdateTabId(TEXT("NiagaraEditor_Effect"));
 
 void FNiagaraEffectEditor::RegisterTabSpawners(const TSharedRef<class FTabManager>& TabManager)
@@ -68,14 +69,6 @@ void FNiagaraEffectEditor::InitNiagaraEffectEditor(const EToolkitMode::Type Mode
 
 	ExtendToolbar();
 	RegenerateMenusAndToolbars();
-
-	// @todo toolkit world centric editing
-	/*// Setup our tool's layout
-	if( IsWorldCentricAssetEditor() )
-	{
-	const FString TabInitializationPayload(TEXT(""));		// NOTE: Payload not currently used for table properties
-	SpawnToolkitTab( UpdateGraphTabId, TabInitializationPayload, EToolkitTabSpot::Details );
-	}*/
 }
 
 FName FNiagaraEffectEditor::GetToolkitFName() const
@@ -104,16 +97,11 @@ FLinearColor FNiagaraEffectEditor::GetWorldCentricTabColorScale() const
 TSharedRef<SNiagaraEffectEditorWidget> FNiagaraEffectEditor::CreateEditorWidget(UNiagaraEffect* InEffect)
 {
 	check(InEffect != NULL);
-
+	
+	
 	if (!EditorCommands.IsValid())
 	{
 		EditorCommands = MakeShareable(new FUICommandList);
-		/*
-		// Editing commands
-		EditorCommands->MapAction(FGenericCommands::Get().Delete,
-			FExecuteAction::CreateSP(this, &FNiagaraEffectEditor::DeleteSelectedNodes),
-			FCanExecuteAction::CreateSP(this, &FNiagaraEffectEditor::CanDeleteNodes)
-			);*/
 	}
 
 	// Create the appearance info
@@ -139,12 +127,12 @@ TSharedRef<SNiagaraEffectEditorWidget> FNiagaraEffectEditor::CreateEditorWidget(
 			]
 		];
 
-	// Make full graph editor
-	return SNew(SNiagaraEffectEditorWidget)
-		//.AdditionalCommands(EditorCommands)
-		//.Appearance(AppearanceInfo)
-		.TitleBar(TitleBarWidget)
-		.EffectObj(InEffect);
+	// make preview pane
+	
+		
+		
+	// Make the effect editor widget
+	return SNew(SNiagaraEffectEditorWidget).TitleBar(TitleBarWidget).EffectObj(InEffect);
 }
 
 
@@ -166,41 +154,41 @@ TSharedRef<SDockTab> FNiagaraEffectEditor::SpawnTab(const FSpawnTabArgs& Args)
 
 void FNiagaraEffectEditor::ExtendToolbar()
 {
-	/*
 	struct Local
 	{
-		static void FillToolbar(FToolBarBuilder& ToolbarBuilder, TSharedRef<SWidget> CompileBox)
+		static void FillToolbar(FToolBarBuilder& ToolbarBuilder, TSharedRef<SWidget> AddEmitterBox)
 		{
-			ToolbarBuilder.BeginSection("Compile");
+			ToolbarBuilder.BeginSection("AddEmitter");
 			{
-				ToolbarBuilder.AddWidget(CompileBox);
+				ToolbarBuilder.AddWidget(AddEmitterBox);
 			}
 			ToolbarBuilder.EndSection();
 		}
 	};
 
 	TSharedPtr<FExtender> ToolbarExtender = MakeShareable(new FExtender);
-	TSharedRef<SWidget> Compilebox = SNew(SHorizontalBox)
+
+	TSharedRef<SWidget> AddEmitterBox = SNew(SHorizontalBox)
 		+ SHorizontalBox::Slot()
 		.AutoWidth()
 		.Padding(4)
 		[
 			SNew(SButton)
-			.OnClicked(this, &FNiagaraEditor::OnCompileClicked)
+			.OnClicked(this, &FNiagaraEffectEditor::OnAddEmitterClicked)
 			.Content()
 			[
 				SNew(SVerticalBox)
 				+ SVerticalBox::Slot()
-				.AutoHeight()
-				[
-					SNew(SImage)
-					.Image(FEditorStyle::GetBrush("LevelEditor.Recompile"))
-				]
+					.AutoHeight()
+					[
+						SNew(SImage)
+						.Image(FEditorStyle::GetBrush("LevelEditor.Add"))
+					]
 				+ SVerticalBox::Slot()
 					.AutoHeight()
 					[
 						SNew(STextBlock)
-						.Text(LOCTEXT("NiagaraToolbar_Compile", "Compile"))
+						.Text(LOCTEXT("NiagaraToolbar_AddEmitter", "Add Emitter"))
 					]
 			]
 		];
@@ -209,14 +197,26 @@ void FNiagaraEffectEditor::ExtendToolbar()
 		"Asset",
 		EExtensionHook::After,
 		GetToolkitCommands(),
-		FToolBarExtensionDelegate::CreateStatic(&Local::FillToolbar, Compilebox)
+		FToolBarExtensionDelegate::CreateStatic(&Local::FillToolbar, AddEmitterBox)
 		);
 
 	AddToolbarExtender(ToolbarExtender);
-	*/
+
 	FNiagaraEditorModule& NiagaraEditorModule = FModuleManager::LoadModuleChecked<FNiagaraEditorModule>("NiagaraEditor");
 	AddToolbarExtender(NiagaraEditorModule.GetToolBarExtensibilityManager()->GetAllExtenders(GetToolkitCommands(), GetEditingObjects()));
 }
 
 
+
+FReply FNiagaraEffectEditor::OnAddEmitterClicked()
+{
+	FNiagaraEditorModule& NiagaraEditorModule = FModuleManager::LoadModuleChecked<FNiagaraEditorModule>("NiagaraEditor");
+
+	Effect->AddEmitter();
+	UpdateEditorPtr.Pin()->GetViewport()->SetPreviewEffect(Effect);
+	UpdateEditorPtr.Pin()->UpdateList();
+	return FReply::Handled();
+}
+
+
 #undef LOCTEXT_NAMESPACE
diff --git a/Engine/Source/Editor/NiagaraEditor/Private/NiagaraEffectEditor.h b/Engine/Source/Editor/NiagaraEditor/Private/NiagaraEffectEditor.h
--- a/Engine/Source/Editor/NiagaraEditor/Private/NiagaraEffectEditor.h
+++ b/Engine/Source/Editor/NiagaraEditor/Private/NiagaraEffectEditor.h
@@ -7,8 +7,8 @@
 
 
 
-
-/** Viewer/editor for a DataTable */
+/** Viewer/editor for a NiagaraEffect
+*/
 class FNiagaraEffectEditor : public INiagaraEffectEditor
 {
 
@@ -30,6 +30,8 @@ public:
 	virtual FLinearColor GetWorldCentricTabColorScale() const override;
 	// End IToolkit interface
 
+
+	virtual UNiagaraEffect *GetEffect() const	override { return Effect; }	
 private:
 	/** Create widget for graph editing */
 	TSharedRef<class SNiagaraEffectEditorWidget> CreateEditorWidget(UNiagaraEffect* InEffect);
@@ -39,6 +41,7 @@ private:
 
 	/** Builds the toolbar widget */
 	void ExtendToolbar();
+	FReply OnAddEmitterClicked();
 private:
 
 	/* The Effect being edited */
diff --git a/Engine/Source/Editor/NiagaraEditor/Private/SNiagaraEffectEditorViewport.cpp b/Engine/Source/Editor/NiagaraEditor/Private/SNiagaraEffectEditorViewport.cpp
new file mode 100644
--- /dev/null
+++ b/Engine/Source/Editor/NiagaraEditor/Private/SNiagaraEffectEditorViewport.cpp
@@ -0,0 +1,265 @@
+// Copyright 1998-2014 Epic Games, Inc. All Rights Reserved.
+
+#include "NiagaraEditorPrivatePCH.h"
+#include "NiagaraEditorModule.h"
+#include "MouseDeltaTracker.h"
+#include "SNiagaraEffectEditorViewport.h"
+#include "PreviewScene.h"
+#include "Runtime/Engine/Public/Slate/SceneViewport.h"
+#include "NiagaraEditor.h"
+#include "ComponentReregisterContext.h"
+
+/** Viewport Client for the preview viewport */
+class FNiagaraEffectEditorViewportClient : public FEditorViewportClient
+{
+public:
+	FNiagaraEffectEditorViewportClient(TWeakPtr<INiagaraEffectEditor> InEffectEditor, FPreviewScene& InPreviewScene);
+	
+	// FEditorViewportClient interface
+	virtual FLinearColor GetBackgroundColor() const override;
+	virtual void Tick(float DeltaSeconds) override;
+	virtual void Draw(FViewport* Viewport,FCanvas* Canvas) override;
+	virtual bool ShouldOrbitCamera() const override;
+	virtual FSceneView* CalcSceneView(FSceneViewFamily* ViewFamily) override;
+	
+	void SetShowGrid(bool bShowGrid);
+	
+private:
+	
+	/** Pointer back to the material editor tool that owns us */
+	TWeakPtr<INiagaraEffectEditor> EffectEditorPtr;
+};
+
+FNiagaraEffectEditorViewportClient::FNiagaraEffectEditorViewportClient(TWeakPtr<INiagaraEffectEditor> InEffectEditor, FPreviewScene& InPreviewScene)
+: FEditorViewportClient( GLevelEditorModeTools(), &InPreviewScene )
+, EffectEditorPtr(InEffectEditor)
+{
+	// Setup defaults for the common draw helper.
+	DrawHelper.bDrawPivot = false;
+	DrawHelper.bDrawWorldBox = false;
+	DrawHelper.bDrawKillZ = false;
+	DrawHelper.bDrawGrid = false;
+	DrawHelper.GridColorAxis = FColor(80,80,80);
+	DrawHelper.GridColorMajor = FColor(72,72,72);
+	DrawHelper.GridColorMinor = FColor(64,64,64);
+	DrawHelper.PerspectiveGridSize = HALF_WORLD_MAX1;
+	
+	SetViewMode(VMI_Lit);
+	
+	EngineShowFlags.DisableAdvancedFeatures();
+	EngineShowFlags.SetSnap(0);
+	
+	OverrideNearClipPlane(1.0f);
+	bUsingOrbitCamera = true;
+}
+
+
+
+void FNiagaraEffectEditorViewportClient::Tick(float DeltaSeconds)
+{
+	FEditorViewportClient::Tick(DeltaSeconds);
+	
+	// Tick the preview scene world.
+	if (!GIntraFrameDebuggingGameThread)
+	{
+		PreviewScene->GetWorld()->Tick(LEVELTICK_All, DeltaSeconds);
+	}
+}
+
+
+void FNiagaraEffectEditorViewportClient::Draw(FViewport* Viewport,FCanvas* Canvas)
+{
+	FEditorViewportClient::Draw(Viewport, Canvas);
+}
+
+bool FNiagaraEffectEditorViewportClient::ShouldOrbitCamera() const
+{
+	return true;
+}
+
+
+FLinearColor FNiagaraEffectEditorViewportClient::GetBackgroundColor() const
+{
+	FLinearColor BackgroundColor = FLinearColor::Black;
+	return BackgroundColor;
+}
+
+FSceneView* FNiagaraEffectEditorViewportClient::CalcSceneView(FSceneViewFamily* ViewFamily)
+{
+	FSceneView* SceneView = FEditorViewportClient::CalcSceneView(ViewFamily);
+	FFinalPostProcessSettings::FCubemapEntry& CubemapEntry = *new(SceneView->FinalPostProcessSettings.ContributingCubemaps) FFinalPostProcessSettings::FCubemapEntry;
+	CubemapEntry.AmbientCubemap = GUnrealEd->GetThumbnailManager()->AmbientCubemap;
+	CubemapEntry.AmbientCubemapTintMulScaleValue = FLinearColor::White;
+	return SceneView;
+}
+
+void FNiagaraEffectEditorViewportClient::SetShowGrid(bool bShowGrid)
+{
+	DrawHelper.bDrawGrid = bShowGrid;
+}
+
+void SNiagaraEffectEditorViewport::Construct(const FArguments& InArgs)
+{
+	bShowGrid = false;
+	bShowBackground = false;
+	
+	// Rotate the light in the preview scene so that it faces the preview object
+	PreviewScene.SetLightDirection(FRotator(-40.0f, 27.5f, 0.0f));
+	
+	SEditorViewport::Construct( SEditorViewport::FArguments() );
+
+	PreviewComponent = ConstructObject<UNiagaraComponent>(UNiagaraComponent::StaticClass(), GetTransientPackage(), NAME_None, RF_Transient);
+}
+
+SNiagaraEffectEditorViewport::~SNiagaraEffectEditorViewport()
+{
+	if (EditorViewportClient.IsValid())
+	{
+		EditorViewportClient->Viewport = NULL;
+	}
+}
+
+void SNiagaraEffectEditorViewport::AddReferencedObjects( FReferenceCollector& Collector )
+{
+	Collector.AddReferencedObject( PreviewComponent );
+}
+
+void SNiagaraEffectEditorViewport::RefreshViewport()
+{
+	//reregister the preview components, so if the preview material changed it will be propagated to the render thread
+	PreviewComponent->MarkRenderStateDirty();
+	SceneViewport->InvalidateDisplay();
+}
+
+void SNiagaraEffectEditorViewport::Tick(const FGeometry& AllottedGeometry, const double InCurrentTime, const float InDeltaTime)
+{
+	SEditorViewport::Tick(AllottedGeometry, InCurrentTime, InDeltaTime);
+}
+
+
+
+void SNiagaraEffectEditorViewport::SetPreviewEffect(UNiagaraEffect *InPreviewEffect)
+{
+	check( PreviewComponent );
+
+	FTransform Transform = FTransform::Identity;
+	PreviewScene.RemoveComponent(PreviewComponent);
+	PreviewScene.AddComponent(PreviewComponent, Transform);
+
+	FComponentReregisterContext ReregisterContext(PreviewComponent);
+	PreviewComponent->Effect = InPreviewEffect;
+	PreviewComponent->Effect->Init(PreviewComponent);
+}
+
+
+void SNiagaraEffectEditorViewport::ToggleRealtime()
+{
+	EditorViewportClient->ToggleRealtime();
+}
+
+/*
+TSharedRef<FUICommandList> SNiagaraEffectEditorViewport::GetEffectEditorCommands() const
+{
+	check(EffectEditorPtr.IsValid());
+	return EffectEditorPtr.GetToolkitCommands();
+}
+*/
+
+void SNiagaraEffectEditorViewport::OnAddedToTab( const TSharedRef<SDockTab>& OwnerTab )
+{
+	ParentTab = OwnerTab;
+}
+
+bool SNiagaraEffectEditorViewport::IsVisible() const
+{
+	return ViewportWidget.IsValid() && (!ParentTab.IsValid() || ParentTab.Pin()->IsForeground()) && SEditorViewport::IsVisible() ;
+}
+
+void SNiagaraEffectEditorViewport::BindCommands()
+{
+	SEditorViewport::BindCommands();
+
+	// TODO: implement the necessary commands here
+	/*	
+	const FNiagaraEffectEditorCommands& Commands = FNiagaraEffectEditorCommands::Get();
+	TSharedRef<FUICommandList> ToolkitCommandList = GetEffectEditorCommands();
+	
+	// Add the commands to the toolkit command list so that the toolbar buttons can find them
+	
+	ToolkitCommandList->MapAction(
+		Commands.TogglePreviewGrid,
+		FExecuteAction::CreateSP( this, &SMaterialEditorViewport::TogglePreviewGrid ),
+								  FCanExecuteAction(),
+								  FIsActionChecked::CreateSP( this, &SMaterialEditorViewport::IsTogglePreviewGridChecked ) );
+	
+	ToolkitCommandList->MapAction(
+		Commands.TogglePreviewBackground,
+		FExecuteAction::CreateSP( this, &SMaterialEditorViewport::TogglePreviewBackground ),
+								  FCanExecuteAction(),
+								  FIsActionChecked::CreateSP( this, &SMaterialEditorViewport::IsTogglePreviewBackgroundChecked ) );
+								  */
+}
+
+void SNiagaraEffectEditorViewport::OnFocusViewportToSelection()
+{
+	if( PreviewComponent )
+	{
+		EditorViewportClient->FocusViewportOnBox( PreviewComponent->Bounds.GetBox() );
+	}
+}
+
+void SNiagaraEffectEditorViewport::TogglePreviewGrid()
+{
+	bShowGrid = !bShowGrid;
+	EditorViewportClient->SetShowGrid(bShowGrid);
+	RefreshViewport();
+}
+
+bool SNiagaraEffectEditorViewport::IsTogglePreviewGridChecked() const
+{
+	return bShowGrid;
+}
+
+void SNiagaraEffectEditorViewport::TogglePreviewBackground()
+{
+	bShowBackground = !bShowBackground;
+	// @todo DB: Set the background mesh for the preview viewport.
+	RefreshViewport();
+}
+
+bool SNiagaraEffectEditorViewport::IsTogglePreviewBackgroundChecked() const
+{
+	return bShowBackground;
+}
+
+TSharedRef<FEditorViewportClient> SNiagaraEffectEditorViewport::MakeEditorViewportClient() 
+{
+	EditorViewportClient = MakeShareable( new FNiagaraEffectEditorViewportClient(nullptr, PreviewScene) );
+	
+	EditorViewportClient->SetViewLocation( FVector::ZeroVector );
+	EditorViewportClient->SetViewRotation( FRotator::ZeroRotator );
+	EditorViewportClient->SetViewLocationForOrbiting( FVector::ZeroVector );
+	EditorViewportClient->bSetListenerPosition = false;
+	
+	EditorViewportClient->SetRealtime( true );
+	EditorViewportClient->VisibilityDelegate.BindSP( this, &SNiagaraEffectEditorViewport::IsVisible );
+	
+	return EditorViewportClient.ToSharedRef();
+}
+
+TSharedPtr<SWidget> SNiagaraEffectEditorViewport::MakeViewportToolbar()
+{
+	//return SNew(SNiagaraEffectEditorViewportToolBar)
+	//.Viewport(SharedThis(this));
+	return SNew(SBox);
+}
+
+EVisibility SNiagaraEffectEditorViewport::OnGetViewportContentVisibility() const
+{
+	EVisibility BaseVisibility = SEditorViewport::OnGetViewportContentVisibility();
+	if (BaseVisibility != EVisibility::Visible)
+	{
+		return BaseVisibility;
+	}
+	return IsVisible() ? EVisibility::Visible : EVisibility::Collapsed;
+}
diff --git a/Engine/Source/Editor/NiagaraEditor/Private/SNiagaraEffectEditorViewport.h b/Engine/Source/Editor/NiagaraEditor/Private/SNiagaraEffectEditorViewport.h
new file mode 100644
--- /dev/null
+++ b/Engine/Source/Editor/NiagaraEditor/Private/SNiagaraEffectEditorViewport.h
@@ -0,0 +1,72 @@
+// Copyright 1998-2014 Epic Games, Inc. All Rights Reserved.
+
+#pragma once
+
+#include "PreviewScene.h"
+#include "SEditorViewport.h"
+
+
+/**
+ * Material Editor Preview viewport widget
+ */
+class SNiagaraEffectEditorViewport : public SEditorViewport, public FGCObject
+{
+public:
+	SLATE_BEGIN_ARGS( SNiagaraEffectEditorViewport ){}
+	//SLATE_ARGUMENT(TWeakPtr<INiagaraEffectEditor>, EffectEditor)
+	SLATE_END_ARGS()
+	
+	void Construct(const FArguments& InArgs);
+	~SNiagaraEffectEditorViewport();
+	
+	virtual void AddReferencedObjects( FReferenceCollector& Collector ) override;
+	
+	virtual void Tick(const FGeometry& AllottedGeometry, const double InCurrentTime, const float InDeltaTime) override;
+	void RefreshViewport();
+	
+	void SetPreviewEffect(UNiagaraEffect* InEffect);
+	
+	void ToggleRealtime();
+	
+	/** @return The list of commands known by the material editor */
+	TSharedRef<FUICommandList> GetNiagaraEffectEditorCommands() const;
+	
+	/** If true, render background object in the preview scene. */
+	bool bShowBackground;
+	
+	/** If true, render grid the preview scene. */
+	bool bShowGrid;
+	
+	FPreviewScene PreviewScene;
+	
+	/** The material editor has been added to a tab */
+	void OnAddedToTab( const TSharedRef<SDockTab>& OwnerTab );
+	
+	/** Event handlers */
+	void TogglePreviewGrid();
+	bool IsTogglePreviewGridChecked() const;
+	void TogglePreviewBackground();
+	bool IsTogglePreviewBackgroundChecked() const;
+	
+protected:
+	/** SEditorViewport interface */
+	virtual TSharedRef<FEditorViewportClient> MakeEditorViewportClient() override;
+	virtual TSharedPtr<SWidget> MakeViewportToolbar() override;
+	virtual EVisibility OnGetViewportContentVisibility() const override;
+	virtual void BindCommands() override;
+	virtual void OnFocusViewportToSelection() override;
+private:
+	/** The parent tab where this viewport resides */
+	TWeakPtr<SDockTab> ParentTab;
+	
+	
+	bool IsVisible() const;
+	
+	/** Pointer back to the material editor tool that owns us */
+	//TWeakPtr<INiagaraEffectEditor> EffectEditorPtr;
+	
+	class UNiagaraComponent *PreviewComponent;
+	
+	/** Level viewport client */
+	TSharedPtr<class FNiagaraEffectEditorViewportClient> EditorViewportClient;
+};
diff --git a/Engine/Source/Editor/NiagaraEditor/Private/SNiagaraEffectEditorWidget.cpp b/Engine/Source/Editor/NiagaraEditor/Private/SNiagaraEffectEditorWidget.cpp
--- a/Engine/Source/Editor/NiagaraEditor/Private/SNiagaraEffectEditorWidget.cpp
+++ b/Engine/Source/Editor/NiagaraEditor/Private/SNiagaraEffectEditorWidget.cpp
@@ -6,20 +6,18 @@
 
 void SNiagaraEffectEditorWidget::Construct(const FArguments& InArgs)
 {
-	FString NiagaraEffectText(TEXT("Niagara Effect"));
-	FString ReadOnlyText(TEXT("Read Only"));
 	/*
-		Commands = MakeShareable(new FUICommandList());
-		IsEditable = InArgs._IsEditable;
-		Appearance = InArgs._Appearance;
-		TitleBarEnabledOnly = InArgs._TitleBarEnabledOnly;
-		TitleBar = InArgs._TitleBar;
-		bAutoExpandActionMenu = InArgs._AutoExpandActionMenu;
-		ShowGraphStateOverlay = InArgs._ShowGraphStateOverlay;
-		OnNavigateHistoryBack = InArgs._OnNavigateHistoryBack;
-		OnNavigateHistoryForward = InArgs._OnNavigateHistoryForward;
-		OnNodeSpawnedByKeymap = InArgs._GraphEvents.OnNodeSpawnedByKeymap;
-		*/
+	Commands = MakeShareable(new FUICommandList());
+	IsEditable = InArgs._IsEditable;
+	Appearance = InArgs._Appearance;
+	TitleBarEnabledOnly = InArgs._TitleBarEnabledOnly;
+	TitleBar = InArgs._TitleBar;
+	bAutoExpandActionMenu = InArgs._AutoExpandActionMenu;
+	ShowGraphStateOverlay = InArgs._ShowGraphStateOverlay;
+	OnNavigateHistoryBack = InArgs._OnNavigateHistoryBack;
+	OnNavigateHistoryForward = InArgs._OnNavigateHistoryForward;
+	OnNodeSpawnedByKeymap = InArgs._GraphEvents.OnNodeSpawnedByKeymap;
+	*/
 
 	/*
 	// Make sure that the editor knows about what kinds
@@ -57,99 +55,370 @@ void SNiagaraEffectEditorWidget::Construct(const FArguments& InArgs)
 	//bNeedsRefresh = false;
 
 
-	FPropertyEditorModule& PropertyEditorModule = FModuleManager::GetModuleChecked<FPropertyEditorModule>("PropertyEditor");
-	const FDetailsViewArgs DetailsViewArgs(false, false, true, false, true, this);
-	EmitterDetails = PropertyEditorModule.CreateDetailView(DetailsViewArgs);
-	//FOnGetDetailCustomizationInstance LayoutMICDetails = FOnGetDetailCustomizationInstance::CreateStatic(&FMaterialInstanceParameterDetails::MakeInstance, MaterialEditorInstance, FGetShowHiddenParameters::CreateSP(this, &FMaterialInstanceEditor::GetShowHiddenParameters));
-	//EmitterDetails->RegisterInstancedCustomPropertyLayout(UMaterialEditorInstanceConstant::StaticClass(), LayoutMICDetails);
-
-	//EmitterDetails->SetObjects(EmitterObj, true);
+	Viewport = SNew(SNiagaraEffectEditorViewport);
 
 
 	TSharedPtr<SOverlay> OverlayWidget;
 	this->ChildSlot
 		[
-			SAssignNew(OverlayWidget, SOverlay)
-
-			// details view on the left
-			+ SOverlay::Slot()
-			.Expose(DetailsViewSlot)
-			.HAlign(HAlign_Left)
-			.VAlign(VAlign_Top)
+			SNew(SHorizontalBox)
+			+ SHorizontalBox::Slot()
+			.AutoWidth()
 			[
-				SNew(SBorder)
-				.Padding(4)
+				SNew(SVerticalBox)
+				+ SVerticalBox::Slot()
+				.AutoHeight()
 				[
-					EmitterDetails.ToSharedRef()
+					SNew(SBox)
+					.WidthOverride(512)
+					.HeightOverride(512)
+					[
+						Viewport.ToSharedRef()
+					]
 				]
 			]
-
-			/*
-			// The graph panel
-			+ SOverlay::Slot()
-			.Expose(GraphPanelSlot)
-			[
-			SAssignNew(GraphPanel, SGraphPanel)
-			.EffectObj(EffectObj)
-			.GraphObjToDiff(InArgs._GraphToDiff)
-			//					.OnGetContextMenuFor(this, &SGraphEditorImpl::GraphEd_OnGetContextMenuFor)
-			.OnSelectionChanged(InArgs._GraphEvents.OnSelectionChanged)
-			.OnNodeDoubleClicked(InArgs._GraphEvents.OnNodeDoubleClicked)
-			//					.IsEditable(this, &SGraphEditorImpl::IsGraphEditable)
-			.OnDropActor(InArgs._GraphEvents.OnDropActor)
-			.OnDropStreamingLevel(InArgs._GraphEvents.OnDropStreamingLevel)
-			//					.IsEnabled(this, &SGraphEditorImpl::GraphEd_OnGetGraphEnabled)
-			.OnVerifyTextCommit(InArgs._GraphEvents.OnVerifyTextCommit)
-			.OnTextCommitted(InArgs._GraphEvents.OnTextCommitted)
-			.OnSpawnNodeByShortcut(InArgs._GraphEvents.OnSpawnNodeByShortcut)
-			//					.OnUpdateGraphPanel(this, &SGraphEditorImpl::GraphEd_OnPanelUpdated)
-			.OnDisallowedPinConnection(InArgs._GraphEvents.OnDisallowedPinConnection)
-			.ShowGraphStateOverlay(InArgs._ShowGraphStateOverlay)
-			]*/
-
-			// Title bar - optional
-			+SOverlay::Slot()
-				.VAlign(VAlign_Top)
+			+ SHorizontalBox::Slot()
+				.AutoWidth()
 				[
+					SAssignNew(OverlayWidget, SOverlay)
+					/*
+					// details view on the left
+					+ SOverlay::Slot()
+					.Expose(DetailsViewSlot)
+					.HAlign(HAlign_Left)
+					.VAlign(VAlign_Top)
+					[
+					SNew(SBorder)
+					.Padding(4)
+					[
+					EmitterDetails.ToSharedRef()
+					]
+					]
+					*/
+
+					// Title bar - optional
+					/*
+					+ SOverlay::Slot()
+					.VAlign(VAlign_Top)
+					[
 					InArgs._TitleBar.IsValid() ? InArgs._TitleBar.ToSharedRef() : (TSharedRef<SWidget>)SNullWidget::NullWidget
-				]
+					]
+					*/
 
-			// Bottom-right corner text indicating the type of tool
-			+ SOverlay::Slot()
-				.Padding(10)
-				.VAlign(VAlign_Bottom)
-				.HAlign(HAlign_Right)
-				[
-					SNew(STextBlock)
-					.Visibility(EVisibility::HitTestInvisible)
-					.TextStyle(FEditorStyle::Get(), "Graph.CornerText")
-					.Text(NiagaraEffectText)
-				]
 
-			// Top-right corner text indicating read only when not simulating
-			+ SOverlay::Slot()
-				.Padding(20)
-				.VAlign(VAlign_Top)
-				.HAlign(HAlign_Right)
-				[
-					SNew(STextBlock)
-					.Visibility(this, &SNiagaraEffectEditorWidget::ReadOnlyVisibility)
-					.TextStyle(FEditorStyle::Get(), "Graph.CornerText")
-					.Text(ReadOnlyText)
+					+SOverlay::Slot()
+					.Padding(10)
+					.VAlign(VAlign_Fill)
+					[
+						SAssignNew(ListView, SListView<TSharedPtr<FNiagaraSimulation> >)
+						.ItemHeight(256)
+						.ListItemsSource(&(this->EffectObj->Emitters))
+						.OnGenerateRow(this, &SNiagaraEffectEditorWidget::OnGenerateWidgetForList)
+					]
+
+					// Bottom-right corner text indicating the type of tool
+					+ SOverlay::Slot()
+						.Padding(10)
+						.VAlign(VAlign_Bottom)
+						.HAlign(HAlign_Right)
+						[
+							SNew(STextBlock)
+							.Visibility(EVisibility::HitTestInvisible)
+							.TextStyle(FEditorStyle::Get(), "Graph.CornerText")
+							.Text(FString("Niagara Effect"))
+						]
+
+					// Top-right corner text indicating read only when not simulating
+					+ SOverlay::Slot()
+						.Padding(20)
+						.VAlign(VAlign_Top)
+						.HAlign(HAlign_Right)
+						[
+							SNew(STextBlock)
+							.Visibility(this, &SNiagaraEffectEditorWidget::ReadOnlyVisibility)
+							.TextStyle(FEditorStyle::Get(), "Graph.CornerText")
+							.Text(FString("Read Only"))
+						]
+
+					// Bottom-right corner text for notification list position
+					+ SOverlay::Slot()
+						.Padding(15.f)
+						.VAlign(VAlign_Bottom)
+						.HAlign(HAlign_Right)
+						[
+							SAssignNew(NotificationListPtr, SNotificationList)
+							.Visibility(EVisibility::HitTestInvisible)
+						]
 				]
+		];
+
+	Viewport->SetPreviewEffect(EffectObj);
+}
+
+
 
-			// Bottom-right corner text for notification list position
-			+ SOverlay::Slot()
-				.Padding(15.f)
-				.VAlign(VAlign_Bottom)
-				.HAlign(HAlign_Right)
+
+
+void SEmitterWidget::Construct(const FArguments& InArgs)
+{
+	CurSpawnScript = nullptr;
+	CurUpdateScript = nullptr;
+	Emitter = InArgs._Emitter;
+	Effect = InArgs._Effect;
+	CurMaterial = nullptr;
+
+	ThumbnailPool = MakeShareable(new FAssetThumbnailPool(1, /*InAreRealTileThumbnailsAllowed=*/false));
+	// Create the thumbnail handle
+	MaterialThumbnail = MakeShareable(new FAssetThumbnail(CurMaterial, 64, 64, ThumbnailPool));
+	ThumbnailWidget = MaterialThumbnail->MakeThumbnailWidget();
+	MaterialThumbnail->GetViewportRenderTargetTexture(); // Access the texture once to trigger it to render
+
+	RenderModuleList.Add(MakeShareable(new FString("Sprites")));
+	RenderModuleList.Add(MakeShareable(new FString("Ribbon")));
+	RenderModuleList.Add(MakeShareable(new FString("Trails")));
+	RenderModuleList.Add(MakeShareable(new FString("Meshes")));
+
+	ChildSlot
+		.Padding(4)
+		[
+			SNew(SBorder)
+			.BorderImage(FEditorStyle::GetBrush("ToolPanel.LightGroupBorder"))
+			.Padding(4.0f)
+			[
+				SNew(SVerticalBox)
+				+ SVerticalBox::Slot()
+				.AutoHeight()
+				.Padding(4)
 				[
-					SAssignNew(NotificationListPtr, SNotificationList)
-					.Visibility(EVisibility::HitTestInvisible)
+					NGED_SECTION_BORDER
+					[
+						// name and status line
+						SNew(SHorizontalBox)
+						+ SHorizontalBox::Slot()
+						.HAlign(HAlign_Left)
+						.AutoWidth()
+						[
+							SNew(SCheckBox)
+							.OnCheckStateChanged(this, &SEmitterWidget::OnEmitterEnabledChanged)
+							.IsChecked(this, &SEmitterWidget::IsEmitterEnabled)
+							.ToolTipText(FText::FromString("Toggles whether this emitter is enabled. Disabled emitters don't simulate or render."))
+						]
+						+ SHorizontalBox::Slot()
+							.FillWidth(0.3)
+							.HAlign(HAlign_Left)
+							[
+								SNew(SEditableText).Text(FText::FromString("Emitter"))
+								.MinDesiredWidth(250)
+								.Font(FEditorStyle::GetFontStyle("ContentBrowser.AssetListViewNameFontDirty"))
+							]
+						+ SHorizontalBox::Slot()
+							.HAlign(HAlign_Right)
+							.FillWidth(0.7)
+							[
+								SNew(STextBlock).Text(this, &SEmitterWidget::GetStatsText)
+								.MinDesiredWidth(500)
+							]
+					]
 				]
-		];
 
-	//GraphPanel->RestoreViewSettings(FVector2D::ZeroVector, -1);
+				+ SVerticalBox::Slot()
+					.AutoHeight()
+					[
+						//------------------------------------------------------------------------------------
+						// horizontal arrangement of property widgets
+						SNew(SHorizontalBox)
 
-	//NotifyGraphChanged();
-}
+
+						// Update and spawn script selectors
+						//
+						+SHorizontalBox::Slot()
+						.Padding(4.0f)
+						.FillWidth(800)
+						.HAlign(HAlign_Left)
+						[
+							NGED_SECTION_BORDER
+							[
+								SNew(SVerticalBox)
+								+ SVerticalBox::Slot()
+								.Padding(4)
+								.AutoHeight()
+								[
+									SNew(SHorizontalBox)
+									+ SHorizontalBox::Slot()
+									.HAlign(HAlign_Right)
+									.VAlign(VAlign_Center)
+									.AutoWidth()
+									.Padding(4)
+									[
+										SNew(STextBlock)
+										.Text(FString("Update Script"))
+									]
+									+ SHorizontalBox::Slot()
+										.HAlign(HAlign_Right)
+										.VAlign(VAlign_Center)
+										[
+											SAssignNew(UpdateScriptSelector, SContentReference)
+											.AllowedClass(UNiagaraScript::StaticClass())
+											.AssetReference(this, &SEmitterWidget::GetUpdateScript)
+											.AllowSelectingNewAsset(true)
+											.AllowClearingReference(true)
+											.OnSetReference(this, &SEmitterWidget::OnUpdateScriptSelectedFromPicker)
+											.WidthOverride(150)
+										]
+
+								]
+
+								+ SVerticalBox::Slot()
+									.Padding(4)
+									.AutoHeight()
+									[
+										SNew(SHorizontalBox)
+										+ SHorizontalBox::Slot()
+										.HAlign(HAlign_Right)
+										.VAlign(VAlign_Center)
+										.AutoWidth()
+										.Padding(4)
+										[
+											SNew(STextBlock)
+											.Text(FString("Spawn Script"))
+										]
+										+ SHorizontalBox::Slot()
+											.HAlign(HAlign_Right)
+											.VAlign(VAlign_Center)
+											[
+												SAssignNew(SpawnScriptSelector, SContentReference)
+												.AllowedClass(UNiagaraScript::StaticClass())
+												.AssetReference(this, &SEmitterWidget::GetSpawnScript)
+												.AllowSelectingNewAsset(true)
+												.AllowClearingReference(true)
+												.OnSetReference(this, &SEmitterWidget::OnSpawnScriptSelectedFromPicker)
+												.WidthOverride(150)
+											]
+									]
+							]
+						]
+
+
+						// Spawn rate
+						+SHorizontalBox::Slot()
+							.AutoWidth()
+							.HAlign(HAlign_Left)
+							.Padding(4)
+							[
+								NGED_SECTION_BORDER
+								[
+									SNew(SVerticalBox)
+									+ SVerticalBox::Slot()
+									.Padding(4)
+									.AutoHeight()
+									[
+										SNew(SHorizontalBox)
+										+ SHorizontalBox::Slot()
+										.HAlign(HAlign_Center)
+										.VAlign(VAlign_Center)
+										.AutoWidth()
+										.Padding(4)
+										[
+											SNew(STextBlock)
+											.Text(FString("Spawn rate"))
+										]
+										+ SHorizontalBox::Slot()
+											[
+												SNew(SEditableTextBox).OnTextChanged(this, &SEmitterWidget::OnSpawnRateChanged)
+												.Text(this, &SEmitterWidget::GetSpawnRateText)
+											]
+									]
+								]
+							]
+
+
+						// Material
+						+SHorizontalBox::Slot()
+							.AutoWidth()
+							.HAlign(HAlign_Left)
+							.Padding(4)
+							[
+								NGED_SECTION_BORDER
+								[
+									SNew(SVerticalBox)
+									+ SVerticalBox::Slot()
+									.AutoHeight()
+									.Padding(4)
+									[
+										SNew(SHorizontalBox)
+										+ SHorizontalBox::Slot()
+										.Padding(4)
+										.AutoWidth()
+										[
+											SNew(STextBlock)
+											.Text(FString("Material"))
+										]
+										+ SHorizontalBox::Slot()
+											.Padding(4)
+											.AutoWidth()
+											[
+												SNew(SVerticalBox)
+												+ SVerticalBox::Slot()
+												.HAlign(HAlign_Left)
+												.VAlign(VAlign_Center)
+												.AutoHeight()
+												.Padding(4)
+												[
+													SNew(SBox)
+													.WidthOverride(64)
+													.HeightOverride(64)
+													[
+														ThumbnailWidget.ToSharedRef()
+													]
+												]
+												+ SVerticalBox::Slot()
+													.HAlign(HAlign_Left)
+													.VAlign(VAlign_Center)
+													.AutoHeight()
+													.Padding(4)
+													[
+														SNew(SContentReference)
+														.AllowedClass(UMaterial::StaticClass())
+														.AssetReference(this, &SEmitterWidget::GetMaterial)
+														.AllowSelectingNewAsset(true)
+														.AllowClearingReference(true)
+														.OnSetReference(this, &SEmitterWidget::OnMaterialSelected)
+														.WidthOverride(150)
+													]
+											]
+									]
+									+ SVerticalBox::Slot()
+										.AutoHeight()
+										.Padding(4)
+										[
+											SNew(SHorizontalBox)
+											+ SHorizontalBox::Slot()
+											.Padding(4)
+											.AutoWidth()
+											[
+												SNew(STextBlock)
+												.Text(FString("Render as"))
+											]
+											+ SHorizontalBox::Slot()
+												.Padding(4)
+												.AutoWidth()
+												[
+													// session combo box
+													SNew(SComboBox<TSharedPtr<FString> >)
+													.ContentPadding(FMargin(6.0f, 2.0f))
+													.OptionsSource(&RenderModuleList)
+													.OnSelectionChanged(this, &SEmitterWidget::OnRenderModuleChanged)
+													//										.ToolTipText(LOCTEXT("SessionComboBoxToolTip", "Select the game session to interact with.").ToString())
+													.OnGenerateWidget(this, &SEmitterWidget::GenerateRenderModuleComboWidget)
+													[
+														SNew(STextBlock)
+														.Text(this, &SEmitterWidget::GetRenderModuleText)
+													]
+												]
+										]
+								]
+							]
+					]
+			]
+		];
+}
\ No newline at end of file
diff --git a/Engine/Source/Editor/NiagaraEditor/Private/SNiagaraEffectEditorWidget.h b/Engine/Source/Editor/NiagaraEditor/Private/SNiagaraEffectEditorWidget.h
--- a/Engine/Source/Editor/NiagaraEditor/Private/SNiagaraEffectEditorWidget.h
+++ b/Engine/Source/Editor/NiagaraEditor/Private/SNiagaraEffectEditorWidget.h
@@ -4,8 +4,169 @@
 
 #include "Visibility.h"
 #include "SlateBasics.h"
-#include "Editor/PropertyEditor/Public/PropertyEditorModule.h"
-#include "Editor/PropertyEditor/Public/IDetailsView.h"
+#include "Editor/ContentBrowser/Public/ContentBrowserModule.h"
+#include "WorkflowOrientedApp/SContentReference.h"
+#include "SCheckBox.h"
+#include "SNameComboBox.h"
+#include "AssetThumbnail.h"
+#include "SNiagaraEffectEditorViewport.h"
+
+
+
+#define NGED_SECTION_BORDER SNew(SBorder).BorderImage(FEditorStyle::GetBrush("ToolPanel.GroupBorder")).Padding(4.0f).HAlign(HAlign_Left)
+
+
+class SEmitterWidget : public SCompoundWidget, public FNotifyHook
+{
+private:
+	TSharedPtr<FNiagaraSimulation> Emitter;
+	UNiagaraEffect *Effect;
+	TSharedPtr<STextBlock> UpdateComboText, SpawnComboText;
+	TSharedPtr<SComboButton> UpdateCombo, SpawnCombo;
+	TSharedPtr<SContentReference> UpdateScriptSelector, SpawnScriptSelector;
+	TSharedPtr<SWidget> ThumbnailWidget;
+	TSharedPtr<FAssetThumbnailPool> ThumbnailPool;
+	TSharedPtr<FAssetThumbnail> MaterialThumbnail;
+	UNiagaraScript *CurUpdateScript, *CurSpawnScript;
+	UMaterial *CurMaterial;
+
+	TArray<TSharedPtr<FString> > RenderModuleList;
+
+public:
+	SLATE_BEGIN_ARGS(SEmitterWidget) :
+		_Emitter(nullptr)
+	{
+	}
+	SLATE_ARGUMENT(TSharedPtr<FNiagaraSimulation>, Emitter)
+	SLATE_ARGUMENT(UNiagaraEffect*, Effect)
+	SLATE_END_ARGS()
+
+	void OnUpdateScriptSelectedFromPicker(UObject *Asset)
+	{
+		// Close the asset picker
+		//AssetPickerAnchor->SetIsOpen(false);
+
+		// Set the object found from the asset picker
+		CurUpdateScript = Cast<UNiagaraScript>(Asset);
+		Emitter->SetUpdateScript(CurUpdateScript);
+	}
+
+	void OnSpawnScriptSelectedFromPicker(UObject *Asset)
+	{
+		// Close the asset picker
+		//AssetPickerAnchor->SetIsOpen(false);
+
+		// Set the object found from the asset picker
+		CurSpawnScript = Cast<UNiagaraScript>(Asset);
+		Emitter->SetSpawnScript(CurSpawnScript);
+	}
+
+	void OnMaterialSelected(UObject *Asset)
+	{
+		MaterialThumbnail->SetAsset(Asset);
+		MaterialThumbnail->RefreshThumbnail();
+		MaterialThumbnail->GetViewportRenderTargetTexture(); // Access the texture once to trigger it to render
+		ThumbnailPool->Tick(0);
+		CurMaterial = Cast<UMaterial>(Asset);
+		Emitter->GetEffectRenderer()->SetMaterial(CurMaterial, ERHIFeatureLevel::SM5);
+	}
+
+	UObject *GetMaterial() const
+	{
+		return CurMaterial;
+	}
+
+	UObject *GetUpdateScript() const
+	{
+		return CurUpdateScript;
+	}
+
+	UObject *GetSpawnScript() const
+	{
+		return CurSpawnScript;
+	}
+
+	void OnEmitterEnabledChanged(ESlateCheckBoxState::Type NewCheckedState)
+	{
+		const bool bNewEnabledState = (NewCheckedState == ESlateCheckBoxState::Checked);
+		Emitter->SetEnabled(bNewEnabledState);
+	}
+
+	ESlateCheckBoxState::Type IsEmitterEnabled() const
+	{
+		return Emitter->IsEnabled() ? ESlateCheckBoxState::Checked : ESlateCheckBoxState::Unchecked;
+	}
+
+	void OnRenderModuleChanged(TSharedPtr<FString> ModName, ESelectInfo::Type SelectInfo)
+	{
+		EEmitterRenderModuleType Type = RMT_Sprites;
+		for (int32 i = 0; i < RenderModuleList.Num(); i++)
+		{
+			if (*RenderModuleList[i] == *ModName)
+			{
+				Type = (EEmitterRenderModuleType)(i+1);
+				break;
+			}
+		}
+
+
+		ENQUEUE_UNIQUE_RENDER_COMMAND_THREEPARAMETER(
+			FChangeNiagaraRenderModule,
+			UNiagaraEffect*, InEffect, this->Effect,
+			TSharedPtr<FNiagaraSimulation>, InEmitter, this->Emitter,
+			EEmitterRenderModuleType, InType, Type,
+		{
+			InEmitter->SetRenderModuleType(InType, InEffect->GetComponent()->GetWorld()->FeatureLevel);
+			InEffect->RenderModuleupdate();
+		});
+		TComponentReregisterContext<UNiagaraComponent> ComponentReregisterContext;
+	}
+
+	void OnSpawnRateChanged(const FText &InText)
+	{
+		Emitter->SetSpawnRate(FCString::Atof(*InText.ToString()) );
+	}
+
+	FText GetSpawnRateText() const
+	{
+		TCHAR Txt[32];
+		FCString::Snprintf(Txt, 32, TEXT("%f"), Emitter->GetSpawnRate() );
+
+		return FText::FromString(Txt);
+	}
+
+	FString GetRenderModuleText() const
+	{
+		if (Emitter->GetRenderModuleType() == RMT_None)
+		{
+			return FString("None");
+		}
+
+		return *RenderModuleList[Emitter->GetRenderModuleType()-1];
+	}
+
+	TSharedRef<SWidget>GenerateRenderModuleComboWidget(TSharedPtr<FString> WidgetString )
+	{
+		FString Txt = *WidgetString;
+		return SNew(STextBlock).Text(Txt);
+	}
+
+
+	FString GetStatsText() const
+	{
+		TCHAR Txt[128];
+		FCString::Snprintf(Txt, 128, TEXT("%i Particles"), Emitter->GetNumParticles());
+		return FString(Txt);
+	}
+
+	void Construct(const FArguments& InArgs);
+
+
+
+
+};
+
+
 
 
 class SNiagaraEffectEditorWidget : public SCompoundWidget, public FNotifyHook
@@ -15,23 +176,51 @@ private:
 
 	/** Notification list to pass messages to editor users  */
 	TSharedPtr<SNotificationList> NotificationListPtr;
-	TSharedPtr<class IDetailsView> EmitterDetails;
 	SOverlay::FOverlaySlot* DetailsViewSlot;
+	TSharedPtr<SNiagaraEffectEditorViewport>	Viewport;
+	
+	TSharedPtr< SListView<TSharedPtr<FNiagaraSimulation> > > ListView;
 
 public:
 	SLATE_BEGIN_ARGS(SNiagaraEffectEditorWidget)
 		: _EffectObj(nullptr)
-	{}
-		SLATE_ARGUMENT(UNiagaraEffect*, EffectObj)
+	{
+	}
+	SLATE_ARGUMENT(UNiagaraEffect*, EffectObj)
 		SLATE_ARGUMENT(TSharedPtr<SWidget>, TitleBar)
 	SLATE_END_ARGS()
 
+
 	/** Function to check whether we should show read-only text in the panel */
 	EVisibility ReadOnlyVisibility() const
 	{
 		return EVisibility::Hidden;
 	}
 
+	void UpdateList()
+	{
+		ListView->RequestListRefresh();
+	}
+
+	TSharedPtr<SNiagaraEffectEditorViewport> GetViewport()	{ return Viewport; }
+
+	TSharedRef<ITableRow> OnGenerateWidgetForList(TSharedPtr<FNiagaraSimulation> InItem, const TSharedRef< STableViewBase >& OwnerTable)
+	{
+		return SNew(STableRow<TSharedPtr<FString>>, OwnerTable)
+			.Content()
+			[
+				SNew(SBorder)
+				//.BorderImage(FEditorStyle::GetBrush("ToolPanel.GroupBorder"))
+				.Padding(2.0f)
+				[
+					SNew(SEmitterWidget).Emitter(InItem).Effect(EffectObj)
+				]
+			];
+	}
+
+
+
 	void Construct(const FArguments& InArgs);
+
 };
 
diff --git a/Engine/Source/Editor/NiagaraEditor/Public/INiagaraEffectEditor.h b/Engine/Source/Editor/NiagaraEditor/Public/INiagaraEffectEditor.h
--- a/Engine/Source/Editor/NiagaraEditor/Public/INiagaraEffectEditor.h
+++ b/Engine/Source/Editor/NiagaraEditor/Public/INiagaraEffectEditor.h
@@ -11,7 +11,7 @@ class INiagaraEffectEditor : public FAssetEditorToolkit
 {
 
 public:
-
+	virtual class UNiagaraEffect *GetEffect() const = 0;
 	// ...
 
 };
diff --git a/Engine/Source/Runtime/Engine/Classes/Components/NiagaraComponent.h b/Engine/Source/Runtime/Engine/Classes/Components/NiagaraComponent.h
--- a/Engine/Source/Runtime/Engine/Classes/Components/NiagaraComponent.h
+++ b/Engine/Source/Runtime/Engine/Classes/Components/NiagaraComponent.h
@@ -3,19 +3,17 @@
 #pragma once
 #include "NiagaraComponent.generated.h"
 
-UENUM()
-enum ERenderModuleType
-{
-	Sprites = 0,
-	Ribbon
-};
-
-
 /* Predefined particle attribute names */
 const FName ParticleAttr_Position = "Position";
 const FName ParticleAttr_Velocity = "Velocity";
 const FName ParticleAttr_Color = "Color";
 
+
+
+/**
+* Contains the simulation data for a single FNiagaraSimulation
+* all buffers for particle attributes are stored here for convenient access
+*/
 class FNiagaraEmitterParticleData
 {
 public:
@@ -82,34 +80,18 @@ private:
 
 
 
+/**
+* UNiagaraComponent is the primitive component for a Niagara effect.
+* @see ANiagaraActor
+* @see UNiagaraEffect
+*/
 UCLASS()
 class ENGINE_API UNiagaraComponent : public UPrimitiveComponent
 {
 	GENERATED_UCLASS_BODY()
 
-	/** The simulation being executed for this component. */
-	class FNiagaraSimulation* Simulation;
-
-	/** Script to run for these particles */
-	UPROPERTY(EditAnywhere, Category=NiagaraComponent)
-	class UNiagaraScript* UpdateScript;
 	UPROPERTY(EditAnywhere, Category = NiagaraComponent)
-	class UNiagaraScript* SpawnScript;
-
-	/** The method for rendering the simulated data */
-	UPROPERTY(EditAnywhere, Category=Rendering)
-	TEnumAsByte<ERenderModuleType> RenderModuleType;
-
-	/** Material with which to render particles. */
-	UPROPERTY(EditAnywhere, Category=Rendering)
-	UMaterialInterface* Material;
-
-	/** TEMP Spawn rate */
-	UPROPERTY(EditAnywhere, Category=NiagaraComponent)
-	float SpawnRate;
-
-	/** age of the emitter in seconds */
-	float EmitterAge;
+	class UNiagaraEffect *Effect;
 
 	// Begin UActorComponent interface.
 	virtual void TickComponent(float DeltaTime, enum ELevelTick TickType, FActorComponentTickFunction *ThisTickFunction) override;
@@ -122,8 +104,6 @@ public:
 
 	// Begin UPrimitiveComponent Interface
 	virtual int32 GetNumMaterials() const override;
-	virtual class UMaterialInterface* GetMaterial(int32 ElementIndex) const override;
-	virtual void SetMaterial(int32 ElementIndex, class UMaterialInterface* Material) override;
 	virtual FBoxSphereBounds CalcBounds(const FTransform& LocalToWorld) const override;
 	virtual FPrimitiveSceneProxy* CreateSceneProxy() override;
 	// End UPrimitiveComponent Interface
@@ -152,7 +132,9 @@ public:
 
 	/** Called on render thread to assign new dynamic data */
 	void SetDynamicData_RenderThread(struct FNiagaraDynamicDataBase* NewDynamicData);
-	class NiagaraEffectRenderer *GetEffectRenderer() { return EffectRenderer; }
+	TArray<class NiagaraEffectRenderer*> &GetEffectRenderers() { return EffectRenderers; }
+	void AddEffectRenderer(NiagaraEffectRenderer *Renderer)	{ EffectRenderers.Add(Renderer); }
+	ENGINE_API void UpdateEffectRenderers(UNiagaraEffect *InEffect);
 
 private:
 	void ReleaseRenderThreadResources();
@@ -168,6 +150,7 @@ private:
 	virtual void GetDynamicMeshElements(const TArray<const FSceneView*>& Views, const FSceneViewFamily& ViewFamily, uint32 VisibilityMap, FMeshElementCollector& Collector) const override;
 
 	virtual FPrimitiveViewRelevance GetViewRelevance(const FSceneView* View)  override;
+
 	/*
 	virtual bool CanBeOccluded() const override
 	{
@@ -180,5 +163,6 @@ private:
 
 
 private:
-	class NiagaraEffectRenderer *EffectRenderer;
+	//class NiagaraEffectRenderer *EffectRenderer;
+	TArray<class NiagaraEffectRenderer *>EffectRenderers;
 };
diff --git a/Engine/Source/Runtime/Engine/Classes/Engine/NiagaraConstants.h b/Engine/Source/Runtime/Engine/Classes/Engine/NiagaraConstants.h
--- a/Engine/Source/Runtime/Engine/Classes/Engine/NiagaraConstants.h
+++ b/Engine/Source/Runtime/Engine/Classes/Engine/NiagaraConstants.h
@@ -8,7 +8,7 @@
 namespace NiagaraConstants
 {
 	enum {
-		DeltaSeconds = 0,
+		DeltaSeconds = 1,
 		EmitterPosition,
 		EmitterAge,
 		ActorXAxis,
@@ -26,5 +26,5 @@ namespace NiagaraConstants
 		"Emitter Z Axis",
 		""
 	};
-}
+};
 
diff --git a/Engine/Source/Runtime/Engine/Classes/Engine/NiagaraEffect.h b/Engine/Source/Runtime/Engine/Classes/Engine/NiagaraEffect.h
--- a/Engine/Source/Runtime/Engine/Classes/Engine/NiagaraEffect.h
+++ b/Engine/Source/Runtime/Engine/Classes/Engine/NiagaraEffect.h
@@ -1,6 +1,8 @@
 // Copyright 1998-2014 Epic Games, Inc. All Rights Reserved.
 
 #pragma once
+#include "NiagaraSimulation.h"
+#include "NiagaraEffectRenderer.h"
 #include "NiagaraEffect.generated.h"
 
 UCLASS(MinimalAPI)
@@ -8,6 +10,80 @@ class UNiagaraEffect : public UObject
 {
 	GENERATED_UCLASS_BODY()
 
-	UPROPERTY()
-	TArray<FName> Emitters;
+public:
+	ENGINE_API TSharedPtr<FNiagaraSimulation> AddEmitter();
+
+
+	void Init(UNiagaraComponent *InComponent)
+	{
+		Component = InComponent;
+		//ConstantTable.Empty();
+		/*
+		// make room for constants
+		while (ConstantTable.Num() < 10)
+		{
+			ConstantTable.Add(FVector4(0.0f, 0.0f, 0.0f, 0.0f));
+		}
+		*/
+	}
+
+
+
+	void SetConstant(FName ConstantName, const float Value)
+	{
+		Constants.SetOrAdd(ConstantName, Value);
+	}
+
+	void SetConstant(FName ConstantName, const FVector4& Value)
+	{
+		Constants.SetOrAdd(ConstantName, Value);
+	}
+
+	void SetConstant(FName ConstantName, const FMatrix& Value)
+	{
+		Constants.SetOrAdd(ConstantName, Value);
+	}
+
+
+	void Tick(float DeltaSeconds)
+	{
+		// pass the constants down to the emitter
+		// TODO: should probably just pass a pointer to the table
+		for (TSharedPtr<FNiagaraSimulation>&it : Emitters)
+		{
+			/*
+			for (int32 i = 0; i < ConstantTable.Num(); i++)
+			{
+				it->SetConstant(i, ConstantTable[i]);
+			}
+			*/
+			it->SetConstants(Constants);
+
+			it->Tick(DeltaSeconds);
+		}
+	}
+
+	void RenderModuleupdate()
+	{
+		FNiagaraSceneProxy *NiagaraProxy = static_cast<FNiagaraSceneProxy*>(Component->SceneProxy);
+		NiagaraProxy->UpdateEffectRenderers(this);
+	}
+
+	FNiagaraSimulation *GetEmitter(uint32 idx)
+	{
+		return Emitters[idx].Get();
+	}
+
+
+	/** Local constant table. */
+	//TArray<FVector4> ConstantTable;
+	FNiagaraConstantMap Constants;
+
+	//UPROPERTY()
+	TArray< TSharedPtr<FNiagaraSimulation> > Emitters;
+
+	UNiagaraComponent *GetComponent() { return Component; }
+private:
+	UNiagaraComponent *Component;
+	TArray<FNiagaraEmitterProperties>EmitterProps;
 };
diff --git a/Engine/Source/Runtime/Engine/Classes/Engine/NiagaraEffectRenderer.h b/Engine/Source/Runtime/Engine/Classes/Engine/NiagaraEffectRenderer.h
--- a/Engine/Source/Runtime/Engine/Classes/Engine/NiagaraEffectRenderer.h
+++ b/Engine/Source/Runtime/Engine/Classes/Engine/NiagaraEffectRenderer.h
@@ -3,75 +3,38 @@
 /*==============================================================================
 NiagaraEffectRenderer.h: Base class for Niagara render modules
 ==============================================================================*/
+#pragma once
+
 #include "SceneUtils.h"
+#include "ParticleHelper.h"
+#include "ParticleVertexFactory.h"
 #include "ParticleBeamTrailVertexFactory.h"
 
-DECLARE_CYCLE_STAT(TEXT("PreRenderView"), STAT_NiagaraPreRenderView, STATGROUP_Niagara);
-DECLARE_CYCLE_STAT(TEXT("Render Total"), STAT_NiagaraRender, STATGROUP_Niagara);
-DECLARE_CYCLE_STAT(TEXT("Render Sprites"), STAT_NiagaraRenderSprites, STATGROUP_Niagara);
-DECLARE_CYCLE_STAT(TEXT("Render Ribbons"), STAT_NiagaraRenderRibbons, STATGROUP_Niagara);
 DECLARE_CYCLE_STAT(TEXT("Generate Sprite Vertex Data"), STAT_NiagaraGenSpriteVertexData, STATGROUP_Niagara);
 DECLARE_CYCLE_STAT(TEXT("Generate Ribbon Vertex Data"), STAT_NiagaraGenRibbonVertexData, STATGROUP_Niagara);
 
-
-
-
 /** Struct used to pass dynamic data from game thread to render thread */
 struct FNiagaraDynamicDataBase
 {
 };
 
 
-struct FNiagaraDynamicDataSprites : public FNiagaraDynamicDataBase
-{
-	TArray<FParticleSpriteVertex> VertexData;
-};
-
-struct FNiagaraDynamicDataRibbon : public FNiagaraDynamicDataBase
-{
-	TArray<FParticleBeamTrailVertex> VertexData;
-};
-
-
-
-/* Mesh collector classes */
-class FNiagaraMeshCollectorResourcesSprite : public FOneFrameResource
-{
-public:
-	FParticleSpriteVertexFactory VertexFactory;
-	FParticleSpriteUniformBufferRef UniformBuffer;
-
-	virtual ~FNiagaraMeshCollectorResourcesSprite()
-	{
-		VertexFactory.ReleaseResource();
-	}
-};
-
-
-class FNiagaraMeshCollectorResourcesRibbon : public FOneFrameResource
-{
-public:
-	 FParticleBeamTrailVertexFactory VertexFactory;
-	 FParticleBeamTrailUniformBufferRef UniformBuffer;
-
-	virtual ~FNiagaraMeshCollectorResourcesRibbon()
-	{
-		VertexFactory.ReleaseResource();
-	}
-};
-
 
 
 
 
+/**
+* Base class for Niagara effect renderers. Effect renderers handle generating vertex data for and
+* drawing of simulation data coming out of FNiagaraSimulation instances.
+*/
 class NiagaraEffectRenderer
 {
 public:
-	virtual void GetDynamicMeshElements(const TArray<const FSceneView*>& Views, const FSceneViewFamily& ViewFamily, uint32 VisibilityMap, FMeshElementCollector& Collector) const = 0;
+	virtual void GetDynamicMeshElements(const TArray<const FSceneView*>& Views, const FSceneViewFamily& ViewFamily, uint32 VisibilityMap, FMeshElementCollector& Collector, const FNiagaraSceneProxy *SceneProxy) const = 0;
 
-	virtual void PreRenderView(const FSceneViewFamily* ViewFamily, const uint32 VisibilityMap, int32 FrameNumber) = 0;
+	virtual void PreRenderView(const FSceneViewFamily* ViewFamily, const uint32 VisibilityMap, int32 FrameNumber, const FNiagaraSceneProxy *SceneProxy) = 0;
 	virtual void SetDynamicData_RenderThread(FNiagaraDynamicDataBase* NewDynamicData) = 0;
-	virtual void DrawDynamicElements(FPrimitiveDrawInterface* PDI, const FSceneView* View) = 0;
+	virtual void DrawDynamicElements(FPrimitiveDrawInterface* PDI, const FSceneView* View, const FNiagaraSceneProxy *SceneProxy) = 0;
 	virtual void CreateRenderThreadResources() = 0;
 	virtual void ReleaseRenderThreadResources() = 0;
 	virtual FNiagaraDynamicDataBase *GenerateVertexData(const FNiagaraEmitterParticleData &Data) = 0;
@@ -79,7 +42,7 @@ public:
 
 	virtual bool HasDynamicData() = 0;
 
-	virtual FPrimitiveViewRelevance GetViewRelevance(const FSceneView* View)
+	virtual FPrimitiveViewRelevance GetViewRelevance(const FSceneView* View, FNiagaraSceneProxy *SceneProxy)
 	{
 		FPrimitiveViewRelevance Result;
 		bool bHasDynamicData = HasDynamicData();
@@ -98,24 +61,27 @@ public:
 	}
 
 
-
-	~NiagaraEffectRenderer() {}
-
-protected:
-	NiagaraEffectRenderer(const UNiagaraComponent *InComponent, const FPrimitiveSceneProxy *Proxy)	
+	UMaterial *GetMaterial()	{ return Material; }
+	void SetMaterial(UMaterial *InMaterial, ERHIFeatureLevel::Type FeatureLevel)
 	{
-		SceneProxy = static_cast<const FNiagaraSceneProxy*>(Proxy);
-		Material = InComponent->Material;
-		if (!Material || !Material->CheckMaterialUsage(MATUSAGE_ParticleSprites))
+		Material = InMaterial;
+		if (!Material)
 		{
 			Material = UMaterial::GetDefaultMaterial(MD_Surface);
 		}
 		check(Material);
-		MaterialRelevance = Material->GetRelevance(SceneProxy->GetScene().GetFeatureLevel());
+		MaterialRelevance = Material->GetRelevance(FeatureLevel);
+	}
+
+	~NiagaraEffectRenderer() {}
+
+protected:
+	NiagaraEffectRenderer()	
+		: Material(nullptr)
+	{
 	}
 
-	const FNiagaraSceneProxy *SceneProxy;
-	UMaterialInterface* Material;
+	UMaterial* Material;
 
 private:
 	FMaterialRelevance MaterialRelevance;
@@ -123,333 +89,45 @@ private:
 
 
 
+/**
+* NiagaraEffectRendererSprites renders an FNiagaraSimulation as sprite particles
+*/
 class NiagaraEffectRendererSprites : public NiagaraEffectRenderer
 {
 public:	
-	NiagaraEffectRendererSprites(const UNiagaraComponent* InComponent, const FPrimitiveSceneProxy *Proxy) :
-		NiagaraEffectRenderer(InComponent, Proxy),
-		DynamicDataRender(NULL),
-		VertexFactory(PVFT_Sprite, InComponent->GetWorld()->FeatureLevel)		
-	{
-	}
-
+	NiagaraEffectRendererSprites(ERHIFeatureLevel::Type FeatureLevel);
 	~NiagaraEffectRendererSprites()
 	{
 		ReleaseRenderThreadResources();
 	}
 
 
-	virtual void ReleaseRenderThreadResources() override
-	{
-		VertexFactory.ReleaseResource();
-		PerViewUniformBuffers.Empty();
-		WorldSpacePrimitiveUniformBuffer.ReleaseResource();
-	}
+	virtual void ReleaseRenderThreadResources() override;
 
 	// FPrimitiveSceneProxy interface.
-	virtual void CreateRenderThreadResources() override
-	{
-		VertexFactory.InitResource();
-
-		UniformParameters.AxisLockRight = FVector4(0.0f, 0.0f, 0.0f, 0.0f);
-		UniformParameters.AxisLockUp = FVector4(0.0f, 0.0f, 0.0f, 0.0f);
-		UniformParameters.RotationScale = 1.0f;
-		UniformParameters.RotationBias = 0.0f;
-		UniformParameters.TangentSelector = FVector4(0.0f, 0.0f, 0.0f, 1.0f);
-		UniformParameters.InvDeltaSeconds = 0.0f;
-		UniformParameters.SubImageSize = FVector4(1.0f, 1.0f, 1.0f, 1.0f);
-		UniformParameters.NormalsType = 0;
-		UniformParameters.NormalsSphereCenter = FVector4(0.0f, 0.0f, 0.0f, 1.0f);
-		UniformParameters.NormalsCylinderUnitDirection = FVector4(0.0f, 0.0f, 1.0f, 0.0f);
-		UniformParameters.PivotOffset = FVector2D(0.0f, 0.0f);
-	}
-
-
-	virtual void PreRenderView(const FSceneViewFamily* ViewFamily, const uint32 VisibilityMap, int32 FrameNumber) override
-	{
-		SCOPE_CYCLE_COUNTER(STAT_NiagaraPreRenderView);
-
-		check(DynamicDataRender && DynamicDataRender->VertexData.Num());
-
-		int32 SizeInBytes = DynamicDataRender->VertexData.GetTypeSize() * DynamicDataRender->VertexData.Num();
-		DynamicVertexAllocation = FGlobalDynamicVertexBuffer::Get().Allocate(SizeInBytes);
-		if (DynamicVertexAllocation.IsValid())
-		{
-			// Copy the vertex data over.
-			FMemory::Memcpy(DynamicVertexAllocation.Buffer, DynamicDataRender->VertexData.GetData(), SizeInBytes);
-			VertexFactory.SetInstanceBuffer(
-				DynamicVertexAllocation.VertexBuffer,
-				DynamicVertexAllocation.VertexOffset,
-				sizeof(FParticleSpriteVertex),
-				true
-				);
-			VertexFactory.SetDynamicParameterBuffer(NULL, 0, 0, true);
-			// Compute the per-view uniform buffers.
-			const int32 NumViews = ViewFamily->Views.Num();
-			uint32 ViewBit = 1;
-			for (int32 ViewIndex = 0; ViewIndex < NumViews; ++ViewIndex, ViewBit <<= 1)
-			{
-				FParticleSpriteUniformBufferRef* SpriteViewUniformBufferPtr = new(PerViewUniformBuffers)FParticleSpriteUniformBufferRef();
-				if (VisibilityMap & ViewBit)
-				{
-					FParticleSpriteUniformParameters PerViewUniformParameters = UniformParameters;
-					PerViewUniformParameters.MacroUVParameters = FVector4(0.0f, 0.0f, 1.0f, 1.0f);
-					*SpriteViewUniformBufferPtr = FParticleSpriteUniformBufferRef::CreateUniformBufferImmediate(PerViewUniformParameters, UniformBuffer_SingleFrame);
-				}
-			}
-
-			// Update the primitive uniform buffer if needed.
-			if (!WorldSpacePrimitiveUniformBuffer.IsInitialized())
-			{
-				FPrimitiveUniformShaderParameters PrimitiveUniformShaderParameters = GetPrimitiveUniformShaderParameters(
-					FMatrix::Identity,
-					SceneProxy->GetActorPosition(),
-					SceneProxy->GetBounds(),
-					SceneProxy->GetLocalBounds(),
-					SceneProxy->ReceivesDecals(),
-					false,
-					false, 
-					false
-					);
-				WorldSpacePrimitiveUniformBuffer.SetContents(PrimitiveUniformShaderParameters);
-				WorldSpacePrimitiveUniformBuffer.InitResource();
-			}
-		}
-	}
-
-
-	virtual void DrawDynamicElements(FPrimitiveDrawInterface* PDI, const FSceneView* View) override
-	{
-		SCOPE_CYCLE_COUNTER(STAT_NiagaraRender);
-
-		check(DynamicDataRender && DynamicDataRender->VertexData.Num());
-
-		const bool bIsWireframe = View->Family->EngineShowFlags.Wireframe;
-		FMaterialRenderProxy* MaterialRenderProxy = Material->GetRenderProxy(SceneProxy->IsSelected(), SceneProxy->IsHovered());
-
-		if (DynamicVertexAllocation.IsValid()
-			&& (bIsWireframe || !PDI->IsMaterialIgnored(MaterialRenderProxy, View->GetFeatureLevel())))
-		{
-			FMeshBatch MeshBatch;
-			MeshBatch.VertexFactory = &VertexFactory;
-			MeshBatch.CastShadow = SceneProxy->CastsDynamicShadow();
-			MeshBatch.bUseAsOccluder = false;
-			MeshBatch.ReverseCulling = SceneProxy->IsLocalToWorldDeterminantNegative();
-			MeshBatch.Type = PT_TriangleList;
-			MeshBatch.DepthPriorityGroup = SceneProxy->GetDepthPriorityGroup(View);
-			MeshBatch.LCI = NULL;
-			if (bIsWireframe)
-			{
-				MeshBatch.MaterialRenderProxy = UMaterial::GetDefaultMaterial(MD_Surface)->GetRenderProxy(SceneProxy->IsSelected(), SceneProxy->IsHovered());
-			}
-			else
-			{
-				MeshBatch.MaterialRenderProxy = MaterialRenderProxy;
-			}
-
-			FMeshBatchElement& MeshElement = MeshBatch.Elements[0];
-			MeshElement.IndexBuffer = &GParticleIndexBuffer;
-			MeshElement.FirstIndex = 0;
-			MeshElement.NumPrimitives = 2;
-			MeshElement.NumInstances = DynamicDataRender->VertexData.Num();
-			MeshElement.MinVertexIndex = 0;
-			MeshElement.MaxVertexIndex = MeshElement.NumInstances * 4 - 1;
-			MeshElement.PrimitiveUniformBufferResource = &WorldSpacePrimitiveUniformBuffer;
-
-			int32 ViewIndex = View->Family->Views.Find(View);
-			check(PerViewUniformBuffers.IsValidIndex(ViewIndex) && PerViewUniformBuffers[ViewIndex]);
-			VertexFactory.SetSpriteUniformBuffer(PerViewUniformBuffers[ViewIndex]);
-
-			DrawRichMesh(
-				PDI,
-				MeshBatch,
-				FLinearColor(1.0f, 0.0f, 0.0f),	//WireframeColor,
-				FLinearColor(1.0f, 1.0f, 0.0f),	//LevelColor,
-				FLinearColor(1.0f, 1.0f, 1.0f),	//PropertyColor,		
-				SceneProxy,
-				GIsEditor && (View->Family->EngineShowFlags.Selection) ? SceneProxy->IsSelected() : false,
-				bIsWireframe
-				);
-		}
+	virtual void CreateRenderThreadResources() override;
 
-		/*
-		for(int32 i=0; i<DynamicDataRender->VertexData.Num(); i++)
-		{
-		FParticleSpriteVertex& Vertex = DynamicDataRender->VertexData[i];
-		PDI->DrawPoint(Vertex.Position, Vertex.Color, 1.0f, SDPG_World);
-		}
-		*/
-	}
-
-
-	virtual void GetDynamicMeshElements(const TArray<const FSceneView*>& Views, const FSceneViewFamily& ViewFamily, uint32 VisibilityMap, FMeshElementCollector& Collector) const override
-	{
-		SCOPE_CYCLE_COUNTER(STAT_NiagaraRender);
-		SCOPE_CYCLE_COUNTER(STAT_NiagaraRenderSprites);
-
-		check(DynamicDataRender && DynamicDataRender->VertexData.Num());
-
-		const bool bIsWireframe = ViewFamily.EngineShowFlags.Wireframe;
-		FMaterialRenderProxy* MaterialRenderProxy = Material->GetRenderProxy(SceneProxy->IsSelected(), SceneProxy->IsHovered());
+	virtual void PreRenderView(const FSceneViewFamily* ViewFamily, const uint32 VisibilityMap, int32 FrameNumber, const FNiagaraSceneProxy *SceneProxy) override;
+	virtual void DrawDynamicElements(FPrimitiveDrawInterface* PDI, const FSceneView* View, const FNiagaraSceneProxy *SceneProxy) override;
+	virtual void GetDynamicMeshElements(const TArray<const FSceneView*>& Views, const FSceneViewFamily& ViewFamily, uint32 VisibilityMap, FMeshElementCollector& Collector, const FNiagaraSceneProxy *SceneProxy) const override;
 
-		int32 SizeInBytes = DynamicDataRender->VertexData.GetTypeSize() * DynamicDataRender->VertexData.Num();
-		FGlobalDynamicVertexBuffer::FAllocation LocalDynamicVertexAllocation = FGlobalDynamicVertexBuffer::Get().Allocate(SizeInBytes);
-
-		if (LocalDynamicVertexAllocation.IsValid())
-		{
-			// Update the primitive uniform buffer if needed.
-			if (!WorldSpacePrimitiveUniformBuffer.IsInitialized())
-			{
-				FPrimitiveUniformShaderParameters PrimitiveUniformShaderParameters = GetPrimitiveUniformShaderParameters(
-					FMatrix::Identity,
-					SceneProxy->GetActorPosition(),
-					SceneProxy->GetBounds(),
-					SceneProxy->GetLocalBounds(),
-					SceneProxy->ReceivesDecals(),
-					false,
-					SceneProxy->UseEditorDepthTest()
-					);
-				WorldSpacePrimitiveUniformBuffer.SetContents(PrimitiveUniformShaderParameters);
-				WorldSpacePrimitiveUniformBuffer.InitResource();
-			}
-
-			// Copy the vertex data over.
-			FMemory::Memcpy(LocalDynamicVertexAllocation.Buffer, DynamicDataRender->VertexData.GetData(), SizeInBytes);
-
-			// Compute the per-view uniform buffers.
-			for (int32 ViewIndex = 0; ViewIndex < Views.Num(); ViewIndex++)
-			{
-				if (VisibilityMap & (1 << ViewIndex))
-				{
-					const FSceneView* View = Views[ViewIndex];
-
-					FNiagaraMeshCollectorResourcesSprite& CollectorResources = Collector.AllocateOneFrameResource<FNiagaraMeshCollectorResourcesSprite>();
-					FParticleSpriteUniformParameters PerViewUniformParameters = UniformParameters;
-
-					// Collector.AllocateOneFrameResource uses default ctor, initialize the vertex factory
-					CollectorResources.VertexFactory.SetFeatureLevel(ViewFamily.GetFeatureLevel());
-					CollectorResources.VertexFactory.SetParticleFactoryType(PVFT_Sprite);
-
-					PerViewUniformParameters.MacroUVParameters = FVector4(0.0f, 0.0f, 1.0f, 1.0f);
-					CollectorResources.UniformBuffer = FParticleSpriteUniformBufferRef::CreateUniformBufferImmediate(PerViewUniformParameters, UniformBuffer_SingleFrame);
-
-					CollectorResources.VertexFactory.InitResource();
-					CollectorResources.VertexFactory.SetSpriteUniformBuffer(CollectorResources.UniformBuffer);
-					CollectorResources.VertexFactory.SetInstanceBuffer(
-						LocalDynamicVertexAllocation.VertexBuffer,
-						LocalDynamicVertexAllocation.VertexOffset,
-						sizeof(FParticleSpriteVertex),
-						true
-						);
-					CollectorResources.VertexFactory.SetDynamicParameterBuffer(NULL, 0, 0, true);
-
-					FMeshBatch& MeshBatch = Collector.AllocateMesh();
-					MeshBatch.VertexFactory = &CollectorResources.VertexFactory;
-					MeshBatch.CastShadow = SceneProxy->CastsDynamicShadow();
-					MeshBatch.bUseAsOccluder = false;
-					MeshBatch.ReverseCulling = SceneProxy->IsLocalToWorldDeterminantNegative();
-					MeshBatch.Type = PT_TriangleList;
-					MeshBatch.DepthPriorityGroup = SceneProxy->GetDepthPriorityGroup(View);
-					MeshBatch.bCanApplyViewModeOverrides = true;
-					MeshBatch.bUseWireframeSelectionColoring = SceneProxy->IsSelected();
-
-					if (bIsWireframe)
-					{
-						MeshBatch.MaterialRenderProxy = UMaterial::GetDefaultMaterial(MD_Surface)->GetRenderProxy(SceneProxy->IsSelected(), SceneProxy->IsHovered());
-					}
-					else
-					{
-						MeshBatch.MaterialRenderProxy = MaterialRenderProxy;
-					}
-
-					FMeshBatchElement& MeshElement = MeshBatch.Elements[0];
-					MeshElement.IndexBuffer = &GParticleIndexBuffer;
-					MeshElement.FirstIndex = 0;
-					MeshElement.NumPrimitives = 2;
-					MeshElement.NumInstances = DynamicDataRender->VertexData.Num();
-					MeshElement.MinVertexIndex = 0;
-					MeshElement.MaxVertexIndex = MeshElement.NumInstances * 4 - 1;
-					MeshElement.PrimitiveUniformBufferResource = &WorldSpacePrimitiveUniformBuffer;
-
-					Collector.AddMesh(ViewIndex, MeshBatch);
-				}
-			}
-		}
-	}
 
 
 	/** Update render data buffer from attributes */
-	FNiagaraDynamicDataBase *GenerateVertexData(const FNiagaraEmitterParticleData &Data) override
-	{
-		SCOPE_CYCLE_COUNTER(STAT_NiagaraGenSpriteVertexData);
-		FNiagaraDynamicDataSprites *DynamicData = new FNiagaraDynamicDataSprites;
-		TArray<FParticleSpriteVertex>& RenderData = DynamicData->VertexData;
-
-		RenderData.Reset(Data.GetNumParticles());
-		//CachedBounds.Init();
-
-		const FVector4 *PosPtr = Data.GetAttributeData("Position");
-		const FVector4 *ColPtr = Data.GetAttributeData("Color");
-		const FVector4 *AgePtr = Data.GetAttributeData("Age");
-		const FVector4 *RotPtr = Data.GetAttributeData("Rotation");
-
-		float ParticleId = 0.0f, IdInc = 1.0f / Data.GetNumParticles();
-		for (uint32 ParticleIndex = 0; ParticleIndex < Data.GetNumParticles(); ParticleIndex++)
-		{
-			FParticleSpriteVertex& NewVertex = *new(RenderData)FParticleSpriteVertex;
-			NewVertex.Position = PosPtr[ParticleIndex];
-			NewVertex.OldPosition = NewVertex.Position;
-			NewVertex.Color = FLinearColor(ColPtr[ParticleIndex]);
-			NewVertex.ParticleId = ParticleId;
-			ParticleId += IdInc;
-			NewVertex.RelativeTime = AgePtr[ParticleIndex].X;
-			NewVertex.Size = FVector2D(1.0f, 1.0f);
-			NewVertex.Rotation = RotPtr[ParticleIndex].X;
-			NewVertex.SubImageIndex = 0.f;
-
-			//CachedBounds += NewVertex.Position;
-		}
-		//CachedBounds.ExpandBy(MaxSize);
-
-		return DynamicData;
-	}
+	FNiagaraDynamicDataBase *GenerateVertexData(const FNiagaraEmitterParticleData &Data) override;
 
 
 
 
-	virtual void SetDynamicData_RenderThread(FNiagaraDynamicDataBase* NewDynamicData) override
-	{
-		check(IsInRenderingThread());
-
-		if (DynamicDataRender)
-		{
-			delete DynamicDataRender;
-			DynamicDataRender = NULL;
-		}
-		DynamicDataRender = static_cast<FNiagaraDynamicDataSprites*>(NewDynamicData);
-	}
-
-	int GetDynamicDataSize()
-	{
-		uint32 Size = sizeof(FNiagaraDynamicDataSprites);
-		if (DynamicDataRender)
-		{
-			Size += DynamicDataRender->VertexData.GetAllocatedSize();
-		}
-
-		return Size;
-	}
-
-	bool HasDynamicData()
-	{
-		return DynamicDataRender && DynamicDataRender->VertexData.Num() > 0;
-	}
+	virtual void SetDynamicData_RenderThread(FNiagaraDynamicDataBase* NewDynamicData) override;
+	int GetDynamicDataSize();
+	bool HasDynamicData();
 
 private:
-	FNiagaraDynamicDataSprites *DynamicDataRender;
+	struct FNiagaraDynamicDataSprites *DynamicDataRender;
 	mutable TUniformBuffer<FPrimitiveUniformShaderParameters> WorldSpacePrimitiveUniformBuffer;
-	FParticleSpriteVertexFactory VertexFactory;
-	FParticleSpriteUniformParameters UniformParameters;
+	class FParticleSpriteVertexFactory *VertexFactory;
+	//FParticleSpriteUniformParameters UniformParameters;
 	TArray<FParticleSpriteUniformBufferRef, TInlineAllocator<2> > PerViewUniformBuffers;
 	FGlobalDynamicVertexBuffer::FAllocation DynamicVertexAllocation;
 };
@@ -463,310 +141,30 @@ private:
 
 
 
+/**
+* NiagaraEffectRendererRibbon renders an FNiagaraSimulation as a ribbon connecting all particles
+* in order by particle age.
+*/
 class NiagaraEffectRendererRibbon : public NiagaraEffectRenderer
 {
 public:
-	NiagaraEffectRendererRibbon(const UNiagaraComponent* InComponent, FPrimitiveSceneProxy *Proxy) :
-		NiagaraEffectRenderer(InComponent, Proxy),
-		DynamicDataRender(NULL),
-		VertexFactory(PVFT_BeamTrail, InComponent->GetWorld()->FeatureLevel)
-	{
-	}
-
+	NiagaraEffectRendererRibbon(ERHIFeatureLevel::Type FeatureLevel);
 	~NiagaraEffectRendererRibbon()
 	{
 		ReleaseRenderThreadResources();
 	}
 
+	virtual void ReleaseRenderThreadResources() override;
+	virtual void CreateRenderThreadResources() override;
 
-	virtual void ReleaseRenderThreadResources() override
-	{
-		VertexFactory.ReleaseResource();
-		PerViewUniformBuffers.Empty();
-		WorldSpacePrimitiveUniformBuffer.ReleaseResource();
-	}
-
-	// FPrimitiveSceneProxy interface.
-	virtual void CreateRenderThreadResources() override
-	{
-		VertexFactory.InitResource();
-		UniformParameters.CameraUp = FVector4(0.0f, 0.0f, 1.0f, 0.0f);
-		UniformParameters.CameraRight = FVector4(1.0f, 0.0f, 0.0f, 0.0f);
-		UniformParameters.ScreenAlignment = FVector4(0.0f, 0.0f, 0.0f, 0.0f);
-	}
-
-	virtual void GetDynamicMeshElements(const TArray<const FSceneView*>& Views, const FSceneViewFamily& ViewFamily, uint32 VisibilityMap, FMeshElementCollector& Collector) const override
-	{
-		SCOPE_CYCLE_COUNTER(STAT_NiagaraRender);
-		SCOPE_CYCLE_COUNTER(STAT_NiagaraRenderRibbons);
-
-		check(DynamicDataRender && DynamicDataRender->VertexData.Num());
-
-		const bool bIsWireframe = ViewFamily.EngineShowFlags.Wireframe;
-		FMaterialRenderProxy* MaterialRenderProxy = Material->GetRenderProxy(SceneProxy->IsSelected(), SceneProxy->IsHovered());
-
-		int32 SizeInBytes = DynamicDataRender->VertexData.GetTypeSize() * DynamicDataRender->VertexData.Num();
-		FGlobalDynamicVertexBuffer::FAllocation LocalDynamicVertexAllocation = FGlobalDynamicVertexBuffer::Get().Allocate(SizeInBytes);
-
-		if (LocalDynamicVertexAllocation.IsValid())
-		{
-			// Update the primitive uniform buffer if needed.
-			if (!WorldSpacePrimitiveUniformBuffer.IsInitialized())
-			{
-				FPrimitiveUniformShaderParameters PrimitiveUniformShaderParameters = GetPrimitiveUniformShaderParameters(
-					FMatrix::Identity,
-					SceneProxy->GetActorPosition(),
-					SceneProxy->GetBounds(),
-					SceneProxy->GetLocalBounds(),
-					SceneProxy->ReceivesDecals(),
-					false,
-					SceneProxy->UseEditorDepthTest()
-					);
-				WorldSpacePrimitiveUniformBuffer.SetContents(PrimitiveUniformShaderParameters);
-				WorldSpacePrimitiveUniformBuffer.InitResource();
-			}
-
-			// Copy the vertex data over.
-			FMemory::Memcpy(LocalDynamicVertexAllocation.Buffer, DynamicDataRender->VertexData.GetData(), SizeInBytes);
-
-			// Compute the per-view uniform buffers.
-			for (int32 ViewIndex = 0; ViewIndex < Views.Num(); ViewIndex++)
-			{
-				if (VisibilityMap & (1 << ViewIndex))
-				{
-					const FSceneView* View = Views[ViewIndex];
-
-					FNiagaraMeshCollectorResourcesRibbon& CollectorResources = Collector.AllocateOneFrameResource<FNiagaraMeshCollectorResourcesRibbon>();
-					FParticleBeamTrailUniformParameters PerViewUniformParameters = UniformParameters;
-
-					// Collector.AllocateOneFrameResource uses default ctor, initialize the vertex factory
-					CollectorResources.VertexFactory.SetFeatureLevel(ViewFamily.GetFeatureLevel());
-					CollectorResources.VertexFactory.SetParticleFactoryType(PVFT_BeamTrail);
-
-					CollectorResources.UniformBuffer = FParticleBeamTrailUniformBufferRef::CreateUniformBufferImmediate(PerViewUniformParameters, UniformBuffer_SingleFrame);
-
-					CollectorResources.VertexFactory.InitResource();
-					CollectorResources.VertexFactory.SetBeamTrailUniformBuffer(PerViewUniformBuffers[ViewIndex]);
-					CollectorResources.VertexFactory.SetVertexBuffer(DynamicVertexAllocation.VertexBuffer, DynamicVertexAllocation.VertexOffset, sizeof(FParticleBeamTrailVertex));
-					CollectorResources.VertexFactory.SetDynamicParameterBuffer(NULL, 0, 0);
-
-					FMeshBatch& MeshBatch = Collector.AllocateMesh();
-					MeshBatch.VertexFactory = &CollectorResources.VertexFactory;
-					MeshBatch.CastShadow = SceneProxy->CastsDynamicShadow();
-					MeshBatch.bUseAsOccluder = false;
-					MeshBatch.ReverseCulling = SceneProxy->IsLocalToWorldDeterminantNegative();
-					MeshBatch.bDisableBackfaceCulling = true;
-					MeshBatch.Type = PT_TriangleStrip;
-					MeshBatch.DepthPriorityGroup = SceneProxy->GetDepthPriorityGroup(View);
-					MeshBatch.bCanApplyViewModeOverrides = true;
-					MeshBatch.bUseWireframeSelectionColoring = SceneProxy->IsSelected();
-
-					if (bIsWireframe)
-					{
-						MeshBatch.MaterialRenderProxy = UMaterial::GetDefaultMaterial(MD_Surface)->GetRenderProxy(SceneProxy->IsSelected(), SceneProxy->IsHovered());
-					}
-					else
-					{
-						MeshBatch.MaterialRenderProxy = MaterialRenderProxy;
-					}
-
-					FMeshBatchElement& MeshElement = MeshBatch.Elements[0];
-					MeshElement.IndexBuffer = &GParticleIndexBuffer;
-					MeshElement.FirstIndex = 0;
-					MeshElement.NumPrimitives = (DynamicDataRender->VertexData.Num()-1) / 2;
-					MeshElement.NumInstances = 1;
-					MeshElement.MinVertexIndex = 0;
-					MeshElement.MaxVertexIndex = DynamicDataRender->VertexData.Num() - 1;
-					MeshElement.PrimitiveUniformBufferResource = &WorldSpacePrimitiveUniformBuffer;
-
-					Collector.AddMesh(ViewIndex, MeshBatch);
-				}
-			}
-		}
-	}
-
-	virtual void PreRenderView(const FSceneViewFamily* ViewFamily, const uint32 VisibilityMap, int32 FrameNumber) override
-	{
-		SCOPE_CYCLE_COUNTER(STAT_NiagaraPreRenderView);
-
-		check(DynamicDataRender && DynamicDataRender->VertexData.Num());
-
-		int32 SizeInBytes = DynamicDataRender->VertexData.GetTypeSize() * DynamicDataRender->VertexData.Num();
-		DynamicVertexAllocation = FGlobalDynamicVertexBuffer::Get().Allocate(SizeInBytes);
-		if (DynamicVertexAllocation.IsValid())
-		{
-			// Copy the vertex data over.
-			FMemory::Memcpy(DynamicVertexAllocation.Buffer, DynamicDataRender->VertexData.GetData(), SizeInBytes);
-			VertexFactory.SetVertexBuffer(DynamicVertexAllocation.VertexBuffer, DynamicVertexAllocation.VertexOffset, sizeof(FParticleBeamTrailVertex));
-			VertexFactory.SetDynamicParameterBuffer(NULL, 0, 0);
-			// Compute the per-view uniform buffers.
-			const int32 NumViews = ViewFamily->Views.Num();
-			uint32 ViewBit = 1;
-			for (int32 ViewIndex = 0; ViewIndex < NumViews; ++ViewIndex, ViewBit <<= 1)
-			{
-				FParticleBeamTrailUniformBufferRef* SpriteViewUniformBufferPtr = new(PerViewUniformBuffers)FParticleBeamTrailUniformBufferRef();
-				if (VisibilityMap & ViewBit)
-				{
-					FParticleBeamTrailUniformParameters PerViewUniformParameters = UniformParameters;
-					*SpriteViewUniformBufferPtr = FParticleBeamTrailUniformBufferRef::CreateUniformBufferImmediate(PerViewUniformParameters, UniformBuffer_SingleFrame);
-				}
-			}
-
-			// Update the primitive uniform buffer if needed.
-			if (!WorldSpacePrimitiveUniformBuffer.IsInitialized())
-			{
-				FPrimitiveUniformShaderParameters PrimitiveUniformShaderParameters = GetPrimitiveUniformShaderParameters(
-					FMatrix::Identity,
-					SceneProxy->GetActorPosition(),
-					SceneProxy->GetBounds(),
-					SceneProxy->GetLocalBounds(),
-					SceneProxy->ReceivesDecals(),
-					false,
-					false,
-					false
-					);
-				WorldSpacePrimitiveUniformBuffer.SetContents(PrimitiveUniformShaderParameters);
-				WorldSpacePrimitiveUniformBuffer.InitResource();
-			}
-		}
-	}
-
-
-	virtual void DrawDynamicElements(FPrimitiveDrawInterface* PDI, const FSceneView* View) override
-	{
-		SCOPE_CYCLE_COUNTER(STAT_NiagaraRender);
-
-		check(DynamicDataRender && DynamicDataRender->VertexData.Num());
-
-		const bool bIsWireframe = View->Family->EngineShowFlags.Wireframe;
-		FMaterialRenderProxy* MaterialRenderProxy = Material->GetRenderProxy(SceneProxy->IsSelected(), SceneProxy->IsHovered());
-
-		if (DynamicVertexAllocation.IsValid()
-			&& (bIsWireframe || !PDI->IsMaterialIgnored(MaterialRenderProxy, View->GetFeatureLevel())))
-		{
-			FMeshBatch MeshBatch;
-			MeshBatch.VertexFactory = &VertexFactory;
-			MeshBatch.CastShadow = SceneProxy->CastsDynamicShadow();
-			MeshBatch.bUseAsOccluder = false;
-			MeshBatch.bDisableBackfaceCulling = true;
-			MeshBatch.ReverseCulling = SceneProxy->IsLocalToWorldDeterminantNegative();
-			MeshBatch.Type = PT_TriangleStrip;
-			MeshBatch.DepthPriorityGroup = SceneProxy->GetDepthPriorityGroup(View);
-			MeshBatch.LCI = NULL;
-			if (bIsWireframe)
-			{
-				MeshBatch.MaterialRenderProxy = UMaterial::GetDefaultMaterial(MD_Surface)->GetRenderProxy(SceneProxy->IsSelected(), SceneProxy->IsHovered());
-			}
-			else
-			{
-				MeshBatch.MaterialRenderProxy = MaterialRenderProxy;
-			}
-
-			FMeshBatchElement& MeshElement = MeshBatch.Elements[0];
-			MeshElement.IndexBuffer = &GParticleIndexBuffer;
-			MeshElement.FirstIndex = 0;
-			MeshElement.NumPrimitives = DynamicDataRender->VertexData.Num()/2;
-			MeshElement.NumInstances = 1;
-			MeshElement.MinVertexIndex = 0;
-			MeshElement.MaxVertexIndex = DynamicDataRender->VertexData.Num() - 1;
-			MeshElement.PrimitiveUniformBufferResource = &WorldSpacePrimitiveUniformBuffer;
-
-			int32 ViewIndex = View->Family->Views.Find(View);
-			check(PerViewUniformBuffers.IsValidIndex(ViewIndex) && PerViewUniformBuffers[ViewIndex]);
-			VertexFactory.SetBeamTrailUniformBuffer(PerViewUniformBuffers[ViewIndex]);
-
-			DrawRichMesh(
-				PDI,
-				MeshBatch,
-				FLinearColor(1.0f, 0.0f, 0.0f),	//WireframeColor,
-				FLinearColor(1.0f, 1.0f, 0.0f),	//LevelColor,
-				FLinearColor(1.0f, 1.0f, 1.0f),	//PropertyColor,		
-				SceneProxy,
-				GIsEditor && (View->Family->EngineShowFlags.Selection) ? SceneProxy->IsSelected() : false,
-				bIsWireframe
-				);
-		}
-
-		/*
-		for(int32 i=0; i<DynamicDataRender->VertexData.Num(); i++)
-		{
-		FParticleSpriteVertex& Vertex = DynamicDataRender->VertexData[i];
-		PDI->DrawPoint(Vertex.Position, Vertex.Color, 1.0f, SDPG_World);
-		}
-		*/
-	}
 
+	virtual void GetDynamicMeshElements(const TArray<const FSceneView*>& Views, const FSceneViewFamily& ViewFamily, uint32 VisibilityMap, FMeshElementCollector& Collector, const FNiagaraSceneProxy *SceneProxy) const override;
 
+	virtual void PreRenderView(const FSceneViewFamily* ViewFamily, const uint32 VisibilityMap, int32 FrameNumber, const FNiagaraSceneProxy *SceneProxy) override;
+	virtual void DrawDynamicElements(FPrimitiveDrawInterface* PDI, const FSceneView* View, const FNiagaraSceneProxy *SceneProxy) override;
 
 	/** Update render data buffer from attributes */
-	FNiagaraDynamicDataBase *GenerateVertexData(const FNiagaraEmitterParticleData &Data)
-	{
-		SCOPE_CYCLE_COUNTER(STAT_NiagaraGenRibbonVertexData);
-		FNiagaraDynamicDataRibbon *DynamicData = new FNiagaraDynamicDataRibbon;
-		TArray<FParticleBeamTrailVertex>& RenderData = DynamicData->VertexData;
-
-		RenderData.Reset(Data.GetNumParticles() * 2);
-		//CachedBounds.Init();
-
-		// build a sorted list by age, so we always get particles in order 
-		// regardless of them being moved around due to dieing and spawning
-		TArray<int32> SortedIndices;
-		for (uint32 Idx = 0; Idx < Data.GetNumParticles(); Idx++)
-		{
-			SortedIndices.Add(Idx);
-		}
-
-		const FVector4 *AgeData = Data.GetAttributeData("Age");
-		SortedIndices.Sort(
-			[&AgeData](const int32& A, const int32& B) {
-			return AgeData[A].X < AgeData[B].X;
-		}
-		);
-
-
-		FVector2D UVs[4] = { FVector2D(0.0f, 0.0f), FVector2D(1.0f, 0.0f), FVector2D(1.0f, 1.0f), FVector2D(0.0f, 1.0f) };
-
-		const FVector4 *PosPtr = Data.GetAttributeData("Position");
-		const FVector4 *ColorPtr = Data.GetAttributeData("Color");
-		const FVector4 *AgePtr = Data.GetAttributeData("Age");
-		const FVector4 *RotPtr = Data.GetAttributeData("Rotation");
-
-		FVector PrevPos, PrevPos2, PrevDir(0.0f, 0.0f, 0.1f);
-		for (int32 i = 0; i < SortedIndices.Num() - 1; i++)
-		{
-			uint32 Index1 = SortedIndices[i];
-			uint32 Index2 = SortedIndices[i + 1];
-
-			const FVector ParticlePos = PosPtr[Index1];
-			FVector ParticleDir = PosPtr[Index2] - ParticlePos;
-			if (ParticleDir.Size() <= SMALL_NUMBER)
-			{
-				ParticleDir = PrevDir*0.1f;
-			}
-
-			FVector ParticleRight = FVector::CrossProduct(ParticleDir, FVector(0.0f, 0.0f, 1.0f));
-			ParticleRight.Normalize();
-
-			if (i == 0)
-			{
-				AddRibbonVert(RenderData, ParticlePos + ParticleRight, Data, UVs[0], ColorPtr[i], AgePtr[i], RotPtr[i]);
-				AddRibbonVert(RenderData, ParticlePos - ParticleRight, Data, UVs[1], ColorPtr[i], AgePtr[i], RotPtr[i]);
-			}
-			else
-			{
-				AddRibbonVert(RenderData, PrevPos2, Data, UVs[0], ColorPtr[i], AgePtr[i], RotPtr[i]);
-				AddRibbonVert(RenderData, PrevPos, Data, UVs[1], ColorPtr[i], AgePtr[i], RotPtr[i]);
-			}
-
-			AddRibbonVert(RenderData, ParticlePos - ParticleRight + ParticleDir, Data, UVs[2], ColorPtr[i], AgePtr[i], RotPtr[i]);
-			AddRibbonVert(RenderData, ParticlePos + ParticleRight + ParticleDir, Data, UVs[3], ColorPtr[i], AgePtr[i], RotPtr[i]);
-			PrevPos = ParticlePos - ParticleRight + ParticleDir;
-			PrevPos2 = ParticlePos + ParticleRight + ParticleDir;
-			PrevDir = ParticleDir;
-		}
-
-		return DynamicData;
-	}
+	FNiagaraDynamicDataBase *GenerateVertexData(const FNiagaraEmitterParticleData &Data);
 
 	void AddRibbonVert(TArray<FParticleBeamTrailVertex>& RenderData, FVector ParticlePos, const FNiagaraEmitterParticleData &Data, FVector2D UV1,
 		const FVector4 &Color, const FVector4 &Age, const FVector4 &Rotation)
@@ -788,41 +186,17 @@ public:
 
 
 
-	virtual void SetDynamicData_RenderThread(FNiagaraDynamicDataBase* NewDynamicData) override
-	{
-		check(IsInRenderingThread());
-
-		if (DynamicDataRender)
-		{
-			delete DynamicDataRender;
-			DynamicDataRender = NULL;
-		}
-		DynamicDataRender = static_cast<FNiagaraDynamicDataRibbon*>(NewDynamicData);
-	}
-
-	int GetDynamicDataSize()
-	{
-		uint32 Size = sizeof(FNiagaraDynamicDataRibbon);
-		if (DynamicDataRender)
-		{
-			Size += DynamicDataRender->VertexData.GetAllocatedSize();
-		}
-
-		return Size;
-	}
-
-	bool HasDynamicData()
-	{
-		return DynamicDataRender && DynamicDataRender->VertexData.Num() > 0;
-	}
+	virtual void SetDynamicData_RenderThread(FNiagaraDynamicDataBase* NewDynamicData) override;
+	int GetDynamicDataSize();
+	bool HasDynamicData();
 
 
 private:
-	FNiagaraDynamicDataRibbon *DynamicDataRender;
+	struct FNiagaraDynamicDataRibbon *DynamicDataRender;
 	mutable TUniformBuffer<FPrimitiveUniformShaderParameters> WorldSpacePrimitiveUniformBuffer;
-	FParticleBeamTrailVertexFactory VertexFactory;
-	FParticleBeamTrailUniformParameters UniformParameters;
-	TArray<FParticleBeamTrailUniformBufferRef, TInlineAllocator<2> > PerViewUniformBuffers;
+	class FParticleBeamTrailVertexFactory *VertexFactory;
+	//FParticleBeamTrailUniformParameters UniformParameters;
+	TArray<	FParticleBeamTrailUniformBufferRef, TInlineAllocator<2> > PerViewUniformBuffers;
 	FGlobalDynamicVertexBuffer::FAllocation DynamicVertexAllocation;
 };
 
diff --git a/Engine/Source/Runtime/Engine/Classes/Engine/NiagaraSimulation.h b/Engine/Source/Runtime/Engine/Classes/Engine/NiagaraSimulation.h
new file mode 100644
--- /dev/null
+++ b/Engine/Source/Runtime/Engine/Classes/Engine/NiagaraSimulation.h
@@ -0,0 +1,149 @@
+// Copyright 1998-2014 Epic Games, Inc. All Rights Reserved.
+
+/*==============================================================================
+NiagaraSimulation.h: Niagara emitter simulation class
+==============================================================================*/
+#pragma once
+
+#include "NiagaraScript.h"
+#include "Components/NiagaraComponent.h"
+#include "Runtime/VectorVM/Public/VectorVM.h"
+#include "NiagaraSimulation.generated.h"
+
+DECLARE_CYCLE_STAT(TEXT("Tick"), STAT_NiagaraTick, STATGROUP_Niagara);
+DECLARE_CYCLE_STAT(TEXT("Simulate"), STAT_NiagaraSimulate, STATGROUP_Niagara);
+DECLARE_CYCLE_STAT(TEXT("Spawn + Kill"), STAT_NiagaraSpawnAndKill, STATGROUP_Niagara);
+
+
+class NiagaraEffectRenderer;
+
+UENUM()
+enum EEmitterRenderModuleType
+{
+	RMT_None = 0,
+	RMT_Sprites,
+	RMT_Ribbon,
+	RMT_Trails,
+	RMT_Meshes
+};
+
+
+/** 
+ *	FniagaraEmitterProperties stores the attributes of an FNiagaraSimulation
+ *	that need to be serialized and are used for its initialization 
+ */
+USTRUCT()
+struct FNiagaraEmitterProperties
+{
+	GENERATED_USTRUCT_BODY()
+public:
+	UPROPERTY(EditAnywhere, Category = "Emitter data")
+	float SpawnRate;
+
+	UPROPERTY(EditAnywhere, Category = "Emitter data")
+	UNiagaraScript *UpdateScript;
+	UPROPERTY(EditAnywhere, Category = "Emitter data")
+	UNiagaraScript *SpawnScript;
+
+	UPROPERTY(EditAnywhere, Category = "Emitter data")
+	UMaterial *Material;
+};
+
+
+
+/**
+* A niagara particle simulation.
+*/
+class FNiagaraSimulation
+{
+public:
+	explicit FNiagaraSimulation();
+
+	void Tick(float DeltaSeconds);
+
+
+	FBox GetBounds() const { return CachedBounds; }
+
+
+	FNiagaraEmitterParticleData &GetData()	{ return Data; }
+
+	void SetConstants(const FNiagaraConstantMap &InMap)	{ Constants = InMap; }
+	NiagaraEffectRenderer *GetEffectRenderer()	{ return EffectRenderer;  }
+	
+	void SetUpdateScript(UNiagaraScript *InScript)	{ UpdateScript = InScript;  }
+	void SetSpawnScript(UNiagaraScript *InScript)	{ SpawnScript = InScript; }
+
+	bool IsEnabled()	{ return bIsEnabled;  }
+	void SetEnabled(bool bInEnabled)	{ bIsEnabled = bInEnabled;  }
+
+	void SetSpawnRate(float InRate)	{ SpawnRate = InRate; }
+	float GetSpawnRate()		{ return SpawnRate;  }
+
+	EEmitterRenderModuleType GetRenderModuleType()	{ return RenderModuleType; }
+	ENGINE_API void SetRenderModuleType(EEmitterRenderModuleType Type, ERHIFeatureLevel::Type FeatureLevel);
+
+	int32 GetNumParticles()	{ return Data.GetNumParticles(); }
+private:
+	float Age;
+	bool bIsEnabled;
+	/** Temporary stuff for the prototype. */
+	float SpawnRate;
+	EEmitterRenderModuleType RenderModuleType;
+
+	/** The particle update script. */
+	UNiagaraScript *UpdateScript;
+	/** The particle spawn script. */
+	UNiagaraScript *SpawnScript;
+	/** Local constant set. */
+	FNiagaraConstantMap Constants;
+	/** particle simulation data */
+	FNiagaraEmitterParticleData Data;
+	/** Keep partial particle spawns from last frame */
+	float SpawnRemainder;
+	/** The cached ComponentToWorld transform. */
+	FTransform CachedComponentToWorld;
+	/** Cached bounds. */
+	FBox CachedBounds;
+
+	NiagaraEffectRenderer *EffectRenderer;
+
+	/** Calc number to spawn */
+	int32 CalcNumToSpawn(float DeltaSeconds)
+	{
+		float FloatNumToSpawn = SpawnRemainder + (DeltaSeconds * SpawnRate);
+		int32 NumToSpawn = FMath::FloorToInt(FloatNumToSpawn);
+		SpawnRemainder = FloatNumToSpawn - NumToSpawn;
+		return NumToSpawn;
+	}
+
+	/** Run VM to update particle positions */
+	void UpdateParticles(
+		float DeltaSeconds,
+		FVector4* PrevParticles,
+		int32 PrevNumVectorsPerAttribute,
+		FVector4* Particles,
+		int32 NumVectorsPerAttribute,
+		int32 NumParticles
+		);
+
+	int32 SpawnAndKillParticles(int32 NumToSpawn);
+
+
+	/** Spawn a new particle at this index */
+	int32 SpawnParticles(int32 NumToSpawn);
+
+	/** Util to move a particle */
+	void MoveParticleToIndex(int32 SrcIndex, int32 DestIndex)
+	{
+		FVector4 *SrcPtr = Data.GetCurrentBuffer() + SrcIndex;
+		FVector4 *DestPtr = Data.GetCurrentBuffer() + DestIndex;
+
+		for (int32 AttrIndex = 0; AttrIndex < Data.GetNumAttributes(); AttrIndex++)
+		{
+			*DestPtr = *SrcPtr;
+			DestPtr += Data.GetParticleAllocation();
+			SrcPtr += Data.GetParticleAllocation();
+		}
+	}
+
+};
diff --git a/Engine/Source/Runtime/Engine/Engine.Build.cs b/Engine/Source/Runtime/Engine/Engine.Build.cs
--- a/Engine/Source/Runtime/Engine/Engine.Build.cs
+++ b/Engine/Source/Runtime/Engine/Engine.Build.cs
@@ -52,7 +52,8 @@ public class Engine : ModuleRules
 				"EngineMessages",
 				"EngineSettings",
 				"SynthBenchmark",                
-                "AIModule", // Here until all engine's dependencies of AI code are removed
+                "AIModule",
+				"VectorVM",
 			}
 		);
 
diff --git a/Engine/Source/Runtime/Engine/Private/NiagaraComponent.cpp b/Engine/Source/Runtime/Engine/Private/NiagaraComponent.cpp
--- a/Engine/Source/Runtime/Engine/Private/NiagaraComponent.cpp
+++ b/Engine/Source/Runtime/Engine/Private/NiagaraComponent.cpp
@@ -3,36 +3,38 @@
 #include "EnginePrivate.h"
 #include "Components/NiagaraComponent.h"
 #include "Engine/NiagaraScript.h"
+#include "Engine/NiagaraConstants.h"
 #include "VectorVM.h"
 #include "ParticleHelper.h"
 #include "Particles/ParticleResources.h"
 #include "Engine/NiagaraEffectRenderer.h"
+#include "Engine/NiagaraEffect.h"
+#include "Engine/NiagaraSimulation.h"
 #include "MeshBatch.h"
 #include "SceneUtils.h"
 #include "ComponentReregisterContext.h"
 
-DECLARE_CYCLE_STAT(TEXT("Tick"),STAT_NiagaraTick,STATGROUP_Niagara);
-DECLARE_CYCLE_STAT(TEXT("Simulate"),STAT_NiagaraSimulate,STATGROUP_Niagara);
-DECLARE_CYCLE_STAT(TEXT("Spawn + Kill"),STAT_NiagaraSpawnAndKill,STATGROUP_Niagara);
 DECLARE_CYCLE_STAT(TEXT("Gen Verts"),STAT_NiagaraGenerateVertices,STATGROUP_Niagara);
 DECLARE_DWORD_COUNTER_STAT(TEXT("NumParticles"),STAT_NiagaraNumParticles,STATGROUP_Niagara);
 
 
-DEFINE_LOG_CATEGORY_STATIC(LogNiagaraComponent, All, All);
-
 
 FNiagaraSceneProxy::FNiagaraSceneProxy(const UNiagaraComponent* InComponent)
 		:	FPrimitiveSceneProxy(InComponent)
 {
-	if (InComponent->RenderModuleType == Sprites)
-	{
-		EffectRenderer = new NiagaraEffectRendererSprites(InComponent, this);
-	}
-	else  if (InComponent->RenderModuleType == Ribbon)
+	UpdateEffectRenderers(InComponent->Effect);
+}
+
+void FNiagaraSceneProxy::UpdateEffectRenderers(UNiagaraEffect *InEffect)
+{
+	EffectRenderers.Empty();
+	if (InEffect)
 	{
-		EffectRenderer = new NiagaraEffectRendererRibbon(InComponent, this);
+		for (TSharedPtr<FNiagaraSimulation>Emitter : InEffect->Emitters)
+		{
+			AddEffectRenderer(Emitter->GetEffectRenderer());
+		}
 	}
-
 }
 
 FNiagaraSceneProxy::~FNiagaraSceneProxy()
@@ -43,21 +45,30 @@ FNiagaraSceneProxy::~FNiagaraSceneProxy()
 /** Called on render thread to assign new dynamic data */
 void FNiagaraSceneProxy::SetDynamicData_RenderThread(FNiagaraDynamicDataBase* NewDynamicData)
 {
-	EffectRenderer->SetDynamicData_RenderThread(NewDynamicData);
+	for (NiagaraEffectRenderer *Renderer : EffectRenderers)
+	{
+		Renderer->SetDynamicData_RenderThread(NewDynamicData);
+	}
 	return;
 }
 
 
 void FNiagaraSceneProxy::ReleaseRenderThreadResources()
 {
-	EffectRenderer->ReleaseRenderThreadResources();
+	for (NiagaraEffectRenderer *Renderer : EffectRenderers)
+	{
+		Renderer->ReleaseRenderThreadResources();
+	}
 	return;
 }
 
 // FPrimitiveSceneProxy interface.
 void FNiagaraSceneProxy::CreateRenderThreadResources()
 {
-	EffectRenderer->CreateRenderThreadResources();
+	for (NiagaraEffectRenderer *Renderer : EffectRenderers)
+	{
+		Renderer->CreateRenderThreadResources();
+	}
 	return;
 }
 
@@ -73,27 +84,34 @@ void FNiagaraSceneProxy::OnTransformChanged()
 
 void FNiagaraSceneProxy::PreRenderView(const FSceneViewFamily* ViewFamily, const uint32 VisibilityMap, int32 FrameNumber)
 {
-	EffectRenderer->PreRenderView(ViewFamily, VisibilityMap, FrameNumber);
+	for (NiagaraEffectRenderer *Renderer : EffectRenderers)
+	{
+		Renderer->PreRenderView(ViewFamily, VisibilityMap, FrameNumber, this);
+	}
 	return;
-
 }
 		  
 void FNiagaraSceneProxy::DrawDynamicElements(FPrimitiveDrawInterface* PDI, const FSceneView* View) 
 {
-	EffectRenderer->DrawDynamicElements(PDI, View);
+	for (NiagaraEffectRenderer *Renderer : EffectRenderers)
+	{
+		Renderer->DrawDynamicElements(PDI, View, this);
+	}
 	return;
 }
 
 FPrimitiveViewRelevance FNiagaraSceneProxy::GetViewRelevance(const FSceneView* View)
 {
-	return EffectRenderer->GetViewRelevance(View);
-}
-/*
-virtual bool CanBeOccluded() const override
-{
-	return !MaterialRelevance.bDisableDepthTest;
+	FPrimitiveViewRelevance Relevance;
+	Relevance.bDynamicRelevance = true;
+
+	for (NiagaraEffectRenderer *Renderer : EffectRenderers)
+	{
+		Relevance |= EffectRenderers[0]->GetViewRelevance(View, this);
+	}
+	return Relevance;
 }
-*/
+
 
 uint32 FNiagaraSceneProxy::GetMemoryFootprint() const
 { 
@@ -103,13 +121,20 @@ uint32 FNiagaraSceneProxy::GetMemoryFootprint() const
 uint32 FNiagaraSceneProxy::GetAllocatedSize() const
 { 
 	uint32 DynamicDataSize = 0;
-	return FPrimitiveSceneProxy::GetAllocatedSize() + EffectRenderer->GetDynamicDataSize();
+	for (NiagaraEffectRenderer *Renderer : EffectRenderers)
+	{
+		DynamicDataSize += Renderer->GetDynamicDataSize();
+	}
+	return FPrimitiveSceneProxy::GetAllocatedSize() + DynamicDataSize;
 }
 
 
 void FNiagaraSceneProxy::GetDynamicMeshElements(const TArray<const FSceneView*>& Views, const FSceneViewFamily& ViewFamily, uint32 VisibilityMap, FMeshElementCollector& Collector) const
 {
-	EffectRenderer->GetDynamicMeshElements(Views, ViewFamily, VisibilityMap, Collector);
+	for (NiagaraEffectRenderer *Renderer : EffectRenderers)
+	{
+		Renderer->GetDynamicMeshElements(Views, ViewFamily, VisibilityMap, Collector, this);
+	}
 }
 
 
@@ -129,253 +154,6 @@ namespace ENiagaraVectorAttr
 }
 
 
-/**
- * A niagara particle simulation.
- */
-class FNiagaraSimulation
-{
-public:
-	explicit FNiagaraSimulation(const UNiagaraComponent* InComponent)
-		: SpawnRate(InComponent->SpawnRate)
-		, Component(*InComponent)
-		, UpdateScript(*InComponent->UpdateScript)
-		, SpawnScript(InComponent->SpawnScript)
-		, SpawnRemainder(0.0f)
-		, CachedBounds(ForceInit)
-	{
-		check(InComponent->UpdateScript && UpdateScript.ByteCode.Num());
-	}
-
-	void Tick(float DeltaSeconds)
-	{
-		SCOPE_CYCLE_COUNTER(STAT_NiagaraTick);
-
-		const int32 NumAttributes = UpdateScript.Attributes.Num();
-
-		// Cache the ComponentToWorld transform.
-		CachedComponentToWorld = Component.GetComponentToWorld();
-		
-		Data.SwapBuffers();
-
-		// Figure out how many we will spawn.
-		int32 NumToSpawn = CalcNumToSpawn(DeltaSeconds);
-		int32 MaxNewParticles = Data.GetNumParticles() + NumToSpawn;
-
-		// Remember the stride of the original data.
-		int32 PrevNumVectorsPerAttribute = Data.GetParticleAllocation();
-
-		Data.Allocate( MaxNewParticles );
-		// Simulate particles forward by DeltaSeconds.
-		UpdateParticles(
-			DeltaSeconds,
-			Data.GetPreviousBuffer(),
-			PrevNumVectorsPerAttribute,
-			Data.GetCurrentBuffer(),
-			MaxNewParticles,
-			Data.GetNumParticles()
-			);
-
-		SpawnAndKillParticles(
-			NumToSpawn
-			);
-
-		DECLARE_DWORD_COUNTER_STAT(TEXT("NumParticles"), STAT_NiagaraNumParticles, STATGROUP_Niagara);
-		INC_DWORD_STAT_BY(STAT_NiagaraNumParticles, Data.GetNumParticles());
-
-	}
-
-	FBox GetBounds() const { return CachedBounds; }
-
-
-	FNiagaraEmitterParticleData &GetData()	{ return Data; }
-
-	void SetConstant(FName ConstantName, const float Value)
-	{
-		Constants.SetOrAdd(ConstantName,Value);
-	}
-
-	void SetConstant(FName ConstantName, const FVector4& Value)
-	{
-		Constants.SetOrAdd(ConstantName, Value);
-	}
-
-	void SetConstant(FName ConstantName, const FMatrix& Value)
-	{
-		Constants.SetOrAdd(ConstantName, Value);
-	}
-private:
-	/** Temporary stuff for the prototype. */
-	float SpawnRate;
-
-	/** The component. */
-	const UNiagaraComponent& Component;
-	/** The particle update script. */
-	UNiagaraScript& UpdateScript;
-	/** The particle spawn script. */
-	UNiagaraScript *SpawnScript;
-	/** Local constant set. */
-	FNiagaraConstantMap Constants;
-	/** particle simulation data */
-	FNiagaraEmitterParticleData Data;
-	/** Keep partial particle spawns from last frame */
-	float SpawnRemainder;
-	/** The cached ComponentToWorld transform. */
-	FTransform CachedComponentToWorld;
-	/** Cached bounds. */
-	FBox CachedBounds;
-
-	/** Calc number to spawn */
-	int32 CalcNumToSpawn(float DeltaSeconds)
-	{
-		float FloatNumToSpawn = SpawnRemainder + (DeltaSeconds * SpawnRate);
-		int32 NumToSpawn = FMath::FloorToInt(FloatNumToSpawn);
-		SpawnRemainder = FloatNumToSpawn - NumToSpawn;
-		return NumToSpawn;
-	}
-
-	/** Run VM to update particle positions */
-	void UpdateParticles(
-		float DeltaSeconds,
-		FVector4* PrevParticles,
-		int32 PrevNumVectorsPerAttribute,
-		FVector4* Particles,
-		int32 NumVectorsPerAttribute,
-		int32 NumParticles
-		)
-	{
-		SCOPE_CYCLE_COUNTER(STAT_NiagaraSimulate);
-
-		VectorRegister* InputRegisters[VectorVM::MaxInputRegisters] = { 0 };
-		VectorRegister* OutputRegisters[VectorVM::MaxOutputRegisters] = {0};
-		const int32 NumAttr = UpdateScript.Attributes.Num();
-
-		check(NumAttr < VectorVM::MaxInputRegisters);
-		check(NumAttr < VectorVM::MaxOutputRegisters);
-
-		// Setup input and output registers.
-		for (int32 AttrIndex = 0; AttrIndex < NumAttr; ++AttrIndex)
-		{
-			InputRegisters[AttrIndex] = (VectorRegister*)(PrevParticles + AttrIndex * PrevNumVectorsPerAttribute);
-			OutputRegisters[AttrIndex] = (VectorRegister*)(Particles + AttrIndex * NumVectorsPerAttribute);
-		}
-
-		//Fill constant table with required emitter constants and internal script constants.
-		TArray<FVector4> ConstantTable;
-		UpdateScript.ConstantData.FillConstantTable(Constants, ConstantTable);
-
-		VectorVM::Exec(
-			UpdateScript.ByteCode.GetData(),
-			InputRegisters,
-			NumAttr,
-			OutputRegisters,
-			NumAttr,
-			ConstantTable.GetData(),
-			NumParticles
-			);
-	}
-
-	int32 SpawnAndKillParticles(int32 NumToSpawn)
-	{
-		SCOPE_CYCLE_COUNTER(STAT_NiagaraSpawnAndKill);
-		int32 OrigNumParticles = Data.GetNumParticles();
-		int32 CurNumParticles = OrigNumParticles;
-		CurNumParticles = SpawnParticles(NumToSpawn);
-
-		// run the spawn graph over all new particles
-		if (SpawnScript && SpawnScript->ByteCode.Num())
-		{
-			VectorRegister* InputRegisters[VectorVM::MaxInputRegisters] = { 0 };
-			VectorRegister* OutputRegisters[VectorVM::MaxOutputRegisters] = { 0 };
-			const int32 NumAttr = SpawnScript->Attributes.Num();
-			const int32 NumVectors = NumToSpawn;
-
-			check(NumAttr < VectorVM::MaxInputRegisters);
-			check(NumAttr < VectorVM::MaxOutputRegisters);
-
-			FVector4 *NewParticlesStart = Data.GetCurrentBuffer() + OrigNumParticles;
-
-			// Setup input and output registers.
-			for (int32 AttrIndex = 0; AttrIndex < NumAttr; ++AttrIndex)
-			{
-				InputRegisters[AttrIndex] = (VectorRegister*)(NewParticlesStart + AttrIndex * Data.GetParticleAllocation());
-				OutputRegisters[AttrIndex] = (VectorRegister*)(NewParticlesStart + AttrIndex * Data.GetParticleAllocation());
-			}
-
-			//Fill constant table with required emitter constants and internal script constants.
-			TArray<FVector4> ConstantTable;
-			SpawnScript->ConstantData.FillConstantTable(Constants, ConstantTable);
-
-			VectorVM::Exec(
-				SpawnScript->ByteCode.GetData(),
-				InputRegisters,
-				NumAttr,
-				OutputRegisters,
-				NumAttr,
-				ConstantTable.GetData(),
-				NumVectors
-				);
-		}
-
-		// Iterate over looking for dead particles and move from the end of the list to the dead location, compacting in the process
-		int32 ParticleIndex = 0;
-		const FVector4* ParticleRelativeTimes = Data.GetAttributeData("Age");
-		while (ParticleIndex < OrigNumParticles)
-		{
-			if (ParticleRelativeTimes[ParticleIndex].X > 1.0f)
-			{
-				// Particle is dead, move one from the end here.
-				MoveParticleToIndex(--CurNumParticles, ParticleIndex);
-			}
-			ParticleIndex++;
-		}
-
-		Data.SetNumParticles(CurNumParticles);
-		return CurNumParticles;
-	}
-
-	/** Spawn a new particle at this index */
-	int32 SpawnParticles(int32 NumToSpawn)
-	{
-		FVector4 *PosPtr = Data.GetAttributeDataWrite("Position");
-		FVector4 *VelPtr = Data.GetAttributeDataWrite("Velocity");
-		FVector4 *ColPtr = Data.GetAttributeDataWrite("Color");
-		FVector4 *RotPtr = Data.GetAttributeDataWrite("Rotation");
-		FVector4 *AgePtr = Data.GetAttributeDataWrite("Age");
-
-		// Spawn new Particles at the end of the buffer
-		int32 ParticleIndex = Data.GetNumParticles();
-		for (int32 i = 0; i < NumToSpawn; i++)
-		{
-			FVector SpawnLocation = CachedComponentToWorld.GetLocation();
-			SpawnLocation.X += FMath::FRandRange(-20.f, 20.f);
-			SpawnLocation.Y += FMath::FRandRange(-20.f, 20.f);
-			SpawnLocation.Z += FMath::FRandRange(-20.f, 20.f);
-
-			PosPtr[ParticleIndex] = SpawnLocation;
-			ColPtr[ParticleIndex] = FVector4(0.0f, 0.0f, 0.0f, 0.0f);
-			VelPtr[ParticleIndex] = FVector4(0.0f, 0.0f, 2.0f, 0.0f);
-			RotPtr[ParticleIndex] = FVector4(0.0f, 0.0f, 0.0f, 0.0f);
-			AgePtr[ParticleIndex] = FVector4(0.0f, 0.0f, 0.0f, 0.0f);
-			ParticleIndex++;
-		}
-		return ParticleIndex;
-	}
-
-	/** Util to move a particle */
-	void MoveParticleToIndex(int32 SrcIndex, int32 DestIndex)
-	{
-		FVector4 *SrcPtr = Data.GetCurrentBuffer() + SrcIndex;
-		FVector4 *DestPtr = Data.GetCurrentBuffer() + DestIndex;
-
-		for (int32 AttrIndex = 0; AttrIndex < Data.GetNumAttributes(); AttrIndex++)
-		{
-			*DestPtr = *SrcPtr;
-			DestPtr += Data.GetParticleAllocation();
-			SrcPtr += Data.GetParticleAllocation();
-		}
-	}
-
-};
 
 //////////////////////////////////////////////////////////////////////////
 
@@ -384,51 +162,46 @@ UNiagaraComponent::UNiagaraComponent(const FObjectInitializer& ObjectInitializer
 {
 	PrimaryComponentTick.bCanEverTick = true;
 	bTickInEditor = true;
-	RenderModuleType = ERenderModuleType::Sprites;
-	SpawnRate = 20.f;
-	Simulation = nullptr;
 }
 
 
 void UNiagaraComponent::TickComponent(float DeltaSeconds, enum ELevelTick TickType, FActorComponentTickFunction *ThisTickFunction)
 {
-	if (Simulation)
+//	EmitterAge += DeltaSeconds;
+
+	if (Effect)
 	{
-		EmitterAge += DeltaSeconds;
-	
+
 		//Todo, open this up to the UI and setting via code and BPs.
 		static FName Const_Zero(TEXT("ZERO"));
-		Simulation->SetConstant(Const_Zero, FVector4(0.0f, 0.0f, 0.0f, 0.0f));	// zero constant
+		Effect->SetConstant(Const_Zero, FVector4(0.0f, 0.0f, 0.0f, 0.0f));	// zero constant
 		static FName Const_DeltaTime(TEXT("Delta Time"));
-		Simulation->SetConstant(Const_DeltaTime, FVector4(DeltaSeconds, DeltaSeconds, DeltaSeconds, DeltaSeconds));
+		Effect->SetConstant(Const_DeltaTime, FVector4(DeltaSeconds, DeltaSeconds, DeltaSeconds, DeltaSeconds));
 		static FName Const_EmitterPos(TEXT("Emitter Position"));
-		Simulation->SetConstant(Const_EmitterPos, FVector4(ComponentToWorld.GetTranslation()));
-		static FName Const_Age(TEXT("Emitter Age"));
-		Simulation->SetConstant(Const_Age, FVector4(EmitterAge, EmitterAge, EmitterAge, EmitterAge));
+		Effect->SetConstant(Const_EmitterPos, FVector4(ComponentToWorld.GetTranslation()));
+		//static FName Const_Age(TEXT("Emitter Age"));
+		//Effect->SetConstant(Const_Age, FVector4(EmitterAge, EmitterAge, EmitterAge, EmitterAge));
 		static FName Const_EmitterX(TEXT("Emitter X Axis"));
-		Simulation->SetConstant(Const_EmitterX, FVector4(ComponentToWorld.GetUnitAxis(EAxis::X)));
+		Effect->SetConstant(Const_EmitterX, FVector4(ComponentToWorld.GetUnitAxis(EAxis::X)));
 		static FName Const_EmitterY(TEXT("Emitter Y Axis"));
-		Simulation->SetConstant(Const_EmitterY, FVector4(ComponentToWorld.GetUnitAxis(EAxis::Y)));
+		Effect->SetConstant(Const_EmitterY, FVector4(ComponentToWorld.GetUnitAxis(EAxis::Y)));
 		static FName Const_EmitterZ(TEXT("Emitter Z Axis"));
-		Simulation->SetConstant(Const_EmitterZ, FVector4(ComponentToWorld.GetUnitAxis(EAxis::Z)));
-
+		Effect->SetConstant(Const_EmitterZ, FVector4(ComponentToWorld.GetUnitAxis(EAxis::Z)));
 		static FName Const_EmitterTransform(TEXT("Emitter Transform"));
-		Simulation->SetConstant(Const_EmitterTransform, ComponentToWorld.ToMatrixWithScale());
-
-		Simulation->Tick(DeltaSeconds);
-		UpdateComponentToWorld();
-		MarkRenderDynamicDataDirty();
+		Effect->SetConstant(Const_EmitterTransform, ComponentToWorld.ToMatrixWithScale());
+		Effect->Tick(DeltaSeconds);
 	}
+
+	UpdateComponentToWorld();
+	MarkRenderDynamicDataDirty();
 }
 
 void UNiagaraComponent::OnRegister()
 {
 	Super::OnRegister();
-
-	ensure(Simulation == NULL);
-	if (UpdateScript && UpdateScript->ByteCode.Num() && UpdateScript->Attributes.Num() == ENiagaraVectorAttr::MaxVectorAttribs)
+	if (Effect)
 	{
-		Simulation = new FNiagaraSimulation(this);
+		Effect->Init(this);
 	}
 	VectorVM::Init();
 }
@@ -437,29 +210,28 @@ void UNiagaraComponent::OnRegister()
 void UNiagaraComponent::OnUnregister()
 {
 	Super::OnUnregister();
-
-	if(Simulation != NULL)
-	{
-		delete Simulation;
-		Simulation = NULL;
-	}
 }
 
 void UNiagaraComponent::SendRenderDynamicData_Concurrent()
 {
-	if (Simulation && SceneProxy)
+	if (Effect && SceneProxy)
 	{
 		FNiagaraSceneProxy *NiagaraProxy = static_cast<FNiagaraSceneProxy*>(SceneProxy);
-		FNiagaraDynamicDataBase* DynamicData = NiagaraProxy->GetEffectRenderer()->GenerateVertexData(Simulation->GetData());
-
-		ENQUEUE_UNIQUE_RENDER_COMMAND_TWOPARAMETER(
-			FSendNiagaraDynamicData,
-			NiagaraEffectRenderer*, EffectRenderer, (NiagaraEffectRenderer*)NiagaraProxy->GetEffectRenderer(),
-			FNiagaraDynamicDataBase*,DynamicData,DynamicData,
+		for (int32 i = 0; i < Effect->Emitters.Num(); i++)
 		{
-			EffectRenderer->SetDynamicData_RenderThread(DynamicData);
-		});
+			FNiagaraSimulation *Emitter = Effect->GetEmitter(i);
+			FNiagaraDynamicDataBase* DynamicData = Emitter->GetEffectRenderer()->GenerateVertexData(Effect->Emitters[i]->GetData());
+
+			ENQUEUE_UNIQUE_RENDER_COMMAND_TWOPARAMETER(
+				FSendNiagaraDynamicData,
+				NiagaraEffectRenderer*, EffectRenderer, Emitter->GetEffectRenderer(),
+				FNiagaraDynamicDataBase*, DynamicData, DynamicData,
+				{
+				EffectRenderer->SetDynamicData_RenderThread(DynamicData);
+			});
+		}
 	}
+
 }
 
 int32 UNiagaraComponent::GetNumMaterials() const
@@ -467,34 +239,20 @@ int32 UNiagaraComponent::GetNumMaterials() const
 	return 0;
 }
 
-UMaterialInterface* UNiagaraComponent::GetMaterial(int32 ElementIndex) const
-{
-	return NULL;
-
-	UMaterialInterface* RequestedMaterial = NULL;
-	if (ElementIndex == 0)
-	{
-		RequestedMaterial = Material;
-	}
-	return RequestedMaterial;
-}
-
-void UNiagaraComponent::SetMaterial(int32 ElementIndex, UMaterialInterface* InMaterial)
-{
-	if (ElementIndex == 0)
-	{
-		Material = InMaterial;
-	}
-}
 
 FBoxSphereBounds UNiagaraComponent::CalcBounds(const FTransform& LocalToWorld) const
 {
 	FBox SimBounds(ForceInit);
-	//if (Simulation)
-	//{
-		//SimBounds = Simulation->GetBounds();
-	//}
-	//if (!SimBounds.IsValid)
+
+	/*
+	if (Effect)
+	{
+		for (FNiagaraSimulation Sim : Effect->Emitters)
+		{
+			SimBounds += Sim->GetBounds();
+		}
+	}
+	*/
 	{
 		SimBounds.Min = FVector(-HALF_WORLD_MAX,-HALF_WORLD_MAX,-HALF_WORLD_MAX);
 		SimBounds.Max = FVector(+HALF_WORLD_MAX,+HALF_WORLD_MAX,+HALF_WORLD_MAX);
@@ -504,7 +262,8 @@ FBoxSphereBounds UNiagaraComponent::CalcBounds(const FTransform& LocalToWorld) c
 
 FPrimitiveSceneProxy* UNiagaraComponent::CreateSceneProxy()
 {
-	return new FNiagaraSceneProxy(this);
+	FNiagaraSceneProxy *Proxy = new FNiagaraSceneProxy(this);
+	return Proxy;
 }
 
 #if WITH_EDITOR
diff --git a/Engine/Source/Runtime/Engine/Private/NiagaraEffect.cpp b/Engine/Source/Runtime/Engine/Private/NiagaraEffect.cpp
--- a/Engine/Source/Runtime/Engine/Private/NiagaraEffect.cpp
+++ b/Engine/Source/Runtime/Engine/Private/NiagaraEffect.cpp
@@ -10,3 +10,21 @@ UNiagaraEffect::UNiagaraEffect(const FObjectInitializer& ObjectInitializer)
 {
 
 }
+
+
+
+TSharedPtr<FNiagaraSimulation> UNiagaraEffect::AddEmitter()
+{
+	TSharedPtr<FNiagaraSimulation> Sim = MakeShareable(new FNiagaraSimulation());
+
+	Sim->SetRenderModuleType(RMT_Sprites, Component->GetWorld()->FeatureLevel);
+	Emitters.Add(Sim);
+
+	FNiagaraSceneProxy *SceneProxy = static_cast<FNiagaraSceneProxy*>(Component->SceneProxy);
+	SceneProxy->UpdateEffectRenderers(this);
+
+	FNiagaraEmitterProperties Props;
+	EmitterProps.Add(Props);
+
+	return Sim;
+}
diff --git a/Engine/Source/Runtime/Engine/Private/NiagaraEffectRenderer.cpp b/Engine/Source/Runtime/Engine/Private/NiagaraEffectRenderer.cpp
new file mode 100644
--- /dev/null
+++ b/Engine/Source/Runtime/Engine/Private/NiagaraEffectRenderer.cpp
@@ -0,0 +1,717 @@
+#include "EnginePrivate.h"
+#include "Engine/NiagaraEffectRenderer.h"
+
+DECLARE_CYCLE_STAT(TEXT("PreRenderView"), STAT_NiagaraPreRenderView, STATGROUP_Niagara);
+DECLARE_CYCLE_STAT(TEXT("Render Total"), STAT_NiagaraRender, STATGROUP_Niagara);
+DECLARE_CYCLE_STAT(TEXT("Render Sprites"), STAT_NiagaraRenderSprites, STATGROUP_Niagara);
+DECLARE_CYCLE_STAT(TEXT("Render Ribbons"), STAT_NiagaraRenderRibbons, STATGROUP_Niagara);
+
+
+
+
+struct FNiagaraDynamicDataSprites : public FNiagaraDynamicDataBase
+{
+	TArray<FParticleSpriteVertex> VertexData;
+};
+
+struct FNiagaraDynamicDataRibbon : public FNiagaraDynamicDataBase
+{
+	TArray<FParticleBeamTrailVertex> VertexData;
+};
+
+
+
+/* Mesh collector classes */
+class FNiagaraMeshCollectorResourcesSprite : public FOneFrameResource
+{
+public:
+	FParticleSpriteVertexFactory VertexFactory;
+	FParticleSpriteUniformBufferRef UniformBuffer;
+
+	virtual ~FNiagaraMeshCollectorResourcesSprite()
+	{
+		VertexFactory.ReleaseResource();
+	}
+};
+
+
+class FNiagaraMeshCollectorResourcesRibbon : public FOneFrameResource
+{
+public:
+	FParticleBeamTrailVertexFactory VertexFactory;
+	FParticleBeamTrailUniformBufferRef UniformBuffer;
+
+	virtual ~FNiagaraMeshCollectorResourcesRibbon()
+	{
+		VertexFactory.ReleaseResource();
+	}
+};
+
+
+
+
+
+NiagaraEffectRendererSprites::NiagaraEffectRendererSprites(ERHIFeatureLevel::Type FeatureLevel) :
+	NiagaraEffectRenderer(),
+	DynamicDataRender(NULL)
+{
+	VertexFactory = new FParticleSpriteVertexFactory(PVFT_Sprite, FeatureLevel);
+}
+
+
+void NiagaraEffectRendererSprites::ReleaseRenderThreadResources()
+{
+	VertexFactory->ReleaseResource();
+	PerViewUniformBuffers.Empty();
+	WorldSpacePrimitiveUniformBuffer.ReleaseResource();
+}
+
+void NiagaraEffectRendererSprites::CreateRenderThreadResources() 
+{
+	VertexFactory->InitResource();
+}
+
+
+
+void NiagaraEffectRendererSprites::PreRenderView(const FSceneViewFamily* ViewFamily, const uint32 VisibilityMap, int32 FrameNumber, const FNiagaraSceneProxy *SceneProxy)
+{
+	SCOPE_CYCLE_COUNTER(STAT_NiagaraPreRenderView);
+
+	check(DynamicDataRender);
+	if (DynamicDataRender->VertexData.Num() == 0)
+		return;
+
+	int32 SizeInBytes = DynamicDataRender->VertexData.GetTypeSize() * DynamicDataRender->VertexData.Num();
+	DynamicVertexAllocation = FGlobalDynamicVertexBuffer::Get().Allocate(SizeInBytes);
+	if (DynamicVertexAllocation.IsValid())
+	{
+		// Copy the vertex data over.
+		FMemory::Memcpy(DynamicVertexAllocation.Buffer, DynamicDataRender->VertexData.GetData(), SizeInBytes);
+		VertexFactory->SetInstanceBuffer(
+			DynamicVertexAllocation.VertexBuffer,
+			DynamicVertexAllocation.VertexOffset,
+			sizeof(FParticleSpriteVertex),
+			true
+			);
+		VertexFactory->SetDynamicParameterBuffer(NULL, 0, 0, true);
+		// Compute the per-view uniform buffers.
+		const int32 NumViews = ViewFamily->Views.Num();
+		uint32 ViewBit = 1;
+		for (int32 ViewIndex = 0; ViewIndex < NumViews; ++ViewIndex, ViewBit <<= 1)
+		{
+			FParticleSpriteUniformBufferRef* SpriteViewUniformBufferPtr = new(PerViewUniformBuffers)FParticleSpriteUniformBufferRef();
+			if (VisibilityMap & ViewBit)
+			{
+				FParticleSpriteUniformParameters PerViewUniformParameters;// = UniformParameters;
+				PerViewUniformParameters.AxisLockRight = FVector4(0.0f, 0.0f, 0.0f, 0.0f);
+				PerViewUniformParameters.AxisLockUp = FVector4(0.0f, 0.0f, 0.0f, 0.0f);
+				PerViewUniformParameters.RotationScale = 1.0f;
+				PerViewUniformParameters.RotationBias = 0.0f;
+				PerViewUniformParameters.TangentSelector = FVector4(0.0f, 0.0f, 0.0f, 1.0f);
+				PerViewUniformParameters.InvDeltaSeconds = 0.0f;
+				PerViewUniformParameters.SubImageSize = FVector4(1.0f, 1.0f, 1.0f, 1.0f);
+				PerViewUniformParameters.NormalsType = 0;
+				PerViewUniformParameters.NormalsSphereCenter = FVector4(0.0f, 0.0f, 0.0f, 1.0f);
+				PerViewUniformParameters.NormalsCylinderUnitDirection = FVector4(0.0f, 0.0f, 1.0f, 0.0f);
+				PerViewUniformParameters.PivotOffset = FVector2D(0.0f, 0.0f);
+				PerViewUniformParameters.MacroUVParameters = FVector4(0.0f, 0.0f, 1.0f, 1.0f);
+				*SpriteViewUniformBufferPtr = FParticleSpriteUniformBufferRef::CreateUniformBufferImmediate(PerViewUniformParameters, UniformBuffer_SingleFrame);
+			}
+		}
+
+		// Update the primitive uniform buffer if needed.
+		if (!WorldSpacePrimitiveUniformBuffer.IsInitialized())
+		{
+			FPrimitiveUniformShaderParameters PrimitiveUniformShaderParameters = GetPrimitiveUniformShaderParameters(
+				FMatrix::Identity,
+				SceneProxy->GetActorPosition(),
+				SceneProxy->GetBounds(),
+				SceneProxy->GetLocalBounds(),
+				SceneProxy->ReceivesDecals(),
+				false,
+				false,
+				false
+				);
+			WorldSpacePrimitiveUniformBuffer.SetContents(PrimitiveUniformShaderParameters);
+			WorldSpacePrimitiveUniformBuffer.InitResource();
+		}
+	}
+}
+
+
+void NiagaraEffectRendererSprites::DrawDynamicElements(FPrimitiveDrawInterface* PDI, const FSceneView* View, const FNiagaraSceneProxy *SceneProxy)
+{
+	SCOPE_CYCLE_COUNTER(STAT_NiagaraRender);
+
+	check(DynamicDataRender);
+	if (DynamicDataRender->VertexData.Num() == 0 || Material==nullptr)
+		return;
+
+	const bool bIsWireframe = View->Family->EngineShowFlags.Wireframe;
+	FMaterialRenderProxy* MaterialRenderProxy = Material->GetRenderProxy(SceneProxy->IsSelected(), SceneProxy->IsHovered());
+
+	if (DynamicVertexAllocation.IsValid()
+		&& (bIsWireframe || !PDI->IsMaterialIgnored(MaterialRenderProxy, View->GetFeatureLevel())))
+	{
+		FMeshBatch MeshBatch;
+		MeshBatch.VertexFactory = VertexFactory;
+		MeshBatch.CastShadow = SceneProxy->CastsDynamicShadow();
+		MeshBatch.bUseAsOccluder = false;
+		MeshBatch.ReverseCulling = SceneProxy->IsLocalToWorldDeterminantNegative();
+		MeshBatch.Type = PT_TriangleList;
+		MeshBatch.DepthPriorityGroup = SceneProxy->GetDepthPriorityGroup(View);
+		MeshBatch.LCI = NULL;
+		if (bIsWireframe)
+		{
+			MeshBatch.MaterialRenderProxy = UMaterial::GetDefaultMaterial(MD_Surface)->GetRenderProxy(SceneProxy->IsSelected(), SceneProxy->IsHovered());
+		}
+		else
+		{
+			MeshBatch.MaterialRenderProxy = MaterialRenderProxy;
+		}
+
+		FMeshBatchElement& MeshElement = MeshBatch.Elements[0];
+		MeshElement.IndexBuffer = &GParticleIndexBuffer;
+		MeshElement.FirstIndex = 0;
+		MeshElement.NumPrimitives = 2;
+		MeshElement.NumInstances = DynamicDataRender->VertexData.Num();
+		MeshElement.MinVertexIndex = 0;
+		MeshElement.MaxVertexIndex = MeshElement.NumInstances * 4 - 1;
+		MeshElement.PrimitiveUniformBufferResource = &WorldSpacePrimitiveUniformBuffer;
+
+		int32 ViewIndex = View->Family->Views.Find(View);
+		check(PerViewUniformBuffers.IsValidIndex(ViewIndex) && PerViewUniformBuffers[ViewIndex]);
+		VertexFactory->SetSpriteUniformBuffer(PerViewUniformBuffers[ViewIndex]);
+
+		DrawRichMesh(
+			PDI,
+			MeshBatch,
+			FLinearColor(1.0f, 0.0f, 0.0f),	//WireframeColor,
+			FLinearColor(1.0f, 1.0f, 0.0f),	//LevelColor,
+			FLinearColor(1.0f, 1.0f, 1.0f),	//PropertyColor,		
+			SceneProxy,
+			GIsEditor && (View->Family->EngineShowFlags.Selection) ? SceneProxy->IsSelected() : false,
+			bIsWireframe
+			);
+	}
+
+	/*
+	for(int32 i=0; i<DynamicDataRender->VertexData.Num(); i++)
+	{
+	FParticleSpriteVertex& Vertex = DynamicDataRender->VertexData[i];
+	PDI->DrawPoint(Vertex.Position, Vertex.Color, 1.0f, SDPG_World);
+	}
+	*/
+}
+
+
+void NiagaraEffectRendererSprites::GetDynamicMeshElements(const TArray<const FSceneView*>& Views, const FSceneViewFamily& ViewFamily, uint32 VisibilityMap, FMeshElementCollector& Collector, const FNiagaraSceneProxy *SceneProxy) const
+{
+	SCOPE_CYCLE_COUNTER(STAT_NiagaraRender);
+	SCOPE_CYCLE_COUNTER(STAT_NiagaraRenderSprites);
+
+	check(DynamicDataRender && DynamicDataRender->VertexData.Num());
+
+	const bool bIsWireframe = ViewFamily.EngineShowFlags.Wireframe;
+	FMaterialRenderProxy* MaterialRenderProxy = Material->GetRenderProxy(SceneProxy->IsSelected(), SceneProxy->IsHovered());
+
+	int32 SizeInBytes = DynamicDataRender->VertexData.GetTypeSize() * DynamicDataRender->VertexData.Num();
+	FGlobalDynamicVertexBuffer::FAllocation LocalDynamicVertexAllocation = FGlobalDynamicVertexBuffer::Get().Allocate(SizeInBytes);
+
+	if (LocalDynamicVertexAllocation.IsValid())
+	{
+		// Update the primitive uniform buffer if needed.
+		if (!WorldSpacePrimitiveUniformBuffer.IsInitialized())
+		{
+			FPrimitiveUniformShaderParameters PrimitiveUniformShaderParameters = GetPrimitiveUniformShaderParameters(
+				FMatrix::Identity,
+				SceneProxy->GetActorPosition(),
+				SceneProxy->GetBounds(),
+				SceneProxy->GetLocalBounds(),
+				SceneProxy->ReceivesDecals(),
+				false,
+				SceneProxy->UseEditorDepthTest()
+				);
+			WorldSpacePrimitiveUniformBuffer.SetContents(PrimitiveUniformShaderParameters);
+			WorldSpacePrimitiveUniformBuffer.InitResource();
+		}
+
+		// Copy the vertex data over.
+		FMemory::Memcpy(LocalDynamicVertexAllocation.Buffer, DynamicDataRender->VertexData.GetData(), SizeInBytes);
+
+		// Compute the per-view uniform buffers.
+		for (int32 ViewIndex = 0; ViewIndex < Views.Num(); ViewIndex++)
+		{
+			if (VisibilityMap & (1 << ViewIndex))
+			{
+				const FSceneView* View = Views[ViewIndex];
+
+				FNiagaraMeshCollectorResourcesSprite& CollectorResources = Collector.AllocateOneFrameResource<FNiagaraMeshCollectorResourcesSprite>();
+				FParticleSpriteUniformParameters PerViewUniformParameters;// = UniformParameters;
+				PerViewUniformParameters.AxisLockRight = FVector4(0.0f, 0.0f, 0.0f, 0.0f);
+				PerViewUniformParameters.AxisLockUp = FVector4(0.0f, 0.0f, 0.0f, 0.0f);
+				PerViewUniformParameters.RotationScale = 1.0f;
+				PerViewUniformParameters.RotationBias = 0.0f;
+				PerViewUniformParameters.TangentSelector = FVector4(0.0f, 0.0f, 0.0f, 1.0f);
+				PerViewUniformParameters.InvDeltaSeconds = 0.0f;
+				PerViewUniformParameters.SubImageSize = FVector4(1.0f, 1.0f, 1.0f, 1.0f);
+				PerViewUniformParameters.NormalsType = 0;
+				PerViewUniformParameters.NormalsSphereCenter = FVector4(0.0f, 0.0f, 0.0f, 1.0f);
+				PerViewUniformParameters.NormalsCylinderUnitDirection = FVector4(0.0f, 0.0f, 1.0f, 0.0f);
+				PerViewUniformParameters.PivotOffset = FVector2D(0.0f, 0.0f);
+				PerViewUniformParameters.MacroUVParameters = FVector4(0.0f, 0.0f, 1.0f, 1.0f);
+
+				// Collector.AllocateOneFrameResource uses default ctor, initialize the vertex factory
+				CollectorResources.VertexFactory.SetFeatureLevel(ViewFamily.GetFeatureLevel());
+				CollectorResources.VertexFactory.SetParticleFactoryType(PVFT_Sprite);
+
+				PerViewUniformParameters.MacroUVParameters = FVector4(0.0f, 0.0f, 1.0f, 1.0f);
+				CollectorResources.UniformBuffer = FParticleSpriteUniformBufferRef::CreateUniformBufferImmediate(PerViewUniformParameters, UniformBuffer_SingleFrame);
+
+				CollectorResources.VertexFactory.InitResource();
+				CollectorResources.VertexFactory.SetSpriteUniformBuffer(CollectorResources.UniformBuffer);
+				CollectorResources.VertexFactory.SetInstanceBuffer(
+					LocalDynamicVertexAllocation.VertexBuffer,
+					LocalDynamicVertexAllocation.VertexOffset,
+					sizeof(FParticleSpriteVertex),
+					true
+					);
+				CollectorResources.VertexFactory.SetDynamicParameterBuffer(NULL, 0, 0, true);
+
+				FMeshBatch& MeshBatch = Collector.AllocateMesh();
+				MeshBatch.VertexFactory = &CollectorResources.VertexFactory;
+				MeshBatch.CastShadow = SceneProxy->CastsDynamicShadow();
+				MeshBatch.bUseAsOccluder = false;
+				MeshBatch.ReverseCulling = SceneProxy->IsLocalToWorldDeterminantNegative();
+				MeshBatch.Type = PT_TriangleList;
+				MeshBatch.DepthPriorityGroup = SceneProxy->GetDepthPriorityGroup(View);
+				MeshBatch.bCanApplyViewModeOverrides = true;
+				MeshBatch.bUseWireframeSelectionColoring = SceneProxy->IsSelected();
+
+				if (bIsWireframe)
+				{
+					MeshBatch.MaterialRenderProxy = UMaterial::GetDefaultMaterial(MD_Surface)->GetRenderProxy(SceneProxy->IsSelected(), SceneProxy->IsHovered());
+				}
+				else
+				{
+					MeshBatch.MaterialRenderProxy = MaterialRenderProxy;
+				}
+
+				FMeshBatchElement& MeshElement = MeshBatch.Elements[0];
+				MeshElement.IndexBuffer = &GParticleIndexBuffer;
+				MeshElement.FirstIndex = 0;
+				MeshElement.NumPrimitives = 2;
+				MeshElement.NumInstances = DynamicDataRender->VertexData.Num();
+				MeshElement.MinVertexIndex = 0;
+				MeshElement.MaxVertexIndex = MeshElement.NumInstances * 4 - 1;
+				MeshElement.PrimitiveUniformBufferResource = &WorldSpacePrimitiveUniformBuffer;
+
+				Collector.AddMesh(ViewIndex, MeshBatch);
+			}
+		}
+	}
+}
+
+
+/** Update render data buffer from attributes */
+FNiagaraDynamicDataBase *NiagaraEffectRendererSprites::GenerateVertexData(const FNiagaraEmitterParticleData &Data)
+{
+	SCOPE_CYCLE_COUNTER(STAT_NiagaraGenSpriteVertexData);
+	FNiagaraDynamicDataSprites *DynamicData = new FNiagaraDynamicDataSprites;
+	TArray<FParticleSpriteVertex>& RenderData = DynamicData->VertexData;
+
+	RenderData.Reset(Data.GetNumParticles());
+	//CachedBounds.Init();
+
+	const FVector4 *PosPtr = Data.GetAttributeData("Position");
+	const FVector4 *ColPtr = Data.GetAttributeData("Color");
+	const FVector4 *AgePtr = Data.GetAttributeData("Age");
+	const FVector4 *RotPtr = Data.GetAttributeData("Rotation");
+
+	float ParticleId = 0.0f, IdInc = 1.0f / Data.GetNumParticles();
+	for (uint32 ParticleIndex = 0; ParticleIndex < Data.GetNumParticles(); ParticleIndex++)
+	{
+		FParticleSpriteVertex& NewVertex = *new(RenderData)FParticleSpriteVertex;
+		NewVertex.Position = PosPtr[ParticleIndex];
+		NewVertex.OldPosition = NewVertex.Position;
+		NewVertex.Color = FLinearColor(ColPtr[ParticleIndex]);
+		NewVertex.ParticleId = ParticleId;
+		ParticleId += IdInc;
+		NewVertex.RelativeTime = AgePtr[ParticleIndex].X;
+		NewVertex.Size = FVector2D(1.0f, 1.0f);
+		NewVertex.Rotation = RotPtr[ParticleIndex].X;
+		NewVertex.SubImageIndex = 0.f;
+
+		//CachedBounds += NewVertex.Position;
+	}
+	//CachedBounds.ExpandBy(MaxSize);
+
+	return DynamicData;
+}
+
+
+
+void NiagaraEffectRendererSprites::SetDynamicData_RenderThread(FNiagaraDynamicDataBase* NewDynamicData)
+{
+	check(IsInRenderingThread());
+
+	if (DynamicDataRender)
+	{
+		delete DynamicDataRender;
+		DynamicDataRender = NULL;
+	}
+	DynamicDataRender = static_cast<FNiagaraDynamicDataSprites*>(NewDynamicData);
+}
+
+int NiagaraEffectRendererSprites::GetDynamicDataSize()
+{
+	uint32 Size = sizeof(FNiagaraDynamicDataSprites);
+	if (DynamicDataRender)
+	{
+		Size += DynamicDataRender->VertexData.GetAllocatedSize();
+	}
+
+	return Size;
+}
+
+bool NiagaraEffectRendererSprites::HasDynamicData()
+{
+	return DynamicDataRender && DynamicDataRender->VertexData.Num() > 0;
+}
+
+
+
+
+
+NiagaraEffectRendererRibbon::NiagaraEffectRendererRibbon(ERHIFeatureLevel::Type FeatureLevel) :
+NiagaraEffectRenderer(),
+DynamicDataRender(NULL)
+{
+	VertexFactory = new FParticleBeamTrailVertexFactory(PVFT_BeamTrail, FeatureLevel);
+}
+
+
+void NiagaraEffectRendererRibbon::ReleaseRenderThreadResources()
+{
+	VertexFactory->ReleaseResource();
+	PerViewUniformBuffers.Empty();
+	WorldSpacePrimitiveUniformBuffer.ReleaseResource();
+}
+
+// FPrimitiveSceneProxy interface.
+void NiagaraEffectRendererRibbon::CreateRenderThreadResources()
+{
+	VertexFactory->InitResource();
+}
+
+
+
+
+void NiagaraEffectRendererRibbon::GetDynamicMeshElements(const TArray<const FSceneView*>& Views, const FSceneViewFamily& ViewFamily, uint32 VisibilityMap, FMeshElementCollector& Collector, const FNiagaraSceneProxy *SceneProxy) const
+{
+	SCOPE_CYCLE_COUNTER(STAT_NiagaraRender);
+	SCOPE_CYCLE_COUNTER(STAT_NiagaraRenderRibbons);
+
+	check(DynamicDataRender && DynamicDataRender->VertexData.Num());
+
+	const bool bIsWireframe = ViewFamily.EngineShowFlags.Wireframe;
+	FMaterialRenderProxy* MaterialRenderProxy = Material->GetRenderProxy(SceneProxy->IsSelected(), SceneProxy->IsHovered());
+
+	int32 SizeInBytes = DynamicDataRender->VertexData.GetTypeSize() * DynamicDataRender->VertexData.Num();
+	FGlobalDynamicVertexBuffer::FAllocation LocalDynamicVertexAllocation = FGlobalDynamicVertexBuffer::Get().Allocate(SizeInBytes);
+
+	if (LocalDynamicVertexAllocation.IsValid())
+	{
+		// Update the primitive uniform buffer if needed.
+		if (!WorldSpacePrimitiveUniformBuffer.IsInitialized())
+		{
+			FPrimitiveUniformShaderParameters PrimitiveUniformShaderParameters = GetPrimitiveUniformShaderParameters(
+				FMatrix::Identity,
+				SceneProxy->GetActorPosition(),
+				SceneProxy->GetBounds(),
+				SceneProxy->GetLocalBounds(),
+				SceneProxy->ReceivesDecals(),
+				false,
+				SceneProxy->UseEditorDepthTest()
+				);
+			WorldSpacePrimitiveUniformBuffer.SetContents(PrimitiveUniformShaderParameters);
+			WorldSpacePrimitiveUniformBuffer.InitResource();
+		}
+
+		// Copy the vertex data over.
+		FMemory::Memcpy(LocalDynamicVertexAllocation.Buffer, DynamicDataRender->VertexData.GetData(), SizeInBytes);
+
+		// Compute the per-view uniform buffers.
+		for (int32 ViewIndex = 0; ViewIndex < Views.Num(); ViewIndex++)
+		{
+			if (VisibilityMap & (1 << ViewIndex))
+			{
+				const FSceneView* View = Views[ViewIndex];
+
+				FNiagaraMeshCollectorResourcesRibbon& CollectorResources = Collector.AllocateOneFrameResource<FNiagaraMeshCollectorResourcesRibbon>();
+				FParticleBeamTrailUniformParameters PerViewUniformParameters;// = UniformParameters;
+				PerViewUniformParameters.CameraUp = FVector4(0.0f, 0.0f, 1.0f, 0.0f);
+				PerViewUniformParameters.CameraRight = FVector4(1.0f, 0.0f, 0.0f, 0.0f);
+				PerViewUniformParameters.ScreenAlignment = FVector4(0.0f, 0.0f, 0.0f, 0.0f);
+
+				// Collector.AllocateOneFrameResource uses default ctor, initialize the vertex factory
+				CollectorResources.VertexFactory.SetFeatureLevel(ViewFamily.GetFeatureLevel());
+				CollectorResources.VertexFactory.SetParticleFactoryType(PVFT_BeamTrail);
+
+				CollectorResources.UniformBuffer = FParticleBeamTrailUniformBufferRef::CreateUniformBufferImmediate(PerViewUniformParameters, UniformBuffer_SingleFrame);
+
+				CollectorResources.VertexFactory.InitResource();
+				CollectorResources.VertexFactory.SetBeamTrailUniformBuffer(PerViewUniformBuffers[ViewIndex]);
+				CollectorResources.VertexFactory.SetVertexBuffer(DynamicVertexAllocation.VertexBuffer, DynamicVertexAllocation.VertexOffset, sizeof(FParticleBeamTrailVertex));
+				CollectorResources.VertexFactory.SetDynamicParameterBuffer(NULL, 0, 0);
+
+				FMeshBatch& MeshBatch = Collector.AllocateMesh();
+				MeshBatch.VertexFactory = &CollectorResources.VertexFactory;
+				MeshBatch.CastShadow = SceneProxy->CastsDynamicShadow();
+				MeshBatch.bUseAsOccluder = false;
+				MeshBatch.ReverseCulling = SceneProxy->IsLocalToWorldDeterminantNegative();
+				MeshBatch.bDisableBackfaceCulling = true;
+				MeshBatch.Type = PT_TriangleStrip;
+				MeshBatch.DepthPriorityGroup = SceneProxy->GetDepthPriorityGroup(View);
+				MeshBatch.bCanApplyViewModeOverrides = true;
+				MeshBatch.bUseWireframeSelectionColoring = SceneProxy->IsSelected();
+
+				if (bIsWireframe)
+				{
+					MeshBatch.MaterialRenderProxy = UMaterial::GetDefaultMaterial(MD_Surface)->GetRenderProxy(SceneProxy->IsSelected(), SceneProxy->IsHovered());
+				}
+				else
+				{
+					MeshBatch.MaterialRenderProxy = MaterialRenderProxy;
+				}
+
+				FMeshBatchElement& MeshElement = MeshBatch.Elements[0];
+				MeshElement.IndexBuffer = &GParticleIndexBuffer;
+				MeshElement.FirstIndex = 0;
+				MeshElement.NumPrimitives = (DynamicDataRender->VertexData.Num() - 1) / 2;
+				MeshElement.NumInstances = 1;
+				MeshElement.MinVertexIndex = 0;
+				MeshElement.MaxVertexIndex = DynamicDataRender->VertexData.Num() - 1;
+				MeshElement.PrimitiveUniformBufferResource = &WorldSpacePrimitiveUniformBuffer;
+
+				Collector.AddMesh(ViewIndex, MeshBatch);
+			}
+		}
+	}
+}
+
+void NiagaraEffectRendererRibbon::PreRenderView(const FSceneViewFamily* ViewFamily, const uint32 VisibilityMap, int32 FrameNumber, const FNiagaraSceneProxy *SceneProxy)
+{
+	SCOPE_CYCLE_COUNTER(STAT_NiagaraPreRenderView);
+
+	check(DynamicDataRender && DynamicDataRender->VertexData.Num());
+
+	int32 SizeInBytes = DynamicDataRender->VertexData.GetTypeSize() * DynamicDataRender->VertexData.Num();
+	DynamicVertexAllocation = FGlobalDynamicVertexBuffer::Get().Allocate(SizeInBytes);
+	if (DynamicVertexAllocation.IsValid())
+	{
+		// Copy the vertex data over.
+		FMemory::Memcpy(DynamicVertexAllocation.Buffer, DynamicDataRender->VertexData.GetData(), SizeInBytes);
+		VertexFactory->SetVertexBuffer(DynamicVertexAllocation.VertexBuffer, DynamicVertexAllocation.VertexOffset, sizeof(FParticleBeamTrailVertex));
+		VertexFactory->SetDynamicParameterBuffer(NULL, 0, 0);
+		// Compute the per-view uniform buffers.
+		const int32 NumViews = ViewFamily->Views.Num();
+		uint32 ViewBit = 1;
+		for (int32 ViewIndex = 0; ViewIndex < NumViews; ++ViewIndex, ViewBit <<= 1)
+		{
+			FParticleBeamTrailUniformBufferRef* SpriteViewUniformBufferPtr = new(PerViewUniformBuffers)FParticleBeamTrailUniformBufferRef();
+			if (VisibilityMap & ViewBit)
+			{
+				FParticleBeamTrailUniformParameters PerViewUniformParameters;// = UniformParameters;
+				PerViewUniformParameters.CameraUp = FVector4(0.0f, 0.0f, 1.0f, 0.0f);
+				PerViewUniformParameters.CameraRight = FVector4(1.0f, 0.0f, 0.0f, 0.0f);
+				PerViewUniformParameters.ScreenAlignment = FVector4(0.0f, 0.0f, 0.0f, 0.0f);
+				*SpriteViewUniformBufferPtr = FParticleBeamTrailUniformBufferRef::CreateUniformBufferImmediate(PerViewUniformParameters, UniformBuffer_SingleFrame);
+			}
+		}
+
+		// Update the primitive uniform buffer if needed.
+		if (!WorldSpacePrimitiveUniformBuffer.IsInitialized())
+		{
+			FPrimitiveUniformShaderParameters PrimitiveUniformShaderParameters = GetPrimitiveUniformShaderParameters(
+				FMatrix::Identity,
+				SceneProxy->GetActorPosition(),
+				SceneProxy->GetBounds(),
+				SceneProxy->GetLocalBounds(),
+				SceneProxy->ReceivesDecals(),
+				false,
+				false,
+				false
+				);
+			WorldSpacePrimitiveUniformBuffer.SetContents(PrimitiveUniformShaderParameters);
+			WorldSpacePrimitiveUniformBuffer.InitResource();
+		}
+	}
+}
+
+
+void NiagaraEffectRendererRibbon::DrawDynamicElements(FPrimitiveDrawInterface* PDI, const FSceneView* View, const FNiagaraSceneProxy *SceneProxy)
+{
+	SCOPE_CYCLE_COUNTER(STAT_NiagaraRender);
+
+	check(DynamicDataRender && DynamicDataRender->VertexData.Num());
+
+	const bool bIsWireframe = View->Family->EngineShowFlags.Wireframe;
+	FMaterialRenderProxy* MaterialRenderProxy = Material->GetRenderProxy(SceneProxy->IsSelected(), SceneProxy->IsHovered());
+
+	if (DynamicVertexAllocation.IsValid()
+		&& (bIsWireframe || !PDI->IsMaterialIgnored(MaterialRenderProxy, View->GetFeatureLevel())))
+	{
+		FMeshBatch MeshBatch;
+		MeshBatch.VertexFactory = VertexFactory;
+		MeshBatch.CastShadow = SceneProxy->CastsDynamicShadow();
+		MeshBatch.bUseAsOccluder = false;
+		MeshBatch.bDisableBackfaceCulling = true;
+		MeshBatch.ReverseCulling = SceneProxy->IsLocalToWorldDeterminantNegative();
+		MeshBatch.Type = PT_TriangleStrip;
+		MeshBatch.DepthPriorityGroup = SceneProxy->GetDepthPriorityGroup(View);
+		MeshBatch.LCI = NULL;
+		if (bIsWireframe)
+		{
+			MeshBatch.MaterialRenderProxy = UMaterial::GetDefaultMaterial(MD_Surface)->GetRenderProxy(SceneProxy->IsSelected(), SceneProxy->IsHovered());
+		}
+		else
+		{
+			MeshBatch.MaterialRenderProxy = MaterialRenderProxy;
+		}
+
+		FMeshBatchElement& MeshElement = MeshBatch.Elements[0];
+		MeshElement.IndexBuffer = &GParticleIndexBuffer;
+		MeshElement.FirstIndex = 0;
+		MeshElement.NumPrimitives = DynamicDataRender->VertexData.Num() / 2;
+		MeshElement.NumInstances = 1;
+		MeshElement.MinVertexIndex = 0;
+		MeshElement.MaxVertexIndex = DynamicDataRender->VertexData.Num() - 1;
+		MeshElement.PrimitiveUniformBufferResource = &WorldSpacePrimitiveUniformBuffer;
+
+		int32 ViewIndex = View->Family->Views.Find(View);
+		check(PerViewUniformBuffers.IsValidIndex(ViewIndex) && PerViewUniformBuffers[ViewIndex]);
+		VertexFactory->SetBeamTrailUniformBuffer(PerViewUniformBuffers[ViewIndex]);
+
+		DrawRichMesh(
+			PDI,
+			MeshBatch,
+			FLinearColor(1.0f, 0.0f, 0.0f),	//WireframeColor,
+			FLinearColor(1.0f, 1.0f, 0.0f),	//LevelColor,
+			FLinearColor(1.0f, 1.0f, 1.0f),	//PropertyColor,		
+			SceneProxy,
+			GIsEditor && (View->Family->EngineShowFlags.Selection) ? SceneProxy->IsSelected() : false,
+			bIsWireframe
+			);
+	}
+
+	/*
+	for(int32 i=0; i<DynamicDataRender->VertexData.Num(); i++)
+	{
+	FParticleSpriteVertex& Vertex = DynamicDataRender->VertexData[i];
+	PDI->DrawPoint(Vertex.Position, Vertex.Color, 1.0f, SDPG_World);
+	}
+	*/
+}
+
+
+
+
+void NiagaraEffectRendererRibbon::SetDynamicData_RenderThread(FNiagaraDynamicDataBase* NewDynamicData)
+{
+	check(IsInRenderingThread());
+
+	if (DynamicDataRender)
+	{
+		delete DynamicDataRender;
+		DynamicDataRender = NULL;
+	}
+	DynamicDataRender = static_cast<FNiagaraDynamicDataRibbon*>(NewDynamicData);
+}
+
+int NiagaraEffectRendererRibbon::GetDynamicDataSize()
+{
+	uint32 Size = sizeof(FNiagaraDynamicDataRibbon);
+	if (DynamicDataRender)
+	{
+		Size += DynamicDataRender->VertexData.GetAllocatedSize();
+	}
+
+	return Size;
+}
+
+bool NiagaraEffectRendererRibbon::HasDynamicData()
+{
+	return DynamicDataRender && DynamicDataRender->VertexData.Num() > 0;
+}
+
+
+
+FNiagaraDynamicDataBase *NiagaraEffectRendererRibbon::GenerateVertexData(const FNiagaraEmitterParticleData &Data)
+{
+	SCOPE_CYCLE_COUNTER(STAT_NiagaraGenRibbonVertexData);
+	FNiagaraDynamicDataRibbon *DynamicData = new FNiagaraDynamicDataRibbon;
+	TArray<FParticleBeamTrailVertex>& RenderData = DynamicData->VertexData;
+
+	RenderData.Reset(Data.GetNumParticles() * 2);
+	//CachedBounds.Init();
+
+	// build a sorted list by age, so we always get particles in order 
+	// regardless of them being moved around due to dieing and spawning
+	TArray<int32> SortedIndices;
+	for (uint32 Idx = 0; Idx < Data.GetNumParticles(); Idx++)
+	{
+		SortedIndices.Add(Idx);
+	}
+
+	const FVector4 *AgeData = Data.GetAttributeData("Age");
+	SortedIndices.Sort(
+		[&AgeData](const int32& A, const int32& B) {
+		return AgeData[A].X < AgeData[B].X;
+	}
+	);
+
+
+	FVector2D UVs[4] = { FVector2D(0.0f, 0.0f), FVector2D(1.0f, 0.0f), FVector2D(1.0f, 1.0f), FVector2D(0.0f, 1.0f) };
+
+	const FVector4 *PosPtr = Data.GetAttributeData("Position");
+	const FVector4 *ColorPtr = Data.GetAttributeData("Color");
+	const FVector4 *AgePtr = Data.GetAttributeData("Age");
+	const FVector4 *RotPtr = Data.GetAttributeData("Rotation");
+
+	FVector PrevPos, PrevPos2, PrevDir(0.0f, 0.0f, 0.1f);
+	for (int32 i = 0; i < SortedIndices.Num() - 1; i++)
+	{
+		uint32 Index1 = SortedIndices[i];
+		uint32 Index2 = SortedIndices[i + 1];
+
+		const FVector ParticlePos = PosPtr[Index1];
+		FVector ParticleDir = PosPtr[Index2] - ParticlePos;
+		if (ParticleDir.Size() <= SMALL_NUMBER)
+		{
+			ParticleDir = PrevDir*0.1f;
+		}
+
+		FVector ParticleRight = FVector::CrossProduct(ParticleDir, FVector(0.0f, 0.0f, 1.0f));
+		ParticleRight.Normalize();
+
+		if (i == 0)
+		{
+			AddRibbonVert(RenderData, ParticlePos + ParticleRight, Data, UVs[0], ColorPtr[i], AgePtr[i], RotPtr[i]);
+			AddRibbonVert(RenderData, ParticlePos - ParticleRight, Data, UVs[1], ColorPtr[i], AgePtr[i], RotPtr[i]);
+		}
+		else
+		{
+			AddRibbonVert(RenderData, PrevPos2, Data, UVs[0], ColorPtr[i], AgePtr[i], RotPtr[i]);
+			AddRibbonVert(RenderData, PrevPos, Data, UVs[1], ColorPtr[i], AgePtr[i], RotPtr[i]);
+		}
+
+		AddRibbonVert(RenderData, ParticlePos - ParticleRight + ParticleDir, Data, UVs[2], ColorPtr[i], AgePtr[i], RotPtr[i]);
+		AddRibbonVert(RenderData, ParticlePos + ParticleRight + ParticleDir, Data, UVs[3], ColorPtr[i], AgePtr[i], RotPtr[i]);
+		PrevPos = ParticlePos - ParticleRight + ParticleDir;
+		PrevPos2 = ParticlePos + ParticleRight + ParticleDir;
+		PrevDir = ParticleDir;
+	}
+
+	return DynamicData;
+}
\ No newline at end of file
diff --git a/Engine/Source/Runtime/VectorVM/Private/VectorVM.cpp b/Engine/Source/Runtime/VectorVM/Private/VectorVM.cpp
--- a/Engine/Source/Runtime/VectorVM/Private/VectorVM.cpp
+++ b/Engine/Source/Runtime/VectorVM/Private/VectorVM.cpp
@@ -370,6 +370,29 @@ struct FVectorKernelSin4 : public TUnaryVectorKernel<FVectorKernelSin4>
 	}
 };
 
+struct FVectorKernelLessThan : public TBinaryVectorKernel<FVectorKernelLessThan>
+{
+	static void FORCEINLINE DoKernel(VectorRegister* RESTRICT Dst, VectorRegister Src0, VectorRegister Src1)
+	{
+		const VectorRegister large = MakeVectorRegister(BIG_NUMBER, BIG_NUMBER, BIG_NUMBER, BIG_NUMBER);
+		const VectorRegister one = MakeVectorRegister(1.0f, 1.0f, 1.0f, 1.0f);
+		const VectorRegister zero = MakeVectorRegister(0.0f, 0.0f, 0.0f, 0.0f);
+
+		float const* FloatSrc0 = reinterpret_cast<float const*>(&Src0);
+		float const* FloatSrc1 = reinterpret_cast<float const*>(&Src1);
+		VectorRegister tmp = VectorSubtract(Src1, Src0);
+		tmp = VectorMultiply(tmp, large);
+		tmp = VectorMin(tmp, one);
+		*Dst = VectorMax(tmp, zero);
+		/*
+		float f1 = FloatSrc0[0] < FloatSrc1[0] ? 1.0f : 0.0f;
+		float f2 = FloatSrc0[1] < FloatSrc1[1] ? 1.0f : 0.0f;
+		float f3 = FloatSrc0[2] < FloatSrc1[2] ? 1.0f : 0.0f;
+		float f4 = FloatSrc0[3] < FloatSrc1[3] ? 1.0f : 0.0f;
+		*Dst = MakeVectorRegister(f1, f2, f3, f4);
+		*/
+	}
+};
 
 struct FVectorKernelDot : public TBinaryVectorKernel<FVectorKernelDot>
 {
@@ -467,7 +490,7 @@ struct FVectorKernelNoise : public TUnaryVectorKernel<FVectorKernelNoise>
 			float FX = FloatSrc[0] / 5 / (1<<i);
 			float FY = FloatSrc[1] / 5 / (1<<i);
 			float FZ = FloatSrc[2] / 5 / (1<<i);
-			uint32 X = FMath::Abs((int32)(FX) % 8);
+			uint32 X = FMath::Abs( (int32)(FX) % 8 );
 			uint32 Y = FMath::Abs( (int32)(FY) % 8 );
 			uint32 Z = FMath::Abs( (int32)(FZ) % 8 );
 
@@ -482,7 +505,7 @@ struct FVectorKernelNoise : public TUnaryVectorKernel<FVectorKernelNoise>
 				}
 			}
 
-			const uint32 IntCoords[3] = { static_cast<uint32>(FX), static_cast<uint32>(FY), static_cast<uint32>(FZ) };
+			const int32 IntCoords[3] = { static_cast<int32>(FX), static_cast<int32>(FY), static_cast<int32>(FZ) };
 			const float Fractionals[3] = { FX - IntCoords[0], FY - IntCoords[1], FZ - IntCoords[2] };
 			VectorRegister Alpha = MakeVectorRegister(Fractionals[0], Fractionals[0], Fractionals[0], Fractionals[0]);
 
