Adding missing file.

diff --git a/Engine/Source/Runtime/Engine/Private/NiagaraSimulation.cpp b/Engine/Source/Runtime/Engine/Private/NiagaraSimulation.cpp
new file mode 100644
--- /dev/null
+++ b/Engine/Source/Runtime/Engine/Private/NiagaraSimulation.cpp
@@ -0,0 +1,228 @@
+
+#include "EnginePrivate.h"
+#include "Engine/NiagaraConstants.h"
+#include "Engine/NiagaraSimulation.h"
+#include "VectorVM.h"
+
+
+FNiagaraSimulation::FNiagaraSimulation()
+: SpawnRate(0)
+, Age(0.0f)
+, UpdateScript(nullptr)
+, SpawnScript(nullptr)
+, SpawnRemainder(0.0f)
+, CachedBounds(ForceInit)
+, RenderModuleType(RMT_None)
+, EffectRenderer(nullptr)
+{
+}
+
+
+void FNiagaraSimulation::Tick(float DeltaSeconds)
+{
+	SCOPE_CYCLE_COUNTER(STAT_NiagaraTick);
+
+
+	// Cache the ComponentToWorld transform.
+//	CachedComponentToWorld = Component.GetComponentToWorld();
+
+	Data.SwapBuffers();
+	// Remember the stride of the original data.
+	int32 PrevNumVectorsPerAttribute = Data.GetParticleAllocation();
+	int32 MaxNewParticles = Data.GetNumParticles();
+	int32 NumToSpawn = 0;
+
+	// Figure out how many we will spawn.
+	if (SpawnScript)
+	{
+		NumToSpawn = CalcNumToSpawn(DeltaSeconds);
+		MaxNewParticles = Data.GetNumParticles() + NumToSpawn;
+		Data.Allocate(MaxNewParticles);
+	}
+
+	if (UpdateScript)
+	{
+		// Simulate particles forward by DeltaSeconds.
+		UpdateParticles(
+			DeltaSeconds,
+			Data.GetPreviousBuffer(),
+			PrevNumVectorsPerAttribute,
+			Data.GetCurrentBuffer(),
+			MaxNewParticles,
+			Data.GetNumParticles()
+			);
+	}
+
+	if (SpawnScript)
+	{
+		SpawnAndKillParticles(NumToSpawn);
+	}
+
+	DECLARE_DWORD_COUNTER_STAT(TEXT("NumParticles"), STAT_NiagaraNumParticles, STATGROUP_Niagara);
+	INC_DWORD_STAT_BY(STAT_NiagaraNumParticles, Data.GetNumParticles());
+
+}
+
+
+void FNiagaraSimulation::UpdateParticles(
+	float DeltaSeconds,
+	FVector4* PrevParticles,
+	int32 PrevNumVectorsPerAttribute,
+	FVector4* Particles,
+	int32 NumVectorsPerAttribute,
+	int32 NumParticles
+	)
+{
+	SCOPE_CYCLE_COUNTER(STAT_NiagaraSimulate);
+	Age += DeltaSeconds;
+	Constants.SetOrAdd(TEXT("Emitter Age"), FVector4(Age, Age, Age, Age));
+
+	VectorRegister* InputRegisters[VectorVM::MaxInputRegisters] = { 0 };
+	VectorRegister* OutputRegisters[VectorVM::MaxOutputRegisters] = { 0 };
+	const int32 NumAttr = UpdateScript->Attributes.Num();
+
+	check(NumAttr < VectorVM::MaxInputRegisters);
+	check(NumAttr < VectorVM::MaxOutputRegisters);
+
+	// Setup input and output registers.
+	for (int32 AttrIndex = 0; AttrIndex < NumAttr; ++AttrIndex)
+	{
+		InputRegisters[AttrIndex] = (VectorRegister*)(PrevParticles + AttrIndex * PrevNumVectorsPerAttribute);
+		OutputRegisters[AttrIndex] = (VectorRegister*)(Particles + AttrIndex * NumVectorsPerAttribute);
+	}
+
+	//Fill constant table with required emitter constants and internal script constants.
+	TArray<FVector4> ConstantTable;
+	UpdateScript->ConstantData.FillConstantTable(Constants, ConstantTable);
+
+	VectorVM::Exec(
+		UpdateScript->ByteCode.GetData(),
+		InputRegisters,
+		NumAttr,
+		OutputRegisters,
+		NumAttr,
+		ConstantTable.GetData(),
+		NumParticles
+		);
+}
+
+
+
+
+
+int32 FNiagaraSimulation::SpawnAndKillParticles(int32 NumToSpawn)
+{
+	SCOPE_CYCLE_COUNTER(STAT_NiagaraSpawnAndKill);
+	int32 OrigNumParticles = Data.GetNumParticles();
+	int32 CurNumParticles = OrigNumParticles;
+	CurNumParticles = SpawnParticles(NumToSpawn);
+
+	// run the spawn graph over all new particles
+	if (SpawnScript && SpawnScript->ByteCode.Num())
+	{
+		VectorRegister* InputRegisters[VectorVM::MaxInputRegisters] = { 0 };
+		VectorRegister* OutputRegisters[VectorVM::MaxOutputRegisters] = { 0 };
+		const int32 NumAttr = SpawnScript->Attributes.Num();
+		const int32 NumVectors = NumToSpawn;
+
+		check(NumAttr < VectorVM::MaxInputRegisters);
+		check(NumAttr < VectorVM::MaxOutputRegisters);
+
+		FVector4 *NewParticlesStart = Data.GetCurrentBuffer() + OrigNumParticles;
+
+		// Setup input and output registers.
+		for (int32 AttrIndex = 0; AttrIndex < NumAttr; ++AttrIndex)
+		{
+			InputRegisters[AttrIndex] = (VectorRegister*)(NewParticlesStart + AttrIndex * Data.GetParticleAllocation());
+			OutputRegisters[AttrIndex] = (VectorRegister*)(NewParticlesStart + AttrIndex * Data.GetParticleAllocation());
+		}
+
+		//Fill constant table with required emitter constants and internal script constants.
+		TArray<FVector4> ConstantTable;
+		SpawnScript->ConstantData.FillConstantTable(Constants, ConstantTable);
+
+		VectorVM::Exec(
+			SpawnScript->ByteCode.GetData(),
+			InputRegisters,
+			NumAttr,
+			OutputRegisters,
+			NumAttr,
+			ConstantTable.GetData(),
+			NumVectors
+			);
+	}
+
+	// Iterate over looking for dead particles and move from the end of the list to the dead location, compacting in the process
+	int32 ParticleIndex = 0;
+	const FVector4* ParticleRelativeTimes = Data.GetAttributeData("Age");
+	while (ParticleIndex < OrigNumParticles)
+	{
+		if (ParticleRelativeTimes[ParticleIndex].X > 1.0f)
+		{
+			// Particle is dead, move one from the end here.
+			MoveParticleToIndex(--CurNumParticles, ParticleIndex);
+		}
+		ParticleIndex++;
+	}
+
+	Data.SetNumParticles(CurNumParticles);
+	return CurNumParticles;
+}
+
+
+/** Replace the current effect renderer with a new one of Type, transferring the existing material over
+ *	to the new renderer. Don't forget to call RenderModuleUpdate on the SceneProxy after calling this! 
+ */
+void FNiagaraSimulation::SetRenderModuleType(EEmitterRenderModuleType Type, ERHIFeatureLevel::Type FeatureLevel)
+{
+	if (Type != RenderModuleType)
+	{
+		UMaterial *Material = UMaterial::GetDefaultMaterial(MD_Surface);
+
+		if (EffectRenderer)
+		{
+			Material = EffectRenderer->GetMaterial();
+			delete EffectRenderer;
+		}
+
+		RenderModuleType = Type;
+		switch (Type)
+		{
+		case RMT_Sprites: EffectRenderer = new NiagaraEffectRendererSprites(FeatureLevel);
+			break;
+		case RMT_Ribbon: EffectRenderer = new NiagaraEffectRendererRibbon(FeatureLevel);
+			break;
+		}
+
+		EffectRenderer->SetMaterial(Material, FeatureLevel);
+	}
+}
+
+
+
+int32 FNiagaraSimulation::SpawnParticles(int32 NumToSpawn)
+{
+	FVector4 *PosPtr = Data.GetAttributeDataWrite("Position");
+	FVector4 *VelPtr = Data.GetAttributeDataWrite("Velocity");
+	FVector4 *ColPtr = Data.GetAttributeDataWrite("Color");
+	FVector4 *RotPtr = Data.GetAttributeDataWrite("Rotation");
+	FVector4 *AgePtr = Data.GetAttributeDataWrite("Age");
+
+	// Spawn new Particles at the end of the buffer
+	int32 ParticleIndex = Data.GetNumParticles();
+	FVector SpawnLocation = *Constants.FindVector(TEXT("Emitter Position"));
+	for (int32 i = 0; i < NumToSpawn; i++)
+	{
+		SpawnLocation.X += FMath::FRandRange(-20.f, 20.f);
+		SpawnLocation.Y += FMath::FRandRange(-20.f, 20.f);
+		SpawnLocation.Z += FMath::FRandRange(-20.f, 20.f);
+
+		PosPtr[ParticleIndex] = SpawnLocation;
+		ColPtr[ParticleIndex] = FVector4(0.0f, 0.0f, 0.0f, 0.0f);
+		VelPtr[ParticleIndex] = FVector4(0.0f, 0.0f, 2.0f, 0.0f);
+		RotPtr[ParticleIndex] = FVector4(0.0f, 0.0f, 0.0f, 0.0f);
+		AgePtr[ParticleIndex] = FVector4(0.0f, 0.0f, 0.0f, 0.0f);
+		ParticleIndex++;
+	}
+	return ParticleIndex;
+}
\ No newline at end of file
