#ue4 - MovementComponent::TwoWallAdjust() uses ComputeSlideVector() rather than an explicit vector projection.

- Extra call to HandleSlopeBoosting() removed in CharacterMovementComponent::PhysFalling() since it's handled by TwoWallAdjust()'s call to ComputeSlideVector() now.

diff --git a/Engine/Source/Runtime/Engine/Private/CharacterMovementComponent.cpp b/Engine/Source/Runtime/Engine/Private/CharacterMovementComponent.cpp
--- a/Engine/Source/Runtime/Engine/Private/CharacterMovementComponent.cpp
+++ b/Engine/Source/Runtime/Engine/Private/CharacterMovementComponent.cpp
@@ -3010,11 +3010,6 @@ void UCharacterMovementComponent::PhysFalling(float deltaTime, int32 Iterations)
 							}
 						}
 
-						if (Delta.Z > 0.f)
-						{
-							Delta = HandleSlopeBoosting(Delta, PreTwoWallDelta, 1.f - Hit.Time, Hit.Normal, Hit);
-						}
-
 						// Compute velocity after deflection (only gravity component for RootMotion)
 						if (subTimeTickRemaining > KINDA_SMALL_NUMBER && !bJustTeleported)
 						{
diff --git a/Engine/Source/Runtime/Engine/Private/MovementComponent.cpp b/Engine/Source/Runtime/Engine/Private/MovementComponent.cpp
--- a/Engine/Source/Runtime/Engine/Private/MovementComponent.cpp
+++ b/Engine/Source/Runtime/Engine/Private/MovementComponent.cpp
@@ -510,7 +510,7 @@ void UMovementComponent::TwoWallAdjust(FVector& OutDelta, const FHitResult& Hit,
 	else //adjust to new wall
 	{
 		const FVector DesiredDir = Delta;
-		Delta = (Delta - HitNormal * (Delta | HitNormal)) * (1.f - Hit.Time);
+		Delta = ComputeSlideVector(Delta, 1.f - Hit.Time, HitNormal, Hit);
 		if ((Delta | DesiredDir) <= 0.f)
 		{
 			Delta = FVector::ZeroVector;
