#docs - Fixing missing parens in doc link.

diff --git a/Engine/Source/Runtime/Engine/Classes/Components/PrimitiveComponent.h b/Engine/Source/Runtime/Engine/Classes/Components/PrimitiveComponent.h
--- a/Engine/Source/Runtime/Engine/Classes/Components/PrimitiveComponent.h
+++ b/Engine/Source/Runtime/Engine/Classes/Components/PrimitiveComponent.h
@@ -158,7 +158,7 @@ public:
 	 * If true, this component will generate overlap events when it is overlapping other components (eg Begin Overlap).
 	 * Both components (this and the other) must have this enabled for overlap events to occur.
 	 *
-	 * @see [Overlap Events] https://docs.unrealengine.com/latest/INT/Engine/Physics/Collision/index.html#overlapandgenerateoverlapevents
+	 * @see [Overlap Events](https://docs.unrealengine.com/latest/INT/Engine/Physics/Collision/index.html#overlapandgenerateoverlapevents)
 	 * @see UpdateOverlaps(), BeginComponentOverlap(), EndComponentOverlap()
 	 */
 	UPROPERTY(EditAnywhere, BlueprintReadWrite, Category=Collision)
