Fix for compile problem on Mac

diff --git a/Engine/Source/Runtime/Engine/Private/NiagaraSimulation.cpp b/Engine/Source/Runtime/Engine/Private/NiagaraSimulation.cpp
--- a/Engine/Source/Runtime/Engine/Private/NiagaraSimulation.cpp
+++ b/Engine/Source/Runtime/Engine/Private/NiagaraSimulation.cpp
@@ -5,14 +5,14 @@
 #include "VectorVM.h"
 
 
-FNiagaraSimulation::FNiagaraSimulation()
-: SpawnRate(0)
-, Age(0.0f)
+FNiagaraSimulation::FNiagaraSimulation() : Age(0.0f)
+, bIsEnabled(true)
+, SpawnRate(0)
+, RenderModuleType(RMT_None)
 , UpdateScript(nullptr)
 , SpawnScript(nullptr)
 , SpawnRemainder(0.0f)
 , CachedBounds(ForceInit)
-, RenderModuleType(RMT_None)
 , EffectRenderer(nullptr)
 {
 }
