#ue4 - Clean up deprecation warning push/pop in MovementComponent.

diff --git a/Engine/Source/Runtime/Engine/Private/MovementComponent.cpp b/Engine/Source/Runtime/Engine/Private/MovementComponent.cpp
--- a/Engine/Source/Runtime/Engine/Private/MovementComponent.cpp
+++ b/Engine/Source/Runtime/Engine/Private/MovementComponent.cpp
@@ -546,22 +546,9 @@ float UMovementComponent::GetMaxSpeedModifier() const
 float UMovementComponent::K2_GetMaxSpeedModifier() const
 {
 	// Allow calling old deprecated function to maintain old behavior until it is removed.
-#ifdef __clang__
-	#pragma clang diagnostic push
-	#pragma clang diagnostic ignored "-Wdeprecated-declarations"
-#else
-	#pragma warning(push)
-	#pragma warning(disable:4995)
-	#pragma warning(disable:4996)
-#endif
-
+	PRAGMA_DISABLE_DEPRECATION_WARNINGS
 	return GetMaxSpeedModifier();
-
-#ifdef __clang__
-	#pragma clang diagnostic pop
-#else
-	#pragma warning(pop)
-#endif
+	PRAGMA_ENABLE_DEPRECATION_WARNINGS
 }
 
 // TODO: Deprecated, remove.
@@ -574,22 +561,9 @@ float UMovementComponent::GetModifiedMaxSpeed() const
 float UMovementComponent::K2_GetModifiedMaxSpeed() const
 {
 	// Allow calling old deprecated function to maintain old behavior until it is removed.
-#ifdef __clang__
-	#pragma clang diagnostic push
-	#pragma clang diagnostic ignored "-Wdeprecated-declarations"
-#else
-	#pragma warning(push)
-	#pragma warning(disable:4995)
-	#pragma warning(disable:4996)
-#endif
-
+	PRAGMA_DISABLE_DEPRECATION_WARNINGS
 	return GetModifiedMaxSpeed();
-
-#ifdef __clang__
-	#pragma clang diagnostic pop
-#else
-	#pragma warning(pop)
-#endif
+	PRAGMA_ENABLE_DEPRECATION_WARNINGS
 }
 
 #undef LOCTEXT_NAMESPACE
