Another compile fix

diff --git a/Engine/Source/Runtime/Engine/Classes/Engine/NiagaraSimulation.h b/Engine/Source/Runtime/Engine/Classes/Engine/NiagaraSimulation.h
--- a/Engine/Source/Runtime/Engine/Classes/Engine/NiagaraSimulation.h
+++ b/Engine/Source/Runtime/Engine/Classes/Engine/NiagaraSimulation.h
@@ -58,6 +58,7 @@ class FNiagaraSimulation
 {
 public:
 	explicit FNiagaraSimulation();
+	virtual ~FNiagaraSimulation();
 
 	void Tick(float DeltaSeconds);
 
