Remove inappropriate events from being exposed for Emitter
Expose OnSystemFinished delegate to blueprints

diff --git a/Engine/Source/Runtime/Engine/Classes/Particles/Emitter.h b/Engine/Source/Runtime/Engine/Classes/Particles/Emitter.h
--- a/Engine/Source/Runtime/Engine/Classes/Particles/Emitter.h
+++ b/Engine/Source/Runtime/Engine/Classes/Particles/Emitter.h
@@ -44,7 +44,7 @@ DECLARE_DYNAMIC_MULTICAST_DELEGATE_SixParams( FParticleDeathSignature, FName, Ev
 	*/
 DECLARE_DYNAMIC_MULTICAST_DELEGATE_EightParams( FParticleCollisionSignature, FName, EventName, float, EmitterTime, int32, ParticleTime, FVector, Location, FVector, Velocity, FVector, Direction, FVector, Normal, FName, BoneName);
 
-UCLASS(hidecategories=(Activation,"Components|Activation",Input,Collision), showcategories=("Input|MouseInput", "Input|TouchInput"))
+UCLASS(hidecategories=(Activation,"Components|Activation",Input,Collision,"Game|Damage"))
 class ENGINE_API AEmitter : public AActor
 {
 	GENERATED_UCLASS_BODY()
diff --git a/Engine/Source/Runtime/Engine/Classes/Particles/ParticleSystemComponent.h b/Engine/Source/Runtime/Engine/Classes/Particles/ParticleSystemComponent.h
--- a/Engine/Source/Runtime/Engine/Classes/Particles/ParticleSystemComponent.h
+++ b/Engine/Source/Runtime/Engine/Classes/Particles/ParticleSystemComponent.h
@@ -71,7 +71,7 @@ enum EParticleEventType
 };
 
 // Called when the particle system is done
-DECLARE_DYNAMIC_DELEGATE_OneParam( FOnSystemFinished, class UParticleSystemComponent*, PSystem );
+DECLARE_DYNAMIC_MULTICAST_DELEGATE_OneParam( FOnSystemFinished, class UParticleSystemComponent*, PSystem );
 
 
 /** Struct used for a particular named instance parameter for this ParticleSystemComponent. */
@@ -441,7 +441,7 @@ public:
 	float EmitterDelay;
 
 	// Called when the particle system is done
-	UPROPERTY()
+	UPROPERTY(BlueprintAssignable)
 	FOnSystemFinished OnSystemFinished;
 
 	//
diff --git a/Engine/Source/Runtime/Engine/Private/Particles/ParticleComponents.cpp b/Engine/Source/Runtime/Engine/Private/Particles/ParticleComponents.cpp
--- a/Engine/Source/Runtime/Engine/Private/Particles/ParticleComponents.cpp
+++ b/Engine/Source/Runtime/Engine/Private/Particles/ParticleComponents.cpp
@@ -356,7 +356,7 @@ void AEmitter::PostInitializeComponents()
 	// Set Notification Delegate
 	if (ParticleSystemComponent)
 	{
-		ParticleSystemComponent->OnSystemFinished.BindDynamic( this, &AEmitter::OnParticleSystemFinished );
+		ParticleSystemComponent->OnSystemFinished.AddDynamic( this, &AEmitter::OnParticleSystemFinished );
 		bCurrentlyActive = ParticleSystemComponent->bAutoActivate;
 	}
 
@@ -4033,7 +4033,7 @@ void UParticleSystemComponent::FinalizeTickComponent()
 			TEXT("HasCompleted()==true @ %fs %s"), World->TimeSeconds,
 			Template != NULL ? *Template->GetName() : TEXT("NULL"));
 
-		OnSystemFinished.ExecuteIfBound(this);
+		OnSystemFinished.Broadcast(this);
 
 		// When system is done - destroy all subemitters etc. We don't need them any more.
 		ResetParticles();
