#ue4 - Tweak to air control limit code, to allow pushing in to vertical surfaces so you still get the hit event. Angle surfaces that normally restrict air control still get the event due to contacts caused by gravity.

diff --git a/Engine/Source/Runtime/Engine/Private/CharacterMovementComponent.cpp b/Engine/Source/Runtime/Engine/Private/CharacterMovementComponent.cpp
--- a/Engine/Source/Runtime/Engine/Private/CharacterMovementComponent.cpp
+++ b/Engine/Source/Runtime/Engine/Private/CharacterMovementComponent.cpp
@@ -22,7 +22,7 @@ DEFINE_LOG_CATEGORY_STATIC(LogCharacterMovement, Log, All);
 // MAGIC NUMBERS
 const float MAX_STEP_SIDE_Z = 0.08f;	// maximum z value for the normal on the vertical side of steps
 const float SWIMBOBSPEED = -80.f;
-const float VERTICAL_SLOPE_NORMAL_Z = -0.005f; // Slope is vertical or upward if Normal.Z >= Thresh. Accounts for precision problems that sometimes angle normals slightly off horizontal for vertical surface.
+const float VERTICAL_SLOPE_NORMAL_Z = 0.001f; // Slope is vertical if Abs(Normal.Z) <= this threshold. Accounts for precision problems that sometimes angle normals slightly off horizontal for vertical surface.
 
 const float UCharacterMovementComponent::MIN_TICK_TIME = 0.0002f;
 const float UCharacterMovementComponent::MIN_FLOOR_DIST = 1.9f;
@@ -2988,7 +2988,7 @@ void UCharacterMovementComponent::PhysFalling(float deltaTime, int32 Iterations)
 						}
 
 						// Act as if there was no air control on the last move when computing new deflection.
-						if (bHasAirControl && Hit.Normal.Z >= VERTICAL_SLOPE_NORMAL_Z)
+						if (bHasAirControl && Hit.Normal.Z > VERTICAL_SLOPE_NORMAL_Z)
 						{
 							const FVector LastMoveNoAirControl = VelocityNoAirControl * LastMoveTimeSlice;
 							Delta = ComputeSlideVector(LastMoveNoAirControl, 1.f, OldHitNormal, Hit);
@@ -3095,7 +3095,7 @@ FVector UCharacterMovementComponent::LimitAirControl(float DeltaTime, const FVec
 {
 	FVector Result(FallAcceleration);
 
-	if (HitResult.IsValidBlockingHit() && HitResult.Normal.Z >= VERTICAL_SLOPE_NORMAL_Z)
+	if (HitResult.IsValidBlockingHit() && HitResult.Normal.Z > VERTICAL_SLOPE_NORMAL_Z)
 	{
 		if (!bCheckForValidLandingSpot || !IsValidLandingSpot(HitResult.Location, HitResult))
 		{
