Yet another compilation fix. This time for reals.

diff --git a/Engine/Source/Runtime/Engine/Classes/Engine/NiagaraSimulation.h b/Engine/Source/Runtime/Engine/Classes/Engine/NiagaraSimulation.h
--- a/Engine/Source/Runtime/Engine/Classes/Engine/NiagaraSimulation.h
+++ b/Engine/Source/Runtime/Engine/Classes/Engine/NiagaraSimulation.h
@@ -58,7 +58,8 @@ class FNiagaraSimulation
 {
 public:
 	explicit FNiagaraSimulation();
-	virtual ~FNiagaraSimulation();
+	virtual ~FNiagaraSimulation()
+	{}
 
 	void Tick(float DeltaSeconds);
 
