Fix clang error 'class is abstract but has non-virtual destructor'.

diff --git a/Engine/Source/Runtime/Engine/Classes/Engine/NiagaraEffectRenderer.h b/Engine/Source/Runtime/Engine/Classes/Engine/NiagaraEffectRenderer.h
--- a/Engine/Source/Runtime/Engine/Classes/Engine/NiagaraEffectRenderer.h
+++ b/Engine/Source/Runtime/Engine/Classes/Engine/NiagaraEffectRenderer.h
@@ -73,7 +73,7 @@ public:
 		MaterialRelevance = Material->GetRelevance(FeatureLevel);
 	}
 
-	~NiagaraEffectRenderer() {}
+	virtual ~NiagaraEffectRenderer() {}
 
 protected:
 	NiagaraEffectRenderer()	
