fixed case when a single float was passed in as Subsurface color (should be grey, was red)

diff --git a/Engine/Shaders/ShaderVersion.usf b/Engine/Shaders/ShaderVersion.usf
--- a/Engine/Shaders/ShaderVersion.usf
+++ b/Engine/Shaders/ShaderVersion.usf
@@ -2,4 +2,4 @@
 // Each time the console command is executed it generates a new GUID. As this file is included
 // in common.usf (which should be included in any shader) it allows to invalidate the shader DDC.
 // 
-// GUID = 98C7AD7080954AAFA80A0922AE68EC87afdsfdsdasdfsdfsd
\ No newline at end of file
+// GUID = 4C36424F4976EF3266512DACC6E7D30A
\ No newline at end of file
diff --git a/Engine/Source/Runtime/Engine/Private/Materials/HLSLMaterialTranslator.h b/Engine/Source/Runtime/Engine/Private/Materials/HLSLMaterialTranslator.h
--- a/Engine/Source/Runtime/Engine/Private/Materials/HLSLMaterialTranslator.h
+++ b/Engine/Source/Runtime/Engine/Private/Materials/HLSLMaterialTranslator.h
@@ -345,7 +345,7 @@ public:
 				// Note we don't test for the blend mode as you can have a translucent material using the subsurface shading model
 
 				// another ForceCast as CompilePropertyAndSetMaterialProperty() can return MCT_Float which we don't want here
-				int32 SubsurfaceColor = ForceCast(Material->CompilePropertyAndSetMaterialProperty(MP_SubsurfaceColor, this), MCT_Float3, true);
+				int32 SubsurfaceColor = ForceCast(Material->CompilePropertyAndSetMaterialProperty(MP_SubsurfaceColor, this), MCT_Float3, true, true);
 
 				static FName NameSubsurfaceProfile(TEXT("__SubsurfaceProfile"));
 
