#ue4 - Fix SceneComponent::MoveComponent() so it always fills in the out HitResult no matter the control path.

diff --git a/Engine/Source/Runtime/Engine/Private/SceneComponent.cpp b/Engine/Source/Runtime/Engine/Private/SceneComponent.cpp
--- a/Engine/Source/Runtime/Engine/Private/SceneComponent.cpp
+++ b/Engine/Source/Runtime/Engine/Private/SceneComponent.cpp
@@ -1323,6 +1323,12 @@ bool USceneComponent::MoveComponent( const FVector& Delta, const FRotator& NewRo
 	}
 #endif
 
+	// Fill in optional output param. SceneComponent doesn't sweep, so this is just an empty result.
+	if (OutHit)
+	{
+		*OutHit = FHitResult(1.f);
+	}
+
 	if (!bWorldToComponentUpdated)
 	{
 		UpdateComponentToWorld();
@@ -1349,12 +1355,6 @@ bool USceneComponent::MoveComponent( const FVector& Delta, const FRotator& NewRo
 		UpdateOverlaps();
 	}
 
-	// fill in optional output param
-	if (OutHit)
-	{
-		*OutHit = FHitResult(1.f);
-	}
-
 	return true;
 }
 
