World origin rebasing fixes from CCP

diff --git a/Engine/Source/Runtime/Engine/Classes/Components/LineBatchComponent.h b/Engine/Source/Runtime/Engine/Classes/Components/LineBatchComponent.h
--- a/Engine/Source/Runtime/Engine/Classes/Components/LineBatchComponent.h
+++ b/Engine/Source/Runtime/Engine/Classes/Components/LineBatchComponent.h
@@ -157,6 +157,7 @@ class ULineBatchComponent : public UPrimitiveComponent
 	
 	// Begin UActorComponent interface.
 	virtual void TickComponent(float DeltaTime, enum ELevelTick TickType, FActorComponentTickFunction *ThisTickFunction) override;
+	virtual void ApplyWorldOffset(const FVector& InOffset, bool bWorldShift) override;
 	// End UActorComponent interface.
 
 	ENGINE_API void Flush();
diff --git a/Engine/Source/Runtime/Engine/Private/Level.cpp b/Engine/Source/Runtime/Engine/Private/Level.cpp
--- a/Engine/Source/Runtime/Engine/Private/Level.cpp
+++ b/Engine/Source/Runtime/Engine/Private/Level.cpp
@@ -2311,6 +2311,39 @@ void ULineBatchComponent::TickComponent(float DeltaTime, enum ELevelTick TickTyp
 	}
 }
 
+void ULineBatchComponent::ApplyWorldOffset(const FVector& InOffset, bool bWorldShift)
+{
+	Super::ApplyWorldOffset(InOffset, bWorldShift);
+
+	bool bDirty = false;
+	for (FBatchedLine& Line : BatchedLines)
+	{
+		Line.Start += InOffset;
+		Line.End += InOffset;
+		bDirty = true;
+	}
+
+	for (FBatchedPoint& Point : BatchedPoints)
+	{
+		Point.Position += InOffset;
+		bDirty = true;
+	}
+
+	for (FBatchedMesh& Mesh : BatchedMeshes)
+	{
+		for (FVector& Vert : Mesh.MeshVerts)
+		{
+			Vert += InOffset;
+			bDirty = true;
+		}
+	}
+
+	if (bDirty)
+	{
+		MarkRenderStateDirty();
+	}
+}
+
 /**
 * Creates a new scene proxy for the line batcher component.
 * @return	Pointer to the FLineBatcherSceneProxy
diff --git a/Engine/Source/Runtime/Engine/Private/Particles/ParticleTrail2EmitterInstance.cpp b/Engine/Source/Runtime/Engine/Private/Particles/ParticleTrail2EmitterInstance.cpp
--- a/Engine/Source/Runtime/Engine/Private/Particles/ParticleTrail2EmitterInstance.cpp
+++ b/Engine/Source/Runtime/Engine/Private/Particles/ParticleTrail2EmitterInstance.cpp
@@ -1711,7 +1711,7 @@ bool FParticleRibbonEmitterInstance::Spawn_Source(float DeltaTime)
 					{
 						// Calculate the new tangent from the previous and next position...
 						// LastSourcePosition will be that position of the first particle that will be spawned this frame
-						FVector PositionDelta = (CurrentSourcePosition[TrailIdx] - NextNextSpawnedParticle->Location);
+						FVector PositionDelta = (CurrentSourcePosition[TrailIdx] - PositionOffsetThisTick - NextNextSpawnedParticle->Location);
 						float TimeDelta = ElapsedTime - NextNextSpawnedTrailData->SpawnTime;
 						NewTangent = PositionDelta / TimeDelta;
 						// Calculate new tangents for all the interpolated particles between NextNext and Next
@@ -2812,6 +2812,21 @@ bool FParticleRibbonEmitterInstance::FillReplayData(FDynamicEmitterReplayDataBas
 	return true;
 }
 
+void FParticleRibbonEmitterInstance::ApplyWorldOffset(FVector InOffset, bool bWorldShift)
+{
+	FParticleTrailsEmitterInstance_Base::ApplyWorldOffset(InOffset, bWorldShift);
+
+	for (FVector& Position : CurrentSourcePosition)
+	{
+		Position += InOffset;
+	}
+	
+	for (FVector& Position : LastSourcePosition)
+	{
+		Position += InOffset;
+	}
+}
+
 //
 //	FParticleAnimTrailEmitterInstance
 //
diff --git a/Engine/Source/Runtime/Engine/Private/World.cpp b/Engine/Source/Runtime/Engine/Private/World.cpp
--- a/Engine/Source/Runtime/Engine/Private/World.cpp
+++ b/Engine/Source/Runtime/Engine/Private/World.cpp
@@ -4076,6 +4076,12 @@ bool UWorld::SetNewWorldOrigin(FIntVector InNewOriginLocation)
 		}
 	}
 
+	// Shift navigation meshes
+	if (NavigationSystem)
+	{
+		NavigationSystem->ApplyWorldOffset(Offset, true);
+	}
+
 	// Apply offset to components with no actor (like UGameplayStatics::SpawnEmitterAtLocation) 
 	{
 		TArray <UObject*> WorldChildren; 
@@ -4091,10 +4097,19 @@ bool UWorld::SetNewWorldOrigin(FIntVector InNewOriginLocation)
 		}
 	}
 			
-	// Shift navigation meshes
-	if (NavigationSystem)
+	if (LineBatcher)
 	{
-		NavigationSystem->ApplyWorldOffset(Offset, true);
+		LineBatcher->ApplyWorldOffset(Offset, true);
+	}
+	
+	if (PersistentLineBatcher)
+	{
+		PersistentLineBatcher->ApplyWorldOffset(Offset, true);
+	}
+	
+	if (ForegroundLineBatcher)
+	{
+		ForegroundLineBatcher->ApplyWorldOffset(Offset, true);
 	}
 
 	FIntVector PreviosWorldOriginLocation = OriginLocation;
diff --git a/Engine/Source/Runtime/Engine/Public/ParticleEmitterInstances.h b/Engine/Source/Runtime/Engine/Public/ParticleEmitterInstances.h
--- a/Engine/Source/Runtime/Engine/Public/ParticleEmitterInstances.h
+++ b/Engine/Source/Runtime/Engine/Public/ParticleEmitterInstances.h
@@ -1367,6 +1367,9 @@ struct FParticleRibbonEmitterInstance : public FParticleTrailsEmitterInstance_Ba
 			SourceEmitter = NULL;
 		}
 	}
+	
+	/** Called on world origin changes */
+	virtual void ApplyWorldOffset(FVector InOffset, bool bWorldShift) override;
 
 protected:
 
