Fixed case where FlushLevelStreaming was not flushing state of all levels

diff --git a/Engine/Source/Runtime/Engine/Private/World.cpp b/Engine/Source/Runtime/Engine/Private/World.cpp
--- a/Engine/Source/Runtime/Engine/Private/World.cpp
+++ b/Engine/Source/Runtime/Engine/Private/World.cpp
@@ -2211,7 +2211,9 @@ void UWorld::UpdateLevelStreamingInner( UWorld* PersistentWorld, FSceneViewFamil
 		// on purpose as well so the GC code has a chance to execute between consecutive loads of maps.
 		//
 		// NOTE: AllowLevelLoadRequests not an invariant as streaming might affect the result, do NOT pulled out of the loop.
-		bool bAllowLevelLoadRequests =	(PersistentWorld->AllowLevelLoadRequests() && !bAreLevelsPendingPurge) || bShouldBlockOnLoad;
+		bool bAllowLevelLoadRequests =	(PersistentWorld->AllowLevelLoadRequests() && !bAreLevelsPendingPurge) 
+										|| bShouldBlockOnLoad 
+										|| bFlushingLevelStreaming;
 
 		// Request a 'soft' GC if there are levels pending purge and there are levels to be loaded. In the case of a blocking
 		// load this is going to guarantee GC firing first thing afterwards and otherwise it is going to sneak in right before
@@ -2260,7 +2262,7 @@ void UWorld::UpdateLevelStreamingInner( UWorld* PersistentWorld, FSceneViewFamil
 				if (Level->bIsVisible)
 				{
 					PersistentWorld->RemoveFromWorld(Level);
-			}
+				}
 			}
 
 			if (!bShouldBeLoaded)
@@ -2317,7 +2319,7 @@ void UWorld::UpdateLevelStreaming( FSceneViewFamily* ViewFamily )
 	const int32	NumAsyncLoadingPackages = GetNumAsyncPackages(); 
 	
 	// Update level streaming objects state
-		UpdateLevelStreamingInner(this, ViewFamily);
+	UpdateLevelStreamingInner(this, ViewFamily);
 	
 	// Force initial loading to be "bShouldBlockOnLoad".
 	const bool bLevelsHaveLoadRequestPending = NumAsyncLoadingPackages < GetNumAsyncPackages();
