Static lighting/shadowing support for instanced static meshes (including foliage)

diff --git a/Engine/Shaders/LocalVertexFactory.usf b/Engine/Shaders/LocalVertexFactory.usf
--- a/Engine/Shaders/LocalVertexFactory.usf
+++ b/Engine/Shaders/LocalVertexFactory.usf
@@ -504,36 +504,28 @@ FVertexFactoryInterpolantsVSToPS VertexFactoryGetInterpolantsVSToPS(FVertexFacto
 #endif	// NUM_MATERIAL_TEXCOORDS
 
 #if NEEDS_LIGHTMAP_COORDINATE
+	float2 LightMapCoordinate = 0;
+	float2 ShadowMapCoordinate = 0;
+
 	#if USE_INSTANCING
-	//float2 InstanceLightMapBias = float2(Input.InstanceInverseTransform1.w, Input.InstanceInverseTransform2.w);
-	//Interpolants.LightMapCoordinate.xy = Input.LightMapCoordinate * LightMapCoordinateScaleBias.xy + InstanceLightMapBias;
-	#if (ES2_PROFILE || ES3_1_PROFILE)
-	SetLightMapCoordinate(Interpolants, (float4)0, (float2)0);
+		#if !USE_INSTANCING_EMULATED
+		float2 InstanceLightMapBias = float2(Input.InstanceInverseTransform1.w, Input.InstanceInverseTransform2.w);
 		#else
-	Interpolants.LightMapCoordinate = 0;
-		#endif	// ES2_PROFILE
+		float2 InstanceLightMapBias = float2(CPUInstanceInverseTransform[1].w, CPUInstanceInverseTransform[2].w);
+		#endif // !USE_INSTANCING_EMULATED
+		LightMapCoordinate = Input.LightMapCoordinate * LightMapCoordinateScaleBias.xy + InstanceLightMapBias;
 	#else
-	
-	float4 LightMapCoordinate = 0;
-	float2 ShadowMapCoordinate = 0;
-
-	#if (ES2_PROFILE || ES3_1_PROFILE)
-		LightMapCoordinate.xy = Input.LightMapCoordinate * LightMapCoordinateScaleBias.xy + LightMapCoordinateScaleBias.zw;
-		LightMapCoordinate.y *= 0.5;
-		LightMapCoordinate.zw = LightMapCoordinate.xy;
-		LightMapCoordinate.w += 0.5;
-			#if STATICLIGHTING_TEXTUREMASK
-			ShadowMapCoordinate = Input.LightMapCoordinate * ShadowMapCoordinateScaleBias.xy + ShadowMapCoordinateScaleBias.zw;
-			#endif	// STATICLIGHTING_TEXTUREMASK
+		LightMapCoordinate = Input.LightMapCoordinate * LightMapCoordinateScaleBias.xy + LightMapCoordinateScaleBias.zw;
+	#endif
+	#if STATICLIGHTING_TEXTUREMASK
+		#if USE_INSTANCING
+			ShadowMapCoordinate = Input.LightMapCoordinate * ShadowMapCoordinateScaleBias.xy + GetInstanceShadowMapBias(Input).xy;
 		#else
-		LightMapCoordinate.xy = Input.LightMapCoordinate * LightMapCoordinateScaleBias.xy + LightMapCoordinateScaleBias.zw;
-			#if STATICLIGHTING_TEXTUREMASK
 			ShadowMapCoordinate = Input.LightMapCoordinate * ShadowMapCoordinateScaleBias.xy + ShadowMapCoordinateScaleBias.zw;
-			#endif	// STATICLIGHTING_TEXTUREMASK
-		#endif	// ES2_PROFILE
+		#endif
+	#endif	// STATICLIGHTING_TEXTUREMASK
 
 	SetLightMapCoordinate(Interpolants, LightMapCoordinate, ShadowMapCoordinate);
-	#endif	// USE_INSTANCING
 #endif	// NEEDS_LIGHTMAP_COORDINATE
 
 	SetTangents(Interpolants, Intermediates.TangentToWorld[0], Intermediates.TangentToWorld[2], Intermediates.TangentToWorldSign);
diff --git a/Engine/Shaders/LocalVertexFactoryCommon.usf b/Engine/Shaders/LocalVertexFactoryCommon.usf
--- a/Engine/Shaders/LocalVertexFactoryCommon.usf
+++ b/Engine/Shaders/LocalVertexFactoryCommon.usf
@@ -81,9 +81,9 @@ float2 GetShadowMapCoordinate(FVertexFactoryInterpolantsVSToPS Interpolants)
 	return Interpolants.LightMapCoordinate.zw;
 }
 
-void SetLightMapCoordinate(inout FVertexFactoryInterpolantsVSToPS Interpolants, float4 InLightMapCoordinate, float2 InShadowMapCoordinate)
+void SetLightMapCoordinate(inout FVertexFactoryInterpolantsVSToPS Interpolants, float2 InLightMapCoordinate, float2 InShadowMapCoordinate)
 {
-	Interpolants.LightMapCoordinate.xy = InLightMapCoordinate.xy;
+	Interpolants.LightMapCoordinate.xy = InLightMapCoordinate;
 	Interpolants.LightMapCoordinate.zw = InShadowMapCoordinate;
 }
 #endif
@@ -201,11 +201,11 @@ float2 GetShadowMapCoordinate(FVertexFactoryInterpolantsVSToPS Interpolants)
 	#endif
 }
 
-void SetLightMapCoordinate(inout FVertexFactoryInterpolantsVSToPS Interpolants, float4 InLightMapCoordinate, float2 InShadowMapCoordinate)
+void SetLightMapCoordinate(inout FVertexFactoryInterpolantsVSToPS Interpolants, float2 InLightMapCoordinate, float2 InShadowMapCoordinate)
 {
 	#if USE_INTERPOLATOR_345
-		Interpolants.PackedInterpolator3.xy = InLightMapCoordinate.xy;
-		Interpolants.PackedInterpolator4.xy = InLightMapCoordinate.zw;
+		Interpolants.PackedInterpolator3.xy = InLightMapCoordinate * float2(1, 0.5);
+		Interpolants.PackedInterpolator4.xy = Interpolants.PackedInterpolator3.xy + float2(0, 0.5);
 		Interpolants.PackedInterpolator5.xy = InShadowMapCoordinate;
 	#endif
 }
diff --git a/Engine/Source/Runtime/Core/Public/UObject/ObjectVersion.h b/Engine/Source/Runtime/Core/Public/UObject/ObjectVersion.h
--- a/Engine/Source/Runtime/Core/Public/UObject/ObjectVersion.h
+++ b/Engine/Source/Runtime/Core/Public/UObject/ObjectVersion.h
@@ -510,7 +510,7 @@ enum EUnrealEngineObjectUE4Version
 	VER_UE4_SOFT_CONSTRAINTS_USE_MASS,
 	// Reflection capture data saved in packages
 	VER_UE4_REFLECTION_DATA_IN_PACKAGES,
-	// Fix up old foliage components to have movable mobility
+	// Fix up old foliage components to have movable mobility (superseded by VER_UE4_FOLIAGE_STATIC_LIGHTING_SUPPORT)
 	VER_UE4_FOLIAGE_MOVABLE_MOBILITY,
 	// Undo BreakMaterialAttributes changes as it broke old content
 	VER_UE4_UNDO_BREAK_MATERIALATTRIBUTES_CHANGE,
@@ -626,6 +626,8 @@ enum EUnrealEngineObjectUE4Version
 	VER_UE4_CHARACTER_MOVEMENT_UPPER_IMPACT_BEHAVIOR,
 	// Changed Blueprint math equality functions for vectors and rotators to operate as a "nearly" equals rather than "exact"
 	VER_UE4_BP_MATH_VECTOR_EQUALITY_USES_EPSILON,
+	// Static lighting support was re-added to foliage, and mobility was returned to static
+	VER_UE4_FOLIAGE_STATIC_LIGHTING_SUPPORT,
 
 	// -----<new versions can be added before this line>-------------------------------------------------
 	// - this needs to be the last line (see note below)
diff --git a/Engine/Source/Runtime/Engine/Classes/Components/InstancedStaticMeshComponent.h b/Engine/Source/Runtime/Engine/Classes/Components/InstancedStaticMeshComponent.h
--- a/Engine/Source/Runtime/Engine/Classes/Components/InstancedStaticMeshComponent.h
+++ b/Engine/Source/Runtime/Engine/Classes/Components/InstancedStaticMeshComponent.h
@@ -6,6 +6,10 @@
 
 #include "InstancedStaticMeshComponent.generated.h"
 
+class FStaticLightingTextureMapping_InstancedStaticMesh;
+class FInstancedLightMap2D;
+class FInstancedShadowMap2D;
+
 USTRUCT()
 struct FInstancedStaticMeshInstanceData
 {
@@ -29,13 +33,13 @@ struct FInstancedStaticMeshInstanceData
 	}
 
 
-		friend FArchive& operator<<(FArchive& Ar, FInstancedStaticMeshInstanceData& InstanceData)
-		{
-			// @warning BulkSerialize: FInstancedStaticMeshInstanceData is serialized as memory dump
-			// See TArray::BulkSerialize for detailed description of implied limitations.
-			Ar << InstanceData.Transform << InstanceData.LightmapUVBias << InstanceData.ShadowmapUVBias;
-			return Ar;
-		}
+	friend FArchive& operator<<(FArchive& Ar, FInstancedStaticMeshInstanceData& InstanceData)
+	{
+		// @warning BulkSerialize: FInstancedStaticMeshInstanceData is serialized as memory dump
+		// See TArray::BulkSerialize for detailed description of implied limitations.
+		Ar << InstanceData.Transform << InstanceData.LightmapUVBias << InstanceData.ShadowmapUVBias;
+		return Ar;
+	}
 	
 };
 
@@ -44,19 +48,12 @@ struct FInstancedStaticMeshMappingInfo
 {
 	GENERATED_USTRUCT_BODY()
 
-	UPROPERTY()
-	class UTexture2D* LightmapTexture;
-
+	FStaticLightingTextureMapping_InstancedStaticMesh* Mapping;
 
 	FInstancedStaticMeshMappingInfo()
-		: LightmapTexture(NULL)
+		: Mapping(nullptr)
 	{
 	}
-
-
-	class FInstancedStaticMeshStaticLightingTextureMapping*		Mapping;
-
-	class FInstancedLightMap2D*									Lightmap;
 };
 
 /** A component that efficiently renders multiple instances of the same StaticMesh. */
@@ -67,7 +64,7 @@ class ENGINE_API UInstancedStaticMeshComponent : public UStaticMeshComponent
 
 	/** Array of instances, bulk serialized */
 	UPROPERTY(EditAnywhere, Transient, DuplicateTransient, DisplayName="Instances", Category=Instances, meta=(MakeEditWidget=true))
-	TArray<struct FInstancedStaticMeshInstanceData> PerInstanceSMData;
+	TArray<FInstancedStaticMeshInstanceData> PerInstanceSMData;
 
 	/** Value used to seed the random number stream that generates random numbers for each of this mesh's instances.
 		The random number is stored in a buffer accessible to materials through the PerInstanceRandom expression.  If
@@ -144,6 +141,7 @@ public:
 	virtual bool CanEditSimulatePhysics() override;
 
 	virtual FBoxSphereBounds CalcBounds(const FTransform& BoundTransform) const override;
+	virtual bool SupportsStaticLighting() const override { return true; }
 #if WITH_EDITOR
 	virtual void GetStaticLightingInfo(FStaticLightingPrimitiveInfo& OutPrimitiveInfo,const TArray<ULightComponent*>& InRelevantLights,const FLightingBuildOptions& Options) override;
 #endif
@@ -181,12 +179,18 @@ private:
 	void SetupNewInstanceData(FInstancedStaticMeshInstanceData& InOutNewInstanceData, int32 InInstanceIndex, const FTransform& InInstanceTransform);
 
 protected:
+	/** Number of pending lightmaps still to be calculated (Apply()'d) */
+	UPROPERTY(Transient, DuplicateTransient, TextExportTransient)
+	int32 NumPendingLightmaps;
 
-	/** Whether the component type supports static lighting. */
-	virtual bool SupportsStaticLighting() const override
-	{
-		//@todo - support for instanced meshes
-		return false;
-	}
+	/** The mappings for all the instances of this component */
+	UPROPERTY(Transient, DuplicateTransient, TextExportTransient)
+	TArray<FInstancedStaticMeshMappingInfo> CachedMappings;
+
+	void ApplyLightMapping(FStaticLightingTextureMapping_InstancedStaticMesh* InMapping);
+
+	friend FStaticLightingTextureMapping_InstancedStaticMesh;
+	friend FInstancedLightMap2D;
+	friend FInstancedShadowMap2D;
 };
 
diff --git a/Engine/Source/Runtime/Engine/Classes/Foliage/FoliageType.h b/Engine/Source/Runtime/Engine/Classes/Foliage/FoliageType.h
--- a/Engine/Source/Runtime/Engine/Classes/Foliage/FoliageType.h
+++ b/Engine/Source/Runtime/Engine/Classes/Foliage/FoliageType.h
@@ -206,14 +206,6 @@ class UFoliageType : public UObject
 	UPROPERTY(EditAnywhere, BlueprintReadOnly, Category=Lighting)
 	uint32 bCastStaticShadow:1;
 
-	/** 
-	 *	If true, the foliage will cast shadows even if bHidden is true.
-	 *	Controls whether the primitive should cast shadows when hidden.
-	 *	This flag is only used if CastShadow is true.
-	 */
-	UPROPERTY(EditAnywhere, AdvancedDisplay, BlueprintReadOnly, Category=Lighting)
-	uint32 bCastHiddenShadow:1;
-
 	/** Whether this foliage should cast dynamic shadows as if it were a two sided material. */
 	UPROPERTY(EditAnywhere, AdvancedDisplay, BlueprintReadOnly, Category=Lighting)
 	uint32 bCastShadowAsTwoSided:1;
@@ -221,6 +213,14 @@ class UFoliageType : public UObject
 	/** Whether the foliage receives decals. */
 	UPROPERTY(EditAnywhere, AdvancedDisplay, BlueprintReadOnly, Category=Rendering)
 	uint32 bReceivesDecals : 1;
+	
+	/** Whether to override the lightmap resolution defined in the static mesh. */
+	UPROPERTY(BlueprintReadOnly, Category=Lighting)
+	uint32 bOverrideLightMapRes:1;
+
+	/** Overrides the lightmap resolution defined in the static mesh. A value of 0 disables static lighting/shadowing */
+	UPROPERTY(EditAnywhere, BlueprintReadOnly, Category=Lighting, meta=(DisplayName="Light Map Resolution", EditCondition="bOverrideLightMapRes"))
+	int32 OverriddenLightMapRes;
 
 	/** Custom collision for foliage */
 	UPROPERTY(EditAnywhere, Category=Collision, meta=(HideObjectType=true))
@@ -228,6 +228,10 @@ class UFoliageType : public UObject
 
 	/* Gets the mesh associated with this FoliageType */
 	virtual UStaticMesh* GetStaticMesh() PURE_VIRTUAL(UFoliageType::GetStaticMesh, return nullptr; );
+
+#if WITH_EDITOR
+	virtual void PostEditChangeProperty(struct FPropertyChangedEvent& PropertyChangedEvent) override;
+#endif
 };
 
 
diff --git a/Engine/Source/Runtime/Engine/Private/InstancedFoliage.cpp b/Engine/Source/Runtime/Engine/Private/InstancedFoliage.cpp
--- a/Engine/Source/Runtime/Engine/Private/InstancedFoliage.cpp
+++ b/Engine/Source/Runtime/Engine/Private/InstancedFoliage.cpp
@@ -127,10 +127,12 @@ UFoliageType::UFoliageType(const FObjectInitializer& ObjectInitializer)
 	bCastDynamicShadow = true;
 	bCastStaticShadow = true;
 	bAffectDynamicIndirectLighting = false;
-	bCastHiddenShadow = false;
 	bCastShadowAsTwoSided = false;
 	bReceivesDecals = false;
 
+	bOverrideLightMapRes = false;
+	OverriddenLightMapRes = 32;
+
 	BodyInstance.SetCollisionProfileName(UCollisionProfile::NoCollision_ProfileName);
 }
 
@@ -140,6 +142,16 @@ UFoliageType_InstancedStaticMesh::UFoliageType_InstancedStaticMesh(const FObject
 	Mesh = nullptr;
 }
 
+#if WITH_EDITOR
+void UFoliageType::PostEditChangeProperty(struct FPropertyChangedEvent& PropertyChangedEvent)
+{
+	Super::PostEditChangeProperty(PropertyChangedEvent);
+
+	// Ensure that OverriddenLightMapRes is a factor of 4
+	OverriddenLightMapRes = FMath::Max(OverriddenLightMapRes + 3 & ~3, 4);
+}
+#endif
+
 //
 // FFoliageMeshInfo
 //
@@ -277,23 +289,23 @@ void FFoliageMeshInfo::AddInstance(AInstancedFoliageActor* InIFA, UFoliageType*
 			InSettings->GetStaticMesh()->GetBounds().TransformBy(InstanceToWorld)
 			);
 
-		// Make the instanced static mesh component movable so it doesn't get statically lit; see the comment below about CastShadow for more details
-		BestCluster->ClusterComponent->Mobility = EComponentMobility::Movable;
+		BestCluster->ClusterComponent->Mobility = EComponentMobility::Static;
 
 		BestCluster->ClusterComponent->StaticMesh = InSettings->GetStaticMesh();
 		BestCluster->ClusterComponent->bSelectable = true;
 		BestCluster->ClusterComponent->bHasPerInstanceHitProxies = true;
 		BestCluster->ClusterComponent->InstancingRandomSeed = FMath::Rand();
 		BestCluster->ClusterComponent->InstanceStartCullDistance = InSettings->StartCullDistance;
-		BestCluster->ClusterComponent->InstanceEndCullDistance = InSettings->EndCullDistance;
+		BestCluster->ClusterComponent->InstanceEndCullDistance   = InSettings->EndCullDistance;
 
-		BestCluster->ClusterComponent->CastShadow = InSettings->CastShadow;
-		BestCluster->ClusterComponent->bCastDynamicShadow = InSettings->bCastDynamicShadow;
-		BestCluster->ClusterComponent->bCastStaticShadow = InSettings->bCastStaticShadow;
+		BestCluster->ClusterComponent->CastShadow                     = InSettings->CastShadow;
+		BestCluster->ClusterComponent->bCastDynamicShadow             = InSettings->bCastDynamicShadow;
+		BestCluster->ClusterComponent->bCastStaticShadow              = InSettings->bCastStaticShadow;
 		BestCluster->ClusterComponent->bAffectDynamicIndirectLighting = InSettings->bAffectDynamicIndirectLighting;
-		BestCluster->ClusterComponent->bCastHiddenShadow = InSettings->bCastHiddenShadow;
-		BestCluster->ClusterComponent->bCastShadowAsTwoSided = InSettings->bCastShadowAsTwoSided;
-		BestCluster->ClusterComponent->bReceivesDecals = InSettings->bReceivesDecals;
+		BestCluster->ClusterComponent->bCastShadowAsTwoSided          = InSettings->bCastShadowAsTwoSided;
+		BestCluster->ClusterComponent->bReceivesDecals                = InSettings->bReceivesDecals;
+		BestCluster->ClusterComponent->bOverrideLightMapRes           = InSettings->bOverrideLightMapRes;
+		BestCluster->ClusterComponent->OverriddenLightMapRes          = InSettings->OverriddenLightMapRes;
 
 		BestCluster->ClusterComponent->BodyInstance.CopyBodyInstancePropertiesFrom(&InSettings->BodyInstance);
 
@@ -1470,16 +1482,20 @@ void AInstancedFoliageActor::Serialize(FArchive& Ar)
 		Ar << FoliageMeshes;
 	}
 
-	if (Ar.UE4Ver() < VER_UE4_FOLIAGE_MOVABLE_MOBILITY)
+	if (Ar.UE4Ver() < VER_UE4_FOLIAGE_STATIC_LIGHTING_SUPPORT)
 	{
 		for (auto& MeshInfo : FoliageMeshes)
 		{
 			for (const FFoliageInstanceCluster& Cluster : MeshInfo.Value->InstanceClusters)
 			{
-				if (Ar.UE4Ver() < VER_UE4_FOLIAGE_MOVABLE_MOBILITY)
-					Cluster.ClusterComponent->SetMobility(EComponentMobility::Movable);
+				if (Ar.UE4Ver() < VER_UE4_FOLIAGE_STATIC_LIGHTING_SUPPORT)
+				{
+					Cluster.ClusterComponent->SetMobility(EComponentMobility::Static);
+				}
 				if (Ar.UE4Ver() < VER_UE4_FOLIAGE_COLLISION)
+				{
 					Cluster.ClusterComponent->SetCollisionEnabled(ECollisionEnabled::NoCollision);
+				}
 			}
 		}
 	}
diff --git a/Engine/Source/Runtime/Engine/Private/InstancedStaticMesh.cpp b/Engine/Source/Runtime/Engine/Private/InstancedStaticMesh.cpp
--- a/Engine/Source/Runtime/Engine/Private/InstancedStaticMesh.cpp
+++ b/Engine/Source/Runtime/Engine/Private/InstancedStaticMesh.cpp
@@ -42,6 +42,11 @@
 #include "Collision/PhysXCollision.h"
 #endif
 
+#if WITH_EDITOR
+#include "LightMap.h"
+#include "ShadowMap.h"
+#endif
+
 // This must match the maximum a user could specify in the material (see 
 // FHLSLMaterialTranslator::TextureCoordinate), otherwise the material will attempt 
 // to look up a texture coordinate we didn't provide an element for.
@@ -754,15 +759,12 @@ void FInstancedStaticMeshRenderData::InitStaticMeshVertexFactories(
 
 			if (Parent->LightMapCoordinateIndex >= 0 && Parent->LightMapCoordinateIndex < NumTexCoords)
 			{
-#if 0
-				//@todo UE4 foliage - static lighting/shadowing?
-				Data.ShadowMapCoordinateComponent = FVertexStreamComponent(
+				Data.LightMapCoordinateComponent = FVertexStreamComponent(
 					&RenderData->VertexBuffer,
 					STRUCT_OFFSET(TStaticMeshFullVertexFloat16UVs<MAX_STATIC_TEXCOORDS>, UVs) + sizeof(FVector2DHalf)* Parent->LightMapCoordinateIndex,
 					RenderData->VertexBuffer.GetStride(),
 					VET_Half2
 					);
-#endif
 			}
 		}
 		else
@@ -790,15 +792,12 @@ void FInstancedStaticMeshRenderData::InitStaticMeshVertexFactories(
 
 			if (Parent->LightMapCoordinateIndex >= 0 && Parent->LightMapCoordinateIndex < NumTexCoords)
 			{
-#if 0
-				//@todo UE4 foliage - static lighting/shadowing?
-				Data.ShadowMapCoordinateComponent = FVertexStreamComponent(
+				Data.LightMapCoordinateComponent = FVertexStreamComponent(
 					&RenderData->VertexBuffer,
 					STRUCT_OFFSET(TStaticMeshFullVertexFloat32UVs<InstancedStaticMeshMaxTexCoord>, UVs) + sizeof(FVector2D)* Parent->LightMapCoordinateIndex,
 					RenderData->VertexBuffer.GetStride(),
 					VET_Float2
 					);
-#endif
 			}
 		}
 
@@ -1216,20 +1215,53 @@ UInstancedStaticMeshComponent::UInstancedStaticMeshComponent(const FObjectInitia
 {
 	Mobility = EComponentMobility::Movable;
 	BodyInstance.bSimulatePhysics = false;
+
+	char test[] = "abcefgijkmnoCDEGHIKLM";
+	TRotationInPlaceMerge<TJugglingRotation<FEuclidDivisionGCD>>::Merge(test, 12, 21, [](char a, char b){ return FCharAnsi::ToLower(a) < FCharAnsi::ToLower(b); });
+
+	char test1[] = "abcefgijkCDEGHIKLM";
+	TRotationInPlaceMerge<TJugglingRotation<FEuclidDivisionGCD>>::Merge(test1, 9, 18, [](char a, char b){ return FCharAnsi::ToLower(a) < FCharAnsi::ToLower(b); });
 }
 
 #if WITH_EDITOR
-/** Helper class used to preserve selection state across component duplication */
-class FInstancedStaticMeshSelectionData : public FComponentInstanceDataBase
+/** Helper class used to preserve lighting/selection state across blueprint reinstancing */
+class FInstancedStaticMeshComponentInstanceData : public FComponentInstanceDataBase
 {
 public:
-	FInstancedStaticMeshSelectionData(const UInstancedStaticMeshComponent& InComponent)
+	FInstancedStaticMeshComponentInstanceData(const UInstancedStaticMeshComponent& InComponent)
 		: FComponentInstanceDataBase(&InComponent)
-		, SelectedInstances(InComponent.SelectedInstances)
+		, StaticMesh(InComponent.StaticMesh)
+		, bHasCachedStaticLighting(false)
+	{
+	}
+
+	virtual bool MatchesComponent(const UActorComponent* Component) const override
 	{
+		return (CastChecked<UStaticMeshComponent>(Component)->StaticMesh == StaticMesh && FComponentInstanceDataBase::MatchesComponent(Component));
 	}
 
+	/** Used to store lightmap data during RerunConstructionScripts */
+	struct FLightMapInstanceData
+	{
+		/** Transform of component */
+		FTransform Transform;
+
+		/** Lightmaps from LODData */
+		TArray<FLightMapRef>	LODDataLightMap;
+		TArray<FShadowMapRef>	LODDataShadowMap;
+
+		TArray<FGuid> IrrelevantLights;
+	};
+
 public:
+	/** Mesh being used by component */
+	UStaticMesh* StaticMesh;
+
+	// Static lighting info
+	bool bHasCachedStaticLighting;
+	FLightMapInstanceData CachedStaticLighting;
+	TArray<FInstancedStaticMeshInstanceData> PerInstanceSMData;
+
 	/** The cached selected instances */
 	TBitArray<> SelectedInstances;
 };
@@ -1237,24 +1269,95 @@ public:
 
 FName UInstancedStaticMeshComponent::GetComponentInstanceDataType() const
 {
-	static const FName InstanceStaticMeshComponentInstanceDataName(TEXT("InstancedStaticMeshSelectionData"));
-	return InstanceStaticMeshComponentInstanceDataName;
+	static const FName InstancedStaticMeshComponentInstanceDataName(TEXT("InstancedStaticMeshComponentInstanceData"));
+	return InstancedStaticMeshComponentInstanceDataName;
 }
 
 TSharedPtr<FComponentInstanceDataBase> UInstancedStaticMeshComponent::GetComponentInstanceData() const
 {
-	TSharedPtr<FComponentInstanceDataBase> InstanceData;
 #if WITH_EDITOR
-	InstanceData = MakeShareable(new FInstancedStaticMeshSelectionData(*this));
-#endif
+	TSharedPtr<FInstancedStaticMeshComponentInstanceData> InstanceData;
+
+	// Don't back up static lighting if there isn't any
+	if (bHasCachedStaticLighting || SelectedInstances.Num() > 0)
+	{
+		InstanceData = MakeShareable(new FInstancedStaticMeshComponentInstanceData(*this));
+	}
+
+	// Don't back up static lighting if there isn't any
+	if (bHasCachedStaticLighting)
+	{
+		// Fill in info (copied from UStaticMeshComponent::GetComponentInstanceData)
+		InstanceData->bHasCachedStaticLighting = true;
+		InstanceData->CachedStaticLighting.Transform = ComponentToWorld;
+		InstanceData->CachedStaticLighting.IrrelevantLights = IrrelevantLights;
+		InstanceData->CachedStaticLighting.LODDataLightMap.Empty(LODData.Num());
+		for (const FStaticMeshComponentLODInfo& LODDataEntry : LODData)
+		{
+			InstanceData->CachedStaticLighting.LODDataLightMap.Add(LODDataEntry.LightMap);
+			InstanceData->CachedStaticLighting.LODDataShadowMap.Add(LODDataEntry.ShadowMap);
+		}
+
+		// Back up per-instance lightmap/shadowmap info
+		InstanceData->PerInstanceSMData = PerInstanceSMData;
+	}
+
+	// Back up instance selection
+	if (SelectedInstances.Num() > 0)
+	{
+		InstanceData->SelectedInstances = SelectedInstances;
+	}
+
 	return InstanceData;
+#else
+	return TSharedPtr<FComponentInstanceDataBase>();
+#endif
 }
 
 void UInstancedStaticMeshComponent::ApplyComponentInstanceData(TSharedPtr<FComponentInstanceDataBase> ComponentInstanceData)
 {
 #if WITH_EDITOR
 	check(ComponentInstanceData.IsValid());
-	SelectedInstances = StaticCastSharedPtr<FInstancedStaticMeshSelectionData>(ComponentInstanceData)->SelectedInstances;
+
+	auto* InstancedMeshData = static_cast<FInstancedStaticMeshComponentInstanceData*>(ComponentInstanceData.Get());
+
+	// See if data matches current state
+	if (InstancedMeshData->bHasCachedStaticLighting)
+	{
+		bool bMatch = InstancedMeshData->CachedStaticLighting.Transform.Equals(ComponentToWorld);
+
+		// Check for any instance having moved as that would invalidate static lighting
+		if (PerInstanceSMData.Num() == InstancedMeshData->PerInstanceSMData.Num())
+		{
+			for (int32 InstanceIndex = 0; InstanceIndex < PerInstanceSMData.Num(); ++InstanceIndex)
+			{
+				if (PerInstanceSMData[InstanceIndex].Transform != InstancedMeshData->PerInstanceSMData[InstanceIndex].Transform)
+				{
+					bMatch = false;
+					break;
+				}
+			}
+		}
+
+		// Restore static lighting if appropriate
+		if (bMatch)
+		{
+			const int32 NumLODLightMaps = InstancedMeshData->CachedStaticLighting.LODDataLightMap.Num();
+			SetLODDataCount(NumLODLightMaps, NumLODLightMaps);
+			for (int32 i = 0; i < NumLODLightMaps; ++i)
+			{
+				LODData[i].LightMap = InstancedMeshData->CachedStaticLighting.LODDataLightMap[i];
+				LODData[i].ShadowMap = InstancedMeshData->CachedStaticLighting.LODDataShadowMap[i];
+			}
+			IrrelevantLights = InstancedMeshData->CachedStaticLighting.IrrelevantLights;
+			PerInstanceSMData = InstancedMeshData->PerInstanceSMData;
+			bHasCachedStaticLighting = true;
+		}
+	}
+
+	SelectedInstances = InstancedMeshData->SelectedInstances;
+
+	MarkRenderStateDirty();
 #endif
 }
 
@@ -1427,8 +1530,8 @@ class FStaticLightingMesh_InstancedStaticMesh : public FStaticMeshStaticLighting
 public:
 
 	/** Initialization constructor. */
-	FStaticLightingMesh_InstancedStaticMesh(const UInstancedStaticMeshComponent* InPrimitive, int32 InstanceIndex, const TArray<ULightComponent*>& InRelevantLights)
-		: FStaticMeshStaticLightingMesh(InPrimitive, 0, InRelevantLights)
+	FStaticLightingMesh_InstancedStaticMesh(const UInstancedStaticMeshComponent* InPrimitive, int32 LODIndex, int32 InstanceIndex, const TArray<ULightComponent*>& InRelevantLights)
+		: FStaticMeshStaticLightingMesh(InPrimitive, LODIndex, InRelevantLights)
 	{
 		// override the local to world to combine the per instance transform with the component's standard transform
 		SetLocalToWorld(InPrimitive->PerInstanceSMData[InstanceIndex].Transform * InPrimitive->ComponentToWorld.ToMatrixWithScale());
@@ -1445,16 +1548,36 @@ class FStaticLightingTextureMapping_InstancedStaticMesh : public FStaticMeshStat
 {
 public:
 	/** Initialization constructor. */
-	FStaticLightingTextureMapping_InstancedStaticMesh(UInstancedStaticMeshComponent* InPrimitive, FStaticLightingMesh* InMesh)
-		: FStaticMeshStaticLightingTextureMapping(InPrimitive, 0, InMesh, 0, 0, 0, false)
+	FStaticLightingTextureMapping_InstancedStaticMesh(UInstancedStaticMeshComponent* InPrimitive, int32 LODIndex, int32 InInstanceIndex, FStaticLightingMesh* InMesh, int32 InSizeX, int32 InSizeY, int32 InTextureCoordinateIndex, bool bPerformFullQualityRebuild)
+		: FStaticMeshStaticLightingTextureMapping(InPrimitive, LODIndex, InMesh, InSizeX, InSizeY, InTextureCoordinateIndex, bPerformFullQualityRebuild)
+		, InstanceIndex(InInstanceIndex)
+		, QuantizedData(nullptr)
+		, ShadowMapData()
+		, bComplete(false)
 	{
-		// We don't actually support light/shadow mapping on instanced meshes, only *casting* shadows
-		bProcessMapping = false;
 	}
 
-	virtual void Apply(FQuantizedLightmapData* QuantizedData, const TMap<ULightComponent*, FShadowMapData2D*>& ShadowMapData) override
+	// FStaticLightingTextureMapping interface
+	virtual void Apply(FQuantizedLightmapData* InQuantizedData, const TMap<ULightComponent*, FShadowMapData2D*>& InShadowMapData) override
 	{
-		// Not supported
+		check(bComplete == false);
+
+		UInstancedStaticMeshComponent* InstancedComponent = Cast<UInstancedStaticMeshComponent>(Primitive.Get());
+
+		if (InstancedComponent)
+		{
+			// Save the static lighting until all of the component's static lighting has been built.
+			QuantizedData = TUniquePtr<FQuantizedLightmapData>(InQuantizedData);
+			ShadowMapData.Empty(InShadowMapData.Num());
+			for (auto& ShadowDataPair : InShadowMapData)
+			{
+				ShadowMapData.Add(ShadowDataPair.Key, TUniquePtr<FShadowMapData2D>(ShadowDataPair.Value));
+			}
+
+			InstancedComponent->ApplyLightMapping(this);
+		}
+
+		bComplete = true;
 	}
 
 #if WITH_EDITOR
@@ -1468,23 +1591,154 @@ public:
 	{
 		return FString(TEXT("InstancedSMLightingMapping"));
 	}
+
+private:
+	friend class UInstancedStaticMeshComponent;
+
+	/** The instance of the primitive this mapping represents. */
+	const int32 InstanceIndex;
+
+	// Light/shadow map data stored until all instances for this component are processed
+	// so we can apply them all into one light/shadowmap
+	TUniquePtr<FQuantizedLightmapData> QuantizedData;
+	TMap<ULightComponent*, TUniquePtr<FShadowMapData2D>> ShadowMapData;
+
+	/** Has this mapping already been completed? */
+	bool bComplete;
 };
 
-void UInstancedStaticMeshComponent::GetStaticLightingInfo(FStaticLightingPrimitiveInfo& OutPrimitiveInfo,const TArray<ULightComponent*>& InRelevantLights,const FLightingBuildOptions& Options)
+void UInstancedStaticMeshComponent::GetStaticLightingInfo(FStaticLightingPrimitiveInfo& OutPrimitiveInfo, const TArray<ULightComponent*>& InRelevantLights, const FLightingBuildOptions& Options)
 {
-	// We don't support light/shadow mapping for instanced meshes, only *casting* shadows
-	// we intentionally ignore the mobility setting here, as foliage is marked "Movable" to force
-	// dynamic lighting due to static lighting not being supported, but is actually static
-	if (StaticMesh && bCastStaticShadow)
+	if (HasValidSettingsForStaticLighting())
 	{
-		for (int32 InstanceIndex = 0; InstanceIndex < PerInstanceSMData.Num(); InstanceIndex++)
+		// create static lighting for LOD 0
+		int32 LightMapWidth = 0;
+		int32 LightMapHeight = 0;
+		GetLightMapResolution(LightMapWidth, LightMapHeight);
+
+		const int32 LightMapSize = GetWorld()->GetWorldSettings()->PackedLightAndShadowMapTextureSize;
+		const int32 MaxInstancesInDefaultSizeLightmap = (LightMapSize / LightMapWidth) * ((LightMapSize / 2) / LightMapHeight);
+		if (PerInstanceSMData.Num() > MaxInstancesInDefaultSizeLightmap)
+		{
+			FMessageLog("LightingResults").Message(EMessageSeverity::Warning)
+				->AddToken(FUObjectToken::Create(this))
+				->AddToken(FTextToken::Create(NSLOCTEXT("InstancedStaticMesh", "LargeStaticLightingWarning", "The total lightmap size for this InstancedStaticMeshComponent is large, consider reducing the component's lightmap resolution or number of mesh instances in this component")));
+		}
+
+		bool bCanLODsShareStaticLighting = StaticMesh->CanLODsShareStaticLighting();
+
+		// TODO: We currently only support one LOD of static lighting in instanced meshes
+		// Need to create per-LOD instance data to fix that
+		if (!bCanLODsShareStaticLighting)
+		{
+			UE_LOG(LogStaticMesh, Warning, TEXT("Instanced meshes don't yet support unique static lighting for each LOD, lighting on LOD 1+ may be incorrect"));
+			bCanLODsShareStaticLighting = true;
+		}
+
+		int32 NumLODs = bCanLODsShareStaticLighting ? 1 : StaticMesh->RenderData->LODResources.Num();
+
+		CachedMappings.Reset(PerInstanceSMData.Num() * NumLODs);
+		CachedMappings.AddZeroed(PerInstanceSMData.Num() * NumLODs);
+
+		for (int32 LODIndex = 0; LODIndex < NumLODs; LODIndex++)
+		{
+			const FStaticMeshLODResources& LODRenderData = StaticMesh->RenderData->LODResources[LODIndex];
+
+			for (int32 InstanceIndex = 0; InstanceIndex < PerInstanceSMData.Num(); InstanceIndex++)
+			{
+				auto* StaticLightingMesh = new FStaticLightingMesh_InstancedStaticMesh(this, LODIndex, InstanceIndex, InRelevantLights);
+				OutPrimitiveInfo.Meshes.Add(StaticLightingMesh);
+
+				auto* InstancedMapping = new FStaticLightingTextureMapping_InstancedStaticMesh(this, LODIndex, InstanceIndex, StaticLightingMesh, LightMapWidth, LightMapHeight, StaticMesh->LightMapCoordinateIndex, true);
+				OutPrimitiveInfo.Mappings.Add(InstancedMapping);
+
+				CachedMappings[LODIndex * PerInstanceSMData.Num() + InstanceIndex].Mapping = InstancedMapping;
+				NumPendingLightmaps++;
+			}
+
+			// Shrink LOD texture lightmaps by half for each LOD level (minimum 4x4 px)
+			LightMapWidth  = FMath::Max(LightMapWidth  / 2, 4);
+			LightMapHeight = FMath::Max(LightMapHeight / 2, 4);
+		}
+	}
+}
+
+void UInstancedStaticMeshComponent::ApplyLightMapping(FStaticLightingTextureMapping_InstancedStaticMesh* InMapping)
+{
+	NumPendingLightmaps--;
+
+	if (NumPendingLightmaps == 0)
+	{
+		// Calculate the range of each coefficient in this light-map and repack the data to have the same scale factor and bias across all instances
+		// TODO: Per instance scale?
+
+		// generate the final lightmaps for all the mappings for this component
+		TArray<TUniquePtr<FQuantizedLightmapData>> AllQuantizedData;
+		for (auto& MappingInfo : CachedMappings)
+		{
+			FStaticLightingTextureMapping_InstancedStaticMesh* Mapping = MappingInfo.Mapping;
+			AllQuantizedData.Add(MoveTemp(Mapping->QuantizedData));
+		}
+
+		bool bNeedsShadowMap = false;
+		TArray<TMap<ULightComponent*, TUniquePtr<FShadowMapData2D>>> AllShadowMapData;
+		for (auto& MappingInfo : CachedMappings)
 		{
-			FStaticLightingMesh_InstancedStaticMesh* StaticLightingMesh = new FStaticLightingMesh_InstancedStaticMesh(this, InstanceIndex, InRelevantLights);
-			OutPrimitiveInfo.Meshes.Add(StaticLightingMesh);
+			FStaticLightingTextureMapping_InstancedStaticMesh* Mapping = MappingInfo.Mapping;
+			bNeedsShadowMap = bNeedsShadowMap || (Mapping->ShadowMapData.Num() > 0);
+			AllShadowMapData.Add(MoveTemp(Mapping->ShadowMapData));
+		}
+
+		// Create a light-map for the primitive.
+		const ELightMapPaddingType PaddingType = GAllowLightmapPadding ? LMPT_NormalPadding : LMPT_NoPadding;
+		FLightMap2D* NewLightMap = FLightMap2D::AllocateInstancedLightMap(this, MoveTemp(AllQuantizedData), Bounds, PaddingType, LMF_Streamed);
+
+		// Create a shadow-map for the primitive.
+		FShadowMap2D* NewShadowMap = bNeedsShadowMap ? FShadowMap2D::AllocateInstancedShadowMap(this, MoveTemp(AllShadowMapData), Bounds, PaddingType, SMF_Streamed) : nullptr;
+
+		// Ensure LODData has enough entries in it, free not required.
+		SetLODDataCount(StaticMesh->GetNumLODs(), StaticMesh->GetNumLODs());
 
-			FStaticLightingTextureMapping_InstancedStaticMesh* InstancedMapping = new FStaticLightingTextureMapping_InstancedStaticMesh(this, StaticLightingMesh);
-			OutPrimitiveInfo.Mappings.Add(InstancedMapping);
+		// Share lightmap/shadowmap over all LODs
+		for (FStaticMeshComponentLODInfo& ComponentLODInfo : LODData)
+		{
+			ComponentLODInfo.LightMap = NewLightMap;
+			ComponentLODInfo.ShadowMap = NewShadowMap;
 		}
+
+		// Build the list of statically irrelevant lights.
+		// TODO: This should be stored per LOD.
+		TSet<FGuid> RelevantLights;
+		TSet<FGuid> PossiblyIrrelevantLights;
+		for (auto& MappingInfo : CachedMappings)
+		{
+			for (const ULightComponent* Light : MappingInfo.Mapping->Mesh->RelevantLights)
+			{
+				// Check if the light is stored in the light-map.
+				const bool bIsInLightMap = LODData[0].LightMap && LODData[0].LightMap->LightGuids.Contains(Light->LightGuid);
+
+				// Check if the light is stored in the shadow-map.
+				const bool bIsInShadowMap = LODData[0].ShadowMap && LODData[0].ShadowMap->LightGuids.Contains(Light->LightGuid);
+
+				// If the light isn't already relevant to another mapping, add it to the potentially irrelevant list
+				if (!bIsInLightMap && !bIsInShadowMap && !RelevantLights.Contains(Light->LightGuid))
+				{
+					PossiblyIrrelevantLights.Add(Light->LightGuid);
+				}
+
+				// Light is relevant
+				if (bIsInLightMap || bIsInShadowMap)
+				{
+					RelevantLights.Add(Light->LightGuid);
+					PossiblyIrrelevantLights.Remove(Light->LightGuid);
+				}
+			}
+		}
+
+		IrrelevantLights = PossiblyIrrelevantLights.Array();
+
+		bHasCachedStaticLighting = true;
+		MarkPackageDirty();
 	}
 }
 #endif
@@ -1526,28 +1780,6 @@ struct FComponentInstancedLightmapData
 	TArray< FComponentInstanceSharingData > SharingData;
 };
 
-/**
- * Struct that controls what we use to determine compatible components
- */
-struct FValidCombination
-{
-	/** An optional key for marking components as compatible (eg proc buildings only allow meshes on a single face to join) */
-	int32 JoinKey;
-
-	/** Different meshes are never compatible */
-	UStaticMesh* Mesh;
-
-	friend bool operator==(const FValidCombination& A, const FValidCombination& B)
-	{
-		return A.JoinKey == B.JoinKey && A.Mesh == B.Mesh;
-	}
-
-	friend uint32 GetTypeHash(const FValidCombination& Combo)
-	{
-		return (uint32)(PTRINT)Combo.Mesh * Combo.JoinKey;
-	}
-};
-
 void UInstancedStaticMeshComponent::GetLightAndShadowMapMemoryUsage( int32& LightMapMemoryUsage, int32& ShadowMapMemoryUsage ) const
 {
 	Super::GetLightAndShadowMapMemoryUsage(LightMapMemoryUsage, ShadowMapMemoryUsage);
diff --git a/Engine/Source/Runtime/Engine/Private/LightMap.cpp b/Engine/Source/Runtime/Engine/Private/LightMap.cpp
--- a/Engine/Source/Runtime/Engine/Private/LightMap.cpp
+++ b/Engine/Source/Runtime/Engine/Private/LightMap.cpp
@@ -216,36 +216,50 @@ struct FLightMapAllocation
 		MappedRect.Min.Y = 0;
 		MappedRect.Max.X = 0;
 		MappedRect.Max.Y = 0;
-		Primitive = NULL;
-		LightmapFlags = LMF_None;
+		Primitive = nullptr;
+		InstanceIndex = INDEX_NONE;
 		bSkipEncoding = false;
 	}
 
 	/**
 	 * Copy construct from FQuantizedLightmapData
 	 */
-	FLightMapAllocation(const FQuantizedLightmapData* QuantizedData)
-		: TotalSizeX(QuantizedData->SizeX)
-		, TotalSizeY(QuantizedData->SizeY)
-		, bHasSkyShadowing(QuantizedData->bHasSkyShadowing)
-		, RawData(QuantizedData->Data)
+	FLightMapAllocation(FQuantizedLightmapData&& QuantizedData)
+		: TotalSizeX(QuantizedData.SizeX)
+		, TotalSizeY(QuantizedData.SizeY)
+		, bHasSkyShadowing(QuantizedData.bHasSkyShadowing)
+		, RawData(MoveTemp(QuantizedData.Data))
 	{
-		FMemory::Memcpy(Scale, QuantizedData->Scale, sizeof(Scale));
-		FMemory::Memcpy(Add, QuantizedData->Add, sizeof(Add));
+		FMemory::Memcpy(Scale, QuantizedData.Scale, sizeof(Scale));
+		FMemory::Memcpy(Add, QuantizedData.Add, sizeof(Add));
 		PaddingType = GAllowLightmapPadding ? LMPT_NormalPadding : LMPT_NoPadding;
 		MappedRect.Min.X = 0;
 		MappedRect.Min.Y = 0;
 		MappedRect.Max.X = TotalSizeX;
 		MappedRect.Max.Y = TotalSizeY;
-		Primitive = NULL;
-		LightmapFlags = LMF_None;
+		Primitive = nullptr;
+		InstanceIndex = INDEX_NONE;
 		bSkipEncoding = false;
 	}
 
+	// Called after the lightmap is encoded
+	void PostEncode()
+	{
+		if (InstanceIndex >= 0)
+		{
+			UInstancedStaticMeshComponent* Component = dynamic_cast<UInstancedStaticMeshComponent*>(Primitive);
+			check(Component);
+
+			// TODO: We currently only support one LOD of static lighting in foliage
+			// Need to create per-LOD instance data to fix that
+			Component->PerInstanceSMData[InstanceIndex].LightmapUVBias = LightMap->GetCoordinateBias();
+		}
+	}
+
 	FLightMap2D*	LightMap;
-	UObject*		Outer;
 
 	UObject*		Primitive;
+	int32			InstanceIndex;
 
 	/** Upper-left X-coordinate in the texture atlas. */
 	int32				OffsetX;
@@ -260,17 +274,69 @@ struct FLightMapAllocation
 	bool			bDebug;
 	bool			bHasSkyShadowing;
 	ELightMapPaddingType			PaddingType;
-	ELightMapFlags					LightmapFlags;
 	TArray<FLightMapCoefficients>	RawData;
 	float							Scale[NUM_STORED_LIGHTMAP_COEF][4];
 	float							Add[NUM_STORED_LIGHTMAP_COEF][4];
-	/** Bounds of the primitive that the mapping is applied to. */
-	FBoxSphereBounds				Bounds;
 	/** True if we can skip encoding this allocation because it's similar enough to an existing
 	    allocation at the same offset */
 	bool bSkipEncoding;
 };
 
+struct FLightMapAllocationGroup
+{
+	FLightMapAllocationGroup()
+		: Outer(nullptr)
+		, LightmapFlags(LMF_None)
+		, Bounds(ForceInit)
+		, TotalTexels(0)
+	{
+	}
+
+#if PLATFORM_COMPILER_HAS_DEFAULTED_FUNCTIONS
+	FLightMapAllocationGroup(FLightMapAllocationGroup&&) = default;
+	FLightMapAllocationGroup& operator=(FLightMapAllocationGroup&&) = default;
+
+	FLightMapAllocationGroup(const FLightMapAllocationGroup&) = delete; // non-copy-able
+	FLightMapAllocationGroup& operator=(const FLightMapAllocationGroup&) = delete;
+#else
+	FLightMapAllocationGroup(FLightMapAllocationGroup&& other)
+		: Allocations(MoveTemp(other.Allocations))
+		, Outer(other.Outer)
+		, LightmapFlags(other.LightmapFlags)
+		, Bounds(other.Bounds)
+		, TotalTexels(other.TotalTexels)
+	{
+	}
+
+	FLightMapAllocationGroup& operator=(FLightMapAllocationGroup&& other)
+	{
+		Allocations = MoveTemp(other.Allocations);
+		Outer = other.Outer;
+		LightmapFlags = other.LightmapFlags;
+		Bounds = other.Bounds;
+		TotalTexels = other.TotalTexels;
+
+		return *this;
+	}
+
+private:
+	FLightMapAllocationGroup(const FLightMapAllocationGroup&); // non-copy-able
+	FLightMapAllocationGroup& operator=(const FLightMapAllocationGroup&);
+public:
+#endif
+
+	TArray<TUniquePtr<FLightMapAllocation>, TInlineAllocator<1>> Allocations;
+
+	UObject*		Outer;
+	ELightMapFlags	LightmapFlags;
+
+	// Bounds of the primitive that the mapping is applied to
+	// Used to group nearby allocations into the same lightmap texture
+	FBoxSphereBounds Bounds;
+
+	int32			TotalTexels;
+};
+
 enum ELightmapTextureType
 {
 	LTT_Coefficients = NUM_STORED_LIGHTMAP_COEF,
@@ -325,15 +391,20 @@ struct FLightMapPendingTexture : public FTextureLayout
 	/** Lightmap streaming flags that must match in order to be stored in this texture.	*/
 	ELightMapFlags					LightmapFlags;
 
+	// Optimization to quickly test if a new allocation won't fit
+	// Primarily of benefit to instanced mesh lightmaps
+	int32							UnallocatedTexels;
+
 	int32							NumOutstandingAsyncTasks;
 
 	FLightMapPendingTexture(UWorld* InWorld, uint32 InSizeX,uint32 InSizeY)
-		:	FTextureLayout(4,4,InSizeX,InSizeY,true) // Min size is 4x4 in case of block compression.
-		,	SkyOcclusionTexture(NULL)
-		,	OwningWorld(InWorld)
-		,	Bounds(FBox(0))
-		,	LightmapFlags( LMF_None )
-		,	NumOutstandingAsyncTasks(0)
+		: FTextureLayout(4, 4, InSizeX, InSizeY, /* PowerOfTwo */ true, /* AlignByFour */ true) // Min size is 4x4 in case of block compression.
+		, SkyOcclusionTexture(nullptr)
+		, OwningWorld(InWorld)
+		, Bounds(FBox(0))
+		, LightmapFlags(LMF_None)
+		, UnallocatedTexels(InSizeX * InSizeY)
+		, NumOutstandingAsyncTasks(0)
 	{}
 
 	~FLightMapPendingTexture()
@@ -367,7 +438,7 @@ struct FLightMapPendingTexture : public FTextureLayout
 	 *
 	 * @return	True if succeeded, false otherwise.
 	 */
-	bool AddElement(FLightMapAllocation& Allocation, const bool bForceIntoThisTexture = false );
+	bool AddElement(FLightMapAllocationGroup& AllocationGroup, const bool bForceIntoThisTexture = false);
 
 private:
 	FName GetLightmapName(int32 TextureIndex, int32 CoefficientIndex);
@@ -455,9 +526,6 @@ void FLightMapPendingTexture::FinishEncoding( ELightmapTextureType TextureType )
 	}
 }
 
-/** Whether to try to pack lightmaps/shadowmaps into the same texture. */
-bool GGroupComponentLightmaps = true;
-
 /**
  * Finds a free area in the texture large enough to contain a surface with the given size.
  * If a large enough area is found, it is marked as in use, the output parameters OutBaseX and OutBaseY are
@@ -467,59 +535,92 @@ bool GGroupComponentLightmaps = true;
  * If the allocation succeeded, Allocation.OffsetX and Allocation.OffsetY will be set to the upper-left corner
  * of the allocated area.
  *
- * @param Allocation	Lightmap allocation to try to fit
+ * @param Allocation	Lightmap allocation group to try to fit
  * @param bForceIntoThisTexture	True if we should ignore distance and other factors when considering whether the mapping should be packed into this texture
  *
  * @return	True if succeeded, false otherwise.
  */
-bool FLightMapPendingTexture::AddElement(FLightMapAllocation& Allocation, bool bForceIntoThisTexture )
+bool FLightMapPendingTexture::AddElement(FLightMapAllocationGroup& AllocationGroup, bool bForceIntoThisTexture)
 {
-	if( !bForceIntoThisTexture )
+	if (!bForceIntoThisTexture)
 	{
 		// Don't pack lightmaps from different packages into the same texture.
-		if ( Outer != Allocation.Outer )
+		if (Outer != AllocationGroup.Outer)
 		{
 			return false;
 		}
 	}
 
-	const FBoxSphereBounds NewBounds = Bounds + Allocation.Bounds;
+	// This is a rough test, passing it doesn't guarantee it'll fit
+	// But failing it does guarantee that it _won't_ fit
+	if (UnallocatedTexels < AllocationGroup.TotalTexels)
+	{
+		return false;
+	}
+
 	const bool bEmptyTexture = Allocations.Num() == 0;
+	const FBoxSphereBounds NewBounds = bEmptyTexture ? AllocationGroup.Bounds : Bounds + AllocationGroup.Bounds;
 
-	if ( !bEmptyTexture && !bForceIntoThisTexture )
+	if (!bEmptyTexture && !bForceIntoThisTexture)
 	{
 		// Don't mix streaming lightmaps with non-streaming lightmaps.
-		if ( (LightmapFlags & LMF_Streamed) != (Allocation.LightmapFlags & LMF_Streamed) )
+		if ((LightmapFlags & LMF_Streamed) != (AllocationGroup.LightmapFlags & LMF_Streamed))
 		{
 			return false;
 		}
 
 		// Is this a streaming lightmap?
-		if ( LightmapFlags & LMF_Streamed )
+		if (LightmapFlags & LMF_Streamed)
 		{
 			bool bPerformDistanceCheck = true;
 
 			// Don't pack together lightmaps that are too far apart
-			if ( bPerformDistanceCheck && NewBounds.SphereRadius > GMaxLightmapRadius && NewBounds.SphereRadius > (Bounds.SphereRadius + SMALL_NUMBER) )
+			if (bPerformDistanceCheck && NewBounds.SphereRadius > GMaxLightmapRadius && NewBounds.SphereRadius > (Bounds.SphereRadius + SMALL_NUMBER))
 			{
 				return false;
 			}
 		}
 	}
 
-	uint32 BaseX = 0;
-	uint32 BaseY = 0;
-	if( !FTextureLayout::AddElement( BaseX, BaseY, Allocation.MappedRect.Width(), Allocation.MappedRect.Height() ) )
+	int32 NewUnallocatedTexels = UnallocatedTexels;
+
+	int32 iAllocation = 0;
+	for (; iAllocation < AllocationGroup.Allocations.Num(); ++iAllocation)
 	{
+		auto& Allocation = AllocationGroup.Allocations[iAllocation];
+		uint32 BaseX, BaseY;
+		const uint32 SizeX = Allocation->MappedRect.Width();
+		const uint32 SizeY = Allocation->MappedRect.Height();
+		if (FTextureLayout::AddElement(BaseX, BaseY, SizeX, SizeY))
+		{
+			Allocation->OffsetX = BaseX;
+			Allocation->OffsetY = BaseY;
+
+			// Assumes bAlignByFour
+			NewUnallocatedTexels -= ((SizeX + 3) & ~3) * ((SizeY + 3) & ~3);
+		}
+		else
+		{
+			// failed to add all elements to the texture
+			break;
+		}
+	}
+	if (iAllocation < AllocationGroup.Allocations.Num())
+	{
+		// failed to add all elements to the texture
+		// remove the ones added so far to restore our original state
+		for (--iAllocation; iAllocation >= 0; --iAllocation)
+		{
+			auto& Allocation = AllocationGroup.Allocations[iAllocation];
+			const uint32 SizeX = Allocation->MappedRect.Width();
+			const uint32 SizeY = Allocation->MappedRect.Height();
+			verify(FTextureLayout::RemoveElement(Allocation->OffsetX, Allocation->OffsetY, SizeX, SizeY));
+		}
 		return false;
 	}
 
-	// UE_LOG(LogLightMap, Warning, TEXT( "LightMapPacking: New Allocation (%i texels)" ), Allocation.MappedRect.Area() );
-
-	// Save the position the light-maps (the Allocation.MappedRect portion) in the texture atlas.
-	Allocation.OffsetX = BaseX;
-	Allocation.OffsetY = BaseY;
-	Bounds = bEmptyTexture ? Allocation.Bounds : NewBounds;
+	Bounds = NewBounds;
+	UnallocatedTexels = NewUnallocatedTexels;
 
 	return true;
 }
@@ -999,9 +1100,11 @@ void FLightMapPendingTexture::StartEncoding()
 		FVector2D Scale((float)PaddedSizeX / (float)GetSizeX(), (float)PaddedSizeY / (float)GetSizeY());
 		FVector2D Bias((float)BaseX / (float)GetSizeX(), (float)BaseY / (float)GetSizeY());
 
-		// let the lightmap finish up after being encoded, setting the scale/bias, and returning if it can be deleted
+		// Set the scale/bias of the lightmap
 		check( Allocation->LightMap );
-		Allocation->LightMap->FinalizeEncoding(Scale, Bias, Textures[0]);
+		Allocation->LightMap->CoordinateScale = Scale;
+		Allocation->LightMap->CoordinateBias = Bias;
+		Allocation->PostEncode();
 
 		// Free the light-map's raw data.
 		Allocation->RawData.Empty();
@@ -1050,7 +1153,7 @@ FName FLightMapPendingTexture::GetSkyOcclusionTextureName(int32 TextureIndex)
 }
 
 /** The light-maps which have not yet been encoded into textures. */
-static TIndirectArray<FLightMapAllocation> PendingLightMaps;
+static TArray<FLightMapAllocationGroup> PendingLightMaps;
 static uint32 PendingLightMapSize = 0;
 
 #endif // WITH_EDITOR
@@ -1058,55 +1161,59 @@ static uint32 PendingLightMapSize = 0;
 /** If true, update the status when encoding light maps */
 bool FLightMap2D::bUpdateStatus = true;
 
-FLightMap2D* FLightMap2D::AllocateLightMap(UObject* LightMapOuter, FQuantizedLightmapData*& SourceQuantizedData, const FBoxSphereBounds& Bounds, ELightMapPaddingType InPaddingType, ELightMapFlags InLightmapFlags )
+FLightMap2D* FLightMap2D::AllocateLightMap(UObject* LightMapOuter, FQuantizedLightmapData*& SourceQuantizedData, const FBoxSphereBounds& Bounds, ELightMapPaddingType InPaddingType, ELightMapFlags InLightmapFlags)
 {
 	// If the light-map has no lights in it, return NULL.
-	if(!SourceQuantizedData)
+	if (!SourceQuantizedData)
 	{
 		return NULL;
 	}
 
 #if WITH_EDITOR
-
-	FLightMapAllocation* Allocation = new(PendingLightMaps) FLightMapAllocation(SourceQuantizedData);
-
-	Allocation->Outer		= LightMapOuter->GetOutermost();
-	Allocation->PaddingType	= InPaddingType;
-	Allocation->LightmapFlags = InLightmapFlags;
-	Allocation->Bounds		= Bounds;
-	Allocation->Primitive	= LightMapOuter;
-	if ( !GAllowStreamingLightmaps )
+	FLightMapAllocationGroup AllocationGroup;
+	AllocationGroup.Outer = LightMapOuter->GetOutermost();
+	AllocationGroup.LightmapFlags = InLightmapFlags;
+	AllocationGroup.Bounds = Bounds;
+	if (!GAllowStreamingLightmaps)
 	{
-		Allocation->LightmapFlags = ELightMapFlags( Allocation->LightmapFlags & ~LMF_Streamed );
+		AllocationGroup.LightmapFlags = ELightMapFlags(AllocationGroup.LightmapFlags & ~LMF_Streamed);
 	}
 
 	// Create a new light-map.
 	FLightMap2D* LightMap = new FLightMap2D(SourceQuantizedData->LightGuids);
-	Allocation->LightMap = LightMap;
 
-#if ALLOW_LIGHTMAP_SAMPLE_DEBUGGING && WITH_EDITOR
-	// Detect if this allocation belongs to the texture mapping that was being debugged
-	//@todo - this only works for mappings that can be uniquely identified by a single component, BSP for example does not work.
-	if (GCurrentSelectedLightmapSample.Component && GCurrentSelectedLightmapSample.Component == LightMapOuter)
+	// Create allocation and add it to the group
 	{
-		GCurrentSelectedLightmapSample.Lightmap = LightMap;
-		Allocation->bDebug = true;
-	}
-	else
-	{
-		Allocation->bDebug = false;
+		TUniquePtr<FLightMapAllocation> Allocation = MakeUnique<FLightMapAllocation>(MoveTemp(*SourceQuantizedData));
+		Allocation->PaddingType = InPaddingType;
+		Allocation->LightMap = LightMap;
+		Allocation->Primitive = LightMapOuter;
+
+#if ALLOW_LIGHTMAP_SAMPLE_DEBUGGING && WITH_EDITOR
+		// Detect if this allocation belongs to the texture mapping that was being debugged
+		//@todo - this only works for mappings that can be uniquely identified by a single component, BSP for example does not work.
+		if (GCurrentSelectedLightmapSample.Component && GCurrentSelectedLightmapSample.Component == LightMapOuter)
+		{
+			GCurrentSelectedLightmapSample.Lightmap = LightMap;
+			Allocation->bDebug = true;
+		}
+		else
+		{
+			Allocation->bDebug = false;
 	}
 #endif
 
-	// Quantized data is no longer needed now that FLightMapAllocation has made a copy of it
-	if (SourceQuantizedData != NULL)
-	{
+		// SourceQuantizedData is no longer needed now that FLightMapAllocation has what it needs
 		delete SourceQuantizedData;
 		SourceQuantizedData = NULL;
+
+		// Track the size of pending light-maps.
+		PendingLightMapSize += ((Allocation->TotalSizeX + 3) & ~3) * ((Allocation->TotalSizeY + 3) & ~3);
+
+		AllocationGroup.Allocations.Add(MoveTemp(Allocation));
 	}
 
-	// Track the size of pending light-maps.
-	PendingLightMapSize += ((Allocation->TotalSizeX + 3) & ~3) * ((Allocation->TotalSizeY + 3) & ~3);
+	PendingLightMaps.Add(MoveTemp(AllocationGroup));
 
 	return LightMap;
 #else
@@ -1114,13 +1221,189 @@ FLightMap2D* FLightMap2D::AllocateLightMap(UObject* LightMapOuter, FQuantizedLig
 #endif // WITH_EDITOR
 }
 
+FLightMap2D* FLightMap2D::AllocateInstancedLightMap(UInstancedStaticMeshComponent* Component, TArray<TUniquePtr<FQuantizedLightmapData>> InstancedSourceQuantizedData,
+	const FBoxSphereBounds& Bounds, ELightMapPaddingType InPaddingType, ELightMapFlags InLightmapFlags)
+{
+#if WITH_EDITOR
+	check(InstancedSourceQuantizedData.Num() > 0);
+
+	// Verify all instance lightmaps are the same size
+	for (int32 InstanceIndex = 1; InstanceIndex < InstancedSourceQuantizedData.Num(); ++InstanceIndex)
+	{
+		auto& SourceQuantizedData = InstancedSourceQuantizedData[InstanceIndex];
+		check(SourceQuantizedData->SizeX == InstancedSourceQuantizedData[0]->SizeX);
+		check(SourceQuantizedData->SizeY == InstancedSourceQuantizedData[0]->SizeY);
+	}
+
+	// Requantize source data to the same quantization
+	// Most of the following code is cloned from UModel::ApplyStaticLighting(), possibly it can be shared in future?
+	// or removed, if instanced mesh components can be given per-instance lightmap unpack coefficients
+	float MinCoefficient[NUM_STORED_LIGHTMAP_COEF][4];
+	float MaxCoefficient[NUM_STORED_LIGHTMAP_COEF][4];
+	for (int32 CoefficientIndex = 0; CoefficientIndex < NUM_STORED_LIGHTMAP_COEF; CoefficientIndex += 2)
+	{
+		for (int32 ColorIndex = 0; ColorIndex < 4; ColorIndex++)
+		{
+			// Color
+			MinCoefficient[CoefficientIndex][ColorIndex] = FLT_MAX;
+			MaxCoefficient[CoefficientIndex][ColorIndex] = 0.0f;
+
+			// Direction
+			MinCoefficient[CoefficientIndex + 1][ColorIndex] = FLT_MAX;
+			MaxCoefficient[CoefficientIndex + 1][ColorIndex] = -FLT_MAX;
+		}
+	}
+
+	// first, we need to find the max scale for all mappings, and that will be the scale across all instances of this component
+	for (auto& SourceQuantizedData : InstancedSourceQuantizedData)
+	{
+		for (int32 CoefficientIndex = 0; CoefficientIndex < NUM_STORED_LIGHTMAP_COEF; CoefficientIndex++)
+		{
+			for (int32 ColorIndex = 0; ColorIndex < 4; ColorIndex++)
+			{
+				// The lightmap data for directional coefficients was packed in lightmass with
+				// Pack: y = (x - Min) / (Max - Min)
+				// We need to solve for Max and Min in order to combine BSP mappings into a lighting group.
+				// QuantizedData->Scale and QuantizedData->Add were calculated in lightmass in order to unpack the lightmap data like so
+				// Unpack: x = y * UnpackScale + UnpackAdd
+				// Which means
+				// Scale = Max - Min
+				// Add = Min
+				// Therefore we can solve for min and max using substitution
+
+				float Scale = SourceQuantizedData->Scale[CoefficientIndex][ColorIndex];
+				float Add = SourceQuantizedData->Add[CoefficientIndex][ColorIndex];
+				float Min = Add;
+				float Max = Scale + Add;
+
+				MinCoefficient[CoefficientIndex][ColorIndex] = FMath::Min(MinCoefficient[CoefficientIndex][ColorIndex], Min);
+				MaxCoefficient[CoefficientIndex][ColorIndex] = FMath::Max(MaxCoefficient[CoefficientIndex][ColorIndex], Max);
+			}
+		}
+	}
+
+	// Now calculate the new unpack scale and add based on the composite min and max
+	float Scale[NUM_STORED_LIGHTMAP_COEF][4];
+	float Add[NUM_STORED_LIGHTMAP_COEF][4];
+	for (int32 CoefficientIndex = 0; CoefficientIndex < NUM_STORED_LIGHTMAP_COEF; CoefficientIndex++)
+	{
+		for (int32 ColorIndex = 0; ColorIndex < 4; ColorIndex++)
+		{
+			Scale[CoefficientIndex][ColorIndex] = FMath::Max(MaxCoefficient[CoefficientIndex][ColorIndex] - MinCoefficient[CoefficientIndex][ColorIndex], DELTA);
+			Add[CoefficientIndex][ColorIndex] = MinCoefficient[CoefficientIndex][ColorIndex];
+		}
+	}
+
+	// perform requantization
+	for (auto& SourceQuantizedData : InstancedSourceQuantizedData)
+	{
+		for (uint32 Y = 0; Y < SourceQuantizedData->SizeY; Y++)
+		{
+			for (uint32 X = 0; X < SourceQuantizedData->SizeX; X++)
+			{
+				// get source from input, dest from the rectangular offset in the group
+				FLightMapCoefficients& LightmapSample = SourceQuantizedData->Data[Y * SourceQuantizedData->SizeX + X];
+
+				// Treat alpha special because of residual
+				{
+					// Decode LogL
+					float LogL = (float)LightmapSample.Coefficients[0][3] / 255.0f;
+					float Residual = (float)LightmapSample.Coefficients[1][3] / 255.0f;
+					LogL += (Residual - 0.5f) / 255.0f;
+					LogL = LogL * SourceQuantizedData->Scale[0][3] + SourceQuantizedData->Add[0][3];
+
+					// Encode LogL
+					LogL = (LogL - Add[0][3]) / Scale[0][3];
+					Residual = LogL * 255.0f - FMath::RoundToFloat(LogL * 255.0f) + 0.5f;
+
+					LightmapSample.Coefficients[0][3] = (uint8)FMath::Clamp<int32>(FMath::RoundToInt(LogL * 255.0f), 0, 255);
+					LightmapSample.Coefficients[1][3] = (uint8)FMath::Clamp<int32>(FMath::RoundToInt(Residual * 255.0f), 0, 255);
+				}
+
+				// go over each color coefficient and dequantize and requantize with new Scale/Add
+				for (int32 CoefficientIndex = 0; CoefficientIndex < NUM_STORED_LIGHTMAP_COEF; CoefficientIndex++)
+				{
+					// Don't touch alpha here
+					for (int32 ColorIndex = 0; ColorIndex < 3; ColorIndex++)
+					{
+						// dequantize it
+						float Dequantized = (float)LightmapSample.Coefficients[CoefficientIndex][ColorIndex] / 255.0f;
+						const float Exponent = CoefficientIndex == 0 ? 2.0f : 1.0f;
+						Dequantized = FMath::Pow(Dequantized, Exponent);
+
+						const float Unpacked = Dequantized * SourceQuantizedData->Scale[CoefficientIndex][ColorIndex] + SourceQuantizedData->Add[CoefficientIndex][ColorIndex];
+						const float Repacked = (Unpacked - Add[CoefficientIndex][ColorIndex]) / Scale[CoefficientIndex][ColorIndex];
+
+						// requantize it
+						LightmapSample.Coefficients[CoefficientIndex][ColorIndex] = (uint8)FMath::Clamp<int32>(FMath::RoundToInt(FMath::Pow(Repacked, 1.0f / Exponent) * 255.0f), 0, 255);
+					}
+				}
+			}
+		}
+	}
+
+	FLightMapAllocationGroup AllocationGroup = FLightMapAllocationGroup();
+	AllocationGroup.Outer = Component->GetOutermost();
+	AllocationGroup.LightmapFlags = InLightmapFlags;
+	AllocationGroup.Bounds = Bounds;
+	if (!GAllowStreamingLightmaps)
+	{
+		AllocationGroup.LightmapFlags = ELightMapFlags(AllocationGroup.LightmapFlags & ~LMF_Streamed);
+	}
+
+	FLightMap2D* BaseLightmap = nullptr;
+
+	for (int32 InstanceIndex = 0; InstanceIndex < InstancedSourceQuantizedData.Num(); ++InstanceIndex)
+	{
+		auto& SourceQuantizedData = InstancedSourceQuantizedData[InstanceIndex];
+
+		// Create a new light-map.
+		FLightMap2D* LightMap = new FLightMap2D(SourceQuantizedData->LightGuids);
+
+		if (InstanceIndex == 0)
+		{
+			BaseLightmap = LightMap;
+		}
+		else
+		{
+			// we need the base lightmap to contain all of the lights used by all lightmaps in the group
+			for (auto& LightGuid : SourceQuantizedData->LightGuids)
+			{
+				BaseLightmap->LightGuids.AddUnique(LightGuid);
+			}
+		}
+
+		TUniquePtr<FLightMapAllocation> Allocation = MakeUnique<FLightMapAllocation>(MoveTemp(*SourceQuantizedData));
+		Allocation->PaddingType = InPaddingType;
+		Allocation->LightMap = LightMap;
+		Allocation->Primitive = Component;
+		Allocation->InstanceIndex = InstanceIndex;
+
+		// SourceQuantizedData is no longer needed now that FLightMapAllocation has what it needs
+		SourceQuantizedData.Reset();
+
+		// Track the size of pending light-maps.
+		PendingLightMapSize += ((Allocation->TotalSizeX + 3) & ~3) * ((Allocation->TotalSizeY + 3) & ~3);
+
+		AllocationGroup.Allocations.Add(MoveTemp(Allocation));
+	}
+
+	PendingLightMaps.Add(MoveTemp(AllocationGroup));
+
+	return BaseLightmap;
+#else
+	return nullptr;
+#endif //WITH_EDITOR
+}
+
 #if WITH_EDITOR
 
 struct FCompareLightmaps
 {
-	FORCEINLINE bool operator()( const FLightMapAllocation& A, const FLightMapAllocation& B ) const
+	FORCEINLINE bool operator()(const FLightMapAllocationGroup& A, const FLightMapAllocationGroup& B) const
 	{
-		return FMath::Max(B.TotalSizeX,B.TotalSizeY) < FMath::Max(A.TotalSizeX,A.TotalSizeY);
+		// Order descending by total size of allocation
+		return A.TotalTexels > B.TotalTexels;
 	}
 };
 
@@ -1133,7 +1416,7 @@ struct FCompareLightmaps
 void FLightMap2D::EncodeTextures( UWorld* InWorld, bool bLightingSuccessful, bool bForceCompletion )
 {
 #if WITH_EDITOR
-	if ( bLightingSuccessful )
+	if (bLightingSuccessful)
 	{
 		GWarn->BeginSlowTask( NSLOCTEXT("LightMap2D", "BeginEncodingLightMapsTask", "Encoding light-maps"), false );
 		int32 PackedLightAndShadowMapTextureSize = InWorld->GetWorldSettings()->PackedLightAndShadowMapTextureSize;
@@ -1141,66 +1424,112 @@ void FLightMap2D::EncodeTextures( UWorld* InWorld, bool bLightingSuccessful, boo
 		// Reset the pending light-map size.
 		PendingLightMapSize = 0;
 
+		// Crop lightmaps if allowed
+		if (GAllowLightmapCropping)
+		{
+			for (FLightMapAllocationGroup& PendingGroup : PendingLightMaps)
+			{
+				if (!ensure(PendingGroup.Allocations.Num() >= 1))
+				{
+					continue;
+				}
+
+				// TODO: Crop all allocations in a group to the same size for instanced meshes
+				if (PendingGroup.Allocations.Num() == 1)
+				{
+					for (auto& Allocation : PendingGroup.Allocations)
+					{
+						CropUnmappedTexels(Allocation->RawData, Allocation->TotalSizeX, Allocation->TotalSizeY, Allocation->MappedRect);
+					}
+				}
+			}
+		}
+
+		// Calculate size of pending allocations for sorting
+		for (FLightMapAllocationGroup& PendingGroup : PendingLightMaps)
+		{
+			PendingGroup.TotalTexels = 0;
+			for (auto& Allocation : PendingGroup.Allocations)
+			{
+				// Assumes bAlignByFour
+				PendingGroup.TotalTexels += ((Allocation->MappedRect.Width() + 3) & ~3) * ((Allocation->MappedRect.Height() + 3) & ~3);
+			}
+		}
+
 		// Sort the light-maps in descending order by size.
 		Sort(PendingLightMaps.GetData(), PendingLightMaps.Num(), FCompareLightmaps());
 
 		// Allocate texture space for each light-map.
 		TArray<FLightMapPendingTexture*> PendingTextures;
 
-		for(int32 LightMapIndex = 0;LightMapIndex < PendingLightMaps.Num();LightMapIndex++)
+		for (FLightMapAllocationGroup& PendingGroup : PendingLightMaps)
 		{
-			FLightMapAllocation& Allocation = PendingLightMaps[LightMapIndex];
-			FLightMapPendingTexture* Texture = NULL;
+			if (!ensure(PendingGroup.Allocations.Num() >= 1))
+			{
+				continue;
+			}
 
-			if ( GAllowLightmapCropping )
+			int32 MaxWidth = 0;
+			int32 MaxHeight = 0;
+			for (auto& Allocation : PendingGroup.Allocations)
 			{
-				CropUnmappedTexels( Allocation.RawData, Allocation.TotalSizeX, Allocation.TotalSizeY, Allocation.MappedRect );
+				MaxWidth = FMath::Max(MaxWidth, Allocation->MappedRect.Width());
+				MaxHeight = FMath::Max(MaxHeight, Allocation->MappedRect.Height());
 			}
 
+			FLightMapPendingTexture* Texture = nullptr;
+
 			// Find an existing texture which the light-map can be stored in.
-			int32 FoundIndex = -1;
-			for(int32 TextureIndex = 0;TextureIndex < PendingTextures.Num();TextureIndex++)
+			// Lightmaps will always be 4-pixel aligned...
+			for (FLightMapPendingTexture* ExistingTexture : PendingTextures)
 			{
-				FLightMapPendingTexture* ExistingTexture = PendingTextures[TextureIndex];
-
-				// Lightmaps will always be 4-pixel aligned...
-				if ( ExistingTexture->AddElement( Allocation ) )
+				if (ExistingTexture->AddElement(PendingGroup))
 				{
 					Texture = ExistingTexture;
-					FoundIndex = TextureIndex;
 					break;
 				}
 			}
 
-			if(!Texture)
+			if (!Texture)
 			{
 				int32 NewTextureSizeX = PackedLightAndShadowMapTextureSize;
 				int32 NewTextureSizeY = PackedLightAndShadowMapTextureSize / 2;
-				if(Allocation.MappedRect.Width() > NewTextureSizeX || Allocation.MappedRect.Height() > NewTextureSizeY)
+
+				// Assumes identically-sized allocations, fit into the smallest 2x1 rectangle
+				const int32 AllocationCountX = FMath::CeilToInt(FMath::Sqrt(FMath::DivideAndRoundUp(PendingGroup.Allocations.Num() * 2 * MaxHeight, MaxWidth)));
+				const int32 AllocationCountY = FMath::DivideAndRoundUp(PendingGroup.Allocations.Num(), AllocationCountX);
+				const int32 AllocationSizeX = AllocationCountX * MaxWidth;
+				const int32 AllocationSizeY = AllocationCountY * MaxHeight;
+
+				if (AllocationSizeX > NewTextureSizeX || AllocationSizeY > NewTextureSizeY)
 				{
-					NewTextureSizeX = FMath::RoundUpToPowerOfTwo(Allocation.MappedRect.Width());
-					NewTextureSizeY = FMath::RoundUpToPowerOfTwo(Allocation.MappedRect.Height());
+					NewTextureSizeX = FMath::RoundUpToPowerOfTwo(AllocationSizeX);
+					NewTextureSizeY = FMath::RoundUpToPowerOfTwo(AllocationSizeY);
 
 					// Force 2:1 aspect
-					NewTextureSizeX = FMath::Max( NewTextureSizeX, NewTextureSizeY * 2 );
-					NewTextureSizeY = FMath::Max( NewTextureSizeY, NewTextureSizeX / 2 );
+					NewTextureSizeX = FMath::Max(NewTextureSizeX, NewTextureSizeY * 2);
+					NewTextureSizeY = FMath::Max(NewTextureSizeY, NewTextureSizeX / 2);
 				}
 
 				// If there is no existing appropriate texture, create a new one.
-				Texture = new FLightMapPendingTexture( InWorld, NewTextureSizeX, NewTextureSizeY );
-				PendingTextures.Add( Texture );
-				Texture->Outer = Allocation.Outer;
-				Texture->Bounds = Allocation.Bounds;
-				Texture->LightmapFlags = Allocation.LightmapFlags;
-				verify( Texture->AddElement( Allocation ) );
-				FoundIndex = PendingTextures.Num() - 1;
+				Texture = new FLightMapPendingTexture(InWorld, NewTextureSizeX, NewTextureSizeY);
+				PendingTextures.Add(Texture);
+				Texture->Outer = PendingGroup.Outer;
+				Texture->Bounds = PendingGroup.Bounds;
+				Texture->LightmapFlags = PendingGroup.LightmapFlags;
+				verify(Texture->AddElement(PendingGroup));
 			}
 
-			Texture->Allocations.Add(&Allocation);
+			// Give the texture ownership of the allocations
+			for (auto& Allocation : PendingGroup.Allocations)
+			{
+				Texture->Allocations.Add(Allocation.Release());
+			}
 		}
+		PendingLightMaps.Empty();
 
 		// Encode all the pending textures.
-		for(int32 TextureIndex = 0; TextureIndex < PendingTextures.Num(); TextureIndex++)
+		for (int32 TextureIndex = 0; TextureIndex < PendingTextures.Num(); TextureIndex++)
 		{
 			if (bUpdateStatus && (TextureIndex % 20) == 0)
 			{
@@ -1209,17 +1538,14 @@ void FLightMap2D::EncodeTextures( UWorld* InWorld, bool bLightingSuccessful, boo
 			FLightMapPendingTexture* PendingTexture = PendingTextures[TextureIndex];
 			PendingTexture->StartEncoding();
 		}
-
 		PendingTextures.Empty();
-		PendingLightMaps.Empty();
 
-		if ( bForceCompletion )
+		if (bForceCompletion)
 		{
 			// Block until there are 0 unfinished tasks, making sure all compression has completed.
-			FLightMapPendingTexture::FinishCompletedTasks( 0 );
+			FLightMapPendingTexture::FinishCompletedTasks(0);
 		}
-	
-		
+
 		// End the encoding lighmaps slow task
 		GWarn->EndSlowTask();
 		
@@ -1280,22 +1606,6 @@ struct FLegacyLightMapTextureInfo
 	}
 };
 
-/**
- * Finalizes the lightmap after encodeing, including setting the UV scale/bias for this lightmap 
- * inside the larger UTexture2D that this lightmap is in
- * 
- * @param Scale UV scale (size)
- * @param Bias UV Bias (offset)
- * @param ALightmapTexture One of the lightmap textures that this lightmap was put into
- *
- * @return true if the lightmap should be deleted after this call
- */
-void FLightMap2D::FinalizeEncoding(const FVector2D& Scale, const FVector2D& Bias, UTexture2D* ALightmapTexture)
-{
-	CoordinateScale = Scale;
-	CoordinateBias = Bias;
-}
-
 void FLightMap2D::AddReferencedObjects( FReferenceCollector& Collector )
 {
 	Collector.AddReferencedObject(Textures[0]);
diff --git a/Engine/Source/Runtime/Engine/Private/ShadowMap.cpp b/Engine/Source/Runtime/Engine/Private/ShadowMap.cpp
--- a/Engine/Source/Runtime/Engine/Private/ShadowMap.cpp
+++ b/Engine/Source/Runtime/Engine/Private/ShadowMap.cpp
@@ -54,7 +54,10 @@ void FShadowMap::FinishCleanup()
 struct FShadowMapAllocation
 {
 	FShadowMap2D*				ShadowMap;
-	UObject*					TextureOuter;
+
+	UObject*					Primitive;
+	int32						InstanceIndex;
+
 	/** Upper-left X-coordinate in the texture atlas. */
 	int32						OffsetX;
 	/** Upper-left Y-coordinate in the texture atlas. */
@@ -69,11 +72,6 @@ struct FShadowMapAllocation
 
 	TMap<ULightComponent*, TArray<FQuantizedSignedDistanceFieldShadowSample> > ShadowMapData;
 
-	/** Bounds of the primitive that the mapping is applied to. */
-	FBoxSphereBounds	Bounds;
-	/** Bit-field with shadowmap flags. */
-	EShadowMapFlags		ShadowmapFlags;
-
 	FShadowMapAllocation()
 	{
 		PaddingType = GAllowLightmapPadding ? LMPT_NormalPadding : LMPT_NoPadding;
@@ -81,7 +79,78 @@ struct FShadowMapAllocation
 		MappedRect.Min.Y = 0;
 		MappedRect.Max.X = 0;
 		MappedRect.Max.Y = 0;
+		Primitive = nullptr;
+		InstanceIndex = INDEX_NONE;
 	}
+
+	// Called after the shadowmap is encoded
+	void PostEncode()
+	{
+		if (InstanceIndex >= 0)
+		{
+			UInstancedStaticMeshComponent* Component = dynamic_cast<UInstancedStaticMeshComponent*>(Primitive);
+			check(Component);
+
+			// TODO: We currently only support one LOD of static lighting in foliage
+			// Need to create per-LOD instance data to fix that
+			Component->PerInstanceSMData[InstanceIndex].ShadowmapUVBias = ShadowMap->GetCoordinateBias();
+		}
+	}
+};
+
+struct FShadowMapAllocationGroup
+{
+	FShadowMapAllocationGroup()
+		: TextureOuter(nullptr)
+		, Bounds(ForceInit)
+		, ShadowmapFlags(SMF_None)
+		, TotalTexels(0)
+	{
+	}
+
+#if PLATFORM_COMPILER_HAS_DEFAULTED_FUNCTIONS
+	FShadowMapAllocationGroup(FShadowMapAllocationGroup&&) = default;
+	FShadowMapAllocationGroup& operator=(FShadowMapAllocationGroup&&) = default;
+
+	FShadowMapAllocationGroup(const FShadowMapAllocationGroup&) = delete; // non-copy-able
+	FShadowMapAllocationGroup& operator=(const FShadowMapAllocationGroup&) = delete;
+#else
+	FShadowMapAllocationGroup(FShadowMapAllocationGroup&& other)
+		: Allocations(MoveTemp(other.Allocations))
+		, TextureOuter(other.TextureOuter)
+		, Bounds(other.Bounds)
+		, ShadowmapFlags(other.ShadowmapFlags)
+		, TotalTexels(other.TotalTexels)
+	{
+	}
+
+	FShadowMapAllocationGroup& operator=(FShadowMapAllocationGroup&& other)
+	{
+		Allocations = MoveTemp(other.Allocations);
+		TextureOuter = other.TextureOuter;
+		Bounds = other.Bounds;
+		ShadowmapFlags = other.ShadowmapFlags;
+		TotalTexels = other.TotalTexels;
+
+		return *this;
+	}
+
+private:
+	FShadowMapAllocationGroup(const FShadowMapAllocationGroup&); // non-copy-able
+	FShadowMapAllocationGroup& operator=(const FShadowMapAllocationGroup&);
+public:
+#endif
+
+	TArray<TUniquePtr<FShadowMapAllocation>, TInlineAllocator<1>> Allocations;
+
+	UObject*			TextureOuter;
+
+	/** Bounds of the primitive that the mapping is applied to. */
+	FBoxSphereBounds	Bounds;
+	/** Bit-field with shadowmap flags. */
+	EShadowMapFlags		ShadowmapFlags;
+
+	int32				TotalTexels;
 };
 
 /** Largest boundingsphere radius to use when packing shadowmaps into a texture atlas. */
@@ -101,14 +170,22 @@ struct FShadowMapPendingTexture : FTextureLayout
 	/** Bit-field with shadowmap flags that are shared among all shadowmaps in this texture. */
 	EShadowMapFlags		ShadowmapFlags;
 
+	// Optimization to quickly test if a new allocation won't fit
+	// Primarily of benefit to instanced mesh shadowmaps
+	int32				UnallocatedTexels;
+
 	/**
 	 * Minimal initialization constructor.
 	 */
-	FShadowMapPendingTexture(uint32 InSizeX,uint32 InSizeY):
-		FTextureLayout(1,1,InSizeX,InSizeY,true)
-	{}
+	FShadowMapPendingTexture(uint32 InSizeX,uint32 InSizeY)
+		: FTextureLayout(4, 4, InSizeX, InSizeY, /* PowerOfTwo */ true, /* AlignByFour */ true) // Min size is 4x4 in case of block compression.
+		, Bounds(FBox(0))
+		, ShadowmapFlags(SMF_None)
+		, UnallocatedTexels(InSizeX * InSizeY)
+	{
+	}
 
-	bool AddElement(FShadowMapAllocation& Allocation, const bool bForceIntoThisTexture = false);
+	bool AddElement(FShadowMapAllocationGroup& AllocationGroup, const bool bForceIntoThisTexture = false);
 
 	/**
 	 * Begin encoding the textures
@@ -118,57 +195,87 @@ struct FShadowMapPendingTexture : FTextureLayout
 	void StartEncoding(UWorld* InWorld);
 };
 
-bool FShadowMapPendingTexture::AddElement(FShadowMapAllocation& Allocation, const bool bForceIntoThisTexture)
+bool FShadowMapPendingTexture::AddElement(FShadowMapAllocationGroup& AllocationGroup, const bool bForceIntoThisTexture)
 {
-	if( !bForceIntoThisTexture )
+	if (!bForceIntoThisTexture)
 	{
-		// Must be in the same package
-		if ( Outer != Allocation.TextureOuter )
+		// Don't pack shadowmaps from different packages into the same texture.
+		if (Outer != AllocationGroup.TextureOuter)
 		{
 			return false;
 		}
 	}
 
-	const FBoxSphereBounds NewBounds = Bounds + Allocation.Bounds;
+	// This is a rough test, passing it doesn't guarantee it'll fit
+	// But failing it does guarantee that it _won't_ fit
+	if (UnallocatedTexels < AllocationGroup.TotalTexels)
+	{
+		return false;
+	}
+
 	const bool bEmptyTexture = Allocations.Num() == 0;
-	if ( !bEmptyTexture && !bForceIntoThisTexture )
+	const FBoxSphereBounds NewBounds = bEmptyTexture ? AllocationGroup.Bounds : Bounds + AllocationGroup.Bounds;
+
+	if (!bEmptyTexture && !bForceIntoThisTexture)
 	{
-		// Don't mix streaming lightmaps with non-streaming lightmaps.
-		if ( (ShadowmapFlags & SMF_Streamed) != (Allocation.ShadowmapFlags & SMF_Streamed) )
+		// Don't mix streaming shadowmaps with non-streaming shadowmaps.
+		if ((ShadowmapFlags & LMF_Streamed) != (AllocationGroup.ShadowmapFlags & LMF_Streamed))
 		{
 			return false;
 		}
 
-		// If this is a streaming shadowmap?
-		if ( ShadowmapFlags & SMF_Streamed )
+		// Is this a streaming shadowmap?
+		if (ShadowmapFlags & LMF_Streamed)
 		{
 			bool bPerformDistanceCheck = true;
 
-			// Don't pack together shadowmaps that are too far apart.
-			if ( bPerformDistanceCheck && NewBounds.SphereRadius > GMaxShadowmapRadius && NewBounds.SphereRadius > (Bounds.SphereRadius + SMALL_NUMBER) )
+			// Don't pack together shadowmaps that are too far apart
+			if (bPerformDistanceCheck && NewBounds.SphereRadius > GMaxShadowmapRadius && NewBounds.SphereRadius > (Bounds.SphereRadius + SMALL_NUMBER))
 			{
 				return false;
 			}
 		}
 	}
 
-	// 4-byte alignment, no more padding
-	uint32 PaddingBaseX = 0;
-	uint32 PaddingBaseY = 0;
+	int32 NewUnallocatedTexels = UnallocatedTexels;
 
-	// See if the new one will fit in the existing texture
-	if ( !FTextureLayout::AddElement(PaddingBaseX,PaddingBaseY,Allocation.MappedRect.Width(),Allocation.MappedRect.Height()) )
+	int32 iAllocation = 0;
+	for (; iAllocation < AllocationGroup.Allocations.Num(); ++iAllocation)
 	{
+		auto& Allocation = AllocationGroup.Allocations[iAllocation];
+		uint32 BaseX, BaseY;
+		const uint32 SizeX = Allocation->MappedRect.Width();
+		const uint32 SizeY = Allocation->MappedRect.Height();
+		if (FTextureLayout::AddElement(BaseX, BaseY, SizeX, SizeY))
+		{
+			Allocation->OffsetX = BaseX;
+			Allocation->OffsetY = BaseY;
+
+			// Assumes bAlignByFour
+			NewUnallocatedTexels -= ((SizeX + 3) & ~3) * ((SizeY + 3) & ~3);
+		}
+		else
+		{
+			// failed to add all elements to the texture
+			break;
+		}
+	}
+	if (iAllocation < AllocationGroup.Allocations.Num())
+	{
+		// failed to add all elements to the texture
+		// remove the ones added so far to restore our original state
+		for (--iAllocation; iAllocation >= 0; --iAllocation)
+		{
+			auto& Allocation = AllocationGroup.Allocations[iAllocation];
+			const uint32 SizeX = Allocation->MappedRect.Width();
+			const uint32 SizeY = Allocation->MappedRect.Height();
+			verify(FTextureLayout::RemoveElement(Allocation->OffsetX, Allocation->OffsetY, SizeX, SizeY));
+		}
 		return false;
 	}
 
-	// Position the shadow-maps in the middle of their padded space.
-	Allocation.OffsetX = PaddingBaseX;
-	Allocation.OffsetY = PaddingBaseY;
-	Bounds = bEmptyTexture ? Allocation.Bounds : NewBounds;
-
-	// Add the shadow-map to the list of shadow-maps allocated space in this texture.
-	Allocations.Add(&Allocation);
+	Bounds = NewBounds;
+	UnallocatedTexels = NewUnallocatedTexels;
 
 	return true;
 }
@@ -239,7 +346,7 @@ void FShadowMapPendingTexture::StartEncoding(UWorld* InWorld)
 	Texture->UpdateResource();
 }
 
-static TIndirectArray<FShadowMapAllocation> PendingShadowMaps;
+static TArray<FShadowMapAllocationGroup> PendingShadowMaps;
 static uint32 PendingShadowMapSize;
 /** If true, update the status when encoding light maps */
 bool FShadowMap2D::bUpdateStatus = true;
@@ -254,53 +361,53 @@ FShadowMap2D* FShadowMap2D::AllocateShadowMap(
 	EShadowMapFlags InShadowmapFlags)
 {
 #if WITH_EDITOR
-
 	check(ShadowMapData.Num() > 0);
 
+	FShadowMapAllocationGroup AllocationGroup;
+	AllocationGroup.TextureOuter = Outer->GetOutermost();
+	AllocationGroup.ShadowmapFlags = InShadowmapFlags;
+	AllocationGroup.Bounds = Bounds;
+	if (!GAllowStreamingLightmaps)
+	{
+		AllocationGroup.ShadowmapFlags = EShadowMapFlags(AllocationGroup.ShadowmapFlags & ~SMF_Streamed);
+	}
+
 	// Create a new shadow-map.
 	FShadowMap2D* ShadowMap = new FShadowMap2D(ShadowMapData);
 
+	// Calculate Shadowmap size
 	int32 SizeX = -1;
 	int32 SizeY = -1;
 	int32 LightIndex = 0;
-
-	for (TMap<ULightComponent*, FShadowMapData2D*>::TConstIterator It(ShadowMapData); It; ++It)
+	for (const auto& ShadowDataPair : ShadowMapData)
 	{
 		if (LightIndex == 0)
 		{
-			SizeX = It.Value()->GetSizeX();
-			SizeY = It.Value()->GetSizeY();
+			SizeX = ShadowDataPair.Value->GetSizeX();
+			SizeY = ShadowDataPair.Value->GetSizeY();
 		}
 		else
 		{
-			check(SizeX == It.Value()->GetSizeX() && SizeY == It.Value()->GetSizeY());
+			check(SizeX == ShadowDataPair.Value->GetSizeX() && SizeY == ShadowDataPair.Value->GetSizeY());
 		}
 
 		LightIndex++;
 	}
-
 	check(SizeX != -1 && SizeY != -1);
 
 	// Add a pending allocation for this shadow-map.
-	FShadowMapAllocation* Allocation = new(PendingShadowMaps) FShadowMapAllocation;
-	Allocation->ShadowMap		= ShadowMap;
-	Allocation->TextureOuter	= Outer->GetOutermost();
-	Allocation->TotalSizeX		= SizeX;
-	Allocation->TotalSizeY		= SizeY;
-	Allocation->MappedRect		= FIntRect( 0, 0, SizeX, SizeY );
-	Allocation->PaddingType		= InPaddingType;
-	Allocation->Bounds			= Bounds;
-	Allocation->ShadowmapFlags	= InShadowmapFlags;
-
-	if ( !GAllowStreamingLightmaps )
+	TUniquePtr<FShadowMapAllocation> Allocation = MakeUnique<FShadowMapAllocation>();
+	Allocation->PaddingType = InPaddingType;
+	Allocation->ShadowMap = ShadowMap;
+	Allocation->TotalSizeX = SizeX;
+	Allocation->TotalSizeY = SizeY;
+	Allocation->MappedRect = FIntRect(0, 0, SizeX, SizeY);
+	Allocation->PaddingType = InPaddingType;
+
+	for (const auto& ShadowDataPair : ShadowMapData)
 	{
-		Allocation->ShadowmapFlags = EShadowMapFlags( Allocation->ShadowmapFlags & ~SMF_Streamed );
-	}
-
-	for (TMap<ULightComponent*, FShadowMapData2D*>::TConstIterator It(ShadowMapData); It; ++It)
-	{
-		const FShadowMapData2D* RawData = It.Value();
-		TArray<FQuantizedSignedDistanceFieldShadowSample>& DistanceFieldShadowData = Allocation->ShadowMapData.Add(It.Key(), TArray<FQuantizedSignedDistanceFieldShadowSample>());
+		const FShadowMapData2D* RawData = ShadowDataPair.Value;
+		TArray<FQuantizedSignedDistanceFieldShadowSample>& DistanceFieldShadowData = Allocation->ShadowMapData.Add(ShadowDataPair.Key, TArray<FQuantizedSignedDistanceFieldShadowSample>());
 
 		switch (RawData->GetType())
 		{
@@ -317,12 +424,18 @@ FShadowMap2D* FShadowMap2D::AllocateShadowMap(
 
 		// Track the size of pending light-maps.
 		PendingShadowMapSize += Allocation->TotalSizeX * Allocation->TotalSizeY;
+		// Assumes bAlignByFour
+		AllocationGroup.TotalTexels += ((Allocation->MappedRect.Width() + 3) & ~3) * ((Allocation->MappedRect.Height() + 3) & ~3);
 	}
 
+	AllocationGroup.Allocations.Add(MoveTemp(Allocation));
+
+	PendingShadowMaps.Add(MoveTemp(AllocationGroup));
+
 	return ShadowMap;
 #else
-	return NULL;
-#endif 
+	return nullptr;
+#endif
 }
 
 FShadowMap2D::FShadowMap2D() :
@@ -352,6 +465,18 @@ FShadowMap2D::FShadowMap2D(const TMap<ULightComponent*,FShadowMapData2D*>& Shado
 	}
 }
 
+FShadowMap2D::FShadowMap2D(TArray<FGuid> LightGuids)
+	: FShadowMap(MoveTemp(LightGuids))
+	, Texture(nullptr)
+	, CoordinateScale(FVector2D(0, 0))
+	, CoordinateBias(FVector2D(0, 0))
+{
+	for (int Channel = 0; Channel < ARRAY_COUNT(bChannelValid); Channel++)
+	{
+		bChannelValid[Channel] = false;
+	}
+}
+
 void FShadowMap2D::AddReferencedObjects(FReferenceCollector& Collector)
 {
 	Collector.AddReferencedObject(Texture);
@@ -389,22 +514,141 @@ FShadowMapInteraction FShadowMap2D::GetInteraction() const
 }
 
 
+FShadowMap2D* FShadowMap2D::AllocateInstancedShadowMap(UInstancedStaticMeshComponent* Component, TArray<TMap<ULightComponent*, TUniquePtr<FShadowMapData2D>>> InstancedShadowMapData,
+	const FBoxSphereBounds& Bounds, ELightMapPaddingType InPaddingType, EShadowMapFlags InShadowmapFlags)
+{
+#if WITH_EDITOR
+	check(InstancedShadowMapData.Num() > 0);
+
+	// Verify all instance shadowmaps are the same size, and build complete list of shadow lights
+	int32 SizeX = -1;
+	int32 SizeY = -1;
+	TSet<ULightComponent*> AllLights;
+	for (auto& ShadowMapData : InstancedShadowMapData)
+	{
+		for (const auto& ShadowDataPair : ShadowMapData)
+		{
+			if (SizeX == -1)
+			{
+				SizeX = ShadowDataPair.Value->GetSizeX();
+				SizeY = ShadowDataPair.Value->GetSizeY();
+			}
+			else
+			{
+				check(ShadowDataPair.Value->GetSizeX() == SizeX);
+				check(ShadowDataPair.Value->GetSizeY() == SizeY);
+			}
+			AllLights.Add(ShadowDataPair.Key);
+		}
+	}
+
+	check(SizeX != -1 && SizeY != -1); // No valid shadowmaps
+
+	TArray<FGuid> LightGuids;
+	LightGuids.Reserve(AllLights.Num());
+	for (ULightComponent* Light : AllLights)
+	{
+		LightGuids.Add(Light->LightGuid);
+	}
+
+	// Unify all the shadow map data to contain the same lights in the same order
+	for (auto& ShadowMapData : InstancedShadowMapData)
+	{
+		for (ULightComponent* Light : AllLights)
+		{
+			if (!ShadowMapData.Contains(Light))
+			{
+				ShadowMapData.Add(Light, MakeUnique<FQuantizedShadowSignedDistanceFieldData2D>(SizeX, SizeY));
+			}
+		}
+	}
+
+	FShadowMapAllocationGroup AllocationGroup;
+	AllocationGroup.TextureOuter = Component->GetOutermost();
+	AllocationGroup.ShadowmapFlags = InShadowmapFlags;
+	AllocationGroup.Bounds = Bounds;
+	if (!GAllowStreamingLightmaps)
+	{
+		AllocationGroup.ShadowmapFlags = EShadowMapFlags(AllocationGroup.ShadowmapFlags & ~SMF_Streamed);
+	}
+
+	FShadowMap2D* BaseShadowmap = nullptr;
+
+	for (int32 InstanceIndex = 0; InstanceIndex < InstancedShadowMapData.Num(); ++InstanceIndex)
+	{
+		auto& ShadowMapData = InstancedShadowMapData[InstanceIndex];
+		check(ShadowMapData.Num() > 0);
+
+		// Create a new shadow-map.
+		FShadowMap2D* ShadowMap = new FShadowMap2D(LightGuids);
+
+		if (InstanceIndex == 0)
+		{
+			BaseShadowmap = ShadowMap;
+		}
+
+		// Add a pending allocation for this shadow-map.
+		TUniquePtr<FShadowMapAllocation> Allocation = MakeUnique<FShadowMapAllocation>();
+		Allocation->PaddingType = InPaddingType;
+		Allocation->ShadowMap = ShadowMap;
+		Allocation->TotalSizeX = SizeX;
+		Allocation->TotalSizeY = SizeY;
+		Allocation->MappedRect = FIntRect(0, 0, SizeX, SizeY);
+		Allocation->Primitive = Component;
+		Allocation->InstanceIndex = InstanceIndex;
+
+		for (auto& ShadowDataPair : ShadowMapData)
+		{
+			auto& RawData = ShadowDataPair.Value;
+			auto& DistanceFieldShadowData = Allocation->ShadowMapData.Add(ShadowDataPair.Key, TArray<FQuantizedSignedDistanceFieldShadowSample>());
+
+			switch (RawData->GetType())
+			{
+			case FShadowMapData2D::SHADOW_SIGNED_DISTANCE_FIELD_DATA:
+			case FShadowMapData2D::SHADOW_SIGNED_DISTANCE_FIELD_DATA_QUANTIZED:
+				// If the data is already quantized, this will just copy the data
+				RawData->Quantize(DistanceFieldShadowData);
+				break;
+			default:
+				check(0);
+			}
+
+			RawData.Reset();
+
+			// Track the size of pending light-maps.
+			PendingShadowMapSize += Allocation->TotalSizeX * Allocation->TotalSizeY;
+			// Assumes bAlignByFour
+			AllocationGroup.TotalTexels += ((Allocation->MappedRect.Width() + 3) & ~3) * ((Allocation->MappedRect.Height() + 3) & ~3);
+		}
+
+		AllocationGroup.Allocations.Add(MoveTemp(Allocation));
+	}
+
+	PendingShadowMaps.Add(MoveTemp(AllocationGroup));
+
+	return BaseShadowmap;
+#else
+	return nullptr;
+#endif
+}
+
+
 #if WITH_EDITOR
 
 struct FCompareShadowMaps
 {
-	FORCEINLINE bool operator()( const FShadowMapAllocation& A, const FShadowMapAllocation& B ) const
+	FORCEINLINE bool operator()(const FShadowMapAllocationGroup& A, const FShadowMapAllocationGroup& B) const
 	{
-		return FMath::Max(B.TotalSizeX, B.TotalSizeY) < FMath::Max(A.TotalSizeX, A.TotalSizeY);
+		return A.TotalTexels > B.TotalTexels;
 	}
 };
 
 /**
  * Executes all pending shadow-map encoding requests.
  * @param	InWorld				World in which the textures exist
- * @param	bLightingSuccessful	Whether the lighting build was successful or not.	 
+ * @param	bLightingSuccessful	Whether the lighting build was successful or not.
  */
-void FShadowMap2D::EncodeTextures(UWorld* InWorld , bool bLightingSuccessful )
+void FShadowMap2D::EncodeTextures(UWorld* InWorld , bool bLightingSuccessful)
 {
 	if ( bLightingSuccessful )
 	{
@@ -416,49 +660,73 @@ void FShadowMap2D::EncodeTextures(UWorld* InWorld , bool bLightingSuccessful )
 
 		Sort(PendingShadowMaps.GetData(), PendingShadowMaps.Num(), FCompareShadowMaps());
 
-		// Allocate texture space for each light-map.
+		// Allocate texture space for each shadow-map.
 		TIndirectArray<FShadowMapPendingTexture> PendingTextures;
 
-		for(int32 ShadowMapIndex = 0;ShadowMapIndex < PendingShadowMaps.Num();ShadowMapIndex++)
+		for (FShadowMapAllocationGroup& PendingGroup : PendingShadowMaps)
 		{
-			FShadowMapAllocation& Allocation = PendingShadowMaps[ShadowMapIndex];
+			if (!ensure(PendingGroup.Allocations.Num() >= 1))
+			{
+				continue;
+			}
 
-			// Find an existing texture which the light-map can be stored in.
-			FShadowMapPendingTexture* Texture = NULL;
-			for(int32 TextureIndex = 0;TextureIndex < PendingTextures.Num();TextureIndex++)
+			int32 MaxWidth = 0;
+			int32 MaxHeight = 0;
+			for (auto& Allocation : PendingGroup.Allocations)
 			{
-				FShadowMapPendingTexture& ExistingTexture = PendingTextures[TextureIndex];
+				MaxWidth = FMath::Max(MaxWidth, Allocation->MappedRect.Width());
+				MaxHeight = FMath::Max(MaxHeight, Allocation->MappedRect.Height());
+			}
+
+			FShadowMapPendingTexture* Texture = nullptr;
 
-				// See if the new one will fit in the existing texture
-				if ( ExistingTexture.AddElement( Allocation ) )
+			// Find an existing texture which the shadow-map can be stored in.
+			// Shadowmaps will always be 4-pixel aligned...
+			for (FShadowMapPendingTexture& ExistingTexture : PendingTextures)
+			{
+				if (ExistingTexture.AddElement(PendingGroup))
 				{
 					Texture = &ExistingTexture;
 					break;
 				}
 			}
 
-			// If there is no appropriate texture, create a new one.
-			if(!Texture)
+			if (!Texture)
 			{
 				int32 NewTextureSizeX = PackedLightAndShadowMapTextureSize;
 				int32 NewTextureSizeY = PackedLightAndShadowMapTextureSize;
 
-				if(Allocation.MappedRect.Width() > NewTextureSizeX || Allocation.MappedRect.Height() > NewTextureSizeY)
+				// Assumes identically-sized allocations, fit into the smallest square
+				const int32 AllocationCountX = FMath::CeilToInt(FMath::Sqrt(FMath::DivideAndRoundUp(PendingGroup.Allocations.Num() * MaxHeight, MaxWidth)));
+				const int32 AllocationCountY = FMath::DivideAndRoundUp(PendingGroup.Allocations.Num(), AllocationCountX);
+				const int32 AllocationSizeX = AllocationCountX * MaxWidth;
+				const int32 AllocationSizeY = AllocationCountY * MaxHeight;
+
+				if (AllocationSizeX > NewTextureSizeX || AllocationSizeY > NewTextureSizeY)
 				{
-					NewTextureSizeX = FMath::RoundUpToPowerOfTwo(Allocation.MappedRect.Width());
-					NewTextureSizeY = FMath::RoundUpToPowerOfTwo(Allocation.MappedRect.Height());
+					NewTextureSizeX = FMath::RoundUpToPowerOfTwo(AllocationSizeX);
+					NewTextureSizeY = FMath::RoundUpToPowerOfTwo(AllocationSizeY);
 				}
 
 				// If there is no existing appropriate texture, create a new one.
-				Texture				= ::new(PendingTextures) FShadowMapPendingTexture(NewTextureSizeX,NewTextureSizeY);
-				Texture->Outer		= Allocation.TextureOuter;
-				Texture->Bounds		= Allocation.Bounds;
-				Texture->ShadowmapFlags = Allocation.ShadowmapFlags;
-				verify( Texture->AddElement( Allocation, true ) );
+				Texture = new FShadowMapPendingTexture(NewTextureSizeX, NewTextureSizeY);
+				PendingTextures.Add(Texture);
+				Texture->Outer = PendingGroup.TextureOuter;
+				Texture->Bounds = PendingGroup.Bounds;
+				Texture->ShadowmapFlags = PendingGroup.ShadowmapFlags;
+				verify(Texture->AddElement(PendingGroup));
+			}
+
+			// Give the texture ownership of the allocations
+			for (auto& Allocation : PendingGroup.Allocations)
+			{
+				Texture->Allocations.Add(Allocation.Release());
 			}
 		}
-		
-		for(int32 TextureIndex = 0;TextureIndex < PendingTextures.Num();TextureIndex++)
+		PendingShadowMaps.Empty();
+
+		// Encode all the pending textures.
+		for (int32 TextureIndex = 0; TextureIndex < PendingTextures.Num(); TextureIndex++)
 		{
 			if (bUpdateStatus && (TextureIndex % 20) == 0)
 			{
@@ -468,9 +736,7 @@ void FShadowMap2D::EncodeTextures(UWorld* InWorld , bool bLightingSuccessful )
 			FShadowMapPendingTexture& PendingTexture = PendingTextures[TextureIndex];
 			PendingTexture.StartEncoding(InWorld);
 		}
-
 		PendingTextures.Empty();
-		PendingShadowMaps.Empty();
 
 		GWarn->EndSlowTask();
 	}
@@ -495,12 +761,12 @@ void FShadowMap2D::EncodeSingleTexture(FShadowMapPendingTexture& PendingTexture,
 
 		for (int32 ChannelIndex = 0; ChannelIndex < 4; ChannelIndex++)
 		{
-			for (TMap<ULightComponent*, TArray<FQuantizedSignedDistanceFieldShadowSample>>::TIterator It(Allocation.ShadowMapData); It; ++It)
+			for (const auto& ShadowMapPair : Allocation.ShadowMapData)
 			{
-				if (It.Key()->ShadowMapChannel == ChannelIndex)
+				if (ShadowMapPair.Key->ShadowMapChannel == ChannelIndex)
 				{
 					bChannelUsed[ChannelIndex] = true;
-					const TArray<FQuantizedSignedDistanceFieldShadowSample>& SourceSamples = It.Value();
+					const TArray<FQuantizedSignedDistanceFieldShadowSample>& SourceSamples = ShadowMapPair.Value;
 
 					// Copy the raw data for this light-map into the raw texture data array.
 					for (int32 Y = Allocation.MappedRect.Min.Y; Y < Allocation.MappedRect.Max.Y; ++Y)
@@ -537,9 +803,9 @@ void FShadowMap2D::EncodeSingleTexture(FShadowMapPendingTexture& PendingTexture,
 		Allocation.ShadowMap->Texture = Texture;
 
 		// Free the shadow-map's raw data.
-		for (TMap<ULightComponent*, TArray<FQuantizedSignedDistanceFieldShadowSample> >::TIterator It(Allocation.ShadowMapData); It; ++It)
+		for (auto& ShadowMapPair : Allocation.ShadowMapData)
 		{
-			It.Value().Empty();
+			ShadowMapPair.Value.Empty();
 		}
 		
 		int32 PaddedSizeX = Allocation.TotalSizeX;
@@ -744,6 +1010,11 @@ void FShadowMap2D::EncodeSingleTexture(FShadowMapPendingTexture& PendingTexture,
 			}
 		}
 	}
+
+	for (FShadowMapAllocation* Allocation : PendingTexture.Allocations)
+	{
+		Allocation->PostEncode();
+	}
 }
 
 #endif
diff --git a/Engine/Source/Runtime/Engine/Public/LightMap.h b/Engine/Source/Runtime/Engine/Public/LightMap.h
--- a/Engine/Source/Runtime/Engine/Public/LightMap.h
+++ b/Engine/Source/Runtime/Engine/Public/LightMap.h
@@ -10,6 +10,9 @@
 #include "SceneManagement.h"
 #include "Engine/LightMapTexture2D.h"
 
+// Forward declarations
+struct FQuantizedLightmapData;
+
 /** Whether to use bilinear filtering on lightmaps */
 extern ENGINE_API bool GUseBilinearLightmaps;
 
@@ -214,16 +217,6 @@ public:
 	const FVector2D& GetCoordinateScale() const { return CoordinateScale; }
 	const FVector2D& GetCoordinateBias() const { return CoordinateBias; }
 
-	/**
-	 * Finalizes the lightmap after encoding, including setting the UV scale/bias for this lightmap 
-	 * inside the larger UTexture2D that this lightmap is in
-	 * 
-	 * @param Scale UV scale (size)
-	 * @param Bias UV Bias (offset)
-	 * @param ALightmapTexture One of the lightmap textures that this lightmap was put into
-	 */
-	virtual void FinalizeEncoding(const FVector2D& Scale, const FVector2D& Bias, UTexture2D* ALightmapTexture);
-
 	// FLightMap interface.
 	virtual void AddReferencedObjects( FReferenceCollector& Collector );
 
@@ -237,16 +230,29 @@ public:
 	/**
 	 * Allocates texture space for the light-map and stores the light-map's raw data for deferred encoding.
 	 * If the light-map has no lights in it, it will return NULL.
-	 * RawData and SourceQuantizedData will be deleted by this function.
+	 * SourceQuantizedData will be deleted by this function.
 	 * @param	LightMapOuter - The package to create the light-map and textures in.
 	 * @param	SourceQuantizedData - If the data is already quantized, the values will be in here, and not in RawData.  
 	 * @param	Bounds - The bounds of the primitive the light-map will be rendered on.  Used as a hint to pack light-maps on nearby primitives in the same texture.
 	 * @param	InPaddingType - the method for padding the lightmap.
 	 * @param	LightmapFlags - flags that determine how the lightmap is stored (e.g. streamed or not)
 	 */
-	static ENGINE_API class FLightMap2D* AllocateLightMap(UObject* LightMapOuter, struct FQuantizedLightmapData*& SourceQuantizedData,
+	static ENGINE_API class FLightMap2D* AllocateLightMap(UObject* LightMapOuter, FQuantizedLightmapData*& SourceQuantizedData,
 		const FBoxSphereBounds& Bounds, ELightMapPaddingType InPaddingType, ELightMapFlags InLightmapFlags );
 
+	/**
+	 * Allocates texture space for the light-map and stores the light-map's raw data for deferred encoding.
+	 * If the light-map has no lights in it, it will return NULL.
+	 * SourceQuantizedData will be deleted by this function.
+	 * @param	Component - The component that owns this lightmap
+	 * @param	SourceQuantizedData - If the data is already quantized, the values will be in here, and not in RawData.
+	 * @param	Bounds - The bounds of the primitive the light-map will be rendered on.  Used as a hint to pack light-maps on nearby primitives in the same texture.
+	 * @param	InPaddingType - the method for padding the lightmap.
+	 * @param	LightmapFlags - flags that determine how the lightmap is stored (e.g. streamed or not)
+	 */
+	static FLightMap2D* AllocateInstancedLightMap(UInstancedStaticMeshComponent* Component, TArray<TUniquePtr<FQuantizedLightmapData>> SourceQuantizedData,
+		const FBoxSphereBounds& Bounds, ELightMapPaddingType InPaddingType, ELightMapFlags LightmapFlags);
+
 	/**
 	 * Executes all pending light-map encoding requests.
 	 * @param	bLightingSuccessful	Whether the lighting build was successful or not.
diff --git a/Engine/Source/Runtime/Engine/Public/ShadowMap.h b/Engine/Source/Runtime/Engine/Public/ShadowMap.h
--- a/Engine/Source/Runtime/Engine/Public/ShadowMap.h
+++ b/Engine/Source/Runtime/Engine/Public/ShadowMap.h
@@ -15,7 +15,7 @@ struct FSignedDistanceFieldShadowSample
 
 struct FQuantizedSignedDistanceFieldShadowSample
 {
-	enum {NumFilterableComponents = 2};
+	enum { NumFilterableComponents = 2 };
 	uint8 Distance;
 	uint8 PenumbraSize;
 	uint8 Coverage;
@@ -37,28 +37,28 @@ struct FQuantizedSignedDistanceFieldShadowSample
 	{
 		if (Index == 0)
 		{
-			Distance = (uint8)FMath::Clamp<int32>(FMath::TruncToInt(InComponent * 255.0f),0,255);
+			Distance = (uint8)FMath::Clamp<int32>(FMath::TruncToInt(InComponent * 255.0f), 0, 255);
 		}
 		else
 		{
 			checkSlow(Index == 1);
-			PenumbraSize = (uint8)FMath::Clamp<int32>(FMath::TruncToInt(InComponent * 255.0f),0,255);
+			PenumbraSize = (uint8)FMath::Clamp<int32>(FMath::TruncToInt(InComponent * 255.0f), 0, 255);
 		}
 	}
 
 	/** Equality operator */
-	bool operator==( const FQuantizedSignedDistanceFieldShadowSample& RHS ) const
+	bool operator==(const FQuantizedSignedDistanceFieldShadowSample& RHS) const
 	{
 		return Distance == RHS.Distance &&
-			   PenumbraSize == RHS.PenumbraSize &&
-			   Coverage == RHS.Coverage;
+			PenumbraSize == RHS.PenumbraSize &&
+			Coverage == RHS.Coverage;
 	}
 
 	FQuantizedSignedDistanceFieldShadowSample() {}
 	FQuantizedSignedDistanceFieldShadowSample(const FSignedDistanceFieldShadowSample& InSample)
 	{
-		Distance = (uint8)FMath::Clamp<int32>(FMath::TruncToInt(InSample.Distance * 255.0f),0,255);
-		PenumbraSize = (uint8)FMath::Clamp<int32>(FMath::TruncToInt(InSample.PenumbraSize * 255.0f),0,255);
+		Distance = (uint8)FMath::Clamp<int32>(FMath::TruncToInt(InSample.Distance * 255.0f), 0, 255);
+		PenumbraSize = (uint8)FMath::Clamp<int32>(FMath::TruncToInt(InSample.PenumbraSize * 255.0f), 0, 255);
 		Coverage = InSample.IsMapped ? 255 : 0;
 	}
 };
@@ -88,14 +88,13 @@ public:
 	};
 	virtual ShadowMapDataType GetType() const { return UNKNOWN; }
 
-	virtual void Quantize( TArray<struct FQuantizedShadowSample>& OutData ) const {}
-	virtual void Quantize( TArray<struct FQuantizedSignedDistanceFieldShadowSample>& OutData ) const {}
+	virtual void Quantize(TArray<FQuantizedSignedDistanceFieldShadowSample>& OutData) const {}
 
-	virtual void Serialize( FArchive* OutShadowMap ) {}
+	virtual void Serialize(FArchive* OutShadowMap) {}
 
 protected:
 
-	FShadowMapData2D(uint32 InSizeX,uint32 InSizeY):
+	FShadowMapData2D(uint32 InSizeX, uint32 InSizeY) :
 		SizeX(InSizeX),
 		SizeY(InSizeY)
 	{
@@ -115,30 +114,30 @@ class FShadowSignedDistanceFieldData2D : public FShadowMapData2D
 {
 public:
 
-	FShadowSignedDistanceFieldData2D(uint32 InSizeX,uint32 InSizeY)
-	:	FShadowMapData2D(InSizeX, InSizeY)
+	FShadowSignedDistanceFieldData2D(uint32 InSizeX, uint32 InSizeY)
+		: FShadowMapData2D(InSizeX, InSizeY)
 	{
 		Data.Empty(SizeX * SizeY);
 		Data.AddZeroed(SizeX * SizeY);
 	}
 
-	const FSignedDistanceFieldShadowSample& operator()(uint32 X,uint32 Y) const { return Data[SizeX * Y + X]; }
-	FSignedDistanceFieldShadowSample& operator()(uint32 X,uint32 Y) { return Data[SizeX * Y + X]; }
+	const FSignedDistanceFieldShadowSample& operator()(uint32 X, uint32 Y) const { return Data[SizeX * Y + X]; }
+	FSignedDistanceFieldShadowSample& operator()(uint32 X, uint32 Y) { return Data[SizeX * Y + X]; }
 
-	virtual ShadowMapDataType GetType() const { return SHADOW_SIGNED_DISTANCE_FIELD_DATA; }
-	virtual void Quantize( TArray<struct FQuantizedSignedDistanceFieldShadowSample>& OutData ) const
+	virtual ShadowMapDataType GetType() const override { return SHADOW_SIGNED_DISTANCE_FIELD_DATA; }
+	virtual void Quantize(TArray<FQuantizedSignedDistanceFieldShadowSample>& OutData) const override
 	{
-		OutData.Empty( Data.Num() );
-		OutData.AddUninitialized( Data.Num() );
-		for( int32 Index = 0; Index < Data.Num(); Index++ )
+		OutData.Empty(Data.Num());
+		OutData.AddUninitialized(Data.Num());
+		for (int32 Index = 0; Index < Data.Num(); Index++)
 		{
 			OutData[Index] = Data[Index];
 		}
 	}
 
-	virtual void Serialize( FArchive* OutShadowMap )
+	virtual void Serialize(FArchive* OutShadowMap) override
 	{
-		for( int32 Index = 0; Index < Data.Num(); Index++ )
+		for (int32 Index = 0; Index < Data.Num(); Index++)
 		{
 			OutShadowMap->Serialize(&Data[Index], sizeof(FSignedDistanceFieldShadowSample));
 		}
@@ -155,8 +154,8 @@ class FQuantizedShadowSignedDistanceFieldData2D : public FShadowMapData2D
 {
 public:
 
-	FQuantizedShadowSignedDistanceFieldData2D(uint32 InSizeX,uint32 InSizeY)
-	:	FShadowMapData2D(InSizeX, InSizeY)
+	FQuantizedShadowSignedDistanceFieldData2D(uint32 InSizeX, uint32 InSizeY)
+		: FShadowMapData2D(InSizeX, InSizeY)
 	{
 		Data.Empty(SizeX * SizeY);
 		Data.AddZeroed(SizeX * SizeY);
@@ -164,23 +163,18 @@ public:
 
 	FQuantizedSignedDistanceFieldShadowSample* GetData() { return Data.GetData(); }
 
-	const FQuantizedSignedDistanceFieldShadowSample& operator()(uint32 X,uint32 Y) const { return Data[SizeX * Y + X]; }
-	FQuantizedSignedDistanceFieldShadowSample& operator()(uint32 X,uint32 Y) { return Data[SizeX * Y + X]; }
+	const FQuantizedSignedDistanceFieldShadowSample& operator()(uint32 X, uint32 Y) const { return Data[SizeX * Y + X]; }
+	FQuantizedSignedDistanceFieldShadowSample& operator()(uint32 X, uint32 Y) { return Data[SizeX * Y + X]; }
 
-	virtual ShadowMapDataType GetType() const { return SHADOW_SIGNED_DISTANCE_FIELD_DATA_QUANTIZED; }
-	virtual void Quantize( TArray<struct FQuantizedSignedDistanceFieldShadowSample>& OutData ) const
+	virtual ShadowMapDataType GetType() const override { return SHADOW_SIGNED_DISTANCE_FIELD_DATA_QUANTIZED; }
+	virtual void Quantize(TArray<FQuantizedSignedDistanceFieldShadowSample>& OutData) const override
 	{
-		OutData.Empty( Data.Num() );
-		OutData.AddUninitialized( Data.Num() );
-		for( int32 Index = 0; Index < Data.Num(); Index++ )
-		{
-			OutData[Index] = Data[Index];
-		}
+		OutData = Data;
 	}
 
-	virtual void Serialize( FArchive* OutShadowMap )
+	virtual void Serialize(FArchive* OutShadowMap) override
 	{
-		for( int32 Index = 0; Index < Data.Num(); Index++ )
+		for (int32 Index = 0; Index < Data.Num(); Index++)
 		{
 			OutShadowMap->Serialize(&Data[Index], sizeof(FQuantizedSignedDistanceFieldShadowSample));
 		}
@@ -212,8 +206,14 @@ public:
 	TArray<FGuid> LightGuids;
 
 	/** Default constructor. */
-	FShadowMap() : 
-		NumRefs(0) 
+	FShadowMap() :
+		NumRefs(0)
+	{
+	}
+
+	FShadowMap(TArray<FGuid> LightGuids)
+		: LightGuids(MoveTemp(LightGuids))
+		, NumRefs(0)
 	{
 	}
 
@@ -231,7 +231,7 @@ public:
 	}
 
 	// FShadowMap interface.
-	virtual void AddReferencedObjects( FReferenceCollector& Collector ) {}
+	virtual void AddReferencedObjects(FReferenceCollector& Collector) {}
 	virtual void Serialize(FArchive& Ar);
 	virtual FShadowMapInteraction GetInteraction() const = 0;
 
@@ -242,14 +242,14 @@ public:
 	// Reference counting.
 	void AddRef()
 	{
-		check( IsInGameThread() );
+		check(IsInGameThread());
 		NumRefs++;
 	}
 	void Release()
 	{
-		check( IsInGameThread() );
+		check(IsInGameThread());
 		checkSlow(NumRefs > 0);
-		if(--NumRefs == 0)
+		if (--NumRefs == 0)
 		{
 			Cleanup();
 		}
@@ -264,7 +264,7 @@ protected:
 
 private:
 	int32 NumRefs;
-	
+
 	// FDeferredCleanupInterface
 	ENGINE_API virtual void FinishCleanup();
 };
@@ -276,28 +276,41 @@ public:
 	/**
 	 * Executes all pending shadow-map encoding requests.
 	 * @param	InWorld				World in which the textures exist
-	 * @param	bLightingSuccessful	Whether the lighting build was successful or not.	 
+	 * @param	bLightingSuccessful	Whether the lighting build was successful or not.
 	 */
-	ENGINE_API static void EncodeTextures( UWorld* InWorld, bool bLightingSuccessful );
+	ENGINE_API static void EncodeTextures(UWorld* InWorld, bool bLightingSuccessful);
 
 	/**
 	 * Constructs mip maps for a single shadowmap texture.
 	 */
-	static void EncodeSingleTexture(struct FShadowMapPendingTexture& PendingTexture, UShadowMapTexture2D* Texture, TArray< TArray<FFourDistanceFieldSamples> >& MipData);
+	static void EncodeSingleTexture(struct FShadowMapPendingTexture& PendingTexture, UShadowMapTexture2D* Texture, TArray< TArray<FFourDistanceFieldSamples>>& MipData);
 
 	static ENGINE_API FShadowMap2D* AllocateShadowMap(
-		UObject* LightMapOuter, 
-		const TMap<ULightComponent*,FShadowMapData2D*>& ShadowMapData,
-		const FBoxSphereBounds& Bounds, 
+		UObject* LightMapOuter,
+		const TMap<ULightComponent*, FShadowMapData2D*>& ShadowMapData,
+		const FBoxSphereBounds& Bounds,
 		ELightMapPaddingType InPaddingType,
 		EShadowMapFlags InShadowmapFlags);
 
+	/**
+	 * Allocates texture space for the shadow-map and stores the shadow-map's raw data for deferred encoding.
+	 * If the shadow-map has no shadows in it, it will return NULL.
+	 * @param	Component - The component that owns this shadowmap
+	 * @param	InstancedShadowMapData - The raw shadow-map data to fill the texture with.
+	 * @param	Bounds - The bounds of the primitive the shadow-map will be rendered on.  Used as a hint to pack shadow-maps on nearby primitives in the same texture.
+	 * @param	InPaddingType - the method for padding the shadowmap.
+	 * @param	ShadowmapFlags - flags that determine how the shadowmap is stored (e.g. streamed or not)
+	 */
+	static FShadowMap2D* AllocateInstancedShadowMap(UInstancedStaticMeshComponent* Component, TArray<TMap<ULightComponent*, TUniquePtr<FShadowMapData2D>>> InstancedShadowMapData,
+		const FBoxSphereBounds& Bounds, ELightMapPaddingType PaddingType, EShadowMapFlags ShadowmapFlags);
+
 	FShadowMap2D();
 
-	FShadowMap2D(const TMap<ULightComponent*,FShadowMapData2D*>& ShadowMapData);
+	FShadowMap2D(const TMap<ULightComponent*, FShadowMapData2D*>& ShadowMapData);
+	FShadowMap2D(TArray<FGuid> LightGuids);
 
 	// Accessors.
-	class UShadowMapTexture2D* GetTexture() const { check(IsValid()); return Texture; }
+	UShadowMapTexture2D* GetTexture() const { check(IsValid()); return Texture; }
 	const FVector2D& GetCoordinateScale() const { check(IsValid()); return CoordinateScale; }
 	const FVector2D& GetCoordinateBias() const { check(IsValid()); return CoordinateBias; }
 	bool IsValid() const { return Texture != NULL; }
@@ -309,7 +322,7 @@ public:
 	virtual FShadowMapInteraction GetInteraction() const;
 
 	// Runtime type casting.
-	virtual class FShadowMap2D* GetShadowMap2D() { return this; }
+	virtual FShadowMap2D* GetShadowMap2D() { return this; }
 	virtual const FShadowMap2D* GetShadowMap2D() const { return this; }
 
 	/** Call to enable/disable status update of LightMap encoding */
diff --git a/Engine/Source/Runtime/Engine/Public/TextureLayout.h b/Engine/Source/Runtime/Engine/Public/TextureLayout.h
--- a/Engine/Source/Runtime/Engine/Public/TextureLayout.h
+++ b/Engine/Source/Runtime/Engine/Public/TextureLayout.h
@@ -102,6 +102,13 @@ public:
 	 */
 	bool RemoveElement(uint32 ElementBaseX, uint32 ElementBaseY, uint32 ElementSizeX, uint32 ElementSizeY)
 	{
+		if (bAlignByFour)
+		{
+			// Pad to 4 to ensure alignment
+			ElementSizeX = (ElementSizeX + 3) & ~3;
+			ElementSizeY = (ElementSizeY + 3) & ~3;
+		}
+
 		int32 FoundNodeIndex = INDEX_NONE;
 		// Search through nodes to find the element to remove
 		//@todo - traverse the tree instead of iterating through all nodes
