Remove some unrelated test code accidentally added in CL 2335533

diff --git a/Engine/Source/Runtime/Engine/Private/InstancedStaticMesh.cpp b/Engine/Source/Runtime/Engine/Private/InstancedStaticMesh.cpp
--- a/Engine/Source/Runtime/Engine/Private/InstancedStaticMesh.cpp
+++ b/Engine/Source/Runtime/Engine/Private/InstancedStaticMesh.cpp
@@ -1215,12 +1215,6 @@ UInstancedStaticMeshComponent::UInstancedStaticMeshComponent(const FObjectInitia
 {
 	Mobility = EComponentMobility::Movable;
 	BodyInstance.bSimulatePhysics = false;
-
-	char test[] = "abcefgijkmnoCDEGHIKLM";
-	TRotationInPlaceMerge<TJugglingRotation<FEuclidDivisionGCD>>::Merge(test, 12, 21, [](char a, char b){ return FCharAnsi::ToLower(a) < FCharAnsi::ToLower(b); });
-
-	char test1[] = "abcefgijkCDEGHIKLM";
-	TRotationInPlaceMerge<TJugglingRotation<FEuclidDivisionGCD>>::Merge(test1, 9, 18, [](char a, char b){ return FCharAnsi::ToLower(a) < FCharAnsi::ToLower(b); });
 }
 
 #if WITH_EDITOR
