fixed navmesh not updating properly on BSP

diff --git a/Engine/Source/Runtime/Engine/Classes/AI/Navigation/NavigationSystem.h b/Engine/Source/Runtime/Engine/Classes/AI/Navigation/NavigationSystem.h
--- a/Engine/Source/Runtime/Engine/Classes/AI/Navigation/NavigationSystem.h
+++ b/Engine/Source/Runtime/Engine/Classes/AI/Navigation/NavigationSystem.h
@@ -562,6 +562,9 @@ public:
 	/** */
 	virtual void OnWorldInitDone(FNavigationSystem::EMode Mode);
 
+	/** adds BSP collisions of currently streamed in levels to octree */
+	void InitializeLevelCollisions();
+
 #if WITH_EDITOR
 	/** allow editor to toggle whether seamless navigation building is enabled */
 	static void SetNavigationAutoUpdateEnabled(bool bNewEnable, UNavigationSystem* InNavigationsystem);
@@ -659,6 +662,7 @@ protected:
 	uint32 bNavigationBuildingLocked:1;
 	uint32 bInitialBuildingLockActive:1;
 	uint32 bInitialSetupHasBeenPerformed:1;
+	uint32 bInitialLevelsAdded:1;
 
 	/** cached navigable world bounding box*/
 	mutable FBox NavigableWorldBounds;
diff --git a/Engine/Source/Runtime/Engine/Private/AI/Navigation/NavigationSystem.cpp b/Engine/Source/Runtime/Engine/Private/AI/Navigation/NavigationSystem.cpp
--- a/Engine/Source/Runtime/Engine/Private/AI/Navigation/NavigationSystem.cpp
+++ b/Engine/Source/Runtime/Engine/Private/AI/Navigation/NavigationSystem.cpp
@@ -191,6 +191,7 @@ UNavigationSystem::UNavigationSystem(const FObjectInitializer& ObjectInitializer
 	, bNavigationBuildingLocked(false)
 	, bInitialBuildingLockActive(false)
 	, bInitialSetupHasBeenPerformed(false)
+	, bInitialLevelsAdded(false)
 	, CurrentlyDrawnNavDataIndex(0)
 	, DirtyAreasUpdateTime(0)
 {
@@ -291,26 +292,7 @@ void UNavigationSystem::PostInitProperties()
 			: FNavigationOctree::SkipNavGeometry);
 	
 		bInitialBuildingLockActive = bInitialBuildingLocked;
-
-#if WITH_NAVIGATION_GENERATOR
-		UWorld* World = GetWorld();
-		if (World)
-		{
-			if (World->PersistentLevel)
-		{
-			OnLevelAddedToWorld(World->PersistentLevel, World);
-		}
-
-			for (int32 Idx = 0; Idx < World->StreamingLevels.Num(); Idx++)
-			{
-				ULevelStreaming* StreamingInfo = World->StreamingLevels[Idx];
-				if (StreamingInfo && StreamingInfo->GetLoadedLevel())
-				{
-					OnLevelAddedToWorld(StreamingInfo->GetLoadedLevel(), World);
-				}
-			}
-		}
-#endif
+		InitializeLevelCollisions();
 
 		// register for any actor move change
 #if WITH_EDITOR
@@ -361,6 +343,7 @@ void UNavigationSystem::OnWorldInitDone(FNavigationSystem::EMode Mode)
 	}
 	else
 	{
+		InitializeLevelCollisions();
 		PopulateNavOctree();
 
 		// gather all navigation data instances and register all not-yet-registered
@@ -2321,6 +2304,31 @@ void UNavigationSystem::ReleaseInitialBuildingLock()
 	}
 }
 
+void UNavigationSystem::InitializeLevelCollisions()
+{
+#if WITH_NAVIGATION_GENERATOR
+	UWorld* World = GetWorld();
+	if (!bInitialLevelsAdded && UNavigationSystem::GetCurrent(World) == this)
+	{
+		if (World->PersistentLevel)
+		{
+			OnLevelAddedToWorld(World->PersistentLevel, World);
+		}
+
+		for (int32 Idx = 0; Idx < World->StreamingLevels.Num(); Idx++)
+		{
+			ULevelStreaming* StreamingInfo = World->StreamingLevels[Idx];
+			if (StreamingInfo && StreamingInfo->GetLoadedLevel())
+			{
+				OnLevelAddedToWorld(StreamingInfo->GetLoadedLevel(), World);
+			}
+		}
+
+		bInitialLevelsAdded = true;
+	}
+#endif
+}
+
 #if WITH_EDITOR
 void UNavigationSystem::UpdateLevelCollision(ULevel* InLevel)
 {
@@ -2717,6 +2725,8 @@ void UNavigationSystem::AddLevelCollisionToOctree(ULevel* Level)
 		{
 			NavOctree->AddNode(Level, NULL, Bounds, BSPElem);
 			AddDirtyArea(Bounds, ENavigationDirtyFlag::All);
+
+			UE_LOG(LogNavOctree, Log, TEXT("ADD %s"), *GetNameSafe(Level));
 		}
 	}
 #endif // NAVOCTREE_CONTAINS_COLLISION_DATA
@@ -2726,6 +2736,8 @@ void UNavigationSystem::RemoveLevelCollisionFromOctree(ULevel* Level)
 {
 #if NAVOCTREE_CONTAINS_COLLISION_DATA
 	const FOctreeElementId* ElementId = GetObjectsNavOctreeId(Level);
+	UE_LOG(LogNavOctree, Log, TEXT("UNREG %s %s"), *GetNameSafe(Level), ElementId ? TEXT("[exists]") : TEXT(""));
+
 	if (ElementId != NULL)
 	{
 		if (NavOctree->IsValidElementId(*ElementId))
