Fix Mac build (missing forward declarations)

diff --git a/Engine/Source/Runtime/Engine/Public/LightMap.h b/Engine/Source/Runtime/Engine/Public/LightMap.h
--- a/Engine/Source/Runtime/Engine/Public/LightMap.h
+++ b/Engine/Source/Runtime/Engine/Public/LightMap.h
@@ -12,6 +12,7 @@
 
 // Forward declarations
 struct FQuantizedLightmapData;
+class UInstancedStaticMeshComponent;
 
 /** Whether to use bilinear filtering on lightmaps */
 extern ENGINE_API bool GUseBilinearLightmaps;
diff --git a/Engine/Source/Runtime/Engine/Public/ShadowMap.h b/Engine/Source/Runtime/Engine/Public/ShadowMap.h
--- a/Engine/Source/Runtime/Engine/Public/ShadowMap.h
+++ b/Engine/Source/Runtime/Engine/Public/ShadowMap.h
@@ -2,6 +2,10 @@
 
 #pragma once
 
+// Forward declarations
+class ULightComponent;
+class UInstancedStaticMeshComponent;
+
 struct FSignedDistanceFieldShadowSample
 {
 	/** Normalized and encoded distance to the nearest shadow transition, in the range [0, 1], where .5 is at the transition. */
