Remove GRHIFeatureLevel_DEPRECATED.

Skeletal uses max feature level to determine CPU access type for skin VBs.
pass through feature level to FGPUSkinCache::DispatchSkinCacheProcess (only called when feature level >= SM5)

diff --git a/Engine/Source/Runtime/Engine/Private/GPUSkinCache.cpp b/Engine/Source/Runtime/Engine/Private/GPUSkinCache.cpp
--- a/Engine/Source/Runtime/Engine/Private/GPUSkinCache.cpp
+++ b/Engine/Source/Runtime/Engine/Private/GPUSkinCache.cpp
@@ -350,7 +350,7 @@ int32 FGPUSkinCache::StartCacheMesh(FRHICommandListImmediate& RHICmdList, int32
 
 	auto& ShaderData = GPUVertexFactory->GetShaderData();
 	
-	DispatchSkinCacheProcess(RHICmdList, InputStreamFloatOffset, ECSInfo->StreamOffset, ShaderData.GetBoneBuffer(), ShaderData.GetUniformBuffer(), VBInfo, StreamStrides[0], ElementVertexCount, ShaderData.MeshOrigin, ShaderData.MeshExtension, bExtraBoneInfluences);
+	DispatchSkinCacheProcess(RHICmdList, InputStreamFloatOffset, ECSInfo->StreamOffset, ShaderData.GetBoneBuffer(), ShaderData.GetUniformBuffer(), VBInfo, StreamStrides[0], ElementVertexCount, ShaderData.MeshOrigin, ShaderData.MeshExtension, bExtraBoneInfluences, Skin->GetFeatureLevel());
 
 	CacheCurrentFloatOffset += ProcessedFloatCount;
 	CachedChunksThisFrameCount++;
@@ -467,11 +467,8 @@ FGPUSkinCache::FElementCacheStatusInfo* FGPUSkinCache::FindEvictableCacheStatusI
 	}
 }
 
-void FGPUSkinCache::DispatchSkinCacheProcess(FRHICommandListImmediate& RHICmdList, uint32 InputStreamFloatOffset, uint32 OutputBufferFloatOffset, FBoneBufferTypeRef BoneBuffer, FUniformBufferRHIRef UniformBuffer, const FVertexBufferInfo* VBInfo, uint32 VertexStride, uint32 VertexCount, const FVector& MeshOrigin, const FVector& MeshExtension, bool bUseExtraBoneInfluences)
+void FGPUSkinCache::DispatchSkinCacheProcess(FRHICommandListImmediate& RHICmdList, uint32 InputStreamFloatOffset, uint32 OutputBufferFloatOffset, FBoneBufferTypeRef BoneBuffer, FUniformBufferRHIRef UniformBuffer, const FVertexBufferInfo* VBInfo, uint32 VertexStride, uint32 VertexCount, const FVector& MeshOrigin, const FVector& MeshExtension, bool bUseExtraBoneInfluences, ERHIFeatureLevel::Type FeatureLevel)
 {
-	// MOBILEPREVIEWTODO: Proper value for this
-	const auto FeatureLevel = GRHIFeatureLevel_DEPRECATED;
-
 	if (FeatureLevel < ERHIFeatureLevel::SM5)
 	{
 		return;
diff --git a/Engine/Source/Runtime/Engine/Private/GPUSkinCache.h b/Engine/Source/Runtime/Engine/Private/GPUSkinCache.h
--- a/Engine/Source/Runtime/Engine/Private/GPUSkinCache.h
+++ b/Engine/Source/Runtime/Engine/Private/GPUSkinCache.h
@@ -113,7 +113,7 @@ private:
 
 	FElementCacheStatusInfo* FindEvictableCacheStatusInfo();
 
-	void	DispatchSkinCacheProcess(FRHICommandListImmediate& RHICmdList, uint32 InputStreamFloatOffset, uint32 OutputBufferFloatOffset, FBoneBufferTypeRef BoneBuffer, FUniformBufferRHIRef UniformBuffer, const FVertexBufferInfo* VBInfo, uint32 VertexStride, uint32 VertexCount, const FVector& MeshOrigin, const FVector& MeshExtension, bool bUseExtraBoneInfluences);
+	void	DispatchSkinCacheProcess(FRHICommandListImmediate& RHICmdList, uint32 InputStreamFloatOffset, uint32 OutputBufferFloatOffset, FBoneBufferTypeRef BoneBuffer, FUniformBufferRHIRef UniformBuffer, const FVertexBufferInfo* VBInfo, uint32 VertexStride, uint32 VertexCount, const FVector& MeshOrigin, const FVector& MeshExtension, bool bUseExtraBoneInfluences, ERHIFeatureLevel::Type FeatureLevel);
 
 	bool	Initialized : 1;
 
diff --git a/Engine/Source/Runtime/Engine/Private/PrimitiveDrawingUtils.cpp b/Engine/Source/Runtime/Engine/Private/PrimitiveDrawingUtils.cpp
--- a/Engine/Source/Runtime/Engine/Private/PrimitiveDrawingUtils.cpp
+++ b/Engine/Source/Runtime/Engine/Private/PrimitiveDrawingUtils.cpp
@@ -15,7 +15,7 @@ void EmitMeshDrawEvents_Inner(FRHICommandList& RHICmdList, const FPrimitiveScene
 {
 	// Only show material name at the top level
 	// Note: this is the parent's material name, not the material instance
-	SCOPED_DRAW_EVENTF(RHICmdList, MaterialEvent, *Mesh.MaterialRenderProxy->GetMaterial(PrimitiveSceneProxy ? PrimitiveSceneProxy->GetScene().GetFeatureLevel() : GRHIFeatureLevel_DEPRECATED)->GetFriendlyName());
+	SCOPED_DRAW_EVENTF(RHICmdList, MaterialEvent, *Mesh.MaterialRenderProxy->GetMaterial(PrimitiveSceneProxy ? PrimitiveSceneProxy->GetScene().GetFeatureLevel() : GMaxRHIFeatureLevel)->GetFriendlyName());
 	if (PrimitiveSceneProxy)
 	{
 		// Show Actor, level and resource name inside the material name
diff --git a/Engine/Source/Runtime/Engine/Private/SkeletalMesh.cpp b/Engine/Source/Runtime/Engine/Private/SkeletalMesh.cpp
--- a/Engine/Source/Runtime/Engine/Private/SkeletalMesh.cpp
+++ b/Engine/Source/Runtime/Engine/Private/SkeletalMesh.cpp
@@ -1186,7 +1186,7 @@ void FStaticLODModel::Serialize( FArchive& Ar, UObject* Owner, int32 Idx )
 		if( Ar.IsLoading() )
 		{
 			// set cpu skinning flag on the vertex buffer so that the resource arrays know if they need to be CPU accessible
-			VertexBufferGPUSkin.SetNeedsCPUAccess(bKeepBuffersInCPUMemory || SkelMeshOwner->GetImportedResource()->RequiresCPUSkinning(GRHIFeatureLevel_DEPRECATED));
+			VertexBufferGPUSkin.SetNeedsCPUAccess(bKeepBuffersInCPUMemory || SkelMeshOwner->GetImportedResource()->RequiresCPUSkinning(GMaxRHIFeatureLevel));
 		}
 		Ar << NumTexCoords;
 		Ar << VertexBufferGPUSkin;
@@ -4351,7 +4351,7 @@ void FSkeletalMeshSceneProxy::GetDynamicElementsSection(const TArray<const FScen
 			BatchElement.NumPrimitives = Section.NumTriangles;
 			if( GIsEditor && MeshObject->ProgressiveDrawingFraction != 1.f )
 			{
-				if (Mesh.MaterialRenderProxy->GetMaterial(GRHIFeatureLevel_DEPRECATED)->GetBlendMode() == BLEND_Translucent)
+				if (Mesh.MaterialRenderProxy->GetMaterial(FeatureLevel)->GetBlendMode() == BLEND_Translucent)
 				{
 					BatchElement.NumPrimitives = FMath::RoundToInt(((float)Section.NumTriangles)*FMath::Clamp<float>(MeshObject->ProgressiveDrawingFraction,0.f,1.f));
 					if( BatchElement.NumPrimitives == 0 )
diff --git a/Engine/Source/Runtime/Engine/Private/SkeletalMeshMerge.cpp b/Engine/Source/Runtime/Engine/Private/SkeletalMeshMerge.cpp
--- a/Engine/Source/Runtime/Engine/Private/SkeletalMeshMerge.cpp
+++ b/Engine/Source/Runtime/Engine/Private/SkeletalMeshMerge.cpp
@@ -272,7 +272,7 @@ static void BoneMapToNewRefSkel(const TArray<FBoneIndexType>& InBoneMap, const T
 */
 void FSkeletalMeshMerge::GenerateNewSectionArray( TArray<FNewSectionInfo>& NewSectionArray, int32 LODIdx )
 {
-	const int32 MaxGPUSkinBones = GetFeatureLevelMaxNumberOfBones(GRHIFeatureLevel_DEPRECATED);
+	const int32 MaxGPUSkinBones = GetFeatureLevelMaxNumberOfBones(GMaxRHIFeatureLevel);
 
 	NewSectionArray.Empty();
 	for( int32 MeshIdx=0; MeshIdx < SrcMeshList.Num(); MeshIdx++ )
@@ -605,7 +605,7 @@ void FSkeletalMeshMerge::GenerateLODModel( int32 LODIdx )
 	// copy the new vertices and indices to the vertex buffer for the new model
 	MergeLODModel.VertexBufferGPUSkin.SetUseFullPrecisionUVs(MergeMesh->bUseFullPrecisionUVs);
 	// set CPU skinning on vertex buffer since it affects the type of TResourceArray needed
-	MergeLODModel.VertexBufferGPUSkin.SetNeedsCPUAccess(MergeResource->RequiresCPUSkinning(GRHIFeatureLevel_DEPRECATED));
+	MergeLODModel.VertexBufferGPUSkin.SetNeedsCPUAccess(MergeResource->RequiresCPUSkinning(GMaxRHIFeatureLevel));
 	MergeLODModel.VertexBufferGPUSkin.SetHasExtraBoneInfluences(MergeResource->HasExtraBoneInfluences());
 	// Set the number of tex coords on this vertex buffer
 	MergeLODModel.VertexBufferGPUSkin.SetNumTexCoords(TotalNumUVs);
diff --git a/Engine/Source/Runtime/Engine/Public/SkeletalRenderPublic.h b/Engine/Source/Runtime/Engine/Public/SkeletalRenderPublic.h
--- a/Engine/Source/Runtime/Engine/Public/SkeletalRenderPublic.h
+++ b/Engine/Source/Runtime/Engine/Public/SkeletalRenderPublic.h
@@ -203,6 +203,12 @@ public:
 	int32				SectionIndexPreview;
 #endif
 
+	/** returns the feature level this FSkeletalMeshObject was created with */
+	ERHIFeatureLevel::Type GetFeatureLevel() const
+	{
+		return FeatureLevel;
+	}
+
 protected:
 	/** The skeletal mesh resource with which to render. */
 	FSkeletalMeshResource* SkeletalMeshResource;
diff --git a/Engine/Source/Runtime/RHI/Public/RHI.h b/Engine/Source/Runtime/RHI/Public/RHI.h
--- a/Engine/Source/Runtime/RHI/Public/RHI.h
+++ b/Engine/Source/Runtime/RHI/Public/RHI.h
@@ -245,7 +245,6 @@ extern RHI_API EShaderPlatform GMaxRHIShaderPlatform;
 // should be used for all rendering, rather a specific set for each view. The deprecated names are being kept around temporarily 
 // while the last instances are removed.
 #define GRHIShaderPlatform_DEPRECATED GMaxRHIShaderPlatform
-#define GRHIFeatureLevel_DEPRECATED GMaxRHIFeatureLevel
 
 /** Table for finding out which shader platform corresponds to a given feature level for this RHI. */
 extern RHI_API EShaderPlatform GShaderPlatformForFeatureLevel[ERHIFeatureLevel::Num];
