Reenabling Landscape module for any build, cause lack of it was causing shader serialization assert bug. We have to figure out how the get rid of this, but it's a fix for now.

diff --git a/Engine/Source/Runtime/Engine/Engine.Build.cs b/Engine/Source/Runtime/Engine/Engine.Build.cs
--- a/Engine/Source/Runtime/Engine/Engine.Build.cs
+++ b/Engine/Source/Runtime/Engine/Engine.Build.cs
@@ -65,10 +65,12 @@ public class Engine : ModuleRules
 				"SlateCore",
 				"SlateReflector",
 				"VectorVM",
+				"Landscape"
 			}
         );
 
         CircularlyReferencedDependentModules.Add("AIModule");
+		CircularlyReferencedDependentModules.Add("Landscape");
 
 		DynamicallyLoadedModuleNames.AddRange(
 			new string[]
diff --git a/Engine/Source/Runtime/UE4Game/UE4Game.Build.cs b/Engine/Source/Runtime/UE4Game/UE4Game.Build.cs
--- a/Engine/Source/Runtime/UE4Game/UE4Game.Build.cs
+++ b/Engine/Source/Runtime/UE4Game/UE4Game.Build.cs
@@ -7,7 +7,6 @@ public class UE4Game : ModuleRules
 	public UE4Game(TargetInfo Target)
 	{
 		PrivateDependencyModuleNames.Add("Core");
-		PrivateDependencyModuleNames.Add("Landscape");
 	
 		if (Target.Platform == UnrealTargetPlatform.IOS)
 		{
