Hopefully fix non-unity build

diff --git a/Engine/Source/Runtime/Engine/Private/InstancedStaticMesh.cpp b/Engine/Source/Runtime/Engine/Private/InstancedStaticMesh.cpp
--- a/Engine/Source/Runtime/Engine/Private/InstancedStaticMesh.cpp
+++ b/Engine/Source/Runtime/Engine/Private/InstancedStaticMesh.cpp
@@ -45,6 +45,7 @@
 #if WITH_EDITOR
 #include "LightMap.h"
 #include "ShadowMap.h"
+#include "Logging/MessageLog.h"
 #endif
 
 // This must match the maximum a user could specify in the material (see 
@@ -1574,12 +1575,10 @@ public:
 		bComplete = true;
 	}
 
-#if WITH_EDITOR
 	virtual bool DebugThisMapping() const override
 	{
 		return false;
 	}
-#endif	//WITH_EDITOR
 
 	virtual FString GetDescription() const override
 	{
