Fix for GSystemResolution not being setup in time for some objects being created resulting in UTextureRenderTarget2D's having a size of 0,0.

diff --git a/Engine/Source/Runtime/Engine/Private/GameEngine.cpp b/Engine/Source/Runtime/Engine/Private/GameEngine.cpp
--- a/Engine/Source/Runtime/Engine/Private/GameEngine.cpp
+++ b/Engine/Source/Runtime/Engine/Private/GameEngine.cpp
@@ -196,7 +196,7 @@ void UGameEngine::ConditionallyOverrideSettings(int32& ResolutionX, int32& Resol
 	{
 		// consume available desktop area
 		FDisplayMetrics DisplayMetrics;
-		FSlateApplication::Get().GetDisplayMetrics( DisplayMetrics );
+		FDisplayMetrics::GetDisplayMetrics( DisplayMetrics );
 		
 		ResolutionX = DisplayMetrics.PrimaryDisplayWidth;
 		ResolutionY = DisplayMetrics.PrimaryDisplayHeight;
@@ -234,7 +234,7 @@ void UGameEngine::ConditionallyOverrideSettings(int32& ResolutionX, int32& Resol
 	{
 		// Always use the device's actual resolution that has been setup earlier
 		FDisplayMetrics DisplayMetrics;
-		FSlateApplication::Get().GetDisplayMetrics( DisplayMetrics );
+		FDisplayMetrics::GetDisplayMetrics( DisplayMetrics );
 
 		// We need to pass the resolution back out to GameUserSettings, or it will just override it again
 		ResolutionX = DisplayMetrics.PrimaryDisplayWorkAreaRect.Right - DisplayMetrics.PrimaryDisplayWorkAreaRect.Left;
diff --git a/Engine/Source/Runtime/Engine/Private/GameUserSettings.cpp b/Engine/Source/Runtime/Engine/Private/GameUserSettings.cpp
--- a/Engine/Source/Runtime/Engine/Private/GameUserSettings.cpp
+++ b/Engine/Source/Runtime/Engine/Private/GameUserSettings.cpp
@@ -337,7 +337,7 @@ void UGameUserSettings::PreloadResolutionSettings()
 		CVar->Set((int32)WindowMode);
 	}
 
-	RequestResolutionChange(ResolutionX, ResolutionY, WindowMode, false);
+	RequestResolutionChange(ResolutionX, ResolutionY, WindowMode);
 
 	IConsoleManager::Get().CallAllConsoleVariableSinks();
 }
