Removing include of NiagaraConstants.h; no longer needed here

diff --git a/Engine/Source/Runtime/Engine/Private/NiagaraComponent.cpp b/Engine/Source/Runtime/Engine/Private/NiagaraComponent.cpp
--- a/Engine/Source/Runtime/Engine/Private/NiagaraComponent.cpp
+++ b/Engine/Source/Runtime/Engine/Private/NiagaraComponent.cpp
@@ -3,7 +3,6 @@
 #include "EnginePrivate.h"
 #include "Components/NiagaraComponent.h"
 #include "Engine/NiagaraScript.h"
-#include "Engine/NiagaraConstants.h"
 #include "VectorVM.h"
 #include "ParticleHelper.h"
 #include "Particles/ParticleResources.h"
diff --git a/Engine/Source/Runtime/Engine/Private/NiagaraSimulation.cpp b/Engine/Source/Runtime/Engine/Private/NiagaraSimulation.cpp
--- a/Engine/Source/Runtime/Engine/Private/NiagaraSimulation.cpp
+++ b/Engine/Source/Runtime/Engine/Private/NiagaraSimulation.cpp
@@ -1,6 +1,5 @@
 
 #include "EnginePrivate.h"
-#include "Engine/NiagaraConstants.h"
 #include "Engine/NiagaraSimulation.h"
 #include "VectorVM.h"
 
