Move World teardown to GameEngine::PreExit instead of HandleExitCommand
Move SlateApplication shutdown after Engine and App PreExit

diff --git a/Engine/Source/Runtime/Engine/Private/GameEngine.cpp b/Engine/Source/Runtime/Engine/Private/GameEngine.cpp
--- a/Engine/Source/Runtime/Engine/Private/GameEngine.cpp
+++ b/Engine/Source/Runtime/Engine/Private/GameEngine.cpp
@@ -500,6 +500,21 @@ void UGameEngine::PreExit()
 		UWorld* const World = WorldList[WorldIndex].World();
 		if ( World != NULL )
 		{
+			World->bIsTearingDown = true;
+
+			// Cancel any pending connection to a server
+			CancelPending(World);
+
+			// Shut down any existing game connections
+			ShutdownWorldNetDriver(World);
+
+			for (FActorIterator ActorIt(World); ActorIt; ++ActorIt)
+			{
+				ActorIt->RouteEndPlay(EEndPlayReason::Quit);
+			}
+
+			World->GetGameInstance()->Shutdown();
+
 			World->FlushLevelStreaming( NULL, true );
 			World->CleanupWorld();
 		}
@@ -619,25 +634,6 @@ bool UGameEngine::Exec( UWorld* InWorld, const TCHAR* Cmd, FOutputDevice& Ar )
 
 bool UGameEngine::HandleExitCommand( const TCHAR* Cmd, FOutputDevice& Ar )
 {
-	for (int32 WorldIndex = 0; WorldIndex < WorldList.Num(); ++WorldIndex)
-	{
-		UWorld* const World = WorldList[WorldIndex].World();
-		World->bIsTearingDown = true;
-
-		// Cancel any pending connection to a server
-		CancelPending(WorldList[WorldIndex]);
-
-		// Shut down any existing game connections
-		ShutdownWorldNetDriver(World);
-
-		for (FActorIterator ActorIt(World); ActorIt; ++ActorIt)
-		{
-			ActorIt->RouteEndPlay(EEndPlayReason::Quit);
-		}
-
-		World->GetGameInstance()->Shutdown();
-	}
-
 	Ar.Log( TEXT("Closing by request") );
 	FPlatformMisc::RequestExit( 0 );
 	return true;
