#ue4 - FHitResult default constructor initializes 'Time' variable to 1.0.

diff --git a/Engine/Source/Runtime/Engine/Classes/Engine/EngineTypes.h b/Engine/Source/Runtime/Engine/Classes/Engine/EngineTypes.h
--- a/Engine/Source/Runtime/Engine/Classes/Engine/EngineTypes.h
+++ b/Engine/Source/Runtime/Engine/Classes/Engine/EngineTypes.h
@@ -1433,6 +1433,7 @@ struct ENGINE_API FHitResult
 	FHitResult()
 	{
 		FMemory::Memzero(this, sizeof(FHitResult));
+		Time = 1.f;
 	}
 	
 	explicit FHitResult(float InTime)
@@ -1444,6 +1445,7 @@ struct ENGINE_API FHitResult
 	explicit FHitResult(EForceInit Init)
 	{
 		FMemory::Memzero(this, sizeof(FHitResult));
+		Time = 1.f;
 	}
 
 	/** Ctor for easily creating "fake" hits from limited data. */
