#ue4 - Added 'bBlockingHit' boolean to BreakHitResult.

diff --git a/Engine/Source/Runtime/Engine/Classes/Kismet/GameplayStatics.h b/Engine/Source/Runtime/Engine/Classes/Kismet/GameplayStatics.h
--- a/Engine/Source/Runtime/Engine/Classes/Kismet/GameplayStatics.h
+++ b/Engine/Source/Runtime/Engine/Classes/Kismet/GameplayStatics.h
@@ -389,11 +389,12 @@ class ENGINE_API UGameplayStatics : public UBlueprintFunctionLibrary
 
 	/** Extracts data from a HitResult. 
 	 * @param Hit			The source HitResult.
+	 * @param bBlockingHit	True if there was a blocking hit, false otherwise.
+	 * @param Time			'Time' of impact along trace direction ranging from [0.0 to 1.0) if there is a hit, indicating time between start and end. Equals 1.0 if there is no hit.
 	 * @param Location		Location of the hit in world space. If this was a swept shape test, this is the location where we can place the shape in the world where it will not penetrate.
 	 * @param Normal		Normal of the hit in world space, for the object that was swept (e.g. for a sphere trace this points towards the sphere's center). Equal to ImpactNormal for line tests.
 	 * @param ImpactPoint	Location of the actual contact point of the trace shape with the surface of the hit object.
 	 * @param ImpactNormal	Normal of the hit in world space, for the object that was hit by the sweep.
-	 * @param Time			'Time' of impact along trace direction ranging from [0.0 to 1.0) if there is a hit, indicating time between start and end. Equals 1.0 if there is no hit.
 	 * @param PhysMat		Physical material that was hit. Must set bReturnPhysicalMaterial to true in the query params for this to be returned.
 	 * @param HitActor		Actor hit by the trace.
 	 * @param HitComponent	PrimitiveComponent hit by the trace.
@@ -401,7 +402,7 @@ class ENGINE_API UGameplayStatics : public UBlueprintFunctionLibrary
 	 * @param HitItem		Primitive-specific data recording which item in the primitive was hit
 	 */
 	UFUNCTION(BlueprintPure, Category = "Collision", meta=(NativeBreakFunc))
-	static void BreakHitResult(const struct FHitResult& Hit, FVector& Location, FVector& Normal, FVector& ImpactPoint, FVector& ImpactNormal, float& Time, class UPhysicalMaterial*& PhysMat, class AActor*& HitActor, class UPrimitiveComponent*& HitComponent, FName& HitBoneName, int32& HitItem);
+	static void BreakHitResult(const struct FHitResult& Hit, bool& bBlockingHit, float& Time, FVector& Location, FVector& Normal, FVector& ImpactPoint, FVector& ImpactNormal, class UPhysicalMaterial*& PhysMat, class AActor*& HitActor, class UPrimitiveComponent*& HitComponent, FName& HitBoneName, int32& HitItem);
 
 	/** Returns the EPhysicalSurface type of the given Hit. 
 	 * To edit surface type for your project, use ProjectSettings/Physics/PhysicalSurface section
diff --git a/Engine/Source/Runtime/Engine/Private/GameplayStatics.cpp b/Engine/Source/Runtime/Engine/Private/GameplayStatics.cpp
--- a/Engine/Source/Runtime/Engine/Private/GameplayStatics.cpp
+++ b/Engine/Source/Runtime/Engine/Private/GameplayStatics.cpp
@@ -557,16 +557,17 @@ UParticleSystemComponent* UGameplayStatics::SpawnEmitterAttached(UParticleSystem
 	return PSC;
 }
 
-void UGameplayStatics::BreakHitResult(const FHitResult& Hit, FVector& Location, FVector& Normal, FVector& ImpactPoint, FVector& ImpactNormal, float& Time, UPhysicalMaterial*& PhysMat, AActor*& HitActor, UPrimitiveComponent*& HitComponent, FName& HitBoneName, int32& HitItem)
+void UGameplayStatics::BreakHitResult(const FHitResult& Hit, bool& bBlockingHit, float& Time, FVector& Location, FVector& Normal, FVector& ImpactPoint, FVector& ImpactNormal, UPhysicalMaterial*& PhysMat, AActor*& HitActor, UPrimitiveComponent*& HitComponent, FName& HitBoneName, int32& HitItem)
 {
+	bBlockingHit = Hit.bBlockingHit;
+	Time = Hit.Time;
 	Location = Hit.Location;
 	Normal = Hit.Normal;
 	ImpactPoint = Hit.ImpactPoint;
-	ImpactNormal = Hit.ImpactNormal;
-	Time = Hit.Time;
+	ImpactNormal = Hit.ImpactNormal;	
 	PhysMat = Hit.PhysMaterial.Get();
 	HitActor = Hit.GetActor();
-	HitComponent = Hit.Component.Get();
+	HitComponent = Hit.GetComponent();
 	HitBoneName = Hit.BoneName;
 	HitItem = Hit.Item;
 }
