Fix uninitialized trimesh and heightfield in case where mtd fails on overlap conversion.

diff --git a/Engine/Source/Runtime/Engine/Private/Collision/CollisionConversions.cpp b/Engine/Source/Runtime/Engine/Private/Collision/CollisionConversions.cpp
--- a/Engine/Source/Runtime/Engine/Private/Collision/CollisionConversions.cpp
+++ b/Engine/Source/Runtime/Engine/Private/Collision/CollisionConversions.cpp
@@ -630,7 +630,7 @@ static bool ConvertOverlappedShapeToImpactHit(const PxLocationHit& PHit, const F
 
 		PxGeometryType::Enum GeometryType = PShape->getGeometryType();
 
-		if (GeometryType == PxGeometryType::eTRIANGLEMESH || GeometryType == PxGeometryType::eHEIGHTFIELD)
+		if (PShape->getTriangleMeshGeometry(PTriMeshGeom) || PShape->getHeightFieldGeometry(PHeightfieldGeom))
 		{
 			if (ComputeInflatedMTD(MtdInflation, PHit, OutResult, QueryTM, Geom, PShapeWorldPose))
 			{
@@ -641,6 +641,7 @@ static bool ConvertOverlappedShapeToImpactHit(const PxLocationHit& PHit, const F
 				const bool bIsTriMesh = GeometryType == PxGeometryType::eTRIANGLEMESH;
 				PxU32 HitTris[64];
 				bool bOverflow = false;
+
 				const int32 NumTrisHit =  bIsTriMesh? PxMeshQuery::findOverlapTriangleMesh(Geom, QueryTM, PTriMeshGeom,     PShapeWorldPose, HitTris, ARRAY_COUNT(HitTris), 0, bOverflow) :
 													  PxMeshQuery::findOverlapHeightField( Geom, QueryTM, PHeightfieldGeom, PShapeWorldPose, HitTris, ARRAY_COUNT(HitTris), 0, bOverflow);
 				if (NumTrisHit > 0)
