Adding ForceSpawn to GPU emitters, implementing forced burst and regular spawning, and enabling the EventReceiver modules for GPU emitters.

diff --git a/Engine/Config/BaseEditor.ini b/Engine/Config/BaseEditor.ini
--- a/Engine/Config/BaseEditor.ini
+++ b/Engine/Config/BaseEditor.ini
@@ -222,7 +222,7 @@ bUseGrayScale=false
 +ModuleMenu_TypeDataToBaseModuleRejections=(ObjName="ParticleModuleTypeDataBeam2",InvalidObjNames=("ParticleModuleTrailBase","ParticleModuleVectorFieldBase","ParticleModuleLightBase"))
 +ModuleMenu_TypeDataToBaseModuleRejections=(ObjName="ParticleModuleTypeDataMesh",InvalidObjNames=("ParticleModuleBeamBase","ParticleModuleTrailBase","ParticleModuleVectorFieldBase"))
 +ModuleMenu_TypeDataToBaseModuleRejections=(ObjName="ParticleModuleTypeDataRibbon",InvalidObjNames=("ParticleModuleVectorFieldBase","ParticleModuleLightBase"))
-+ModuleMenu_TypeDataToBaseModuleRejections=(ObjName="ParticleModuleTypeDataGpu",InvalidObjNames=("ParticleModuleBeamBase","ParticleModuleTrailBase","ParticleModuleCameraBase","ParticleModuleKillBase","ParticleModuleParameterBase","ParticleModuleEventBase","ParticleModuleLightBase"))
++ModuleMenu_TypeDataToBaseModuleRejections=(ObjName="ParticleModuleTypeDataGpu",InvalidObjNames=("ParticleModuleBeamBase","ParticleModuleTrailBase","ParticleModuleCameraBase","ParticleModuleKillBase","ParticleModuleParameterBase","ParticleModuleLightBase"))
 +ModuleMenu_TypeDataToSpecificModuleRejections=(ObjName="None",InvalidObjNames=("ParticleModuleCollisionGPU","ParticleModuleMeshMaterial","ParticleModuleMeshRotation","ParticleModuleMeshRotation_Seeded","ParticleModuleMeshRotationRate","ParticleModuleMeshRotationRate_Seeded","ParticleModuleMeshRotationRateMultiplyLife","ParticleModuleMeshRotationRateOverLife","ParticleModuleTrailSpawnPerUnit", "ParticleModuleAttractorPointGravity"))
 +ModuleMenu_TypeDataToSpecificModuleRejections=(ObjName="ParticleModuleTypeDataBeam2",InvalidObjNames=("ParticleModuleSizeScaleBySpeed","ParticleModuleMeshMaterial","ParticleModuleMeshRotation","ParticleModuleMeshRotation_Seeded","ParticleModuleMeshRotationRate","ParticleModuleMeshRotationRate_Seeded","ParticleModuleMeshRotationRateMultiplyLife","ParticleModuleMeshRotationRateOverLife","ParticleModuleTrailSpawnPerUnit","ParticleModuleAccelerationDrag"))
 +ModuleMenu_TypeDataToSpecificModuleRejections=(ObjName="ParticleModuleTypeDataRibbon",InvalidObjNames=("ParticleModuleSizeScaleBySpeed","ParticleModuleMeshMaterial","ParticleModuleMeshRotation","ParticleModuleMeshRotation_Seeded","ParticleModuleMeshRotationRate","ParticleModuleMeshRotationRate_Seeded","ParticleModuleMeshRotationRateMultiplyLife","ParticleModuleMeshRotationRateOverLife","ParticleModuleLocationEmitterDirect" ))
diff --git a/Engine/Source/Runtime/Engine/Private/Particles/ParticleGpuSimulation.cpp b/Engine/Source/Runtime/Engine/Private/Particles/ParticleGpuSimulation.cpp
--- a/Engine/Source/Runtime/Engine/Private/Particles/ParticleGpuSimulation.cpp
+++ b/Engine/Source/Runtime/Engine/Private/Particles/ParticleGpuSimulation.cpp
@@ -2899,6 +2899,10 @@ class FGPUSpriteParticleEmitterInstance : public FParticleEmitterInstance
 	TArray<uint32> TilesToClear;
 	/** The list of new particles generated this time step. */
 	TArray<FNewParticle> NewParticles;
+	/** The list of force spawned particles from events */
+	TArray<FNewParticle> ForceSpawnedParticles;
+	/** The list of force spawned particles from events using Bursts */
+	TArray<FNewParticle> ForceBurstSpawnedParticles;
 	/** The rotation to apply to the local vector field. */
 	FRotator LocalVectorFieldRotation;
 	/** The strength of the point attractor. */
@@ -3243,17 +3247,25 @@ public:
 			{
 				float BurstDeltaTime = DeltaSeconds;
 				GetCurrentBurstRateOffset(BurstDeltaTime, BurstInfo.Count);
+
+				BurstInfo.Count += ForceBurstSpawnedParticles.Num();
+
 				if (BurstInfo.Count > FXConsoleVariables::MaxGPUParticlesSpawnedPerFrame)
 				{
 					LeftoverBurst = BurstInfo.Count - FXConsoleVariables::MaxGPUParticlesSpawnedPerFrame;
 					BurstInfo.Count = FXConsoleVariables::MaxGPUParticlesSpawnedPerFrame;
 				}
 			}
+
+
+
 			int32 FirstBurstParticleIndex = NewParticles.Num();
 			BurstInfo.Count = AllocateTilesForParticles(NewParticles, BurstInfo.Count, ActiveTileCount);
 
 			// Determine spawn count based on rate.
 			FSpawnInfo SpawnInfo = GetNumParticlesToSpawn(DeltaSeconds);
+			SpawnInfo.Count += ForceSpawnedParticles.Num();
+
 			int32 FirstSpawnParticleIndex = NewParticles.Num();
 			SpawnInfo.Count = AllocateTilesForParticles(NewParticles, SpawnInfo.Count, ActiveTileCount);
 			SpawnFraction += LeftoverBurst;
@@ -3261,15 +3273,18 @@ public:
 			if (BurstInfo.Count > 0)
 			{
 				// Spawn burst particles.
-				BuildNewParticles(NewParticles.GetData() + FirstBurstParticleIndex, BurstInfo);
+				BuildNewParticles(NewParticles.GetData() + FirstBurstParticleIndex, BurstInfo, ForceBurstSpawnedParticles);
 			}
 
 			if (SpawnInfo.Count > 0)
 			{
 				// Spawn normal particles.
-				BuildNewParticles(NewParticles.GetData() + FirstSpawnParticleIndex, SpawnInfo);
+				BuildNewParticles(NewParticles.GetData() + FirstSpawnParticleIndex, SpawnInfo, ForceSpawnedParticles);
 			}
 
+			ForceBurstSpawnedParticles.Empty();
+			ForceSpawnedParticles.Empty();
+
 			int32 NewParticleCount = BurstInfo.Count + SpawnInfo.Count;
 			INC_DWORD_STAT_BY(STAT_GPUSpritesSpawned, NewParticleCount);
 #if STATS
@@ -3700,7 +3715,7 @@ private:
 	 * @param SpawnTime - The time at which to begin spawning particles.
 	 * @param Increment - The amount by which to increment time for each particle spawned.
 	 */
-	void BuildNewParticles(FNewParticle* InNewParticles, FSpawnInfo SpawnInfo)
+	void BuildNewParticles(FNewParticle* InNewParticles, FSpawnInfo SpawnInfo, TArray<FNewParticle> &ForceSpawned)
 	{
 		const float OneOverTwoPi = 1.0f / (2.0f * PI);
 		UParticleModuleRequired* RequiredModule = EmitterInfo.RequiredModule;
@@ -3730,6 +3745,15 @@ private:
 
 			// Set the particle's location and invoke each spawn module on the particle.
 			TempParticle->Location = EmitterToSimulation.GetOrigin();
+
+			int32 ForceSpawnedOffset = SpawnInfo.Count - ForceSpawned.Num();
+			if (ForceSpawned.Num() && i > ForceSpawnedOffset)
+			{
+				TempParticle->Location = ForceSpawned[i - ForceSpawnedOffset - 1].Position;
+				TempParticle->RelativeTime = ForceSpawned[i - ForceSpawnedOffset - 1].RelativeTime;
+				TempParticle->Velocity += ForceSpawned[i - ForceSpawnedOffset - 1].Velocity;
+			}
+
 			for (int32 ModuleIndex = 0; ModuleIndex < EmitterInfo.SpawnModules.Num(); ModuleIndex++)
 			{
 				UParticleModule* SpawnModule = EmitterInfo.SpawnModules[ModuleIndex];
@@ -3878,6 +3902,28 @@ private:
 
 	virtual void ForceSpawn(float DeltaTime, int32 InSpawnCount, int32 InBurstCount, FVector& InLocation, FVector& InVelocity)
 	{
+		const bool bUseLocalSpace = GetCurrentLODLevelChecked()->RequiredModule->bUseLocalSpace;
+		FVector SpawnLocation = bUseLocalSpace ? FVector::ZeroVector : InLocation;
+
+		float Increment = DeltaTime / InSpawnCount;
+		for (int32 i = 0; i < InSpawnCount; i++)
+		{
+
+			FNewParticle Particle;
+			Particle.Position = SpawnLocation;
+			Particle.Velocity = InVelocity;
+			Particle.RelativeTime = Increment*i;
+			ForceSpawnedParticles.Add(Particle);
+		}
+
+		for (int32 i = 0; i < InBurstCount; i++)
+		{
+			FNewParticle Particle;
+			Particle.Position = SpawnLocation;
+			Particle.Velocity = InVelocity;
+			Particle.RelativeTime = 0.0f;
+			ForceBurstSpawnedParticles.Add(Particle);
+		}
 	}
 
 	virtual void PreSpawn(FBaseParticle* Particle, const FVector& InitialLocation, const FVector& InitialVelocity)
