AnimBP: fixed issue where entry node is always active first time in the state machine.

diff --git a/Engine/Source/Runtime/Engine/Classes/Animation/AnimNode_StateMachine.h b/Engine/Source/Runtime/Engine/Classes/Animation/AnimNode_StateMachine.h
--- a/Engine/Source/Runtime/Engine/Classes/Animation/AnimNode_StateMachine.h
+++ b/Engine/Source/Runtime/Engine/Classes/Animation/AnimNode_StateMachine.h
@@ -148,12 +148,17 @@ protected:
 	
 	// Used during transitions to make sure we don't double tick a state if it appears multiple times
 	TArray<int32> StatesUpdated;
-	
+
+private:
+	// true if it is the first update.
+	bool bFirstUpdate;
+
 public:
 	FAnimNode_StateMachine()
 		: MaxTransitionsPerFrame(3)
 		, PRIVATE_MachineDescription(NULL)
 		, CurrentState(INDEX_NONE)
+		, bFirstUpdate(true)
 	{
 	}
 
diff --git a/Engine/Source/Runtime/Engine/Private/Animation/AnimNode_StateMachine.cpp b/Engine/Source/Runtime/Engine/Private/Animation/AnimNode_StateMachine.cpp
--- a/Engine/Source/Runtime/Engine/Private/Animation/AnimNode_StateMachine.cpp
+++ b/Engine/Source/Runtime/Engine/Private/Animation/AnimNode_StateMachine.cpp
@@ -209,6 +209,8 @@ void FAnimNode_StateMachine::Initialize(const FAnimationInitializeContext& Conte
 			// Move to the default state
 			SetState(Context, Machine->InitialState);
 
+			// initialize first update
+			bFirstUpdate = true;
 			// Fire off any entry notifications the default state has
 			Context.AnimInstance->AddAnimNotifyFromGeneratedClass(GetStateInfo(Machine->InitialState).StartNotify);
 		}
@@ -334,6 +336,13 @@ void FAnimNode_StateMachine::Update(const FAnimationUpdateContext& Context)
 	}
 	while (bFoundValidTransition && (TransitionCountThisFrame < MaxTransitionsPerFrame));
 
+	// in the first update, we don't like to transition from entry state
+	// so we throw out any transition data at the first update
+	if (bFirstUpdate)
+	{
+		ActiveTransitionArray.Empty();
+		bFirstUpdate = false;
+	}
 
 	StatesUpdated.Empty(StatesUpdated.Num());
 
