Collision profile: fixed so that you can't have object type as collision profile object type.

diff --git a/Engine/Source/Runtime/Engine/Private/Collision/CollisionProfile.cpp b/Engine/Source/Runtime/Engine/Private/Collision/CollisionProfile.cpp
--- a/Engine/Source/Runtime/Engine/Private/Collision/CollisionProfile.cpp
+++ b/Engine/Source/Runtime/Engine/Private/Collision/CollisionProfile.cpp
@@ -513,9 +513,9 @@ void UCollisionProfile::FillProfileData(TArray<FCollisionResponseTemplate>& Prof
 	// First using "DisplayName" will be more user-friendly
 	// Second, allowing default reponse simplifies a lot of things
 	// --------------------------------------------------------------------------
-	for (auto ProfileIter = ProfileList.CreateIterator(); ProfileIter; ++ProfileIter)
+	for (int32 ProfileIndex = 0; ProfileIndex<ProfileList.Num(); ++ProfileIndex)
 	{
-		FCollisionResponseTemplate& Template = *ProfileIter;
+		FCollisionResponseTemplate& Template = ProfileList[ProfileIndex];
 
 		if (Template.ObjectTypeName!=NAME_None)
 		{
@@ -523,12 +523,31 @@ void UCollisionProfile::FillProfileData(TArray<FCollisionResponseTemplate>& Prof
 			int32 EnumIndex = ReturnContainerIndexFromChannelName(Template.ObjectTypeName);
 			if (EnumIndex != INDEX_NONE)
 			{
-				Template.ObjectType = (ECollisionChannel)EnumIndex;
+				// first verify if this is real object type
+				ECollisionChannel ObjectTypeEnum = (ECollisionChannel)EnumIndex;
+				EObjectTypeQuery ObjectTypeQuery = ConvertToObjectType(ObjectTypeEnum);
+				if (ObjectTypeQuery != ObjectTypeQuery_MAX)
+				{
+					Template.ObjectType = ObjectTypeEnum;
+				}
+				else
+				{
+					UE_LOG(LogCollisionProfile, Warning, TEXT("Profile (%s) ObjectTypeName (%s) is Trace Type. You can set Object Type Channel to Object Type."),
+						*Template.Name.ToString(), *Template.ObjectTypeName.ToString());
+
+					ProfileList.RemoveAt(ProfileIndex);
+					--ProfileIndex;
+					continue;
+				}
 			}
 			else
 			{
 				UE_LOG(LogCollisionProfile, Warning, TEXT("Profile (%s) ObjectTypeName (%s) is invalid. "), 
 					*Template.Name.ToString(), *Template.ObjectTypeName.ToString());
+
+				ProfileList.RemoveAt(ProfileIndex);
+				--ProfileIndex;
+				continue;
 			}
 		}
 
