Track referenced objects in replication shadow state during GC

diff --git a/Engine/Source/Runtime/Engine/Private/DataChannel.cpp b/Engine/Source/Runtime/Engine/Private/DataChannel.cpp
--- a/Engine/Source/Runtime/Engine/Private/DataChannel.cpp
+++ b/Engine/Source/Runtime/Engine/Private/DataChannel.cpp
@@ -1982,6 +1982,13 @@ FString UActorChannel::Describe()
 void UActorChannel::AddReferencedObjects(UObject* InThis, FReferenceCollector& Collector)
 {
 	UActorChannel* This = CastChecked<UActorChannel>(InThis);	
+
+	// Let GC know that we're referencing UObject's in the replication shadow state
+	for ( auto MapIt = This->ReplicationMap.CreateIterator(); MapIt; ++MapIt )
+	{
+		MapIt.Value()->AddReferencedObjects( InThis, Collector );
+	}
+
 	Super::AddReferencedObjects( This, Collector );
 }
 
diff --git a/Engine/Source/Runtime/Engine/Private/DataReplication.cpp b/Engine/Source/Runtime/Engine/Private/DataReplication.cpp
--- a/Engine/Source/Runtime/Engine/Private/DataReplication.cpp
+++ b/Engine/Source/Runtime/Engine/Private/DataReplication.cpp
@@ -977,6 +977,17 @@ void FObjectReplicator::ForceRefreshUnreliableProperties()
 	bOpenAckCalled = true;
 }
 
+void FObjectReplicator::AddReferencedObjects( UObject* InThis, FReferenceCollector& Collector )
+{
+	if ( GetObject() == NULL )
+	{
+		UE_LOG( LogNet, Verbose, TEXT( "AddReferencedObjects: Object == NULL" ) );
+		return;
+	}
+
+	RepLayout->AddReferencedObjects( RepState, InThis, Collector );
+}
+
 void FObjectReplicator::PostSendBunch( FPacketIdRange & PacketRange, uint8 bReliable )
 {
 	if ( GetObject() == NULL )
diff --git a/Engine/Source/Runtime/Engine/Private/RepLayout.cpp b/Engine/Source/Runtime/Engine/Private/RepLayout.cpp
--- a/Engine/Source/Runtime/Engine/Private/RepLayout.cpp
+++ b/Engine/Source/Runtime/Engine/Private/RepLayout.cpp
@@ -24,6 +24,8 @@ FAutoConsoleVariable CVarDoReplicationContextString( TEXT( "net.ContextDebug" ),
 
 //#define ENABLE_SUPER_CHECKSUMS
 
+//#pragma optimize("", off)
+
 #ifdef USE_CUSTOM_COMPARE
 static FORCEINLINE bool CompareBool( const FRepLayoutCmd& Cmd, const void* A, const void* B )
 {
@@ -195,6 +197,65 @@ static void SerializeReadWritePropertyChecksum( const FRepLayoutCmd& Cmd, const
 	}
 }
 
+class FMyCmdIteratorBaseStackState
+{
+public:
+	FMyCmdIteratorBaseStackState( FScriptArray* InArray ) : Array( InArray ), ArrayNum( InArray->Num() ) {}
+
+	FScriptArray*	Array;
+	const uint16	ArrayNum;
+};
+
+template< typename TImpl, typename TStackState >
+class FRepLayoutCmdIterator
+{
+public:
+	FRepLayoutCmdIterator( TImpl& InImpl, const TArray< FRepLayoutCmd > & InCmds ) : Impl( InImpl ), Cmds( InCmds ) {}
+
+	void IterateCmdArray_r( const FRepLayoutCmd& Cmd, const int32 CmdIndex, const uint8* RESTRICT ShadowData )
+	{
+		TStackState StackState( (FScriptArray *)ShadowData );
+
+		uint8* LocalData = (uint8*)StackState.Array->GetData();
+
+		Impl.PreArray( StackState );
+
+		for ( int32 i = 0; i < StackState.ArrayNum; i++ )
+		{
+			ProcessCmd_r( CmdIndex + 1, Cmd.EndCmd - 1, LocalData + i * Cmd.ElementSize );
+		}
+
+		Impl.PostArray( StackState );
+	}
+
+	void ProcessCmd_r( const int32 CmdStart, const int32 CmdEnd, const uint8* RESTRICT ShadowData )
+	{
+		for ( int32 CmdIndex = CmdStart; CmdIndex < CmdEnd; CmdIndex++ )
+		{
+			const FRepLayoutCmd& Cmd = Cmds[ CmdIndex ];
+
+			check( Cmd.Type != REPCMD_Return );
+
+			if ( Cmd.Type == REPCMD_DynamicArray )
+			{
+				IterateCmdArray_r( Cmd, CmdIndex, ShadowData + Cmd.Offset );
+				CmdIndex = Cmd.EndCmd - 1;	// Jump past children of this array (-1 for ++ in for loop)
+				continue;
+			}
+
+			Impl.ProcessCmd( Cmd, ShadowData + Cmd.Offset );
+		}
+	}
+
+	void Start( FRepState * RepState )
+	{
+		ProcessCmd_r( 0, Cmds.Num() - 1, (uint8*)RepState->StaticBuffer.GetData() );
+	}
+
+	TImpl&							Impl;
+	const TArray< FRepLayoutCmd > & Cmds;
+};
+
 uint16 FRepLayout::CompareProperties_r(
 	const int32				CmdStart,
 	const int32				CmdEnd,
@@ -2456,6 +2517,42 @@ void FRepLayout::RebuildConditionalProperties( FRepState * RESTRICT	RepState, co
 	RepState->RepFlags = RepFlags;
 }
 
+class FAddReferencedObjectsImpl
+{
+public:
+	FAddReferencedObjectsImpl( UObject* InReferencingObject, FReferenceCollector& InCollector ) : ReferencingObject( InReferencingObject ), Collector( InCollector ) {}
+
+	void PreArray( FMyCmdIteratorBaseStackState & StackState ) { }
+	void PostArray( FMyCmdIteratorBaseStackState & StackState ) { }
+
+	void ProcessCmd( const FRepLayoutCmd& Cmd, const uint8* RESTRICT ShadowData )
+	{
+		if ( Cmd.Type == REPCMD_PropertyObject )
+		{
+			UObjectPropertyBase * ObjProperty = CastChecked< UObjectPropertyBase>( Cmd.Property );
+
+			UObject* ReferencedObject = ObjProperty->GetObjectPropertyValue( ShadowData );
+
+			if ( ReferencedObject != NULL )
+			{
+				Collector.AddReferencedObject( ReferencedObject, ReferencingObject );
+			}
+		}
+	}
+
+	UObject*				ReferencingObject;
+	FReferenceCollector&	Collector;
+};
+
+void FRepLayout::AddReferencedObjects( FRepState * RepState, UObject* ReferencingObject, FReferenceCollector& Collector ) const
+{
+	FAddReferencedObjectsImpl AddReferencedObjectsImpl( ReferencingObject, Collector );
+
+	FRepLayoutCmdIterator< FAddReferencedObjectsImpl, FMyCmdIteratorBaseStackState > Iterator( AddReferencedObjectsImpl, Cmds );
+
+	Iterator.Start( RepState );
+}
+
 void FRepLayout::InitChangedTracker( FRepChangedPropertyTracker * ChangedTracker ) const
 {
 	ChangedTracker->Parents.SetNum( Parents.Num() );
diff --git a/Engine/Source/Runtime/Engine/Public/Net/DataReplication.h b/Engine/Source/Runtime/Engine/Public/Net/DataReplication.h
--- a/Engine/Source/Runtime/Engine/Public/Net/DataReplication.h
+++ b/Engine/Source/Runtime/Engine/Public/Net/DataReplication.h
@@ -114,7 +114,10 @@ public:
 	bool	ReceivedBunch( FInBunch & Bunch, const FReplicationFlags& RepFlags, bool & bOutHasUnmapped );
 	void	PostReceivedBunch();
 
-	void ForceRefreshUnreliableProperties();
+	void	ForceRefreshUnreliableProperties();
+
+	void	AddReferencedObjects( UObject* InThis, FReferenceCollector& Collector );
+
 
 	bool bHasReplicatedProperties;
 
diff --git a/Engine/Source/Runtime/Engine/Public/Net/RepLayout.h b/Engine/Source/Runtime/Engine/Public/Net/RepLayout.h
--- a/Engine/Source/Runtime/Engine/Public/Net/RepLayout.h
+++ b/Engine/Source/Runtime/Engine/Public/Net/RepLayout.h
@@ -357,6 +357,8 @@ public:
 	void SerializePropertiesForStruct( UStruct * Struct, FArchive & Ar, UPackageMap	* Map, void* Data, bool & bHasUnmapped ) const;	
 	void InitFromStruct( UStruct * InStruct );
 
+	void AddReferencedObjects( FRepState * RepState, UObject* ReferencingObject, FReferenceCollector& Collector ) const;
+
 private:
 	void RebuildConditionalProperties( FRepState * RESTRICT	RepState, const FRepChangedPropertyTracker& ChangedTracker, const FReplicationFlags& RepFlags ) const;
 
