Improved "add looping interpolation"

diff --git a/Engine/Source/Runtime/Engine/Private/Animation/AnimSequence.cpp b/Engine/Source/Runtime/Engine/Private/Animation/AnimSequence.cpp
--- a/Engine/Source/Runtime/Engine/Private/Animation/AnimSequence.cpp
+++ b/Engine/Source/Runtime/Engine/Private/Animation/AnimSequence.cpp
@@ -2414,6 +2414,9 @@ bool UAnimSequence::AddLoopingInterpolation()
 	if(Interval > 0.f)
 	{
 		// 2d array of animated time [track index][time key]
+		// the reason I have to query the time is because 
+		// the RawAnimation doesn't have all keys
+		// it removes duplicated keys if not necessary. 
 		TArray< TArray<FTransform> > AnimatedLocalSpaces;
 		AnimatedLocalSpaces.AddZeroed(NumTracks);
 
@@ -2428,7 +2431,6 @@ bool UAnimSequence::AddLoopingInterpolation()
 				GetBoneTransform(AnimatedLocalSpaces[TrackIndex][Key], TrackIndex, Interval*Key, false);
 			}
 
-			// now add extra, source should be available all the time since we check interval already
 			AnimatedLocalSpaces[TrackIndex][DestKeyIndex] = AnimatedLocalSpaces[TrackIndex][SrcKeyIndex];
 		}
 
@@ -2456,8 +2458,9 @@ bool UAnimSequence::AddLoopingInterpolation()
 		}
 
 		RawAnimationData = NewRawAnimationData;
-		NumFrames += 1;
-		SequenceLength += Interval;
+		NumFrames = NewNumKeys;
+
+		// we don't increase sequence length as the original sequence length should stay
 
 		PostProcessSequence();
 		return true;
