Cleaned up GameViewportClient header (Comments, structs etc)

diff --git a/Engine/Source/Runtime/Engine/Classes/Engine/DebugDisplayProperty.h b/Engine/Source/Runtime/Engine/Classes/Engine/DebugDisplayProperty.h
new file mode 100644
--- /dev/null
+++ b/Engine/Source/Runtime/Engine/Classes/Engine/DebugDisplayProperty.h
@@ -0,0 +1,42 @@
+// Copyright 1998-2014 Epic Games, Inc. All Rights Reserved.
+
+#pragma once
+
+#include "DebugDisplayProperty.generated.h"
+
+/** 
+ * Debug property display functionality
+ * to interact with this, use "display", "displayall", "displayclear"
+ *
+ * @see UGameViewportClient
+ * @see FDebugDisplayProperty
+ * @see DrawStatsHUD
+ */
+USTRUCT()
+struct FDebugDisplayProperty
+{
+	GENERATED_USTRUCT_BODY()
+
+	/** the object whose property to display. If this is a class, all objects of that class are drawn. */
+	UPROPERTY()
+	class UObject* Obj;
+
+	/** if Obj is a class and WithinClass is not NULL, further limit the display to objects that have an Outer of WithinClass */
+	UPROPERTY()
+	TSubclassOf<class UObject>  WithinClass;
+
+	/** name of the property to display */
+	FName PropertyName;
+
+	/** whether PropertyName is a "special" value not directly mapping to a real property (e.g. state name) */
+	uint32 bSpecialProperty:1;
+
+
+	FDebugDisplayProperty()
+		: Obj(NULL)
+		, WithinClass(NULL)
+		, bSpecialProperty(false)
+	{
+	}
+
+};
diff --git a/Engine/Source/Runtime/Engine/Classes/Engine/GameViewportClient.h b/Engine/Source/Runtime/Engine/Classes/Engine/GameViewportClient.h
--- a/Engine/Source/Runtime/Engine/Classes/Engine/GameViewportClient.h
+++ b/Engine/Source/Runtime/Engine/Classes/Engine/GameViewportClient.h
@@ -1,50 +1,27 @@
 // Copyright 1998-2014 Epic Games, Inc. All Rights Reserved.
 
-/**
- * A game viewport (FViewport) is a high-level abstract interface for the
- * platform specific rendering, audio, and input subsystems.
- * GameViewportClient is the engine's interface to a game viewport.
- * Exactly one GameViewportClient is created for each instance of the game.  The
- * only case (so far) where you might have a single instance of Engine, but
- * multiple instances of the game (and thus multiple GameViewportClients) is when
- * you have more than one PIE window running.
- *
- * Responsibilities:
- * propagating input events to the global interactions list
- *
- */
+
 #pragma once
 
 #include "ShowFlags.h"
-#include "Engine/ScriptViewportClient.h"
+#include "ScriptViewportClient.h"
+#include "ViewportSplitScreen.h"
+#include "DebugDisplayProperty.h"
+#include "TitleSafeZone.h"
+#include "GameViewportDelegates.h"
+
 #include "GameViewportClient.generated.h"
 
+
 class UGameInstance;
 class UNetDriver;
+class ULocalPlayer;
+class FSceneViewport;
+class SViewport;
+class SWindow;
+class SOverlay;
+class UCanvas;
 
-/**
- * Enum of the different splitscreen types
- */
-namespace ESplitScreenType
-{
-	enum Type
-	{
-		// No split
-		None,
-		// 2 player horizontal split
-		TwoPlayer_Horizontal,
-		// 2 player vertical split
-		TwoPlayer_Vertical,
-		// 3 Player split with 1 player on top and 2 on bottom
-		ThreePlayer_FavorTop,
-		// 3 Player split with 1 player on bottom and 2 on top
-		ThreePlayer_FavorBottom,
-		// 4 Player split
-		FourPlayer,
-
-		SplitTypeCount
-	};
-}
 
 /**
  * Stereoscopic rendering passes.  FULL implies stereoscopic rendering isn't enabled for this pass
@@ -56,129 +33,21 @@ enum EStereoscopicPass
 	eSSP_RIGHT_EYE
 };
 
-/**
- * The 4 different kinds of safezones
- */
-enum ESafeZoneType
-{
-	eSZ_TOP,
-	eSZ_BOTTOM,
-	eSZ_LEFT,
-	eSZ_RIGHT,
-	eSZ_MAX,
-};
-
-/** Max/Recommended screen viewable extents as a percentage */
-struct FTitleSafeZoneArea
-{
-	float MaxPercentX;
-	float MaxPercentY;
-	float RecommendedPercentX;
-	float RecommendedPercentY;
-
-	FTitleSafeZoneArea()
-		: MaxPercentX(0)
-		, MaxPercentY(0)
-		, RecommendedPercentX(0)
-		, RecommendedPercentY(0)
-	{
-	}
-
-};
-
-/** Structure to store splitscreen data. */
-struct FPerPlayerSplitscreenData
-{
-	float SizeX;
-	float SizeY;
-	float OriginX;
-	float OriginY;
-
-
-	FPerPlayerSplitscreenData()
-		: SizeX(0)
-		, SizeY(0)
-		, OriginX(0)
-		, OriginY(0)
-	{
-	}
-
-	FPerPlayerSplitscreenData(float NewSizeX, float NewSizeY, float NewOriginX, float NewOriginY)
-		: SizeX(NewSizeX)
-		, SizeY(NewSizeY)
-		, OriginX(NewOriginX)
-		, OriginY(NewOriginY)
-	{
-	}
-
-};
-
-/** Structure containing all the player splitscreen datas per splitscreen configuration. */
-struct FSplitscreenData
-{
-	TArray<struct FPerPlayerSplitscreenData> PlayerData;
-};
-
-/** debug property display functionality
- * to interact with this, use "display", "displayall", "displayclear"
- */
-USTRUCT()
-struct FDebugDisplayProperty
-{
-	GENERATED_USTRUCT_BODY()
-
-	/** the object whose property to display. If this is a class, all objects of that class are drawn. */
-	UPROPERTY()
-	class UObject* Obj;
-
-	/** if Obj is a class and WithinClass is not NULL, further limit the display to objects that have an Outer of WithinClass */
-	UPROPERTY()
-	TSubclassOf<class UObject>  WithinClass;
-
-	/** name of the property to display */
-	FName PropertyName;
-
-	/** whether PropertyName is a "special" value not directly mapping to a real property (e.g. state name) */
-	uint32 bSpecialProperty:1;
-
-
-	FDebugDisplayProperty()
-		: Obj(NULL)
-		, WithinClass(NULL)
-		, bSpecialProperty(false)
-	{
-	}
-
-};
-
-class ULocalPlayer;
 
 /**
- * Delegate type for when a screenshot has been captured
+ * A game viewport (FViewport) is a high-level abstract interface for the
+ * platform specific rendering, audio, and input subsystems.
+ * GameViewportClient is the engine's interface to a game viewport.
+ * Exactly one GameViewportClient is created for each instance of the game.  The
+ * only case (so far) where you might have a single instance of Engine, but
+ * multiple instances of the game (and thus multiple GameViewportClients) is when
+ * you have more than one PIE window running.
  *
- * The first parameter is the width.
- * The second parameter is the height.
- * The third parameter is the array of bitmap data.
- */
-DECLARE_MULTICAST_DELEGATE_ThreeParams(FOnScreenshotCaptured, int32 /*Width*/, int32 /*Height*/, const TArray<FColor>& /*Colors*/);
-
-/**
- * Delegate type for when a png screenshot has been captured
+ * Responsibilities:
+ * propagating input events to the global interactions list
  *
- * The first parameter is the width.
- * The second parameter is the height.
- * The third parameter is the array of bitmap data.
- * The fourth parameter is the screen shot filename.
+ * @see UGameViewportClient
  */
-DECLARE_DELEGATE_FourParams(FOnPNGScreenshotCaptured, int32, int32, const TArray<FColor>&, const FString&);
-
-/**
-* Delegate type for when call is made to close a viewport
-*
-* The first parameter is the viewport being closed.
-*/
-DECLARE_DELEGATE_OneParam(FOnCloseRequested, FViewport*);
-
 UCLASS(Within=Engine, transient, config=Engine)
 class ENGINE_API UGameViewportClient : public UScriptViewportClient, public FExec
 {
@@ -219,6 +88,7 @@ protected:
 	 */
 	TEnumAsByte<ESplitScreenType::Type> ActiveSplitscreenType;
 
+	/* The relative world context for this viewport */
 	UPROPERTY()
 	UWorld* World;
 
@@ -229,8 +99,6 @@ protected:
 	bool bSuppressTransitionMessage;
 
 public:
-	/** @todo document */
-	float ProgressFadeTime;
 
 	/** see enum EViewModeIndex */
 	int32 ViewModeIndex;
@@ -250,10 +118,13 @@ public:
 	/** Returns a relative world context for this viewport.	 */
 	virtual UWorld* GetWorld() const override;
 
-	class FSceneViewport* GetGameViewport();
+	/* Returns the game viewport */
+	FSceneViewport* GetGameViewport();
 
-	TSharedPtr<class SViewport> GetGameViewportWidget();
+	/* Returns the widget for this viewport */
+	TSharedPtr<SViewport> GetGameViewportWidget();
 
+	/* Returns the relevant game instance for this viewport */
 	UGameInstance* GetGameInstance() const;
 
 	virtual void Init(struct FWorldContext& WorldContext, UGameInstance* OwningGameInstance);
@@ -363,29 +234,34 @@ public:
 
 	/**
 	 * Set this GameViewportClient's viewport and viewport frame to the viewport specified
+	 * @param	InViewportFrame	The viewportframe to set
 	 */
 	virtual void SetViewportFrame( FViewportFrame* InViewportFrame );
 
 	/**
 	 * Set this GameViewportClient's viewport to the viewport specified
+	 * @param	InViewportFrame	The viewport to set
 	 */
 	virtual void SetViewport( FViewport* InViewportFrame );
 
 	/** Assigns the viewport overlay widget to use for this viewport client.  Should only be called when first created */
-	void SetViewportOverlayWidget( TSharedPtr< class SWindow > InWindow, TSharedRef<class SOverlay> InViewportOverlayWidget )
+	void SetViewportOverlayWidget( TSharedPtr< SWindow > InWindow, TSharedRef<SOverlay> InViewportOverlayWidget )
 	{
 		Window = InWindow;
 		ViewportOverlayWidget = InViewportOverlayWidget;
 	}
 
 	/** Returns access to this viewport's Slate window */
-	TSharedPtr< class SWindow > GetWindow()
+	TSharedPtr< SWindow > GetWindow()
 	{
 	 	 return Window.Pin();
 	}
 	 
-	/** sets bDropDetail and other per-frame detail level flags on the current WorldSettings
+	/** 
+ 	 * Sets bDropDetail and other per-frame detail level flags on the current WorldSettings
+	 *
 	 * @param DeltaSeconds - amount of time passed since last tick
+	 * @see UWorld
 	 */
 	virtual void SetDropDetail(float DeltaSeconds);
 
@@ -441,7 +317,7 @@ public:
 
 	/**
 	 * Sets the value of ActiveSplitscreenConfiguration based on the desired split-screen layout type, current number of players, and any other
-	 * factors that might affect the way the screen should be layed out.
+	 * factors that might affect the way the screen should be laid out.
 	 */
 	virtual void UpdateActiveSplitscreenType();
 
@@ -459,9 +335,10 @@ public:
 
 	/**
 	* Convert a LocalPlayer to it's index in the GamePlayer array
-	* Returns -1 if the index could not be found.
+	* @param LPlayer Player to get the index of
+	* @returns -1 if the index could not be found.
 	*/
-	int32 ConvertLocalPlayerToGamePlayerIndex( class ULocalPlayer* LPlayer );
+	int32 ConvertLocalPlayerToGamePlayerIndex( ULocalPlayer* LPlayer );
 
 	/** Whether the player at LocalPlayerIndex's viewport has a "top of viewport" safezone or not. */
 	bool HasTopSafeZone( int32 LocalPlayerIndex );
@@ -479,35 +356,46 @@ public:
 	* Get the total pixel size of the screen.
 	* This is different from the pixel size of the viewport since we could be in splitscreen
 	*/
-	void GetPixelSizeOfScreen( float& Width, float& Height, class UCanvas* Canvas, int32 LocalPlayerIndex );
+	void GetPixelSizeOfScreen( float& Width, float& Height, UCanvas* Canvas, int32 LocalPlayerIndex );
 
 	/** Calculate the amount of safezone needed for a single side for both vertical and horizontal dimensions*/
-	void CalculateSafeZoneValues( float& Horizontal, float& Vertical, class UCanvas* Canvas, int32 LocalPlayerIndex, bool bUseMaxPercent );
+	void CalculateSafeZoneValues( float& Horizontal, float& Vertical, UCanvas* Canvas, int32 LocalPlayerIndex, bool bUseMaxPercent );
 
 	/**
 	* pixel size of the deadzone for all sides (right/left/top/bottom) based on which local player it is
 	* @return true if the safe zone exists
 	*/
-	bool CalculateDeadZoneForAllSides( class ULocalPlayer* LPlayer, class UCanvas* Canvas, float& fTopSafeZone, float& fBottomSafeZone, float& fLeftSafeZone, float& fRightSafeZone, bool bUseMaxPercent = false );
+	bool CalculateDeadZoneForAllSides( ULocalPlayer* LPlayer, UCanvas* Canvas, float& fTopSafeZone, float& fBottomSafeZone, float& fLeftSafeZone, float& fRightSafeZone, bool bUseMaxPercent = false );
 
-	/**  Draw the safe area using the current TitleSafeZone settings. */
-	virtual void DrawTitleSafeArea( class UCanvas* Canvas );
+	/**  
+	 * Draw the safe area using the current TitleSafeZone settings. 
+	 * 
+	 * @param Canvas	Canvas on which to draw
+	 */
+	virtual void DrawTitleSafeArea( UCanvas* Canvas );
 
 	/**
 	 * Called after rendering the player views and HUDs to render menus, the console, etc.
 	 * This is the last rendering call in the render loop
-	 * @param Canvas - The canvas to use for rendering.
+	 *
+	 * @param Canvas	The canvas to use for rendering.
 	 */
-	virtual void PostRender(class UCanvas* Canvas);
+	virtual void PostRender( UCanvas* Canvas );
 
 	/**
 	 * Displays the transition screen.
-	 * @param Canvas - The canvas to use for rendering.
+	 *
+	 * @param Canvas	The canvas to use for rendering.
 	 */
-	virtual void DrawTransition(class UCanvas* Canvas);
+	virtual void DrawTransition( UCanvas* Canvas );
 
-	/** Print a centered transition message with a drop shadow. */
-	virtual void DrawTransitionMessage(class UCanvas* Canvas,const FString& Message);
+	/** 
+	 * Print a centered transition message with a drop shadow. 
+	 * 
+	 * @param Canvas	The canvas to use for rendering.
+	 * @param Message	Transition message
+	 */
+	virtual void DrawTransitionMessage( UCanvas* Canvas, const FString& Message );
 
 	/**
 	 * Notifies all interactions that a new player has been added to the list of active players.
@@ -617,6 +505,8 @@ public:
 
 	/**
 	 * Sets all the stats that should be enabled for the viewport
+	 *
+	 * @param InEnabledStats	Stats to enable
 	 */
 	virtual void SetEnabledStats(const TArray<FString>& InEnabledStats) override
 	{
@@ -625,6 +515,8 @@ public:
 
 	/**
 	 * Check whether a specific stat is enabled for this viewport
+	 *
+	 * @param	InName	Name of the stat to check
 	 */
 	virtual bool IsStatEnabled(const TCHAR* InName) const override
 	{
@@ -734,24 +626,25 @@ private:
 
 private:
 	/** Slate window associated with this viewport client.  The same window may host more than one viewport client. */
-	TWeakPtr<class SWindow> Window;
+	TWeakPtr<SWindow> Window;
 
 	/** Overlay widget that contains widgets to draw on top of the game viewport */
-	TWeakPtr<class SOverlay> ViewportOverlayWidget;
+	TWeakPtr<SOverlay> ViewportOverlayWidget;
 
 	/** Current buffer visualization mode for this game viewport */
 	FName CurrentBufferVisualizationMode;
 
 	/** Weak pointer to the highres screenshot dialog if it's open */
-	TWeakPtr<class SWindow> HighResScreenshotDialog;
+	TWeakPtr<SWindow> HighResScreenshotDialog;
 
-	// Function that handles bug screen-shot requests w/ or w/o extra HUD info (project-specific)
+	/* Function that handles bug screen-shot requests w/ or w/o extra HUD info (project-specific) */
 	bool RequestBugScreenShot(const TCHAR* Cmd, bool bDisplayHUDInfo);
 
-	/** Applies requested changes to display configuration 
-	* @param	Dimensions - Pointer to new dimensions of the display. NULL for no change.
-	* @param	WindowMode - What window mode do we want to st the display to.
-	*/
+	/** 
+	 * Applies requested changes to display configuration 
+	 * @param	Dimensions - Pointer to new dimensions of the display. NULL for no change.
+	 * @param	WindowMode - What window mode do we want to st the display to.
+	 */
 	bool SetDisplayConfiguration( const FIntPoint* Dimensions, EWindowMode::Type WindowMode);
 
 	/** Delegate called at the end of the frame when a screenshot is captured and a .png is requested */
diff --git a/Engine/Source/Runtime/Engine/Classes/Engine/GameViewportDelegates.h b/Engine/Source/Runtime/Engine/Classes/Engine/GameViewportDelegates.h
new file mode 100644
--- /dev/null
+++ b/Engine/Source/Runtime/Engine/Classes/Engine/GameViewportDelegates.h
@@ -0,0 +1,32 @@
+// Copyright 1998-2014 Epic Games, Inc. All Rights Reserved.
+
+#pragma once
+
+/**
+ * Delegate type used by UGameViewportClient when a screenshot has been captured
+ *
+ * The first parameter is the width.
+ * The second parameter is the height.
+ * The third parameter is the array of bitmap data.
+ * @see UGameViewportClient
+ */
+DECLARE_MULTICAST_DELEGATE_ThreeParams(FOnScreenshotCaptured, int32 /*Width*/, int32 /*Height*/, const TArray<FColor>& /*Colors*/);
+
+/**
+ * Delegate type used by UGameViewportClient when a png screenshot has been captured
+ *
+ * The first parameter is the width.
+ * The second parameter is the height.
+ * The third parameter is the array of bitmap data.
+ * The fourth parameter is the screen shot filename.
+ * @see UGameViewportClient
+ */
+DECLARE_DELEGATE_FourParams(FOnPNGScreenshotCaptured, int32, int32, const TArray<FColor>&, const FString&);
+
+/**
+ * Delegate type used by UGameViewportClient when call is made to close a viewport
+ *
+ * The first parameter is the viewport being closed.
+ * @see UGameViewportClient
+ */
+DECLARE_DELEGATE_OneParam(FOnCloseRequested, FViewport*);
diff --git a/Engine/Source/Runtime/Engine/Classes/Engine/TitleSafeZone.h b/Engine/Source/Runtime/Engine/Classes/Engine/TitleSafeZone.h
new file mode 100644
--- /dev/null
+++ b/Engine/Source/Runtime/Engine/Classes/Engine/TitleSafeZone.h
@@ -0,0 +1,41 @@
+// Copyright 1998-2014 Epic Games, Inc. All Rights Reserved.
+
+#pragma once
+
+/** 
+ * Safezone related declarations. Used by GameViewportClient for rendering 'safe zones' on screen.
+ *
+ * @see UGameViewportClient 
+ * @see ESafeZoneType
+ * @see FTitleSafeZoneArea 
+ */
+
+/**
+ * The 4 different kinds of safezones
+ */
+enum ESafeZoneType
+{
+	eSZ_TOP,
+	eSZ_BOTTOM,
+	eSZ_LEFT,
+	eSZ_RIGHT,
+	eSZ_MAX,
+};
+
+/** Max/Recommended screen viewable extents as a percentage */
+struct FTitleSafeZoneArea
+{
+	float MaxPercentX;
+	float MaxPercentY;
+	float RecommendedPercentX;
+	float RecommendedPercentY;
+
+	FTitleSafeZoneArea()
+		: MaxPercentX(0)
+		, MaxPercentY(0)
+		, RecommendedPercentX(0)
+		, RecommendedPercentY(0)
+	{
+	}
+
+};
\ No newline at end of file
diff --git a/Engine/Source/Runtime/Engine/Classes/Engine/ViewportSplitScreen.h b/Engine/Source/Runtime/Engine/Classes/Engine/ViewportSplitScreen.h
new file mode 100644
--- /dev/null
+++ b/Engine/Source/Runtime/Engine/Classes/Engine/ViewportSplitScreen.h
@@ -0,0 +1,69 @@
+// Copyright 1998-2014 Epic Games, Inc. All Rights Reserved.
+
+#pragma once
+
+/** 
+ * Various data declarations relating to split screen on a GameViewportClient
+ *
+ * @see ESplitScreenType
+ * @see FPerPlayerSplitscreenData
+ * @see FSplitscreenData
+ */
+ 
+/**
+ * Enum of the different splitscreen types
+ */
+namespace ESplitScreenType
+{
+	enum Type
+	{
+		// No split
+		None,
+		// 2 player horizontal split
+		TwoPlayer_Horizontal,
+		// 2 player vertical split
+		TwoPlayer_Vertical,
+		// 3 Player split with 1 player on top and 2 on bottom
+		ThreePlayer_FavorTop,
+		// 3 Player split with 1 player on bottom and 2 on top
+		ThreePlayer_FavorBottom,
+		// 4 Player split
+		FourPlayer,
+
+		SplitTypeCount
+	};
+}
+
+/** Structure to store splitscreen data. */
+struct FPerPlayerSplitscreenData
+{
+	float SizeX;
+	float SizeY;
+	float OriginX;
+	float OriginY;
+
+
+	FPerPlayerSplitscreenData()
+		: SizeX(0)
+		, SizeY(0)
+		, OriginX(0)
+		, OriginY(0)
+	{
+	}
+
+	FPerPlayerSplitscreenData(float NewSizeX, float NewSizeY, float NewOriginX, float NewOriginY)
+		: SizeX(NewSizeX)
+		, SizeY(NewSizeY)
+		, OriginX(NewOriginX)
+		, OriginY(NewOriginY)
+	{
+	}
+
+};
+
+/** Structure containing all the player splitscreen datas per splitscreen configuration. */
+struct FSplitscreenData
+{
+	TArray<struct FPerPlayerSplitscreenData> PlayerData;
+};
+
diff --git a/Engine/Source/Runtime/Engine/Private/GameViewportClient.cpp b/Engine/Source/Runtime/Engine/Private/GameViewportClient.cpp
--- a/Engine/Source/Runtime/Engine/Private/GameViewportClient.cpp
+++ b/Engine/Source/Runtime/Engine/Private/GameViewportClient.cpp
@@ -95,7 +95,6 @@ UGameViewportClient::UGameViewportClient(const FObjectInitializer& ObjectInitial
 	TitleSafeZone.RecommendedPercentY = 0.8f;
 
 	bIsPlayInEditorViewport = false;
-	ProgressFadeTime = 1.0f;
 	ViewModeIndex = VMI_Lit;
 
 	SplitscreenInfo.Init(FSplitscreenData(), ESplitScreenType::SplitTypeCount);
