special area marking for too low navigable spaces is now removed after applying static modifiers to reduce amount of polys on slopes / inside terrain geometry

diff --git a/Engine/Source/Runtime/Engine/Private/AI/Navigation/RecastNavMeshGenerator.cpp b/Engine/Source/Runtime/Engine/Private/AI/Navigation/RecastNavMeshGenerator.cpp
--- a/Engine/Source/Runtime/Engine/Private/AI/Navigation/RecastNavMeshGenerator.cpp
+++ b/Engine/Source/Runtime/Engine/Private/AI/Navigation/RecastNavMeshGenerator.cpp
@@ -1989,6 +1989,12 @@ bool FRecastTileGenerator::GenerateCompressedLayers(FNavMeshBuildContext& BuildC
 	// (Optional) Mark areas.
 	MarkStaticAreas(BuildContext, *RasterContext.CompactHF);
 
+	// remove all low area marking at this point
+	if (TileConfig.bMarkLowHeightAreas)
+	{
+		rcReplaceBoxArea(&BuildContext, TileConfig.bmin, TileConfig.bmax, RECAST_NULL_AREA, RECAST_LOW_AREA, *RasterContext.CompactHF);
+	}
+
 	// Build layers
 	{
 		RECAST_STAT(STAT_Navigation_Async_Recast_Layers);
