Adde Visual Log's device to save log files in binary format. It's disabled for now.
Minor optymizations for Visual Logger.
Fixed names for vlog files saved in json format.

diff --git a/Engine/Config/BaseEngine.ini b/Engine/Config/BaseEngine.ini
--- a/Engine/Config/BaseEngine.ini
+++ b/Engine/Config/BaseEngine.ini
@@ -2322,3 +2322,7 @@ AISystemClassName=/Script/AIModule.AISystem
 bTrackHistory=false
 NumberOfHistoryItemsTracked=5
 
+[VisualLogger]
+UseBinaryFileDevice=false
+FrameCacheLenght=1.0f ;in seconds, to batch log data between file serializations
+UseCompression=false ;works only with binary files
diff --git a/Engine/Source/Runtime/Engine/Classes/VisualLogger/VisualLoggerBinaryFileDevice.h b/Engine/Source/Runtime/Engine/Classes/VisualLogger/VisualLoggerBinaryFileDevice.h
new file mode 100644
--- /dev/null
+++ b/Engine/Source/Runtime/Engine/Classes/VisualLogger/VisualLoggerBinaryFileDevice.h
@@ -0,0 +1,40 @@
+// Copyright 1998-2014 Epic Games, Inc. All Rights Reserved.
+#pragma once
+#include "VisualLogger/VisualLogger.h"
+
+#if ENABLE_VISUAL_LOG
+class FVisualLoggerBinaryFileDevice : public FVisualLogDevice
+{
+public:
+	static FVisualLoggerBinaryFileDevice& Get()
+	{
+		static FVisualLoggerBinaryFileDevice GDevice;
+		return GDevice;
+	}
+
+	FVisualLoggerBinaryFileDevice();
+	virtual void Cleanup(bool bReleaseMemory = false) override;
+	virtual void StartRecordingToFile(float TImeStamp) override;
+	virtual void StopRecordingToFile(float TImeStamp) override;
+	virtual void SetFileName(const FString& InFileName) override;
+	virtual void Serialize(const class UObject* LogOwner, FName OwnerName, const FVisualLogEntry& LogEntry) override;
+
+	struct FFrameCacheItem
+	{
+		FFrameCacheItem() {}
+		FFrameCacheItem(FName InOwnerName, const FVisualLogEntry& LogEntry) : OwnerName(OwnerName), Entry(LogEntry) { }
+
+		FName OwnerName;
+		FVisualLogEntry Entry;
+	};
+protected:
+	int32 bUseCompression : 1;
+	float FrameCacheLenght;
+	float StartRecordingTime;
+	float LastLogTimeStamp;
+	FArchive* FileArchive;
+	FString TempFileName;
+	FString FileName;
+	TArray<FFrameCacheItem> FrameCache;
+};
+#endif
diff --git a/Engine/Source/Runtime/Engine/Private/VisualLog.cpp b/Engine/Source/Runtime/Engine/Private/VisualLog.cpp
--- a/Engine/Source/Runtime/Engine/Private/VisualLog.cpp
+++ b/Engine/Source/Runtime/Engine/Private/VisualLog.cpp
@@ -461,7 +461,7 @@ void FVisualLog::StopRecordingToFile(float TImeStamp)
 
 void FVisualLog::SetFileName(const FString& InFileName)
 {
-	BaseFileName = InFileName;
+	FileName = InFileName;
 }
 
 void FVisualLog::Serialize(const class UObject* LogOwner, FName OwnerName, const FVisualLogEntry& LogEntry)
@@ -478,11 +478,13 @@ void FVisualLog::Serialize(const class UObject* LogOwner, FName OwnerName, const
 	{
 		OnNewLogCreated.ExecuteIfBound(LogOwner, Log);
 	}
+
+	LastLogTimeStamp = LogEntry.TimeStamp;
 }
 
 FString FVisualLog::GetLogFileFullName() const
 {
-	FString FileName = FString::Printf(TEXT("%s_%.0f-%.0f_%s.vlog"), *BaseFileName, StartRecordingTime, GWorld ? GWorld->TimeSeconds : 0, *FDateTime::Now().ToString());
+	FString FileName = FString::Printf(TEXT("%s_%.0f-%.0f_%s.vlog"), *FileName, StartRecordingTime, GWorld ? GWorld->TimeSeconds : 0, *FDateTime::Now().ToString());
 
 	const FString FullFilename = FString::Printf(TEXT("%slogs/%s"), *FPaths::GameSavedDir(), *FileName);
 
@@ -546,9 +548,25 @@ void FVisualLog::SetIsRecording(bool NewRecording, bool bRecordToFile)
 			const FString HeadetStr = TEXT("]}");
 			auto AnsiAdditionalData = StringCast<UCS2CHAR>(*HeadetStr);
 			FileAr->Serialize((UCS2CHAR*)AnsiAdditionalData.Get(), HeadetStr.Len() * sizeof(UCS2CHAR));
+
+			const int64 TotalSize = FileAr->TotalSize();
 			FileAr->Close();
 			delete FileAr;
 			FileAr = NULL;
+
+			const FString TempFullFilename = FString::Printf(TEXT("%slogs/%s"), *FPaths::GameSavedDir(), *TempFileName);
+			const FString NewFileName = FVisualLoggerHelpers::GenerateFilename(TempFileName, FileName, StartRecordingTime, LastLogTimeStamp);
+
+			if (TotalSize > 0)
+			{
+				// rename file when we serialized some data
+				IFileManager::Get().Move(*NewFileName, *TempFullFilename, true);
+			}
+			else
+			{
+				// or remove file if nothing serialized
+				IFileManager::Get().Delete(*TempFullFilename, false, true, true);
+			}
 		}
 
 		Cleanup(true);
@@ -559,10 +577,13 @@ void FVisualLog::SetIsRecording(bool NewRecording, bool bRecordToFile)
 	{
 		bIsRecordingToFile = bRecordToFile;
 		StartRecordingTime = GWorld ? GWorld->TimeSeconds : 0.0f;
+		LastLogTimeStamp = StartRecordingTime;
 
 		if (!FileAr)
 		{
-			FileAr = IFileManager::Get().CreateFileWriter(*GetLogFileFullName());
+			TempFileName = FVisualLoggerHelpers::GenerateTemporaryFilename(TEXT("vlog"));
+			const FString FullFilename = FString::Printf(TEXT("%slogs/%s"), *FPaths::GameSavedDir(), *TempFileName);
+			FileAr = IFileManager::Get().CreateFileWriter(*FullFilename);
 
 			const FString HeadetStr = TEXT("{\"Logs\":[{}");
 			auto AnsiAdditionalData = StringCast<UCS2CHAR>(*HeadetStr);
diff --git a/Engine/Source/Runtime/Engine/Private/VisualLogger/VisualLogger.cpp b/Engine/Source/Runtime/Engine/Private/VisualLogger/VisualLogger.cpp
--- a/Engine/Source/Runtime/Engine/Private/VisualLogger/VisualLogger.cpp
+++ b/Engine/Source/Runtime/Engine/Private/VisualLogger/VisualLogger.cpp
@@ -3,6 +3,7 @@
 #include "EnginePrivate.h"
 #include "VisualLog.h"
 #include "VisualLogger/VisualLogger.h"
+#include "VisualLogger/VisualLoggerBinaryFileDevice.h"
 
 #if ENABLE_VISUAL_LOG 
 
@@ -12,8 +13,18 @@ TMap<class UObject*, TArray<TWeakObjectPtr<const class UObject> > > FVisualLogge
 
 FVisualLogger::FVisualLogger()
 {
+	bool UseBinaryFileDevice = false;
+	GConfig->GetBool(TEXT("VisualLogger"), TEXT("UseBinaryFileDevice"), UseBinaryFileDevice, GEngineIni);
+
 	BlockAllCategories(false);
-	AddDevice(&FVisualLog::GetStatic());
+	if (UseBinaryFileDevice)
+	{
+		AddDevice(&FVisualLoggerBinaryFileDevice::Get());
+	}
+	else
+	{
+		AddDevice(&FVisualLog::GetStatic());
+	}
 	SetIsRecording(GEngine ? !!GEngine->bEnableVisualLogRecordingOnStart : false);
 	SetIsRecordingOnServer(false);
 
@@ -26,6 +37,16 @@ FVisualLogger::FVisualLogger()
 
 FVisualLogger::~FVisualLogger()
 {
+	bool UseBinaryFileDevice = false;
+	GConfig->GetBool(TEXT("VisualLogger"), TEXT("UseBinaryFileDevice"), UseBinaryFileDevice, GEngineIni);
+	if (UseBinaryFileDevice)
+	{
+		RemoveDevice(&FVisualLoggerBinaryFileDevice::Get());
+	}
+	else
+	{
+		RemoveDevice(&FVisualLoggerBinaryFileDevice::Get());
+	}
 	RemoveDevice(&FVisualLog::GetStatic());
 }
 
@@ -117,24 +138,10 @@ void FVisualLogger::SetIsRecordingToFile(bool InIsRecording)
 		SetIsRecording(true);
 	}
 
-	FString BaseFileName = TEXT("VisualLog");
-	FString UserDefinedName;
-	FString MapName = GWorld.GetReference() ? GWorld->GetMapName() : TEXT("");
-	if (LogFileNameGetter.IsBound())
-	{
-		UserDefinedName = LogFileNameGetter.Execute();
-	}
-	if (UserDefinedName.Len() > 0)
-	{
-		if (MapName.Len() > 0)
-		{
-			BaseFileName = FString::Printf(TEXT("%s_%s"), *UserDefinedName, *MapName);
-		}
-		else
-		{
-			BaseFileName = UserDefinedName;
-		}
-	}
+	const FString BaseFileName = LogFileNameGetter.IsBound() ? LogFileNameGetter.Execute() : TEXT("VisualLog");
+	const FString MapName = GWorld.GetReference() ? GWorld->GetMapName() : TEXT("");
+
+	FString OutputFileName = FString::Printf(TEXT("%s_%s"), *BaseFileName, *MapName);
 
 	if (bIsRecordingToFile && !InIsRecording)
 	{
@@ -142,6 +149,7 @@ void FVisualLogger::SetIsRecordingToFile(bool InIsRecording)
 		{
 			if (Device->HasFlags(VisualLogger::CanSaveToFile))
 			{
+				Device->SetFileName(OutputFileName);
 				Device->StopRecordingToFile(GWorld.GetReference() ? GWorld->TimeSeconds : StartRecordingToFileTime);
 			}
 		}
@@ -153,7 +161,6 @@ void FVisualLogger::SetIsRecordingToFile(bool InIsRecording)
 		{
 			if (Device->HasFlags(VisualLogger::CanSaveToFile))
 			{
-				Device->SetFileName(BaseFileName);
 				Device->StartRecordingToFile(StartRecordingToFileTime);
 			}
 		}
diff --git a/Engine/Source/Runtime/Engine/Private/VisualLogger/VisualLoggerAutomationTests.cpp b/Engine/Source/Runtime/Engine/Private/VisualLogger/VisualLoggerAutomationTests.cpp
--- a/Engine/Source/Runtime/Engine/Private/VisualLogger/VisualLoggerAutomationTests.cpp
+++ b/Engine/Source/Runtime/Engine/Private/VisualLogger/VisualLoggerAutomationTests.cpp
@@ -5,6 +5,7 @@
 #include "VisualLog.h"
 #include "VisualLogger/VisualLogger.h"
 #include "VisualLogger/VisualLoggerAutomationTests.h"
+#include "VisualLogger/VisualLoggerBinaryFileDevice.h"
 
 UVisualLoggerAutomationTests::UVisualLoggerAutomationTests(const FObjectInitializer& ObjectInitializer) : Super(ObjectInitializer)
 {
@@ -73,6 +74,7 @@ if ((__Test__)) \
 	return false; \
 }
 
+template<typename TYPE = FVisualLoggerTestDevice>
 struct FTestDeviceContext
 {
 	FTestDeviceContext() 
@@ -91,7 +93,7 @@ struct FTestDeviceContext
 		Device.Cleanup();
 	}
 
-	FVisualLoggerTestDevice Device;
+	TYPE Device;
 };
 
 IMPLEMENT_SIMPLE_AUTOMATION_TEST(FVisualLogTest, "Engine.VisualLogger.Logging simple text", EAutomationTestFlags::ATF_Game | EAutomationTestFlags::ATF_Editor)
@@ -104,7 +106,7 @@ IMPLEMENT_SIMPLE_AUTOMATION_TEST(FVisualLogTest, "Engine.VisualLogger.Logging si
 
 bool FVisualLogTest::RunTest(const FString& Parameters)
 {
-	FTestDeviceContext Context;
+	FTestDeviceContext<FVisualLoggerTestDevice> Context;
 
 	FVisualLogger::Get().SetIsRecording(false);
 	CHECK_FAIL(FVisualLogger::Get().IsRecording());
@@ -146,7 +148,7 @@ bool FVisualLogTest::RunTest(const FString& Parameters)
 IMPLEMENT_SIMPLE_AUTOMATION_TEST(FVisualLogSegmentsTest, "Engine.VisualLogger.Logging segment shape", EAutomationTestFlags::ATF_Game | EAutomationTestFlags::ATF_Editor)
 bool FVisualLogSegmentsTest::RunTest(const FString& Parameters)
 {
-	FTestDeviceContext Context;
+	FTestDeviceContext<FVisualLoggerTestDevice> Context;
 
 	FVisualLogger::Get().SetIsRecording(false);
 	CHECK_FAIL(FVisualLogger::Get().IsRecording());
@@ -183,30 +185,26 @@ bool FVisualLogSegmentsTest::RunTest(const FString& Parameters)
 	return true;
 }
 
-IMPLEMENT_SIMPLE_AUTOMATION_TEST(FVisualLogEventsTest, "Engine.VisualLogger.Logging events", EAutomationTestFlags::ATF_Game | EAutomationTestFlags::ATF_Editor)
-
-DECLARE_VLOG_EVENT(EventTest, All, "Simple event for tests")
-DEFINE_VLOG_EVENT(EventTest)
+DEFINE_VLOG_EVENT(EventTest, All, "Simple event for tests")
 
-DECLARE_VLOG_EVENT(EventTest2, All, "Second simple event for tests")
-DEFINE_VLOG_EVENT(EventTest2)
+DEFINE_VLOG_EVENT(EventTest2, All, "Second simple event for tests")
 
-DECLARE_VLOG_EVENT(EventTest3, All, "Third simple event for tests")
-DEFINE_VLOG_EVENT(EventTest3)
+DEFINE_VLOG_EVENT(EventTest3, All, "Third simple event for tests")
 
+IMPLEMENT_SIMPLE_AUTOMATION_TEST(FVisualLogEventsTest, "Engine.VisualLogger.Logging events", EAutomationTestFlags::ATF_Game | EAutomationTestFlags::ATF_Editor)
 bool FVisualLogEventsTest::RunTest(const FString& Parameters)
 {
-	FTestDeviceContext Context;
+	FTestDeviceContext<FVisualLoggerTestDevice> Context;
 	FVisualLogger::Get().SetIsRecording(true);
 
-	CHECK_SUCCESS(EventTest.GetName() == TEXT("EventTest"));
-	CHECK_SUCCESS(EventTest.GetUserFriendlyDesc() == TEXT("Simple event for tests"));
+	CHECK_SUCCESS(EventTest.Name == TEXT("EventTest"));
+	CHECK_SUCCESS(EventTest.FriendlyDesc == TEXT("Simple event for tests"));
 
-	CHECK_SUCCESS(EventTest2.GetName() == TEXT("EventTest2"));
-	CHECK_SUCCESS(EventTest2.GetUserFriendlyDesc() == TEXT("Second simple event for tests"));
+	CHECK_SUCCESS(EventTest2.Name == TEXT("EventTest2"));
+	CHECK_SUCCESS(EventTest2.FriendlyDesc == TEXT("Second simple event for tests"));
 
-	CHECK_SUCCESS(EventTest3.GetName() == TEXT("EventTest3"));
-	CHECK_SUCCESS(EventTest3.GetUserFriendlyDesc() == TEXT("Third simple event for tests"));
+	CHECK_SUCCESS(EventTest3.Name == TEXT("EventTest3"));
+	CHECK_SUCCESS(EventTest3.FriendlyDesc == TEXT("Third simple event for tests"));
 
 	FVisualLogEntry* CurrentEntry = FVisualLogger::Get().GetEntryToWrite(GWorld, GWorld->TimeSeconds, VisualLogger::DontCreate);
 	float CurrentTimestamp = GWorld->TimeSeconds;
diff --git a/Engine/Source/Runtime/Engine/Private/VisualLogger/VisualLoggerBinaryFileDevice.cpp b/Engine/Source/Runtime/Engine/Private/VisualLogger/VisualLoggerBinaryFileDevice.cpp
new file mode 100644
--- /dev/null
+++ b/Engine/Source/Runtime/Engine/Private/VisualLogger/VisualLoggerBinaryFileDevice.cpp
@@ -0,0 +1,104 @@
+// Copyright 1998-2014 Epic Games, Inc. All Rights Reserved.
+#include "EnginePrivate.h"
+#include "VisualLogger/VisualLoggerBinaryFileDevice.h"
+
+#if ENABLE_VISUAL_LOG
+FArchive& operator<<(FArchive& Ar, FVisualLoggerBinaryFileDevice::FFrameCacheItem& FrameCacheItem)
+{
+	Ar << FrameCacheItem.OwnerName;
+	Ar << FrameCacheItem.Entry;
+	return Ar;
+}
+
+FVisualLoggerBinaryFileDevice::FVisualLoggerBinaryFileDevice()
+	: FileArchive(NULL)
+{
+	Cleanup();
+
+	bool DefaultFrameCacheLenght = 0;
+	GConfig->GetBool(TEXT("VisualLogger"), TEXT("FrameCacheLenght"), DefaultFrameCacheLenght, GEngineIni);
+	FrameCacheLenght = DefaultFrameCacheLenght;
+
+	bool UseCompression = false;
+	GConfig->GetBool(TEXT("VisualLogger"), TEXT("UseCompression"), UseCompression, GEngineIni);
+	bUseCompression = UseCompression;
+}
+
+void FVisualLoggerBinaryFileDevice::Cleanup(bool bReleaseMemory)
+{
+
+}
+
+void FVisualLoggerBinaryFileDevice::StartRecordingToFile(float TimeStamp)
+{
+	if (FileArchive != NULL)
+	{
+		return;
+	}
+
+	StartRecordingTime = TimeStamp;
+	LastLogTimeStamp = StartRecordingTime;
+	TempFileName = FVisualLoggerHelpers::GenerateTemporaryFilename(TEXT("bvlog"));
+	const FString FullFilename = FString::Printf(TEXT("%slogs/%s"), *FPaths::GameSavedDir(), *TempFileName);
+
+	FileArchive = IFileManager::Get().CreateFileWriter(*FullFilename);
+}
+
+void FVisualLoggerBinaryFileDevice::StopRecordingToFile(float TimeStamp)
+{
+	if (FileArchive == NULL)
+	{
+		return;
+	}
+
+	const int32 NumEntries = FrameCache.Num();
+	if (NumEntries> 0)
+	{
+		*FileArchive << FrameCache;
+		FrameCache.Reset();
+	}
+
+	const int64 TotalSize = FileArchive->TotalSize();
+	FileArchive->Close();
+	delete FileArchive;
+	FileArchive = NULL;
+
+	const FString TempFullFilename = FString::Printf(TEXT("%slogs/%s"), *FPaths::GameSavedDir(), *TempFileName);
+	const FString NewFileName = FVisualLoggerHelpers::GenerateFilename(TempFileName, FileName, StartRecordingTime, LastLogTimeStamp);
+
+	if (TotalSize > 0)
+	{
+		// rename file when we serialized some data
+		IFileManager::Get().Move(*NewFileName, *TempFullFilename, true);
+	}
+	else
+	{
+		// or remove file if nothing serialized
+		IFileManager::Get().Delete(*TempFullFilename, false, true, true);
+	}
+}
+
+void FVisualLoggerBinaryFileDevice::SetFileName(const FString& InFileName)
+{
+	FileName = InFileName;
+}
+
+void FVisualLoggerBinaryFileDevice::Serialize(const class UObject* LogOwner, FName OwnerName, const FVisualLogEntry& LogEntry)
+{
+	const int32 NumEntries = FrameCache.Num();
+	if (NumEntries> 0 && FrameCache[NumEntries - 1].Entry.TimeStamp + FrameCacheLenght < LogEntry.TimeStamp && FileArchive)
+	{
+		int32 FrameTag = 0xFAFA;
+		*FileArchive << FrameTag;
+
+		int32 FrameCacheNum = FrameCache.Num();
+		*FileArchive << FrameCacheNum;
+
+		*FileArchive << FrameCache;
+		FrameCache.Reset();
+	}
+
+	LastLogTimeStamp = LogEntry.TimeStamp;
+	FrameCache.Add(FFrameCacheItem(OwnerName, LogEntry));
+}
+#endif
diff --git a/Engine/Source/Runtime/Engine/Public/VisualLog.h b/Engine/Source/Runtime/Engine/Public/VisualLog.h
--- a/Engine/Source/Runtime/Engine/Public/VisualLog.h
+++ b/Engine/Source/Runtime/Engine/Public/VisualLog.h
@@ -166,8 +166,9 @@ private:
 	FOnNewLogCreatedDelegate OnNewLogCreated;
 
 	float StartRecordingTime;
-
-	FString BaseFileName;
+	float LastLogTimeStamp;
+	FString TempFileName;
+	FString FileName;
 
 	int32 bIsRecordingOnServer : 1;
 	int32 bIsRecordingToFile : 1;
diff --git a/Engine/Source/Runtime/Engine/Public/VisualLogger/VisualLogger.h b/Engine/Source/Runtime/Engine/Public/VisualLogger/VisualLogger.h
--- a/Engine/Source/Runtime/Engine/Public/VisualLogger/VisualLogger.h
+++ b/Engine/Source/Runtime/Engine/Public/VisualLogger/VisualLogger.h
@@ -31,25 +31,21 @@
 #define UE_CVLOG_CONE(Condition, LogOwner, CategoryName, Verbosity, Orgin, Direction, Length, Angle, Color, Format, ...)  if(Condition) {UE_VLOG_CONE(Actor, CategoryName, Verbosity, Orgin, Direction, Length, Angle, Color, Format, ##__VA_ARGS__);} 
 // Cylinder shape
 #define UE_VLOG_CYLINDER(LogOwner, CategoryName, Verbosity, Start, End, Radius, Color, Format, ...) FVisualLogger::GeometryShapeLogf(LogOwner, CategoryName, ELogVerbosity::Verbosity, INDEX_NONE, Start, End, Radius, Color, Format, ##__VA_ARGS__)
-#define UE_CVLOG_CYLINDER(LogOwner, CategoryName, Verbosity, Start, End, Radius, Color, Format, ...)  if(Condition) {UE_VLOG_CYLINDER(LogOwner, CategoryName, Verbosity, Start, End, Radius, Color, Format, ##__VA_ARGS__);} 
+#define UE_CVLOG_CYLINDER(Condition, LogOwner, CategoryName, Verbosity, Start, End, Radius, Color, Format, ...)  if(Condition) {UE_VLOG_CYLINDER(LogOwner, CategoryName, Verbosity, Start, End, Radius, Color, Format, ##__VA_ARGS__);} 
 // Capsule shape
 #define UE_VLOG_CAPSULE(LogOwner, CategoryName, Verbosity, Center, HalfHeight, Radius, Rotation, Color, Format, ...) FVisualLogger::GeometryShapeLogf(LogOwner, CategoryName, ELogVerbosity::Verbosity, INDEX_NONE, Center, HalfHeight, Radius, Rotation, Color, Format, ##__VA_ARGS__)
-#define UE_CVLOG_CAPSULE(LogOwner, CategoryName, Verbosity, Center, HalfHeight, Radius, Rotation, Color, Format, ...)  if(Condition) {UE_VLOG_CAPSULE(LogOwner, CategoryName, Verbosity, Center, HalfHeight, Radius, Rotation, Color, Format, ##__VA_ARGS__);} 
+#define UE_CVLOG_CAPSULE(Condition, LogOwner, CategoryName, Verbosity, Center, HalfHeight, Radius, Rotation, Color, Format, ...)  if(Condition) {UE_VLOG_CAPSULE(LogOwner, CategoryName, Verbosity, Center, HalfHeight, Radius, Rotation, Color, Format, ##__VA_ARGS__);} 
 // Histogram data for 2d graphs 
 #define UE_VLOG_HISTOGRAM(LogOwner, CategoryName, Verbosity, GraphName, DataName, Data) FVisualLogger::HistogramDataLogf(LogOwner, CategoryName, ELogVerbosity::Verbosity, INDEX_NONE, GraphName, DataName, Data, FColor::White, TEXT(""))
-#define UE_CVLOG_HISTOGRAM(LogOwner, CategoryName, Verbosity, GraphName, DataName, Data) if(Condition) {UE_VLOG_HISTOGRAM(LogOwner, CategoryName, Verbosity, GraphName, DataName, Data);} 
+#define UE_CVLOG_HISTOGRAM(Condition, LogOwner, CategoryName, Verbosity, GraphName, DataName, Data) if(Condition) {UE_VLOG_HISTOGRAM(LogOwner, CategoryName, Verbosity, GraphName, DataName, Data);} 
+
+#define DECLARE_VLOG_EVENT(EventName) extern FVisualLogEventBase EventName;
+#define DEFINE_VLOG_EVENT(EventName, Verbosity, UserFriendlyDesc) FVisualLogEventBase EventName(TEXT(#EventName), TEXT(UserFriendlyDesc), ELogVerbosity::Verbosity); 
 
 #define UE_VLOG_EVENTS(LogOwner, TagNameToLog, ...) FVisualLogger::EventLog(LogOwner, TagNameToLog, ##__VA_ARGS__)
+#define UE_CVLOG_EVENTS(Condition, LogOwner, TagNameToLog, ...) if(Condition) {UE_VLOG_EVENTS(LogOwner, TagNameToLog, ##__VA_ARGS__);}
 #define UE_VLOG_EVENT_WITH_DATA(LogOwner, LogEvent, ...) FVisualLogger::EventLog(LogOwner, LogEvent, ##__VA_ARGS__)
-#define DECLARE_VLOG_EVENT(EventName, Verbosity, UserFriendlyDesc) \
-	struct FVisualLogEventBase_##EventName : public FVisualLogEventBase \
-	{ \
-		virtual FString GetName() const override { return TEXT(#EventName); } \
-		virtual FString GetUserFriendlyDesc() const override { return TEXT(UserFriendlyDesc); } \
-		virtual ELogVerbosity::Type GetVerbosity() const override { return ELogVerbosity::Verbosity; } \
-	};
-
-#define DEFINE_VLOG_EVENT(EventName) FVisualLogEventBase_##EventName EventName;
+#define UE_CVLOG_EVENT_WITH_DATA(Condition, LogOwner, LogEvent, ...) if(Condition) {UE_VLOG_EVENT_WITH_DATA(LogOwner, LogEvent, ##__VA_ARGS__);}
 
 #else
 #define REDIRECT_TO_VLOG(Dest)
@@ -60,26 +56,29 @@
 #define UE_VLOG(Actor, CategoryName, Verbosity, Format, ...)
 #define UE_CVLOG(Condition, Actor, CategoryName, Verbosity, Format, ...)
 #define UE_VLOG_SEGMENT(Actor, CategoryName, Verbosity, SegmentStart, SegmentEnd, Color, DescriptionFormat, ...)
-#define UE_CVLOG_SEGMENT(Actor, CategoryName, Verbosity, SegmentStart, SegmentEnd, Color, DescriptionFormat, ...)
+#define UE_CVLOG_SEGMENT(Condition, Actor, CategoryName, Verbosity, SegmentStart, SegmentEnd, Color, DescriptionFormat, ...)
 #define UE_VLOG_SEGMENT_THICK(Actor, CategoryName, Verbosity, SegmentStart, SegmentEnd, Color, DescriptionFormat, ...)
-#define UE_CVLOG_SEGMENT_THICK(Actor, CategoryName, Verbosity, SegmentStart, SegmentEnd, Color, DescriptionFormat, ...)
+#define UE_CVLOG_SEGMENT_THICK(Condition, Actor, CategoryName, Verbosity, SegmentStart, SegmentEnd, Color, DescriptionFormat, ...)
 #define UE_VLOG_LOCATION(Actor, CategoryName, Verbosity, Location, Radius, Color, DescriptionFormat, ...)
-#define UE_CVLOG_LOCATION(Actor, CategoryName, Verbosity, Location, Radius, Color, DescriptionFormat, ...)
+#define UE_CVLOG_LOCATION(Condition, Actor, CategoryName, Verbosity, Location, Radius, Color, DescriptionFormat, ...)
 #define UE_VLOG_BOX(Actor, CategoryName, Verbosity, Box, Color, DescriptionFormat, ...) 
-#define UE_CVLOG_BOX(Actor, CategoryName, Verbosity, Box, Color, DescriptionFormat, ...) 
+#define UE_CVLOG_BOX(Condition, Actor, CategoryName, Verbosity, Box, Color, DescriptionFormat, ...) 
 #define UE_VLOG_CONE(Object, CategoryName, Verbosity, Orgin, Direction, Length, Angle, Color, DescriptionFormat, ...)
-#define UE_CVLOG_CONE(Object, CategoryName, Verbosity, Orgin, Direction, Length, Angle, Color, DescriptionFormat, ...)
+#define UE_CVLOG_CONE(Condition, Object, CategoryName, Verbosity, Orgin, Direction, Length, Angle, Color, DescriptionFormat, ...)
 #define UE_VLOG_CYLINDER(Object, CategoryName, Verbosity, Start, End, Radius, Color, DescriptionFormat, ...)
-#define UE_CVLOG_CYLINDER(Object, CategoryName, Verbosity, Start, End, Radius, Color, DescriptionFormat, ...)
+#define UE_CVLOG_CYLINDER(Condition, Object, CategoryName, Verbosity, Start, End, Radius, Color, DescriptionFormat, ...)
 #define UE_VLOG_CAPSULE(Object, CategoryName, Verbosity, Center, HalfHeight, Radius, Rotation, Color, DescriptionFormat, ...)
-#define UE_CVLOG_CAPSULE(Object, CategoryName, Verbosity, Center, HalfHeight, Radius, Rotation, Color, DescriptionFormat, ...)
+#define UE_CVLOG_CAPSULE(Condition, Object, CategoryName, Verbosity, Center, HalfHeight, Radius, Rotation, Color, DescriptionFormat, ...)
 #define UE_VLOG_HISTOGRAM(Actor, CategoryName, Verbosity, GraphName, DataName, Data)
-#define UE_CVLOG_HISTOGRAM(Actor, CategoryName, Verbosity, GraphName, DataName, Data)
+#define UE_CVLOG_HISTOGRAM(Condition, Actor, CategoryName, Verbosity, GraphName, DataName, Data)
 
+#define DECLARE_VLOG_EVENT(EventName)
+#define DEFINE_VLOG_EVENT(EventName, Verbosity, UserFriendlyDesc)
 #define UE_VLOG_EVENTS(LogOwner, TagNameToLog, ...) 
+#define UE_CVLOG_EVENTS(Condition, LogOwner, TagNameToLog, ...) 
 #define UE_VLOG_EVENT_WITH_DATA(LogOwner, LogEvent, ...)
-#define DECLARE_VLOG_EVENT(EventName, Verbosity, UserFriendlyDesc)
-#define DEFINE_VLOG_EVENT(EventName) 
+#define UE_CVLOG_EVENT_WITH_DATA(Condition, LogOwner, LogEvent, ...)
+
 
 #endif //ENABLE_VISUAL_LOG
 
@@ -268,4 +267,5 @@ public:
 
 #include "VisualLogger.inl"
 
+
 #endif //ENABLE_VISUAL_LOG
diff --git a/Engine/Source/Runtime/Engine/Public/VisualLogger/VisualLogger.inl b/Engine/Source/Runtime/Engine/Public/VisualLogger/VisualLogger.inl
--- a/Engine/Source/Runtime/Engine/Public/VisualLogger/VisualLogger.inl
+++ b/Engine/Source/Runtime/Engine/Public/VisualLogger/VisualLogger.inl
@@ -279,7 +279,7 @@ void FVisualLogger::EventLog(const class UObject* Object, const FName EventTag1,
 {
 	UWorld *World = NULL;
 	FVisualLogEntry *CurrentEntry = NULL;
-	const FLogCategory<ELogVerbosity::Log, ELogVerbosity::Log> Category(*Event.GetName());
+	const FLogCategory<ELogVerbosity::Log, ELogVerbosity::Log> Category(*Event.Name);
 	if (CheckVisualLogInputInternal(Object, Category, ELogVerbosity::Log, &World, &CurrentEntry) == false)
 	{
 		return;
diff --git a/Engine/Source/Runtime/Engine/Public/VisualLogger/VisualLoggerTypes.h b/Engine/Source/Runtime/Engine/Public/VisualLogger/VisualLoggerTypes.h
--- a/Engine/Source/Runtime/Engine/Public/VisualLogger/VisualLoggerTypes.h
+++ b/Engine/Source/Runtime/Engine/Public/VisualLogger/VisualLoggerTypes.h
@@ -18,11 +18,16 @@ namespace VisualLogger
 	};
 }
 
-struct FVisualLogEventBase
+struct ENGINE_API FVisualLogEventBase
 {
-	virtual FString GetName() const { return FString(); }
-	virtual FString GetUserFriendlyDesc() const { return FString(); }
-	virtual ELogVerbosity::Type GetVerbosity() const { return ELogVerbosity::All; }
+	const FString Name;
+	const FString FriendlyDesc;
+	const ELogVerbosity::Type Verbosity;
+
+	FVisualLogEventBase(const FString& InName, const FString& InFriendlyDesc, ELogVerbosity::Type InVerbosity)
+		: Name(InName), FriendlyDesc(InFriendlyDesc), Verbosity(InVerbosity)
+	{
+	}
 };
 
 struct ENGINE_API FVisualLogEntry
@@ -40,16 +45,16 @@ struct ENGINE_API FVisualLogEntry
 		FLogEvent() : Counter(1) {}
 		FLogEvent(const FVisualLogEventBase& Event) : Counter(1)
 		{
-			Name = Event.GetName();
-			UserFriendlyDesc = Event.GetUserFriendlyDesc();
-			Verbosity = Event.GetVerbosity();
+			Name = Event.Name;
+			UserFriendlyDesc = Event.FriendlyDesc;
+			Verbosity = Event.Verbosity;
 		}
 
 		FLogEvent& operator = (const FVisualLogEventBase& Event)
 		{
-			Name = Event.GetName();
-			UserFriendlyDesc = Event.GetUserFriendlyDesc();
-			Verbosity = Event.GetVerbosity();
+			Name = Event.Name;
+			UserFriendlyDesc = Event.FriendlyDesc;
+			Verbosity = Event.Verbosity;
 			return *this;
 		}
 	};
@@ -62,6 +67,8 @@ struct ENGINE_API FVisualLogEntry
 		int64 UserData;
 		FName TagName;
 
+		FLogLine() {}
+
 		FLogLine(const FName& InCategory, ELogVerbosity::Type InVerbosity, const FString& InLine)
 			: Line(InLine), Category(InCategory), Verbosity(InVerbosity), UserData(0)
 		{}
@@ -246,24 +253,134 @@ struct ENGINE_API FVisualLogEntry
 		// find index of status category
 
 #endif // ENABLE_VISUAL_LOG
-	};
+};
 
 #if  ENABLE_VISUAL_LOG
-	FORCEINLINE int32 FVisualLogEntry::FindStatusIndex(const FString& CategoryName)
+FORCEINLINE
+FArchive& operator<<(FArchive& Ar, FVisualLogEntry::FDataBlock& Data)
+{
+	Ar << Data.TagName;
+	Ar << Data.Category;
+	Ar << Data.Verbosity;
+	Ar << Data.Data;
+	Ar << Data.UniqueId;
+
+	return Ar;
+}
+
+FORCEINLINE
+FArchive& operator<<(FArchive& Ar, FVisualLogEntry::FHistogramSample& Sample)
+{
+	Ar << Sample.Category;
+	Ar << Sample.Verbosity;
+	Ar << Sample.GraphName;
+	Ar << Sample.DataName;
+	Ar << Sample.SampleValue;
+	Ar << Sample.UniqueId;
+
+	return Ar;
+}
+
+FORCEINLINE
+FArchive& operator<<(FArchive& Ar, FVisualLogEntry::FElementToDraw& Element)
+{
+	Ar << Element.Description;
+	Ar << Element.Category;
+	Ar << Element.Verbosity;
+	Ar << Element.Points;
+	Ar << Element.UniqueId;
+	Ar << Element.Type;
+	Ar << Element.Color;
+	Ar << Element.Thicknes;
+
+	return Ar;
+}
+
+FORCEINLINE
+FArchive& operator<<(FArchive& Ar, FVisualLogEntry::FLogEvent& Event)
+{
+	Ar << Event.Name;
+	Ar << Event.UserFriendlyDesc;
+	Ar << Event.Verbosity;
+	Ar << Event.EventTags;
+	Ar << Event.Counter;
+	Ar << Event.UserData;
+	Ar << Event.TagName;
+
+	return Ar;
+}
+
+FORCEINLINE
+FArchive& operator<<(FArchive& Ar, FVisualLogEntry::FLogLine& LogLine)
+{
+	Ar << LogLine.Category;
+	Ar << LogLine.Verbosity;
+	Ar << LogLine.TagName;
+	Ar << LogLine.UniqueId;
+	Ar << LogLine.UserData;
+	Ar << LogLine.Line;
+	return Ar;
+}
+
+FORCEINLINE
+FArchive& operator<<(FArchive& Ar, FVisualLogEntry::FStatusCategory& Status)
+{
+	Ar << Status.Category;
+	Ar << Status.Data;
+
+	return Ar;
+}
+
+FORCEINLINE
+FArchive& operator<<(FArchive& Ar, FVisualLogEntry& LogEntry)
+{
+	Ar << LogEntry.TimeStamp;
+	Ar << LogEntry.Location;
+	Ar << LogEntry.LogLines;
+	Ar << LogEntry.Status;
+	Ar << LogEntry.Events;
+	Ar << LogEntry.ElementsToDraw;
+
+	return Ar;
+}
+
+FORCEINLINE int32 FVisualLogEntry::FindStatusIndex(const FString& CategoryName)
+{
+	for (int32 TestCategoryIndex = 0; TestCategoryIndex < Status.Num(); TestCategoryIndex++)
 	{
-		for (int32 TestCategoryIndex = 0; TestCategoryIndex < Status.Num(); TestCategoryIndex++)
+		if (Status[TestCategoryIndex].Category == CategoryName)
 		{
-			if (Status[TestCategoryIndex].Category == CategoryName)
-			{
-				return TestCategoryIndex;
-			}
+			return TestCategoryIndex;
 		}
-
-		return INDEX_NONE;
 	}
 
-	FORCEINLINE bool operator==(const FVisualLogEntry::FLogEvent& Left, const FVisualLogEntry::FLogEvent& Right)
-	{
-		return Left.Name == Right.Name;
-	}
+	return INDEX_NONE;
+}
+
+FORCEINLINE bool operator==(const FVisualLogEntry::FLogEvent& Left, const FVisualLogEntry::FLogEvent& Right)
+{
+	return Left.Name == Right.Name;
+}
+
+struct FVisualLoggerHelpers
+{
+	static FString GenerateTemporaryFilename(const FString& FileExt);
+	static FString GenerateFilename(const FString& TempFileName, const FString& Prefix, float StartRecordingTime, float EndTimeStamp);
+};
+
+FORCEINLINE 
+FString FVisualLoggerHelpers::GenerateTemporaryFilename(const FString& FileExt)
+{ 
+	return FString::Printf(TEXT("VTEMP_%s.%s"), *FDateTime::Now().ToString(), *FileExt);
+}
+
+FORCEINLINE
+FString FVisualLoggerHelpers::GenerateFilename(const FString& TempFileName, const FString& Prefix, float StartRecordingTime, float EndTimeStamp)
+{
+	const FString TempFullFilename = FString::Printf(TEXT("%slogs/%s"), *FPaths::GameSavedDir(), *TempFileName);
+	const FString FullFilename = FString::Printf(TEXT("%slogs/%s_%s"), *FPaths::GameSavedDir(), *Prefix, *TempFileName);
+	const FString TimeFrameString = FString::Printf(TEXT("%d-%d_"), FMath::TruncToInt(StartRecordingTime), FMath::TruncToInt(EndTimeStamp));
+	return FullFilename.Replace(TEXT("VTEMP_"), *TimeFrameString, ESearchCase::CaseSensitive);
+}
+
 #endif // ENABLE_VISUAL_LOG
