CIS fix for OSX

diff --git a/Engine/Source/Runtime/Engine/Private/VisualLog.cpp b/Engine/Source/Runtime/Engine/Private/VisualLog.cpp
--- a/Engine/Source/Runtime/Engine/Private/VisualLog.cpp
+++ b/Engine/Source/Runtime/Engine/Private/VisualLog.cpp
@@ -482,15 +482,6 @@ void FVisualLog::Serialize(const class UObject* LogOwner, FName OwnerName, const
 	LastLogTimeStamp = LogEntry.TimeStamp;
 }
 
-FString FVisualLog::GetLogFileFullName() const
-{
-	FString FileName = FString::Printf(TEXT("%s_%.0f-%.0f_%s.vlog"), *FileName, StartRecordingTime, GWorld ? GWorld->TimeSeconds : 0, *FDateTime::Now().ToString());
-
-	const FString FullFilename = FString::Printf(TEXT("%slogs/%s"), *FPaths::GameSavedDir(), *FileName);
-
-	return FullFilename;
-}
-
 void FVisualLog::DumpRecordedLogs()
 {
 	if (!FileAr)
diff --git a/Engine/Source/Runtime/Engine/Public/VisualLog.h b/Engine/Source/Runtime/Engine/Public/VisualLog.h
--- a/Engine/Source/Runtime/Engine/Public/VisualLog.h
+++ b/Engine/Source/Runtime/Engine/Public/VisualLog.h
@@ -144,8 +144,6 @@ public:
 
 protected:
 
-	FString GetLogFileFullName() const;
-
 	void LogElementImpl(FVisLogSelfDrawingElement* Element);
 
 	FORCEINLINE TSharedPtr<FActorsVisLog> GetLog(const class UObject* Object)
