Compilation fix

diff --git a/Engine/Source/Runtime/Engine/Classes/VisualLogger/VisualLoggerBinaryFileDevice.h b/Engine/Source/Runtime/Engine/Classes/VisualLogger/VisualLoggerBinaryFileDevice.h
--- a/Engine/Source/Runtime/Engine/Classes/VisualLogger/VisualLoggerBinaryFileDevice.h
+++ b/Engine/Source/Runtime/Engine/Classes/VisualLogger/VisualLoggerBinaryFileDevice.h
@@ -22,7 +22,7 @@ public:
 	struct FFrameCacheItem
 	{
 		FFrameCacheItem() {}
-		FFrameCacheItem(FName InOwnerName, const FVisualLogEntry& LogEntry) : OwnerName(OwnerName), Entry(LogEntry) { }
+		FFrameCacheItem(FName InOwnerName, const FVisualLogEntry& LogEntry) : OwnerName(InOwnerName), Entry(LogEntry) { }
 
 		FName OwnerName;
 		FVisualLogEntry Entry;
