fix for crash on editor exit due to accessing already deleted object
#ue4

diff --git a/Engine/Source/Runtime/Engine/Private/VisualLogger/VisualLogger.cpp b/Engine/Source/Runtime/Engine/Private/VisualLogger/VisualLogger.cpp
--- a/Engine/Source/Runtime/Engine/Private/VisualLogger/VisualLogger.cpp
+++ b/Engine/Source/Runtime/Engine/Private/VisualLogger/VisualLogger.cpp
@@ -13,7 +13,7 @@ TMap<class UObject*, TArray<TWeakObjectPtr<const class UObject> > > FVisualLogge
 
 FVisualLogger::FVisualLogger()
 {
-	bool UseBinaryFileDevice = false;
+	UseBinaryFileDevice = false;
 	GConfig->GetBool(TEXT("VisualLogger"), TEXT("UseBinaryFileDevice"), UseBinaryFileDevice, GEngineIni);
 
 	BlockAllCategories(false);
@@ -37,17 +37,14 @@ FVisualLogger::FVisualLogger()
 
 FVisualLogger::~FVisualLogger()
 {
-	bool UseBinaryFileDevice = false;
-	GConfig->GetBool(TEXT("VisualLogger"), TEXT("UseBinaryFileDevice"), UseBinaryFileDevice, GEngineIni);
 	if (UseBinaryFileDevice)
 	{
 		RemoveDevice(&FVisualLoggerBinaryFileDevice::Get());
 	}
 	else
 	{
-		RemoveDevice(&FVisualLoggerBinaryFileDevice::Get());
+		RemoveDevice(&FVisualLog::GetStatic());
 	}
-	RemoveDevice(&FVisualLog::GetStatic());
 }
 
 void FVisualLogger::Shutdown()
diff --git a/Engine/Source/Runtime/Engine/Public/VisualLogger/VisualLogger.h b/Engine/Source/Runtime/Engine/Public/VisualLogger/VisualLogger.h
--- a/Engine/Source/Runtime/Engine/Public/VisualLogger/VisualLogger.h
+++ b/Engine/Source/Runtime/Engine/Public/VisualLogger/VisualLogger.h
@@ -237,6 +237,8 @@ protected:
 	float StartRecordingToFileTime;
 	/** Delegate to set project specific file name for vlogs */
 	FVisualLogFilenameGetterDelegate LogFileNameGetter;
+	/** Specifies if the Binary Device is being used */
+	bool UseBinaryFileDevice;
 };
 
 /**
