UE4 - Metal - Enable ddx/ddy

diff --git a/Engine/Source/Runtime/Engine/Private/Materials/HLSLMaterialTranslator.h b/Engine/Source/Runtime/Engine/Private/Materials/HLSLMaterialTranslator.h
--- a/Engine/Source/Runtime/Engine/Private/Materials/HLSLMaterialTranslator.h
+++ b/Engine/Source/Runtime/Engine/Private/Materials/HLSLMaterialTranslator.h
@@ -2144,7 +2144,7 @@ protected:
 		{
 			// Mobile: Sampling of a particular level depends on an extension; iOS does have it by default but
 			// there's a driver as of 7.0.2 that will cause a GPU hang if used with an Aniso > 1 sampler, so show an error for now
-			if (ErrorUnlessFeatureLevelSupported(ERHIFeatureLevel::SM4) == INDEX_NONE)
+			if (ErrorUnlessFeatureLevelSupported(ERHIFeatureLevel::ES3_1) == INDEX_NONE)
 			{
 				Errorf(TEXT("Sampling for a specific mip-level is not supported for mobile"));
 				return INDEX_NONE;
@@ -3225,7 +3225,7 @@ protected:
 
 	virtual int32 DDX( int32 X ) override
 	{
-		if (ErrorUnlessFeatureLevelSupported(ERHIFeatureLevel::SM4) == INDEX_NONE)
+		if (ErrorUnlessFeatureLevelSupported(ERHIFeatureLevel::ES3_1) == INDEX_NONE)
 		{
 			return INDEX_NONE;
 		}
@@ -3251,7 +3251,7 @@ protected:
 
 	virtual int32 DDY( int32 X ) override
 	{
-		if (ErrorUnlessFeatureLevelSupported(ERHIFeatureLevel::SM4) == INDEX_NONE)
+		if (ErrorUnlessFeatureLevelSupported(ERHIFeatureLevel::ES3_1) == INDEX_NONE)
 		{
 			return INDEX_NONE;
 		}
