Make sure that native subcomponents of blueprint created components get properly registered and destroyed as their parent component is created/destroyed by construction script

diff --git a/Engine/Source/Runtime/Engine/Private/ActorComponent.cpp b/Engine/Source/Runtime/Engine/Private/ActorComponent.cpp
--- a/Engine/Source/Runtime/Engine/Private/ActorComponent.cpp
+++ b/Engine/Source/Runtime/Engine/Private/ActorComponent.cpp
@@ -624,6 +624,23 @@ void UActorComponent::RegisterComponentWithWorld(UWorld* InWorld)
 			InitializeComponent();
 		}
 	}
+
+	// If this is a blueprint created component and it has component children they can miss getting registered in some scenarios
+	if (bCreatedByConstructionScript)
+	{
+		TArray<UObject*> Children;
+		GetObjectsWithOuter(this, Children, true, RF_PendingKill);
+
+		for (UObject* Child : Children)
+		{
+			UActorComponent* ChildComponent = Cast<UActorComponent>(Child);
+			if (ChildComponent && !ChildComponent->IsRegistered())
+			{
+				ChildComponent->RegisterComponentWithWorld(InWorld);
+			}
+		}
+
+	}
 }
 
 void UActorComponent::RegisterComponent()
diff --git a/Engine/Source/Runtime/Engine/Private/ActorConstruction.cpp b/Engine/Source/Runtime/Engine/Private/ActorConstruction.cpp
--- a/Engine/Source/Runtime/Engine/Private/ActorConstruction.cpp
+++ b/Engine/Source/Runtime/Engine/Private/ActorConstruction.cpp
@@ -88,20 +88,42 @@ void AActor::DestroyConstructedComponents()
 	GetComponents(PreviouslyAttachedComponents);
 	for (int32 i = 0; i < PreviouslyAttachedComponents.Num(); i++)
 	{
-		UActorComponent*& Component = PreviouslyAttachedComponents[i];
-		if (Component && Component->bCreatedByConstructionScript)
+		UActorComponent* Component = PreviouslyAttachedComponents[i];
+		if (Component)
 		{
-			if (Component == RootComponent)
+			bool bDestroyComponent = false;
+			if (Component->bCreatedByConstructionScript)
 			{
-				RootComponent = NULL;
+				bDestroyComponent = true;
+			}
+			else
+			{
+				UActorComponent* OuterComponent = Component->GetTypedOuter<UActorComponent>();
+				while (OuterComponent)
+				{
+					if (OuterComponent->bCreatedByConstructionScript)
+					{
+						bDestroyComponent = true;
+						break;
+					}
+					OuterComponent = OuterComponent->GetTypedOuter<UActorComponent>();
+				}
 			}
 
-			Component->DestroyComponent();
+			if (bDestroyComponent)
+			{
+				if (Component == RootComponent)
+				{
+					RootComponent = NULL;
+				}
 
-			// Rename component to avoid naming conflicts in the case where we rerun the SCS and name the new components the same way.
-			FName const NewBaseName( *(FString::Printf(TEXT("TRASH_%s"), *Component->GetClass()->GetName())) );
-			FName const NewObjectName = MakeUniqueObjectName(this, GetClass(), NewBaseName);
-			Component->Rename(*NewObjectName.ToString(), this, REN_ForceNoResetLoaders);
+				Component->DestroyComponent();
+
+				// Rename component to avoid naming conflicts in the case where we rerun the SCS and name the new components the same way.
+				FName const NewBaseName( *(FString::Printf(TEXT("TRASH_%s"), *Component->GetClass()->GetName())) );
+				FName const NewObjectName = MakeUniqueObjectName(this, GetClass(), NewBaseName);
+				Component->Rename(*NewObjectName.ToString(), this, REN_ForceNoResetLoaders);
+			}
 		}
 	}
 }
