Navmesh generation on BSP fix #UE4

- issue UE-4507

diff --git a/Engine/Source/Runtime/Engine/Private/AI/Navigation/NavigationSystem.cpp b/Engine/Source/Runtime/Engine/Private/AI/Navigation/NavigationSystem.cpp
--- a/Engine/Source/Runtime/Engine/Private/AI/Navigation/NavigationSystem.cpp
+++ b/Engine/Source/Runtime/Engine/Private/AI/Navigation/NavigationSystem.cpp
@@ -2292,7 +2292,6 @@ void UNavigationSystem::ReleaseInitialBuildingLock()
 
 void UNavigationSystem::InitializeLevelCollisions()
 {
-#if WITH_NAVIGATION_GENERATOR
 	UWorld* World = GetWorld();
 	if (!bInitialLevelsAdded && UNavigationSystem::GetCurrent(World) == this)
 	{
@@ -2312,7 +2311,6 @@ void UNavigationSystem::InitializeLevelCollisions()
 
 		bInitialLevelsAdded = true;
 	}
-#endif
 }
 
 #if WITH_EDITOR
