Fix for JISSUE UE-3925.
bIsMasked was being set to false if there was no experssion into OpacityMask. Now checking for bUsesMaterialAttributes too.

diff --git a/Engine/Source/Runtime/Engine/Private/Materials/Material.cpp b/Engine/Source/Runtime/Engine/Private/Materials/Material.cpp
--- a/Engine/Source/Runtime/Engine/Private/Materials/Material.cpp
+++ b/Engine/Source/Runtime/Engine/Private/Materials/Material.cpp
@@ -2276,7 +2276,7 @@ void UMaterial::PostEditChangeProperty(FPropertyChangedEvent& PropertyChangedEve
 	}
 
 	// Check if the material is masked and uses a custom opacity (that's not 1.0f).
-	bIsMasked = ((EBlendMode(BlendMode) == BLEND_Masked) && (OpacityMask.Expression || (OpacityMask.UseConstant && OpacityMask.Constant < 0.999f)));
+	bIsMasked = ((EBlendMode(BlendMode) == BLEND_Masked) && (OpacityMask.Expression || (OpacityMask.UseConstant && OpacityMask.Constant < 0.999f)) || bUseMaterialAttributes);
 
 	bool bRequiresCompilation = true;
 	if( PropertyThatChanged ) 
