visualizebuffer now outputs basecolor and specular instead of diffuse color and specualar color

diff --git a/Engine/Source/Runtime/Engine/Private/BufferVisualizationData.cpp b/Engine/Source/Runtime/Engine/Private/BufferVisualizationData.cpp
--- a/Engine/Source/Runtime/Engine/Private/BufferVisualizationData.cpp
+++ b/Engine/Source/Runtime/Engine/Private/BufferVisualizationData.cpp
@@ -95,7 +95,7 @@ void FBufferVisualizationData::ConfigureConsoleCommand()
 	ConsoleDocumentationOverviewTargets += AvailableVisualizationMaterials;
 
 	IConsoleManager::Get().RegisterConsoleVariable(TEXT("r.BufferVisualizationOverviewTargets"),
-		TEXT("DiffuseColor,SpecularColor,SubsurfaceColor,WorldNormal,SeparateTranslucencyRGB,,,Opacity,SeparateTranslucencyA,,,DecalMask,SceneDepth,Roughness,MaterialAO,ShadingModel,,SceneDepthWorldUnits,SceneColor"), 
+		TEXT("BaseColor,Specular,SubsurfaceColor,WorldNormal,SeparateTranslucencyRGB,,,Opacity,SeparateTranslucencyA,,,DecalMask,SceneDepth,Roughness,Metallic,ShadingModel,,SceneDepthWorldUnits,SceneColor"), 
 		*ConsoleDocumentationOverviewTargets,
 		ECVF_Default);
 }
