Changing the location of the performance capture screenshot

Removed the MyEvent from the RenderTestMap matinee.

diff --git a/Engine/Source/Runtime/Engine/Private/Interpolation.cpp b/Engine/Source/Runtime/Engine/Private/Interpolation.cpp
--- a/Engine/Source/Runtime/Engine/Private/Interpolation.cpp
+++ b/Engine/Source/Runtime/Engine/Private/Interpolation.cpp
@@ -408,7 +408,7 @@ void AMatineeActor::NotifyEventTriggered(FName EventName, float EventTime)
 		FString FolderName;
 		PackageName.Split(TEXT("/"), &FolderName, &MapName, ESearchCase::CaseSensitive, ESearchDir::FromEnd);
 
-		FString CaptureName = FString::Printf(TEXT("%s/%s_%4.2f"), *MapName, *GetName(), EventTime);
+		FString CaptureName = FString::Printf(TEXT("PerformanceCapture/%s/%s_%4.2f"), *MapName, *GetName(), EventTime);
 		GEngine->PerformanceCapture(CaptureName);
 	}
 }
