#docs - Fix comment on Pawn::PossessedBy()

diff --git a/Engine/Source/Runtime/Engine/Classes/GameFramework/Pawn.h b/Engine/Source/Runtime/Engine/Classes/GameFramework/Pawn.h
--- a/Engine/Source/Runtime/Engine/Classes/GameFramework/Pawn.h
+++ b/Engine/Source/Runtime/Engine/Classes/GameFramework/Pawn.h
@@ -263,7 +263,7 @@ public:
 
 	/** 
 	 * Called when this Pawn is possessed. Only called on the server (or in standalone).
-	 *	@param C is the controller possessing this pawn
+	 *	@param C The controller possessing this pawn
 	 */
 	virtual void PossessedBy(AController* NewController);
 
