UMG - Disabling UMG being depended on by Engine in dedicated server builds.

diff --git a/Engine/Source/Runtime/Engine/Engine.Build.cs b/Engine/Source/Runtime/Engine/Engine.Build.cs
--- a/Engine/Source/Runtime/Engine/Engine.Build.cs
+++ b/Engine/Source/Runtime/Engine/Engine.Build.cs
@@ -77,12 +77,22 @@ public class Engine : ModuleRules
 			{
 				"MovieSceneCore",
 				"MovieSceneCoreTypes",
-				"UMG",
 				"HeadMountedDisplay",
 				"StreamingPauseRendering",
 			}
 		);
 
+        // Prevent adding UMG as a dependency in dedicated server builds.
+        if (!(Target.Type == TargetRules.TargetType.Server || Target.Type == TargetRules.TargetType.Program))
+        {
+            DynamicallyLoadedModuleNames.AddRange(
+                new string[]
+			    {
+				    "UMG",
+			    }
+            );
+        }
+
 		if (UEBuildConfiguration.bBuildDeveloperTools)
 		{
 			// Add "BlankModule" so that it gets compiled as an example and will be maintained and tested.  This can be removed
