Update DefaultPhysicsVolume header comments

diff --git a/Engine/Source/Runtime/Engine/Classes/GameFramework/DefaultPhysicsVolume.h b/Engine/Source/Runtime/Engine/Classes/GameFramework/DefaultPhysicsVolume.h
--- a/Engine/Source/Runtime/Engine/Classes/GameFramework/DefaultPhysicsVolume.h
+++ b/Engine/Source/Runtime/Engine/Classes/GameFramework/DefaultPhysicsVolume.h
@@ -9,6 +9,11 @@
 #include "GameFramework/PhysicsVolume.h"
 #include "DefaultPhysicsVolume.generated.h"
 
+/**
+* DefaultPhysicsVolume determines the physical effects an actor will experience if they are not inside any user specified PhysicsVolume
+*
+* @see APhysicsVolume
+*/
 UCLASS(notplaceable, transient, MinimalAPI)
 class ADefaultPhysicsVolume : public APhysicsVolume
 {
