#ue4 - Add hit result to all Actor and SceneComponent transform functions, for both C++ and Blueprints.

- Sweep params are now hidden (advanced display) by default for rotation-only nodes, since we don't support swept rotation.

diff --git a/Engine/Config/BaseEngine.ini b/Engine/Config/BaseEngine.ini
--- a/Engine/Config/BaseEngine.ini
+++ b/Engine/Config/BaseEngine.ini
@@ -793,7 +793,6 @@ bUseStreamingPause=false
 +K2FieldRedirects=(OldFieldName="Actor.PlaySound",NewFieldName="PlaySoundAtLocation")
 +K2FieldRedirects=(OldFieldName="Actor.GetLocation",NewFieldName="K2_GetActorLocation")
 +K2FieldRedirects=(OldFieldName="Actor.GetRotation",NewFieldName="K2_GetActorRotation")
-+K2FieldRedirects=(OldFieldName="Actor.SetLocation",NewFieldName="K2_SetActorLocation")
 +K2FieldRedirects=(OldFieldName="Actor.SetRotation",NewFieldName="SetActorRotation")
 +K2FieldRedirects=(OldFieldName="Actor.GetController",NewFieldName="Pawn.GetController")
 +K2FieldRedirects=(OldFieldName="SceneComponent.GetLocation",NewFieldName="K2_GetComponentLocation")
@@ -833,7 +832,6 @@ bUseStreamingPause=false
 +K2FieldRedirects=(OldFieldName="Pawn.LaunchPawn",NewFieldName="Character.LaunchCharacter")
 +K2FieldRedirects=(OldFieldName="Character.Launch",NewFieldName="Character.LaunchCharacter")
 +K2FieldRedirects=(OldFieldName="Character.IsJumping",NewFieldName="Character.IsJumpProvidingForce")
-+K2FieldRedirects=(OldFieldName="Actor.SetActorLocation",NewFieldName="Actor.K2_SetActorLocation")
 +K2FieldRedirects=(OldFieldName="ProjectileMovementComponent.StopMovement",NewFieldName="ProjectileMovementComponent.StopSimulating")
 +K2FieldRedirects=(OldFieldName="PawnMovementComponent.StopActiveMovement",NewFieldName="NavMovementComponent.StopActiveMovement")
 
@@ -897,10 +895,32 @@ bUseStreamingPause=false
 +K2FieldRedirects=(OldFieldName="PrimitiveComponent.GetTouchingComponents",NewFieldName="PrimitiveComponent.GetOverlappingComponents")
 +K2ParamRedirects=(NodeName="/Script/BlueprintGraph.K2Node_CallFunction", OldParamName="PrimitiveComponent.GetOverlappingComponents.TouchingComponents", NewParamName="OverlappingComponents")
 
-+K2FieldRedirects=(OldFieldName="Actor.AddActorLocalTranslation",NewFieldName="Actor.AddActorLocalOffset")
-+K2ParamRedirects=(NodeName="/Script/BlueprintGraph.K2Node_CallFunction", OldParamName="Actor.AddActorLocalOffset.DeltaTranslation", NewParamName="DeltaLocation")
-+K2FieldRedirects=(OldFieldName="Actor.SetActorRelativeTranslation",NewFieldName="Actor.SetActorRelativeLocation")
-+K2ParamRedirects=(NodeName="/Script/BlueprintGraph.K2Node_CallFunction", OldParamName="Actor.SetActorRelativeLocation.NewRelativeTranslation", NewParamName="NewRelativeLocation")
++K2FieldRedirects=(OldFieldName="Actor.AddActorLocalTranslation",NewFieldName="Actor.K2_AddActorLocalOffset")
++K2ParamRedirects=(NodeName="/Script/BlueprintGraph.K2Node_CallFunction", OldParamName="Actor.K2_AddActorLocalOffset.DeltaTranslation", NewParamName="DeltaLocation")
+
+;
+; Add FHitResult to swept movement, and use wrappers for K2 methods
+;
++K2FieldRedirects=(OldFieldName="Actor.SetActorLocation",NewFieldName="Actor.K2_SetActorLocation")
++K2FieldRedirects=(OldFieldName="Actor.SetActorTransform",NewFieldName="Actor.K2_SetActorTransform")
++K2FieldRedirects=(OldFieldName="Actor.SetActorLocationAndRotation",NewFieldName="Actor.K2_SetActorLocationAndRotation")
+
++K2FieldRedirects=(OldFieldName="Actor.AddActorWorldOffset",NewFieldName="Actor.K2_AddActorWorldOffset")
++K2FieldRedirects=(OldFieldName="Actor.AddActorWorldRotation",NewFieldName="Actor.K2_AddActorWorldRotation")
++K2FieldRedirects=(OldFieldName="Actor.AddActorWorldTransform",NewFieldName="Actor.K2_AddActorWorldTransform")
+
++K2FieldRedirects=(OldFieldName="Actor.AddActorLocalOffset",NewFieldName="Actor.K2_AddActorLocalOffset")
++K2FieldRedirects=(OldFieldName="Actor.AddActorLocalRotation",NewFieldName="Actor.K2_AddActorLocalRotation")
++K2FieldRedirects=(OldFieldName="Actor.AddActorLocalTransform",NewFieldName="Actor.K2_AddActorLocalTransform")
+
++K2FieldRedirects=(OldFieldName="Actor.SetActorRelativeLocation",NewFieldName="Actor.K2_SetActorRelativeLocation")
++K2FieldRedirects=(OldFieldName="Actor.SetActorRelativeRotation",NewFieldName="Actor.K2_SetActorRelativeRotation")
++K2FieldRedirects=(OldFieldName="Actor.SetActorRelativeTransform",NewFieldName="Actor.K2_SetActorRelativeTransform")
+
+
+
++K2FieldRedirects=(OldFieldName="Actor.SetActorRelativeTranslation",NewFieldName="Actor.K2_SetActorRelativeLocation")
++K2ParamRedirects=(NodeName="/Script/BlueprintGraph.K2Node_CallFunction", OldParamName="Actor.K2_SetActorRelativeLocation.NewRelativeTranslation", NewParamName="NewRelativeLocation")
 
 +K2FieldRedirects=(OldFieldName="LevelScriptActor.BeginGame",NewFieldName="Actor.ReceiveBeginPlay")
 +K2FieldRedirects=(OldFieldName="LevelScriptActor.BeginGameSignature__DelegateSignature", NewFieldName="Actor.ReceiveBeginPlaySignature__DelegateSignature")
@@ -967,14 +987,33 @@ bUseStreamingPause=false
 +K2ParamRedirects=(NodeName="/Script/BlueprintGraph.K2Node_CallFunction", OldParamName="KismetSystemLibrary.ComponentOverlapActors_NEW.ObjectsToTrace", NewParamName="ObjectTypes")
 +K2ParamRedirects=(NodeName="/Script/BlueprintGraph.K2Node_CallFunction", OldParamName="KismetSystemLibrary.ComponentOverlapComponents_NEW.ObjectsToTrace", NewParamName="ObjectTypes")
 
-+K2FieldRedirects=(OldFieldName="SceneComponent.SetRelativeTranslation",NewFieldName="SceneComponent.SetRelativeLocation")
-+K2ParamRedirects=(NodeName="/Script/BlueprintGraph.K2Node_CallFunction", OldParamName="SceneComponent.SetRelativeLocation.NewTranslation", NewParamName="NewLocation")
-+K2FieldRedirects=(OldFieldName="SceneComponent.AddLocalTranslation",NewFieldName="SceneComponent.AddLocalOffset")
-+K2ParamRedirects=(NodeName="/Script/BlueprintGraph.K2Node_CallFunction", OldParamName="SceneComponent.AddLocalOffset.DeltaTranslation", NewParamName="DeltaLocation")
-+K2FieldRedirects=(OldFieldName="SceneComponent.SetWorldTranslation",NewFieldName="SceneComponent.SetWorldLocation")
-+K2ParamRedirects=(NodeName="/Script/BlueprintGraph.K2Node_CallFunction", OldParamName="SceneComponent.SetWorldLocation.NewTranslation", NewParamName="NewLocation")
-+K2FieldRedirects=(OldFieldName="SceneComponent.AddRelativeTranslation",NewFieldName="SceneComponent.AddRelativeLocation")
-+K2ParamRedirects=(NodeName="/Script/BlueprintGraph.K2Node_CallFunction", OldParamName="Actor.AddRelativeLocation.DeltaTranslation", NewParamName="DeltaLocation")
++K2FieldRedirects=(OldFieldName="SceneComponent.SetRelativeTranslation",NewFieldName="SceneComponent.K2_SetRelativeLocation")
++K2ParamRedirects=(NodeName="/Script/BlueprintGraph.K2Node_CallFunction", OldParamName="SceneComponent.K2_SetRelativeLocation.NewTranslation", NewParamName="NewLocation")
++K2FieldRedirects=(OldFieldName="SceneComponent.AddLocalTranslation",NewFieldName="SceneComponent.K2_AddLocalOffset")
++K2ParamRedirects=(NodeName="/Script/BlueprintGraph.K2Node_CallFunction", OldParamName="SceneComponent.K2_AddLocalOffset.DeltaTranslation", NewParamName="DeltaLocation")
++K2FieldRedirects=(OldFieldName="SceneComponent.SetWorldTranslation",NewFieldName="SceneComponent.K2_SetWorldLocation")
++K2ParamRedirects=(NodeName="/Script/BlueprintGraph.K2Node_CallFunction", OldParamName="SceneComponent.K2_SetWorldLocation.NewTranslation", NewParamName="NewLocation")
++K2FieldRedirects=(OldFieldName="SceneComponent.AddRelativeTranslation",NewFieldName="SceneComponent.K2_AddRelativeLocation")
++K2ParamRedirects=(NodeName="/Script/BlueprintGraph.K2Node_CallFunction", OldParamName="SceneComponent.K2_AddRelativeLocation.DeltaTranslation", NewParamName="DeltaLocation")
+
+;
+; Add FHitResult to swept movement, and use wrappers for K2 methods
+;
++K2FieldRedirects=(OldFieldName="SceneComponent.SetRelativeLocation",NewFieldName="SceneComponent.K2_SetRelativeLocation")
++K2FieldRedirects=(OldFieldName="SceneComponent.SetRelativeRotation",NewFieldName="SceneComponent.K2_SetRelativeRotation")
++K2FieldRedirects=(OldFieldName="SceneComponent.SetRelativeTransform",NewFieldName="SceneComponent.K2_SetRelativeTransform")
++K2FieldRedirects=(OldFieldName="SceneComponent.AddRelativeLocation",NewFieldName="SceneComponent.K2_AddRelativeLocation")
++K2FieldRedirects=(OldFieldName="SceneComponent.AddRelativeRotation",NewFieldName="SceneComponent.K2_AddRelativeRotation")
++K2FieldRedirects=(OldFieldName="SceneComponent.AddLocalOffset",NewFieldName="SceneComponent.K2_AddLocalOffset")
++K2FieldRedirects=(OldFieldName="SceneComponent.AddLocalRotation",NewFieldName="SceneComponent.K2_AddLocalRotation")
++K2FieldRedirects=(OldFieldName="SceneComponent.AddLocalTransform",NewFieldName="SceneComponent.K2_AddLocalTransform")
++K2FieldRedirects=(OldFieldName="SceneComponent.SetWorldLocation",NewFieldName="SceneComponent.K2_SetWorldLocation")
++K2FieldRedirects=(OldFieldName="SceneComponent.SetWorldRotation",NewFieldName="SceneComponent.K2_SetWorldRotation")
++K2FieldRedirects=(OldFieldName="SceneComponent.SetWorldTransform",NewFieldName="SceneComponent.K2_SetWorldTransform")
++K2FieldRedirects=(OldFieldName="SceneComponent.AddWorldOffset",NewFieldName="SceneComponent.K2_AddWorldOffset")
++K2FieldRedirects=(OldFieldName="SceneComponent.AddWorldRotation",NewFieldName="SceneComponent.K2_AddWorldRotation")
++K2FieldRedirects=(OldFieldName="SceneComponent.AddWorldTransform",NewFieldName="SceneComponent.K2_AddWorldTransform")
+
 
 +K2FieldRedirects=(OldFieldName="Actor.SetTickPrerequisite",NewFieldName="Actor.AddTickPrerequisiteActor")
 +K2ParamRedirects=(NodeName="/Script/BlueprintGraph.K2Node_CallFunction", OldParamName="Actor.AddTickPrerequisiteActor.DependentActor", NewParamName="PrerequisiteActor")
diff --git a/Engine/Source/Runtime/Engine/Classes/Components/SceneComponent.h b/Engine/Source/Runtime/Engine/Classes/Components/SceneComponent.h
--- a/Engine/Source/Runtime/Engine/Classes/Components/SceneComponent.h
+++ b/Engine/Source/Runtime/Engine/Classes/Components/SceneComponent.h
@@ -222,16 +222,20 @@ protected:
 public:
 
 	/** Set the location of this component relative to its parent */
-	UFUNCTION(BlueprintCallable, Category="Utilities|Transformation")
-	void SetRelativeLocation(FVector NewLocation, bool bSweep=false);
+	UFUNCTION(BlueprintCallable, Category="Utilities|Transformation", meta=(FriendlyName="SetRelativeLocation"))
+	void K2_SetRelativeLocation(FVector NewLocation, bool bSweep, FHitResult& SweepHitResult);
+	void SetRelativeLocation(FVector NewLocation, bool bSweep=false, FHitResult* OutSweepHitResult=nullptr);
+
 
 	/** Set the rotation of this component relative to its parent */
-	UFUNCTION(BlueprintCallable, Category="Utilities|Transformation")
-	void SetRelativeRotation(FRotator NewRotation, bool bSweep=false);
+	UFUNCTION(BlueprintCallable, Category="Utilities|Transformation", meta=(FriendlyName="SetRelativeRotation", AdvancedDisplay="bSweep,SweepHitResult"))
+	void K2_SetRelativeRotation(FRotator NewRotation, bool bSweep, FHitResult& SweepHitResult);
+	void SetRelativeRotation(FRotator NewRotation, bool bSweep=false, FHitResult* OutSweepHitResult=nullptr);
 
 	/** Set the transform of this component relative to its parent */
-	UFUNCTION(BlueprintCallable, Category="Utilities|Transformation")
-	void SetRelativeTransform(const FTransform& NewTransform, bool bSweep=false);
+	UFUNCTION(BlueprintCallable, Category="Utilities|Transformation", meta=(FriendlyName="SetRelativeTransform"))
+	void K2_SetRelativeTransform(const FTransform& NewTransform, bool bSweep, FHitResult& SweepHitResult);
+	void SetRelativeTransform(const FTransform& NewTransform, bool bSweep=false, FHitResult* OutSweepHitResult=nullptr);
 
 	/** Returns the transform of this component relative to its parent */
 	UFUNCTION(BlueprintCallable, Category="Utilities|Transformation")
@@ -246,52 +250,70 @@ public:
 	virtual void SetRelativeScale3D(FVector NewScale3D);
 
 	/** Adds a delta to the translation of this component relative to its parent */
-	UFUNCTION(BlueprintCallable, Category="Utilities|Transformation")
-	void AddRelativeLocation(FVector DeltaLocation, bool bSweep=false);
+	UFUNCTION(BlueprintCallable, Category="Utilities|Transformation", meta=(FriendlyName="AddRelativeLocation"))
+	void K2_AddRelativeLocation(FVector DeltaLocation, bool bSweep, FHitResult& SweepHitResult);
+	void AddRelativeLocation(FVector DeltaLocation, bool bSweep=false, FHitResult* OutSweepHitResult=nullptr);
 
 	/** Adds a delta the rotation of this component relative to its parent */
-	UFUNCTION(BlueprintCallable, Category="Utilities|Transformation")
-	void AddRelativeRotation(FRotator DeltaRotation, bool bSweep=false);
+	UFUNCTION(BlueprintCallable, Category="Utilities|Transformation", meta=(FriendlyName="AddRelativeRotation", AdvancedDisplay="bSweep,SweepHitResult"))
+	void K2_AddRelativeRotation(FRotator DeltaRotation, bool bSweep, FHitResult& SweepHitResult);
+	void AddRelativeRotation(FRotator DeltaRotation, bool bSweep=false, FHitResult* OutSweepHitResult=nullptr);
 
 	/** Adds a delta to the location of this component in its local reference frame */
-	UFUNCTION(BlueprintCallable, Category="Utilities|Transformation", meta=(Keywords="location position"))
-	void AddLocalOffset(FVector DeltaLocation, bool bSweep=false);
+	UFUNCTION(BlueprintCallable, Category="Utilities|Transformation", meta=(FriendlyName="AddLocalOffset", Keywords="location position"))
+	void K2_AddLocalOffset(FVector DeltaLocation, bool bSweep, FHitResult& SweepHitResult);
+	void AddLocalOffset(FVector DeltaLocation, bool bSweep=false, FHitResult* OutSweepHitResult=nullptr);
+
 	
 	/** Adds a delta to the rotation of this component in its local reference frame */
-	UFUNCTION(BlueprintCallable, Category="Utilities|Transformation")
-	void AddLocalRotation(FRotator DeltaRotation, bool bSweep=false);
+	UFUNCTION(BlueprintCallable, Category="Utilities|Transformation", meta=(FriendlyName="AddLocalRotation", AdvancedDisplay="bSweep,SweepHitResult"))
+	void K2_AddLocalRotation(FRotator DeltaRotation, bool bSweep, FHitResult& SweepHitResult);
+	void AddLocalRotation(FRotator DeltaRotation, bool bSweep=false, FHitResult* OutSweepHitResult=nullptr);
+
 
 	/** Adds a delta to the transform of this component in its local reference frame. Scale is unchanged. */
-	UFUNCTION(BlueprintCallable, Category="Utilities|Transformation")
-	void AddLocalTransform(const FTransform& DeltaTransform, bool bSweep=false);
+	UFUNCTION(BlueprintCallable, Category="Utilities|Transformation", meta=(FriendlyName="AddLocalTransform"))
+	void K2_AddLocalTransform(const FTransform& DeltaTransform, bool bSweep, FHitResult& SweepHitResult);
+	void AddLocalTransform(const FTransform& DeltaTransform, bool bSweep=false, FHitResult* OutSweepHitResult=nullptr);
+
 
 	/** Set the relative location of this component to put it at the supplied location in world space. */
-	UFUNCTION(BlueprintCallable, Category="Utilities|Transformation")
-	void SetWorldLocation(FVector NewLocation, bool bSweep=false);
+	UFUNCTION(BlueprintCallable, Category="Utilities|Transformation", meta=(FriendlyName="SetWorldLocation", AdvancedDisplay="bSweep,SweepHitResult"))
+	void K2_SetWorldLocation(FVector NewLocation, bool bSweep, FHitResult& SweepHitResult);
+	void SetWorldLocation(FVector NewLocation, bool bSweep=false, FHitResult* OutSweepHitResult=nullptr);
+
 
 	/** Set the relative rotation of this component to put it at the supplied orientation in world space. */
-	UFUNCTION(BlueprintCallable, Category="Utilities|Transformation")
-	void SetWorldRotation(FRotator NewRotation, bool bSweep=false);
+	UFUNCTION(BlueprintCallable, Category="Utilities|Transformation", meta=(FriendlyName="SetWorldRotation"))
+	void K2_SetWorldRotation(FRotator NewRotation, bool bSweep, FHitResult& SweepHitResult);
+	void SetWorldRotation(FRotator NewRotation, bool bSweep=false, FHitResult* OutSweepHitResult=nullptr);
+
 
 	/** Set the relative scale of this component to put it at the supplied scale in world space. */
 	UFUNCTION(BlueprintCallable, Category="Utilities|Transformation")
 	void SetWorldScale3D(FVector NewScale);
 
 	/** Set the transform of this component in world space. */
-	UFUNCTION(BlueprintCallable, Category="Utilities|Transformation")
-	void SetWorldTransform(const FTransform& NewTransform, bool bSweep=false);
+	UFUNCTION(BlueprintCallable, Category="Utilities|Transformation", meta=(FriendlyName="SetWorldTransform"))
+	void K2_SetWorldTransform(const FTransform& NewTransform, bool bSweep, FHitResult& SweepHitResult);
+	void SetWorldTransform(const FTransform& NewTransform, bool bSweep=false, FHitResult* OutSweepHitResult=nullptr);
+
 
 	/** Adds a delta to the location of this component in world space. */
-	UFUNCTION(BlueprintCallable, Category="Utilities|Transformation", meta=(Keywords="location position"))
-	void AddWorldOffset(FVector DeltaLocation, bool bSweep=false);
+	UFUNCTION(BlueprintCallable, Category="Utilities|Transformation", meta=(FriendlyName="AddWorldOffset", Keywords="location position"))
+	void K2_AddWorldOffset(FVector DeltaLocation, bool bSweep, FHitResult& SweepHitResult);
+	void AddWorldOffset(FVector DeltaLocation, bool bSweep=false, FHitResult* OutSweepHitResult=nullptr);
+
 
 	/** Adds a delta to the rotation of this component in world space. */
-	UFUNCTION(BlueprintCallable, Category="Utilities|Transformation")
-	void AddWorldRotation(FRotator DeltaRotation, bool bSweep=false);
+	UFUNCTION(BlueprintCallable, Category="Utilities|Transformation", meta=(FriendlyName="AddWorldRotation", AdvancedDisplay="bSweep,SweepHitResult"))
+	void K2_AddWorldRotation(FRotator DeltaRotation, bool bSweep, FHitResult& SweepHitResult);
+	void AddWorldRotation(FRotator DeltaRotation, bool bSweep=false, FHitResult* OutSweepHitResult=nullptr);
 
 	/** Adds a delta to the transform of this component in world space. Scale is unchanged. */
-	UFUNCTION(BlueprintCallable, Category="Utilities|Transformation")
-	void AddWorldTransform(const FTransform& DeltaTransform, bool bSweep=false);
+	UFUNCTION(BlueprintCallable, Category="Utilities|Transformation", meta=(FriendlyName="AddWorldTransform"))
+	void K2_AddWorldTransform(const FTransform& DeltaTransform, bool bSweep, FHitResult& SweepHitResult);
+	void AddWorldTransform(const FTransform& DeltaTransform, bool bSweep=false, FHitResult* OutSweepHitResult=nullptr);
 
 	/** Return location of the component, in world space */
 	UFUNCTION(BlueprintCallable, meta=(FriendlyName = "GetWorldLocation"), Category="Utilities|Transformation")
@@ -670,17 +692,21 @@ protected:
 	
 public:
 	/** Set the location and rotation of this component relative to its parent */
-	void SetRelativeLocationAndRotation(FVector NewLocation, FRotator NewRotation, bool bSweep=false);
+	UFUNCTION(BlueprintCallable, Category="Utilities|Transformation", meta=(FriendlyName="SetRelativeLocationAndRotation"))
+	void K2_SetRelativeLocationAndRotation(FVector NewLocation, FRotator NewRotation, bool bSweep, FHitResult& SweepHitResult);
+	void SetRelativeLocationAndRotation(FVector NewLocation, FRotator NewRotation, bool bSweep=false, FHitResult* OutSweepHitResult=nullptr);
 
 	/** Set which parts of the relative transform should be relative to parent, and which should be relative to world */
 	UFUNCTION(BlueprintCallable, Category="Utilities|Transformation")
 	void SetAbsolute(bool bNewAbsoluteLocation = false, bool bNewAbsoluteRotation = false, bool bNewAbsoluteScale = false);
 
 	/** Set the relative location & FRotator rotation of this component to put it at the supplied pose in world space. */
-	void SetWorldLocationAndRotation(FVector NewLocation, FRotator NewRotation, bool bSweep=false);
+	UFUNCTION(BlueprintCallable, Category="Utilities|Transformation", meta=(FriendlyName="SetWorldLocationAndRotation"))
+	void K2_SetWorldLocationAndRotation(FVector NewLocation, FRotator NewRotation, bool bSweep, FHitResult& SweepHitResult);
+	void SetWorldLocationAndRotation(FVector NewLocation, FRotator NewRotation, bool bSweep=false, FHitResult* OutSweepHitResult=nullptr);
 
 	/** Set the relative location & FQuat rotation of this component to put it at the supplied pose in world space. */
-	void SetWorldLocationAndRotation(FVector NewLocation, const FQuat& NewRotation, bool bSweep=false);
+	void SetWorldLocationAndRotation(FVector NewLocation, const FQuat& NewRotation, bool bSweep=false, FHitResult* OutSweepHitResult=nullptr);
 
 	/** Special version of SetWorldLocationAndRotation that does not affect physics. */
 	void SetWorldLocationAndRotationNoPhysics(const FVector& NewLocation, const FRotator& NewRotation);
diff --git a/Engine/Source/Runtime/Engine/Classes/GameFramework/Actor.h b/Engine/Source/Runtime/Engine/Classes/GameFramework/Actor.h
--- a/Engine/Source/Runtime/Engine/Classes/GameFramework/Actor.h
+++ b/Engine/Source/Runtime/Engine/Classes/GameFramework/Actor.h
@@ -587,13 +587,14 @@ public:
 	FVector K2_GetActorLocation() const;
 
 	/** 
-	 *	Move the Actor instantly to the specified location. 
-	 *	@param NewLocation	The new location to teleport the Actor to.
-	 *	@param bSweep		Should we sweep to the destination location. If true, will stop short of the target if blocked by something.
+	 *	Move the Actor to the specified location.
+	 *	@param NewLocation	The new location to move the Actor to.
+	 *	@param bSweep		Should we sweep to the destination location, stopping short of the target if blocked by something.
+	 *  @param SweepHitResult	The hit result from the move if swept.
 	 *	@return	Whether the location was successfully set (if not swept), or whether movement occurred at all (if swept).
 	 */
 	UFUNCTION(BlueprintCallable, meta=(FriendlyName = "SetActorLocation", Keywords="position"), Category="Utilities|Transformation")
-	bool K2_SetActorLocation(FVector NewLocation, bool bSweep=false);
+	bool K2_SetActorLocation(FVector NewLocation, bool bSweep, FHitResult& SweepHitResult);
 
 	/** Returns rotation of the RootComponent of this Actor. */
 	UFUNCTION(BlueprintCallable, meta=(FriendlyName = "GetActorRotation"), Category="Utilities|Transformation")
@@ -612,8 +613,8 @@ public:
 	FVector GetActorRightVector() const;
 
 	/**
-	 *	Returns the bounding box of all components that make up this Actor.
-	 *	@param	bOnlyCollidingComponents	If true, will only return the bounding box for components with collision enabled.
+	 * Returns the bounding box of all components that make up this Actor.
+	 * @param	bOnlyCollidingComponents	If true, will only return the bounding box for components with collision enabled.
 	 */
 	UFUNCTION(BlueprintCallable, Category="Collision", meta=(FriendlyName = "GetActorBounds"))
 	void GetActorBounds(bool bOnlyCollidingComponents, FVector& Origin, FVector& BoxExtent) const;
@@ -627,30 +628,46 @@ public:
 	virtual FVector GetVelocity() const;
 
 	/** 
-	 *	Move the Actor instantly to the specified location. 
-	 *	@param NewLocation	The new location to teleport the Actor to.
-	 *	@param bSweep		Should we sweep to the destination location. If true, will stop short of the target if blocked by something.
-	 *	@return	Whether the location was successfully set if not swept, or whether movement occurred if swept.
+	 * Move the actor instantly to the specified location. 
+	 * 
+	 * @param NewLocation		The new location to teleport the Actor to.
+	 * @param bSweep			Whether to sweep to the destination location, triggering overlaps along the way and stopping at the first blocking hit.
+	 * @param OutSweepHitResult The hit result from the move if swept.
+	 * @return	Whether the location was successfully set if not swept, or whether movement occurred if swept.
 	 */
-	bool SetActorLocation(const FVector& NewLocation, bool bSweep=false);
+	bool SetActorLocation(const FVector& NewLocation, bool bSweep=false, FHitResult* OutSweepHitResult=nullptr);
 
 	/** 
-	 *	Set the Actor's rotation instantly to the specified rotation.
-	 *	@param	NewRotation	The new rotation for the Actor.
-	 *	@return	Whether the rotation was successfully set.
+	 * Set the Actor's rotation instantly to the specified rotation.
+	 * 
+	 * @param	NewRotation	The new rotation for the Actor.
+	 * @return	Whether the rotation was successfully set.
 	 */
 	UFUNCTION(BlueprintCallable, Category="Utilities|Transformation")
 	bool SetActorRotation(FRotator NewRotation);
 
 	/** 
-	 *	Set the Actor's location and rotation instantly to the specified location and rotation.
-	 *	@param NewLocation	The new location to teleport the Actor to.
-	 *	@param NewRotation	The new rotation for the Actor.
-	 *	@param bSweep		Should we sweep to the destination location. If true, will stop short of the target if blocked by something.
-	 *	@return	Whether the rotation was successfully set.
+	 * Move the actor instantly to the specified location and rotation.
+	 * 
+	 * @param NewLocation			The new location to teleport the Actor to.
+	 * @param NewRotation			The new rotation for the Actor.
+	 * @param bSweep				Whether to sweep to the destination location, triggering overlaps along the way and stopping at the first blocking hit.
+	 * @param SweepHitResult		The hit result from the move if swept.
+	 * @return	Whether the rotation was successfully set.
+	 */
+	UFUNCTION(BlueprintCallable, Category="Utilities|Transformation", meta=(FriendlyName="SetActorLocationAndRotation"))
+	bool K2_SetActorLocationAndRotation(FVector NewLocation, FRotator NewRotation, bool bSweep, FHitResult& SweepHitResult);
+	
+	/** 
+	 * Move the actor instantly to the specified location and rotation.
+	 * 
+	 * @param NewLocation			The new location to teleport the Actor to.
+	 * @param NewRotation			The new rotation for the Actor.
+	 * @param bSweep				Whether to sweep to the destination location, triggering overlaps along the way and stopping at the first blocking hit.
+	 * @param OutSweepHitResult	The hit result from the move if swept.
+	 * @return	Whether the rotation was successfully set.
 	 */
-	UFUNCTION(BlueprintCallable, Category="Utilities|Transformation")
-	bool SetActorLocationAndRotation(const FVector& NewLocation, FRotator NewRotation, bool bSweep=false);
+	bool SetActorLocationAndRotation(FVector NewLocation, FRotator NewRotation, bool bSweep=false, FHitResult* OutSweepHitResult=nullptr);
 
 	/** Set the Actor's world-space scale. */
 	UFUNCTION(BlueprintCallable, Category="Utilities|Transformation")
@@ -680,60 +697,98 @@ public:
 	UFUNCTION(BlueprintCallable, Category = "Utilities|Transformation")
 	float GetHorizontalDotProductTo(AActor* OtherActor);
 
-	/** Adds a delta to the location of this actor in world space. */
-	UFUNCTION(BlueprintCallable, Category="Utilities|Transformation", meta=(Keywords="location position"))
-	void AddActorWorldOffset(FVector DeltaLocation, bool bSweep=false);
 
-	/** Adds a delta to the rotation of this actor in world space. */
-	UFUNCTION(BlueprintCallable, Category="Utilities|Transformation")
-	void AddActorWorldRotation(FRotator DeltaRotation, bool bSweep=false);
+	/**
+	 * Adds a delta to the location of this actor in world space.
+	 * 
+	 * @param  DeltaLocation		The change in location.
+	 * @param  bSweep				Whether to sweep to the destination location, triggering overlaps along the way and stopping at the first blocking hit.
+	 * @param  SweepHitResult		The hit result from the move if swept.
+	 */
+	UFUNCTION(BlueprintCallable, Category="Utilities|Transformation", meta=(FriendlyName="AddActorWorldOffset", Keywords="location position"))
+	void K2_AddActorWorldOffset(FVector DeltaLocation, bool bSweep, FHitResult& SweepHitResult);
+
+	/**
+	 * Adds a delta to the location of this actor in world space.
+	 * 
+	 * @param  DeltaLocation		The change in location.
+	 * @param  bSweep				Whether to sweep to the destination location, triggering overlaps along the way and stopping at the first blocking hit.
+	 * @param  SweepHitResult		The hit result from the move if swept.
+	 */
+	void AddActorWorldOffset(FVector DeltaLocation, bool bSweep=false, FHitResult* OutSweepHitResult=nullptr);
+
+
+	/**
+	 * Adds a delta to the rotation of this actor in world space.
+	 * 
+	 * @param  DeltaRotation		The change in rotation.
+	 * @param  bSweep				Whether to sweep to the target rotation (not currently supported).
+	 * @param  SweepHitResult		The hit result from the move if swept.
+	 */
+	UFUNCTION(BlueprintCallable, Category="Utilities|Transformation", meta=(FriendlyName="AddActorWorldRotation", AdvancedDisplay="bSweep,SweepHitResult"))
+	void K2_AddActorWorldRotation(FRotator DeltaRotation, bool bSweep, FHitResult& SweepHitResult);
+	void AddActorWorldRotation(FRotator DeltaRotation, bool bSweep=false, FHitResult* OutSweepHitResult=nullptr);
+
 
 	/** Adds a delta to the transform of this actor in world space. Scale is unchanged. */
-	UFUNCTION(BlueprintCallable, Category="Utilities|Transformation")
-	void AddActorWorldTransform(const FTransform& DeltaTransform, bool bSweep=false);
+	UFUNCTION(BlueprintCallable, Category="Utilities|Transformation", meta=(FriendlyName="AddActorWorldTransform"))
+	void K2_AddActorWorldTransform(const FTransform& DeltaTransform, bool bSweep, FHitResult& SweepHitResult);
+	void AddActorWorldTransform(const FTransform& DeltaTransform, bool bSweep=false, FHitResult* OutSweepHitResult=nullptr);
+
 
 	/** 
-	 *	Set the Actors transform to the specified one.
-	 *	@param bSweep		Should we sweep to the destination location. If true, will stop short of the target if blocked by something.
+	 * Set the Actors transform to the specified one.
+	 * @param bSweep		Whether to sweep to the destination location, stopping short of the target if blocked by something.
 	 */
-	UFUNCTION(BlueprintCallable, Category="Utilities|Transformation")
-	bool SetActorTransform(const FTransform& NewTransform, bool bSweep=false);
+	UFUNCTION(BlueprintCallable, Category="Utilities|Transformation", meta=(FriendlyName="SetActorTransform"))
+	bool K2_SetActorTransform(const FTransform& NewTransform, bool bSweep, FHitResult& SweepHitResult);
+	bool SetActorTransform(const FTransform& NewTransform, bool bSweep=false, FHitResult* OutSweepHitResult=nullptr);
+
 
 	/** Adds a delta to the location of this component in its local reference frame */
-	UFUNCTION(BlueprintCallable, Category="Utilities|Transformation", meta=(Keywords="location position"))
-	void AddActorLocalOffset(FVector DeltaLocation, bool bSweep=false);
+	UFUNCTION(BlueprintCallable, Category="Utilities|Transformation", meta=(FriendlyName="AddActorLocalOffset", Keywords="location position"))
+	void K2_AddActorLocalOffset(FVector DeltaLocation, bool bSweep, FHitResult& SweepHitResult);
+	void AddActorLocalOffset(FVector DeltaLocation, bool bSweep=false, FHitResult* OutSweepHitResult=nullptr);
+
 
 	/** Adds a delta to the rotation of this component in its local reference frame */
-	UFUNCTION(BlueprintCallable, Category="Utilities|Transformation")
-	void AddActorLocalRotation(FRotator DeltaRotation, bool bSweep=false);
+	UFUNCTION(BlueprintCallable, Category="Utilities|Transformation", meta=(FriendlyName="AddActorLocalRotation", AdvancedDisplay="bSweep,SweepHitResult"))
+	void K2_AddActorLocalRotation(FRotator DeltaRotation, bool bSweep, FHitResult& SweepHitResult);
+	void AddActorLocalRotation(FRotator DeltaRotation, bool bSweep=false, FHitResult* OutSweepHitResult=nullptr);
+
 
 	/** Adds a delta to the transform of this component in its local reference frame */
-	UFUNCTION(BlueprintCallable, Category="Utilities|Transformation")
-	void AddActorLocalTransform(const FTransform& NewTransform, bool bSweep=false);
+	UFUNCTION(BlueprintCallable, Category="Utilities|Transformation", meta=(FriendlyName="AddActorLocalTransform"))
+	void K2_AddActorLocalTransform(const FTransform& NewTransform, bool bSweep, FHitResult& SweepHitResult);
+	void AddActorLocalTransform(const FTransform& NewTransform, bool bSweep=false, FHitResult* OutSweepHitResult=nullptr);
+
 
 	/**
 	 * Set the actor's RootComponent to the specified relative location
 	 * @param NewRelativeLocation	New relative location to set the actor's RootComponent to
 	 * @param bSweep				Should we sweep to the destination location. If true, will stop short of the target if blocked by something
 	 */
-	UFUNCTION(BlueprintCallable, Category="Utilities|Transformation")
-	void SetActorRelativeLocation(FVector NewRelativeLocation, bool bSweep=false);
+	UFUNCTION(BlueprintCallable, Category="Utilities|Transformation", meta=(FriendlyName="SetActorRelativeLocation"))
+	void K2_SetActorRelativeLocation(FVector NewRelativeLocation, bool bSweep, FHitResult& SweepHitResult);
+	void SetActorRelativeLocation(FVector NewRelativeLocation, bool bSweep=false, FHitResult* OutSweepHitResult=nullptr);
 
 	/**
 	 * Set the actor's RootComponent to the specified relative rotation
 	 * @param NewRelativeRotation		New relative rotation to set the actor's RootComponent to
 	 * @param bSweep					Should we sweep to the destination rotation. If true, will stop short of the target if blocked by something
 	 */
-	UFUNCTION(BlueprintCallable, Category="Utilities|Transformation")
-	void SetActorRelativeRotation(FRotator NewRelativeRotation, bool bSweep=false);
+	UFUNCTION(BlueprintCallable, Category="Utilities|Transformation", meta=(FriendlyName="SetActorRelativeRotation", AdvancedDisplay="bSweep,SweepHitResult"))
+	void K2_SetActorRelativeRotation(FRotator NewRelativeRotation, bool bSweep, FHitResult& SweepHitResult);
+	void SetActorRelativeRotation(FRotator NewRelativeRotation, bool bSweep=false, FHitResult* OutSweepHitResult=nullptr);
 
 	/**
 	 * Set the actor's RootComponent to the specified relative transform
 	 * @param NewRelativeTransform		New relative transform to set the actor's RootComponent to
 	 * @param bSweep					Should we sweep to the destination transform. If true, will stop short of the target if blocked by something
 	 */
-	UFUNCTION(BlueprintCallable, Category="Utilities|Transformation")
-	void SetActorRelativeTransform(const FTransform& NewRelativeTransform, bool bSweep=false);
+	UFUNCTION(BlueprintCallable, Category="Utilities|Transformation", meta=(FriendlyName="SetActorRelativeTransform"))
+	void K2_SetActorRelativeTransform(const FTransform& NewRelativeTransform, bool bSweep, FHitResult& SweepHitResult);
+	void SetActorRelativeTransform(const FTransform& NewRelativeTransform, bool bSweep=false, FHitResult* OutSweepHitResult=nullptr);
 
 	/**
 	 * Set the actor's RootComponent to the specified relative scale 3d
@@ -2228,9 +2283,9 @@ FORCEINLINE FVector AActor::GetSimpleCollisionCylinderExtent() const
 	FRotator GetActorRotation() const { return Super::GetActorRotation(); } \
 	FQuat GetActorQuat() const { return Super::GetActorQuat(); } \
 	FVector GetActorScale() const { return Super::GetActorScale(); } \
-	bool SetActorLocation(const FVector& NewLocation, bool bSweep=false) { return Super::SetActorLocation(NewLocation, bSweep); } \
+	bool SetActorLocation(const FVector& NewLocation, bool bSweep=false, FHitResult* OutSweepHitResult=nullptr) { return Super::SetActorLocation(NewLocation, bSweep, OutSweepHitResult); } \
 	bool SetActorRotation(FRotator NewRotation) { return Super::SetActorRotation(NewRotation); } \
-	bool SetActorLocationAndRotation(const FVector& NewLocation, FRotator NewRotation, bool bSweep=false) { return Super::SetActorLocationAndRotation(NewLocation, NewRotation, bSweep); } \
+	bool SetActorLocationAndRotation(FVector NewLocation, FRotator NewRotation, bool bSweep=false, FHitResult* OutSweepHitResult=nullptr) { return Super::SetActorLocationAndRotation(NewLocation, NewRotation, bSweep, OutSweepHitResult); } \
 	virtual bool TeleportTo( const FVector& DestLocation, const FRotator& DestRotation, bool bIsATest, bool bNoCheck ) override { return Super::TeleportTo(DestLocation, DestRotation, bIsATest, bNoCheck); } \
 	virtual FVector GetVelocity() const override { return Super::GetVelocity(); } \
 	float GetDistanceTo(AActor* OtherActor) { return Super::GetDistanceTo(OtherActor); }
diff --git a/Engine/Source/Runtime/Engine/Private/Actor.cpp b/Engine/Source/Runtime/Engine/Private/Actor.cpp
--- a/Engine/Source/Runtime/Engine/Private/Actor.cpp
+++ b/Engine/Source/Runtime/Engine/Private/Actor.cpp
@@ -2385,22 +2385,21 @@ FVector AActor::GetInputVectorAxisValue(const FKey InputAxisKey) const
 	return Value;
 }
 
-bool AActor::SetActorLocation(const FVector& NewLocation, bool bSweep)
+bool AActor::SetActorLocation(const FVector& NewLocation, bool bSweep, FHitResult* OutSweepHitResult)
 {
 	if (RootComponent)
 	{
 		const FVector Delta = NewLocation - GetActorLocation();
-		return RootComponent->MoveComponent( Delta, GetActorRotation(), bSweep );
+		return RootComponent->MoveComponent( Delta, GetActorRotation(), bSweep, OutSweepHitResult );
+	}
+	else if (OutSweepHitResult)
+	{
+		*OutSweepHitResult = FHitResult();
 	}
 
 	return false;
 }
 
-bool AActor::K2_SetActorLocation(FVector NewLocation, bool bSweep)
-{
-	return SetActorLocation(NewLocation, bSweep);
-}
-
 bool AActor::SetActorRotation(FRotator NewRotation)
 {
 	if (RootComponent)
@@ -2411,12 +2410,16 @@ bool AActor::SetActorRotation(FRotator NewRotation)
 	return false;
 }
 
-bool AActor::SetActorLocationAndRotation(const FVector& NewLocation, FRotator NewRotation, bool bSweep)
+bool AActor::SetActorLocationAndRotation(FVector NewLocation, FRotator NewRotation, bool bSweep, FHitResult* OutSweepHitResult)
 {
 	if (RootComponent)
 	{
 		const FVector Delta = NewLocation - GetActorLocation();
-		return RootComponent->MoveComponent( Delta, NewRotation, bSweep );
+		return RootComponent->MoveComponent( Delta, NewRotation, bSweep, OutSweepHitResult );
+	}
+	else if (OutSweepHitResult)
+	{
+		*OutSweepHitResult = FHitResult();
 	}
 
 	return false;
@@ -2440,33 +2443,44 @@ FVector AActor::GetActorScale3D() const
 	return FVector(1,1,1);
 }
 
-void AActor::AddActorWorldOffset(FVector DeltaLocation, bool bSweep)
+void AActor::AddActorWorldOffset(FVector DeltaLocation, bool bSweep, FHitResult* OutSweepHitResult)
 {
 	if (RootComponent)
 	{
-		RootComponent->AddWorldOffset(DeltaLocation, bSweep);
+		RootComponent->AddWorldOffset(DeltaLocation, bSweep, OutSweepHitResult);
+	}
+	else if (OutSweepHitResult)
+	{
+		*OutSweepHitResult = FHitResult();
 	}
 }
 
-void AActor::AddActorWorldRotation(FRotator DeltaRotation, bool bSweep)
+void AActor::AddActorWorldRotation(FRotator DeltaRotation, bool bSweep, FHitResult* OutSweepHitResult)
 {
 	if (RootComponent)
 	{
-		RootComponent->AddWorldRotation(DeltaRotation, bSweep);
-	}	
+		RootComponent->AddWorldRotation(DeltaRotation, bSweep, OutSweepHitResult);
+	}
+	else if (OutSweepHitResult)
+	{
+		*OutSweepHitResult = FHitResult();
+	}
 }
 
-void AActor::AddActorWorldTransform(const FTransform& DeltaTransform, bool bSweep)
+void AActor::AddActorWorldTransform(const FTransform& DeltaTransform, bool bSweep, FHitResult* OutSweepHitResult)
 {
 	if (RootComponent)
 	{
-		RootComponent->AddWorldTransform(DeltaTransform, bSweep);
-	}	
+		RootComponent->AddWorldTransform(DeltaTransform, bSweep, OutSweepHitResult);
+	}
+	else if (OutSweepHitResult)
+	{
+		*OutSweepHitResult = FHitResult();
+	}
 }
 
 
-
-bool AActor::SetActorTransform(const FTransform& NewTransform, bool bSweep)
+bool AActor::SetActorTransform(const FTransform& NewTransform, bool bSweep, FHitResult* OutSweepHitResult)
 {
 	// we have seen this gets NAN from kismet, and would like to see if this
 	// happens, and if so, something else is giving NAN as output
@@ -2474,63 +2488,95 @@ bool AActor::SetActorTransform(const FTransform& NewTransform, bool bSweep)
 	{
 		if (ensure(!NewTransform.ContainsNaN()))
 		{
-			RootComponent->SetWorldTransform(NewTransform, bSweep);
+			RootComponent->SetWorldTransform(NewTransform, bSweep, OutSweepHitResult);
 		}
 		else
 		{
 			UE_LOG(LogScript, Warning, TEXT("SetActorTransform: Get NAN Transform data for %s: %s"), *GetNameSafe(this), *NewTransform.ToString());
+			if (OutSweepHitResult)
+			{
+				*OutSweepHitResult = FHitResult();
+			}
 		}
 		return true;
 	}
 
+	if (OutSweepHitResult)
+	{
+		*OutSweepHitResult = FHitResult();
+	}
 	return false;
 }
 
-void AActor::AddActorLocalOffset(FVector DeltaLocation, bool bSweep)
+void AActor::AddActorLocalOffset(FVector DeltaLocation, bool bSweep, FHitResult* OutSweepHitResult)
 {
 	if(RootComponent)
 	{
-		RootComponent->AddLocalOffset(DeltaLocation, bSweep);
+		RootComponent->AddLocalOffset(DeltaLocation, bSweep, OutSweepHitResult);
+	}
+	else if (OutSweepHitResult)
+	{
+		*OutSweepHitResult = FHitResult();
 	}
 }
 
-void AActor::AddActorLocalRotation(FRotator DeltaRotation, bool bSweep)
+void AActor::AddActorLocalRotation(FRotator DeltaRotation, bool bSweep, FHitResult* OutSweepHitResult)
 {
 	if(RootComponent)
 	{
-		RootComponent->AddLocalRotation(DeltaRotation, bSweep);
+		RootComponent->AddLocalRotation(DeltaRotation, bSweep, OutSweepHitResult);
+	}
+	else if (OutSweepHitResult)
+	{
+		*OutSweepHitResult = FHitResult();
 	}
 }
 
-void AActor::AddActorLocalTransform(const FTransform& NewTransform, bool bSweep)
+void AActor::AddActorLocalTransform(const FTransform& NewTransform, bool bSweep, FHitResult* OutSweepHitResult)
 {
 	if(RootComponent)
 	{
-		RootComponent->AddLocalTransform(NewTransform, bSweep);
+		RootComponent->AddLocalTransform(NewTransform, bSweep, OutSweepHitResult);
+	}
+	else if (OutSweepHitResult)
+	{
+		*OutSweepHitResult = FHitResult();
 	}
 }
 
-void AActor::SetActorRelativeLocation(FVector NewRelativeLocation, bool bSweep/*=false*/)
+void AActor::SetActorRelativeLocation(FVector NewRelativeLocation, bool bSweep, FHitResult* OutSweepHitResult)
 {
 	if (RootComponent)
 	{
-		RootComponent->SetRelativeLocation(NewRelativeLocation, bSweep);
+		RootComponent->SetRelativeLocation(NewRelativeLocation, bSweep, OutSweepHitResult);
+	}
+	else if (OutSweepHitResult)
+	{
+		*OutSweepHitResult = FHitResult();
 	}
 }
 
-void AActor::SetActorRelativeRotation(FRotator NewRelativeRotation, bool bSweep/*=false*/)
+void AActor::SetActorRelativeRotation(FRotator NewRelativeRotation, bool bSweep, FHitResult* OutSweepHitResult)
 {
 	if (RootComponent)
 	{
-		RootComponent->SetRelativeRotation(NewRelativeRotation, bSweep);
+		RootComponent->SetRelativeRotation(NewRelativeRotation, bSweep, OutSweepHitResult);
+	}
+	else if (OutSweepHitResult)
+	{
+		*OutSweepHitResult = FHitResult();
 	}
 }
 
-void AActor::SetActorRelativeTransform(const FTransform& NewRelativeTransform, bool bSweep/*=false*/)
+void AActor::SetActorRelativeTransform(const FTransform& NewRelativeTransform, bool bSweep, FHitResult* OutSweepHitResult)
 {
 	if (RootComponent)
 	{
-		RootComponent->SetRelativeTransform(NewRelativeTransform);
+		RootComponent->SetRelativeTransform(NewRelativeTransform, bSweep, OutSweepHitResult);
+	}
+	else if (OutSweepHitResult)
+	{
+		*OutSweepHitResult = FHitResult();
 	}
 }
 
@@ -3431,4 +3477,71 @@ const int32 AActor::GetNumUncachedStaticLightingInteractions() const
 }
 #endif // WITH_EDITOR
 
+
+// K2 versions of various transform changing operations.
+// Note: we pass null for the hit result if not sweeping, for better perf.
+// This assumes this K2 function is only used by blueprints, which initializes the param for each function call.
+
+bool AActor::K2_SetActorLocation(FVector NewLocation, bool bSweep, FHitResult& SweepHitResult)
+{
+	return SetActorLocation(NewLocation, bSweep, (bSweep ? &SweepHitResult : nullptr));
+}
+
+bool AActor::K2_SetActorLocationAndRotation(FVector NewLocation, FRotator NewRotation, bool bSweep, FHitResult& SweepHitResult)
+{
+	return SetActorLocationAndRotation(NewLocation, NewRotation, bSweep, (bSweep ? &SweepHitResult : nullptr));
+}
+
+void AActor::K2_AddActorWorldOffset(FVector DeltaLocation, bool bSweep, FHitResult& SweepHitResult)
+{
+	AddActorWorldOffset(DeltaLocation, bSweep, (bSweep ? &SweepHitResult : nullptr));
+}
+
+void AActor::K2_AddActorWorldRotation(FRotator DeltaRotation, bool bSweep, FHitResult& SweepHitResult)
+{
+	AddActorWorldRotation(DeltaRotation, bSweep, (bSweep ? &SweepHitResult : nullptr));
+}
+
+void AActor::K2_AddActorWorldTransform(const FTransform& DeltaTransform, bool bSweep, FHitResult& SweepHitResult)
+{
+	AddActorWorldTransform(DeltaTransform, bSweep, (bSweep ? &SweepHitResult : nullptr));
+}
+
+bool AActor::K2_SetActorTransform(const FTransform& NewTransform, bool bSweep, FHitResult& SweepHitResult)
+{
+	return SetActorTransform(NewTransform, bSweep, (bSweep ? &SweepHitResult : nullptr));
+}
+
+void AActor::K2_AddActorLocalOffset(FVector DeltaLocation, bool bSweep, FHitResult& SweepHitResult)
+{
+	AddActorLocalOffset(DeltaLocation, bSweep, (bSweep ? &SweepHitResult : nullptr));
+}
+
+void AActor::K2_AddActorLocalRotation(FRotator DeltaRotation, bool bSweep, FHitResult& SweepHitResult)
+{
+	AddActorLocalRotation(DeltaRotation, bSweep, (bSweep ? &SweepHitResult : nullptr));
+}
+
+void AActor::K2_AddActorLocalTransform(const FTransform& NewTransform, bool bSweep, FHitResult& SweepHitResult)
+{
+	AddActorLocalTransform(NewTransform, bSweep, (bSweep ? &SweepHitResult : nullptr));
+}
+
+void AActor::K2_SetActorRelativeLocation(FVector NewRelativeLocation, bool bSweep, FHitResult& SweepHitResult)
+{
+	SetActorRelativeLocation(NewRelativeLocation, bSweep, (bSweep ? &SweepHitResult : nullptr));
+}
+
+void AActor::K2_SetActorRelativeRotation(FRotator NewRelativeRotation, bool bSweep, FHitResult& SweepHitResult)
+{
+	SetActorRelativeRotation(NewRelativeRotation, bSweep, (bSweep ? &SweepHitResult : nullptr));
+}
+
+void AActor::K2_SetActorRelativeTransform(const FTransform& NewRelativeTransform, bool bSweep, FHitResult& SweepHitResult)
+{
+	SetActorRelativeTransform(NewRelativeTransform, bSweep, (bSweep ? &SweepHitResult : nullptr));
+}
+
+
+
 #undef LOCTEXT_NAMESPACE
diff --git a/Engine/Source/Runtime/Engine/Private/PrimitiveComponent.cpp b/Engine/Source/Runtime/Engine/Private/PrimitiveComponent.cpp
--- a/Engine/Source/Runtime/Engine/Private/PrimitiveComponent.cpp
+++ b/Engine/Source/Runtime/Engine/Private/PrimitiveComponent.cpp
@@ -1256,6 +1256,10 @@ bool UPrimitiveComponent::MoveComponent( const FVector& Delta, const FRotator& N
 	if ( IsPendingKill() )
 	{
 		//UE_LOG(LogPrimitiveComponent, Log, TEXT("%s deleted move physics %d"),*Actor->GetName(),Actor->Physics);
+		if (OutHit)
+		{
+			*OutHit = FHitResult();
+		}
 		return false;
 	}
 
@@ -1266,6 +1270,10 @@ bool UPrimitiveComponent::MoveComponent( const FVector& Delta, const FRotator& N
 	{
 		// TODO: Static components without an owner can move, should they be able to?
 		UE_LOG(LogPrimitiveComponent, Warning, TEXT("Trying to move static component '%s' after initialization"), *GetFullName());
+		if (OutHit)
+		{
+			*OutHit = FHitResult();
+		}
 		return false;
 	}
 
@@ -1297,7 +1305,6 @@ bool UPrimitiveComponent::MoveComponent( const FVector& Delta, const FRotator& N
 			{
 				*OutHit = BlockingHit;
 			}
-
 			return true;
 		}
 		DeltaSizeSq = 0.f;
diff --git a/Engine/Source/Runtime/Engine/Private/SceneComponent.cpp b/Engine/Source/Runtime/Engine/Private/SceneComponent.cpp
--- a/Engine/Source/Runtime/Engine/Private/SceneComponent.cpp
+++ b/Engine/Source/Runtime/Engine/Private/SceneComponent.cpp
@@ -302,17 +302,17 @@ void USceneComponent::UpdateBounds()
 #endif
 }
 
-void USceneComponent::SetRelativeLocation(FVector NewLocation, bool bSweep)
+void USceneComponent::SetRelativeLocation(FVector NewLocation, bool bSweep, FHitResult* OutSweepHitResult)
 {
-	SetRelativeLocationAndRotation(NewLocation, RelativeRotation, bSweep);
+	SetRelativeLocationAndRotation(NewLocation, RelativeRotation, bSweep, OutSweepHitResult);
 }
 
-void USceneComponent::SetRelativeRotation(FRotator NewRotation, bool bSweep)
+void USceneComponent::SetRelativeRotation(FRotator NewRotation, bool bSweep, FHitResult* OutSweepHitResult)
 {
-	SetRelativeLocationAndRotation(RelativeLocation, NewRotation, bSweep);
+	SetRelativeLocationAndRotation(RelativeLocation, NewRotation, bSweep, OutSweepHitResult);
 }
 
-void USceneComponent::SetRelativeLocationAndRotation(FVector NewLocation, FRotator NewRotation, bool bSweep)
+void USceneComponent::SetRelativeLocationAndRotation(FVector NewLocation, FRotator NewRotation, bool bSweep, FHitResult* OutSweepHitResult)
 {
 	if(!NewLocation.Equals(RelativeLocation) || !NewRotation.Equals(RelativeRotation))
 	{
@@ -325,32 +325,40 @@ void USceneComponent::SetRelativeLocationAndRotation(FVector NewLocation, FRotat
 		const FTransform DesiredWorldTransform = CalcNewComponentToWorld(DesiredRelTransform);
 		const FVector DesiredDelta = DesiredWorldTransform.GetTranslation() - ComponentToWorld.GetTranslation();
 		
-		MoveComponent(DesiredDelta, DesiredWorldTransform.Rotator(), bSweep);
+		MoveComponent(DesiredDelta, DesiredWorldTransform.Rotator(), bSweep, OutSweepHitResult);
+	}
+	else if (OutSweepHitResult)
+	{
+		*OutSweepHitResult = FHitResult();
 	}
 }
 
-void USceneComponent::AddRelativeLocation(FVector DeltaLocation, bool bSweep)
+void USceneComponent::AddRelativeLocation(FVector DeltaLocation, bool bSweep, FHitResult* OutSweepHitResult)
 {
-	SetRelativeLocationAndRotation(RelativeLocation + DeltaLocation, RelativeRotation, bSweep);
+	SetRelativeLocationAndRotation(RelativeLocation + DeltaLocation, RelativeRotation, bSweep, OutSweepHitResult);
 }
 
-void USceneComponent::AddRelativeRotation(FRotator DeltaRotation, bool bSweep)
+void USceneComponent::AddRelativeRotation(FRotator DeltaRotation, bool bSweep, FHitResult* OutSweepHitResult)
 {
-	SetRelativeLocationAndRotation(RelativeLocation, RelativeRotation + DeltaRotation, bSweep);
+	SetRelativeLocationAndRotation(RelativeLocation, RelativeRotation + DeltaRotation, bSweep, OutSweepHitResult);
 }
 
-void USceneComponent::AddLocalOffset(FVector DeltaLocation, bool bSweep)
+void USceneComponent::AddLocalOffset(FVector DeltaLocation, bool bSweep, FHitResult* OutSweepHitResult)
 {
 	if (!DeltaLocation.IsNearlyZero())
 	{
 		const FQuat RelativeRotQuat = RelativeRotation.Quaternion();
 		const FVector LocalOffset = RelativeRotQuat.RotateVector(DeltaLocation);
 
-		SetRelativeLocationAndRotation(RelativeLocation + LocalOffset, RelativeRotation, bSweep);
+		SetRelativeLocationAndRotation(RelativeLocation + LocalOffset, RelativeRotation, bSweep, OutSweepHitResult);
+	}
+	else if (OutSweepHitResult)
+	{
+		*OutSweepHitResult = FHitResult();
 	}
 }
 
-void USceneComponent::AddLocalRotation(FRotator DeltaRotation, bool bSweep)
+void USceneComponent::AddLocalRotation(FRotator DeltaRotation, bool bSweep, FHitResult* OutSweepHitResult)
 {
 	if (!DeltaRotation.IsZero())
 	{
@@ -358,37 +366,49 @@ void USceneComponent::AddLocalRotation(FRotator DeltaRotation, bool bSweep)
 		const FQuat DeltaRotQuat = DeltaRotation.Quaternion();
 		const FQuat NewRelRotQuat = RelativeRotQuat * DeltaRotQuat;
 
-		SetRelativeLocationAndRotation(RelativeLocation, NewRelRotQuat.Rotator(), bSweep);
+		SetRelativeLocationAndRotation(RelativeLocation, NewRelRotQuat.Rotator(), bSweep, OutSweepHitResult);
+	}
+	else if (OutSweepHitResult)
+	{
+		*OutSweepHitResult = FHitResult();
 	}
 }
 
-void USceneComponent::AddLocalTransform(const FTransform& DeltaTransform, bool bSweep)
+void USceneComponent::AddLocalTransform(const FTransform& DeltaTransform, bool bSweep, FHitResult* OutSweepHitResult)
 {
 	const FTransform RelativeTransform( RelativeRotation, RelativeLocation, FVector(1,1,1) ); // don't use scaling, so it matches how AddLocalRotation/Offset work
 	const FTransform NewRelTransform = DeltaTransform * RelativeTransform;
 
-	SetRelativeTransform(NewRelTransform, bSweep);
+	SetRelativeTransform(NewRelTransform, bSweep, OutSweepHitResult);
 }
 
-void USceneComponent::AddWorldOffset(FVector DeltaLocation, bool bSweep)
+void USceneComponent::AddWorldOffset(FVector DeltaLocation, bool bSweep, FHitResult* OutSweepHitResult)
 {
 	if (!DeltaLocation.IsZero())
 	{
 		const FVector NewWorldLocation = DeltaLocation + ComponentToWorld.GetTranslation();
-		SetWorldLocation(NewWorldLocation, bSweep);
+		SetWorldLocation(NewWorldLocation, bSweep, OutSweepHitResult);
+	}
+	else if (OutSweepHitResult)
+	{
+		*OutSweepHitResult = FHitResult();
 	}
 }
 
-void USceneComponent::AddWorldRotation(FRotator DeltaRotation, bool bSweep)
+void USceneComponent::AddWorldRotation(FRotator DeltaRotation, bool bSweep, FHitResult* OutSweepHitResult)
 {
 	if (!DeltaRotation.IsZero())
 	{
 		const FQuat NewWorldRotation = DeltaRotation.Quaternion() * ComponentToWorld.GetRotation();
-		SetWorldRotation(NewWorldRotation.Rotator(), bSweep);
+		SetWorldRotation(NewWorldRotation.Rotator(), bSweep, OutSweepHitResult);
+	}
+	else if (OutSweepHitResult)
+	{
+		*OutSweepHitResult = FHitResult();
 	}
 }
 
-void USceneComponent::AddWorldTransform(const FTransform& DeltaTransform, bool bSweep)
+void USceneComponent::AddWorldTransform(const FTransform& DeltaTransform, bool bSweep, FHitResult* OutSweepHitResult)
 {
 	const FQuat NewWorldRotation = DeltaTransform.GetRotation() * ComponentToWorld.GetRotation();
 	const FVector NewWorldLocation = DeltaTransform.GetTranslation() + ComponentToWorld.GetTranslation();
@@ -421,9 +441,9 @@ void USceneComponent::ResetRelativeTransform()
 	SetRelativeScale3D(FVector(1.f));
 }
 
-void USceneComponent::SetRelativeTransform(const FTransform& NewTransform, bool bSweep)
+void USceneComponent::SetRelativeTransform(const FTransform& NewTransform, bool bSweep, FHitResult* OutSweepHitResult)
 {
-	SetRelativeLocationAndRotation(NewTransform.GetTranslation(), NewTransform.Rotator(), bSweep);
+	SetRelativeLocationAndRotation(NewTransform.GetTranslation(), NewTransform.Rotator(), bSweep, OutSweepHitResult);
 	SetRelativeScale3D(NewTransform.GetScale3D());
 }
 
@@ -433,7 +453,7 @@ FTransform USceneComponent::GetRelativeTransform()
 	return RelativeTransform;
 }
 
-void USceneComponent::SetWorldLocation(FVector NewLocation, bool bSweep)
+void USceneComponent::SetWorldLocation(FVector NewLocation, bool bSweep, FHitResult* OutSweepHitResult)
 {
 	FVector NewRelLocation = NewLocation;
 
@@ -444,10 +464,10 @@ void USceneComponent::SetWorldLocation(FVector NewLocation, bool bSweep)
 		NewRelLocation = ParentToWorld.InverseTransformPosition(NewLocation);
 	}
 
-	SetRelativeLocation(NewRelLocation, bSweep);
+	SetRelativeLocation(NewRelLocation, bSweep, OutSweepHitResult);
 }
 
-void USceneComponent::SetWorldRotation(FRotator NewRotation, bool bSweep)
+void USceneComponent::SetWorldRotation(FRotator NewRotation, bool bSweep, FHitResult* OutSweepHitResult)
 {
 	FRotator NewRelRotation = NewRotation;
 
@@ -460,7 +480,7 @@ void USceneComponent::SetWorldRotation(FRotator NewRotation, bool bSweep)
 		NewRelRotation = NewRelQuat.Rotator();
 	}
 
-	SetRelativeRotation(NewRelRotation, bSweep);
+	SetRelativeRotation(NewRelRotation, bSweep, OutSweepHitResult);
 }
 
 void USceneComponent::SetWorldScale3D(FVector NewScale)
@@ -477,7 +497,7 @@ void USceneComponent::SetWorldScale3D(FVector NewScale)
 	SetRelativeScale3D(NewRelScale);
 }
 
-void USceneComponent::SetWorldTransform(const FTransform& NewTransform, bool bSweep)
+void USceneComponent::SetWorldTransform(const FTransform& NewTransform, bool bSweep, FHitResult* OutSweepHitResult)
 {
 	// If attached to something, transform into local space
 	FQuat NewRotation = NewTransform.GetRotation();
@@ -505,15 +525,15 @@ void USceneComponent::SetWorldTransform(const FTransform& NewTransform, bool bSw
 		}
 	}
 
-	SetRelativeTransform(FTransform(NewRotation, NewLocation, NewScale), bSweep);
+	SetRelativeTransform(FTransform(NewRotation, NewLocation, NewScale), bSweep, OutSweepHitResult);
 }
 
-void USceneComponent::SetWorldLocationAndRotation(FVector NewLocation, FRotator NewRotation, bool bSweep)
+void USceneComponent::SetWorldLocationAndRotation(FVector NewLocation, FRotator NewRotation, bool bSweep, FHitResult* OutSweepHitResult)
 {
-	SetWorldLocationAndRotation(NewLocation, NewRotation.Quaternion(), bSweep);
+	SetWorldLocationAndRotation(NewLocation, NewRotation.Quaternion(), bSweep, OutSweepHitResult);
 }
 
-void USceneComponent::SetWorldLocationAndRotation(FVector NewLocation, const FQuat& NewRotation, bool bSweep)
+void USceneComponent::SetWorldLocationAndRotation(FVector NewLocation, const FQuat& NewRotation, bool bSweep, FHitResult* OutSweepHitResult)
 {
 	// If attached to something, transform into local space
 	FQuat NewFinalRotation = NewRotation;
@@ -534,7 +554,7 @@ void USceneComponent::SetWorldLocationAndRotation(FVector NewLocation, const FQu
 		}
 	}
 
-	SetRelativeLocationAndRotation(NewLocation, NewFinalRotation.Rotator(), bSweep);
+	SetRelativeLocationAndRotation(NewLocation, NewFinalRotation.Rotator(), bSweep, OutSweepHitResult);
 }
 
 void USceneComponent::SetWorldLocationAndRotationNoPhysics(const FVector& NewLocation, const FRotator& NewRotation)
@@ -1795,4 +1815,89 @@ void USceneComponent::UpdateNavigationData()
 	}
 }
 
+
+// K2 versions of various transform changing operations.
+// Note: we pass null for the hit result if not sweeping, for better perf.
+// This assumes this K2 function is only used by blueprints, which initializes the param for each function call.
+
+void USceneComponent::K2_SetRelativeLocationAndRotation(FVector NewLocation, FRotator NewRotation, bool bSweep, FHitResult& SweepHitResult)
+{
+	SetRelativeLocationAndRotation(NewLocation, NewRotation, bSweep, (bSweep ? &SweepHitResult : nullptr));
+}
+
+void USceneComponent::K2_SetWorldLocationAndRotation(FVector NewLocation, FRotator NewRotation, bool bSweep, FHitResult& SweepHitResult)
+{
+	SetWorldLocationAndRotation(NewLocation, NewRotation, bSweep, (bSweep ? &SweepHitResult : nullptr));
+}
+
+void USceneComponent::K2_SetRelativeLocation(FVector NewLocation, bool bSweep, FHitResult& SweepHitResult)
+{
+	SetRelativeLocation(NewLocation, bSweep, (bSweep ? &SweepHitResult : nullptr));
+}
+
+void USceneComponent::K2_SetRelativeRotation(FRotator NewRotation, bool bSweep, FHitResult& SweepHitResult)
+{
+	SetRelativeRotation(NewRotation, bSweep, (bSweep ? &SweepHitResult : nullptr));
+}
+
+void USceneComponent::K2_SetRelativeTransform(const FTransform& NewTransform, bool bSweep, FHitResult& SweepHitResult)
+{
+	SetRelativeTransform(NewTransform, bSweep, (bSweep ? &SweepHitResult : nullptr));
+}
+
+void USceneComponent::K2_AddRelativeLocation(FVector DeltaLocation, bool bSweep, FHitResult& SweepHitResult)
+{
+	AddRelativeLocation(DeltaLocation, bSweep, (bSweep ? &SweepHitResult : nullptr));
+}
+
+void USceneComponent::K2_AddRelativeRotation(FRotator DeltaRotation, bool bSweep, FHitResult& SweepHitResult)
+{
+	AddRelativeRotation(DeltaRotation, bSweep, (bSweep ? &SweepHitResult : nullptr));
+}
+
+void USceneComponent::K2_AddLocalOffset(FVector DeltaLocation, bool bSweep, FHitResult& SweepHitResult)
+{
+	AddLocalOffset(DeltaLocation, bSweep, (bSweep ? &SweepHitResult : nullptr));
+}
+
+void USceneComponent::K2_AddLocalRotation(FRotator DeltaRotation, bool bSweep, FHitResult& SweepHitResult)
+{
+	AddLocalRotation(DeltaRotation, bSweep, (bSweep ? &SweepHitResult : nullptr));
+}
+
+void USceneComponent::K2_AddLocalTransform(const FTransform& DeltaTransform, bool bSweep, FHitResult& SweepHitResult)
+{
+	AddLocalTransform(DeltaTransform, bSweep, (bSweep ? &SweepHitResult : nullptr));
+}
+
+void USceneComponent::K2_SetWorldLocation(FVector NewLocation, bool bSweep, FHitResult& SweepHitResult)
+{
+	SetWorldLocation(NewLocation, bSweep, (bSweep ? &SweepHitResult : nullptr));
+}
+
+void USceneComponent::K2_SetWorldRotation(FRotator NewRotation, bool bSweep, FHitResult& SweepHitResult)
+{
+	SetWorldRotation(NewRotation, bSweep, (bSweep ? &SweepHitResult : nullptr));
+}
+
+void USceneComponent::K2_SetWorldTransform(const FTransform& NewTransform, bool bSweep, FHitResult& SweepHitResult)
+{
+	SetWorldTransform(NewTransform, bSweep, (bSweep ? &SweepHitResult : nullptr));
+}
+
+void USceneComponent::K2_AddWorldOffset(FVector DeltaLocation, bool bSweep, FHitResult& SweepHitResult)
+{
+	AddWorldOffset(DeltaLocation, bSweep, (bSweep ? &SweepHitResult : nullptr));
+}
+
+void USceneComponent::K2_AddWorldRotation(FRotator DeltaRotation, bool bSweep, FHitResult& SweepHitResult)
+{
+	AddWorldRotation(DeltaRotation, bSweep, (bSweep ? &SweepHitResult : nullptr));
+}
+
+void USceneComponent::K2_AddWorldTransform(const FTransform& DeltaTransform, bool bSweep, FHitResult& SweepHitResult)
+{
+	AddWorldTransform(DeltaTransform, bSweep, (bSweep ? &SweepHitResult : nullptr));
+}
+
 #undef LOCTEXT_NAMESPACE
