AddUnique array function added to  BP.

diff --git a/Engine/Source/Runtime/Engine/Classes/Kismet/KismetArrayLibrary.h b/Engine/Source/Runtime/Engine/Classes/Kismet/KismetArrayLibrary.h
--- a/Engine/Source/Runtime/Engine/Classes/Kismet/KismetArrayLibrary.h
+++ b/Engine/Source/Runtime/Engine/Classes/Kismet/KismetArrayLibrary.h
@@ -15,7 +15,16 @@ class UKismetArrayLibrary : public UBlueprintFunctionLibrary
 	*/
 	UFUNCTION(BlueprintCallable, CustomThunk, meta=(FriendlyName = "Add", CompactNodeTitle = "ADD", ArrayParm = "TargetArray|ArrayProperty", ArrayTypeDependentParams = "NewItem", AutoCreateRefTerm = "NewItem"), Category="Utilities|Array")
 	static int32 Array_Add(const TArray<int32>& TargetArray, const UArrayProperty* ArrayProperty, const int32& NewItem);
-	
+
+	/**
+	*Add item to array (unique)
+	*
+	*@param	TargetArray		The array to add item to
+	*@param	NewItem			The item to add to the array
+	*/
+	UFUNCTION(BlueprintCallable, CustomThunk, meta = (FriendlyName = "Add Unique", CompactNodeTitle = "ADDUNIQUE", ArrayParm = "TargetArray|ArrayProperty", ArrayTypeDependentParams = "NewItem", AutoCreateRefTerm = "NewItem"), Category = "Utilities|Array")
+	static int32 Array_AddUnique(const TArray<int32>& TargetArray, const UArrayProperty* ArrayProperty, const int32& NewItem);
+
 	/** 
 	 * Shuffle (randomize) the elements of an array
 	 *
@@ -158,6 +167,7 @@ class UKismetArrayLibrary : public UBlueprintFunctionLibrary
 
 	// Native functions that will be called by the below custom thunk layers, which read off the property address, and call the appropriate native handler
 	ENGINE_API static int32 GenericArray_Add(void* TargetArray, const UArrayProperty* ArrayProp, const void* NewItem);
+	ENGINE_API static int32 GenericArray_AddUnique(void* TargetArray, const UArrayProperty* ArrayProp, const void* NewItem);
 	ENGINE_API static void GenericArray_Shuffle(void* TargetArray, const UArrayProperty* ArrayProp);
 	ENGINE_API static void GenericArray_Append(void* TargetArray, const UArrayProperty* TargetArrayProp, void* SourceArray, const UArrayProperty* SourceArrayProperty);
 	ENGINE_API static void GenericArray_Insert(void* TargetArray, const UArrayProperty* ArrayProp, const void* NewItem, int32 Index);
@@ -207,6 +217,30 @@ public:
 		InnerProp->DestroyValue(StorageSpace);
 	}
 
+	DECLARE_FUNCTION(execArray_AddUnique)
+	{
+		Stack.StepCompiledIn<UArrayProperty>(NULL);
+		void* ArrayAddr = Stack.MostRecentPropertyAddress;
+
+		P_GET_OBJECT(UArrayProperty, ArrayProperty);
+
+		// Since NewItem isn't really an int, step the stack manually
+		const UProperty* InnerProp = ArrayProperty->Inner;
+		const int32 PropertySize = InnerProp->ElementSize * InnerProp->ArrayDim;
+		void* StorageSpace = FMemory_Alloca(PropertySize);
+		InnerProp->InitializeValue(StorageSpace);
+
+		Stack.MostRecentPropertyAddress = NULL;
+		Stack.StepCompiledIn<UProperty>(StorageSpace);
+		void* NewItemPtr = (Stack.MostRecentPropertyAddress != NULL) ? Stack.MostRecentPropertyAddress : StorageSpace;
+
+		P_FINISH;
+
+		*(int32*)Result = GenericArray_AddUnique(ArrayAddr, ArrayProperty, NewItemPtr);
+
+		InnerProp->DestroyValue(StorageSpace);
+	}
+
 	DECLARE_FUNCTION(execArray_Shuffle)
 	{
 		Stack.StepCompiledIn<UArrayProperty>(NULL);
diff --git a/Engine/Source/Runtime/Engine/Private/KismetArrayLibrary.cpp b/Engine/Source/Runtime/Engine/Private/KismetArrayLibrary.cpp
--- a/Engine/Source/Runtime/Engine/Private/KismetArrayLibrary.cpp
+++ b/Engine/Source/Runtime/Engine/Private/KismetArrayLibrary.cpp
@@ -28,7 +28,7 @@ void UKismetArrayLibrary::FilterArray(const TArray<AActor*>& TargetArray, TSubcl
 
 int32 UKismetArrayLibrary::GenericArray_Add(void* TargetArray, const UArrayProperty* ArrayProp, const void* NewItem)
 {
-	int32 NewIndex = 0;
+	int32 NewIndex = INDEX_NONE;
 	if( TargetArray )
 	{
 		FScriptArrayHelper ArrayHelper(ArrayProp, TargetArray);
@@ -40,6 +40,23 @@ int32 UKismetArrayLibrary::GenericArray_Add(void* TargetArray, const UArrayPrope
 	return NewIndex;
 }
 
+int32 UKismetArrayLibrary::GenericArray_AddUnique(void* TargetArray, const UArrayProperty* ArrayProp, const void* NewItem)
+{
+	int32 NewIndex = INDEX_NONE;
+	if (TargetArray)
+	{
+		FScriptArrayHelper ArrayHelper(ArrayProp, TargetArray);
+		UProperty* InnerProp = ArrayProp->Inner;
+
+		if (GenericArray_Find(TargetArray, ArrayProp, NewItem) == INDEX_NONE)
+		{
+			NewIndex = ArrayHelper.AddValue();
+			InnerProp->CopySingleValueToScriptVM(ArrayHelper.GetRawPtr(NewIndex), NewItem);
+		}
+	}
+	return NewIndex;
+}
+
 void UKismetArrayLibrary::GenericArray_Append(void* TargetArray, const UArrayProperty* TargetArrayProp, void* SourceArray, const UArrayProperty* SourceArrayProperty)
 {
 	if(TargetArray && SourceArray)
diff --git a/Engine/Source/Runtime/Engine/Public/GeneratedCodeHelpers.h b/Engine/Source/Runtime/Engine/Public/GeneratedCodeHelpers.h
--- a/Engine/Source/Runtime/Engine/Public/GeneratedCodeHelpers.h
+++ b/Engine/Source/Runtime/Engine/Public/GeneratedCodeHelpers.h
@@ -21,6 +21,12 @@ struct FCustomThunkTemplates
 		return UKismetArrayLibrary::GenericArray_Add(&TargetArray, ArrayProperty, &NewItem);
 	}
 
+	template<typename T>
+	static int32 Array_AddUnique(TArray<T>& TargetArray, const UArrayProperty* ArrayProperty, const T& NewItem)
+	{
+		return UKismetArrayLibrary::GenericArray_AddUnique(&TargetArray, ArrayProperty, &NewItem);
+	}
+
 	template<typename T>
 	static void Array_Shuffle(TArray<T>& TargetArray, const UArrayProperty* ArrayProperty)
 	{
