Exposed axis locking to blueprints (from user wolfadex). Axis locking allows for runtime constraint of primitive component to a plane.

GitHub pull request #524

diff --git a/Engine/Source/Runtime/Engine/Classes/Components/PrimitiveComponent.h b/Engine/Source/Runtime/Engine/Classes/Components/PrimitiveComponent.h
--- a/Engine/Source/Runtime/Engine/Classes/Components/PrimitiveComponent.h
+++ b/Engine/Source/Runtime/Engine/Classes/Components/PrimitiveComponent.h
@@ -655,6 +655,13 @@ public:
 	 */
 	virtual bool CanEditSimulatePhysics();
 
+	/**
+	 * Sets the specified axis as locked, preventing movement along that axis.
+	 * @param LockedAxis	The axis to lock.
+	 */
+	UFUNCTION(BlueprintCallable, meta = (FriendlyName = "Set Locked Axis", Keywords = "set locked axis constraint physics"), Category = Physics)
+	virtual void SetLockedAxis(ELockedAxis::Type LockedAxis);
+
 	/**
 	 *	Add an impulse to a single rigid body. Good for one time instant burst.
 	 *
diff --git a/Engine/Source/Runtime/Engine/Classes/PhysicsEngine/BodyInstance.h b/Engine/Source/Runtime/Engine/Classes/PhysicsEngine/BodyInstance.h
--- a/Engine/Source/Runtime/Engine/Classes/PhysicsEngine/BodyInstance.h
+++ b/Engine/Source/Runtime/Engine/Classes/PhysicsEngine/BodyInstance.h
@@ -27,7 +27,7 @@ namespace physx
 }
 #endif // WITH_PHYSX
 
-UENUM()
+UENUM(BlueprintType)
 namespace ELockedAxis
 {
 	enum Type
@@ -216,6 +216,9 @@ public:
 	UPROPERTY(EditAnywhere, Category = Physics)
 	FVector CustomLockedAxis;
 
+	/** Locks physical movement along axis. */
+	void SetDOFLock(ELockedAxis::Type NewAxisMode);
+
 	FVector GetLockedAxis() const;
 	void CreateDOFLock();
 
diff --git a/Engine/Source/Runtime/Engine/Private/PhysicsEngine/BodyInstance.cpp b/Engine/Source/Runtime/Engine/Private/PhysicsEngine/BodyInstance.cpp
--- a/Engine/Source/Runtime/Engine/Private/PhysicsEngine/BodyInstance.cpp
+++ b/Engine/Source/Runtime/Engine/Private/PhysicsEngine/BodyInstance.cpp
@@ -547,6 +547,13 @@ FVector FBodyInstance::GetLockedAxis() const
 	return FVector::ZeroVector;
 }
 
+void FBodyInstance::SetDOFLock(ELockedAxis::Type NewAxisMode)
+{
+	LockedAxisMode = NewAxisMode;
+
+	CreateDOFLock();
+}
+
 void FBodyInstance::CreateDOFLock()
 {
 	if (DOFConstraint)
diff --git a/Engine/Source/Runtime/Engine/Private/PrimitiveComponentPhysics.cpp b/Engine/Source/Runtime/Engine/Private/PrimitiveComponentPhysics.cpp
--- a/Engine/Source/Runtime/Engine/Private/PrimitiveComponentPhysics.cpp
+++ b/Engine/Source/Runtime/Engine/Private/PrimitiveComponentPhysics.cpp
@@ -148,6 +148,18 @@ void UPrimitiveComponent::SetSimulatePhysics(bool bSimulate)
 	BodyInstance.SetInstanceSimulatePhysics(bSimulate);
 }
 
+void UPrimitiveComponent::SetLockedAxis(ELockedAxis::Type LockedAxis)
+{
+	FBodyInstance * RootBI = GetBodyInstance(NAME_None, false);
+
+	if (RootBI == NULL || IsPendingKill())
+	{
+		return;
+	}
+
+	RootBI->SetDOFLock(LockedAxis);
+}
+
 void UPrimitiveComponent::AddImpulse(FVector Impulse, FName BoneName, bool bVelChange)
 {
 	FBodyInstance* BI = GetBodyInstance(BoneName);
