Navmesh with fixed tile pool size

Added two properties to ARecastNavmesh (bFixedTilePoolSize, TilePoolSize) which can be used to specify fixed tile pool size for a Navmesh (disabled by default).
Navmesh allocates space for specified number of tiles once at creation and never changes it. This enables support for adding/removing/updating navigation bounds in the world without rebuilding whole Navmesh. So navigation bounds can be placed in a streaming levels and navigation system will update only part of Navmesh where bounds were streamed in/out. This only works with runtime generation for now, next step is to make it work with serialized Navmesh.

diff --git a/Engine/Source/Runtime/Engine/Classes/AI/Navigation/NavigationData.h b/Engine/Source/Runtime/Engine/Classes/AI/Navigation/NavigationData.h
--- a/Engine/Source/Runtime/Engine/Classes/AI/Navigation/NavigationData.h
+++ b/Engine/Source/Runtime/Engine/Classes/AI/Navigation/NavigationData.h
@@ -410,7 +410,8 @@ class ENGINE_API ANavigationData : public AActor
 	virtual void RerunConstructionScripts() override;
 
 	virtual bool NeedsRebuild() const { return false; }
-	virtual bool CanRebuild() const;
+	virtual bool SupportsRuntimeGeneration() const;
+	virtual void OnNavigationBoundsChanged();
 	
 	//----------------------------------------------------------------------//
 	// Generation & data access                                                      
@@ -428,9 +429,9 @@ protected:
 	virtual void FillConfig(FNavDataConfig& Dest) { Dest = NavDataConfig; }
 
 public:
+	virtual void ConstructGenerator();
 	virtual void RebuildAll();
 	virtual void TickAsyncBuild(float DeltaSeconds);
-
 	/** Retrieves navmesh's generator */
 	FNavDataGenerator* GetGenerator() { return NavDataGenerator.Get(); }
 	const FNavDataGenerator* GetGenerator() const { return NavDataGenerator.Get(); }
diff --git a/Engine/Source/Runtime/Engine/Classes/AI/Navigation/NavigationSystem.h b/Engine/Source/Runtime/Engine/Classes/AI/Navigation/NavigationSystem.h
--- a/Engine/Source/Runtime/Engine/Classes/AI/Navigation/NavigationSystem.h
+++ b/Engine/Source/Runtime/Engine/Classes/AI/Navigation/NavigationSystem.h
@@ -174,8 +174,8 @@ class ENGINE_API UNavigationSystem : public UBlueprintFunctionLibrary
 
 	TSet<FNavigationDirtyElement> PendingOctreeUpdates;
 
-	UPROPERTY(transient)
-	TArray<ANavMeshBoundsVolume*> PendingNavVolumeUpdates;
+	// List of pending navigation bounds update requests (add, remove, update size)
+	TArray<FNavigationBoundsUpdateRequest> PendingNavBoundsUpdates;
 
  	UPROPERTY(/*BlueprintAssignable, */Transient)
 	FOnNavDataRegistered OnNavDataRegisteredEvent;
@@ -190,6 +190,9 @@ private:
 	/** set to true when navigation processing was blocked due to missing nav bounds */
 	uint32 bNavDataRemovedDueToMissingNavBounds:1;
 
+	/** All areas where we build/have navigation */
+	TSet<FNavigationBounds> RegisteredNavBounds;
+
 public:
 	//----------------------------------------------------------------------//
 	// Kismet functions
@@ -391,6 +394,8 @@ public:
 
 	bool IsNavigationRelevant(const AActor* TestActor) const;
 
+	const TSet<FNavigationBounds>& GetNavigationBounds() const;
+
 	/** @return default walkable area class */
 	FORCEINLINE static TSubclassOf<UNavArea> GetDefaultWalkableArea() { return DefaultWalkableArea; }
 
@@ -428,7 +433,10 @@ protected:
 	void ProcessNavAreaPendingRegistration();
 
 	/** used to apply updates of nav volumes in navigation system's tick */
-	void PerformNavigationBoundsUpdate(ANavMeshBoundsVolume* NavVolume);
+	void PerformNavigationBoundsUpdate(const TArray<FNavigationBoundsUpdateRequest>& UpdateRequests);
+	
+	/** Searches for all valid navigation bounds in the world and stores them */
+	void GatherNavigationBounds();
 
 	/** @return pointer to ANavigationData instance of given ID, or NULL if it was not found. Note it looks only through registered navigation data */
 	ANavigationData* GetNavDataWithID(const uint16 NavDataID) const;
@@ -530,6 +538,8 @@ public:
 	// @todo document
 	UFUNCTION(BlueprintCallable, Category = Navigation)
 	void OnNavigationBoundsUpdated(ANavMeshBoundsVolume* NavVolume);
+	void OnNavigationBoundsAdded(ANavMeshBoundsVolume* NavVolume);
+	void OnNavigationBoundsRemoved(ANavMeshBoundsVolume* NavVolume);
 
 	/** Used to display "navigation building in progress" notify */
 	bool IsNavigationBuildInProgress(bool bCheckDirtyToo = true);
@@ -729,6 +739,8 @@ private:
 	void NavigationBuildingLock();
 	// @todo document
 	void NavigationBuildingUnlock(bool bForce = false);
+	// adds navigation bounds update request to a pending list
+	void AddNavigationBoundsUpdateRequest(const FNavigationBoundsUpdateRequest& UpdateRequest);
 
 	void SpawnMissingNavigationData();
 
diff --git a/Engine/Source/Runtime/Engine/Classes/AI/Navigation/NavigationTypes.h b/Engine/Source/Runtime/Engine/Classes/AI/Navigation/NavigationTypes.h
--- a/Engine/Source/Runtime/Engine/Classes/AI/Navigation/NavigationTypes.h
+++ b/Engine/Source/Runtime/Engine/Classes/AI/Navigation/NavigationTypes.h
@@ -65,9 +65,10 @@ namespace ENavigationDirtyFlag
 {
 	enum Type
 	{
-		Geometry		= (1 << 0),
-		DynamicModifier	= (1 << 1),
-		UseAgentHeight  = (1 << 2),
+		Geometry			= (1 << 0),
+		DynamicModifier		= (1 << 1),
+		UseAgentHeight		= (1 << 2),
+		NavigationBounds	= (1 << 3),
 
 		All				= Geometry | DynamicModifier,		// all rebuild steps here without additional flags
 	};
@@ -83,6 +84,36 @@ struct FNavigationDirtyArea
 	FORCEINLINE bool HasFlag(ENavigationDirtyFlag::Type Flag) const { return (Flags & Flag) != 0; }
 };
 
+struct FNavigationBounds
+{
+	uint32	UniqueID;
+	FBox	AreaBox;	
+
+	bool operator==(const FNavigationBounds& Other) const 
+	{ 
+		return UniqueID == Other.UniqueID; 
+	}
+
+	friend uint32 GetTypeHash(const FNavigationBounds& NavBounds)
+	{
+		return GetTypeHash(NavBounds.UniqueID);
+	}
+};
+
+struct FNavigationBoundsUpdateRequest 
+{
+	FNavigationBounds NavBounds;
+	
+	enum Type
+	{
+		Added,
+		Removed,
+		Updated,
+	};
+
+	Type UpdateRequest;
+};
+
 struct FNavigationDirtyElement
 {
 	/** object owning this element */
diff --git a/Engine/Source/Runtime/Engine/Classes/AI/Navigation/RecastNavMesh.h b/Engine/Source/Runtime/Engine/Classes/AI/Navigation/RecastNavMesh.h
--- a/Engine/Source/Runtime/Engine/Classes/AI/Navigation/RecastNavMesh.h
+++ b/Engine/Source/Runtime/Engine/Classes/AI/Navigation/RecastNavMesh.h
@@ -448,6 +448,14 @@ class ENGINE_API ARecastNavMesh : public ANavigationData
 	// NavMesh generation parameters
 	//----------------------------------------------------------------------//
 
+	/** if true, the NavMesh will allocate fixed size pool for tiles, should be enabled to support streaming */
+	UPROPERTY(EditAnywhere, Category=Generation, config)
+	uint32 bFixedTilePoolSize:1;
+
+	/** maximum number of tiles NavMesh can hold */
+	UPROPERTY(EditAnywhere, Category=Generation, config, meta=(editcondition = "bFixedTilePoolSize"))
+	int32 TilePoolSize;
+
 	/** size of single tile, expressed in uu */
 	UPROPERTY(EditAnywhere, Category=Generation, config, meta=(ClampMin = "300.0"))
 	float TileSizeUU;
@@ -656,6 +664,9 @@ protected:
 	void SerializeRecastNavMesh(FArchive& Ar, FPImplRecastNavMesh*& NavMesh);
 
 public:
+	/** Whether NavMesh should adjust his tile pool size when NavBounds are changed */
+	bool IsResizable() const;
+
 	/** Returns bounding box for the whole navmesh. */
 	FBox GetNavMeshBounds() const;
 
@@ -837,7 +848,8 @@ public:
 	FORCEINLINE static FNavPolyFlags GetNavLinkFlag() { return NavLinkFlag; }
 	
 	virtual bool NeedsRebuild() const override;
-	virtual void RebuildAll() override;
+	virtual bool SupportsRuntimeGeneration() const override;
+	virtual void ConstructGenerator() override;
 
 protected:
 	// @todo docuement
diff --git a/Engine/Source/Runtime/Engine/Private/AI/Navigation/NavigationData.cpp b/Engine/Source/Runtime/Engine/Private/AI/Navigation/NavigationData.cpp
--- a/Engine/Source/Runtime/Engine/Private/AI/Navigation/NavigationData.cpp
+++ b/Engine/Source/Runtime/Engine/Private/AI/Navigation/NavigationData.cpp
@@ -253,7 +253,8 @@ void ANavigationData::OnRegistered()
 { 
 	InstantiateAndRegisterRenderingComponent();
 
-	bRegistered = true; 
+	bRegistered = true;
+	ConstructGenerator();
 }
 
 void ANavigationData::OnUnregistered()
@@ -346,14 +347,37 @@ void ANavigationData::CleanUpAndMarkPendingKill()
 	MarkComponentsAsPendingKill();
 }
 
-bool ANavigationData::CanRebuild() const
+bool ANavigationData::SupportsRuntimeGeneration() const
+{
+	return false;
+}
+
+void ANavigationData::ConstructGenerator()
 {
-	return NavDataGenerator.Get() != nullptr;
 }
 
 void ANavigationData::RebuildAll()
 {
+	ConstructGenerator(); //recreate generator
+	
+	if (NavDataGenerator)
+	{
+		NavDataGenerator->RebuildAll();
+	}
+}
 
+void ANavigationData::OnNavigationBoundsChanged()
+{
+	// Create generator if it wasn't yet
+	if (SupportsRuntimeGeneration() && NavDataGenerator.Get() == nullptr)
+	{
+		ConstructGenerator();
+	}
+	
+	if (NavDataGenerator)
+	{
+		NavDataGenerator->OnNavigationBoundsChanged();
+	}
 }
 
 void ANavigationData::TickAsyncBuild(float DeltaSeconds)
diff --git a/Engine/Source/Runtime/Engine/Private/AI/Navigation/NavigationSystem.cpp b/Engine/Source/Runtime/Engine/Private/AI/Navigation/NavigationSystem.cpp
--- a/Engine/Source/Runtime/Engine/Private/AI/Navigation/NavigationSystem.cpp
+++ b/Engine/Source/Runtime/Engine/Private/AI/Navigation/NavigationSystem.cpp
@@ -343,6 +343,9 @@ void UNavigationSystem::OnWorldInitDone(FNavigationSystem::EMode Mode)
 		InitializeLevelCollisions();
 		PopulateNavOctree();
 		
+		// gather navigable bounds
+		GatherNavigationBounds();
+		
 		// gather all navigation data instances and register all not-yet-registered
 		// (since it's quite possible navigation system was not ready by the time
 		// those instances were serialized-in or spawned)
@@ -372,18 +375,6 @@ void UNavigationSystem::OnWorldInitDone(FNavigationSystem::EMode Mode)
 			SpawnMissingNavigationData();
 			// in case anything spawned has registered
 			ProcessRegistrationCandidates();
-
-			for (int32 NavDataIndex = 0; NavDataIndex < NavDataSet.Num(); ++NavDataIndex)
-			{
-				ANavigationData* NavData = NavDataSet[NavDataIndex];
-				if (NavData != NULL)
-				{
-					if (bInitialBuildingLockActive == false && bNavigationAutoUpdateEnabled)
-					{
-						NavData->RebuildAll();
-					}
-				}
-			}
 		}
 		else
 		{
@@ -454,13 +445,10 @@ void UNavigationSystem::Tick(float DeltaSeconds)
 		ProcessNavAreaPendingRegistration();
 	}
 
-	if (PendingNavVolumeUpdates.Num() > 0)
+	if (PendingNavBoundsUpdates.Num() > 0)
 	{
-		for (auto Volume : PendingNavVolumeUpdates)
-		{
-			PerformNavigationBoundsUpdate(Volume);
-		}
-		PendingNavVolumeUpdates.Reset();
+		PerformNavigationBoundsUpdate(PendingNavBoundsUpdates);
+		PendingNavBoundsUpdates.Reset();
 	}
 
 	if (PendingOctreeUpdates.Num() > 0)
@@ -1243,6 +1231,11 @@ FBox UNavigationSystem::GetLevelBounds(ULevel* InLevel) const
 	return NavigableLevelBounds;
 }
 
+const TSet<FNavigationBounds>& UNavigationSystem::GetNavigationBounds() const
+{
+	return RegisteredNavBounds;
+}
+
 void UNavigationSystem::ApplyWorldOffset(const FVector& InOffset, bool bWorldShift)
 {
 	for (int32 NavDataIndex = 0; NavDataIndex < NavDataSet.Num(); ++NavDataIndex)
@@ -2354,10 +2347,59 @@ void UNavigationSystem::OnNavigationBoundsUpdated(ANavMeshBoundsVolume* NavVolum
 		return;
 	}
 
-	PendingNavVolumeUpdates.AddUnique(NavVolume);
+	FNavigationBoundsUpdateRequest UpdateRequest;
+	UpdateRequest.NavBounds.UniqueID	= NavVolume->GetUniqueID();
+	UpdateRequest.NavBounds.AreaBox		= NavVolume->GetComponentsBoundingBox(true);
+	UpdateRequest.UpdateRequest			= FNavigationBoundsUpdateRequest::Updated;
+	AddNavigationBoundsUpdateRequest(UpdateRequest);
 }
 
-void UNavigationSystem::PerformNavigationBoundsUpdate(ANavMeshBoundsVolume* NavVolume)
+void UNavigationSystem::OnNavigationBoundsAdded(ANavMeshBoundsVolume* NavVolume)
+{
+	if (NavVolume == NULL)
+	{
+		return;
+	}
+
+	FNavigationBoundsUpdateRequest UpdateRequest;
+	UpdateRequest.NavBounds.UniqueID	= NavVolume->GetUniqueID();
+	UpdateRequest.NavBounds.AreaBox		= NavVolume->GetComponentsBoundingBox(true);
+	UpdateRequest.UpdateRequest	= FNavigationBoundsUpdateRequest::Added;
+	AddNavigationBoundsUpdateRequest(UpdateRequest);
+}
+
+void UNavigationSystem::OnNavigationBoundsRemoved(ANavMeshBoundsVolume* NavVolume)
+{
+	if (NavVolume == NULL)
+	{
+		return;
+	}
+	
+	FNavigationBoundsUpdateRequest UpdateRequest;
+	UpdateRequest.NavBounds.UniqueID	= NavVolume->GetUniqueID();
+	UpdateRequest.NavBounds.AreaBox		= NavVolume->GetComponentsBoundingBox(true);
+	UpdateRequest.UpdateRequest	= FNavigationBoundsUpdateRequest::Removed;
+	AddNavigationBoundsUpdateRequest(UpdateRequest);
+}
+
+void UNavigationSystem::AddNavigationBoundsUpdateRequest(const FNavigationBoundsUpdateRequest& UpdateRequest)
+{
+	int32 ExistingIdx = PendingNavBoundsUpdates.IndexOfByPredicate([&](const FNavigationBoundsUpdateRequest& Element) {
+		return UpdateRequest.NavBounds.UniqueID == Element.NavBounds.UniqueID;
+	});
+
+	if (ExistingIdx != INDEX_NONE)
+	{
+		// Overwrite any previous updates
+		PendingNavBoundsUpdates[ExistingIdx] = UpdateRequest;
+	}
+	else
+	{
+		PendingNavBoundsUpdates.Add(UpdateRequest);
+	}
+}
+
+void UNavigationSystem::PerformNavigationBoundsUpdate(const TArray<FNavigationBoundsUpdateRequest>& UpdateRequests)
 {
 	if (bNavDataRemovedDueToMissingNavBounds)
 	{
@@ -2378,17 +2420,87 @@ void UNavigationSystem::PerformNavigationBoundsUpdate(ANavMeshBoundsVolume* NavV
 			ProcessRegistrationCandidates();
 		}
 	}
+	
+	// Create list of areas that needs to be updated
+	TArray<FBox> UpdatedAreas;
+	for (const FNavigationBoundsUpdateRequest& Request : UpdateRequests)
+	{
+		FSetElementId ExistingElementId = RegisteredNavBounds.FindId(Request.NavBounds);
+
+		switch (Request.UpdateRequest)
+		{
+		case FNavigationBoundsUpdateRequest::Removed:
+			{
+				if (ExistingElementId.IsValidId())
+				{
+					UpdatedAreas.Add(RegisteredNavBounds[ExistingElementId].AreaBox);
+					RegisteredNavBounds.Remove(ExistingElementId);
+				}
+			}
+			break;
+
+		case FNavigationBoundsUpdateRequest::Added:
+		case FNavigationBoundsUpdateRequest::Updated:
+			{
+				if (ExistingElementId.IsValidId())
+				{
+					FBox ExistingBox = RegisteredNavBounds[ExistingElementId].AreaBox;
+
+					if (!(ExistingBox == Request.NavBounds.AreaBox))
+					{
+						UpdatedAreas.Add(ExistingBox);
+						RegisteredNavBounds[ExistingElementId] = Request.NavBounds;
+					}
+				}
+				else
+				{
+					ExistingElementId = RegisteredNavBounds.Add(Request.NavBounds);
+				}
+				
+				UpdatedAreas.Add(Request.NavBounds.AreaBox);
+			}
+
+			break;
+		}
+	}
 
 #if WITH_RECAST
 	if (IsNavigationBuildingLocked() == false)
 	{
-		// rebuild all navmeshes
-		// NOTE: this will most probably be done in editor, if required to be performed at runtime may require some optimizations
-		RebuildAll();
+		if (UpdatedAreas.Num())
+		{
+			for (ANavigationData* NavData : NavDataSet)
+			{
+				if (NavData)
+				{
+					NavData->OnNavigationBoundsChanged();	
+				}
+			}
+		}
+				
+		// Propagate to generators areas that needs to be updated
+		AddDirtyAreas(UpdatedAreas, ENavigationDirtyFlag::All | ENavigationDirtyFlag::NavigationBounds);
 	}
 #endif // WITH_RECAST
 }
 
+void UNavigationSystem::GatherNavigationBounds()
+{
+	// Gather all available navigation bounds
+	RegisteredNavBounds.Empty();
+	for (TActorIterator<ANavMeshBoundsVolume> It(GetWorld()); It; ++It)
+	{
+		ANavMeshBoundsVolume* V = (*It);
+		if (V != nullptr && !V->IsPendingKill())
+		{
+			FNavigationBounds NavBounds;
+			NavBounds.UniqueID	= V->GetUniqueID();
+			NavBounds.AreaBox	= V->GetComponentsBoundingBox(true);
+			RegisteredNavBounds.Add(NavBounds);
+		}
+	}
+}
+
 void UNavigationSystem::Build()
 {
 	if (IsThereAnywhereToBuildNavigation() == false)
@@ -2581,6 +2693,8 @@ void UNavigationSystem::RebuildAll()
 {
 	const bool bIsInGame = GetWorld()->IsGameWorld();
 	
+	GatherNavigationBounds();
+
 	for (int32 NavDataIndex = 0; NavDataIndex < NavDataSet.Num(); ++NavDataIndex)
 	{
 		ANavigationData* NavData = NavDataSet[NavDataIndex];
@@ -2976,7 +3090,7 @@ bool UNavigationSystem::CanRebuildDirtyNavigation() const
 	for (int32 NavDataIndex = 0; NavDataIndex < NavDataSet.Num(); ++NavDataIndex)
 	{
 		const bool bIsDirty = NavDataSet[NavDataIndex]->NeedsRebuild();
-		const bool bCanRebuild = NavDataSet[NavDataIndex]->CanRebuild();
+		const bool bCanRebuild = NavDataSet[NavDataIndex]->SupportsRuntimeGeneration();
 
 		if (bIsDirty && !bCanRebuild)
 		{
diff --git a/Engine/Source/Runtime/Engine/Private/AI/Navigation/RecastNavMesh.cpp b/Engine/Source/Runtime/Engine/Private/AI/Navigation/RecastNavMesh.cpp
--- a/Engine/Source/Runtime/Engine/Private/AI/Navigation/RecastNavMesh.cpp
+++ b/Engine/Source/Runtime/Engine/Private/AI/Navigation/RecastNavMesh.cpp
@@ -164,6 +164,7 @@ ARecastNavMesh::ARecastNavMesh(const FObjectInitializer& ObjectInitializer)
 	, bDrawNavLinks(true)
 	, bDistinctlyDrawTilesBeingBuilt(true)
 	, DrawOffset(10.f)
+	, TilePoolSize(1024)
 	, MaxSimplificationError(1.3f)	// from RecastDemo
 	, DefaultMaxSearchNodes(RECAST_MAX_SEARCH_NODES)
 	, DefaultMaxHierarchicalSearchNodes(RECAST_MAX_SEARCH_NODES)
@@ -787,6 +788,11 @@ int32 ARecastNavMesh::GetNavMeshTilesCount() const
 	return NumTiles;
 }
 
+bool ARecastNavMesh::IsResizable() const
+{
+	return !bFixedTilePoolSize;
+}
+
 void ARecastNavMesh::GetEdgesForPathCorridor(const TArray<NavNodeRef>* PathCorridor, TArray<FNavigationPortalEdge>* PathCorridorEdges) const
 {
 	check(PathCorridor != NULL && PathCorridorEdges != NULL);
@@ -1576,16 +1582,18 @@ bool ARecastNavMesh::NeedsRebuild() const
 	return bLooksLikeNeeded;
 }
 
-void ARecastNavMesh::RebuildAll()
+bool ARecastNavMesh::SupportsRuntimeGeneration() const
 {
-	Super::RebuildAll();
+	// Generator should be enabled in the editor and if navmesh supports runtime generation
+	return (bRebuildAtRuntime || (GetWorld() && !GetWorld()->IsGameWorld()));
+}
 
+void ARecastNavMesh::ConstructGenerator()
+{
 	NavDataGenerator.Reset();
-	// Generator should be enabled in the editor and if navmesh supports runtime generation
-	if (bRebuildAtRuntime || (GetWorld() && !GetWorld()->IsGameWorld()))
+	if (SupportsRuntimeGeneration())
 	{
 		NavDataGenerator.Reset(new FRecastNavMeshGenerator(this));
-		NavDataGenerator->RebuildAll();
 	}
 }
 
@@ -1685,4 +1693,4 @@ void FRecastNavMeshCachedData::OnAreaAdded(const UClass* AreaClass, int32 AreaID
 	}		
 }
 
-#endif// WITH_RECAST
\ No newline at end of file
+#endif// WITH_RECAST
diff --git a/Engine/Source/Runtime/Engine/Private/AI/Navigation/RecastNavMeshGenerator.cpp b/Engine/Source/Runtime/Engine/Private/AI/Navigation/RecastNavMeshGenerator.cpp
--- a/Engine/Source/Runtime/Engine/Private/AI/Navigation/RecastNavMeshGenerator.cpp
+++ b/Engine/Source/Runtime/Engine/Private/AI/Navigation/RecastNavMeshGenerator.cpp
@@ -1425,59 +1425,55 @@ FRecastTileGenerator::FRecastTileGenerator(
 	rcVcopy(TileConfig.bmax, &RCBox.Max.X);
 			
 	// from passed in boxes pick the ones overlapping with tile bounds
-	InclusionBounds.Reset();
-	bOutsideOfInclusionBounds = false;
+	bFullyEncapsulatedByInclusionBounds = true;
 	const auto& ParentBounds = ParentGenerator->GetInclusionBounds();
 	if (ParentBounds.Num() > 0)
 	{
 		bFullyEncapsulatedByInclusionBounds = false;
-
 		InclusionBounds.Reserve(ParentBounds.Num());
-		const FBox* Bounds = ParentBounds.GetData();	
-		for (int32 i = 0; i < ParentBounds.Num() && bFullyEncapsulatedByInclusionBounds == false; ++i, ++Bounds)
-		{	
-			if (Bounds->Intersect(TileBB))
+		for (const FBox& Bounds : ParentBounds)
+		{
+			if (Bounds.Intersect(TileBB))
 			{
-				InclusionBounds.Add(*Bounds);
-				bFullyEncapsulatedByInclusionBounds = DoesBoxContainBox(*Bounds, TileBB);
+				InclusionBounds.Add(Bounds);
+				bFullyEncapsulatedByInclusionBounds = DoesBoxContainBox(Bounds, TileBB);
 			}
 		}
-
-		bOutsideOfInclusionBounds = (InclusionBounds.Num() == 0);
-	}
-	else
-	{
-		bFullyEncapsulatedByInclusionBounds = true;
 	}
 	
-	// In case we don't need to regenerate tile from scratch
-	// get compressed layers data from cache
-	if (DirtyAreas.Num() > 0)
+	// Gather geometry for tile if it inside navigable bounds
+	bRegenerateCompressedLayers = true;
+	if (InclusionBounds.Num())
 	{
-		CompressedLayers = ParentGenerator->GetCachedLayersData(FIntPoint(TileX, TileY));
-	}
+		// In case we don't need to regenerate tile from scratch
+		// get compressed layers data from cache
+		if (DirtyAreas.Num() > 0)
+		{
+			CompressedLayers = ParentGenerator->GetCachedLayersData(FIntPoint(TileX, TileY));
+		}
 	
-	// We have to regenerate layers data in case it's missing
-	bRegenerateCompressedLayers = (CompressedLayers.Num() == 0);
+		// We have to regenerate layers data in case it's missing
+		bRegenerateCompressedLayers = (CompressedLayers.Num() == 0);
 
-	if (!bRegenerateCompressedLayers)
-	{
-		// Mark layers that needs to be updated
-		DirtyLayers.Init(false, CompressedLayers.Num());
-		for (const FNavMeshTileData& LayerData : CompressedLayers)
+		if (!bRegenerateCompressedLayers)
 		{
-			for (FBox DirtyBox : DirtyAreas)
+			// Mark layers that needs to be updated
+			DirtyLayers.Init(false, CompressedLayers.Num());
+			for (const FNavMeshTileData& LayerData : CompressedLayers)
 			{
-				if (DirtyBox.Intersect(LayerData.LayerBBox))
+				for (FBox DirtyBox : DirtyAreas)
 				{
-					DirtyLayers[LayerData.LayerIndex] = true;
+					if (DirtyBox.Intersect(LayerData.LayerBBox))
+					{
+						DirtyLayers[LayerData.LayerIndex] = true;
+					}
 				}
 			}
 		}
-	}
-
-	GatherGeometry(ParentGenerator, bRegenerateCompressedLayers);
 
+		GatherGeometry(ParentGenerator, bRegenerateCompressedLayers);
+	}
+	
 	//
 	UsedMemoryOnStartup = GetUsedMemCount() + sizeof(FRecastTileGenerator);
 }
@@ -2684,19 +2680,19 @@ uint32 FRecastTileGenerator::GetUsedMemCount() const
 	return TotalMemory;
 }
 
-static int32 CaclulateGridCellsCount(const TNavStatArray<FBox>& NavigableAreas, float TileSizeinWorldUnits)
+static int32 CaclulateMaxTilesCount(const TNavStatArray<FBox>& NavigableAreas, float TileSizeinWorldUnits, float AvgLayersPerGridCell)
 {
-	int32 CellsCount = 0;
+	int32 GridCellsCount = 0;
 	for (FBox AreaBounds : NavigableAreas)
 	{
 		// TODO: need more precise calculation, currently we don't take into account that volumes can be overlapped
 		FBox RCBox = Unreal2RecastBox(AreaBounds);
 		int32 XSize = FMath::CeilToInt(RCBox.GetSize().X/TileSizeinWorldUnits) + 1;
 		int32 YSize = FMath::CeilToInt(RCBox.GetSize().Z/TileSizeinWorldUnits) + 1;
-		CellsCount+= (XSize*YSize);
+		GridCellsCount+= (XSize*YSize);
 	}
 	
-	return CellsCount;
+	return FMath::CeilToInt(GridCellsCount * AvgLayersPerGridCell);
 }
 
 //----------------------------------------------------------------------//
@@ -2704,9 +2700,9 @@ static int32 CaclulateGridCellsCount(const TNavStatArray<FBox>& NavigableAreas,
 //----------------------------------------------------------------------//
 
 FRecastNavMeshGenerator::FRecastNavMeshGenerator(class ARecastNavMesh* InDestNavMesh)
-	: MaxActiveTiles(-1)
-	, NumActiveTiles(0)
+	: NumActiveTiles(0)
 	, MaxTileGeneratorTasks(1)
+	, AvgLayersPerTile(8.0f)
 	, DestNavMesh(InDestNavMesh)
 	, bInitialized(false)
 	, Version(0)
@@ -2787,52 +2783,7 @@ void FRecastNavMeshGenerator::Init()
 	Config.regionPartitioning = DestNavMesh->LayerPartitioning;
 	Config.TileCachePartitionType = DestNavMesh->RegionPartitioning;
 
-	TotalNavBounds = FBox(0);
-	// Collect bounding geometry
-	if (NavSys->ShouldGenerateNavigationEverywhere() == false)
-	{
-		InclusionBounds.Empty();
-
-		for (TActorIterator<ANavMeshBoundsVolume> It(NavSys->GetWorld()); It; ++It)
-		{
-			ANavMeshBoundsVolume const* const V = (*It);
-			if (V != NULL)
-			{
-				FBox Bounds = GrowBoundingBox(V->GetComponentsBoundingBox(true), true);
-				TotalNavBounds += Bounds;
-				InclusionBounds.Add(Bounds);
-			}
-		}
-	}
-
-	if (!TotalNavBounds.IsValid)
-	{
-		// if no bounds give Navigation System a chance to specify bounds
-		// will also happen if NavSystem->ShouldGenerateNavigationEverywhere() == true
-		TotalNavBounds = NavSys->GetWorldBounds();
-		InclusionBounds.Add(TotalNavBounds);
-	}
-
-	// limit max amount of tiles: 64 bit poly 
-	const int32 MaxTileBits = 30;
-	const int32 MaxTiles = (1 << 30);
-	const float AvgLayersPerGridCell = 8.0f; // TODO: should be tweakable per navmesh
-	MaxActiveTiles = FMath::CeilToInt(CaclulateGridCellsCount(InclusionBounds, Config.tileSize * Config.cs) * AvgLayersPerGridCell);
-	if (MaxActiveTiles < 0 || MaxActiveTiles > MaxTiles)
-	{
-		UE_LOG(LogNavigation, Error, TEXT("Navmesh bounds are too large! Limiting requested tiles count (%d) to: (%d)"), MaxActiveTiles, MaxTiles);
-		MaxActiveTiles = MaxTiles;
-	}
-
-	// do this once
-	++Version;
-
-	// Max tiles and max polys affect how the tile IDs are calculated.
-	// There are (sizeof(dtPolyRef)*8 - DT_MIN_SALT_BITS) bits available for 
-	// identifying a tile and a polygon.
-	int32 TileBits = FMath::CeilLogTwo(MaxActiveTiles);
-	const int32 PolyBits = FMath::Min(30, (int)(sizeof(dtPolyRef)*8 - DT_MIN_SALT_BITS) - TileBits);
-	Config.MaxPolysPerTile = 1 << PolyBits;
+	UpdateNavigationBounds();
 
 	/** setup maximum number of active tile generator*/
 	const int32 NumberOfWorkerThreads = FTaskGraphInterface::Get().GetNumWorkerThreads();
@@ -2848,6 +2799,33 @@ void FRecastNavMeshGenerator::Init()
 	bInitialized = true;
 }
 
+void FRecastNavMeshGenerator::UpdateNavigationBounds()
+{
+	const UNavigationSystem* NavSys = UNavigationSystem::GetCurrent(GetWorld());
+	const TSet<FNavigationBounds>& NavigationBoundsSet = NavSys->GetNavigationBounds();
+
+	TotalNavBounds = FBox(0);
+	InclusionBounds.Empty(NavigationBoundsSet.Num());
+
+	// Collect bounding geometry
+	if (NavSys->ShouldGenerateNavigationEverywhere() == false)
+	{
+		for (const auto& NavigationBounds : NavigationBoundsSet)
+		{
+			InclusionBounds.Add(NavigationBounds.AreaBox);
+			TotalNavBounds+= NavigationBounds.AreaBox;
+		}
+	}
+	else
+	{
+		TotalNavBounds = NavSys->GetWorldBounds();
+		if (!TotalNavBounds.IsValid)
+		{
+			InclusionBounds.Add(TotalNavBounds);
+		}
+	}
+}
+
 bool FRecastNavMeshGenerator::ConstructTiledNavMesh() 
 {
 	bool bSuccess = false;
@@ -2856,19 +2834,49 @@ bool FRecastNavMeshGenerator::ConstructTiledNavMesh()
 	dtNavMesh* DetourMesh = dtAllocNavMesh();	
 	if (DetourMesh)
 	{
-		float Origin[3] = {0,0,0};
+		++Version;
 		
 		dtNavMeshParams TiledMeshParameters;
 		FMemory::MemZero(TiledMeshParameters);	
-		rcVcopy(TiledMeshParameters.orig, Origin);
+		rcVcopy(TiledMeshParameters.orig, &FVector::ZeroVector.X);
 		TiledMeshParameters.tileWidth = Config.tileSize * Config.cs;
 		TiledMeshParameters.tileHeight = Config.tileSize * Config.cs;
-		TiledMeshParameters.maxTiles = MaxActiveTiles;
-		TiledMeshParameters.maxPolys = Config.MaxPolysPerTile;
 
+		// limit max amount of tiles
+#if USE_64BIT_ADDRESS
+		const int32 MaxTileBits = 30;
+#else
+		const int32 MaxTileBits = 14;
+#endif//USE_64BIT_ADDRESS
+
+		const int32 MaxTiles = (1 << MaxTileBits);
+		int32 MaxRequestedTiles = 0;
+		if (DestNavMesh->IsResizable())
+		{
+			MaxRequestedTiles = CaclulateMaxTilesCount(InclusionBounds, Config.tileSize * Config.cs, AvgLayersPerTile);
+		}
+		else
+		{
+			MaxRequestedTiles = DestNavMesh->TilePoolSize;
+		}
+		
+		if (MaxRequestedTiles < 0 || MaxRequestedTiles > MaxTiles)
+		{
+			UE_LOG(LogNavigation, Error, TEXT("Navmesh bounds are too large! Limiting requested tiles count (%d) to: (%d)"), MaxRequestedTiles, MaxTiles);
+			MaxRequestedTiles = MaxTiles;
+		}
+
+		// Max tiles and max polys affect how the tile IDs are calculated.
+		// There are (sizeof(dtPolyRef)*8 - DT_MIN_SALT_BITS) bits available for 
+		// identifying a tile and a polygon.
+		Config.MaxPolysPerTile = 1 << ((sizeof(dtPolyRef)*8 - DT_MIN_SALT_BITS) - MaxTileBits);
+
+		TiledMeshParameters.maxTiles = MaxRequestedTiles;
+		TiledMeshParameters.maxPolys = Config.MaxPolysPerTile;
+	
 		const dtStatus status = DetourMesh->init(&TiledMeshParameters);
 
-		if ( dtStatusFailed(status) )
+		if (dtStatusFailed(status))
 		{
 			UE_LOG(LogNavigation, Warning, TEXT("ConstructTiledNavMesh: Could not init navmesh.") );
 			bSuccess = false;
@@ -2877,7 +2885,6 @@ bool FRecastNavMeshGenerator::ConstructTiledNavMesh()
 		{
 			bSuccess = true;
 			NumActiveTiles = GetTilesCountHelper(DetourMesh);
-
 			DestNavMesh->GetRecastNavMeshImpl()->SetRecastMesh(DetourMesh);
 		}
 	}
@@ -2898,7 +2905,7 @@ bool FRecastNavMeshGenerator::RebuildAll()
 	TArray<FNavigationDirtyArea> DirtyAreas;
 	for (FBox AreaBounds : InclusionBounds)
 	{
-		FNavigationDirtyArea DirtyArea(AreaBounds, ENavigationDirtyFlag::All);
+		FNavigationDirtyArea DirtyArea(AreaBounds, ENavigationDirtyFlag::All | ENavigationDirtyFlag::NavigationBounds);
 		DirtyAreas.Add(DirtyArea);
 	}
 	
@@ -2983,9 +2990,35 @@ void FRecastNavMeshGenerator::TickAsyncBuild(float DeltaSeconds)
 	}
 }
 
+void FRecastNavMeshGenerator::OnNavigationBoundsChanged()
+{
+	UpdateNavigationBounds();
+	
+	dtNavMesh* DetourMesh = DestNavMesh->GetRecastNavMeshImpl()->GetRecastMesh();
+	if (DestNavMesh->IsResizable() && DetourMesh)
+	{
+		// Check whether Navmesh size needs to be changed
+		int32 MaxRequestedTiles = CaclulateMaxTilesCount(InclusionBounds, Config.tileSize * Config.cs, AvgLayersPerTile);
+		if (DetourMesh->getMaxTiles() != MaxRequestedTiles)
+		{
+			// Destroy current NavMesh, it will be allocated with a new size on next build request
+			DestNavMesh->GetRecastNavMeshImpl()->SetRecastMesh(nullptr);
+		}
+	}
+}
+
 void FRecastNavMeshGenerator::RebuildDirtyAreas(const TArray<FNavigationDirtyArea>& InDirtyAreas)
 {
-	MarkDirtyTiles(InDirtyAreas);
+	dtNavMesh* DetourMesh = DestNavMesh->GetRecastNavMeshImpl()->GetRecastMesh();
+	if (DetourMesh == nullptr)
+	{
+		ConstructTiledNavMesh();
+		RebuildAll();
+	}
+	else
+	{
+		MarkDirtyTiles(InDirtyAreas);
+	}
 }
 
 void FRecastNavMeshGenerator::OnAreaAdded(const UClass* AreaClass, int32 AreaID)
@@ -3047,73 +3080,76 @@ TArray<uint32> FRecastNavMeshGenerator::AddTile(const FRecastTileGenerator& Tile
 	const int32 TileY = TileGenerator.GetTileY();
 	
 	// 
-	TArray<FNavMeshTileData> TileLayers = TileGenerator.GetNavigationData();
 	if (TileGenerator.IsFullyRegenerated())
 	{
 		// remove all layers
 		ResultTileIndices = RemoveTileLayers(TileX, TileY);
 	}
-	ResultTileIndices.Reserve(TileLayers.Num());
-
+	
 	dtNavMesh* DetourMesh = DestNavMesh->GetRecastNavMeshImpl()->GetRecastMesh();
-		
-	bool bHasNavmesh = true;
-	for (int32 i = 0; i < TileLayers.Num(); i++)
+	if (DetourMesh != nullptr)
 	{
-		const int32 LayerIndex = TileLayers[i].LayerIndex;
-		if (!TileGenerator.IsLayerChanged(TileLayers[i].LayerIndex))
+		TArray<FNavMeshTileData> TileLayers = TileGenerator.GetNavigationData();
+		ResultTileIndices.Reserve(TileLayers.Num());
+		
+		bool bHasNavmesh = true;
+		for (int32 i = 0; i < TileLayers.Num(); i++)
 		{
-			continue;
-		}
+			const int32 LayerIndex = TileLayers[i].LayerIndex;
+			if (!TileGenerator.IsLayerChanged(TileLayers[i].LayerIndex))
+			{
+				continue;
+			}
 				
-		const dtTileRef OldTileRef = DetourMesh->getTileRefAt(TileX, TileY, LayerIndex);
-
-		if (OldTileRef)
-		{
-			NumActiveTiles--;
-			UE_LOG(LogNavigation, Log, TEXT("%s> Tile (%d,%d:%d), removing TileRef: 0x%X (active:%d)"),
-				*DestNavMesh->GetName(), TileX, TileY, LayerIndex, OldTileRef, NumActiveTiles);
-
-			// this call will fill OldRawNavData with address of previously added data
-			// if OldRawNavData is NULL it means either this tile was empty, or the memory
-			// allocated has been owned by navmesh itself (i.e. has been serialized-in along
-			// with the level)
-			uint8* OldRawNavData = nullptr;
-			DetourMesh->removeTile(OldTileRef, &OldRawNavData, nullptr);
-			// data should always be owned by navmesh itself
-			check(OldRawNavData == nullptr);
+			const dtTileRef OldTileRef = DetourMesh->getTileRefAt(TileX, TileY, LayerIndex);
+
+			if (OldTileRef)
+			{
+				NumActiveTiles--;
+				UE_LOG(LogNavigation, Log, TEXT("%s> Tile (%d,%d:%d), removing TileRef: 0x%X (active:%d)"),
+					*DestNavMesh->GetName(), TileX, TileY, LayerIndex, OldTileRef, NumActiveTiles);
+
+				// this call will fill OldRawNavData with address of previously added data
+				// if OldRawNavData is NULL it means either this tile was empty, or the memory
+				// allocated has been owned by navmesh itself (i.e. has been serialized-in along
+				// with the level)
+				uint8* OldRawNavData = nullptr;
+				DetourMesh->removeTile(OldTileRef, &OldRawNavData, nullptr);
+				// data should always be owned by navmesh itself
+				check(OldRawNavData == nullptr);
 			
-			ResultTileIndices.AddUnique(DetourMesh->decodePolyIdTile(OldTileRef));
-		}
+				ResultTileIndices.AddUnique(DetourMesh->decodePolyIdTile(OldTileRef));
+			}
 
-		if (TileLayers[i].IsValid()) 
-		{
-			bool bRejectNavmesh = false;
-			dtTileRef ResultTileRef = 0;
+			if (TileLayers[i].IsValid()) 
+			{
+				bool bRejectNavmesh = false;
+				dtTileRef ResultTileRef = 0;
 
-			// let navmesh know it's tile generator who owns the data
-			dtStatus status = DetourMesh->addTile(TileLayers[i].GetData(), TileLayers[i].DataSize, DT_TILE_FREE_DATA, 0, &ResultTileRef);
+				// let navmesh know it's tile generator who owns the data
+				dtStatus status = DetourMesh->addTile(TileLayers[i].GetData(), TileLayers[i].DataSize, DT_TILE_FREE_DATA, 0, &ResultTileRef);
 
-			if (dtStatusFailed(status))
-			{
-				if (dtStatusDetail(status, DT_OUT_OF_MEMORY))
+				if (dtStatusFailed(status))
 				{
-					UE_LOG(LogNavigation, Error, TEXT("%s> Tile (%d,%d:%d), tile limit reached!! (%d)"),
-						*DestNavMesh->GetName(), TileX, TileY, LayerIndex, DetourMesh->getMaxTiles());
-				}
+					if (dtStatusDetail(status, DT_OUT_OF_MEMORY))
+					{
+						UE_LOG(LogNavigation, Error, TEXT("%s> Tile (%d,%d:%d), tile limit reached!! (%d)"),
+							*DestNavMesh->GetName(), TileX, TileY, LayerIndex, DetourMesh->getMaxTiles());
+					}
 
-				bHasNavmesh = false;
-			}
-			else
-			{
-				ResultTileIndices.AddUnique(DetourMesh->decodePolyIdTile(ResultTileRef));
-				NumActiveTiles++;
+					bHasNavmesh = false;
+				}
+				else
+				{
+					ResultTileIndices.AddUnique(DetourMesh->decodePolyIdTile(ResultTileRef));
+					NumActiveTiles++;
 
-				UE_LOG(LogNavigation, Log, TEXT("%s> Tile (%d,%d:%d), added TileRef: 0x%X (active:%d)"),
-					*DestNavMesh->GetName(), TileX, TileY, LayerIndex, ResultTileRef, NumActiveTiles);
+					UE_LOG(LogNavigation, Log, TEXT("%s> Tile (%d,%d:%d), added TileRef: 0x%X (active:%d)"),
+						*DestNavMesh->GetName(), TileX, TileY, LayerIndex, ResultTileRef, NumActiveTiles);
 				
-				// NavMesh took the ownership of generated data, so we don't need to deallocate it
-				TileLayers[i].Release();
+					// NavMesh took the ownership of generated data, so we don't need to deallocate it
+					TileLayers[i].Release();
+				}
 			}
 		}
 	}
@@ -3173,12 +3209,13 @@ void FRecastNavMeshGenerator::MarkDirtyTiles(const TArray<FNavigationDirtyArea>&
 	const FVector NavMeshOrigin = FVector::ZeroVector;
 	int32 NumTilesMarked = 0;
 
-	// find all tiles that need regeneration:
+	// find all tiles that need regeneration
 	TSet<FTileElement> DirtyTiles;
 	for (const FNavigationDirtyArea& DirtyArea : DirtyAreas)
 	{
 		const FBox AdjustedAreaBounds = GrowBoundingBox(DirtyArea.Bounds, DirtyArea.HasFlag(ENavigationDirtyFlag::UseAgentHeight));
-		if (!IntercestsInclusionBounds(AdjustedAreaBounds))
+		if (!DirtyArea.HasFlag(ENavigationDirtyFlag::NavigationBounds) && 
+			!IntercestsInclusionBounds(AdjustedAreaBounds))
 		{
 			// Skip whole area
 			continue;
@@ -3195,7 +3232,8 @@ void FRecastNavMeshGenerator::MarkDirtyTiles(const TArray<FNavigationDirtyArea>&
 			for (int32 x = XMin; x <= XMax; ++x)
 			{
 				const FBox TileBox = CalculateTileBounds(x, y, NavMeshOrigin, TotalNavBounds, TileSizeInWorldUnits);
-				if (!IntercestsInclusionBounds(TileBox))
+				if (!DirtyArea.HasFlag(ENavigationDirtyFlag::NavigationBounds) && 
+					!IntercestsInclusionBounds(TileBox))
 				{
 					// Skip this tile
 					continue;
@@ -3203,7 +3241,7 @@ void FRecastNavMeshGenerator::MarkDirtyTiles(const TArray<FNavigationDirtyArea>&
 								
 				FTileElement Element;
 				Element.Coord = FIntPoint(x, y);
-				Element.bRebuildGeometry = (DirtyArea.Flags & ENavigationDirtyFlag::Geometry);
+				Element.bRebuildGeometry = DirtyArea.HasFlag(ENavigationDirtyFlag::Geometry) || DirtyArea.HasFlag(ENavigationDirtyFlag::NavigationBounds);
 				if (Element.bRebuildGeometry == false)
 				{
 					Element.DirtyAreas.Add(AdjustedAreaBounds);
diff --git a/Engine/Source/Runtime/Engine/Private/NavMeshBoundsVolume.cpp b/Engine/Source/Runtime/Engine/Private/NavMeshBoundsVolume.cpp
--- a/Engine/Source/Runtime/Engine/Private/NavMeshBoundsVolume.cpp
+++ b/Engine/Source/Runtime/Engine/Private/NavMeshBoundsVolume.cpp
@@ -42,7 +42,7 @@ void ANavMeshBoundsVolume::PostRegisterAllComponents()
 	UNavigationSystem* NavSys = UNavigationSystem::GetCurrent(GetWorld());
 	if (NavSys && Role == ROLE_Authority)
 	{
-		NavSys->OnNavigationBoundsUpdated(this);
+		NavSys->OnNavigationBoundsAdded(this);
 	}
 }
 
@@ -53,6 +53,6 @@ void ANavMeshBoundsVolume::PostUnregisterAllComponents()
 	UNavigationSystem* NavSys = UNavigationSystem::GetCurrent(GetWorld());
 	if (NavSys && Role == ROLE_Authority)
 	{
-		NavSys->OnNavigationBoundsUpdated(this);
+		NavSys->OnNavigationBoundsRemoved(this);
 	}
 }
diff --git a/Engine/Source/Runtime/Engine/Public/AI/NavDataGenerator.h b/Engine/Source/Runtime/Engine/Public/AI/NavDataGenerator.h
--- a/Engine/Source/Runtime/Engine/Public/AI/NavDataGenerator.h
+++ b/Engine/Source/Runtime/Engine/Public/AI/NavDataGenerator.h
@@ -20,6 +20,9 @@ public:
 
 	/**  */
 	virtual void TickAsyncBuild(float DeltaSeconds) {};
+	
+	/**  */
+	virtual void OnNavigationBoundsChanged() {};
 
 	/** Asks generator to update navigation affected by DirtyAreas */
 	virtual void RebuildDirtyAreas(const TArray<FNavigationDirtyArea>& DirtyAreas) {}
diff --git a/Engine/Source/Runtime/Engine/Public/AI/Navigation/RecastNavMeshGenerator.h b/Engine/Source/Runtime/Engine/Public/AI/Navigation/RecastNavMeshGenerator.h
--- a/Engine/Source/Runtime/Engine/Public/AI/Navigation/RecastNavMeshGenerator.h
+++ b/Engine/Source/Runtime/Engine/Public/AI/Navigation/RecastNavMeshGenerator.h
@@ -174,7 +174,6 @@ private:
 protected:
 	uint32 bSucceeded : 1;
 	uint32 bRegenerateCompressedLayers : 1;
-	uint32 bOutsideOfInclusionBounds : 1;
 	uint32 bFullyEncapsulatedByInclusionBounds : 1;
 	
 	int32 TileX;
@@ -281,6 +280,7 @@ public:
 	virtual void EnsureBuildCompletion() override;
 	virtual void CancelBuild() override;
 	virtual void TickAsyncBuild(float DeltaSeconds) override;
+	virtual void OnNavigationBoundsChanged() override;
 
 	/** Asks generator to update navigation affected by DirtyAreas */
 	virtual void RebuildDirtyAreas(const TArray<FNavigationDirtyArea>& DirtyAreas) override;
@@ -338,6 +338,9 @@ public:
 private:
 	// Performs initial setup of member variables so that generator is ready to do its thing from this point on
 	void Init();
+
+	// Updates cached list of navigation bounds
+	void UpdateNavigationBounds();
 		
 	// Sorts pending build tiles by proximity to player, so tiles closer to player will get generated first
 	void SortPendingBuildTiles();
@@ -370,10 +373,9 @@ private:
 	struct dtNavMeshParams TiledMeshParams;
 	struct FRecastBuildConfig Config;
 	
-	int32 MaxActiveTiles;
 	int32 NumActiveTiles;
 	int32 MaxTileGeneratorTasks;
-	static int32 ActiveTileGeneratorTasks;
+	float AvgLayersPerTile;
 
 	/** Total bounding box that includes all volumes, in unreal units. */
 	FBox TotalNavBounds;
