In-editor experimental feature cook assert fix.

diff --git a/Engine/Source/Runtime/Engine/Private/World.cpp b/Engine/Source/Runtime/Engine/Private/World.cpp
--- a/Engine/Source/Runtime/Engine/Private/World.cpp
+++ b/Engine/Source/Runtime/Engine/Private/World.cpp
@@ -521,7 +521,7 @@ void UWorld::PostLoad()
 			const FString ShortPackageName = FPackageName::GetLongPackageAssetName(GetOutermost()->GetName());
 			if (GetName() != ShortPackageName)
 			{
-				Rename(*ShortPackageName, NULL, REN_NonTransactional | REN_ForceNoResetLoaders);
+				Rename(*ShortPackageName, NULL, REN_NonTransactional | REN_ForceNoResetLoaders | REN_DontCreateRedirectors);
 			}
 
 			// Worlds are assets so they need RF_Public and RF_Standalone (for the editor)
