Navmesh generation: do not run tasks for tiles that will not produce any data, directly remove those tiles from navmesh

diff --git a/Engine/Source/Runtime/Engine/Private/AI/Navigation/RecastNavMeshGenerator.cpp b/Engine/Source/Runtime/Engine/Private/AI/Navigation/RecastNavMeshGenerator.cpp
--- a/Engine/Source/Runtime/Engine/Private/AI/Navigation/RecastNavMeshGenerator.cpp
+++ b/Engine/Source/Runtime/Engine/Private/AI/Navigation/RecastNavMeshGenerator.cpp
@@ -1449,7 +1449,7 @@ FRecastTileGenerator::FRecastTileGenerator(
 		// get compressed layers data from cache
 		if (DirtyAreas.Num() > 0)
 		{
-			CompressedLayers = ParentGenerator->GetCachedLayersData(FIntPoint(TileX, TileY));
+			CompressedLayers = ParentGenerator->GetIntermediateLayersData(FIntPoint(TileX, TileY));
 		}
 	
 		// We have to regenerate layers data in case it's missing
@@ -1478,6 +1478,17 @@ FRecastTileGenerator::FRecastTileGenerator(
 	UsedMemoryOnStartup = GetUsedMemCount() + sizeof(FRecastTileGenerator);
 }
 
+bool FRecastTileGenerator::HasDataToBuild() const
+{
+	return
+		CompressedLayers.Num()
+		|| StaticAreas.Num()
+		|| DynamicAreas.Num()
+		|| OffmeshLinks.Num()
+		|| GeomCoords.Num()
+		|| GeomIndices.Num();
+}
+
 void FRecastTileGenerator::DoWork()
 {
 	bSucceeded = GenerateTile();
@@ -2908,26 +2919,33 @@ bool FRecastNavMeshGenerator::RebuildAll()
 		FNavigationDirtyArea DirtyArea(AreaBounds, ENavigationDirtyFlag::All | ENavigationDirtyFlag::NavigationBounds);
 		DirtyAreas.Add(DirtyArea);
 	}
-	
-	MarkDirtyTiles(DirtyAreas);
-	
+
+	if (DirtyAreas.Num())
+	{
+		MarkDirtyTiles(DirtyAreas);
+	}
+	else
+	{
+		// There are no navigation bounds to build, probably navmesh was resided and we just need to update debug draw
+		DestNavMesh->RequestDrawingUpdate();
+	}
+
 	return true;
 }
 
 void FRecastNavMeshGenerator::EnsureBuildCompletion()
 {
-	while (PendingDirtyTiles.Num() || RunningDirtyTiles.Num())
+	do 
 	{
-		SubmitDirtyTiles(16);
+		ProcessTileTasks(16);
 		
 		// Block until tasks are finished
-		for (FTileElement& Element : RunningDirtyTiles)
+		for (FRunningTileElement& Element : RunningDirtyTiles)
 		{
 			Element.AsyncTask->EnsureCompletion();
 		}
-		//
-		CollectCompletedTiles();
 	}
+	while (PendingDirtyTiles.Num() || RunningDirtyTiles.Num());
 }
 
 void FRecastNavMeshGenerator::CancelBuild()
@@ -2935,7 +2953,7 @@ void FRecastNavMeshGenerator::CancelBuild()
 	DiscardCurrentBuildingTasks();
 	RunningDirtyTiles.Empty();
 	PendingDirtyTiles.Empty();
-	CachedLayerDataMap.Empty();
+	IntermediateLayerDataMap.Empty();
 	SET_MEMORY_STAT(STAT_Navigation_TileCacheMemory, 0);
 
 #if	WITH_EDITOR	
@@ -2952,41 +2970,32 @@ void FRecastNavMeshGenerator::TickAsyncBuild(float DeltaSeconds)
 		RecentlyBuiltTiles.RemoveAllSwap([&](const FTileTimestamp& Tile) { return (Timestamp - Tile.Timestamp) > 0.5; });
 	}
 #endif//WITH_EDITOR
-		
-	// Collect and apply results from async tile build tasks in case we have any
-	TArray<uint32> CompletedTileIndices = CollectCompletedTiles();
-	if (CompletedTileIndices.Num() > 0)
+
+	// Submit async tile build tasks in case we have dirty tiles and have room for them
+	const UNavigationSystem* NavSys = UNavigationSystem::GetCurrent(GetWorld());
+	check(NavSys);
+	const int32 NumRunningTasks = NavSys->GetNumRunningBuildTasks();
+	const int32 NumTasksToSubmit = MaxTileGeneratorTasks - NumRunningTasks;
+	TArray<uint32> UpdatedTileIndices = ProcessTileTasks(NumTasksToSubmit);
+			
+	if (UpdatedTileIndices.Num() > 0)
 	{
 		// Invalidate active paths that go through regenerated tiles
-		DestNavMesh->InvalidateAffectedPaths(CompletedTileIndices);
+		DestNavMesh->InvalidateAffectedPaths(UpdatedTileIndices);
+		DestNavMesh->RequestDrawingUpdate();
 
 #if	WITH_EDITOR
 		// Store completed tiles with timestamps to have ability to distinguish during debug draw
+		const double Timestamp = FPlatformTime::Seconds();
+		RecentlyBuiltTiles.Reserve(RecentlyBuiltTiles.Num() + UpdatedTileIndices.Num());
+		for (uint32 TiledIdx : UpdatedTileIndices)
 		{
-			const double Timestamp = FPlatformTime::Seconds();
-			RecentlyBuiltTiles.Reserve(RecentlyBuiltTiles.Num() + CompletedTileIndices.Num());
-			for (uint32 TiledIdx : CompletedTileIndices)
-			{
-				FTileTimestamp TileTimestamp;
-				TileTimestamp.TileIdx = TiledIdx;
-				TileTimestamp.Timestamp = Timestamp;
-				RecentlyBuiltTiles.Add(TileTimestamp);
-			}
+			FTileTimestamp TileTimestamp;
+			TileTimestamp.TileIdx = TiledIdx;
+			TileTimestamp.Timestamp = Timestamp;
+			RecentlyBuiltTiles.Add(TileTimestamp);
 		}
 #endif//WITH_EDITOR
-
-		DestNavMesh->RequestDrawingUpdate();
-	}
-
-	// Submit async tile build tasks in case we have dirty tiles and have room for them
-	const UNavigationSystem* NavSys = UNavigationSystem::GetCurrent(GetWorld());
-	check(NavSys);
-	const int32 NumRunningTasks = NavSys->GetNumRunningBuildTasks();
-	const int32 NumTasksToSubmit = MaxTileGeneratorTasks - NumRunningTasks;
-		
-	if (NumTasksToSubmit > 0)
-	{
-		SubmitDirtyTiles(NumTasksToSubmit);
 	}
 }
 
@@ -3039,7 +3048,6 @@ bool FRecastNavMeshGenerator::IntercestsInclusionBounds(const FBox& Box) const
 	return false;
 }
 
-
 TArray<uint32> FRecastNavMeshGenerator::RemoveTileLayers(const int32 TileX, const int32 TileY)
 {
 	TArray<uint32> ResultTileIndices;
@@ -3068,6 +3076,9 @@ TArray<uint32> FRecastNavMeshGenerator::RemoveTileLayers(const int32 TileX, cons
 		}
 	}
 
+	// Remove intermediate layers data at this grid location
+	IntermediateLayerDataMap.Remove(FIntPoint(TileX, TileY));
+
 	return ResultTileIndices;
 }
 
@@ -3159,7 +3170,7 @@ TArray<uint32> FRecastNavMeshGenerator::AddTile(const FRecastTileGenerator& Tile
 
 void FRecastNavMeshGenerator::DiscardCurrentBuildingTasks()
 {
-	for (FTileElement& Element : RunningDirtyTiles)
+	for (FRunningTileElement& Element : RunningDirtyTiles)
 	{
 		if (Element.AsyncTask)
 		{
@@ -3172,7 +3183,6 @@ void FRecastNavMeshGenerator::DiscardCurrentBuildingTasks()
 	RunningDirtyTiles.Empty();
 }
 
-
 bool FRecastNavMeshGenerator::HasDirtyTiles() const
 {
 	return (PendingDirtyTiles.Num() > 0 || RunningDirtyTiles.Num() > 0);
@@ -3186,9 +3196,9 @@ FBox FRecastNavMeshGenerator::GrowBoundingBox(const FBox& BBox, bool bIncludeAge
 	return FBox(BBox.Min - BBoxGrowOffsetBoth - BBoxGrowOffsetMin, BBox.Max + BBoxGrowOffsetBoth);
 }
 
-TArray<FNavMeshTileData> FRecastNavMeshGenerator::GetCachedLayersData(FIntPoint GridCoord) const
+TArray<FNavMeshTileData> FRecastNavMeshGenerator::GetIntermediateLayersData(FIntPoint GridCoord) const
 {
-	const TArray<FNavMeshTileData>* ExistingData = CachedLayerDataMap.Find(GridCoord);
+	const TArray<FNavMeshTileData>* ExistingData = IntermediateLayerDataMap.Find(GridCoord);
 	if (ExistingData)
 	{
 		return *ExistingData;
@@ -3210,7 +3220,7 @@ void FRecastNavMeshGenerator::MarkDirtyTiles(const TArray<FNavigationDirtyArea>&
 	int32 NumTilesMarked = 0;
 
 	// find all tiles that need regeneration
-	TSet<FTileElement> DirtyTiles;
+	TSet<FPendingTileElement> DirtyTiles;
 	for (const FNavigationDirtyArea& DirtyArea : DirtyAreas)
 	{
 		const FBox AdjustedAreaBounds = GrowBoundingBox(DirtyArea.Bounds, DirtyArea.HasFlag(ENavigationDirtyFlag::UseAgentHeight));
@@ -3238,8 +3248,8 @@ void FRecastNavMeshGenerator::MarkDirtyTiles(const TArray<FNavigationDirtyArea>&
 					// Skip this tile
 					continue;
 				}
-								
-				FTileElement Element;
+												
+				FPendingTileElement Element;
 				Element.Coord = FIntPoint(x, y);
 				Element.bRebuildGeometry = DirtyArea.HasFlag(ENavigationDirtyFlag::Geometry) || DirtyArea.HasFlag(ENavigationDirtyFlag::NavigationBounds);
 				if (Element.bRebuildGeometry == false)
@@ -3247,7 +3257,7 @@ void FRecastNavMeshGenerator::MarkDirtyTiles(const TArray<FNavigationDirtyArea>&
 					Element.DirtyAreas.Add(AdjustedAreaBounds);
 				}
 				
-				FTileElement* ExistingElement = DirtyTiles.Find(Element);
+				FPendingTileElement* ExistingElement = DirtyTiles.Find(Element);
 				if (ExistingElement)
 				{
 					ExistingElement->bRebuildGeometry|= Element.bRebuildGeometry;
@@ -3272,9 +3282,9 @@ void FRecastNavMeshGenerator::MarkDirtyTiles(const TArray<FNavigationDirtyArea>&
 	NumTilesMarked = DirtyTiles.Num();
 
 	// Merge all pending tiles into one container
-	for (const FTileElement& Element : PendingDirtyTiles)
+	for (const FPendingTileElement& Element : PendingDirtyTiles)
 	{
-		FTileElement* ExistingElement = DirtyTiles.Find(Element);
+		FPendingTileElement* ExistingElement = DirtyTiles.Find(Element);
 		if (ExistingElement)
 		{
 			ExistingElement->bRebuildGeometry|= Element.bRebuildGeometry;
@@ -3296,7 +3306,7 @@ void FRecastNavMeshGenerator::MarkDirtyTiles(const TArray<FNavigationDirtyArea>&
 	
 	// Dump results into array
 	PendingDirtyTiles.Empty(DirtyTiles.Num());
-	for(const FTileElement& Element : DirtyTiles)
+	for(const FPendingTileElement& Element : DirtyTiles)
 	{
 		PendingDirtyTiles.Add(Element);
 	}
@@ -3338,7 +3348,7 @@ void FRecastNavMeshGenerator::SortPendingBuildTiles()
 		const float TileSizeInWorldUnits = Config.tileSize * Config.cs;
 		
 		// Calculate shortest distances between tiles and players
-		for (FTileElement& Element : PendingDirtyTiles)
+		for (FPendingTileElement& Element : PendingDirtyTiles)
 		{
 			const FBox TileBox = CalculateTileBounds(Element.Coord.X, Element.Coord.Y, FVector::ZeroVector, TotalNavBounds, TileSizeInWorldUnits);
 			FVector2D TileCenter2D = FVector2D(TileBox.GetCenter());
@@ -3357,37 +3367,50 @@ void FRecastNavMeshGenerator::SortPendingBuildTiles()
 	}
 }
 
-int32 FRecastNavMeshGenerator::SubmitDirtyTiles(const int32 NumTasksToSubmit)
+TArray<uint32> FRecastNavMeshGenerator::ProcessTileTasks(const int32 NumTasksToSubmit)
 {
-	QUICK_SCOPE_CYCLE_COUNTER(STAT_RecastNavMeshGenerator_SubmitDirtyTiles);
+	QUICK_SCOPE_CYCLE_COUNTER(STAT_RecastNavMeshGenerator_ProcessTileTasks);
+	
+	TArray<uint32> UpdatedTiles;
 	
-	if (PendingDirtyTiles.Num() == 0)
-	{
-		return 0;
-	}
-		
-	// 
 	int32 NumSubmittedTasks = 0;
+	// Submit pending tile elements
 	for (int32 ElementIdx = PendingDirtyTiles.Num()-1; ElementIdx >= 0 && NumSubmittedTasks < NumTasksToSubmit; ElementIdx--)
 	{
-		FTileElement& Element = PendingDirtyTiles[ElementIdx];
+		FPendingTileElement& PendingElement = PendingDirtyTiles[ElementIdx];
+		FRunningTileElement RunningElement(PendingElement.Coord);
+		
 		// Make sure that we are not submitting generator for grid cell that is currently being regenerated
-		if (!RunningDirtyTiles.Contains(Element))
+		if (!RunningDirtyTiles.Contains(RunningElement))
 		{
-			check(Element.AsyncTask == nullptr);
-
 			// Spawn async task
-			Element.AsyncTask = new FRecastTileGeneratorTask(
+			TUniquePtr<FRecastTileGeneratorTask> TileTask = MakeUnique<FRecastTileGeneratorTask>(
 				this,
-				Element.Coord.X,
-				Element.Coord.Y,
-				Element.DirtyAreas
+				PendingElement.Coord.X,
+				PendingElement.Coord.Y,
+				PendingElement.DirtyAreas
 				);
-			Element.AsyncTask->StartBackgroundTask();
-			
-			RunningDirtyTiles.Add(Element);
-			NumSubmittedTasks++;
 
+			// Start it in background in case it has something to build
+			if (TileTask->GetTask().HasDataToBuild())
+			{
+				RunningElement.AsyncTask = TileTask.Release();
+				RunningElement.AsyncTask->StartBackgroundTask();
+			
+				RunningDirtyTiles.Add(RunningElement);
+				NumSubmittedTasks++;
+			}
+			else
+			{
+				// If there is nothing to generate remove all tiles from navmesh at specified grid coordinates
+				UpdatedTiles.Append(
+					RemoveTileLayers(PendingElement.Coord.X, PendingElement.Coord.Y)
+				);
+					
+				// TODO: should we increment NumSubmittedTasks here? 
+				// We can count removing as a tasks to avoid hitches when there are large number of pending tiles to remove
+			}
+			
 			// Remove submitted element from pending list
 			PendingDirtyTiles.RemoveAt(ElementIdx);
 
@@ -3398,38 +3421,33 @@ int32 FRecastNavMeshGenerator::SubmitDirtyTiles(const int32 NumTasksToSubmit)
 			}
 		}
 	}
-	return NumSubmittedTasks;
-}
-
-TArray<uint32> FRecastNavMeshGenerator::CollectCompletedTiles()
-{
-	QUICK_SCOPE_CYCLE_COUNTER(STAT_RecastNavMeshGenerator_CollectCompletedTiles);
-
-	TArray<uint32> CompletedTiles;
-	CompletedTiles.Reserve(RunningDirtyTiles.Num());
-
+	
+	// Collect completed tasks and apply generated data to navmesh
 	for (int32 Idx = RunningDirtyTiles.Num() - 1; Idx >=0; --Idx)
 	{
-		FTileElement& Element = RunningDirtyTiles[Idx];
+		FRunningTileElement& Element = RunningDirtyTiles[Idx];
 		check(Element.AsyncTask);
 
 		if (Element.AsyncTask->IsDone())
 		{
 			// Add generated tiles to navmesh
-			const FRecastTileGenerator& TileGenerator = Element.AsyncTask->GetTask();
-			TArray<uint32> UpdatedTileIndices = AddTile(TileGenerator);
-			CompletedTiles.Append(UpdatedTileIndices);
-			
-			// Store layers data in the cache, so it can be reused later
-			// TODO: make this optional?
-			TArray<FNavMeshTileData> ComressedLayers = TileGenerator.GetCompressedLayers();
-			if (ComressedLayers.Num())
+			if (!Element.bShouldDiscard)
 			{
-				CachedLayerDataMap.Add(Element.Coord, ComressedLayers);
-			}
-			else
-			{
-				CachedLayerDataMap.Remove(Element.Coord);
+				const FRecastTileGenerator& TileGenerator = Element.AsyncTask->GetTask();
+				TArray<uint32> UpdatedTileIndices = AddTile(TileGenerator);
+				UpdatedTiles.Append(UpdatedTileIndices);
+			
+				// Store intermediate layers data, so it can be reused later
+				// TODO: make this optional?
+				TArray<FNavMeshTileData> ComressedLayers = TileGenerator.GetCompressedLayers();
+				if (ComressedLayers.Num())
+				{
+					IntermediateLayerDataMap.Add(Element.Coord, ComressedLayers);
+				}
+				else
+				{
+					IntermediateLayerDataMap.Remove(Element.Coord);
+				}
 			}
 
 			// Destroy tile generator task
@@ -3440,7 +3458,7 @@ TArray<uint32> FRecastNavMeshGenerator::CollectCompletedTiles()
 		}
 	}
 
-	return CompletedTiles;
+	return UpdatedTiles;
 }
 
 void FRecastNavMeshGenerator::ExportComponentGeometry(UActorComponent* Component, FNavigationRelevantData& Data)
@@ -3546,7 +3564,7 @@ uint32 FRecastNavMeshGenerator::LogMemUsed() const
 	UE_LOG(LogNavigation, Display, TEXT("    FRecastNavMeshGenerator: self %d"), sizeof(FRecastNavMeshGenerator));
 	
 	uint32 GeneratorsMem = 0;
-	for (FTileElement Element : RunningDirtyTiles)
+	for (const FRunningTileElement& Element : RunningDirtyTiles)
 	{
 		GeneratorsMem+= Element.AsyncTask->GetTask().UsedMemoryOnStartup; 
 	}
diff --git a/Engine/Source/Runtime/Engine/Public/AI/Navigation/RecastNavMeshGenerator.h b/Engine/Source/Runtime/Engine/Public/AI/Navigation/RecastNavMeshGenerator.h
--- a/Engine/Source/Runtime/Engine/Public/AI/Navigation/RecastNavMeshGenerator.h
+++ b/Engine/Source/Runtime/Engine/Public/AI/Navigation/RecastNavMeshGenerator.h
@@ -127,6 +127,8 @@ public:
 	FORCEINLINE bool IsLayerChanged(int32 LayerIdx) const { return DirtyLayers[LayerIdx]; }
 	/** Whether tile data was fully regenerated */
 	FORCEINLINE bool IsFullyRegenerated() const { return bRegenerateCompressedLayers; }
+	/** Whether tile task has anything to build */
+	bool HasDataToBuild() const;
 
 	TArray<FNavMeshTileData> GetCompressedLayers() const { return CompressedLayers; }
 	TArray<FNavMeshTileData> GetNavigationData() const { return NavigationData; }
@@ -211,45 +213,70 @@ protected:
 
 typedef FAsyncTask<FRecastTileGenerator> FRecastTileGeneratorTask;
 
-struct FTileElement
+struct FPendingTileElement
 {
 	/** tile coordinates on a grid in recast space */
 	FIntPoint	Coord;
 	/** distance to seed, used for sorting pending tiles */
 	float		SeedDistance; 
-	/** generator task associated with this tile */
-	FRecastTileGeneratorTask* AsyncTask;
 	/** Whether we need a full rebuild for this tile grid cell */
 	bool		bRebuildGeometry;
 	/** We need to store dirty area bounds to check which cached layers needs to be regenerated
 	 *  In case geometry is changed cached layers data will be fully regenerated without using dirty areas list
 	 */
 	TArray<FBox> DirtyAreas;
-	
-	FTileElement()
+
+	FPendingTileElement()
 		: Coord(FIntPoint::NoneValue)
 		, SeedDistance(MAX_flt)
-		, AsyncTask(nullptr)
 		, bRebuildGeometry(false)
 	{
 	}
 
-	bool operator == (const FTileElement& Other) const
+	bool operator == (const FPendingTileElement& Other) const
 	{
 		return Coord == Other.Coord;
 	}
 
-	bool operator < (const FTileElement& Other) const
+	bool operator < (const FPendingTileElement& Other) const
 	{
 		return Other.SeedDistance < SeedDistance;
 	}
 
-	friend uint32 GetTypeHash(const FTileElement& Element)
+	friend uint32 GetTypeHash(const FPendingTileElement& Element)
 	{
 		return GetTypeHash(Element.Coord);
 	}
 };
 
+struct FRunningTileElement
+{
+	FRunningTileElement()
+		: Coord(FIntPoint::NoneValue)
+		, bShouldDiscard(false)
+		, AsyncTask(nullptr)
+	{
+	}
+	
+	FRunningTileElement(FIntPoint InCoord)
+		: Coord(InCoord)
+		, bShouldDiscard(false)
+		, AsyncTask(nullptr)
+	{
+	}
+
+	bool operator == (const FRunningTileElement& Other) const
+	{
+		return Coord == Other.Coord;
+	}
+	
+	/** tile coordinates on a grid in recast space */
+	FIntPoint					Coord;
+	/** whether generated results should be discarded */
+	bool						bShouldDiscard; 
+	FRecastTileGeneratorTask*	AsyncTask;
+};
+
 struct FTileTimestamp
 {
 	uint32 TileIdx;
@@ -320,7 +347,7 @@ public:
 	FBox GrowBoundingBox(const FBox& BBox, bool bIncludeAgentHeight) const;
 
 	/**  */
-	TArray<FNavMeshTileData> GetCachedLayersData(FIntPoint GridCoord) const;
+	TArray<FNavMeshTileData> GetIntermediateLayersData(FIntPoint GridCoord) const;
 
 #if !(UE_BUILD_SHIPPING || UE_BUILD_TEST)
 	virtual void ExportNavigationData(const FString& FileName) const;
@@ -351,11 +378,8 @@ private:
 	/** Marks grid tiles affected by specified areas as dirty */
 	void MarkDirtyTiles(const TArray<FNavigationDirtyArea>& DirtyAreas);
 	
-	/** Submits up to specified number of async generator tasks for grid tiles that was marked as dirty */
-	int32 SubmitDirtyTiles(const int32 NumTasksToSubmit);
-	
-	/** Collect completed async tasks and applies new tile data to a navmesh*/
-	TArray<uint32> CollectCompletedTiles();
+	/** Processes pending tile generattion tasks */
+	TArray<uint32> ProcessTileTasks(const int32 NumTasksToSubmit);
 
 	/** Removes all tiles at specified grid location */
 	TArray<uint32> RemoveTileLayers(const int32 TileX, const int32 TileY);
@@ -387,10 +411,10 @@ private:
 	ARecastNavMesh*	DestNavMesh;
 	
 	/** List of dirty tiles that needs to be regenerated */
-	TNavStatArray<FTileElement> PendingDirtyTiles;			
+	TNavStatArray<FPendingTileElement> PendingDirtyTiles;			
 	
 	/** List of dirty tiles currently being regenerated */
-	TNavStatArray<FTileElement> RunningDirtyTiles;
+	TNavStatArray<FRunningTileElement> RunningDirtyTiles;
 
 #if WITH_EDITOR
 	/** List of tiles that were recently regenerated */
@@ -401,7 +425,7 @@ private:
 	FRecastNavMeshCachedData AdditionalCachedData;
 
 	/** Compressed layers data, can be reused for tiles generation */
-	TMap<FIntPoint, TArray<FNavMeshTileData>> CachedLayerDataMap;
+	TMap<FIntPoint, TArray<FNavMeshTileData>> IntermediateLayerDataMap;
 
 	/** */
 	TMapBase<const AActor*, FBox, false> ActorToAreaMap;
