Fix per poly skeletal mesh so that it works with runtime spawned data and blueprints.

- Moved cooked bodysetup into skeletal mesh that instances copy from
- Added bEnablePerPolyCollision to the skeletal mesh which must be set to true if instances wish to use per poly collision

diff --git a/Engine/Source/Runtime/Engine/Classes/Components/SkeletalMeshComponent.h b/Engine/Source/Runtime/Engine/Classes/Components/SkeletalMeshComponent.h
--- a/Engine/Source/Runtime/Engine/Classes/Components/SkeletalMeshComponent.h
+++ b/Engine/Source/Runtime/Engine/Classes/Components/SkeletalMeshComponent.h
@@ -465,16 +465,6 @@ public:
 
 	void CreateBodySetup();
 
-	// Begin Interface_CollisionDataProvider Interface
-	virtual bool GetPhysicsTriMeshData(struct FTriMeshCollisionData* CollisionData, bool InUseAllTriData) override;
-	virtual bool ContainsPhysicsTriMeshData(bool InUseAllTriData) const override;
-	virtual bool WantsNegXTriMesh() override
-	{
-		return true;
-	}
-	// End Interface_CollisionDataProvider Interface
-
-
 	/**
 	 * Misc 
 	 */
diff --git a/Engine/Source/Runtime/Engine/Classes/Engine/DestructibleMesh.h b/Engine/Source/Runtime/Engine/Classes/Engine/DestructibleMesh.h
--- a/Engine/Source/Runtime/Engine/Classes/Engine/DestructibleMesh.h
+++ b/Engine/Source/Runtime/Engine/Classes/Engine/DestructibleMesh.h
@@ -457,10 +457,6 @@ class UDestructibleMesh : public USkeletalMesh
 	UPROPERTY(EditAnywhere, editfixedsize, Category=Effects)
 	TArray<struct FFractureEffect> FractureEffects;
 
-	/** Physics data.  Fields from BodySetup which are relevant to the DestructibleMesh will be used. */
-	UPROPERTY(EditAnywhere, Instanced, Category=DestructibleMesh)
-	class UBodySetup* BodySetup;
-
 #if WITH_EDITORONLY_DATA
 	/** Information used to author an NxDestructibleAsset*/
 	UPROPERTY(instanced)
@@ -516,11 +512,6 @@ public:
 	/** Fills DefaultDestructibleParameters with parameters from the NxDestructibleAsset. */
 	ENGINE_API void				LoadDefaultDestructibleParametersFromApexAsset();
 
-	/**
-	 * Create BodySetup for this DestructibleMesh if it doesn't have one
-	 */
-	ENGINE_API void				CreateBodySetup();
-
 	/**
 	 * Create DestructibleFractureSettings for this DestructibleMesh if it doesn't have one
 	 */
diff --git a/Engine/Source/Runtime/Engine/Classes/Engine/SkeletalMesh.h b/Engine/Source/Runtime/Engine/Classes/Engine/SkeletalMesh.h
--- a/Engine/Source/Runtime/Engine/Classes/Engine/SkeletalMesh.h
+++ b/Engine/Source/Runtime/Engine/Classes/Engine/SkeletalMesh.h
@@ -9,6 +9,7 @@
 #include "SkeletalMeshTypes.h"
 #include "Animation/PreviewAssetAttachComponent.h"
 #include "BoneContainer.h"
+#include "Interfaces/Interface_CollisionDataProvider.h"
 #include "SkeletalMesh.generated.h"
 
 /** The maximum number of skeletal mesh LODs allowed. */
@@ -563,7 +564,7 @@ class FSkeletalMeshResource;
  * @see https://docs.unrealengine.com/latest/INT/Engine/Content/Types/SkeletalMeshes/
  */
 UCLASS(hidecategories=Object, MinimalAPI, BlueprintType)
-class USkeletalMesh : public UObject
+class USkeletalMesh : public UObject, public IInterface_CollisionDataProvider
 {
 	GENERATED_UCLASS_BODY()
 
@@ -615,6 +616,15 @@ public:
 	UPROPERTY()
 	uint32 bHasVertexColors:1;
 
+
+	/** Uses skinned data for collision data. Per poly collision cannot be used for simulation, in most cases you are better off using the physics asset */
+	UPROPERTY(EditAnywhere, Category = Physics)
+	uint32 bEnablePerPolyCollision : 1;
+
+	// Physics data for the per poly collision case. In 99% of cases you will not need this and are better off using simple ragdoll collision (physics asset)
+	UPROPERTY(transient)
+	class UBodySetup* BodySetup;
+
 	/**
 	 *	Physics and collision information used for this USkeletalMesh, set up in PhAT.
 	 *	This is used for per-bone hit detection, accurate bounding box calculation and ragdoll physics for example.
@@ -871,6 +881,22 @@ public:
 	ENGINE_API int32 GetClothAssetIndex(int32 LODIndex, int32 SectionIndex);
 #endif// #if WITH_APEX_CLOTHING
 
+	ENGINE_API void CreateBodySetup();
+	ENGINE_API UBodySetup* GetBodySetup();
+
+	/** Trigger a physics build to ensure per poly collision is created */
+	ENGINE_API void BuildPhysicsData();
+	
+
+	// Begin Interface_CollisionDataProvider Interface
+	virtual bool GetPhysicsTriMeshData(struct FTriMeshCollisionData* CollisionData, bool InUseAllTriData) override;
+	virtual bool ContainsPhysicsTriMeshData(bool InUseAllTriData) const override;
+	virtual bool WantsNegXTriMesh() override
+	{
+		return true;
+	}
+	// End Interface_CollisionDataProvider Interface
+
 private:
 
 #if WITH_EDITOR
diff --git a/Engine/Source/Runtime/Engine/Classes/PhysicsEngine/BodySetup.h b/Engine/Source/Runtime/Engine/Classes/PhysicsEngine/BodySetup.h
--- a/Engine/Source/Runtime/Engine/Classes/PhysicsEngine/BodySetup.h
+++ b/Engine/Source/Runtime/Engine/Classes/PhysicsEngine/BodySetup.h
@@ -449,6 +449,10 @@ class UBodySetup : public UObject
 	UPROPERTY()
 	uint32 bGenerateNonMirroredCollision:1;
 
+	/** Whether the cooked data is shared by multiple body setups. This is needed for per poly collision case where we don't want to duplicate cooked data, but still need multiple body setups for in place geometry changes */
+	UPROPERTY()
+	uint32 bSharedCookedData : 1;
+
 	/** 
 	 *	Should we generate data necessary to support collision on mirrored versions of this mesh. 
 	 *	This halves the collision data size for this mesh, but disables collision on mirrored instances of the body.
@@ -489,6 +493,9 @@ class UBodySetup : public UObject
 	/** Cooked physics data for each format */
 	FFormatContainer CookedFormatData;
 
+	/** Cooked physics data override. This is needed in cases where some other body setup has the cooked data and you don't want to own it or copy it. See per poly skeletal mesh */
+	FFormatContainer* CookedFormatDataOverride;
+
 #if WITH_PHYSX
 	/** Physics triangle mesh, created from cooked data in CreatePhysicsMeshes */
 	physx::PxTriangleMesh* TriMesh;
@@ -519,7 +526,7 @@ public:
 	//
 	// UBodySetup interface.
 	//
-	ENGINE_API void CopyBodyPropertiesFrom(class UBodySetup* FromSetup);
+	ENGINE_API void CopyBodyPropertiesFrom(const UBodySetup* FromSetup);
 
 	/** Add collision shapes from another body setup to this one */
 	ENGINE_API void AddCollisionFrom(class UBodySetup* FromSetup);
diff --git a/Engine/Source/Runtime/Engine/Private/DestructibleMesh.cpp b/Engine/Source/Runtime/Engine/Private/DestructibleMesh.cpp
--- a/Engine/Source/Runtime/Engine/Private/DestructibleMesh.cpp
+++ b/Engine/Source/Runtime/Engine/Private/DestructibleMesh.cpp
@@ -521,14 +521,6 @@ void UDestructibleMesh::LoadDefaultDestructibleParametersFromApexAsset()
 #endif
 }
 
-void UDestructibleMesh::CreateBodySetup()
-{
-	if (BodySetup == NULL)
-	{
-		BodySetup = ConstructObject<UBodySetup>(UBodySetup::StaticClass(), this);
-	}
-}
-
 void UDestructibleMesh::CreateFractureSettings()
 {
 #if WITH_EDITORONLY_DATA
diff --git a/Engine/Source/Runtime/Engine/Private/PhysicsEngine/BodySetup.cpp b/Engine/Source/Runtime/Engine/Private/PhysicsEngine/BodySetup.cpp
--- a/Engine/Source/Runtime/Engine/Private/PhysicsEngine/BodySetup.cpp
+++ b/Engine/Source/Runtime/Engine/Private/PhysicsEngine/BodySetup.cpp
@@ -52,9 +52,11 @@ UBodySetup::UBodySetup(const FObjectInitializer& ObjectInitializer)
 	BuildScale_DEPRECATED = 1.0f;
 	BuildScale3D = FVector(1.0f, 1.0f, 1.0f);
 	SetFlags(RF_Transactional);
+	bSharedCookedData = false;
+	CookedFormatDataOverride = nullptr;
 }
 
-void UBodySetup::CopyBodyPropertiesFrom(class UBodySetup* FromSetup)
+void UBodySetup::CopyBodyPropertiesFrom(const UBodySetup* FromSetup)
 {
 	AggGeom = FromSetup->AggGeom;
 
@@ -179,7 +181,7 @@ void UBodySetup::CreatePhysicsMeshes()
 	}
 
 	// Clear the cooked data
-	if (!GIsEditor)
+	if (!GIsEditor && !bSharedCookedData)
 	{
 		CookedFormatData.FlushData();
 	}
@@ -672,7 +674,10 @@ void UBodySetup::InvalidatePhysicsData()
 {
 	ClearPhysicsMeshes();
 	BodySetupGuid = FGuid::NewGuid(); // change the guid
-	CookedFormatData.FlushData();
+	if (!bSharedCookedData)
+	{
+		CookedFormatData.FlushData();
+	}
 }
 #endif // WITH_EDITOR
 
@@ -697,6 +702,7 @@ void UBodySetup::Serialize(FArchive& Ar)
 {
 	Super::Serialize(Ar);
 
+
 	// Load GUID (or create one for older versions)
 	Ar << BodySetupGuid;
 
@@ -857,8 +863,10 @@ FByteBulkData* UBodySetup::GetCookedData(FName Format)
 		return NULL;
 	}
 
-	bool bContainedData = CookedFormatData.Contains(Format);
-	FByteBulkData* Result = &CookedFormatData.GetFormat(Format);
+	FFormatContainer* UseCookedData = CookedFormatDataOverride ? CookedFormatDataOverride : &CookedFormatData;
+
+	bool bContainedData = UseCookedData->Contains(Format);
+	FByteBulkData* Result = &UseCookedData->GetFormat(Format);
 #if WITH_PHYSX
 	if (!bContainedData)
 	{
diff --git a/Engine/Source/Runtime/Engine/Private/PhysicsEngine/PhysDerivedData.cpp b/Engine/Source/Runtime/Engine/Private/PhysicsEngine/PhysDerivedData.cpp
--- a/Engine/Source/Runtime/Engine/Private/PhysicsEngine/PhysDerivedData.cpp
+++ b/Engine/Source/Runtime/Engine/Private/PhysicsEngine/PhysDerivedData.cpp
@@ -195,9 +195,9 @@ bool FDerivedDataPhysXCooker::BuildTriMesh( TArray<uint8>& OutData, bool bInMirr
 
 		UE_LOG(LogPhysics, Log, TEXT("Cook TriMesh: %s (FlipX: %d)"), *CollisionDataProvider->GetPathName(), bInMirrored);
 		bool bPerPolySkeletalMesh = false;
-		if (USkeletalMeshComponent * SkeletalMeshComponent = Cast<USkeletalMeshComponent>(CollisionDataProvider))
+		if (USkeletalMesh* SkeletalMesh = Cast<USkeletalMesh>(CollisionDataProvider))
 		{
-			ensure(SkeletalMeshComponent->bEnablePerPolyCollision);
+			ensure(SkeletalMesh->bEnablePerPolyCollision);
 			bPerPolySkeletalMesh = true;
 		}
 		bResult = Cooker->CookTriMesh( Format, *MeshVertices, TriangleMeshDesc.Indices, TriangleMeshDesc.MaterialIndices, bInMirrored ? !TriangleMeshDesc.bFlipNormals : TriangleMeshDesc.bFlipNormals, OutData, bPerPolySkeletalMesh );
diff --git a/Engine/Source/Runtime/Engine/Private/SkeletalMesh.cpp b/Engine/Source/Runtime/Engine/Private/SkeletalMesh.cpp
--- a/Engine/Source/Runtime/Engine/Private/SkeletalMesh.cpp
+++ b/Engine/Source/Runtime/Engine/Private/SkeletalMesh.cpp
@@ -2218,6 +2218,14 @@ void USkeletalMesh::PostEditChangeProperty(FPropertyChangedEvent& PropertyChange
 		InitMorphTargets();
 	}
 
+	if ( GIsEditor &&
+		 PropertyThatChanged &&
+		 PropertyThatChanged->GetFName() == FName(TEXT("bEnablePerPolyCollision"))
+		)
+	{
+		BuildPhysicsData();
+	}
+
 	Super::PostEditChangeProperty(PropertyChangedEvent);
 }
 
@@ -2397,6 +2405,11 @@ void USkeletalMesh::Serialize( FArchive& Ar )
 		PreviewAttachedAssetContainer.SaveAttachedObjectsFromDeprecatedProperties();
 	}
 #endif
+
+	if (bEnablePerPolyCollision)
+	{
+		Ar << BodySetup;
+	}
 }
 
 void USkeletalMesh::AddReferencedObjects(UObject* InThis, FReferenceCollector& Collector)
@@ -3344,6 +3357,116 @@ bool USkeletalMesh::MirrorTableIsGood(FString& ProblemBones)
 	}
 }
 
+void USkeletalMesh::CreateBodySetup()
+{
+	if (BodySetup == nullptr)
+	{
+		BodySetup = ConstructObject<UBodySetup>(UBodySetup::StaticClass(), this);
+		BodySetup->bSharedCookedData = true;
+	}
+}
+
+UBodySetup* USkeletalMesh::GetBodySetup()
+{
+	CreateBodySetup();
+	return BodySetup;
+}
+
+void USkeletalMesh::BuildPhysicsData()
+{
+	CreateBodySetup();
+	BodySetup->CookedFormatData.FlushData();	//we need to force a re-cook because we're essentially re-creating the bodysetup so that it swaps whether or not it has a trimesh
+	BodySetup->InvalidatePhysicsData();
+	BodySetup->CreatePhysicsMeshes();
+}
+
+bool USkeletalMesh::ContainsPhysicsTriMeshData(bool InUseAllTriData) const
+{
+	return bEnablePerPolyCollision;
+}
+
+bool USkeletalMesh::GetPhysicsTriMeshData(FTriMeshCollisionData* CollisionData, bool bInUseAllTriData)
+{
+#if WITH_EDITORONLY_DATA
+
+	// Fail if no mesh or not per poly collision
+	if (!ImportedResource.IsValid() || !bEnablePerPolyCollision)
+	{
+		return false;
+	}
+
+	const FStaticLODModel& Model = ImportedResource->LODModels[0];
+
+	{
+		// Copy all verts into collision vertex buffer.
+		CollisionData->Vertices.Empty();
+		CollisionData->Vertices.AddUninitialized(Model.NumVertices);
+		const uint32 NumChunks = Model.Chunks.Num();
+
+		for (uint32 ChunkIdx = 0; ChunkIdx < NumChunks; ++ChunkIdx)
+		{
+			const FSkelMeshChunk& Chunk = Model.Chunks[ChunkIdx];
+			{
+				//rigid
+				const uint32 RigidOffset = Chunk.GetRigidVertexBufferIndex();
+				const uint32 NumRigidVerts = Chunk.GetNumRigidVertices();
+				for (uint32 RigidIdx = 0; RigidIdx < NumRigidVerts; ++RigidIdx)
+				{
+					CollisionData->Vertices[RigidIdx + RigidOffset] = Chunk.RigidVertices[RigidIdx].Position;
+				}
+			}
+			{
+				//soft
+				const uint32 SoftOffset = Chunk.GetSoftVertexBufferIndex();
+				const uint32 NumSoftVerts = Chunk.GetNumSoftVertices();
+				for (uint32 SoftIdx = 0; SoftIdx < NumSoftVerts; ++SoftIdx)
+				{
+					CollisionData->Vertices[SoftIdx + SoftOffset] = Chunk.SoftVertices[SoftIdx].Position;
+				}
+			}
+
+		}
+	}
+
+	{
+		// Copy indices into collision index buffer
+		const FMultiSizeIndexContainer& IndexBufferContainer = Model.MultiSizeIndexContainer;
+
+		TArray<uint32> Indices;
+		IndexBufferContainer.GetIndexBuffer(Indices);
+
+		const uint32 NumTris = Indices.Num() / 3;
+		CollisionData->Indices.Empty();
+		CollisionData->Indices.Reserve(NumTris);
+
+		FTriIndices TriIndex;
+		for (int32 SectionIndex = 0; SectionIndex < Model.Sections.Num(); ++SectionIndex)
+		{
+			const FSkelMeshSection& Section = Model.Sections[SectionIndex];
+
+			const uint32 OnePastLastIndex = Section.BaseIndex + Section.NumTriangles * 3;
+
+			for (uint32 i = Section.BaseIndex; i < OnePastLastIndex; i += 3)
+			{
+				TriIndex.v0 = Indices[i];
+				TriIndex.v1 = Indices[i + 1];
+				TriIndex.v2 = Indices[i + 2];
+
+				CollisionData->Indices.Add(TriIndex);
+				CollisionData->MaterialIndices.Add(Section.MaterialIndex);
+			}
+		}
+	}
+
+	CollisionData->bFlipNormals = true;
+
+	// We only have a valid TriMesh if the CollisionData has vertices AND indices. For meshes with disabled section collision, it
+	// can happen that the indices will be empty, in which case we do not want to consider that as valid trimesh data
+	return CollisionData->Vertices.Num() > 0 && CollisionData->Indices.Num() > 0;
+#else // #if WITH_EDITORONLY_DATA
+	return false;
+#endif // #if WITH_EDITORONLY_DATA
+}
 
 
 ////// SKELETAL MESH THUMBNAIL SUPPORT ////////
@@ -4778,6 +4901,11 @@ void USkeletalMeshComponent::Serialize(FArchive& Ar)
 		}
 	}
 #endif
+
+	if (Ar.IsLoading() && (Ar.UE4Ver() < VER_UE4_AUTO_WELDING))
+	{
+		BodyInstance.bAutoWeld = false;
+	}
 }
 
 
diff --git a/Engine/Source/Runtime/Engine/Private/SkeletalMeshComponentPhysics.cpp b/Engine/Source/Runtime/Engine/Private/SkeletalMeshComponentPhysics.cpp
--- a/Engine/Source/Runtime/Engine/Private/SkeletalMeshComponentPhysics.cpp
+++ b/Engine/Source/Runtime/Engine/Private/SkeletalMeshComponentPhysics.cpp
@@ -365,99 +365,21 @@ int32 USkeletalMesh::GetClothAssetIndex(int32 LODIndex, int32 SectionIndex)
 
 void USkeletalMeshComponent::CreateBodySetup()
 {
-	if (BodySetup == NULL)
+	if (BodySetup == NULL && SkeletalMesh)
 	{
 		BodySetup = ConstructObject<UBodySetup>(UBodySetup::StaticClass(), this);
-		BodySetup->CollisionTraceFlag = CTF_UseComplexAsSimple;
 	}
-}
-
-bool USkeletalMeshComponent::GetPhysicsTriMeshData(FTriMeshCollisionData* CollisionData, bool bInUseAllTriData)
-{
-#if WITH_EDITORONLY_DATA
-
-	// Fail if no mesh or not per poly collision
-	if (!SkeletalMesh || !MeshObject || !bEnablePerPolyCollision)
-	{
-		return false;
-	}
-
-	const FStaticLODModel& Model = MeshObject->GetSkeletalMeshResource().LODModels[0];
-
-	{
-		// Copy all verts into collision vertex buffer.
-		CollisionData->Vertices.Empty();
-		CollisionData->Vertices.AddUninitialized(Model.NumVertices);
-		const uint32 NumChunks = Model.Chunks.Num();
 
-		for (uint32 ChunkIdx = 0; ChunkIdx < NumChunks; ++ChunkIdx)
-		{
-			const FSkelMeshChunk& Chunk = Model.Chunks[ChunkIdx];
-			{
-				//rigid
-				const uint32 RigidOffset = Chunk.GetRigidVertexBufferIndex();
-				const uint32 NumRigidVerts = Chunk.GetNumRigidVertices();
-				for (uint32 RigidIdx = 0; RigidIdx < NumRigidVerts; ++RigidIdx)
-				{
-					CollisionData->Vertices[RigidIdx + RigidOffset] = Chunk.RigidVertices[RigidIdx].Position;
-				}
-			}
-			{
-				//soft
-				const uint32 SoftOffset = Chunk.GetSoftVertexBufferIndex();
-				const uint32 NumSoftVerts = Chunk.GetNumSoftVertices();
-				for (uint32 SoftIdx = 0; SoftIdx < NumSoftVerts; ++SoftIdx)
-				{
-					CollisionData->Vertices[SoftIdx + SoftOffset] = Chunk.SoftVertices[SoftIdx].Position;
-				}
-			}
-			
-		}
-	}
+	UBodySetup* OriginalBodySetup = SkeletalMesh->GetBodySetup();
 	
-	{
-		// Copy indices into collision index buffer
-		const FMultiSizeIndexContainer& IndexBufferContainer = Model.MultiSizeIndexContainer;
-
-		TArray<uint32> Indices;
-		IndexBufferContainer.GetIndexBuffer(Indices);
+	BodySetup->CopyBodyPropertiesFrom(OriginalBodySetup);
+	BodySetup->CollisionTraceFlag = CTF_UseComplexAsSimple;
 
-		const uint32 NumTris = Indices.Num() / 3;
-		CollisionData->Indices.Empty();
-		CollisionData->Indices.Reserve(NumTris);
-
-		FTriIndices TriIndex;
-		for (int32 SectionIndex = 0; SectionIndex < Model.Sections.Num(); ++SectionIndex)
-		{
-			const FSkelMeshSection& Section = Model.Sections[SectionIndex];
-
-			const uint32 OnePastLastIndex = Section.BaseIndex + Section.NumTriangles * 3;
-
-			for (uint32 i = Section.BaseIndex; i < OnePastLastIndex; i += 3)
-			{
-				TriIndex.v0 = Indices[i];
-				TriIndex.v1 = Indices[i + 1];
-				TriIndex.v2 = Indices[i + 2];
+	BodySetup->CookedFormatDataOverride = &OriginalBodySetup->CookedFormatData;
 
-				CollisionData->Indices.Add(TriIndex);
-				CollisionData->MaterialIndices.Add(Section.MaterialIndex);
-			}
-		}
-	}
-
-	CollisionData->bFlipNormals = true;
-
-	// We only have a valid TriMesh if the CollisionData has vertices AND indices. For meshes with disabled section collision, it
-	// can happen that the indices will be empty, in which case we do not want to consider that as valid trimesh data
-	return CollisionData->Vertices.Num() > 0 && CollisionData->Indices.Num() > 0;
-#else // #if WITH_EDITORONLY_DATA
-	return false;
-#endif // #if WITH_EDITORONLY_DATA
-}
-
-bool USkeletalMeshComponent::ContainsPhysicsTriMeshData(bool InUseAllTriData) const
-{
-	return bEnablePerPolyCollision;
+	//need to recreate meshes
+	BodySetup->ClearPhysicsMeshes();
+	BodySetup->CreatePhysicsMeshes();
 }
 
 //
@@ -1374,6 +1296,7 @@ void USkeletalMeshComponent::CreatePhysicsState()
 	else
 	{
 		CreateBodySetup();
+		BodySetup->CreatePhysicsMeshes();
 		Super::CreatePhysicsState();	//If we're doing per poly we'll use the body instance of the primitive component
 	}
 
