Android audio fixes
- Fixes looping issue with ADPCM samples
- Fixes Android not looping sound cues correctly
- Fixes one shot sound ADPCM problems where sound would get cut early
- Removed Android Audio Device update function - was never called and did nothing anyway so this will improve clarity

Changed ADPCM decompressor so that it uses the same code to decompress a full file as it does when streaming - old code was correct it was just silly have it duplicated.

diff --git a/Engine/Source/Runtime/Android/AndroidAudio/Private/AndroidAudioDevice.cpp b/Engine/Source/Runtime/Android/AndroidAudio/Private/AndroidAudioDevice.cpp
--- a/Engine/Source/Runtime/Android/AndroidAudio/Private/AndroidAudioDevice.cpp
+++ b/Engine/Source/Runtime/Android/AndroidAudio/Private/AndroidAudioDevice.cpp
@@ -137,19 +137,6 @@ bool FSLESAudioDevice::InitializeHardware( void )
 	return true;
 }
 
-/**
- * Update the audio device and calculates the cached inverse transform later
- * on used for spatialization.
- *
- * @param	Realtime	whether we are paused or not
- */
-void FSLESAudioDevice::Update( bool Realtime )
-{
-	//Super::Update( Realtime );
-	
-	//@todo android: UDPATE LISTENERS - Android OpenSLES doesn't support 3D, is there anything to do here?
-}
-
 FSoundSource* FSLESAudioDevice::CreateSoundSource()
 {
 	return new FSLESSoundSource(this);
diff --git a/Engine/Source/Runtime/Android/AndroidAudio/Private/AndroidAudioSource.cpp b/Engine/Source/Runtime/Android/AndroidAudio/Private/AndroidAudioSource.cpp
--- a/Engine/Source/Runtime/Android/AndroidAudio/Private/AndroidAudioSource.cpp
+++ b/Engine/Source/Runtime/Android/AndroidAudio/Private/AndroidAudioSource.cpp
@@ -34,7 +34,14 @@ void FSLESSoundSource::OnRequeueBufferCallback( SLAndroidSimpleBufferQueueItf In
 	else
 	{
 		// Sound decoding is complete, just waiting to finish playing
-		if (bBuffersToFlush) return;
+		if (bBuffersToFlush)
+		{
+			// set the player's state to stopped
+			SLresult result = (*SL_PlayerPlayInterface)->SetPlayState(SL_PlayerPlayInterface, SL_PLAYSTATE_STOPPED);
+			check(SL_RESULT_SUCCESS == result);
+
+			return;
+		}
 
 		// Enqueue the previously decoded buffer
 		SLresult result = (*SL_PlayerBufferQueue)->Enqueue(SL_PlayerBufferQueue, AudioBuffers[BufferInUse].AudioData, AudioBuffers[BufferInUse].AudioDataSize );
@@ -415,7 +422,6 @@ void FSLESSoundSource::Stop( void )
 {
 	if( WaveInstance )
 	{
-		
 		// set the player's state to stopped
 		SLresult result = (*SL_PlayerPlayInterface)->SetPlayState(SL_PlayerPlayInterface, SL_PLAYSTATE_STOPPED);
 		check(SL_RESULT_SUCCESS == result);
@@ -485,8 +491,6 @@ bool FSLESSoundSource::IsFinished( void )
 		// Check for a non starved, stopped source
 		if( IsSourceFinished() )
 		{
-			
-			
 			// Notify the wave instance that it has finished playing
 			WaveInstance->NotifyFinished();
 			return true;
@@ -516,6 +520,5 @@ bool FSLESSoundSource::IsFinished( void )
 		
 		return false;
 	}
-	
 	return true;
 }
diff --git a/Engine/Source/Runtime/Android/AndroidAudio/Public/AndroidAudioDevice.h b/Engine/Source/Runtime/Android/AndroidAudio/Public/AndroidAudioDevice.h
--- a/Engine/Source/Runtime/Android/AndroidAudio/Public/AndroidAudioDevice.h
+++ b/Engine/Source/Runtime/Android/AndroidAudio/Public/AndroidAudioDevice.h
@@ -271,14 +271,6 @@ public:
 	/** Starts up any platform specific hardware/APIs */
 	virtual bool InitializeHardware() override;
  
-	/**
-	 * Update the audio device and calculates the cached inverse transform later
-	 * on used for spatialization.
-	 *
-	 * @param	Realtime	whether we are paused or not
-	 */
-	virtual void Update( bool bGameTicking );
-
 	/** Check if any background music or sound is playing through the audio device */
 	virtual bool IsExernalBackgroundSoundActive() override;
 
diff --git a/Engine/Source/Runtime/Engine/Private/ADPCMAudioInfo.cpp b/Engine/Source/Runtime/Engine/Private/ADPCMAudioInfo.cpp
--- a/Engine/Source/Runtime/Engine/Private/ADPCMAudioInfo.cpp
+++ b/Engine/Source/Runtime/Engine/Private/ADPCMAudioInfo.cpp
@@ -49,13 +49,13 @@ bool FADPCMAudioInfo::ReadCompressedInfo(const uint8* InSrcBufferData, uint32 In
 	const uint32 targetBlocks = MONO_PCM_BUFFER_SAMPLES / uncompressedBlockSamples;
 	StreamBufferSize = targetBlocks * UncompressedBlockSize;
 	StreamBufferSizeInBlocks = CompressedBlockSize * targetBlocks;
-
+	TotalDecodedSize = (WaveInfo.SampleDataSize / CompressedBlockSize) * UncompressedBlockSize;
 	if (QualityInfo)
 	{
 		QualityInfo->SampleRate = *WaveInfo.pSamplesPerSec;
 		QualityInfo->NumChannels = *WaveInfo.pChannels;
-		QualityInfo->SampleDataSize = (WaveInfo.SampleDataSize / CompressedBlockSize) * UncompressedBlockSize;
-		//QualityInfo->Duration = (float)TrueSampleCount / QualityInfo->SampleRate;
+		QualityInfo->SampleDataSize = TotalDecodedSize;
+		QualityInfo->Duration = 0.0f; // (float)TrueSampleCount / QualityInfo->SampleRate;
 	}
 
 	return true;
@@ -72,31 +72,7 @@ void FADPCMAudioInfo::ExpandFile(uint8* DstBuffer, struct FSoundQualityInfo* Qua
 {
 	check(DstBuffer);
 
-	if (*WaveInfo.pChannels == 1)
-	{
-		const uint8* EncodedADPCMBlockStart = reinterpret_cast<uint8*>(WaveInfo.SampleDataStart);
-		uint8* EncodedADPCMBlock = (uint8*)EncodedADPCMBlockStart;
-		while (EncodedADPCMBlock < EncodedADPCMBlockStart + (WaveInfo.SampleDataSize))
-		{
-			ADPCM::DecodeBlock(EncodedADPCMBlock, CompressedBlockSize, (int16*)DstBuffer);
-			EncodedADPCMBlock += CompressedBlockSize;
-			DstBuffer += UncompressedBlockSize;
-		}
-	}
-	else
-	{
-		const uint8* EncodedADPCMBlockLeftStart = reinterpret_cast<uint8*>(WaveInfo.SampleDataStart);
-		const uint8* EncodedADPCMBlockRightStart = reinterpret_cast<uint8*>(WaveInfo.SampleDataStart + (WaveInfo.SampleDataSize / 2));
-		uint8* EncodedADPCMLeftBlock = (uint8*)EncodedADPCMBlockLeftStart;
-		uint8* EncodedADPCMRightBlock = (uint8*)EncodedADPCMBlockRightStart;
-		while (EncodedADPCMLeftBlock < EncodedADPCMBlockLeftStart + (WaveInfo.SampleDataSize / 2))
-		{
-			ADPCM::DecodeBlockStereo(EncodedADPCMLeftBlock, EncodedADPCMRightBlock, CompressedBlockSize, (int16*)DstBuffer);
-			EncodedADPCMLeftBlock += CompressedBlockSize;
-			EncodedADPCMRightBlock += CompressedBlockSize;
-			DstBuffer += (UncompressedBlockSize * 2);
-		}
-	}
+	ReadCompressedData(DstBuffer, false, TotalDecodedSize);
 }
 
 int FADPCMAudioInfo::GetStreamBufferSize() const
@@ -111,6 +87,8 @@ bool FADPCMAudioInfo::DecodeStereoData(uint8* Destination, bool bLooping, uint32
 	uint8* EncodedADPCMLeftBlock = (uint8*)EncodedADPCMBlockLeftStart + SrcBufferOffset;
 	uint8* EncodedADPCMRightBlock = (uint8*)EncodedADPCMBlockRightStart + SrcBufferOffset;
 	
+	bool looped = false;
+
 	// Check to see if we have enough data left to do a complete decode
 	if (EncodedADPCMLeftBlock + StreamBufferSizeInBlocks < EncodedADPCMBlockRightStart)
 	{
@@ -137,6 +115,8 @@ bool FADPCMAudioInfo::DecodeStereoData(uint8* Destination, bool bLooping, uint32
 			BufferSize -= (UncompressedBlockSize * 2);
 		}
 
+		looped = true;
+
 		// and then decide if we zero pad or reset and loop
 		if (bLooping)
 		{
@@ -159,7 +139,7 @@ bool FADPCMAudioInfo::DecodeStereoData(uint8* Destination, bool bLooping, uint32
 		}
 	}
 
-	return true;
+	return looped;
 }
 
 bool FADPCMAudioInfo::DecodeMonoData(uint8* Destination, bool bLooping, uint32 BufferSize)
@@ -168,6 +148,8 @@ bool FADPCMAudioInfo::DecodeMonoData(uint8* Destination, bool bLooping, uint32 B
 	uint8* EncodedADPCMBlock = (uint8*)EncodedADPCMBlockStart + SrcBufferOffset;
 	const uint8* EncodedDataEnd = EncodedADPCMBlockStart + WaveInfo.SampleDataSize;
 
+	bool looped = false;
+
 	// See if we are still completely inside the source data..
 	if (EncodedADPCMBlock + StreamBufferSizeInBlocks < EncodedDataEnd)
 	{
@@ -194,6 +176,9 @@ bool FADPCMAudioInfo::DecodeMonoData(uint8* Destination, bool bLooping, uint32 B
 		}
 
 		// and then decide if we zero pad or reset and loop
+
+		looped = true;
+
 		if (bLooping)
 		{
 			EncodedADPCMBlock = (uint8*)EncodedADPCMBlockStart;
@@ -213,5 +198,5 @@ bool FADPCMAudioInfo::DecodeMonoData(uint8* Destination, bool bLooping, uint32 B
 		}
 	}
 
-	return true;
+	return looped;
 }
diff --git a/Engine/Source/Runtime/Engine/Private/AudioDecompress.cpp b/Engine/Source/Runtime/Engine/Private/AudioDecompress.cpp
--- a/Engine/Source/Runtime/Engine/Private/AudioDecompress.cpp
+++ b/Engine/Source/Runtime/Engine/Private/AudioDecompress.cpp
@@ -316,7 +316,7 @@ void FAsyncAudioDecompressWorker::DoWork( void )
 			// Extract the data
 			Wave->SampleRate = QualityInfo.SampleRate;
 			Wave->NumChannels = QualityInfo.NumChannels;
-			Wave->Duration = QualityInfo.Duration;
+			if(QualityInfo.Duration > 0.0f) Wave->Duration = QualityInfo.Duration;
 
 			Wave->RawPCMDataSize = QualityInfo.SampleDataSize;
 			Wave->RawPCMData = ( uint8* )FMemory::Malloc( Wave->RawPCMDataSize );
diff --git a/Engine/Source/Runtime/Engine/Public/ADPCMAudioInfo.h b/Engine/Source/Runtime/Engine/Public/ADPCMAudioInfo.h
--- a/Engine/Source/Runtime/Engine/Public/ADPCMAudioInfo.h
+++ b/Engine/Source/Runtime/Engine/Public/ADPCMAudioInfo.h
@@ -35,6 +35,7 @@ public:
 	uint32 BlockSize;
 	int32 StreamBufferSize;
 	uint32 StreamBufferSizeInBlocks;
+	uint32 TotalDecodedSize;
 
 private:
 	bool DecodeStereoData(uint8* Destination, bool bLooping, uint32 BufferSize);
