Expose network and travel failure to blueprints as events on the GameInstance.

diff --git a/Engine/Source/Runtime/Engine/Classes/Engine/Engine.h b/Engine/Source/Runtime/Engine/Classes/Engine/Engine.h
--- a/Engine/Source/Runtime/Engine/Classes/Engine/Engine.h
+++ b/Engine/Source/Runtime/Engine/Classes/Engine/Engine.h
@@ -2525,6 +2525,12 @@ public:
 private:
 	void CreateGameUserSettings();
 
+	/** Allows subclasses to pass the failure to a UGameInstance if possible (mainly for blueprints) */
+	virtual void HandleNetworkFailure_NotifyGameInstance(UWorld* World, UNetDriver* NetDriver, ENetworkFailure::Type FailureType);
+
+	/** Allows subclasses to pass the failure to a UGameInstance if possible (mainly for blueprints) */
+	virtual void HandleTravelFailure_NotifyGameInstance(UWorld* World, ETravelFailure::Type FailureType);
+
 public:
 	/**
 	 * Delegate we fire every time a new stat has been registered
diff --git a/Engine/Source/Runtime/Engine/Classes/Engine/EngineBaseTypes.h b/Engine/Source/Runtime/Engine/Classes/Engine/EngineBaseTypes.h
--- a/Engine/Source/Runtime/Engine/Classes/Engine/EngineBaseTypes.h
+++ b/Engine/Source/Runtime/Engine/Classes/Engine/EngineBaseTypes.h
@@ -406,6 +406,7 @@ struct TStructOpsTypeTraits<FPrimitiveComponentPostPhysicsTickFunction> : public
 };
 
 /** Types of network failures broadcast from the engine */
+UENUM(BlueprintType)
 namespace ENetworkFailure
 {
 	enum Type
@@ -429,7 +430,11 @@ namespace ENetworkFailure
 		/** There was an error during connection to the game */
 		PendingConnectionFailure
 	};
+}
+
 
+namespace ENetworkFailure
+{
 	inline const TCHAR* ToString(ENetworkFailure::Type FailureType)
 	{
 		switch (FailureType)
@@ -458,6 +463,7 @@ namespace ENetworkFailure
 }
 
 /** Types of server travel failures broadcast by the engine */
+UENUM(BlueprintType)
 namespace ETravelFailure
 {
 	enum Type
@@ -487,7 +493,10 @@ namespace ETravelFailure
 		/** There was an error during a client travel to a new map */
 		ClientTravelFailure,
 	};
+}
 
+namespace ETravelFailure
+{
 	inline const TCHAR* ToString(ETravelFailure::Type FailureType)
 	{
 		switch (FailureType)
diff --git a/Engine/Source/Runtime/Engine/Classes/Engine/GameEngine.h b/Engine/Source/Runtime/Engine/Classes/Engine/GameEngine.h
--- a/Engine/Source/Runtime/Engine/Classes/Engine/GameEngine.h
+++ b/Engine/Source/Runtime/Engine/Classes/Engine/GameEngine.h
@@ -137,5 +137,9 @@ public:
 
 protected:
 	float GetGameViewportDPIScale(UGameViewportClient* GameViewportClient) const;
+
+private:
+	virtual void HandleNetworkFailure_NotifyGameInstance(UWorld* World, UNetDriver* NetDriver, ENetworkFailure::Type FailureType) override;
+	virtual void HandleTravelFailure_NotifyGameInstance(UWorld* World, ETravelFailure::Type FailureType) override;
 };
 
diff --git a/Engine/Source/Runtime/Engine/Classes/Engine/GameInstance.h b/Engine/Source/Runtime/Engine/Classes/Engine/GameInstance.h
--- a/Engine/Source/Runtime/Engine/Classes/Engine/GameInstance.h
+++ b/Engine/Source/Runtime/Engine/Classes/Engine/GameInstance.h
@@ -1,6 +1,7 @@
 // Copyright 1998-2014 Epic Games, Inc. All Rights Reserved.
 
 #pragma once
+#include "EngineBaseTypes.h"
 #include "GameInstance.generated.h"
 
 
@@ -69,6 +70,14 @@ public:
 	UFUNCTION(BlueprintImplementableEvent, meta=(FriendlyName = "Shutdown"))
 	virtual void ReceiveShutdown();
 
+	/** Opportunity for blueprints to handle network errors. */
+	UFUNCTION(BlueprintImplementableEvent, meta=(FriendlyName = "NetworkError"))
+	virtual void HandleNetworkError(ENetworkFailure::Type FailureType, bool bIsServer);
+
+	/** Opportunity for blueprints to handle travel errors. */
+	UFUNCTION(BlueprintImplementableEvent, meta=(FriendlyName = "TravelError"))
+	virtual void HandleTravelError(ETravelFailure::Type FailureType);
+
 	/* Called to initialize the game instance for standalone instances of the game */
 	void InitializeStandalone();
 #if WITH_EDITOR
diff --git a/Engine/Source/Runtime/Engine/Private/GameEngine.cpp b/Engine/Source/Runtime/Engine/Private/GameEngine.cpp
--- a/Engine/Source/Runtime/Engine/Private/GameEngine.cpp
+++ b/Engine/Source/Runtime/Engine/Private/GameEngine.cpp
@@ -1067,3 +1067,24 @@ UWorld* UGameEngine::GetGameWorld()
 
 	return NULL;
 }
+
+void UGameEngine::HandleNetworkFailure_NotifyGameInstance(UWorld* World, UNetDriver* NetDriver, ENetworkFailure::Type FailureType)
+{
+	if (GameInstance != nullptr)
+	{
+		bool bIsServer = true;
+		if (NetDriver != nullptr)
+		{
+			bIsServer = NetDriver->GetNetMode() != NM_Client;
+		}
+		GameInstance->HandleNetworkError(FailureType, bIsServer);
+	}
+}
+
+void UGameEngine::HandleTravelFailure_NotifyGameInstance(UWorld* World, ETravelFailure::Type FailureType)
+{
+	if (GameInstance != nullptr)
+	{
+		GameInstance->HandleTravelError(FailureType);
+	}
+}
diff --git a/Engine/Source/Runtime/Engine/Private/UnrealEngine.cpp b/Engine/Source/Runtime/Engine/Private/UnrealEngine.cpp
--- a/Engine/Source/Runtime/Engine/Private/UnrealEngine.cpp
+++ b/Engine/Source/Runtime/Engine/Private/UnrealEngine.cpp
@@ -7654,6 +7654,9 @@ void UEngine::HandleTravelFailure(UWorld* InWorld, ETravelFailure::Type FailureT
 
 	UE_LOG(LogNet, Log, TEXT("TravelFailure: %s, Reason for Failure: '%s'"), ETravelFailure::ToString(FailureType), *ErrorString);
 
+	// Give the GameInstance a chance to handle the failure.
+	HandleTravelFailure_NotifyGameInstance(InWorld, FailureType);
+
 	ENetMode NetMode = GetNetMode(InWorld);
 
 	switch (FailureType)
@@ -7695,6 +7698,9 @@ void UEngine::HandleNetworkFailure(UWorld *World, UNetDriver *NetDriver, ENetwor
 			}
 		}
 
+		// Give the GameInstance a chance to handle the failure.
+		HandleNetworkFailure_NotifyGameInstance(World, NetDriver, FailureType);
+
 		ENetMode FailureNetMode = NetDriver->GetNetMode();	// NetMode of the driver that failed
 		bool bShouldTravel = true;
 
@@ -7728,6 +7734,44 @@ void UEngine::HandleNetworkFailure(UWorld *World, UNetDriver *NetDriver, ENetwor
 	}
 }
 
+void UEngine::HandleNetworkFailure_NotifyGameInstance(UWorld *World, UNetDriver *NetDriver, ENetworkFailure::Type FailureType)
+{
+	bool bIsServer = true;
+
+	if (NetDriver != nullptr)
+	{
+		bIsServer = NetDriver->GetNetMode() != NM_Client;
+	}
+
+	if (World != nullptr && World->GetGameInstance() != nullptr)
+	{
+		World->GetGameInstance()->HandleNetworkError(FailureType, bIsServer);
+	}
+	else
+	{
+		// Since the UWorld passed in might be null, as well as the NetDriver's UWorld,
+		// go through the world contexts until we find the one with this net driver.
+		for (auto& Context : WorldList)
+		{
+			if (Context.PendingNetGame != nullptr &&
+				Context.PendingNetGame->NetDriver == NetDriver &&
+				Context.OwningGameInstance != nullptr)
+			{
+				// Use the GameInstance from the current context.
+				Context.OwningGameInstance->HandleNetworkError(FailureType, bIsServer);
+			}
+		}
+	}
+}
+
+void UEngine::HandleTravelFailure_NotifyGameInstance(UWorld *World, ETravelFailure::Type FailureType)
+{
+	if (World != nullptr && World->GetGameInstance() != nullptr)
+	{
+		World->GetGameInstance()->HandleTravelError(FailureType);
+	}
+}
+
 void UEngine::SpawnServerActors(UWorld *World)
 {
 	for( int32 i=0; i < ServerActors.Num(); i++ )
