* refactored code, code pass passing down half radius where it shouldn't. Some code might rely on that but haven't seen anything yet. Could affect culling, shadow correctness or quality

diff --git a/Engine/Source/Runtime/Engine/Private/SpotLightComponent.cpp b/Engine/Source/Runtime/Engine/Private/SpotLightComponent.cpp
--- a/Engine/Source/Runtime/Engine/Private/SpotLightComponent.cpp
+++ b/Engine/Source/Runtime/Engine/Private/SpotLightComponent.cpp
@@ -127,11 +127,10 @@ public:
 		OutInitializer.WorldToLight = GetWorldToLight().RemoveTranslation();
 		OutInitializer.Scales = FVector(1.0f,InvTanOuterCone,InvTanOuterCone);
 		OutInitializer.FaceDirection = FVector(1,0,0);
-		//@todo - these half radius bounds are not correct, but quite a few things depend on it at this point
 		OutInitializer.SubjectBounds = FBoxSphereBounds(
-			GetLightToWorld().RemoveTranslation().TransformPosition(FVector(Radius / 2.0f,0,0)),
-			FVector(Radius/2.0f,Radius/2.0f,Radius/2.0f),
-			Radius / 2.0f
+			GetLightToWorld().RemoveTranslation().TransformPosition(FVector(Radius, 0, 0)),
+			FVector(Radius, Radius, Radius),
+			Radius
 			);
 		OutInitializer.WAxis = FVector4(0,0,1,0);
 		OutInitializer.MinLightW = 0.1f;
