Instead of marking objects in the shadow state as referenced, just make sure we are only comparing object pointers, which is safe even if they are GC'd

diff --git a/Engine/Source/Runtime/Engine/Private/RepLayout.cpp b/Engine/Source/Runtime/Engine/Private/RepLayout.cpp
--- a/Engine/Source/Runtime/Engine/Private/RepLayout.cpp
+++ b/Engine/Source/Runtime/Engine/Private/RepLayout.cpp
@@ -34,7 +34,18 @@ static FORCEINLINE bool CompareBool( const FRepLayoutCmd& Cmd, const void* A, co
 
 static FORCEINLINE bool CompareObject( const FRepLayoutCmd& Cmd, const void* A, const void* B )
 {
+#if 1
+	// Until UObjectPropertyBase::Identical is made safe for GC'd objects, we need to do it manually
+	// This saves us from having to add referenced objects during GC
+	UObjectPropertyBase * ObjProperty = CastChecked< UObjectPropertyBase>( Cmd.Property );
+
+	UObject* ObjectA = ObjProperty->GetObjectPropertyValue( A );
+	UObject* ObjectB = ObjProperty->GetObjectPropertyValue( B );
+
+	return ObjectA == ObjectB;
+#else
 	return Cmd.Property->Identical( A, B );
+#endif
 }
 
 template< typename T >
@@ -2546,11 +2557,13 @@ public:
 
 void FRepLayout::AddReferencedObjects( FRepState * RepState, UObject* ReferencingObject, FReferenceCollector& Collector ) const
 {
+#if 0
 	FAddReferencedObjectsImpl AddReferencedObjectsImpl( ReferencingObject, Collector );
 
 	FRepLayoutCmdIterator< FAddReferencedObjectsImpl, FMyCmdIteratorBaseStackState > Iterator( AddReferencedObjectsImpl, Cmds );
 
 	Iterator.Start( RepState );
+#endif
 }
 
 void FRepLayout::InitChangedTracker( FRepChangedPropertyTracker * ChangedTracker ) const
