Removed error-like looking on-screen runtime information about navmesh generation in progress #UE4

- Jira UE-4613

#code Dmitriy.Dyomin

diff --git a/Engine/Source/Runtime/Engine/Private/UnrealEngine.cpp b/Engine/Source/Runtime/Engine/Private/UnrealEngine.cpp
--- a/Engine/Source/Runtime/Engine/Private/UnrealEngine.cpp
+++ b/Engine/Source/Runtime/Engine/Private/UnrealEngine.cpp
@@ -6718,17 +6718,6 @@ void DrawStatsHUD( UWorld* World, FViewport* Viewport, FCanvas* Canvas, UCanvas*
 				MessageY += 20;
 			}
 
-			if (NavSys && NavSys->bBuildNavigationAtRuntime)
-			{
-				const int32 NumTasks = NavSys->GetNumRemainingBuildTasks();
-				if (NumTasks > 0)
-				{
-					SmallTextItem.Text = FText::FromString(FString::Printf(TEXT("Building Navigation (%d)"), NumTasks));
-					Canvas->DrawItem(SmallTextItem, FVector2D(MessageX, MessageY));
-					MessageY += 20;
-				}
-			}
-
 			if( World->bKismetScriptError )
 			{
 				SmallTextItem.Text = LOCTEXT("BlueprintInLevelHadCompileErrorMessage", "BLUEPRINT COMPILE ERROR" );
