Fixed compile error

diff --git a/Engine/Source/Runtime/Engine/Private/InstancedStaticMesh.cpp b/Engine/Source/Runtime/Engine/Private/InstancedStaticMesh.cpp
--- a/Engine/Source/Runtime/Engine/Private/InstancedStaticMesh.cpp
+++ b/Engine/Source/Runtime/Engine/Private/InstancedStaticMesh.cpp
@@ -906,7 +906,7 @@ public:
 
 	~FInstancedStaticMeshSceneProxy()
 	{
-		InstancedRenderData.ReleaseResources(GetScene( ), StaticMesh);
+		InstancedRenderData.ReleaseResources(&GetScene( ), StaticMesh);
 	}
 
 	// FPrimitiveSceneProxy interface.
