Convert more code to use new generalized FRepLayout cmd iterator + remove unused code

diff --git a/Engine/Source/Runtime/Engine/Private/DataChannel.cpp b/Engine/Source/Runtime/Engine/Private/DataChannel.cpp
--- a/Engine/Source/Runtime/Engine/Private/DataChannel.cpp
+++ b/Engine/Source/Runtime/Engine/Private/DataChannel.cpp
@@ -1982,13 +1982,6 @@ FString UActorChannel::Describe()
 void UActorChannel::AddReferencedObjects(UObject* InThis, FReferenceCollector& Collector)
 {
 	UActorChannel* This = CastChecked<UActorChannel>(InThis);	
-
-	// Let GC know that we're referencing UObject's in the replication shadow state
-	for ( auto MapIt = This->ReplicationMap.CreateIterator(); MapIt; ++MapIt )
-	{
-		MapIt.Value()->AddReferencedObjects( InThis, Collector );
-	}
-
 	Super::AddReferencedObjects( This, Collector );
 }
 
diff --git a/Engine/Source/Runtime/Engine/Private/DataReplication.cpp b/Engine/Source/Runtime/Engine/Private/DataReplication.cpp
--- a/Engine/Source/Runtime/Engine/Private/DataReplication.cpp
+++ b/Engine/Source/Runtime/Engine/Private/DataReplication.cpp
@@ -977,17 +977,6 @@ void FObjectReplicator::ForceRefreshUnreliableProperties()
 	bOpenAckCalled = true;
 }
 
-void FObjectReplicator::AddReferencedObjects( UObject* InThis, FReferenceCollector& Collector )
-{
-	if ( GetObject() == NULL )
-	{
-		UE_LOG( LogNet, Verbose, TEXT( "AddReferencedObjects: Object == NULL" ) );
-		return;
-	}
-
-	RepLayout->AddReferencedObjects( RepState, InThis, Collector );
-}
-
 void FObjectReplicator::PostSendBunch( FPacketIdRange & PacketRange, uint8 bReliable )
 {
 	if ( GetObject() == NULL )
diff --git a/Engine/Source/Runtime/Engine/Private/RepLayout.cpp b/Engine/Source/Runtime/Engine/Private/RepLayout.cpp
--- a/Engine/Source/Runtime/Engine/Private/RepLayout.cpp
+++ b/Engine/Source/Runtime/Engine/Private/RepLayout.cpp
@@ -208,40 +208,91 @@ static void SerializeReadWritePropertyChecksum( const FRepLayoutCmd& Cmd, const
 	}
 }
 
-class FMyCmdIteratorBaseStackState
+#define INIT_STACK( TStack )							\
+	void InitStack( TStack& StackState )				\
+
+#define SHOULD_PROCESS_NEXT_CMD()						\
+	bool ShouldProcessNextCmd()							\
+
+#define PROCESS_ARRAY_CMD( TStack )						\
+	void ProcessArrayCmd_r(								\
+	TStack&							PrevStackState,		\
+	TStack&							StackState,			\
+	const FRepLayoutCmd&			Cmd,				\
+	const int32						CmdIndex,			\
+	uint8* RESTRICT					ShadowData,			\
+	uint8* RESTRICT					Data )				\
+
+
+#define PROCESS_CMD( TStack )							\
+	void ProcessCmd(									\
+	TStack&							StackState,			\
+	const FRepLayoutCmd&			Cmd,				\
+	const int32						CmdIndex,			\
+	uint8* RESTRICT					ShadowData,			\
+	uint8* RESTRICT					Data )				\
+
+class FCmdIteratorBaseStackState
 {
 public:
-	FMyCmdIteratorBaseStackState( FScriptArray* InArray ) : Array( InArray ), ArrayNum( InArray->Num() ) {}
+	FCmdIteratorBaseStackState( const int32 InCmdStart, const int32 InCmdEnd, FScriptArray*	InShadowArray, FScriptArray* InDataArray, uint8* RESTRICT InShadowBaseData, uint8* RESTRICT	InBaseData ) : 
+		CmdStart( InCmdStart ),
+		CmdEnd( InCmdEnd ),
+		ShadowArray( InShadowArray ),
+		DataArray( InDataArray ),
+		ShadowBaseData( InShadowBaseData ),
+		BaseData( InBaseData )
+	{
+	}
+
+	const int32			CmdStart; 
+	const int32			CmdEnd;
 
-	FScriptArray*	Array;
-	const uint16	ArrayNum;
+	FScriptArray*		ShadowArray;
+	FScriptArray*		DataArray;
+
+	uint8* RESTRICT		ShadowBaseData;
+	uint8* RESTRICT		BaseData;
 };
 
+// This uses the "Curiously recurring template pattern" (CRTP) ideas
 template< typename TImpl, typename TStackState >
 class FRepLayoutCmdIterator
 {
 public:
-	FRepLayoutCmdIterator( TImpl& InImpl, const TArray< FRepLayoutCmd > & InCmds ) : Impl( InImpl ), Cmds( InCmds ) {}
+	FRepLayoutCmdIterator( const TArray< FRepParentCmd >& InParents, const TArray< FRepLayoutCmd >& InCmds ) : Parents( InParents ), Cmds( InCmds ) {}
 
-	void IterateCmdArray_r( const FRepLayoutCmd& Cmd, const int32 CmdIndex, const uint8* RESTRICT ShadowData )
+	void ProcessDataArrayElements_r( TStackState& StackState, const FRepLayoutCmd& Cmd )
 	{
-		TStackState StackState( (FScriptArray *)ShadowData );
-
-		uint8* LocalData = (uint8*)StackState.Array->GetData();
-
-		Impl.PreArray( StackState );
+		for ( int32 i = 0; i < StackState.DataArray->Num(); i++ )
+		{
+			const int32 ElementOffset = i * Cmd.ElementSize;
+			ProcessCmds_r( StackState, StackState.ShadowBaseData + ElementOffset, StackState.BaseData + ElementOffset );
+		}
+	}
 
-		for ( int32 i = 0; i < StackState.ArrayNum; i++ )
+	void ProcessShadowArrayElements_r( TStackState& StackState, const FRepLayoutCmd& Cmd )
+	{
+		for ( int32 i = 0; i < StackState.ShadowArray->Num(); i++ )
 		{
-			ProcessCmd_r( CmdIndex + 1, Cmd.EndCmd - 1, LocalData + i * Cmd.ElementSize );
+			const int32 ElementOffset = i * Cmd.ElementSize;
+			ProcessCmds_r( StackState, StackState.ShadowBaseData + ElementOffset, StackState.BaseData + ElementOffset );
 		}
+	}
+
+	void ProcessArrayCmd_r( TStackState & PrevStackState, const FRepLayoutCmd& Cmd, const int32 CmdIndex, uint8* RESTRICT ShadowData, uint8* RESTRICT Data )
+	{
+		FScriptArray* ShadowArray	= (FScriptArray*)ShadowData;
+		FScriptArray* DataArray		= (FScriptArray*)Data;
+
+		TStackState StackState( CmdIndex + 1, Cmd.EndCmd - 1, ShadowArray, DataArray, (uint8*)ShadowArray->GetData(), (uint8*)DataArray->GetData() );
 
-		Impl.PostArray( StackState );
+		static_cast< TImpl* >( this )->ProcessArrayCmd_r( PrevStackState, StackState, Cmd, CmdIndex, ShadowData, Data );
 	}
 
-	void ProcessCmd_r( const int32 CmdStart, const int32 CmdEnd, const uint8* RESTRICT ShadowData )
+	void ProcessCmds_r( TStackState& StackState, uint8* RESTRICT ShadowData, uint8* RESTRICT Data )
 	{
-		for ( int32 CmdIndex = CmdStart; CmdIndex < CmdEnd; CmdIndex++ )
+		for ( int32 CmdIndex = StackState.CmdStart; CmdIndex < StackState.CmdEnd; CmdIndex++ )
 		{
 			const FRepLayoutCmd& Cmd = Cmds[ CmdIndex ];
 
@@ -249,22 +300,33 @@ public:
 
 			if ( Cmd.Type == REPCMD_DynamicArray )
 			{
-				IterateCmdArray_r( Cmd, CmdIndex, ShadowData + Cmd.Offset );
+				if ( static_cast< TImpl* >( this )->ShouldProcessNextCmd() )
+				{
+					ProcessArrayCmd_r( StackState, Cmd, CmdIndex, ShadowData + Cmd.Offset, Data + Cmd.Offset );
+				}
 				CmdIndex = Cmd.EndCmd - 1;	// Jump past children of this array (-1 for ++ in for loop)
-				continue;
 			}
-
-			Impl.ProcessCmd( Cmd, ShadowData + Cmd.Offset );
+			else
+			{
+				if ( static_cast< TImpl* >( this )->ShouldProcessNextCmd() )
+				{
+					static_cast< TImpl* >( this )->ProcessCmd( StackState, Cmd, CmdIndex, ShadowData, Data );
+				}
+			}
 		}
 	}
 
-	void Start( FRepState * RepState )
+	void ProcessCmds( FRepState* RepState, uint8* RESTRICT Data )
 	{
-		ProcessCmd_r( 0, Cmds.Num() - 1, (uint8*)RepState->StaticBuffer.GetData() );
+		TStackState StackState( 0, Cmds.Num() - 1, NULL, NULL, (uint8*)RepState->StaticBuffer.GetData(), Data );
+
+		static_cast< TImpl* >( this )->InitStack( StackState );
+
+		ProcessCmds_r( StackState, (uint8*)RepState->StaticBuffer.GetData(), Data );
 	}
 
-	TImpl&							Impl;
-	const TArray< FRepLayoutCmd > & Cmds;
+	const TArray< FRepParentCmd >&	Parents;
+	const TArray< FRepLayoutCmd >&	Cmds;
 };
 
 uint16 FRepLayout::CompareProperties_r(
@@ -657,7 +719,7 @@ bool FRepLayout::ReplicateProperties(
 			{
 				TArray< uint16 > Temp = Changed;
 				Changed.Empty();
-				MergeDirtyList( (void*)Data, Temp, RepState->PreOpenAckHistory[i].Changed, Changed );
+				MergeDirtyList( RepState, (void*)Data, Temp, RepState->PreOpenAckHistory[i].Changed, Changed );
 			}
 			RepState->PreOpenAckHistory.Empty();
 		}
@@ -727,7 +789,7 @@ void FRepLayout::UpdateChangelistHistory( FRepState * RepState, UClass * ObjectC
 				check( OutMerged != NULL );
 				TArray< uint16 > Temp = *OutMerged;
 				OutMerged->Empty();
-				MergeDirtyList( (void*)Data, Temp, HistoryItem.Changed, *OutMerged );
+				MergeDirtyList( RepState, (void*)Data, Temp, HistoryItem.Changed, *OutMerged );
 				HistoryItem.Changed.Empty();
 
 #ifdef SANITY_CHECK_MERGES
@@ -1067,262 +1129,211 @@ void FRepLayout::SendProperties(
 	WritePropertyHandle( Writer, 0, bDoChecksum );
 }
 
-static void ReadNextHandle( FRepReaderState & ReaderState )
+class FReceivedPropertiesStackState : public FCmdIteratorBaseStackState
 {
-	ReaderState.Bunch.SerializeIntPacked( ReaderState.WaitingHandle );
-
-#ifdef ENABLE_PROPERTY_CHECKSUMS
-	if ( ReaderState.bDoChecksum )
-	{
-		SerializeGenericChecksum( ReaderState.Bunch );
-	}
-#endif
-}
-
-static bool ShouldReadProperty( FRepReaderState & ReaderState )
-{
-	ReaderState.CurrentHandle++;
-
-	if ( ReaderState.CurrentHandle == ReaderState.WaitingHandle )
-	{
-		check( ReaderState.WaitingHandle != 0 );
-		return true;
-	}
+public:
+	FReceivedPropertiesStackState( const int32 InCmdStart, const int32 InCmdEnd, FScriptArray*	InShadowArray, FScriptArray* InDataArray, uint8* RESTRICT InShadowBaseData, uint8* RESTRICT	InBaseData ) : 
+		UnmappedGuids( NULL ),
+		FCmdIteratorBaseStackState( InCmdStart, InCmdEnd, InShadowArray, InDataArray, InShadowBaseData, InBaseData ) 
+	{}
 
-	return false;
-}
+	FUnmappedGuidMgr* UnmappedGuids;
+};
 
-bool FRepLayout::ReadProperty( 
-	FRepReaderState &		ReaderState, 
-	FUnmappedGuidMgr *		UnmappedGuids,
-	const int32				AbsOffset,
-	const FRepLayoutCmd &	Cmd, 
-	const int32				CmdIndex, 
-	uint8* RESTRICT		StoredData, 
-	uint8* RESTRICT		Data ) const
+class FReceivePropertiesImpl : public FRepLayoutCmdIterator< FReceivePropertiesImpl, FReceivedPropertiesStackState >
 {
-	if ( !ShouldReadProperty( ReaderState ) )
-	{
-		// This property didn't change, nothing to read
-		return true;
-	}
-
-	const FRepParentCmd& Parent = Parents[Cmd.ParentIndex];
-
-	// This swaps Role/RemoteRole as we write it
-	const FRepLayoutCmd& SwappedCmd = Parent.RoleSwapIndex != -1 ? Cmds[Parents[Parent.RoleSwapIndex].CmdStart] : Cmd;
-
-	// Let package map know we want to track and know about any guids that are unmapped during the serialize call
-	ReaderState.Bunch.PackageMap->ResetTrackedUnmappedGuids( true );
-
-	// Remember where we started reading from, so that if we have unmapped properties, we can re-deserialize from this data later
-	FBitReaderMark Mark( ReaderState.Bunch );
+public:
+	FReceivePropertiesImpl( FNetBitReader & InBunch, FRepState * InRepState, bool bInDoChecksum, const TArray< FRepParentCmd >& InParents, const TArray< FRepLayoutCmd >& InCmds ) :
+		WaitingHandle( 0 ), 
+		CurrentHandle( 0 ), 
+		Bunch( InBunch ),
+		RepState( InRepState ),
+		bDoChecksum( bInDoChecksum ),
+		bHasUnmapped( false ),
+		FRepLayoutCmdIterator( InParents, InCmds ) 
+	{}
 
-	if ( Parent.Property->HasAnyPropertyFlags( CPF_RepNotify ) )
+	void ReadNextHandle()
 	{
-		// Copy current value over so we can check to see if it changed
-		StoreProperty( Cmd, StoredData + Cmd.Offset, Data + SwappedCmd.Offset );
-
-		// Read the property
-		Cmd.Property->NetSerializeItem( ReaderState.Bunch, ReaderState.Bunch.PackageMap, Data + SwappedCmd.Offset );
+		Bunch.SerializeIntPacked( WaitingHandle );
 
-		// Check to see if this property changed
-		if (Parent.RepNotifyCondition == REPNOTIFY_Always || !PropertiesAreIdentical(Cmd, StoredData + Cmd.Offset, Data + SwappedCmd.Offset))
+#ifdef ENABLE_PROPERTY_CHECKSUMS
+		if ( bDoChecksum )
 		{
-			ReaderState.RepState->RepNotifies.AddUnique( Parent.Property );
+			SerializeGenericChecksum( Bunch );
 		}
-	} 
-	else
-	{
-		Cmd.Property->NetSerializeItem( ReaderState.Bunch, ReaderState.Bunch.PackageMap, Data + SwappedCmd.Offset );
+#endif
 	}
-	
-	const TArray< FNetworkGUID > & TrackedUnmappedGuids = ReaderState.Bunch.PackageMap->GetTrackedUnmappedGuids();
 
-	if ( TrackedUnmappedGuids.Num() > 0 )
+	INIT_STACK( FReceivedPropertiesStackState )
 	{
-		// If we have an unmapped guid, we need to remember it, so we can fix up this object pointer when it finally arrives at a other time
-		const int32 LocalAbsOffset = AbsOffset + SwappedCmd.Offset;
-		
-		// Note - If we already have an existing entry for this offset, we will replace it with this most recent data
-		UnmappedGuids->Map.Add( LocalAbsOffset, FUnmappedGuidMgrElement( ReaderState.Bunch, Mark, TrackedUnmappedGuids, Cmd.ParentIndex, CmdIndex ) );
+		StackState.UnmappedGuids = &RepState->UnmappedGuids;
+	}
 
-		UE_LOG( LogNet, Verbose, TEXT( "ADDED unmapped property: Offset: %i, Name: %s"), LocalAbsOffset, *Cmd.Property->GetName() );
+	SHOULD_PROCESS_NEXT_CMD()
+	{
+		CurrentHandle++;
 
-		// List all the guids that were unmapped
-		for ( int32 i = 0; i < TrackedUnmappedGuids.Num(); i++ )
+		if ( CurrentHandle == WaitingHandle )
 		{
-			UE_LOG( LogNet, Verbose, TEXT( "  Guid: %s"), *TrackedUnmappedGuids[i].ToString() );
+			check( WaitingHandle != 0 );
+			return true;
 		}
 
-		ReaderState.bHasUnmapped = true;
+		return false;
 	}
-	else
-	{
-		// If we don't have any unmapped guids, then make sure to remove the entry so we don't serialize old data when we update unmapped objects
-		const int32 LocalAbsOffset = AbsOffset + SwappedCmd.Offset;
 
-		// Remove the entry
-		UnmappedGuids->Map.Remove( LocalAbsOffset );
-	}
+	PROCESS_ARRAY_CMD( FReceivedPropertiesStackState )
+	{
+		// Read array size
+		uint16 ArrayNum = 0;
+		Bunch << ArrayNum;
 
-	// Stop tracking unmapped objects
-	ReaderState.Bunch.PackageMap->ResetTrackedUnmappedGuids( false );
+		// Read the next property handle
+		ReadNextHandle();
 
-#ifdef ENABLE_PROPERTY_CHECKSUMS
-	if ( ReaderState.bDoChecksum )
-	{
-		SerializeReadWritePropertyChecksum( Cmd, CmdIndex, Data + SwappedCmd.Offset, ReaderState.Bunch );
-	}
-#endif
+		const int32 AbsOffset = Data - PrevStackState.BaseData;
 
-	// Store the property we just received FIXME: Can't do this because of rep notifies. This needs to be done for super checksums though, FIXME!!!
-	//StoreProperty( Cmd, StoredData + Cmd.Offset, Data + SwappedCmd.Offset );
+		// Since we don't know yet if something under us could be unmapped, go ahead and allocate an array container now
+		FUnmappedGuidMgrElement * NewArrayElement = PrevStackState.UnmappedGuids->Map.Find( AbsOffset );
 
-	// Read the next property handle
-	ReadNextHandle( ReaderState );
+		if ( NewArrayElement == NULL )
+		{
+			NewArrayElement = &PrevStackState.UnmappedGuids->Map.FindOrAdd( AbsOffset );
+			NewArrayElement->Array = new FUnmappedGuidMgr;
+			NewArrayElement->ParentIndex = Cmd.ParentIndex;
+			NewArrayElement->CmdIndex = CmdIndex;
+		}
 
-	return true;
-}
+		check( NewArrayElement != NULL );
+		check( NewArrayElement->ParentIndex == Cmd.ParentIndex );
+		check( NewArrayElement->CmdIndex == CmdIndex );
 
-bool FRepLayout::ReceiveProperties_AnyArray_r( 
-	FRepReaderState &	ReaderState, 
-	FUnmappedGuidMgr *	UnmappedGuids,
-	const int32			AbsOffset,
-	const int32			ArrayNum, 
-	const int32			ElementSize, 
-	const int32			CmdIndex, 
-	uint8* RESTRICT	StoredData, 
-	uint8 *	RESTRICT	Data ) const
-{
-	const int32 CmdStart	= CmdIndex + 1;
-	const int32 CmdEnd		= Cmds[CmdIndex].EndCmd - 1;
+		StackState.UnmappedGuids = NewArrayElement->Array;
 
-	for ( int32 i = 0; i < ArrayNum; i++ )
-	{
-		const int32 ElementOffset = i * ElementSize;
+		const FRepParentCmd& Parent = Parents[Cmd.ParentIndex];
 
-		if ( !ReceiveProperties_r( ReaderState, UnmappedGuids, AbsOffset + ElementOffset, CmdStart, CmdEnd, StoredData + ElementOffset, Data + ElementOffset ) )
+		if ( ( StackState.DataArray->Num() != ArrayNum || Parent.RepNotifyCondition == REPNOTIFY_Always ) && Parent.Property->HasAnyPropertyFlags( CPF_RepNotify ) )
 		{
-			return false;
+			RepState->RepNotifies.AddUnique( Parent.Property );
 		}
-	}
-
-	return true;
-}
-
-bool FRepLayout::ReceiveProperties_DynamicArray_r( 
-	FRepReaderState &		ReaderState, 
-	FUnmappedGuidMgr *		UnmappedGuids,
-	const int32				AbsOffset,
-	const FRepLayoutCmd &	Cmd, 
-	const int32				CmdIndex, 
-	uint8* RESTRICT		StoredData, 
-	uint8* RESTRICT		Data ) const
-{
-	if ( !ShouldReadProperty( ReaderState ) )
-	{
-		// This array didn't change, nothing to read
-		return true;
-	}
 
-	// Read array size
-	uint16 ArrayNum = 0;
-	ReaderState.Bunch << ArrayNum;
+		check( CastChecked< UArrayProperty >( Cmd.Property ) != NULL );
 
-	// Read the next property handle
-	ReadNextHandle( ReaderState );
+		// Resize arrays if needed
+		FScriptArrayHelper ArrayHelper( (UArrayProperty *)Cmd.Property, Data );
+		ArrayHelper.Resize( ArrayNum );
 
-	FScriptArray * Array = (FScriptArray *)Data;
-	FScriptArray * StoredArray = (FScriptArray *)StoredData;
+		FScriptArrayHelper ShadowArrayHelper( (UArrayProperty *)Cmd.Property, ShadowData );
+		ShadowArrayHelper.Resize( ArrayNum );
 
-	// Since we don't know yet if something under us could be unmapped, go ahead and allocate an array container now
-	FUnmappedGuidMgrElement * NewArrayElement = UnmappedGuids->Map.Find( AbsOffset );
+		// Re-compute the base data values since they could have changed after the resize above
+		StackState.BaseData			= (uint8*)StackState.DataArray->GetData();
+		StackState.ShadowBaseData	= (uint8*)StackState.ShadowArray->GetData();
 
-	if ( NewArrayElement == NULL )
-	{
-		NewArrayElement = &UnmappedGuids->Map.FindOrAdd( AbsOffset );
-		NewArrayElement->Array = new FUnmappedGuidMgr;
-		NewArrayElement->ParentIndex = Cmd.ParentIndex;
-		NewArrayElement->CmdIndex = CmdIndex;
-	}
+		// Save the old handle so we can restore it when we pop out of the array
+		const uint16 OldHandle = CurrentHandle;
 
-	check( NewArrayElement != NULL );
-	check( NewArrayElement->ParentIndex == Cmd.ParentIndex );
-	check( NewArrayElement->CmdIndex == CmdIndex );
+		// Array children handles are always relative to their immediate parent
+		CurrentHandle = 0;
 
-	UnmappedGuids = NewArrayElement->Array;
+		// Loop over array
+		ProcessDataArrayElements_r( StackState, Cmd );
 
-	const FRepParentCmd& Parent = Parents[Cmd.ParentIndex];
+		// Restore the current handle to what it was before we processed this array
+		CurrentHandle = OldHandle;
 
-	if ( (Array->Num() != ArrayNum || Parent.RepNotifyCondition == REPNOTIFY_Always) && Parent.Property->HasAnyPropertyFlags( CPF_RepNotify ) )
-	{
-		ReaderState.RepState->RepNotifies.AddUnique( Parent.Property );
+		// We should be waiting on the NULL terminator handle at this point
+		check( WaitingHandle == 0 );
+		ReadNextHandle();
 	}
 
-	// Set the array to the correct size
-	FScriptArrayHelper ArrayHelper( (UArrayProperty *)Cmd.Property, Data );
-	ArrayHelper.Resize( ArrayNum );
+	PROCESS_CMD( FReceivedPropertiesStackState )
+	{
+		check( StackState.UnmappedGuids != NULL );
 
-	FScriptArrayHelper StoredArrayHelper( (UArrayProperty *)Cmd.Property, StoredData );
-	StoredArrayHelper.Resize( ArrayNum );
+		const FRepParentCmd& Parent = Parents[Cmd.ParentIndex];
 
-	Data		= (uint8*)Array->GetData();
-	StoredData	= (uint8*)StoredArray->GetData();
+		// This swaps Role/RemoteRole as we write it
+		const FRepLayoutCmd& SwappedCmd = Parent.RoleSwapIndex != -1 ? Cmds[Parents[Parent.RoleSwapIndex].CmdStart] : Cmd;
 
-	const uint16 OldHandle = ReaderState.CurrentHandle;
+		// Let package map know we want to track and know about any guids that are unmapped during the serialize call
+		Bunch.PackageMap->ResetTrackedUnmappedGuids( true );
 
-	// Array children handles are always relative to their immediate parent
-	ReaderState.CurrentHandle = 0;
+		// Remember where we started reading from, so that if we have unmapped properties, we can re-deserialize from this data later
+		FBitReaderMark Mark( Bunch );
 
-	// Read any changed properties into the array
-	if ( !ReceiveProperties_AnyArray_r( ReaderState, UnmappedGuids, 0, ArrayNum, Cmd.ElementSize, CmdIndex, StoredData, Data ) )
-	{
-		return false;
-	}
+		if ( Parent.Property->HasAnyPropertyFlags( CPF_RepNotify ) )
+		{
+			// Copy current value over so we can check to see if it changed
+			StoreProperty( Cmd, ShadowData + Cmd.Offset, Data + SwappedCmd.Offset );
 
-	ReaderState.CurrentHandle = OldHandle;
+			// Read the property
+			Cmd.Property->NetSerializeItem( Bunch, Bunch.PackageMap, Data + SwappedCmd.Offset );
 
-	// We should be waiting on the NULL terminator handle at this point
-	check( ReaderState.WaitingHandle == 0 );
-	ReadNextHandle( ReaderState );
+			// Check to see if this property changed
+			if ( Parent.RepNotifyCondition == REPNOTIFY_Always || !PropertiesAreIdentical( Cmd, ShadowData + Cmd.Offset, Data + SwappedCmd.Offset ) )
+			{
+				RepState->RepNotifies.AddUnique( Parent.Property );
+			}
+		} 
+		else
+		{
+			Cmd.Property->NetSerializeItem( Bunch, Bunch.PackageMap, Data + SwappedCmd.Offset );
+		}
 
-	return true;
-}
+		const TArray< FNetworkGUID > & TrackedUnmappedGuids = Bunch.PackageMap->GetTrackedUnmappedGuids();
 
-bool FRepLayout::ReceiveProperties_r( 
-	FRepReaderState &	ReaderState, 
-	FUnmappedGuidMgr *	UnmappedGuids,
-	const int32			AbsOffset,
-	const int32			CmdStart, 
-	const int32			CmdEnd, 
-	uint8* RESTRICT	StoredData, 
-	uint8* RESTRICT	Data ) const
-{
-	for ( int32 CmdIndex = CmdStart; CmdIndex < CmdEnd; CmdIndex++ )
-	{
-		const FRepLayoutCmd& Cmd = Cmds[ CmdIndex ];
+		const int32 AbsOffset = Data - StackState.BaseData;
 
-		check( Cmd.Type != REPCMD_Return );
+		const int32 AbsCmdOffset = AbsOffset + SwappedCmd.Offset;
 
-		if ( Cmd.Type == REPCMD_DynamicArray )
+		if ( TrackedUnmappedGuids.Num() > 0 )
 		{
-			if ( !ReceiveProperties_DynamicArray_r( ReaderState, UnmappedGuids, AbsOffset + Cmd.Offset, Cmd, CmdIndex, StoredData + Cmd.Offset, Data + Cmd.Offset ) )
+			// If we have an unmapped guid, we need to remember it, so we can fix up this object pointer when it finally arrives at a other time
+			// Note - If we already have an existing entry for this offset, we will replace it with this most recent data
+			StackState.UnmappedGuids->Map.Add( AbsCmdOffset, FUnmappedGuidMgrElement( Bunch, Mark, TrackedUnmappedGuids, Cmd.ParentIndex, CmdIndex ) );
+
+			UE_LOG( LogNet, Verbose, TEXT( "ADDED unmapped property: Offset: %i, Name: %s"), AbsCmdOffset, *Cmd.Property->GetName() );
+
+			// List all the guids that were unmapped
+			for ( int32 i = 0; i < TrackedUnmappedGuids.Num(); i++ )
 			{
-				return false;
+				UE_LOG( LogNet, Verbose, TEXT( "  Guid: %s"), *TrackedUnmappedGuids[i].ToString() );
 			}
-			CmdIndex = Cmd.EndCmd - 1;	// Jump past children of this array ( -1 for the ++ in the for loop)
-			continue;
+
+			bHasUnmapped = true;
+		}
+		else
+		{
+			// If we don't have any unmapped guids, then make sure to remove the entry so we don't serialize old data when we update unmapped objects
+			StackState.UnmappedGuids->Map.Remove( AbsCmdOffset );
 		}
 
-		if ( !ReadProperty( ReaderState, UnmappedGuids, AbsOffset, Cmd, CmdIndex, StoredData, Data ) )
+		// Stop tracking unmapped objects
+		Bunch.PackageMap->ResetTrackedUnmappedGuids( false );
+
+#ifdef ENABLE_PROPERTY_CHECKSUMS
+		if ( bDoChecksum )
 		{
-			return false;
+			SerializeReadWritePropertyChecksum( Cmd, CmdIndex, Data + SwappedCmd.Offset, Bunch );
 		}
+#endif
+
+		// Store the property we just received FIXME: Can't do this because of rep notifies. This needs to be done for super checksums though, FIXME!!!
+		//StoreProperty( Cmd, StoredData + Cmd.Offset, Data + SwappedCmd.Offset );
+
+		// Read the next property handle
+		ReadNextHandle();
 	}
 
-	return true;
-}
+	uint32					WaitingHandle;
+	uint32					CurrentHandle;
+	FNetBitReader &			Bunch;
+	FRepState *				RepState;
+	bool					bDoChecksum;
+	bool					bHasUnmapped;
+};
 
 bool FRepLayout::ReceiveProperties( UClass * InObjectClass, FRepState * RESTRICT RepState, void* RESTRICT Data, FNetBitReader & InBunch, bool & bOutHasUnmapped ) const
 {
@@ -1336,19 +1347,16 @@ bool FRepLayout::ReceiveProperties( UClass * InObjectClass, FRepState * RESTRICT
 
 	bOutHasUnmapped = false;
 
-	FRepReaderState ReaderState( InBunch, RepState, bDoChecksum );
+	FReceivePropertiesImpl ReceivePropertiesImpl( InBunch, RepState, bDoChecksum, Parents, Cmds );
 
 	// Read first handle
-	ReadNextHandle( ReaderState );
+	ReceivePropertiesImpl.ReadNextHandle();
 
 	// Read all properties
-	if ( !ReceiveProperties_r( ReaderState, &RepState->UnmappedGuids, 0, 0, Cmds.Num() - 1, RepState->StaticBuffer.GetData(), (uint8*)Data ) )
-	{
-		return false;
-	}
+	ReceivePropertiesImpl.ProcessCmds( RepState, (uint8*)Data );
 
 	// Make sure we're waiting on the last NULL terminator
-	if ( ReaderState.WaitingHandle != 0 )
+	if ( ReceivePropertiesImpl.WaitingHandle != 0 )
 	{
 		UE_LOG( LogNet, Warning, TEXT( "Read out of sync." ) );
 		return false;
@@ -1361,7 +1369,7 @@ bool FRepLayout::ReceiveProperties( UClass * InObjectClass, FRepState * RESTRICT
 	}
 #endif
 
-	bOutHasUnmapped = ReaderState.bHasUnmapped;
+	bOutHasUnmapped = ReceivePropertiesImpl.bHasUnmapped;
 
 	return true;
 }
@@ -1586,148 +1594,161 @@ uint32 FRepLayout::GenerateChecksum( const FRepState* RepState ) const
 	return FCrc::MemCrc32( Writer.GetData(), Writer.GetNumBytes(), 0 );
 }
 
-void FRepLayout::MergeDirtyList_AnyArray_r( FMergeDirtyListState & MergeState, const FRepLayoutCmd& Cmd, const int32 ArrayNum, const int32 ElementSize, const int32 CmdIndex, const uint8* RESTRICT Data ) const
-{
-	for ( int32 i = 0; i < ArrayNum; i++ )
-	{
-		MergeDirtyList_r( MergeState, CmdIndex + 1, Cmd.EndCmd - 1, Data + i * ElementSize );
-	}
-}
-
-void FRepLayout::MergeDirtyList_DynamicArray_r( FMergeDirtyListState & MergeState, const FRepLayoutCmd& Cmd, const int32 CmdIndex, const uint8* RESTRICT Data ) const
+class FMergeDirtyListImpl : public FRepLayoutCmdIterator< FMergeDirtyListImpl, FCmdIteratorBaseStackState >
 {
-	// At least one of the list should be valid to be here
-	check( MergeState.DirtyValid1 || MergeState.DirtyValid2 );
+public:
+	FMergeDirtyListImpl( const TArray< uint16 > & InDirty1, const TArray< uint16 > & InDirty2, TArray< uint16 > & OutMergedDirty, const TArray< FRepParentCmd >& InParents, const TArray< FRepLayoutCmd >& InCmds ) : 
+		DirtyList1( InDirty1 ), 
+		DirtyList2( InDirty2 ), 
+		DirtyListIndex1( 0 ),
+		DirtyListIndex2( 0 ),
+		Handle( 0 ),
+		bDirtyValid1( true ),
+		bDirtyValid2( true ),
+		MergedDirtyList( OutMergedDirty ),
+		bLastDirty1Matches( false ), 
+		bLastDirty2Matches( false ),
+		FRepLayoutCmdIterator( InParents, InCmds )
+	{}
+
+	INIT_STACK( FCmdIteratorBaseStackState ) { }
+
+	SHOULD_PROCESS_NEXT_CMD() 
+	{ 
+		Handle++;
 
-	MergeState.Handle++;
+		bLastDirty1Matches = bDirtyValid1 && DirtyList1[DirtyListIndex1] == Handle;
+		bLastDirty2Matches = bDirtyValid2 && DirtyList2[DirtyListIndex2] == Handle;
 
-	const bool Dirty1Matches = MergeState.DirtyValid1 && MergeState.DirtyList1[MergeState.DirtyListIndex1] == MergeState.Handle;
-	const bool Dirty2Matches = MergeState.DirtyValid2 && MergeState.DirtyList2[MergeState.DirtyListIndex2] == MergeState.Handle;
+		return bLastDirty1Matches || bLastDirty2Matches;
+	}
 
-	if ( !Dirty1Matches && !Dirty2Matches )
+	PROCESS_ARRAY_CMD( FCmdIteratorBaseStackState ) 
 	{
-		// Neither match, we're done with this branch
-		return;
-	}
+		// At least one of the list should be valid to be here
+		check( bDirtyValid1 || bDirtyValid2 );
+		check( bLastDirty1Matches || bLastDirty2Matches )
 
-	// This will be a new merged dirty entry
-	MergeState.MergedDirtyList.Add( MergeState.Handle );
+		const bool bDirty1Matches = bLastDirty1Matches;
+		const bool bDirty2Matches = bLastDirty2Matches;
 
-	const int32 RememberIndex = MergeState.MergedDirtyList.AddUninitialized();
-	const int32 OldMergedListSize = MergeState.MergedDirtyList.Num();
+		// This will be a new merged dirty entry
+		MergedDirtyList.Add( Handle );
 
-	// Advance the matching dirty lists
-	if ( Dirty1Matches )
-	{
-		MergeState.DirtyListIndex1++;
-	}
+		const int32 OriginalMergedDirtyListIndex = MergedDirtyList.AddUninitialized();
+		check( OriginalMergedDirtyListIndex == MergedDirtyList.Num() - 1 );
 
-	if ( Dirty2Matches )
-	{
-		MergeState.DirtyListIndex2++;
-	}
+		// Advance the matching dirty lists
+		if ( bLastDirty1Matches )
+		{
+			DirtyListIndex1++;
+		}
 
-	// Remember valid dirty lists
-	const bool OldDirtyValid1 = MergeState.DirtyValid1;
-	const bool OldDirtyValid2 = MergeState.DirtyValid2;
+		if ( bLastDirty2Matches )
+		{
+			DirtyListIndex2++;
+		}
 
-	// Update which lists are still valid from this point
-	MergeState.DirtyValid1 = Dirty1Matches;
-	MergeState.DirtyValid2 = Dirty2Matches;
+		// Remember valid dirty lists
+		const bool bOldDirtyValid1 = bDirtyValid1;
+		const bool bOldDirtyValid2 = bDirtyValid2;
 
-	const int32 JumpToIndex1 = Dirty1Matches ? MergeState.DirtyList1[MergeState.DirtyListIndex1++] : -1;
-	const int32 JumpToIndex2 = Dirty2Matches ? MergeState.DirtyList2[MergeState.DirtyListIndex2++] : -1;
+		// Update which lists are still valid from this point
+		bDirtyValid1 = bLastDirty1Matches;
+		bDirtyValid2 = bLastDirty2Matches;
 
-	const int32 OldDirtyListIndex1 = MergeState.DirtyListIndex1;
-	const int32 OldDirtyListIndex2 = MergeState.DirtyListIndex2;
+		const int32 JumpToIndex1 = bLastDirty1Matches ? DirtyList1[DirtyListIndex1++] : -1;
+		const int32 JumpToIndex2 = bLastDirty2Matches ? DirtyList2[DirtyListIndex2++] : -1;
 
-	FScriptArray * Array = (FScriptArray *)Data;
+		const int32 OldDirtyListIndex1 = DirtyListIndex1;
+		const int32 OldDirtyListIndex2 = DirtyListIndex2;
 
-	const int32 OldHandle = MergeState.Handle;
-	MergeState.Handle = 0;
-	MergeDirtyList_AnyArray_r( MergeState, Cmd, Array->Num(), Cmd.ElementSize, CmdIndex, (uint8*)Array->GetData() );
-	MergeState.Handle = OldHandle;
+		const int32 OldHandle = Handle;
+		Handle = 0;
 
-	if ( Dirty1Matches )
-	{
-		check( MergeState.DirtyListIndex1 - OldDirtyListIndex1 <= JumpToIndex1 );
-		MergeState.DirtyListIndex1 = OldDirtyListIndex1 + JumpToIndex1;
-		check( MergeState.DirtyList1[MergeState.DirtyListIndex1] == 0 );
-		MergeState.DirtyListIndex1++;
-	}
+		// Process the array elements
+		ProcessDataArrayElements_r( StackState, Cmd );
 
-	if ( Dirty2Matches )
-	{
-		check( MergeState.DirtyListIndex2 - OldDirtyListIndex2 <= JumpToIndex2 );
-		MergeState.DirtyListIndex2 = OldDirtyListIndex2 + JumpToIndex2;
-		check( MergeState.DirtyList2[MergeState.DirtyListIndex2] == 0 );
-		MergeState.DirtyListIndex2++;
-	}
+		// Restore the handle
+		Handle = OldHandle;
 
-	MergeState.DirtyValid1 = OldDirtyValid1;
-	MergeState.DirtyValid2 = OldDirtyValid2;
+		if ( bDirty1Matches )
+		{
+			check( DirtyListIndex1 - OldDirtyListIndex1 <= JumpToIndex1 );
+			DirtyListIndex1 = OldDirtyListIndex1 + JumpToIndex1;
+			check( DirtyList1[DirtyListIndex1] == 0 );
+			DirtyListIndex1++;
+		}
 
-	// Patch in the jump offset
-	MergeState.MergedDirtyList[RememberIndex] = MergeState.MergedDirtyList.Num() - OldMergedListSize;
+		if ( bDirty2Matches )
+		{
+			check( DirtyListIndex2 - OldDirtyListIndex2 <= JumpToIndex2 );
+			DirtyListIndex2 = OldDirtyListIndex2 + JumpToIndex2;
+			check( DirtyList2[DirtyListIndex2] == 0 );
+			DirtyListIndex2++;
+		}
 
-	// Add the array terminator
-	MergeState.MergedDirtyList.Add( 0 );	
-}
+		bDirtyValid1 = bOldDirtyValid1;
+		bDirtyValid2 = bOldDirtyValid2;
 
-void FRepLayout::MergeDirtyList_r( FMergeDirtyListState & MergeState, const int32 CmdStart, const int32 CmdEnd, const uint8* RESTRICT Data ) const
-{
-	for ( int32 CmdIndex = CmdStart; CmdIndex < CmdEnd; CmdIndex++ )
+		// Patch in the jump offset
+		MergedDirtyList[OriginalMergedDirtyListIndex] = MergedDirtyList.Num() - ( OriginalMergedDirtyListIndex + 1 );
+
+		// Add the array terminator
+		MergedDirtyList.Add( 0 );	
+	}
+
+	PROCESS_CMD( FCmdIteratorBaseStackState ) 
 	{
-		const FRepLayoutCmd& Cmd = Cmds[ CmdIndex ];
+		check( bLastDirty1Matches || bLastDirty2Matches )
 
-		check( Cmd.Type != REPCMD_Return );
+		// This will be a new merged dirty entry
+		MergedDirtyList.Add( Handle );
 
-		if ( Cmd.Type == REPCMD_DynamicArray )
+		// Advance matching dirty indices
+		if ( bLastDirty1Matches )
 		{
-			MergeDirtyList_DynamicArray_r( MergeState, Cmd, CmdIndex, Data + Cmd.Offset );
-			CmdIndex = Cmd.EndCmd - 1;	// Jump past children of this array (-1 for ++ in for loop)
-			continue;
+			DirtyListIndex1++;
 		}
 
-		MergeState.Handle++;
+		if ( bLastDirty2Matches )
+		{
+			DirtyListIndex2++;
+		}
+	}
 
-		const bool Dirty1Matches = MergeState.DirtyValid1 && MergeState.DirtyList1[MergeState.DirtyListIndex1] == MergeState.Handle;
-		const bool Dirty2Matches = MergeState.DirtyValid2 && MergeState.DirtyList2[MergeState.DirtyListIndex2] == MergeState.Handle;
+	const TArray< uint16 >& DirtyList1;
+	const TArray< uint16 >& DirtyList2;
 
-		if ( Dirty1Matches || Dirty2Matches )
-		{
-			// This will be a new merged dirty entry
-			MergeState.MergedDirtyList.Add( MergeState.Handle );
+	int32					DirtyListIndex1;
+	int32					DirtyListIndex2;
+	uint16					Handle;
 
-			// Advance matching dirty indices
-			if ( Dirty1Matches )
-			{
-				MergeState.DirtyListIndex1++;
-			}
+	bool					bDirtyValid1;
+	bool					bDirtyValid2;
 
-			if ( Dirty2Matches )
-			{
-				MergeState.DirtyListIndex2++;
-			}
-		}
-	}
-}
+	TArray< uint16 >& 		MergedDirtyList;
 
-void FRepLayout::MergeDirtyList( const void* RESTRICT Data, const TArray< uint16 > & Dirty1, const TArray< uint16 > & Dirty2, TArray< uint16 > & MergedDirty ) const
+	bool					bLastDirty1Matches;
+	bool					bLastDirty2Matches;
+};
+
+void FRepLayout::MergeDirtyList( FRepState * RepState, const void* RESTRICT Data, const TArray< uint16 > & Dirty1, const TArray< uint16 > & Dirty2, TArray< uint16 > & MergedDirty ) const
 {
 	check( Dirty1.Num() > 0 || Dirty2.Num() > 0 );
 
-	FMergeDirtyListState MergeState( Dirty1, Dirty2, MergedDirty );
-
 	MergedDirty.Empty();
 
+	FMergeDirtyListImpl ReceivePropertiesImpl( Dirty1, Dirty2, MergedDirty, Parents, Cmds );
+
 	// Even though one of these can be empty, we need to send the single one through, so we can prune it to the current shape of the tree
-	MergeState.DirtyValid1 = Dirty1.Num() > 0;
-	MergeState.DirtyValid2 = Dirty2.Num() > 0;
+	ReceivePropertiesImpl.bDirtyValid1 = Dirty1.Num() > 0;
+	ReceivePropertiesImpl.bDirtyValid2 = Dirty2.Num() > 0;
 
-	MergeDirtyList_r( MergeState, 0, Cmds.Num() - 1, (const uint8*)Data );
+	// Merge lists
+	ReceivePropertiesImpl.ProcessCmds( RepState, (uint8*)Data );
 
-	MergeState.MergedDirtyList.Add( 0 );
+	ReceivePropertiesImpl.MergedDirtyList.Add( 0 );
 }
 
 void FRepLayout::SanityCheckChangeList_DynamicArray_r( 
@@ -2528,44 +2549,6 @@ void FRepLayout::RebuildConditionalProperties( FRepState * RESTRICT	RepState, co
 	RepState->RepFlags = RepFlags;
 }
 
-class FAddReferencedObjectsImpl
-{
-public:
-	FAddReferencedObjectsImpl( UObject* InReferencingObject, FReferenceCollector& InCollector ) : ReferencingObject( InReferencingObject ), Collector( InCollector ) {}
-
-	void PreArray( FMyCmdIteratorBaseStackState & StackState ) { }
-	void PostArray( FMyCmdIteratorBaseStackState & StackState ) { }
-
-	void ProcessCmd( const FRepLayoutCmd& Cmd, const uint8* RESTRICT ShadowData )
-	{
-		if ( Cmd.Type == REPCMD_PropertyObject )
-		{
-			UObjectPropertyBase * ObjProperty = CastChecked< UObjectPropertyBase>( Cmd.Property );
-
-			UObject* ReferencedObject = ObjProperty->GetObjectPropertyValue( ShadowData );
-
-			if ( ReferencedObject != NULL )
-			{
-				Collector.AddReferencedObject( ReferencedObject, ReferencingObject );
-			}
-		}
-	}
-
-	UObject*				ReferencingObject;
-	FReferenceCollector&	Collector;
-};
-
-void FRepLayout::AddReferencedObjects( FRepState * RepState, UObject* ReferencingObject, FReferenceCollector& Collector ) const
-{
-#if 0
-	FAddReferencedObjectsImpl AddReferencedObjectsImpl( ReferencingObject, Collector );
-
-	FRepLayoutCmdIterator< FAddReferencedObjectsImpl, FMyCmdIteratorBaseStackState > Iterator( AddReferencedObjectsImpl, Cmds );
-
-	Iterator.Start( RepState );
-#endif
-}
-
 void FRepLayout::InitChangedTracker( FRepChangedPropertyTracker * ChangedTracker ) const
 {
 	ChangedTracker->Parents.SetNum( Parents.Num() );
@@ -2652,6 +2635,8 @@ void FRepLayout::DestructProperties( FRepState * RepState ) const
 			Parents[i].Property->DestroyValue( StoredData + Offset );
 		}
 	}
+
+	RepState->StaticBuffer.Empty();
 }
 
 void FRepLayout::GetLifetimeCustomDeltaProperties( TArray< int32 > & OutCustom )
diff --git a/Engine/Source/Runtime/Engine/Public/Net/DataReplication.h b/Engine/Source/Runtime/Engine/Public/Net/DataReplication.h
--- a/Engine/Source/Runtime/Engine/Public/Net/DataReplication.h
+++ b/Engine/Source/Runtime/Engine/Public/Net/DataReplication.h
@@ -116,9 +116,6 @@ public:
 
 	void	ForceRefreshUnreliableProperties();
 
-	void	AddReferencedObjects( UObject* InThis, FReferenceCollector& Collector );
-
-
 	bool bHasReplicatedProperties;
 
 	void QueueRemoteFunctionBunch( UFunction* Func, FOutBunch &Bunch );
diff --git a/Engine/Source/Runtime/Engine/Public/Net/RepLayout.h b/Engine/Source/Runtime/Engine/Public/Net/RepLayout.h
--- a/Engine/Source/Runtime/Engine/Public/Net/RepLayout.h
+++ b/Engine/Source/Runtime/Engine/Public/Net/RepLayout.h
@@ -229,53 +229,6 @@ public:
 	bool				bDoChecksum;
 };
 
-class FRepReaderState
-{
-public:
-	FRepReaderState( FNetBitReader & InBunch, FRepState * InRepState, bool bInDoChecksum ) : 
-		WaitingHandle( 0 ), 
-		CurrentHandle( 0 ), 
-		Bunch( InBunch ),
-		RepState( InRepState ),
-		bDoChecksum( bInDoChecksum ),
-		bHasUnmapped( false )
-	{}
-
-	uint32					WaitingHandle;
-	uint32					CurrentHandle;
-	FNetBitReader &			Bunch;
-	FRepState *				RepState;
-	bool					bDoChecksum;
-	bool					bHasUnmapped;
-};
-
-class FMergeDirtyListState
-{
-public:
-	FMergeDirtyListState( const TArray< uint16 > & InDirty1, const TArray< uint16 > & InDirty2, TArray< uint16 > & OutMergedDirty ) : 
-		DirtyList1( InDirty1 ), 
-		DirtyList2( InDirty2 ), 
-		DirtyListIndex1( 0 ),
-		DirtyListIndex2( 0 ),
-		Handle( 0 ),
-		DirtyValid1( true ),
-		DirtyValid2( true ),
-		MergedDirtyList( OutMergedDirty )
-	{}
-
-	const TArray< uint16 > & DirtyList1;
-	const TArray< uint16 > & DirtyList2;
-
-	int32				DirtyListIndex1;
-	int32				DirtyListIndex2;
-	uint16				Handle;
-
-	bool				DirtyValid1;
-	bool				DirtyValid2;
-
-	TArray< uint16 > & 	MergedDirtyList;
-};
-
 /** FRepLayout
  *  This class holds all replicated properties for a parent property, and all its children
  *	Helpers functions exist to read/write and compare property state.
@@ -342,7 +295,7 @@ public:
 	void ValidateWithChecksum( const void* RESTRICT Data, FArchive & Ar ) const;
 	uint32 GenerateChecksum( const FRepState* RepState ) const;
 
-	void MergeDirtyList( const void* RESTRICT Data, const TArray< uint16 > & Dirty1, const TArray< uint16 > & Dirty2, TArray< uint16 > & MergedDirty ) const;
+	void MergeDirtyList( FRepState * RepState, const void* RESTRICT Data, const TArray< uint16 > & Dirty1, const TArray< uint16 > & Dirty2, TArray< uint16 > & MergedDirty ) const;
 
 	bool DiffProperties( FRepState * RepState, const void* RESTRICT Data, const bool bSync ) const;
 
@@ -357,8 +310,6 @@ public:
 	void SerializePropertiesForStruct( UStruct * Struct, FArchive & Ar, UPackageMap	* Map, void* Data, bool & bHasUnmapped ) const;	
 	void InitFromStruct( UStruct * InStruct );
 
-	void AddReferencedObjects( FRepState * RepState, UObject* ReferencingObject, FReferenceCollector& Collector ) const;
-
 private:
 	void RebuildConditionalProperties( FRepState * RESTRICT	RepState, const FRepChangedPropertyTracker& ChangedTracker, const FReplicationFlags& RepFlags ) const;
 
@@ -422,43 +373,6 @@ private:
 		const uint8* RESTRICT		Data, 
 		uint16						Handle ) const;
 
-	bool ReadProperty( 
-		FRepReaderState &		ReaderState, 
-		FUnmappedGuidMgr *		UnmappedGuids,
-		const int32				AbsOffset,
-		const FRepLayoutCmd &	Cmd, 
-		const int32				CurrentCmdIndex, 
-		uint8* RESTRICT		StoredData, 
-		uint8* RESTRICT		Data ) const;
-
-	bool ReceiveProperties_AnyArray_r( 
-		FRepReaderState &	ReaderState, 
-		FUnmappedGuidMgr *	UnmappedGuids,
-		const int32			AbsOffset,
-		const int32			ArrayNum, 
-		const int32			ElementSize, 
-		const int32			CmdIndex, 
-		uint8* RESTRICT	StoredData, 
-		uint8 *	RESTRICT	Data ) const;
-
-	bool ReceiveProperties_DynamicArray_r( 
-		FRepReaderState &		ReaderState, 
-		FUnmappedGuidMgr *		UnmappedGuids,
-		const int32				AbsOffset,
-		const FRepLayoutCmd &	Cmd, 
-		const int32				CmdIndex, 
-		uint8* RESTRICT		StoredData, 
-		uint8* RESTRICT		Data ) const;
-
-	bool ReceiveProperties_r( 
-		FRepReaderState &	ReaderState, 
-		FUnmappedGuidMgr *	UnmappedGuids,
-		const int32			AbsOffset,
-		const int32			CmdStart, 
-		const int32			CmdEnd, 
-		uint8* RESTRICT	StoredData, 
-		uint8* RESTRICT	Data ) const;
-
 	void UpdateUnmappedObjects_r( 
 		FRepState *			RepState, 
 		FUnmappedGuidMgr *	UnmappedGuids, 
@@ -473,10 +387,6 @@ private:
 	void ValidateWithChecksum_DynamicArray_r( const FRepLayoutCmd& Cmd, const int32 CmdIndex, const uint8* RESTRICT Data, FArchive & Ar ) const;
 	void ValidateWithChecksum_r( const int32 CmdStart, const int32 CmdEnd, const uint8* RESTRICT Data, FArchive & Ar ) const;
 
-	void MergeDirtyList_AnyArray_r( FMergeDirtyListState & MergeState, const FRepLayoutCmd& Cmd, const int32 ArrayNum, const int32 ElementSize, const int32 CmdIndex, const uint8* RESTRICT Data ) const;
-	void MergeDirtyList_DynamicArray_r( FMergeDirtyListState & MergeState, const FRepLayoutCmd& Cmd, const int32 CmdIndex, const uint8* RESTRICT Data ) const;
-	void MergeDirtyList_r( FMergeDirtyListState & MergeState, const int32 CmdStart, const int32 CmdEnd, const uint8* RESTRICT Data ) const;
-
 	void SanityCheckChangeList_DynamicArray_r( 
 		const int32				CmdIndex, 
 		const uint8* RESTRICT	Data, 
