Refactor active transform cache so that it's sorted by type, and only updated with UE objects that we maintain.
This fixes issues where client code creates and destroys PxActors without us knowing.

JIRA UE-3053

diff --git a/Engine/Source/Runtime/Engine/Private/PhysicsEngine/PhysScene.cpp b/Engine/Source/Runtime/Engine/Private/PhysicsEngine/PhysScene.cpp
--- a/Engine/Source/Runtime/Engine/Private/PhysicsEngine/PhysScene.cpp
+++ b/Engine/Source/Runtime/Engine/Private/PhysicsEngine/PhysScene.cpp
@@ -296,23 +296,15 @@ void FPhysScene::AddTorque(FBodyInstance* BodyInstance, const FVector& Torque, b
 }
 
 #if WITH_PHYSX
-void FPhysScene::RemoveBodyFromActiveTransforms(class PxActor * PActor, uint32 SceneType)
+void FPhysScene::RemoveActiveBody(FBodyInstance* BodyInstance, uint32 SceneType)
 {
-	if (PActor)
+	int32 BodyIndex = ActiveBodyInstances[SceneType].Find(BodyInstance);
+	if (BodyIndex != INDEX_NONE)
 	{
-		TArray<const PxActiveTransform*> & PActiveTransforms = ActiveTransforms[SceneType];
-		for (int32 TransformIndex = 0; TransformIndex < PActiveTransforms.Num(); ++TransformIndex)
-		{
-			const PxActiveTransform * PActiveTransform = PActiveTransforms[TransformIndex];
-			if (PActiveTransform && PActiveTransform->actor == PActor)
-			{
-				PActiveTransforms[TransformIndex] = NULL;	//we NULL out the pointer to indicate that this actor has been destroyed. This is needed because the transform cache doesn't get updated until the next fetchResults
-			}
-		}
+		ActiveBodyInstances[SceneType][BodyIndex] = nullptr;
 	}
 }
 #endif
-
 void FPhysScene::TermBody(FBodyInstance* BodyInstance)
 {
 #if WITH_SUBSTEPPING
@@ -322,11 +314,9 @@ void FPhysScene::TermBody(FBodyInstance* BodyInstance)
 		PhysSubStepper->RemoveBodyInstance(BodyInstance);
 	}
 #endif
-
-
 #if WITH_PHYSX
-	RemoveBodyFromActiveTransforms(BodyInstance->RigidActorSync, PST_Sync);
-	RemoveBodyFromActiveTransforms(BodyInstance->RigidActorAsync, PST_Async);
+	RemoveActiveBody(BodyInstance, PST_Sync);
+	RemoveActiveBody(BodyInstance, PST_Async);
 #endif
 }
 
@@ -671,11 +661,15 @@ void FPhysScene::ProcessPhysScene(uint32 SceneType)
 		}
 	}
 
+
+	// Reset execution flag
+
+//This fetches and gets active transforms. It's important that the function that calls this locks because getting the transforms and using the data must be an atomic operation
 #if WITH_PHYSX
 	PxScene* PScene = GetPhysXScene(SceneType);
 	check(PScene);
-
 	PxU32 OutErrorCode = 0;
+
 #if !WITH_APEX
 	PScene->lockWrite();
 	PScene->fetchResults(true, &OutErrorCode);
@@ -688,7 +682,6 @@ void FPhysScene::ProcessPhysScene(uint32 SceneType)
 #endif	//	#if !WITH_APEX
 
 	UpdateActiveTransforms(SceneType);
-
 	if (OutErrorCode != 0)
 	{
 		UE_LOG(LogPhysics, Log, TEXT("PHYSX FETCHRESULTS ERROR: %d"), OutErrorCode);
@@ -696,11 +689,8 @@ void FPhysScene::ProcessPhysScene(uint32 SceneType)
 #endif // WITH_PHYSX
 
 	PhysicsSubsceneCompletion[SceneType] = NULL;
-
-	// Reset execution flag
 	bPhysXSceneExecuting[SceneType] = false;
 }
-
 #if WITH_PHYSX
 void FPhysScene::UpdateActiveTransforms(uint32 SceneType)
 {
@@ -708,28 +698,34 @@ void FPhysScene::UpdateActiveTransforms(uint32 SceneType)
 	{
 		return;
 	}
-
-
 	PxScene* PScene = GetPhysXScene(SceneType);
 	check(PScene);
-	SCENE_LOCK_READ(PScene);
-	PxU32 NumTransforms = 0;
+	SCOPED_SCENE_WRITE_LOCK(PScene);
 
+	PxU32 NumTransforms = 0;
 	const PxActiveTransform* PActiveTransforms = PScene->getActiveTransforms(NumTransforms);
-	SCENE_UNLOCK_READ(PScene);
-
-	ActiveTransforms[SceneType].Empty(NumTransforms);
+	ActiveBodyInstances[SceneType].Empty(NumTransforms);
+	ActiveDestructibleActors[SceneType].Empty(NumTransforms);
 
 	for (PxU32 TransformIdx = 0; TransformIdx < NumTransforms; ++TransformIdx)
 	{
-		ActiveTransforms[SceneType].Add(&PActiveTransforms[TransformIdx]);	//it's ok to use a pointer here because the physics scene leaves the data for us until the next call to fetchTransform;
-		
 		const PxActiveTransform& PActiveTransform = PActiveTransforms[TransformIdx];
 		PxRigidActor* RigidActor = PActiveTransform.actor->isRigidActor();
-
 		ensure(!RigidActor->userData || !FPhysxUserData::IsGarbage(RigidActor->userData));
-	}
 
+		if (FBodyInstance* BodyInstance = FPhysxUserData::Get<FBodyInstance>(RigidActor->userData))
+		{
+			if (BodyInstance->InstanceBodyIndex == INDEX_NONE && BodyInstance->OwnerComponent.IsValid() && BodyInstance->IsInstanceSimulatingPhysics())
+			{
+				ActiveBodyInstances[SceneType].Add(BodyInstance);
+			}
+		}
+		else if (const FDestructibleChunkInfo* DestructibleChunkInfo = FPhysxUserData::Get<FDestructibleChunkInfo>(RigidActor->userData))
+		{
+			ActiveDestructibleActors[SceneType].Add(RigidActor);
+		}
+		
+	}
 }
 #endif
 
@@ -739,102 +735,39 @@ void FPhysScene::SyncComponentsToBodies(uint32 SceneType)
 {
 	SCOPE_CYCLE_COUNTER(STAT_TotalPhysicsTime);
 	SCOPE_CYCLE_COUNTER(STAT_SyncComponentsToBodies);
-#if WITH_PHYSX
-
-	SCOPED_SCENE_READ_LOCK(GetPhysXScene(SceneType));
-
-	PxU32 NumTransforms = ActiveTransforms[SceneType].Num();
-	const TArray<const PxActiveTransform*> & PActiveTransforms = ActiveTransforms[SceneType];
-
-#if WITH_APEX
-	TArray<const PxRigidActor*> ActiveDestructionActors;
-#endif
 
-	for (PxU32 TransformIdx = 0; TransformIdx<NumTransforms; TransformIdx++)
+	for (FBodyInstance* BodyInstance : ActiveBodyInstances[SceneType])
 	{
-		if (PActiveTransforms[TransformIdx] == NULL)	//it's possible to call TermBody on FBodyInstance after fetchResults, but before SyncComponentsToBodies - in this case a NULL is used to represent the stale data
-		{
-			continue;
-		}
+		if (BodyInstance == nullptr) { continue; }
 
-		const PxActiveTransform& PActiveTransform = *PActiveTransforms[TransformIdx];
-		PxRigidActor* RigidActor = PActiveTransform.actor->isRigidActor();
+		check(BodyInstance->OwnerComponent->IsRegistered()); // shouldn't have a physics body for a non-registered component!
 
-		ensure(!RigidActor->userData || !FPhysxUserData::IsGarbage(RigidActor->userData));
+		AActor* Owner = BodyInstance->OwnerComponent->GetOwner();
 
-#if WITH_APEX
-		//Special code for destructible chunk
-		if (const FDestructibleChunkInfo* DestructibleChunkInfo = FPhysxUserData::Get<FDestructibleChunkInfo>(RigidActor->userData))
+		// See if the transform is actually different, and if so, move the component to match physics
+		const FTransform NewTransform = BodyInstance->GetUnrealWorldTransform();
+		if (!NewTransform.EqualsNoScale(BodyInstance->OwnerComponent->ComponentToWorld))
 		{
-			/*if (GApexModuleDestructible->owns(RigidActor) && DestructibleChunkInfo->OwningComponent.IsValid())
-			{
-				//TODO: waiting on new API to avoid duplicate updates per shape.
-				TArray<PxShape*> Shapes;
-				Shapes.AddZeroed(RigidActor->getNbShapes());
-				int32 NumShapes = RigidActor->getShapes(Shapes.GetData(), Shapes.Num());
-				for (int32 ShapeIdx = 0; ShapeIdx < Shapes.Num(); ++ShapeIdx)
-				{
-					PxShape* Shape = Shapes[ShapeIdx];
-					int32 ChunkIndex;
-					if (NxDestructibleActor* DestructibleActor = GApexModuleDestructible->getDestructibleAndChunk(Shape, &ChunkIndex))
-					{
-						const physx::PxMat44 ChunkPoseRT = DestructibleActor->getChunkPose(ChunkIndex);
-						const physx::PxTransform Transform(ChunkPoseRT);
-						if (UDestructibleComponent * DestructibleComponent = Cast<UDestructibleComponent>(FPhysxUserData::Get<UPrimitiveComponent>(DestructibleActor->userData)))
-						{
-							// if this component was already unregistered, transform data were deallocated
-							// transform data become empty when a piece of destructible actor goes out from a streaming volume
-							if (DestructibleComponent->IsRegistered())
-							{
-								DestructibleComponent->SetChunkWorldRT(ChunkIndex, P2UQuat(Transform.q), P2UVector(Transform.p));
-							}
-						}
-					}
-				}
-			}*/
+			const FVector MoveBy = NewTransform.GetLocation() - BodyInstance->OwnerComponent->ComponentToWorld.GetLocation();
+			const FRotator NewRotation = NewTransform.Rotator();
 
-			ActiveDestructionActors.Add(RigidActor);
+			//@warning: do not reference BodyInstance again after calling MoveComponent() - events from the move could have made it unusable (destroying the actor, SetPhysics(), etc)
+			BodyInstance->OwnerComponent->MoveComponent(MoveBy, NewRotation, false, NULL, MOVECOMP_SkipPhysicsMove);
 		}
-#endif
 
-		FBodyInstance* BodyInst = FPhysxUserData::Get<FBodyInstance>(PActiveTransform.userData);
-		if (BodyInst != NULL &&
-			BodyInst->InstanceBodyIndex == INDEX_NONE &&
-			BodyInst->OwnerComponent != NULL &&
-			BodyInst->IsInstanceSimulatingPhysics())
+		// Check if we didn't fall out of the world
+		if (Owner != NULL && !Owner->IsPendingKill())
 		{
-			check(BodyInst->OwnerComponent->IsRegistered()); // shouldn't have a physics body for a non-registered component!
-
-			AActor* Owner = BodyInst->OwnerComponent->GetOwner();
-
-			// See if the transform is actually different, and if so, move the component to match physics
-			const FTransform NewTransform = BodyInst->GetUnrealWorldTransform();
-			if (!NewTransform.EqualsNoScale(BodyInst->OwnerComponent->ComponentToWorld))
-			{
-				const FVector MoveBy = NewTransform.GetLocation() - BodyInst->OwnerComponent->ComponentToWorld.GetLocation();
-				const FRotator NewRotation = NewTransform.Rotator();
-
-				//UE_LOG(LogTemp, Log, TEXT("MOVING: %s"), *BodyInst->OwnerComponent->GetPathName());
-
-				//@warning: do not reference BodyInstance again after calling MoveComponent() - events from the move could have made it unusable (destroying the actor, SetPhysics(), etc)
-				BodyInst->OwnerComponent->MoveComponent(MoveBy, NewRotation, false, NULL, MOVECOMP_SkipPhysicsMove);
-			}
-
-			// Check if we didn't fall out of the world
-			if (Owner != NULL && !Owner->IsPendingKill())
-			{
-				Owner->CheckStillInWorld();
-			}
+			Owner->CheckStillInWorld();
 		}
 	}
 
 #if WITH_APEX
-	if (ActiveDestructionActors.Num())
+	if (ActiveDestructibleActors[SceneType].Num())
 	{
-		UDestructibleComponent::UpdateDestructibleChunkTM(ActiveDestructionActors);
+		UDestructibleComponent::UpdateDestructibleChunkTM(ActiveDestructibleActors[SceneType]);
 	}
 #endif
-#endif
 }
 
 void FPhysScene::DispatchPhysNotifications()
diff --git a/Engine/Source/Runtime/Engine/Public/PhysicsPublic.h b/Engine/Source/Runtime/Engine/Public/PhysicsPublic.h
--- a/Engine/Source/Runtime/Engine/Public/PhysicsPublic.h
+++ b/Engine/Source/Runtime/Engine/Public/PhysicsPublic.h
@@ -55,6 +55,7 @@ namespace physx
 	class PxSimulationEventCallback;
 	struct PxActiveTransform;
 	class PxActor;
+	class PxRigidActor;
 	
 #if WITH_APEX
 	namespace apex
@@ -447,14 +448,14 @@ private:
 	/** User data wrapper passed to physx */
 	struct FPhysxUserData PhysxUserData;
 
-	/** Cache of active transforms */ //TODO: this solution is not great
-	TArray<const PxActiveTransform*> ActiveTransforms[PST_MAX];
+	/** Cache of active transforms sorted into types */ //TODO: this solution is not great
+	TArray<struct FBodyInstance*> ActiveBodyInstances[PST_MAX];	//body instances that have moved
+	TArray<const physx::PxRigidActor*> ActiveDestructibleActors[PST_MAX];	//destructible actors that have moved
 
-	/** Updates our cache of active transforms */
+	/** Fetch results from simulation and get the active transforms. Make sure to lock before calling this function as the fetch and data you use must be treated as an atomic operation */
 	void UpdateActiveTransforms(uint32 SceneType);
+	void RemoveActiveBody(FBodyInstance* BodyInstance, uint32 SceneType);
 
-	/** When actors are destroyed they must be removed from active transforms */
-	void RemoveBodyFromActiveTransforms(PxActor * PActor, uint32 SceneType);
 #endif
 
 #if WITH_SUBSTEPPING
