Add scene locks for physics handle

diff --git a/Engine/Source/Runtime/Engine/Private/PhysicsEngine/PhysicsHandleComponent.cpp b/Engine/Source/Runtime/Engine/Private/PhysicsEngine/PhysicsHandleComponent.cpp
--- a/Engine/Source/Runtime/Engine/Private/PhysicsEngine/PhysicsHandleComponent.cpp
+++ b/Engine/Source/Runtime/Engine/Private/PhysicsEngine/PhysicsHandleComponent.cpp
@@ -90,6 +90,7 @@ void UPhysicsHandleComponent::GrabComponent(UPrimitiveComponent* InComponent, FN
 	// If we don't already have a handle - make one now.
 	if (!HandleData)
 	{
+		SCOPED_SCENE_WRITE_LOCK(Scene);
 		// Create kinematic actor we are going to create joint with. This will be moved around with calls to SetLocation/SetRotation.
 		PxRigidDynamic* KinActor = Scene->getPhysics().createRigidDynamic(KinPose);
 		KinActor->setRigidDynamicFlag(PxRigidDynamicFlag::eKINEMATIC, true);
@@ -175,6 +176,7 @@ void UPhysicsHandleComponent::ReleaseComponent()
 			PxScene* PScene = GetPhysXSceneFromIndex( SceneIndex );
 			if(PScene)
 			{
+				SCOPED_SCENE_WRITE_LOCK(PScene);
 				// Destroy joint.
 				HandleData->release();
 				
