Expose GetCenterOfMass to blueprints

diff --git a/Engine/Source/Runtime/Engine/Classes/Components/PrimitiveComponent.h b/Engine/Source/Runtime/Engine/Classes/Components/PrimitiveComponent.h
--- a/Engine/Source/Runtime/Engine/Classes/Components/PrimitiveComponent.h
+++ b/Engine/Source/Runtime/Engine/Classes/Components/PrimitiveComponent.h
@@ -789,6 +789,14 @@ public:
 	UFUNCTION(BlueprintCallable, Category="Physics")	
 	FVector GetPhysicsAngularVelocity(FName BoneName = NAME_None);
 
+	/**
+	*	Get the center of mass of a single body. In the case of a welded body this will return the center of mass of the entire welded body (including its parent and children)
+	*   Objects that are not simulated return (0,0,0) as they do not have COM
+	*	@param BoneName			If a SkeletalMeshComponent, name of body to get center of mass of. 'None' indicates root body.
+	*/
+	UFUNCTION(BlueprintCallable, Category = "Physics")
+	FVector GetCenterOfMass(FName BoneName = NAME_None);
+
 	/**
 	 *	'Wake' physics simulation for a single body.
 	 *	@param	BoneName	If a SkeletalMeshComponent, name of body to wake. 'None' indicates root body.
diff --git a/Engine/Source/Runtime/Engine/Private/PrimitiveComponentPhysics.cpp b/Engine/Source/Runtime/Engine/Private/PrimitiveComponentPhysics.cpp
--- a/Engine/Source/Runtime/Engine/Private/PrimitiveComponentPhysics.cpp
+++ b/Engine/Source/Runtime/Engine/Private/PrimitiveComponentPhysics.cpp
@@ -286,6 +286,16 @@ FVector UPrimitiveComponent::GetPhysicsAngularVelocity(FName BoneName)
 	return FVector(0,0,0);
 }
 
+FVector UPrimitiveComponent::GetCenterOfMass(FName BoneName)
+{
+	if (FBodyInstance* BodyInstance = GetBodyInstance(BoneName))
+	{
+		return BodyInstance->GetCOMPosition();
+	}
+
+	return FVector::ZeroVector;
+}
+
 void UPrimitiveComponent::SetAllPhysicsAngularVelocity(FVector const& NewAngVel, bool bAddToCurrent)
 {
 	SetPhysicsAngularVelocity(NewAngVel, bAddToCurrent, NAME_None); 
