- Fixed clang compiler errors (constructor order with members being inited before the base class)

diff --git a/Engine/Source/Runtime/Engine/Private/RepLayout.cpp b/Engine/Source/Runtime/Engine/Private/RepLayout.cpp
--- a/Engine/Source/Runtime/Engine/Private/RepLayout.cpp
+++ b/Engine/Source/Runtime/Engine/Private/RepLayout.cpp
@@ -1133,8 +1133,8 @@ class FReceivedPropertiesStackState : public FCmdIteratorBaseStackState
 {
 public:
 	FReceivedPropertiesStackState( const int32 InCmdStart, const int32 InCmdEnd, FScriptArray*	InShadowArray, FScriptArray* InDataArray, uint8* RESTRICT InShadowBaseData, uint8* RESTRICT	InBaseData ) : 
-		UnmappedGuids( NULL ),
-		FCmdIteratorBaseStackState( InCmdStart, InCmdEnd, InShadowArray, InDataArray, InShadowBaseData, InBaseData ) 
+		FCmdIteratorBaseStackState( InCmdStart, InCmdEnd, InShadowArray, InDataArray, InShadowBaseData, InBaseData ),
+        UnmappedGuids( NULL )
 	{}
 
 	FUnmappedGuidMgr* UnmappedGuids;
@@ -1144,13 +1144,13 @@ class FReceivePropertiesImpl : public FRepLayoutCmdIterator< FReceivePropertiesI
 {
 public:
 	FReceivePropertiesImpl( FNetBitReader & InBunch, FRepState * InRepState, bool bInDoChecksum, const TArray< FRepParentCmd >& InParents, const TArray< FRepLayoutCmd >& InCmds ) :
-		WaitingHandle( 0 ), 
+        FRepLayoutCmdIterator( InParents, InCmds ),
+		WaitingHandle( 0 ),
 		CurrentHandle( 0 ), 
 		Bunch( InBunch ),
 		RepState( InRepState ),
 		bDoChecksum( bInDoChecksum ),
-		bHasUnmapped( false ),
-		FRepLayoutCmdIterator( InParents, InCmds ) 
+		bHasUnmapped( false )
 	{}
 
 	void ReadNextHandle()
@@ -1598,7 +1598,8 @@ class FMergeDirtyListImpl : public FRepLayoutCmdIterator< FMergeDirtyListImpl, F
 {
 public:
 	FMergeDirtyListImpl( const TArray< uint16 > & InDirty1, const TArray< uint16 > & InDirty2, TArray< uint16 > & OutMergedDirty, const TArray< FRepParentCmd >& InParents, const TArray< FRepLayoutCmd >& InCmds ) : 
-		DirtyList1( InDirty1 ), 
+        FRepLayoutCmdIterator( InParents, InCmds ),
+		DirtyList1( InDirty1 ),
 		DirtyList2( InDirty2 ), 
 		DirtyListIndex1( 0 ),
 		DirtyListIndex2( 0 ),
@@ -1607,8 +1608,7 @@ public:
 		bDirtyValid2( true ),
 		MergedDirtyList( OutMergedDirty ),
 		bLastDirty1Matches( false ), 
-		bLastDirty2Matches( false ),
-		FRepLayoutCmdIterator( InParents, InCmds )
+		bLastDirty2Matches( false )
 	{}
 
 	INIT_STACK( FCmdIteratorBaseStackState ) { }
