Less verbose logging when replicated property type not recognized

diff --git a/Engine/Source/Runtime/Engine/Private/RepLayout.cpp b/Engine/Source/Runtime/Engine/Private/RepLayout.cpp
--- a/Engine/Source/Runtime/Engine/Private/RepLayout.cpp
+++ b/Engine/Source/Runtime/Engine/Private/RepLayout.cpp
@@ -2032,7 +2032,7 @@ void FRepLayout::AddPropertyCmd( UProperty * Property, int32 Offset, int32 Relat
 	}
 	else
 	{
-		UE_LOG( LogNet, Verbose, TEXT( "AddPropertyCmd: Falling back to default type for property [%s]" ), *Cmd.Property->GetFullName() );
+		UE_LOG( LogNet, VeryVerbose, TEXT( "AddPropertyCmd: Falling back to default type for property [%s]" ), *Cmd.Property->GetFullName() );
 	}
 }
 
