#ue4 - Falling doesn't need to check air control before the first move, since deflection uses the non-air control velocity to deflect.

- Modified FindAirControlImpact() to distinguish between current and lookahead time, and to take velocity and gravity. However, this might be deprecated as it is not currently used.

diff --git a/Engine/Source/Runtime/Engine/Classes/GameFramework/CharacterMovementComponent.h b/Engine/Source/Runtime/Engine/Classes/GameFramework/CharacterMovementComponent.h
--- a/Engine/Source/Runtime/Engine/Classes/GameFramework/CharacterMovementComponent.h
+++ b/Engine/Source/Runtime/Engine/Classes/GameFramework/CharacterMovementComponent.h
@@ -1034,16 +1034,19 @@ protected:
 	virtual float BoostAirControl(float DeltaTime, float TickAirControl, const FVector& FallAcceleration);
 
 	/**
-	 * Checks if air control will cause the player collision shape to hit something given current conditions.
+	 * Checks if air control will cause the player collision shape to hit something given the current location.
 	 * This function is used internally by PhysFalling().
 	 *
 	 * @param DeltaTime			Time step for the current update.
+	 * @param AdditionalTime	Time to look ahead further, applying acceleration and gravity.
+	 * @param FallVelocity		Velocity to test with.
 	 * @param FallAcceleration	Acceleration used during movement.
+	 * @param Gravity			Gravity to apply to any additional time.
 	 * @param OutHitResult		Result of impact, valid if this function returns true.
 	 * @return True if there is an impact, in which case OutHitResult contains the result of that impact.
 	 * @see GetAirControl()
 	 */
-	virtual bool FindAirControlImpact(float DeltaTime, const FVector& FallAcceleration, FHitResult& OutHitResult);
+	virtual bool FindAirControlImpact(float DeltaTime, float AdditionalTime, const FVector& FallVelocity, const FVector& FallAcceleration, const FVector& Gravity, FHitResult& OutHitResult);
 
 	/**
 	 * Limits the air control to use during falling movement, given an impact from FindAirControlImpact().
diff --git a/Engine/Source/Runtime/Engine/Private/CharacterMovementComponent.cpp b/Engine/Source/Runtime/Engine/Private/CharacterMovementComponent.cpp
--- a/Engine/Source/Runtime/Engine/Private/CharacterMovementComponent.cpp
+++ b/Engine/Source/Runtime/Engine/Private/CharacterMovementComponent.cpp
@@ -2865,6 +2865,7 @@ void UCharacterMovementComponent::PhysFalling(float deltaTime, int32 Iterations)
 		const FVector Gravity(0.f, 0.f, GetGravityZ());
 		Velocity = NewFallVelocity(Velocity, Gravity, timeTick);
 		VelocityNoAirControl = NewFallVelocity(VelocityNoAirControl, Gravity, timeTick);
+		const FVector AirControlAccel = (Velocity - VelocityNoAirControl) / timeTick;
 
 		if( bNotifyApex && CharacterOwner->Controller && (Velocity.Z <= 0.f) )
 		{
@@ -2873,17 +2874,19 @@ void UCharacterMovementComponent::PhysFalling(float deltaTime, int32 Iterations)
 			NotifyJumpApex();
 		}
 
-		// Test air control to see if we should limit it.
-		const FVector AirControlAccel = (Velocity - VelocityNoAirControl) / timeTick;
+		// Not really necessary to limit air control since the first deflection uses non air control velocity for deflection!
+		/*
 		if (bHasAirControl)
 		{
+			const float LookAheadTime = 0.02f;
 			FHitResult TestAirResult(1.f);
-			if (FindAirControlImpact(timeTick, AirControlAccel, TestAirResult))
+			if (FindAirControlImpact(timeTick, LookAheadTime, Velocity, FallAcceleration, Gravity, TestAirResult))
 			{
-				const FVector LimitedAirAccel = LimitAirControl(timeTick, AirControlAccel, TestAirResult, /*bCheckForValidLandingSpot=*/ true);
+				const FVector LimitedAirAccel = LimitAirControl(timeTick, AirControlAccel, TestAirResult, true);
 				Velocity = VelocityNoAirControl + (LimitedAirAccel * timeTick);
 			}
 		}
+		*/
 
 
 		// Move
@@ -2942,7 +2945,7 @@ void UCharacterMovementComponent::PhysFalling(float deltaTime, int32 Iterations)
 				}
 
 				// Limit air control based on what we hit.
-				// We moved to the impact point using air control, but may want to deflect from there based on a non-air control acceleration.
+				// We moved to the impact point using air control, but may want to deflect from there based on a limited air control acceleration.
 				if (bHasAirControl)
 				{
 					const bool bCheckLandingSpot = false; // we already checked above.
@@ -2961,7 +2964,7 @@ void UCharacterMovementComponent::PhysFalling(float deltaTime, int32 Iterations)
 					Velocity = HasRootMotion() ? FVector(Velocity.X, Velocity.Y, NewVelocity.Z) : NewVelocity;
 				}
 
-				if ((Delta | Adjusted) > 0.f)
+				if (subTimeTickRemaining > KINDA_SMALL_NUMBER && (Delta | Adjusted) > 0.f)
 				{
 					// Move in deflected direction.
 					SafeMoveUpdatedComponent( Delta, PawnRotation, true, Hit);
@@ -3066,12 +3069,16 @@ void UCharacterMovementComponent::PhysFalling(float deltaTime, int32 Iterations)
 }
 
 
-bool UCharacterMovementComponent::FindAirControlImpact(float DeltaTime, const FVector& FallAcceleration, FHitResult& OutHitResult)
+bool UCharacterMovementComponent::FindAirControlImpact(float DeltaTime, float AdditionalTime, const FVector& FallVelocity, const FVector& FallAcceleration, const FVector& Gravity, FHitResult& OutHitResult)
 {
-	// test for slope to avoid using air control to climb walls
-	const float TestWalkTime = FMath::Max(DeltaTime, 0.01f);
-	const FVector PostGravityVelocity = NewFallVelocity(Velocity, FVector(0.f,0.f,GetGravityZ()), TestWalkTime);
-	const FVector TestWalk = ((FallAcceleration * TestWalkTime) + PostGravityVelocity) * TestWalkTime;
+	// Test for slope to avoid using air control to climb walls.
+	FVector TestWalk = Velocity * DeltaTime;
+	if (AdditionalTime > 0.f)
+	{
+		const FVector PostGravityVelocity = NewFallVelocity(FallVelocity, Gravity, AdditionalTime);
+		TestWalk += ((FallAcceleration * AdditionalTime) + PostGravityVelocity) * AdditionalTime;
+	}
+	
 	if (!TestWalk.IsZero())
 	{
 		static const FName FallingTraceParamsTag = FName(TEXT("PhysFalling"));
diff --git a/Engine/Source/Runtime/Engine/Private/MovementComponent.cpp b/Engine/Source/Runtime/Engine/Private/MovementComponent.cpp
--- a/Engine/Source/Runtime/Engine/Private/MovementComponent.cpp
+++ b/Engine/Source/Runtime/Engine/Private/MovementComponent.cpp
@@ -519,7 +519,7 @@ void UMovementComponent::TwoWallAdjust(FVector& OutDelta, const FHitResult& Hit,
 		{
 			// we hit the same wall again even after adjusting to move along it the first time
 			// nudge away from it (this can happen due to precision issues)
-			Delta += HitNormal * 0.1f;
+			Delta += HitNormal * 0.01f;
 		}
 	}
 
