updating bounds of NavRelevantComponents after moving owning actor in editor

diff --git a/Engine/Source/Runtime/Engine/Classes/AI/Navigation/NavRelevantComponent.h b/Engine/Source/Runtime/Engine/Classes/AI/Navigation/NavRelevantComponent.h
--- a/Engine/Source/Runtime/Engine/Classes/AI/Navigation/NavRelevantComponent.h
+++ b/Engine/Source/Runtime/Engine/Classes/AI/Navigation/NavRelevantComponent.h
@@ -18,6 +18,7 @@ class ENGINE_API UNavRelevantComponent : public UActorComponent, public INavRele
 	// Begin INavRelevantInterface Interface
 	virtual FBox GetNavigationBounds() const override;
 	virtual bool IsNavigationRelevant() const override;
+	virtual void UpdateNavigationBounds() override;
 	// End INavRelevantInterface Interface
 
 	virtual void CalcBounds();
diff --git a/Engine/Source/Runtime/Engine/Classes/AI/Navigation/NavRelevantInterface.h b/Engine/Source/Runtime/Engine/Classes/AI/Navigation/NavRelevantInterface.h
--- a/Engine/Source/Runtime/Engine/Classes/AI/Navigation/NavRelevantInterface.h
+++ b/Engine/Source/Runtime/Engine/Classes/AI/Navigation/NavRelevantInterface.h
@@ -20,6 +20,9 @@ class INavRelevantInterface
 	/** Get bounds for navigation octree */
 	virtual FBox GetNavigationBounds() const { return FBox(0); }
 
+	/** Update bounds, called after moving owning actor */
+	virtual void UpdateNavigationBounds() {}
+
 	/** Are modifiers active? */
 	virtual bool IsNavigationRelevant() const { return true; }
 
diff --git a/Engine/Source/Runtime/Engine/Classes/AI/Navigation/NavigationSystem.h b/Engine/Source/Runtime/Engine/Classes/AI/Navigation/NavigationSystem.h
--- a/Engine/Source/Runtime/Engine/Classes/AI/Navigation/NavigationSystem.h
+++ b/Engine/Source/Runtime/Engine/Classes/AI/Navigation/NavigationSystem.h
@@ -459,6 +459,8 @@ public:
 	static void UpdateNavOctreeAll(AActor* Actor);
 	/** removes all navoctree entries for actor and its components */
 	static void ClearNavOctreeAll(AActor* Actor);
+	/** updates bounds of all components implementing INavRelevantInterface */
+	static void UpdateNavOctreeBounds(AActor* Actor);
 
 	void AddDirtyArea(const FBox& NewArea, int32 Flags);
 	void AddDirtyAreas(const TArray<FBox>& NewAreas, int32 Flags);
diff --git a/Engine/Source/Runtime/Engine/Private/AI/Navigation/NavRelevantComponent.cpp b/Engine/Source/Runtime/Engine/Private/AI/Navigation/NavRelevantComponent.cpp
--- a/Engine/Source/Runtime/Engine/Private/AI/Navigation/NavRelevantComponent.cpp
+++ b/Engine/Source/Runtime/Engine/Private/AI/Navigation/NavRelevantComponent.cpp
@@ -34,6 +34,11 @@ bool UNavRelevantComponent::IsNavigationRelevant() const
 	return bNavigationRelevant;
 }
 
+void UNavRelevantComponent::UpdateNavigationBounds()
+{
+	CalcBounds();
+}
+
 void UNavRelevantComponent::CalcBounds()
 {
 	AActor* OwnerActor = GetOwner();
diff --git a/Engine/Source/Runtime/Engine/Private/AI/Navigation/NavigationSystem.cpp b/Engine/Source/Runtime/Engine/Private/AI/Navigation/NavigationSystem.cpp
--- a/Engine/Source/Runtime/Engine/Private/AI/Navigation/NavigationSystem.cpp
+++ b/Engine/Source/Runtime/Engine/Private/AI/Navigation/NavigationSystem.cpp
@@ -2059,6 +2059,21 @@ void UNavigationSystem::UpdateNavOctreeAll(class AActor* Actor)
 	}
 }
 
+void UNavigationSystem::UpdateNavOctreeBounds(class AActor* Actor)
+{
+	TArray<UActorComponent*> Components;
+	Actor->GetComponents(Components);
+
+	for (int32 Idx = 0; Idx < Components.Num(); Idx++)
+	{
+		INavRelevantInterface* NavElement = Cast<INavRelevantInterface>(Components[Idx]);
+		if (NavElement)
+		{
+			NavElement->UpdateNavigationBounds();
+		}
+	}
+}
+
 void UNavigationSystem::ClearNavOctreeAll(class AActor* Actor)
 {
 	if (Actor)
diff --git a/Engine/Source/Runtime/Engine/Private/ActorEditor.cpp b/Engine/Source/Runtime/Engine/Private/ActorEditor.cpp
--- a/Engine/Source/Runtime/Engine/Private/ActorEditor.cpp
+++ b/Engine/Source/Runtime/Engine/Private/ActorEditor.cpp
@@ -125,6 +125,7 @@ void AActor::PostEditMove(bool bFinished)
 	{
 		// update actor and all its components in navigation system after finishing move
 		// USceneComponent::UpdateNavigationData works only in game world
+		UNavigationSystem::UpdateNavOctreeBounds(this);
 		UNavigationSystem::UpdateNavOctreeAll(this);
 	}
 }
