fixed updating area of segment links implementing INavLinkCustomInterface

diff --git a/Engine/Source/Runtime/Engine/Private/AI/Navigation/RecastNavMesh.cpp b/Engine/Source/Runtime/Engine/Private/AI/Navigation/RecastNavMesh.cpp
--- a/Engine/Source/Runtime/Engine/Private/AI/Navigation/RecastNavMesh.cpp
+++ b/Engine/Source/Runtime/Engine/Private/AI/Navigation/RecastNavMesh.cpp
@@ -966,7 +966,16 @@ NavNodeRef ARecastNavMesh::FindNearestPoly(FVector const& Loc, FVector const& Ex
 
 void ARecastNavMesh::UpdateCustomLink(const class INavLinkCustomInterface* CustomLink)
 {
-	UpdateNavigationLinkArea(CustomLink->GetLinkId(), CustomLink->GetLinkAreaClass());
+	TSubclassOf<UNavArea> AreaClass = CustomLink->GetLinkAreaClass();
+	const int32 UserId = CustomLink->GetLinkId();
+	const int32 AreaId = GetAreaID(AreaClass);
+	if (AreaId >= 0 && RecastNavMeshImpl)
+	{
+		UNavArea* DefArea = (UNavArea*)(AreaClass->GetDefaultObject());
+
+		RecastNavMeshImpl->UpdateNavigationLinkArea(UserId, AreaId, DefArea->GetAreaFlags());
+		RecastNavMeshImpl->UpdateSegmentLinkArea(UserId, AreaId, DefArea->GetAreaFlags());
+	}
 }
 
 void ARecastNavMesh::UpdateNavigationLinkArea(int32 UserId, TSubclassOf<class UNavArea> AreaClass) const
