Fix collision on instance indices >128 on movable InstancedStaticMeshComponents

diff --git a/Engine/Source/Runtime/Engine/Classes/Components/InstancedStaticMeshComponent.h b/Engine/Source/Runtime/Engine/Classes/Components/InstancedStaticMeshComponent.h
--- a/Engine/Source/Runtime/Engine/Classes/Components/InstancedStaticMeshComponent.h
+++ b/Engine/Source/Runtime/Engine/Classes/Components/InstancedStaticMeshComponent.h
@@ -122,7 +122,7 @@ public:
 
 #if WITH_PHYSX
 	/** Aggregate physx representation of the instances' bodies. */
-	class physx::PxAggregate* Aggregate;
+	TArray<physx::PxAggregate*> Aggregates;
 #endif	//WITH_PHYSX
 
 	/** Physics representation of the instance bodies */
diff --git a/Engine/Source/Runtime/Engine/Private/InstancedStaticMesh.cpp b/Engine/Source/Runtime/Engine/Private/InstancedStaticMesh.cpp
--- a/Engine/Source/Runtime/Engine/Private/InstancedStaticMesh.cpp
+++ b/Engine/Source/Runtime/Engine/Private/InstancedStaticMesh.cpp
@@ -1418,7 +1418,10 @@ void UInstancedStaticMeshComponent::InitInstanceBody(int32 InstanceIdx, FBodyIns
 
 #if WITH_PHYSX
 	// Create physics body instance.
-	InstanceBodyInstance->InitBody( BodySetup, InstanceTransform, this, GetWorld()->GetPhysicsScene(), Aggregate);
+	// Aggregates aren't used for static objects
+	auto* Aggregate = (Mobility == EComponentMobility::Movable) ? Aggregates[FMath::DivideAndRoundDown<int32>(InstanceIdx, AggregateMaxSize)] : nullptr;
+	check(Mobility != EComponentMobility::Movable || Aggregate->getNbActors() < Aggregate->getMaxNbActors());
+	InstanceBodyInstance->InitBody(BodySetup, InstanceTransform, this, GetWorld()->GetPhysicsScene(), Aggregate);
 #endif //WITH_PHYSX
 }
 
@@ -1456,12 +1459,22 @@ void UInstancedStaticMeshComponent::CreatePhysicsState()
 	if (!PhysScene) { return; }
 
 #if WITH_PHYSX
-	check(Aggregate == NULL);
-	Aggregate = GPhysXSDK->createAggregate(AggregateMaxSize, false);
+	check(Aggregates.Num() == 0);
+
+	const int32 NumBodies = PerInstanceSMData.Num();
+
+	// Aggregates aren't used for static objects
+	const int32 NumAggregates = (Mobility == EComponentMobility::Movable) ? FMath::DivideAndRoundUp<int32>(NumBodies, AggregateMaxSize) : 0;
 
 	// Get the scene type from the main BodyInstance
 	const uint32 SceneType = BodyInstance.UseAsyncScene() ? PST_Async : PST_Sync;
-	PhysScene->GetPhysXScene(SceneType)->addAggregate(*Aggregate);
+
+	for (int32 i = 0; i < NumAggregates; i++)
+	{
+		auto* Aggregate = GPhysXSDK->createAggregate(AggregateMaxSize, false);
+		Aggregates.Add(Aggregate);
+		PhysScene->GetPhysXScene(SceneType)->addAggregate(*Aggregate);
+	}
 #endif
 
 	// Create all the bodies.
@@ -1478,13 +1491,13 @@ void UInstancedStaticMeshComponent::DestroyPhysicsState()
 	ClearAllInstanceBodies();
 
 #if WITH_PHYSX
-	// releasing Aggregate, it shouldn't contain any Bodies now, because they are released above
-	if(Aggregate)
+	// releasing Aggregates, they shouldn't contain any Bodies now, because they are released above
+	for (auto* Aggregate : Aggregates)
 	{
 		check(!Aggregate->getNbActors());
 		Aggregate->release();
-		Aggregate = NULL;
 	}
+	Aggregates.Empty();
 #endif //WITH_PHYSX
 }
 
@@ -1940,9 +1953,28 @@ void UInstancedStaticMeshComponent::SetupNewInstanceData(FInstancedStaticMeshIns
 
 	if (bPhysicsStateCreated)
 	{
+		// Add another aggregate if needed
+		// Aggregates aren't used for static objects
+		if (Mobility == EComponentMobility::Movable)
+		{
+			const int32 AggregateIndex = FMath::DivideAndRoundDown<int32>(InInstanceIndex, AggregateMaxSize);
+			if (AggregateIndex >= Aggregates.Num())
+			{
+				// Get the scene type from the main BodyInstance
+				const uint32 SceneType = BodyInstance.UseAsyncScene() ? PST_Async : PST_Sync;
+
+				FPhysScene* PhysScene = GetWorld()->GetPhysicsScene();
+				check(PhysScene);
+
+				auto* Aggregate = GPhysXSDK->createAggregate(AggregateMaxSize, false);
+				const int32 AddedIndex = Aggregates.Add(Aggregate);
+				check(AddedIndex == AggregateIndex);
+				PhysScene->GetPhysXScene(SceneType)->addAggregate(*Aggregate);
+			}
+		}
+
 		FBodyInstance* NewBodyInstance = new FBodyInstance();
 		int32 BodyIndex = InstanceBodies.Insert(NewBodyInstance, InInstanceIndex);
-
 		check(InInstanceIndex == BodyIndex);
 		InitInstanceBody(BodyIndex, NewBodyInstance);
 	}
