Allow vehicles to perform RVO avoidance (GitHub pull request 465)

diff --git a/Engine/Source/Runtime/Engine/Classes/AI/Navigation/AvoidanceManager.h b/Engine/Source/Runtime/Engine/Classes/AI/Navigation/AvoidanceManager.h
--- a/Engine/Source/Runtime/Engine/Classes/AI/Navigation/AvoidanceManager.h
+++ b/Engine/Source/Runtime/Engine/Classes/AI/Navigation/AvoidanceManager.h
@@ -1,10 +1,13 @@
 // Copyright 1998-2014 Epic Games, Inc. All Rights Reserved.
 
 #pragma once
+#include "AI/RVOAvoidanceInterface.h"
 #include "AvoidanceManager.generated.h"
 
 DECLARE_CYCLE_STAT_EXTERN(TEXT("Avoidance Time"),STAT_AI_ObstacleAvoidance,STATGROUP_AI, );
 
+class UAvoidanceManager;
+
 USTRUCT()
 struct FNavAvoidanceData
 {
@@ -40,9 +43,17 @@ struct FNavAvoidanceData
 	/** Do NOT avoid agents is they belong to one of specified groups, takes priority over GroupsToAvoid */
 	int32 GroupsToIgnore;
 
+	/** Radius of the area to consider for avoidance */
+	float TestRadius2D;
+
+	FNavAvoidanceData() {}
+	FNavAvoidanceData(UAvoidanceManager* Manager, IRVOAvoidanceInterface* AvoidanceComp);
+
 	/** Init function for internal use to guard against data changes not being reflected in blueprint-accessible creation functions */
 	void Init(class UAvoidanceManager* Avoidance, const FVector& InCenter, float InRadius, float InHeight,
-		const FVector& InVelocity, float InWeight, int32 InGroupMask, int32 InGroupsToAvoid, int32 InGroupsToIgnore);
+		const FVector& InVelocity, float InWeight = 0.5f,
+		int32 InGroupMask = 1, int32 InGroupsToAvoid = 0xffffffff, int32 InGroupsToIgnore = 0,
+		float InTestRadius2D = 500.0f);
 
 	FORCEINLINE bool ShouldBeIgnored() const
 	{
@@ -85,10 +96,6 @@ class ENGINE_API UAvoidanceManager : public UObject, public FSelfRegisteringExec
 	UPROPERTY(EditAnywhere, Category="Avoidance", config, meta=(ClampMin = "0.0"))
 	float ArtificialRadiusExpansion;
 
-	/** Test against obstacles within given radius from moving agent */
-	UPROPERTY(EditAnywhere, Category="Avoidance", config, meta=(ClampMin = "0.0"))
-	float TestRadius2D;
-
 	/** Test against obstacles within given height difference from moving agent */
 	UPROPERTY(EditAnywhere, Category="Avoidance", config, meta=(ClampMin = "0.0"))
 	float TestHeightDifference;
@@ -105,12 +112,16 @@ class ENGINE_API UAvoidanceManager : public UObject, public FSelfRegisteringExec
 	 *  @param AvoidanceWeight	When avoiding each other, actors divert course in proportion to their relative weights. Range is 0.0 to 1.0. Special: at 1.0, actor will not divert course at all.
 	 */
 	UFUNCTION(BlueprintCallable, Category="AI")
-	bool RegisterMovementComponent(class UCharacterMovementComponent* MovementComp, float AvoidanceWeight = 0.5f);
+	bool RegisterMovementComponent(class UMovementComponent* MovementComp, float AvoidanceWeight = 0.5f);
 
 	/** Get your latest data. */
 	FNavAvoidanceData* GetAvoidanceObjectForUID(int32 AvoidanceUID) { return AvoidanceObjects.Find(AvoidanceUID); }
 	const FNavAvoidanceData* GetAvoidanceObjectForUID(int32 AvoidanceUID) const { return AvoidanceObjects.Find(AvoidanceUID); }
 
+	/** Calculate avoidance velocity for component (avoids collisions with the supplied component) */
+	UFUNCTION(BlueprintCallable, Category="AI")
+	FVector GetAvoidanceVelocityForComponent(UMovementComponent* MovementComp);
+
 	/** Only use if you want manual velocity planning. Provide your AvoidanceUID in order to avoid colliding with yourself. */
 	UFUNCTION(BlueprintCallable, Category="AI")
 	FVector GetAvoidanceVelocityIgnoringUID(const FNavAvoidanceData& AvoidanceData, float DeltaTime, int32 IgnoreThisUID);
@@ -119,8 +130,8 @@ class ENGINE_API UAvoidanceManager : public UObject, public FSelfRegisteringExec
 	UFUNCTION(BlueprintCallable, Category="AI")
 	FVector GetAvoidanceVelocity(const FNavAvoidanceData& AvoidanceData, float DeltaTime);
 
-	/** Inform the avoidance manager of your current information, using the ID you were given as a key. */
-	void UpdateRVO(int32 AvoidanceUID, FVector Center, float Radius, float Height, FVector Velocity, float Weight = 0.5f, int32 GroupMask = 1, int32 AvoidMask = 0xFFFFFFFF, int32 IgnoreMask = 0);
+	/** Update the RVO avoidance data for the participating UMovementComponent */
+	void UpdateRVO(UMovementComponent* MovementComp);
 
 	/** For Duration seconds, set this object to ignore all others. */
 	void OverrideToMaxWeight(int32 AvoidanceUID, float Duration);
diff --git a/Engine/Source/Runtime/Engine/Classes/GameFramework/CharacterMovementComponent.h b/Engine/Source/Runtime/Engine/Classes/GameFramework/CharacterMovementComponent.h
--- a/Engine/Source/Runtime/Engine/Classes/GameFramework/CharacterMovementComponent.h
+++ b/Engine/Source/Runtime/Engine/Classes/GameFramework/CharacterMovementComponent.h
@@ -2,6 +2,7 @@
 
 #pragma once
 #include "AI/Navigation/NavigationAvoidanceTypes.h"
+#include "AI/RVOAvoidanceInterface.h"
 #include "Animation/AnimationAsset.h"
 #include "Engine/EngineBaseTypes.h"
 #include "Engine/EngineTypes.h"
@@ -122,7 +123,7 @@ struct FCharacterMovementComponentPreClothTickFunction : public FTickFunction
  */
 
 UCLASS()
-class ENGINE_API UCharacterMovementComponent : public UPawnMovementComponent, public INetworkPredictionInterface
+class ENGINE_API UCharacterMovementComponent : public UPawnMovementComponent, public IRVOAvoidanceInterface, public INetworkPredictionInterface
 {
 	GENERATED_BODY()
 public:
@@ -622,6 +623,9 @@ public:
 	UPROPERTY(Category="Avoidance", EditAnywhere, BlueprintReadOnly)
 	uint32 bUseRVOAvoidance:1;
 
+	UPROPERTY(Category="Avoidance", EditAnywhere, BlueprintReadOnly)
+	float AvoidanceConsiderationRadius;
+	
 	/**
 	 * Should use acceleration for path following?
 	 * If true, acceleration is applied when path following to reach the target velocity.
@@ -1644,6 +1648,21 @@ protected:
 	/** lock avoidance velocity */
 	void SetAvoidanceVelocityLock(class UAvoidanceManager* Avoidance, float Duration);
 
+	/** BEGIN IRVOAvoidanceInterface */
+	virtual void SetRVOAvoidanceUID(int32 UID) override;
+	virtual int32 GetRVOAvoidanceUID() override;
+	virtual void SetRVOAvoidanceWeight(float Weight) override;
+	virtual float GetRVOAvoidanceWeight() override;
+	virtual FVector GetRVOAvoidanceOrigin() override;
+	virtual float GetRVOAvoidanceRadius() override;
+	virtual float GetRVOAvoidanceHeight() override;
+	virtual float GetRVOAvoidanceConsiderationRadius() override;
+	virtual FVector GetVelocityForRVOConsideration() override;
+	virtual int32 GetAvoidanceGroupMask() override;
+	virtual int32 GetGroupsToAvoidMask() override;
+	virtual int32 GetGroupsToIgnoreMask() override;
+	/** END IRVOAvoidanceInterface */
+
 public:
 
 	/** Minimum delta time considered when ticking. Delta times below this are not considered. This is a very small non-zero value to avoid potential divide-by-zero in simulation code. */
diff --git a/Engine/Source/Runtime/Engine/Classes/Vehicles/WheeledVehicleMovementComponent.h b/Engine/Source/Runtime/Engine/Classes/Vehicles/WheeledVehicleMovementComponent.h
--- a/Engine/Source/Runtime/Engine/Classes/Vehicles/WheeledVehicleMovementComponent.h
+++ b/Engine/Source/Runtime/Engine/Classes/Vehicles/WheeledVehicleMovementComponent.h
@@ -2,6 +2,8 @@
 
 #pragma once
 #include "GameFramework/PawnMovementComponent.h"
+#include "AI/Navigation/NavigationAvoidanceTypes.h"
+#include "AI/RVOAvoidanceInterface.h"
 #include "WheeledVehicleMovementComponent.generated.h"
 
 class UVehicleWheel;
@@ -157,7 +159,7 @@ struct FVehicleInputRate
  * Component to handle the vehicle simulation for an actor.
  */
 UCLASS(Abstract, hidecategories=(PlanarMovement, "Components|Movement|Planar", Activation, "Components|Activation"))
-class ENGINE_API UWheeledVehicleMovementComponent : public UPawnMovementComponent
+class ENGINE_API UWheeledVehicleMovementComponent : public UPawnMovementComponent, public IRVOAvoidanceInterface
 {
 	GENERATED_UCLASS_BODY()
 
@@ -336,6 +338,72 @@ class ENGINE_API UWheeledVehicleMovementComponent : public UPawnMovementComponen
 	UFUNCTION(BlueprintCallable, Category="Game|Components|WheeledVehicleMovement")
 	bool GetUseAutoGears() const;
 
+	// RVO Avoidance
+
+	/** If set, component will use RVO avoidance */
+	UPROPERTY(Category = "Avoidance", EditAnywhere, BlueprintReadWrite)
+	uint32 bUseRVOAvoidance : 1;
+	
+	/** Vehicle Radius to use for RVO avoidance (usually half of vehicle width) */
+	UPROPERTY(Category = "Avoidance", EditAnywhere, BlueprintReadWrite)
+	float RVOAvoidanceRadius;
+	
+	/** Vehicle Height to use for RVO avoidance (usually vehicle height) */
+	UPROPERTY(Category = "Avoidance", EditAnywhere, BlueprintReadWrite)
+	float RVOAvoidanceHeight;
+	
+	/** Area Radius to consider for RVO avoidance */
+	UPROPERTY(Category = "Avoidance", EditAnywhere, BlueprintReadWrite)
+	float AvoidanceConsiderationRadius;
+
+	/** Value by which to alter steering per frame based on calculated avoidance */
+	UPROPERTY(Category = "Avoidance", EditAnywhere, BlueprintReadWrite, meta = (ClampMin = "0.0", UIMin = "0.0", ClampMax = "1.0", UIMax = "1.0"))
+	float RVOSteeringStep;
+
+	/** Value by which to alter throttle per frame based on calculated avoidance */
+	UPROPERTY(Category = "Avoidance", EditAnywhere, BlueprintReadWrite, meta = (ClampMin = "0.0", UIMin = "0.0", ClampMax = "1.0", UIMax = "1.0"))
+	float RVOThrottleStep;
+	
+	/** calculate RVO avoidance and apply it to current velocity */
+	virtual void CalculateAvoidanceVelocity(float DeltaTime);
+
+	/** No default value, for now it's assumed to be valid if GetAvoidanceManager() returns non-NULL. */
+	UPROPERTY(Category = "Avoidance", VisibleAnywhere, BlueprintReadOnly, AdvancedDisplay)
+	int32 AvoidanceUID;
+
+	/** Moving actor's group mask */
+	UPROPERTY(Category = "Avoidance", EditAnywhere, BlueprintReadOnly, AdvancedDisplay)
+	FNavAvoidanceMask AvoidanceGroup;
+	
+	UFUNCTION(BlueprintCallable, Category = "Pawn|Components|WheeledVehicleMovement")
+	void SetAvoidanceGroup(int32 GroupFlags);
+	
+	/** Will avoid other agents if they are in one of specified groups */
+	UPROPERTY(Category = "Avoidance", EditAnywhere, BlueprintReadOnly, AdvancedDisplay)
+	FNavAvoidanceMask GroupsToAvoid;
+
+	UFUNCTION(BlueprintCallable, Category = "Pawn|Components|WheeledVehicleMovement")
+	void SetGroupsToAvoid(int32 GroupFlags);
+	
+	/** Will NOT avoid other agents if they are in one of specified groups, higher priority than GroupsToAvoid */
+	UPROPERTY(Category = "Avoidance", EditAnywhere, BlueprintReadOnly, AdvancedDisplay)
+	FNavAvoidanceMask GroupsToIgnore;
+
+	UFUNCTION(BlueprintCallable, Category = "Pawn|Components|WheeledVehicleMovement")
+	void SetGroupsToIgnore(int32 GroupFlags);
+
+	/** De facto default value 0.5 (due to that being the default in the avoidance registration function), indicates RVO behavior. */
+	UPROPERTY(Category = "Avoidance", EditAnywhere, BlueprintReadOnly)
+	float AvoidanceWeight;
+	
+	/** Temporarily holds launch velocity when pawn is to be launched so it happens at end of movement. */
+	UPROPERTY()
+	FVector PendingLaunchVelocity;
+	
+	/** Change avoidance state and register with RVO manager if necessary */
+	UFUNCTION(BlueprintCallable, Category = "Pawn|Components|WheeledVehicleMovement")
+	void SetAvoidanceEnabled(bool bEnable);
+
 protected:
 
 	// replicated state of vehicle 
@@ -433,6 +501,43 @@ protected:
 	UFUNCTION(reliable, server, WithValidation)
 	void ServerUpdateState(float InSteeringInput, float InThrottleInput, float InBrakeInput, float InHandbrakeInput, int32 CurrentGear);
 
+	/** Update RVO Avoidance for simulation */
+	void UpdateAvoidance(float DeltaTime);
+		
+	/** called in Tick to update data in RVO avoidance manager */
+	void UpdateDefaultAvoidance();
+	
+	/** lock avoidance velocity */
+	void SetAvoidanceVelocityLock(class UAvoidanceManager* Avoidance, float Duration);
+	
+	/** Was avoidance updated in this frame? */
+	UPROPERTY(Transient)
+	uint32 bWasAvoidanceUpdated : 1;
+	
+	/** Calculated avoidance velocity used to adjust steering and throttle */
+	FVector AvoidanceVelocity;
+	
+	/** forced avoidance velocity, used when AvoidanceLockTimer is > 0 */
+	FVector AvoidanceLockVelocity;
+	
+	/** remaining time of avoidance velocity lock */
+	float AvoidanceLockTimer;
+	
+	/** BEGIN IRVOAvoidanceInterface */
+	virtual void SetRVOAvoidanceUID(int32 UID) override;
+	virtual int32 GetRVOAvoidanceUID() override;
+	virtual void SetRVOAvoidanceWeight(float Weight) override;
+	virtual float GetRVOAvoidanceWeight() override;
+	virtual FVector GetRVOAvoidanceOrigin() override;
+	virtual float GetRVOAvoidanceRadius() override;
+	virtual float GetRVOAvoidanceHeight() override;
+	virtual float GetRVOAvoidanceConsiderationRadius() override;
+	virtual FVector GetVelocityForRVOConsideration() override;
+	virtual int32 GetAvoidanceGroupMask() override;
+	virtual int32 GetGroupsToAvoidMask() override;
+	virtual int32 GetGroupsToIgnoreMask() override;
+	/** END IRVOAvoidanceInterface */
+
 #if WITH_VEHICLE
 
 	int32 GearToPhysXGear(const int32 Gear) const;
@@ -510,4 +615,4 @@ inline float M2ToCm2(float M2)
 inline float Cm2ToM2(float Cm2)
 {
 	return Cm2 / (100.f * 100.f);
-}
\ No newline at end of file
+}
diff --git a/Engine/Source/Runtime/Engine/Private/AI/Navigation/AvoidanceManager.cpp b/Engine/Source/Runtime/Engine/Private/AI/Navigation/AvoidanceManager.cpp
--- a/Engine/Source/Runtime/Engine/Private/AI/Navigation/AvoidanceManager.cpp
+++ b/Engine/Source/Runtime/Engine/Private/AI/Navigation/AvoidanceManager.cpp
@@ -2,13 +2,29 @@
 
 #include "EnginePrivate.h"
 #include "AI/Navigation/AvoidanceManager.h"
-#include "GameFramework/Character.h"
-#include "GameFramework/CharacterMovementComponent.h"
+#include "AI/RVOAvoidanceInterface.h"
 
 DEFINE_STAT(STAT_AI_ObstacleAvoidance);
 
+FNavAvoidanceData::FNavAvoidanceData(UAvoidanceManager* Manager, IRVOAvoidanceInterface* AvoidanceComp)
+{
+	Init(Manager,
+		AvoidanceComp->GetRVOAvoidanceOrigin(),
+		AvoidanceComp->GetRVOAvoidanceRadius(),
+		AvoidanceComp->GetRVOAvoidanceHeight(),
+		AvoidanceComp->GetVelocityForRVOConsideration(),
+		AvoidanceComp->GetRVOAvoidanceWeight(),
+		AvoidanceComp->GetAvoidanceGroupMask(),
+		AvoidanceComp->GetGroupsToAvoidMask(),
+		AvoidanceComp->GetGroupsToIgnoreMask(),
+		AvoidanceComp->GetRVOAvoidanceConsiderationRadius()
+		);
+}
+
 void FNavAvoidanceData::Init(class UAvoidanceManager* Avoidance, const FVector& InCenter, float InRadius, float InHeight,
-							 const FVector& InVelocity, float InWeight, int32 InGroupMask, int32 InGroupsToAvoid, int32 InGroupsToIgnore)
+							 const FVector& InVelocity, float InWeight,
+							 int32 InGroupMask, int32 InGroupsToAvoid, int32 InGroupsToIgnore,
+							 float InTestRadius2D)
 {
 	Center = InCenter;
 	Velocity = InVelocity;
@@ -20,6 +36,7 @@ void FNavAvoidanceData::Init(class UAvoidanceManager* Avoidance, const FVector&
 	GroupsToIgnore = InGroupsToIgnore;
 	OverrideWeightTime = 0.0f;
 	RemainingTimeToLive = Avoidance->DefaultTimeToLive;
+	TestRadius2D = InTestRadius2D;
 
 	//RickH - This is kind of a hack, but ultimately this is going to be a 2D solution with 3D culling/broad-phase anyway.
 	Velocity.Z = 0.0f;
@@ -36,7 +53,6 @@ UAvoidanceManager::UAvoidanceManager(const FObjectInitializer& ObjectInitializer
 	LockTimeAfterClean = 0.01f;
 	DeltaTimeToPredict = 0.5f;
 	ArtificialRadiusExpansion = 1.5f;
-	TestRadius2D = 500.0f;
 	TestHeightDifference = 500.0f;
 	bRequestedUpdateTimer = false;
 
@@ -111,28 +127,36 @@ int32 UAvoidanceManager::GetNewAvoidanceUID()
 	return NewUID;
 }
 
-bool UAvoidanceManager::RegisterMovementComponent(UCharacterMovementComponent* MovementComp, float AvoidanceWeight)
+bool UAvoidanceManager::RegisterMovementComponent(UMovementComponent* MovementComp, float AvoidanceWeight)
 {
-	if (MovementComp && MovementComp->GetCharacterOwner())
+	if (IRVOAvoidanceInterface* AvoidingComp = Cast<IRVOAvoidanceInterface>(MovementComp))
 	{
-		if (UCapsuleComponent *OurCapsule = MovementComp->GetCharacterOwner()->GetCapsuleComponent())
-		{
-			const int32 NewAvoidanceUID = GetNewAvoidanceUID();
-			MovementComp->AvoidanceUID = NewAvoidanceUID;
-			MovementComp->AvoidanceWeight = AvoidanceWeight;
+		const int32 NewAvoidanceUID = GetNewAvoidanceUID();
+		AvoidingComp->SetRVOAvoidanceUID(NewAvoidanceUID);
+		AvoidingComp->SetRVOAvoidanceWeight(AvoidanceWeight);
 
-			RequestUpdateTimer();
-			UpdateRVO(NewAvoidanceUID, MovementComp->GetActorFeetLocation(),
-				OurCapsule->GetScaledCapsuleRadius(), OurCapsule->GetScaledCapsuleHalfHeight(),
-				MovementComp->Velocity, AvoidanceWeight,
-				MovementComp->AvoidanceGroup.Packed, MovementComp->GroupsToAvoid.Packed, MovementComp->GroupsToIgnore.Packed);
+		RequestUpdateTimer();
 
-			return true;
-		}
+		FNavAvoidanceData AvoidanceData(this, AvoidingComp);
+		UpdateRVO_Internal(AvoidingComp->GetRVOAvoidanceUID(), AvoidanceData);
+
+		return true;
 	}
+
 	return false;
 }
 
+FVector UAvoidanceManager::GetAvoidanceVelocityForComponent(UMovementComponent* MovementComp)
+{
+	if (IRVOAvoidanceInterface* AvoidingComp = Cast<IRVOAvoidanceInterface>(MovementComp))
+	{
+		FNavAvoidanceData AvoidanceData(this, AvoidingComp);
+		return GetAvoidanceVelocityIgnoringUID(AvoidanceData, DeltaTimeToPredict, AvoidingComp->GetRVOAvoidanceUID());
+	}
+
+	return FVector::ZeroVector;
+}
+
 FVector UAvoidanceManager::GetAvoidanceVelocityIgnoringUID(const FNavAvoidanceData& inAvoidanceData, float DeltaTime, int32 inIgnoreThisUID)
 {
 	return GetAvoidanceVelocity_Internal(inAvoidanceData, DeltaTime, &inIgnoreThisUID);
@@ -143,12 +167,13 @@ FVector UAvoidanceManager::GetAvoidanceVelocity(const FNavAvoidanceData& inAvoid
 	return GetAvoidanceVelocity_Internal(inAvoidanceData, DeltaTime);
 }
 
-void UAvoidanceManager::UpdateRVO(int32 inAvoidanceUID, FVector inCenter, float inRadius, float inHeight, FVector inVelocity, float inWeight, int32 GroupMask, int32 AvoidMask, int32 IgnoreMask)
+void UAvoidanceManager::UpdateRVO(UMovementComponent* MovementComp)
 {
-	FNavAvoidanceData NewAvoidanceData;
-	NewAvoidanceData.Init(this, inCenter, inRadius, inHeight, inVelocity, inWeight, GroupMask, AvoidMask, IgnoreMask);
-
-	UpdateRVO_Internal(inAvoidanceUID, NewAvoidanceData);
+	if (IRVOAvoidanceInterface* AvoidingComp = Cast<IRVOAvoidanceInterface>(MovementComp))
+	{
+		FNavAvoidanceData NewAvoidanceData(this, AvoidingComp);
+		UpdateRVO_Internal(AvoidingComp->GetRVOAvoidanceUID(), NewAvoidanceData);
+	}
 }
 
 void UAvoidanceManager::UpdateRVO_Internal(int32 inAvoidanceUID, const FNavAvoidanceData& inAvoidanceData)
@@ -301,7 +326,7 @@ FVector UAvoidanceManager::GetAvoidanceVelocity_Internal(const FNavAvoidanceData
 
 		//RickH - We should have a max-radius parameter/option here, so I'm just going to hardcode one for now.
 		//if ((OtherObject.Radius + _AvoidanceData.Radius + MaxSpeed + OtherObject.Velocity.Size2D()) < FVector::Dist(OtherObject.Center, _AvoidanceData.Center))
-		if (FVector2D(OtherObject.Center - inAvoidanceData.Center).SizeSquared() > FMath::Square(TestRadius2D))
+		if (FVector2D(OtherObject.Center - inAvoidanceData.Center).SizeSquared() > FMath::Square(inAvoidanceData.TestRadius2D))
 		{
 			continue;
 		}
diff --git a/Engine/Source/Runtime/Engine/Private/AI/Navigation/NavigationInterfaces.cpp b/Engine/Source/Runtime/Engine/Private/AI/Navigation/NavigationInterfaces.cpp
--- a/Engine/Source/Runtime/Engine/Private/AI/Navigation/NavigationInterfaces.cpp
+++ b/Engine/Source/Runtime/Engine/Private/AI/Navigation/NavigationInterfaces.cpp
@@ -8,9 +8,15 @@
 #include "AI/Navigation/NavPathObserverInterface.h"
 #include "AI/Navigation/NavRelevantInterface.h"
 #include "AI/Navigation/NavLinkCustomInterface.h"
+#include "AI/RVOAvoidanceInterface.h"
 
 uint32 INavLinkCustomInterface::NextUniqueId = 1;
 
+URVOAvoidanceInterface::URVOAvoidanceInterface(const FObjectInitializer& ObjectInitializer)
+	: Super(ObjectInitializer)
+{
+}
+
 UNavAgentInterface::UNavAgentInterface(const FObjectInitializer& ObjectInitializer)
 	: Super(ObjectInitializer)
 {
@@ -32,7 +38,7 @@ UNavLinkHostInterface::UNavLinkHostInterface(const FObjectInitializer& ObjectIni
 }
 
 UNavLinkCustomInterface::UNavLinkCustomInterface(const FObjectInitializer& ObjectInitializer)
-	: Super(ObjectInitializer)
+: Super(ObjectInitializer)
 {
 }
 
diff --git a/Engine/Source/Runtime/Engine/Private/CharacterMovementComponent.cpp b/Engine/Source/Runtime/Engine/Private/CharacterMovementComponent.cpp
--- a/Engine/Source/Runtime/Engine/Private/CharacterMovementComponent.cpp
+++ b/Engine/Source/Runtime/Engine/Private/CharacterMovementComponent.cpp
@@ -253,6 +253,7 @@ UCharacterMovementComponent::UCharacterMovementComponent(const FObjectInitialize
 	AvoidanceGroup.bGroup0 = true;
 	GroupsToAvoid.Packed = 0xFFFFFFFF;
 	GroupsToIgnore.Packed = 0;
+	AvoidanceConsiderationRadius = 500.0f;
 
 	OldBaseQuat = FQuat::Identity;
 	OldBaseLocation = FVector::ZeroVector;
@@ -2314,12 +2315,7 @@ void UCharacterMovementComponent::CalcAvoidanceVelocity(float DeltaTime)
 		}
 		else
 		{
-			FNavAvoidanceData CurrentData;
-			CurrentData.Init(AvoidanceManager, GetActorFeetLocation(),
-				OurCapsule->GetScaledCapsuleRadius(), OurCapsule->GetScaledCapsuleHalfHeight(),
-				Velocity, AvoidanceWeight, AvoidanceGroup.Packed, GroupsToAvoid.Packed, GroupsToIgnore.Packed);
-
-			FVector NewVelocity = AvoidanceManager->GetAvoidanceVelocityIgnoringUID(CurrentData, AvoidanceManager->DeltaTimeToPredict, AvoidanceUID);
+			FVector NewVelocity = AvoidanceManager->GetAvoidanceVelocityForComponent(this);
 			if (bUseRVOPostProcess)
 			{
 				PostProcessAvoidanceVelocity(NewVelocity);
@@ -2350,9 +2346,7 @@ void UCharacterMovementComponent::CalcAvoidanceVelocity(float DeltaTime)
 			}
 		}
 		//RickH - We might do better to do this later in our update
-		AvoidanceManager->UpdateRVO(AvoidanceUID, GetActorFeetLocation(),
-			OurCapsule->GetScaledCapsuleRadius(), OurCapsule->GetScaledCapsuleHalfHeight() * 2.0f,
-			Velocity, AvoidanceWeight, AvoidanceGroup.Packed, GroupsToAvoid.Packed, GroupsToIgnore.Packed);
+		AvoidanceManager->UpdateRVO(this);
 
 		bWasAvoidanceUpdated = true;
 	}
@@ -2389,9 +2383,7 @@ void UCharacterMovementComponent::UpdateDefaultAvoidance()
 	{
 		if (UCapsuleComponent *OurCapsule = GetCharacterOwner()->GetCapsuleComponent())
 		{
-			AvoidanceManager->UpdateRVO(AvoidanceUID, GetActorFeetLocation(),
-				OurCapsule->GetScaledCapsuleRadius(), OurCapsule->GetScaledCapsuleHalfHeight() * 2.0f,
-				Velocity, AvoidanceWeight, AvoidanceGroup.Packed, GroupsToAvoid.Packed, GroupsToIgnore.Packed);
+			AvoidanceManager->UpdateRVO(this);
 
 			//Consider this a clean move because we didn't even try to avoid.
 			SetAvoidanceVelocityLock(AvoidanceManager, AvoidanceManager->LockTimeAfterClean);
@@ -2401,6 +2393,68 @@ void UCharacterMovementComponent::UpdateDefaultAvoidance()
 	bWasAvoidanceUpdated = false;		//Reset for next frame
 }
 
+void UCharacterMovementComponent::SetRVOAvoidanceUID(int32 UID)
+{
+	AvoidanceUID = UID;
+}
+
+int32 UCharacterMovementComponent::GetRVOAvoidanceUID()
+{
+	return AvoidanceUID;
+}
+
+void UCharacterMovementComponent::SetRVOAvoidanceWeight(float Weight)
+{
+	AvoidanceWeight = Weight;
+}
+
+float UCharacterMovementComponent::GetRVOAvoidanceWeight()
+{
+	return AvoidanceWeight;
+}
+
+FVector UCharacterMovementComponent::GetRVOAvoidanceOrigin()
+{
+	return GetActorFeetLocation();
+}
+
+float UCharacterMovementComponent::GetRVOAvoidanceRadius()
+{
+	UCapsuleComponent* CapsuleComp = GetCharacterOwner()->GetCapsuleComponent();
+	return CapsuleComp ? CapsuleComp->GetScaledCapsuleRadius() : 0.0f;
+}
+
+float UCharacterMovementComponent::GetRVOAvoidanceConsiderationRadius()
+{
+	return AvoidanceConsiderationRadius;
+}
+
+float UCharacterMovementComponent::GetRVOAvoidanceHeight()
+{
+	UCapsuleComponent* CapsuleComp = GetCharacterOwner()->GetCapsuleComponent();
+	return CapsuleComp ? CapsuleComp->GetScaledCapsuleHalfHeight() : 0.0f;
+}
+
+FVector UCharacterMovementComponent::GetVelocityForRVOConsideration()
+{
+	return Velocity;
+}
+
+int32 UCharacterMovementComponent::GetAvoidanceGroupMask()
+{
+	return AvoidanceGroup.Packed;
+}
+
+int32 UCharacterMovementComponent::GetGroupsToAvoidMask()
+{
+	return GroupsToAvoid.Packed;
+}
+
+int32 UCharacterMovementComponent::GetGroupsToIgnoreMask()
+{
+	return GroupsToIgnore.Packed;
+}
+
 void UCharacterMovementComponent::SetAvoidanceVelocityLock(class UAvoidanceManager* Avoidance, float Duration)
 {
 	Avoidance->OverrideToMaxWeight(AvoidanceUID, Duration);
diff --git a/Engine/Source/Runtime/Engine/Private/Vehicles/WheeledVehicleMovementComponent.cpp b/Engine/Source/Runtime/Engine/Private/Vehicles/WheeledVehicleMovementComponent.cpp
--- a/Engine/Source/Runtime/Engine/Private/Vehicles/WheeledVehicleMovementComponent.cpp
+++ b/Engine/Source/Runtime/Engine/Private/Vehicles/WheeledVehicleMovementComponent.cpp
@@ -107,7 +107,18 @@ UWheeledVehicleMovementComponent::UWheeledVehicleMovementComponent(const FObject
 	SteeringInputRate.RiseRate = 2.5f;
 	SteeringInputRate.FallRate = 5.0f;
 
-
+	bUseRVOAvoidance = false;
+	AvoidanceVelocity = FVector::ZeroVector;
+	AvoidanceLockVelocity = FVector::ZeroVector;
+	AvoidanceLockTimer = 0.0f;
+	AvoidanceGroup.bGroup0 = true;
+	GroupsToAvoid.Packed = 0xFFFFFFFF;
+	GroupsToIgnore.Packed = 0;
+	RVOAvoidanceRadius = 400.0f;
+	RVOAvoidanceHeight = 200.0f;
+	AvoidanceConsiderationRadius = 2000.0f;
+	RVOSteeringStep = 0.5f;
+	RVOThrottleStep = 0.25f;
 
 #if WITH_VEHICLE
 	// tire load filtering
@@ -479,6 +490,14 @@ void UWheeledVehicleMovementComponent::GenerateTireForces( UVehicleWheel* Wheel,
 
 void UWheeledVehicleMovementComponent::PostSetupVehicle()
 {
+	if (bUseRVOAvoidance)
+	{
+		UAvoidanceManager* AvoidanceManager = GetWorld()->GetAvoidanceManager();
+		if (AvoidanceManager)
+		{
+			AvoidanceManager->RegisterMovementComponent(this, AvoidanceWeight);
+		}
+	}
 }
 
 FVector UWheeledVehicleMovementComponent::GetWheelRestingPosition( const FWheelSetup& WheelSetup )
@@ -664,6 +683,11 @@ void UWheeledVehicleMovementComponent::DestroyWheels()
 
 void UWheeledVehicleMovementComponent::TickVehicle( float DeltaTime )
 {
+	if (AvoidanceLockTimer > 0.0f)
+	{
+		AvoidanceLockTimer -= DeltaTime;
+	}
+
 	// movement updates and replication
 	if (PVehicle && UpdatedComponent)
 	{
@@ -725,6 +749,39 @@ void UWheeledVehicleMovementComponent::UpdateSimulation( float DeltaTime )
 
 #endif // WITH_VEHICLE
 
+void UWheeledVehicleMovementComponent::UpdateAvoidance(float DeltaTime)
+{
+	UpdateDefaultAvoidance();
+}
+
+void UWheeledVehicleMovementComponent::UpdateDefaultAvoidance()
+{
+	if (!bUseRVOAvoidance)
+	{
+		return;
+	}
+
+	SCOPE_CYCLE_COUNTER(STAT_AI_ObstacleAvoidance);
+
+	UAvoidanceManager* AvoidanceManager = GetWorld()->GetAvoidanceManager();
+	if (AvoidanceManager && !bWasAvoidanceUpdated)
+	{
+		AvoidanceManager->UpdateRVO(this);
+		
+		//Consider this a clean move because we didn't even try to avoid.
+		SetAvoidanceVelocityLock(AvoidanceManager, AvoidanceManager->LockTimeAfterClean);
+	}
+
+	bWasAvoidanceUpdated = false;		//Reset for next frame
+}
+
+void UWheeledVehicleMovementComponent::SetAvoidanceVelocityLock(class UAvoidanceManager* Avoidance, float Duration)
+{
+	Avoidance->OverrideToMaxWeight(AvoidanceUID, Duration);
+	AvoidanceLockVelocity = AvoidanceVelocity;
+	AvoidanceLockTimer = Duration;
+}
+
 void UWheeledVehicleMovementComponent::UpdateState( float DeltaTime )
 {
 	// update input values
@@ -744,6 +801,12 @@ void UWheeledVehicleMovementComponent::UpdateState( float DeltaTime )
 			}
 		}
 		
+		if (bUseRVOAvoidance)
+		{
+			CalculateAvoidanceVelocity(DeltaTime);
+			UpdateAvoidance(DeltaTime);
+		}
+
 		SteeringInput = SteeringInputRate.InterpInputValue(DeltaTime, SteeringInput, CalcSteeringInput());
 		ThrottleInput = ThrottleInputRate.InterpInputValue( DeltaTime, ThrottleInput, CalcThrottleInput() );
 		BrakeInput = BrakeInputRate.InterpInputValue(DeltaTime, BrakeInput, CalcBrakeInput());
@@ -790,6 +853,19 @@ void UWheeledVehicleMovementComponent::ServerUpdateState_Implementation(float In
 
 float UWheeledVehicleMovementComponent::CalcSteeringInput()
 {
+	if (bUseRVOAvoidance)
+	{
+		const float AngleDiff = AvoidanceVelocity.HeadingAngle() - GetVelocityForRVOConsideration().HeadingAngle();
+		if (AngleDiff > 0.0f)
+		{
+			RawSteeringInput = FMath::Clamp(RawSteeringInput + RVOSteeringStep, 0.0f, 1.0f);
+		}
+		else if (AngleDiff < 0.0f)
+		{
+			RawSteeringInput = FMath::Clamp(RawSteeringInput - RVOSteeringStep, -1.0f, 0.0f);
+		}
+	}
+
 	return RawSteeringInput;
 }
 
@@ -843,6 +919,21 @@ float UWheeledVehicleMovementComponent::CalcHandbrakeInput()
 
 float UWheeledVehicleMovementComponent::CalcThrottleInput()
 {
+	if (bUseRVOAvoidance)
+	{
+		const float AvoidanceSpeedSq = AvoidanceVelocity.SizeSquared();
+		const float DesiredSpeedSq = GetVelocityForRVOConsideration().SizeSquared();
+
+		if (AvoidanceSpeedSq > DesiredSpeedSq)
+		{
+			RawThrottleInput = FMath::Clamp(RawThrottleInput + RVOThrottleStep, -1.0f, 1.0f);
+		}
+		else if (AvoidanceSpeedSq < DesiredSpeedSq)
+		{
+			RawThrottleInput = FMath::Clamp(RawThrottleInput - RVOThrottleStep, -1.0f, 1.0f);
+		}		
+	}
+
 	//If the user is changing direction we should really be braking first and not applying any gas, so wait until they've changed gears
 	if ( (RawThrottleInput > 0.f && GetTargetGear() < 0) || (RawThrottleInput < 0.f && GetTargetGear() > 0)) 
 	{
@@ -1389,6 +1480,178 @@ void UWheeledVehicleMovementComponent::ComputeConstants()
 	MaxEngineRPM = 5000.f;
 }
 
+void UWheeledVehicleMovementComponent::CalculateAvoidanceVelocity(float DeltaTime)
+{
+	if (!bUseRVOAvoidance)
+	{
+		return;
+	}
+	
+	SCOPE_CYCLE_COUNTER(STAT_AI_ObstacleAvoidance);
+	
+	UAvoidanceManager* AvoidanceManager = GetWorld()->GetAvoidanceManager();
+	APawn* MyOwner = UpdatedComponent ? Cast<APawn>(UpdatedComponent->GetOwner()) : NULL;
+	
+	if (AvoidanceWeight >= 1.0f || AvoidanceManager == NULL || MyOwner == NULL)
+	{
+		return;
+	}
+	
+	if (MyOwner->Role != ROLE_Authority)
+	{	
+		return;
+	}
+
+#if !(UE_BUILD_SHIPPING || UE_BUILD_TEST)
+	const bool bShowDebug = AvoidanceManager->IsDebugEnabled(AvoidanceUID);
+#endif
+
+	// since we don't assign the avoidance velocity but instead use it to adjust steering and throttle,
+	// always reset the avoidance velocity to the current velocity
+	AvoidanceVelocity = GetVelocityForRVOConsideration();
+
+	if (!AvoidanceVelocity.IsZero())
+	{
+		//See if we're doing a locked avoidance move already, and if so, skip the testing and just do the move.
+		if (AvoidanceLockTimer > 0.0f)
+		{
+			AvoidanceVelocity = AvoidanceLockVelocity;
+#if !(UE_BUILD_SHIPPING || UE_BUILD_TEST)
+			if (bShowDebug)
+			{
+				DrawDebugLine(GetWorld(), GetRVOAvoidanceOrigin(), GetRVOAvoidanceOrigin() + AvoidanceVelocity, FColor(0, 0, 255), true, 0.5f, SDPG_MAX);
+			}
+#endif
+		}
+		else
+		{
+			FVector NewVelocity = AvoidanceManager->GetAvoidanceVelocityForComponent(this);
+			if (!NewVelocity.Equals(AvoidanceVelocity))		//Really want to branch hint that this will probably not pass
+			{
+				//Had to divert course, lock this avoidance move in for a short time. This will make us a VO, so unlocked others will know to avoid us.
+				AvoidanceVelocity = NewVelocity;
+				SetAvoidanceVelocityLock(AvoidanceManager, AvoidanceManager->LockTimeAfterAvoid);
+#if !(UE_BUILD_SHIPPING || UE_BUILD_TEST)
+				if (bShowDebug)
+				{
+					DrawDebugLine(GetWorld(), GetRVOAvoidanceOrigin(), GetRVOAvoidanceOrigin() + AvoidanceVelocity, FColor(255, 0, 0), true, 20.0f, SDPG_MAX, 10.0f);
+				}
+#endif
+			}
+			else
+			{
+				//Although we didn't divert course, our velocity for this frame is decided. We will not reciprocate anything further, so treat as a VO for the remainder of this frame.
+				SetAvoidanceVelocityLock(AvoidanceManager, AvoidanceManager->LockTimeAfterClean);	//10 ms of lock time should be adequate.
+			}
+		}
+
+		AvoidanceManager->UpdateRVO(this);
+		bWasAvoidanceUpdated = true;
+	}
+#if !(UE_BUILD_SHIPPING || UE_BUILD_TEST)
+	else if (bShowDebug)
+	{
+		DrawDebugLine(GetWorld(), GetRVOAvoidanceOrigin(), GetRVOAvoidanceOrigin() + GetVelocityForRVOConsideration(), FColor(255, 255, 0), true, 0.05f, SDPG_MAX);
+	}
+
+	if (bShowDebug)
+	{
+		FVector UpLine(0, 0, 500);
+		DrawDebugLine(GetWorld(), GetRVOAvoidanceOrigin(), GetRVOAvoidanceOrigin() + UpLine, (AvoidanceLockTimer > 0.01f) ? FColor(255, 0, 0) : FColor(0, 0, 255), true, 0.05f, SDPG_MAX, 5.0f);
+	}
+#endif
+}
+
+void UWheeledVehicleMovementComponent::SetAvoidanceGroup(int32 GroupFlags)
+{
+	AvoidanceGroup.SetFlagsDirectly(GroupFlags);
+}
+
+void UWheeledVehicleMovementComponent::SetGroupsToAvoid(int32 GroupFlags)
+{
+	GroupsToAvoid.SetFlagsDirectly(GroupFlags);
+}
+
+void UWheeledVehicleMovementComponent::SetGroupsToIgnore(int32 GroupFlags)
+{
+	GroupsToIgnore.SetFlagsDirectly(GroupFlags);
+}
+
+void UWheeledVehicleMovementComponent::SetAvoidanceEnabled(bool bEnable)
+{
+	if (bUseRVOAvoidance != bEnable)
+	{
+		bUseRVOAvoidance = bEnable;
+		
+		UAvoidanceManager* AvoidanceManager = GetWorld()->GetAvoidanceManager();
+		if (AvoidanceManager && bEnable && AvoidanceUID == 0)
+		{
+			AvoidanceManager->RegisterMovementComponent(this, AvoidanceWeight);
+		}
+	}
+}
+
+
+void UWheeledVehicleMovementComponent::SetRVOAvoidanceUID(int32 UID)
+{
+	AvoidanceUID = UID;
+}
+
+int32 UWheeledVehicleMovementComponent::GetRVOAvoidanceUID()
+{
+	return AvoidanceUID;
+}
+
+void UWheeledVehicleMovementComponent::SetRVOAvoidanceWeight(float Weight)
+{
+	AvoidanceWeight = Weight;
+}
+
+float UWheeledVehicleMovementComponent::GetRVOAvoidanceWeight()
+{
+	return AvoidanceWeight;
+}
+
+FVector UWheeledVehicleMovementComponent::GetRVOAvoidanceOrigin()
+{
+	return UpdatedComponent->GetComponentLocation();
+}
+
+float UWheeledVehicleMovementComponent::GetRVOAvoidanceRadius()
+{
+	return RVOAvoidanceRadius;
+}
+
+float UWheeledVehicleMovementComponent::GetRVOAvoidanceHeight()
+{
+	return RVOAvoidanceHeight;
+}
+
+float UWheeledVehicleMovementComponent::GetRVOAvoidanceConsiderationRadius()
+{
+	return AvoidanceConsiderationRadius;
+}
+
+FVector UWheeledVehicleMovementComponent::GetVelocityForRVOConsideration()
+{
+	return UpdatedComponent->GetComponentVelocity();
+}
+
+int32 UWheeledVehicleMovementComponent::GetAvoidanceGroupMask()
+{
+	return AvoidanceGroup.Packed;
+}
+
+int32 UWheeledVehicleMovementComponent::GetGroupsToAvoidMask()
+{
+	return GroupsToAvoid.Packed;
+}
+
+int32 UWheeledVehicleMovementComponent::GetGroupsToIgnoreMask()
+{
+	return GroupsToIgnore.Packed;
+}
+
 #undef LOCTEXT_NAMESPACE
 
 FWheelSetup::FWheelSetup()
diff --git a/Engine/Source/Runtime/Engine/Public/AI/RVOAvoidanceInterface.h b/Engine/Source/Runtime/Engine/Public/AI/RVOAvoidanceInterface.h
new file mode 100644
--- /dev/null
+++ b/Engine/Source/Runtime/Engine/Public/AI/RVOAvoidanceInterface.h
@@ -0,0 +1,55 @@
+//=============================================================================
+// IRVOAvoidanceInterface is an interface for objects that want to perform
+// RVO avoidance. See UCharacterMovementComponent and UWheeledVehicleMovementComponent
+// for example implementations.
+//=============================================================================
+
+#pragma once
+#include "RVOAvoidanceInterface.generated.h"
+
+UINTERFACE(MinimalAPI, meta=(CannotImplementInterfaceInBlueprint))
+class URVOAvoidanceInterface : public UInterface
+{
+	GENERATED_UINTERFACE_BODY()
+};
+
+class ENGINE_API IRVOAvoidanceInterface
+{
+	GENERATED_IINTERFACE_BODY()
+
+	/** Store the AvoidanceUID generated by the Avoidance Manager **/
+	virtual void SetRVOAvoidanceUID(int32 UID) = 0;
+
+	/** Return the AvoidanceUID assigned by the Avoidance Manager during registration **/
+	virtual int32 GetRVOAvoidanceUID() = 0;
+
+	/** Store the AvoidanceWeight generated by the Avoidance Manager **/
+	virtual void SetRVOAvoidanceWeight(float Weight) = 0;
+
+	/** Returns the AvoidanceWeight assigned by the Avoidance Manager during registration **/
+	virtual float GetRVOAvoidanceWeight() = 0;
+
+	/** Get the Location from where the RVO avoidance should originate **/
+	virtual FVector GetRVOAvoidanceOrigin() = 0;
+
+	/** The scaled collider radius to consider for RVO avoidance **/
+	virtual float GetRVOAvoidanceRadius() = 0;
+
+	/** The scaled collider height to consider for RVO avoidance **/
+	virtual float GetRVOAvoidanceHeight() = 0;
+
+	/** The scaled collider radius to consider for RVO avoidance **/
+	virtual float GetRVOAvoidanceConsiderationRadius() = 0;
+
+	/** The velocity of the avoiding entity **/
+	virtual FVector GetVelocityForRVOConsideration() = 0;
+
+	/** This actor's avoidance group mask **/
+	virtual int32 GetAvoidanceGroupMask() = 0;
+
+	/** Agent groups to avoid mask **/
+	virtual int32 GetGroupsToAvoidMask() = 0;
+
+	/** Agent groups to ignore **/
+	virtual int32 GetGroupsToIgnoreMask() = 0;
+};
