fixed bounds of NavLinkProxy with active custom link

diff --git a/Engine/Source/Runtime/Engine/Private/AI/Navigation/NavLinkProxy.cpp b/Engine/Source/Runtime/Engine/Private/AI/Navigation/NavLinkProxy.cpp
--- a/Engine/Source/Runtime/Engine/Private/AI/Navigation/NavLinkProxy.cpp
+++ b/Engine/Source/Runtime/Engine/Private/AI/Navigation/NavLinkProxy.cpp
@@ -123,13 +123,15 @@ FBox ANavLinkProxy::GetComponentsBoundingBox(bool bNonColliding) const
 		LinksBB += SegmentLink.RightEnd;
 	}
 
+	LinksBB = LinksBB.TransformBy(RootComponent->ComponentToWorld);
+
 	if (SmartLinkComp->IsNavigationRelevant())
 	{
 		LinksBB += SmartLinkComp->GetStartPoint();
 		LinksBB += SmartLinkComp->GetEndPoint();
 	}
 
-	return LinksBB.TransformBy(RootComponent->ComponentToWorld);
+	return LinksBB;
 }
 
 void ANavLinkProxy::NotifySmartLinkReached(UNavLinkCustomComponent* LinkComp, class UPathFollowingComponent* PathComp, const FVector& DestPoint)
