compilation fix (IRVOAvoidanceInterface)

diff --git a/Engine/Source/Runtime/Engine/Private/AI/Navigation/AvoidanceManager.cpp b/Engine/Source/Runtime/Engine/Private/AI/Navigation/AvoidanceManager.cpp
--- a/Engine/Source/Runtime/Engine/Private/AI/Navigation/AvoidanceManager.cpp
+++ b/Engine/Source/Runtime/Engine/Private/AI/Navigation/AvoidanceManager.cpp
@@ -1,6 +1,7 @@
 // Copyright 1998-2014 Epic Games, Inc. All Rights Reserved.
 
 #include "EnginePrivate.h"
+#include "GameFramework/MovementComponent.h"
 #include "AI/Navigation/AvoidanceManager.h"
 #include "AI/RVOAvoidanceInterface.h"
 
