Prevent Const material parameter values from crashing

#branch UE4

#change Modifiied those Const values to be ran through the compiler than used as is (three occurances that I could see).

diff --git a/Engine/Source/Runtime/Engine/Private/Materials/MaterialExpressions.cpp b/Engine/Source/Runtime/Engine/Private/Materials/MaterialExpressions.cpp
--- a/Engine/Source/Runtime/Engine/Private/Materials/MaterialExpressions.cpp
+++ b/Engine/Source/Runtime/Engine/Private/Materials/MaterialExpressions.cpp
@@ -1327,7 +1327,7 @@ int32 UMaterialExpressionTextureSample::Compile(class FMaterialCompiler* Compile
 				TextureCodeIndex,
 				Coordinates.Expression ? Coordinates.Compile(Compiler) : Compiler->TextureCoordinate(ConstCoordinate, false, false),
 				(EMaterialSamplerType)EffectiveSamplerType,
-				MipValue.Expression ? MipValue.Compile(Compiler) : ConstMipValue,
+				MipValue.Expression ? MipValue.Compile(Compiler) : Compiler->Constant(ConstMipValue),
 				MipValueMode
 				);
 		}
@@ -1468,7 +1468,7 @@ int32 UMaterialExpressionTextureSampleParameter::Compile(class FMaterialCompiler
 					Compiler->TextureParameter(ParameterName, Texture),
 					Coordinates.Expression ? Coordinates.Compile(Compiler) : Compiler->TextureCoordinate(ConstCoordinate, false, false),
 					(EMaterialSamplerType)SamplerType,
-					MipValue.Expression ? MipValue.Compile(Compiler) : ConstMipValue,
+					MipValue.Expression ? MipValue.Compile(Compiler) : Compiler->Constant(ConstMipValue),
 					MipValueMode
 					);
 }
@@ -5009,7 +5009,7 @@ int32 UMaterialExpressionIf::Compile(class FMaterialCompiler* Compiler, int32 Ou
 	}
 
 	int32 Arg3 = AGreaterThanB.Compile(Compiler);
-	int32 Arg4 = AEqualsB.Expression ? AEqualsB.Compile(Compiler) : ConstAEqualsB;
+	int32 Arg4 = AEqualsB.Expression ? AEqualsB.Compile(Compiler) : Compiler->Constant(ConstAEqualsB);
 	int32 Arg5 = ALessThanB.Compile(Compiler);
 	int32 ThresholdArg = Compiler->Constant(EqualsThreshold);
 
