added more transparency what set a console variable the last time (cvar name in the console shows it), allows to implement a order free override

diff --git a/Engine/Source/Editor/LevelEditor/Private/LevelEditorActions.cpp b/Engine/Source/Editor/LevelEditor/Private/LevelEditorActions.cpp
--- a/Engine/Source/Editor/LevelEditor/Private/LevelEditorActions.cpp
+++ b/Engine/Source/Editor/LevelEditor/Private/LevelEditorActions.cpp
@@ -449,7 +449,7 @@ void FLevelEditorActionCallbacks::SetMaterialQualityLevel( EMaterialQualityLevel
 
 	//Ensure the material quality cvar is also set.
 	static IConsoleVariable* MaterialQualityLevelVar = IConsoleManager::Get().FindConsoleVariable(TEXT("r.MaterialQualityLevel"));
-	MaterialQualityLevelVar->Set(NewQualityLevel);
+	MaterialQualityLevelVar->Set(NewQualityLevel, ECVF_SetByEditorSetting);
 
 	GUnrealEd->RedrawAllViewports();
 }
@@ -1911,7 +1911,7 @@ void FLevelEditorActionCallbacks::OnToggleFreezeParticleSimulation()
 	IConsoleVariable* CVar = ConsoleManager.FindConsoleVariable(TEXT("FX.FreezeParticleSimulation"));
 	if (CVar)
 	{
-		CVar->Set(CVar->GetInt() == 0 ? 1 : 0);
+		CVar->Set(CVar->GetInt() == 0 ? 1 : 0, ECVF_SetByEditorSetting);
 	}
 }
 
@@ -1921,7 +1921,7 @@ bool FLevelEditorActionCallbacks::OnIsParticleSimulationFrozen()
 	static const auto* CVar = ConsoleManager.FindTConsoleVariableDataInt(TEXT("FX.FreezeParticleSimulation"));
 	if (CVar)
 	{
-		return CVar->GetValueOnGameThread() == 0 ? false : true;
+		return CVar->GetValueOnGameThread() != 0;
 	}
 	return false;
 }
diff --git a/Engine/Source/Editor/UnrealEd/Private/Editor/EditorUserSettings.cpp b/Engine/Source/Editor/UnrealEd/Private/Editor/EditorUserSettings.cpp
--- a/Engine/Source/Editor/UnrealEd/Private/Editor/EditorUserSettings.cpp
+++ b/Engine/Source/Editor/UnrealEd/Private/Editor/EditorUserSettings.cpp
@@ -21,7 +21,7 @@ void UEditorUserSettings::PostInitProperties()
 
 	//Ensure the material quality cvar is set to the settings loaded.
 	static IConsoleVariable* MaterialQualityLevelVar = IConsoleManager::Get().FindConsoleVariable(TEXT("r.MaterialQualityLevel"));
-	MaterialQualityLevelVar->Set(MaterialQualityLevel);
+	MaterialQualityLevelVar->Set(MaterialQualityLevel, ECVF_SetByEditorSetting);
 }
 
 void UEditorUserSettings::PostEditChangeProperty( FPropertyChangedEvent& PropertyChangedEvent )
diff --git a/Engine/Source/Editor/UnrealEd/Private/PerformanceMonitor.cpp b/Engine/Source/Editor/UnrealEd/Private/PerformanceMonitor.cpp
--- a/Engine/Source/Editor/UnrealEd/Private/PerformanceMonitor.cpp
+++ b/Engine/Source/Editor/UnrealEd/Private/PerformanceMonitor.cpp
@@ -275,9 +275,11 @@ void FPerformanceMonitor::Tick(float DeltaTime)
 
 		if (PercentUnderTarget >= CVarPercentThreshold->GetFloat())
 		{
+			static IConsoleVariable* CVarCoarseSampleTime = IConsoleManager::Get().FindConsoleVariable(TEXT("PerfWarn.CoarseSampleTime"));
+
 			Arguments.Add(TEXT("Framerate"), CVarMinFPS->GetInt());
 			Arguments.Add(TEXT("Percentage"), FMath::FloorToFloat(PercentUnderTarget));
-			SampleTime = IConsoleManager::Get().FindConsoleVariable(TEXT("PerfWarn.CoarseSampleTime"))->GetInt();
+			SampleTime = CVarCoarseSampleTime->GetInt();
 
 			bLowFramerate = true;
 		}
diff --git a/Engine/Source/Runtime/Core/Private/HAL/ConsoleManager.cpp b/Engine/Source/Runtime/Core/Private/HAL/ConsoleManager.cpp
--- a/Engine/Source/Runtime/Core/Private/HAL/ConsoleManager.cpp
+++ b/Engine/Source/Runtime/Core/Private/HAL/ConsoleManager.cpp
@@ -59,12 +59,14 @@ public:
 
 	// ------
 
-	void OnChanged()
+	void OnChanged(EConsoleVariableFlags SetBy)
 	{
 		// only change on main thread
 
 		Flags = (EConsoleVariableFlags)((uint32)Flags | ECVF_Changed);
 
+		Flags = (EConsoleVariableFlags)(((uint32)Flags & ~ECVF_SetByMask) | SetBy);
+
 		OnChangedCallback.ExecuteIfBound(this);
 	}
 
@@ -204,10 +206,10 @@ public:
 	{
 		delete this; 
 	} 
-	virtual void Set(const TCHAR* InValue)
+	virtual void Set(const TCHAR* InValue, EConsoleVariableFlags SetBy)
 	{
 		TTypeFromString<T>::FromString(Data.ShadowedValue[0], InValue);
-		OnChanged();
+		OnChanged(SetBy);
 	}
 
 	virtual int32 GetInt() const;
@@ -228,11 +230,11 @@ private: // ----------------------------------------------------
 		return This->Data.GetReferenceOnAnyThread();
 	}
 
-	void OnChanged()
+	void OnChanged(EConsoleVariableFlags SetBy)
 	{
 		// propagate from main thread to render thread
 		OnCVarChange(Data.ShadowedValue[1], Data.ShadowedValue[0], Flags);
-		FConsoleVariableBase::OnChanged();
+		FConsoleVariableBase::OnChanged(SetBy);
 	}
 };
 
@@ -277,10 +279,10 @@ template<> TConsoleVariableData<float>* FConsoleVariable<float>::AsVariableFloat
 
 // specialization for FString
 
-template<> void FConsoleVariable<FString>::Set(const TCHAR* InValue)
+template<> void FConsoleVariable<FString>::Set(const TCHAR* InValue, EConsoleVariableFlags SetBy)
 {
 	Data.ShadowedValue[0] = InValue;
-	OnChanged();
+	OnChanged(SetBy);
 }
 template<> int32 FConsoleVariable<FString>::GetInt() const
 {
@@ -318,10 +320,10 @@ public:
 	{
 		delete this; 
 	} 
-	virtual void Set(const TCHAR* InValue)
+	virtual void Set(const TCHAR* InValue, EConsoleVariableFlags SetBy)
 	{
 		TTypeFromString<T>::FromString(MainValue, InValue);
-		OnChanged();
+		OnChanged(SetBy);
 	}
 	virtual int32 GetInt() const
 	{
@@ -350,11 +352,11 @@ private: // ----------------------------------------------------
 		return (Index == 0) ? MainValue : RefValue;
 	}
 
-	void OnChanged()
+	void OnChanged(EConsoleVariableFlags SetBy)
 	{
 		// propagate from main thread to render thread or to reference
 		OnCVarChange(RefValue, MainValue, Flags);
-		FConsoleVariableBase::OnChanged();
+		FConsoleVariableBase::OnChanged(SetBy);
 	}
 };
 
@@ -383,7 +385,7 @@ public:
 	{
 		delete this; 
 	} 
-	virtual void Set(const TCHAR* InValue)
+	virtual void Set(const TCHAR* InValue, EConsoleVariableFlags SetBy)
 	{
 		int32 Value = FCString::Atoi(InValue);
 
@@ -391,6 +393,8 @@ public:
 
 		FMath::SetBoolInBitField(Force0MaskPtr, BitNumber, Value == 0);
 		FMath::SetBoolInBitField(Force1MaskPtr, BitNumber, Value == 1);
+
+		OnChanged(SetBy);
 	}
 	virtual int32 GetInt() const
 	{
@@ -429,6 +433,30 @@ IConsoleVariable* FConsoleManager::RegisterConsoleVariableBitRef(const TCHAR* CV
 	return AddConsoleObject(CVarName, new FConsoleVariableBitRef(FlagName, BitNumber, Force0MaskPtr, Force1MaskPtr, Help, (EConsoleVariableFlags)Flags))->AsVariable();
 }
 
+
+static TCHAR* GetLastSetBy(EConsoleVariableFlags InSetBy)
+{
+	EConsoleVariableFlags SetBy = (EConsoleVariableFlags)((uint32)InSetBy & ECVF_SetByMask);
+
+	switch(SetBy)
+	{
+#define CASE(A) case ECVF_SetBy##A: return TEXT(#A);
+		CASE(Constructor)
+		CASE(Code)
+		CASE(Scalability)
+		CASE(ConsoleVariablesIni)
+		CASE(SystemSettingsIni)
+		CASE(DeviceProfile)
+		CASE(Console)
+		CASE(Commandline)
+		CASE(EditorSetting)
+		CASE(GameSetting)
+		CASE(Network)
+#undef CASE 
+	}
+	return TEXT("<UNKNOWN>");
+}
+
 void FConsoleManager::CallAllConsoleVariableSinks()
 {
 	if(bCallAllConsoleVariableSinks)
@@ -890,7 +918,7 @@ bool FConsoleManager::ProcessUserConsoleInput(const TCHAR* InInput, FOutputDevic
 		if(*It == 0)
 		{
 			// get current state
-			Ar.Logf(TEXT("%s = \"%s\""), *Param1, *CVar->GetString());
+			Ar.Logf(TEXT("%s = \"%s\"      LastSetBy: %s"), *Param1, *CVar->GetString(), GetLastSetBy(CVar->GetFlags()));
 		}
 		else
 		{
@@ -920,7 +948,7 @@ bool FConsoleManager::ProcessUserConsoleInput(const TCHAR* InInput, FOutputDevic
 				else
 				{
 					// set value
-					CVar->Set(*Param2);
+					CVar->Set(*Param2, ECVF_SetByConsole);
 
 					Ar.Logf(TEXT("%s = \"%s\""), *Param1, *CVar->GetString());
 
@@ -997,7 +1025,7 @@ IConsoleObject* FConsoleManager::AddConsoleObject(const TCHAR* Name, IConsoleObj
 			if(ExistingVar->TestFlags(ECVF_CreatedFromIni))
 			{
 				// The existing one came from the ini, get the value and destroy the existing one (no need to call sink because that will happen after all ini setting have been loaded)
-				Var->Set(*ExistingVar->GetString());
+				Var->Set(*ExistingVar->GetString(), (EConsoleVariableFlags)((uint32)ExistingVar->GetFlags() & ECVF_SetByMask));
 				ExistingVar->Release();
 
 				ConsoleObjects.Add(Name, Var);
@@ -1265,10 +1293,10 @@ void FConsoleManager::Test()
 
 		// call Set(string)
 
-		VarA->Set(TEXT("3.1"));
-		VarB->Set(TEXT("3.1"));
-		VarD->Set(TEXT("3.1"));
-		VarE->Set(TEXT("3.1"));
+		VarA->Set(TEXT("3.1"), ECVF_SetByConsole);
+		VarB->Set(TEXT("3.1"), ECVF_SetByConsole);
+		VarD->Set(TEXT("3.1"), ECVF_SetByConsole);
+		VarE->Set(TEXT("3.1"), ECVF_SetByConsole);
 
 		check(GConsoleVariableCallbackTestCounter == 1);
 
@@ -1288,7 +1316,7 @@ void FConsoleManager::Test()
 		check(RefD == 3);
 		check(RefE == 3.1f);
 
-		VarB->Set(TEXT("3.1"));
+		VarB->Set(TEXT("3.1"), ECVF_SetByCode);
 		check(GConsoleVariableCallbackTestCounter == 2);
 
 		// unregister
@@ -1322,7 +1350,7 @@ void FConsoleManager::Test()
 			check(FMath::IsNearlyEqual(VarC->GetFloat(), 1.23f, KINDA_SMALL_NUMBER));
 			check(VarC->GetString() == FString(TEXT("1.23")));
 			check(!VarC->TestFlags(ECVF_Changed));
-			VarC->Set(TEXT("3.1"));
+			VarC->Set(TEXT("3.1"), ECVF_SetByCode);
 			check(VarC->TestFlags(ECVF_Changed));
 			check(VarC->GetString() == FString(TEXT("3.1")));
 			UnregisterConsoleObject(TEXT("TestNameC"), false);
diff --git a/Engine/Source/Runtime/Core/Private/Misc/ConfigCacheIni.cpp b/Engine/Source/Runtime/Core/Private/Misc/ConfigCacheIni.cpp
--- a/Engine/Source/Runtime/Core/Private/Misc/ConfigCacheIni.cpp
+++ b/Engine/Source/Runtime/Core/Private/Misc/ConfigCacheIni.cpp
@@ -2785,7 +2785,7 @@ void FConfigCacheIni::LoadLocalIniFile(FConfigFile& ConfigFile, const TCHAR* Ini
 	if (!bGenerateDestIni)
 	{
 		// generate path to the .ini file (not a Default ini, IniName is the complete name of the file, without path)
-		FString SourceIniFilename = FString::Printf(TEXT("%s/%s.ini"), *FPaths::SourceConfigDir(), IniName);
+		FString SourceIniFilename = FString::Printf(TEXT("%s%s.ini"), *FPaths::SourceConfigDir(), IniName);
 
 		// load the .ini file straight up
 		LoadAnIniFile(*SourceIniFilename, ConfigFile);
@@ -2868,7 +2868,7 @@ void FConfigCacheIni::LoadConsoleVariablesFromINI()
 				if(CVar)
 				{
 					// Set if the variable exists.
-					CVar->Set(*ValueString);
+					CVar->Set(*ValueString, ECVF_SetByConsoleVariablesIni);
 				}
 				else
 				{
@@ -2962,7 +2962,7 @@ void FConfigFile::UpdateSections(const TCHAR* DiskFilename, const TCHAR* IniRoot
 	Write(DiskFilename, true, NewFile);
 }
 
-void ApplyCVarSettingsFromIni(const TCHAR* InSectionName, const TCHAR* InIniFilename)
+void ApplyCVarSettingsFromIni(const TCHAR* InSectionName, const TCHAR* InIniFilename, uint32 SetBy)
 {
 	// Lookup the config section for this section and group number
 	TArray<FString> SectionStrings;
@@ -2984,17 +2984,17 @@ void ApplyCVarSettingsFromIni(const TCHAR* InSectionName, const TCHAR* InIniFile
 						|| FCString::Stricmp(*CVarValue, TEXT("Yes")) == 0
 						|| FCString::Stricmp(*CVarValue, TEXT("On")) == 0)
 					{
-						CVar->Set(1);
+						CVar->Set(1, (EConsoleVariableFlags)SetBy);
 					}
 					else if(FCString::Stricmp(*CVarValue, TEXT("False")) == 0
 						|| FCString::Stricmp(*CVarValue, TEXT("No")) == 0
 						|| FCString::Stricmp(*CVarValue, TEXT("Off")) == 0)
 					{
-						CVar->Set(0);
+						CVar->Set(0, (EConsoleVariableFlags)SetBy);
 					}
 					else
 					{
-						CVar->Set(*CVarValue);
+						CVar->Set(*CVarValue, (EConsoleVariableFlags)SetBy);
 					}
 				}
 			}
@@ -3007,9 +3007,9 @@ void ApplyCVarSettingsFromIni(const TCHAR* InSectionName, const TCHAR* InIniFile
 	}
 }
 
-void ApplyCVarSettingsGroupFromIni(const TCHAR* InSectionBaseName, int32 InGroupNumber, const TCHAR* InIniFilename)
+void ApplyCVarSettingsGroupFromIni(const TCHAR* InSectionBaseName, int32 InGroupNumber, const TCHAR* InIniFilename, uint32 SetBy)
 {
 	// Lookup the config section for this section and group number
 	FString SectionName = FString::Printf(TEXT("%s@%d"), InSectionBaseName, InGroupNumber);
-	ApplyCVarSettingsFromIni(*SectionName,InIniFilename);
+	ApplyCVarSettingsFromIni(*SectionName,InIniFilename, SetBy);
 }
diff --git a/Engine/Source/Runtime/Core/Public/HAL/IConsoleManager.h b/Engine/Source/Runtime/Core/Public/HAL/IConsoleManager.h
--- a/Engine/Source/Runtime/Core/Public/HAL/IConsoleManager.h
+++ b/Engine/Source/Runtime/Core/Public/HAL/IConsoleManager.h
@@ -81,6 +81,20 @@ enum EConsoleVariableFlags
 
 	/* those cvars control other cvars with the flag ECVF_Scalability, names should start with "sg." */
 	ECVF_ScalabilityGroup = 0x80,
+
+	/* to get some history of where the last value was set by ( useful for track down why a cvar is in a specific state */
+	ECVF_SetByMask =				0xff000000,
+	ECVF_SetByConstructor =			0x00000000,
+	ECVF_SetByCode =				0x01000000,	// least useful, likely a hack, maybe better to start a new SetBy...
+	ECVF_SetByScalability =			0x02000000,
+	ECVF_SetByConsoleVariablesIni = 0x03000000,
+	ECVF_SetBySystemSettingsIni =	0x04000000,
+	ECVF_SetByDeviceProfile =		0x05000000,
+	ECVF_SetByConsole =				0x06000000,
+	ECVF_SetByCommandline =			0x07000000,	// a minus comman e.g. -VSync
+	ECVF_SetByEditorSetting =		0x08000000,
+	ECVF_SetByGameSetting =			0x09000000,
+	ECVF_SetByNetwork =				0x0a000000,
 };
 
 class IConsoleVariable;
@@ -200,8 +214,11 @@ class IConsoleVariable : public IConsoleObject
 {
 public:
 
-	/** Set the internal value from the specified string. */
-	virtual void Set(const TCHAR* InValue) = 0;
+	/**
+	 * Set the internal value from the specified string. 
+	 * @param SetBy anything in ECVF_LastSetMask e.g. ECVF_SetByScalability
+	 **/
+	virtual void Set(const TCHAR* InValue, EConsoleVariableFlags SetBy = ECVF_SetByCode) = 0;
 	/**
 	 * Get the internal value as int (should not be used on strings).
 	 * @return value is not rounded (simple cast)
@@ -227,16 +244,16 @@ public:
 	// convenience methods
 
 	/** Set the internal value from the specified int. */
-	void Set(int32 InValue)
+	void Set(int32 InValue, EConsoleVariableFlags SetBy = ECVF_SetByCode)
 	{
 		// inefficient but no common code path
-		Set(*FString::Printf(TEXT("%d"), InValue));
+		Set(*FString::Printf(TEXT("%d"), InValue), SetBy);
 	}
 	/** Set the internal value from the specified float. */
-	void Set(float InValue)
+	void Set(float InValue, EConsoleVariableFlags SetBy = ECVF_SetByCode)
 	{
 		// inefficient but no common code path
-		Set(*FString::Printf(TEXT("%g"), InValue));
+		Set(*FString::Printf(TEXT("%g"), InValue), SetBy);
 	}
 
 };
diff --git a/Engine/Source/Runtime/Core/Public/Misc/ConfigCacheIni.h b/Engine/Source/Runtime/Core/Public/Misc/ConfigCacheIni.h
--- a/Engine/Source/Runtime/Core/Public/Misc/ConfigCacheIni.h
+++ b/Engine/Source/Runtime/Core/Public/Misc/ConfigCacheIni.h
@@ -541,12 +541,14 @@ private:
  * @param InSectionBaseName - The base name of the section to apply cvars from (i.e. the bit before the @)
  * @param InGroupNumber - The group number required
  * @param InIniFilename - The ini filename
+ * @param SetBy anything in ECVF_LastSetMask e.g. ECVF_SetByScalability
  */
-CORE_API void ApplyCVarSettingsGroupFromIni(const TCHAR* InSectionBaseName, int32 InGroupNumber, const TCHAR* InIniFilename);
+CORE_API void ApplyCVarSettingsGroupFromIni(const TCHAR* InSectionBaseName, int32 InGroupNumber, const TCHAR* InIniFilename, uint32 SetBy);
 
 /**
  * Helper function to read the contents of an ini file and a specified group of cvar parameters, where sections in the ini file are marked [InName]
  * @param InSectionBaseName - The base name of the section to apply cvars from
  * @param InIniFilename - The ini filename
+ * @param SetBy anything in ECVF_LastSetMask e.g. ECVF_SetByScalability
  */
-CORE_API void ApplyCVarSettingsFromIni(const TCHAR* InSectionBaseName, const TCHAR* InIniFilename);
+CORE_API void ApplyCVarSettingsFromIni(const TCHAR* InSectionBaseName, const TCHAR* InIniFilename, uint32 SetBy);
diff --git a/Engine/Source/Runtime/Engine/Private/ContentStreaming.cpp b/Engine/Source/Runtime/Engine/Private/ContentStreaming.cpp
--- a/Engine/Source/Runtime/Engine/Private/ContentStreaming.cpp
+++ b/Engine/Source/Runtime/Engine/Private/ContentStreaming.cpp
@@ -130,7 +130,7 @@ static TAutoConsoleVariable<float> CVarStreamingBoost(
 	);
 
 #if PLATFORM_SUPPORTS_TEXTURE_STREAMING
-static TAutoConsoleVariable<int32> CVarSetTextureStreamingEnabled(
+static TAutoConsoleVariable<int32> CVarSetTextureStreaming(
 	TEXT("r.TextureStreaming"),
 	1,
 	TEXT("Allows to define if texture streaming is enabled, can be changed at run time.\n")
@@ -1579,7 +1579,7 @@ FStreamingManagerCollection::FStreamingManagerCollection()
 	// Disable texture streaming if that was requested (needs to happen before the call to ProcessNewlyLoadedUObjects, as that can load textures)
 	if( FParse::Param( FCommandLine::Get(), TEXT( "NoTextureStreaming" ) ) )
 	{
-		CVarSetTextureStreamingEnabled.AsVariable()->Set(0);
+		CVarSetTextureStreaming.AsVariable()->Set(0, ECVF_SetByCommandline);
 	}
 #endif
 
@@ -2001,7 +2001,7 @@ void FStreamingManagerCollection::AddOrRemoveTextureStreamingManagerIfNeeded(boo
 
 #if PLATFORM_SUPPORTS_TEXTURE_STREAMING
 	{
-		bUseTextureStreaming = CVarSetTextureStreamingEnabled.GetValueOnGameThread() != 0;
+		bUseTextureStreaming = CVarSetTextureStreaming.GetValueOnGameThread() != 0;
 	}
 
 	if( !GRHISupportsTextureStreaming || IsRunningDedicatedServer() )
@@ -2142,7 +2142,7 @@ FStreamingManagerTexture::FStreamingManagerTexture()
 
 		// Use unlimited texture streaming in the editor. Quality is more important than stutter.
 		static IConsoleVariable* CVar = IConsoleManager::Get().FindConsoleVariable(TEXT("r.Streaming.PoolSize"));
-		CVar->Set(0);
+		CVar->Set(0, ECVF_SetByCode);
 		GTexturePoolSize = 0;
 	}
 	if (GNeverStreamOutTextures)
diff --git a/Engine/Source/Runtime/Engine/Private/DeviceProfiles/DeviceProfileManager.cpp b/Engine/Source/Runtime/Engine/Private/DeviceProfiles/DeviceProfileManager.cpp
--- a/Engine/Source/Runtime/Engine/Private/DeviceProfiles/DeviceProfileManager.cpp
+++ b/Engine/Source/Runtime/Engine/Private/DeviceProfiles/DeviceProfileManager.cpp
@@ -94,7 +94,7 @@ void UDeviceProfileManager::InitializeCVarsForActiveDeviceProfile()
 						if( CVar )
 						{
 							UE_LOG(LogInit, Log, TEXT("Setting Device Profile CVar: [[%s:%s]]"), *CVarKey, *CVarValue);
-							CVar->Set( *CVarValue );
+							CVar->Set( *CVarValue, ECVF_SetByDeviceProfile);
 							CVarsAlreadySetList.Add( CVarKey, CVarValue );
 						}
 						else
diff --git a/Engine/Source/Runtime/Engine/Private/GPUSort.cpp b/Engine/Source/Runtime/Engine/Private/GPUSort.cpp
--- a/Engine/Source/Runtime/Engine/Private/GPUSort.cpp
+++ b/Engine/Source/Runtime/Engine/Private/GPUSort.cpp
@@ -965,11 +965,11 @@ static void RunGPUSortTestWithDebug(FRHICommandListImmediate& RHICmdList, int32
 	static IConsoleVariable* IVarDebugSort = IConsoleManager::Get().FindConsoleVariable(TEXT("GPUSort.DebugSort"));
 	const bool bWasDebuggingOffsets = CVarDebugOffsets.GetValueOnRenderThread() != 0;
 	const bool bWasDebuggingSort = CVarDebugSort.GetValueOnRenderThread() != 0;
-	IVarDebugOffsets->Set(1);
-	IVarDebugSort->Set(1);
+	IVarDebugOffsets->Set(1, ECVF_SetByCode);
+	IVarDebugSort->Set(1, ECVF_SetByCode);
 	RunGPUSortTest(RHICmdList, TestSize, FeatureLevel);
-	IVarDebugOffsets->Set(bWasDebuggingOffsets ? 1 : 0);
-	IVarDebugSort->Set(bWasDebuggingSort ? 1 : 0);
+	IVarDebugOffsets->Set(bWasDebuggingOffsets ? 1 : 0, ECVF_SetByCode);
+	IVarDebugSort->Set(bWasDebuggingSort ? 1 : 0, ECVF_SetByCode);
 }
 
 /**
diff --git a/Engine/Source/Runtime/Engine/Private/GameUserSettings.cpp b/Engine/Source/Runtime/Engine/Private/GameUserSettings.cpp
--- a/Engine/Source/Runtime/Engine/Private/GameUserSettings.cpp
+++ b/Engine/Source/Runtime/Engine/Private/GameUserSettings.cpp
@@ -193,14 +193,14 @@ void UGameUserSettings::ApplyNonResolutionSettings()
 	EWindowMode::Type NewWindowMode = GetFullscreenMode();
 
 	{
-		IConsoleVariable* CVar = IConsoleManager::Get().FindConsoleVariable(TEXT("r.FullScreenMode")); 
-		CVar->Set(NewWindowMode);
+		static IConsoleVariable* CVar = IConsoleManager::Get().FindConsoleVariable(TEXT("r.FullScreenMode")); 
+		CVar->Set(NewWindowMode, ECVF_SetByGameSetting);
 	}
 
 	// Update vsync cvar
 	{
-		auto CVar = IConsoleManager::Get().FindConsoleVariable(TEXT("r.VSync")); 
-		CVar->Set(IsVSyncEnabled());
+		static auto CVar = IConsoleManager::Get().FindConsoleVariable(TEXT("r.VSync")); 
+		CVar->Set(IsVSyncEnabled(), ECVF_SetByGameSetting);
 	}
 
 	// in init those are loaded earlier, after that we apply consolevariables.ini
@@ -219,8 +219,8 @@ void UGameUserSettings::ApplyResolutionSettings(bool bCheckForCommandLineOverrid
 	EWindowMode::Type NewFullscreenMode = GetFullscreenMode();
 
 	{
-		IConsoleVariable* CVar = IConsoleManager::Get().FindConsoleVariable(TEXT("r.FullScreenMode"));
-		CVar->Set(NewFullscreenMode);
+		static IConsoleVariable* CVar = IConsoleManager::Get().FindConsoleVariable(TEXT("r.FullScreenMode"));
+		CVar->Set(NewFullscreenMode, ECVF_SetByGameSetting);
 	}
 
 	// Request a resolution change
@@ -273,8 +273,8 @@ void UGameUserSettings::RequestResolutionChange(int32 InResolutionX, int32 InRes
 
 
 	{
-		IConsoleVariable* CVar = IConsoleManager::Get().FindConsoleVariable(TEXT("r.FullScreenMode"));
-		CVar->Set(InWindowMode);
+		static IConsoleVariable* CVar = IConsoleManager::Get().FindConsoleVariable(TEXT("r.FullScreenMode"));
+		CVar->Set(InWindowMode, ECVF_SetByGameSetting);
 	}
 	FSystemResolution::RequestResolutionChange(InResolutionX, InResolutionY, InWindowMode);
 }
@@ -333,8 +333,8 @@ void UGameUserSettings::PreloadResolutionSettings()
 	}
 
 	{
-		auto CVar = IConsoleManager::Get().FindConsoleVariable(TEXT("r.FullScreenMode")); 
-		CVar->Set((int32)WindowMode);
+		static auto CVar = IConsoleManager::Get().FindConsoleVariable(TEXT("r.FullScreenMode")); 
+		CVar->Set((int32)WindowMode, ECVF_SetByGameSetting);
 	}
 
 	RequestResolutionChange(ResolutionX, ResolutionY, WindowMode);
diff --git a/Engine/Source/Runtime/Engine/Private/GameViewportClient.cpp b/Engine/Source/Runtime/Engine/Private/GameViewportClient.cpp
--- a/Engine/Source/Runtime/Engine/Private/GameViewportClient.cpp
+++ b/Engine/Source/Runtime/Engine/Private/GameViewportClient.cpp
@@ -717,7 +717,8 @@ void UGameViewportClient::Draw(FViewport* InViewport, FCanvas* SceneCanvas)
 					// Mode is out of range, so display a message to the user, and reset the mode back to the previous valid one
 					UE_LOG(LogConsoleResponse, Warning, TEXT("Buffer visualization mode '%s' does not exist"), *ModeNameString);
 					NewBufferVisualizationMode = CurrentBufferVisualizationMode;
-					ICVar->Set(*NewBufferVisualizationMode.GetPlainNameString());
+					// todo: cvars are user settings, here the cvar state is used to avoid log spam and to auto correct for the user (likely not what the user wants)
+					ICVar->Set(*NewBufferVisualizationMode.GetPlainNameString(), ECVF_SetByCode);
 				}
 				else
 				{
diff --git a/Engine/Source/Runtime/Engine/Private/NetworkSettings.cpp b/Engine/Source/Runtime/Engine/Private/NetworkSettings.cpp
--- a/Engine/Source/Runtime/Engine/Private/NetworkSettings.cpp
+++ b/Engine/Source/Runtime/Engine/Private/NetworkSettings.cpp
@@ -60,19 +60,19 @@ void UNetworkSettings::PostEditChangeProperty(FPropertyChangedEvent& PropertyCha
 				UByteProperty* ByteProperty = Cast<UByteProperty>(PropertyChangedEvent.Property);
 				if (ByteProperty != NULL && ByteProperty->Enum != NULL)
 				{
-					CVar->Set(ByteProperty->GetPropertyValue_InContainer(this));
+					CVar->Set(ByteProperty->GetPropertyValue_InContainer(this), ECVF_SetByNetwork);
 				}
 				else if (UBoolProperty* BoolProperty = Cast<UBoolProperty>(PropertyChangedEvent.Property))
 				{
-					CVar->Set((int32)BoolProperty->GetPropertyValue_InContainer(this));
+					CVar->Set((int32)BoolProperty->GetPropertyValue_InContainer(this), ECVF_SetByNetwork);
 				}
 				else if (UIntProperty* IntProperty = Cast<UIntProperty>(PropertyChangedEvent.Property))
 				{
-					CVar->Set(IntProperty->GetPropertyValue_InContainer(this));
+					CVar->Set(IntProperty->GetPropertyValue_InContainer(this), ECVF_SetByNetwork);
 				}
 				else if (UFloatProperty* FloatProperty = Cast<UFloatProperty>(PropertyChangedEvent.Property))
 				{
-					CVar->Set(FloatProperty->GetPropertyValue_InContainer(this));
+					CVar->Set(FloatProperty->GetPropertyValue_InContainer(this), ECVF_SetByNetwork);
 				}
 			}
 		}
diff --git a/Engine/Source/Runtime/Engine/Private/Particles/FXSystem.cpp b/Engine/Source/Runtime/Engine/Private/Particles/FXSystem.cpp
--- a/Engine/Source/Runtime/Engine/Private/Particles/FXSystem.cpp
+++ b/Engine/Source/Runtime/Engine/Private/Particles/FXSystem.cpp
@@ -165,7 +165,9 @@ void FFXSystem::Tick(float DeltaSeconds)
 			TestGPUSort((EGPUSortTest)FXConsoleVariables::TestGPUSort.GetValueOnGameThread(), GetFeatureLevel());
 			// Reset CVar
 			static IConsoleVariable* CVarTestGPUSort = IConsoleManager::Get().FindConsoleVariable(TEXT("FX.TestGPUSort"));
-			CVarTestGPUSort->Set(0);
+
+			// todo: bad use of console variables, this should be a console command 
+			CVarTestGPUSort->Set(0, ECVF_SetByCode);
 		}
 
 		// Before ticking GPU particles, ensure any pending curves have been
diff --git a/Engine/Source/Runtime/Engine/Private/RenderingSettings.cpp b/Engine/Source/Runtime/Engine/Private/RenderingSettings.cpp
--- a/Engine/Source/Runtime/Engine/Private/RenderingSettings.cpp
+++ b/Engine/Source/Runtime/Engine/Private/RenderingSettings.cpp
@@ -79,19 +79,19 @@ void URendererSettings::PostEditChangeProperty(FPropertyChangedEvent& PropertyCh
 				UByteProperty* ByteProperty = Cast<UByteProperty>(PropertyChangedEvent.Property);
 				if (ByteProperty != NULL && ByteProperty->Enum != NULL)
 				{
-					CVar->Set(ByteProperty->GetPropertyValue_InContainer(this));
+					CVar->Set(ByteProperty->GetPropertyValue_InContainer(this), ECVF_SetByEditorSetting);
 				}
 				else if (UBoolProperty* BoolProperty = Cast<UBoolProperty>(PropertyChangedEvent.Property))
 				{
-					CVar->Set((int32)BoolProperty->GetPropertyValue_InContainer(this));
+					CVar->Set((int32)BoolProperty->GetPropertyValue_InContainer(this), ECVF_SetByEditorSetting);
 				}
 				else if (UIntProperty* IntProperty = Cast<UIntProperty>(PropertyChangedEvent.Property))
 				{
-					CVar->Set(IntProperty->GetPropertyValue_InContainer(this));
+					CVar->Set(IntProperty->GetPropertyValue_InContainer(this), ECVF_SetByEditorSetting);
 				}
 				else if (UFloatProperty* FloatProperty = Cast<UFloatProperty>(PropertyChangedEvent.Property))
 				{
-					CVar->Set(FloatProperty->GetPropertyValue_InContainer(this));
+					CVar->Set(FloatProperty->GetPropertyValue_InContainer(this), ECVF_SetByEditorSetting);
 				}
 			}
 			else
diff --git a/Engine/Source/Runtime/Engine/Private/Scalability.cpp b/Engine/Source/Runtime/Engine/Private/Scalability.cpp
--- a/Engine/Source/Runtime/Engine/Private/Scalability.cpp
+++ b/Engine/Source/Runtime/Engine/Private/Scalability.cpp
@@ -154,7 +154,7 @@ static void SetGroupQualityLevel(const TCHAR* InGroupName, int32 InQualityLevel)
 
 //	UE_LOG(LogConsoleResponse, Display, TEXT("  %s %d"), InGroupName, InQualityLevel);
 
-	ApplyCVarSettingsGroupFromIni(InGroupName, InQualityLevel, *GScalabilityIni);
+	ApplyCVarSettingsGroupFromIni(InGroupName, InQualityLevel, *GScalabilityIni, ECVF_SetByScalability);
 }
 
 static void SetResolutionQualityLevel(int32 InResolutionQualityLevel)
@@ -168,7 +168,7 @@ static void SetResolutionQualityLevel(int32 InResolutionQualityLevel)
 	// if it wasn't created yet we either need to change the order or store like we do for ini loading
 	check(CVar);
 
-	CVar->Set(InResolutionQualityLevel);
+	CVar->Set(InResolutionQualityLevel, ECVF_SetByScalability);
 }
 
 void OnChangeResolutionQuality(IConsoleVariable* Var)
@@ -346,13 +346,13 @@ void ProcessCommand(const TCHAR* Cmd, FOutputDevice& Ar)
 
 void SetQualityLevels(const FQualityLevels& QualityLevels)
 {
-	CVarResolutionQuality.AsVariable()->Set(QualityLevels.ResolutionQuality);
-	CVarViewDistanceQuality.AsVariable()->Set(QualityLevels.ViewDistanceQuality);
-	CVarAntiAliasingQuality.AsVariable()->Set(QualityLevels.AntiAliasingQuality);
-	CVarShadowQuality.AsVariable()->Set(QualityLevels.ShadowQuality);
-	CVarPostProcessQuality.AsVariable()->Set(QualityLevels.PostProcessQuality);
-	CVarTextureQuality.AsVariable()->Set(QualityLevels.TextureQuality);
-	CVarEffectsQuality.AsVariable()->Set(QualityLevels.EffectsQuality);
+	CVarResolutionQuality.AsVariable()->Set(QualityLevels.ResolutionQuality, ECVF_SetByScalability);
+	CVarViewDistanceQuality.AsVariable()->Set(QualityLevels.ViewDistanceQuality, ECVF_SetByScalability);
+	CVarAntiAliasingQuality.AsVariable()->Set(QualityLevels.AntiAliasingQuality, ECVF_SetByScalability);
+	CVarShadowQuality.AsVariable()->Set(QualityLevels.ShadowQuality, ECVF_SetByScalability);
+	CVarPostProcessQuality.AsVariable()->Set(QualityLevels.PostProcessQuality, ECVF_SetByScalability);
+	CVarTextureQuality.AsVariable()->Set(QualityLevels.TextureQuality, ECVF_SetByScalability);
+	CVarEffectsQuality.AsVariable()->Set(QualityLevels.EffectsQuality, ECVF_SetByScalability);
 }
 
 FQualityLevels GetQualityLevels()
diff --git a/Engine/Source/Runtime/Engine/Private/SystemSettings.cpp b/Engine/Source/Runtime/Engine/Private/SystemSettings.cpp
--- a/Engine/Source/Runtime/Engine/Private/SystemSettings.cpp
+++ b/Engine/Source/Runtime/Engine/Private/SystemSettings.cpp
@@ -86,18 +86,18 @@ public:
 			|| FCString::Stricmp(Value, TEXT("Yes")) == 0
 			|| FCString::Stricmp(Value, TEXT("On")) == 0)
 			{
-				CVar->Set(1);
+				CVar->Set(1, ECVF_SetBySystemSettingsIni);
 			}
 			else
 			if(FCString::Stricmp(Value, TEXT("False")) == 0
 			|| FCString::Stricmp(Value, TEXT("No")) == 0
 			|| FCString::Stricmp(Value, TEXT("Off")) == 0)
 			{
-				CVar->Set(0);
+				CVar->Set(0, ECVF_SetBySystemSettingsIni);
 			}
 			else
 			{
-				CVar->Set(Value);
+				CVar->Set(Value, ECVF_SetBySystemSettingsIni);
 			}
 		}
 	}
@@ -232,7 +232,6 @@ void FSystemSettingsData::DumpTextureLODGroup(TextureGroup TextureGroupID, const
  */
 FSystemSettings::FSystemSettings() :
 	bIsEditor( false ),
-	bInitialUseHighQualityLightmaps(false),
 	Force0Mask(ESFIM_All0),
 	Force1Mask(ESFIM_All0)
 {
@@ -340,8 +339,6 @@ void FSystemSettings::Initialize( bool bSetupForEditor )
 
 	ApplyOverrides();
 
-	bInitialUseHighQualityLightmaps = AllowHighQualityLightmaps(GMaxRHIFeatureLevel);
-
 	IConsoleManager::Get().RegisterConsoleVariableSink(FConsoleCommandDelegate::CreateRaw(this, &FSystemSettings::CVarSink));
 
 	// intialize a critical texture streaming value used by texture loading, etc
@@ -355,115 +352,12 @@ void FSystemSettings::CVarSink()
 	ApplyOverrides();
 }
 
-bool FSystemSettings::HandleLowendCommand( const TCHAR* Cmd , FOutputDevice& Ar )
-{
-	{
-		static auto CVar = IConsoleManager::Get().FindConsoleVariable(TEXT("r.OnlyStreamInTextures"));
-		CVar->Set(0);
-	}
-
-	{
-		auto CVar = IConsoleManager::Get().FindConsoleVariable(TEXT("r.Shadow.MinResolution"));
-		CVar->Set(32);
-	}
-
-	{
-		auto CVar = IConsoleManager::Get().FindConsoleVariable(TEXT("r.Shadow.PreShadowResolutionFactor"));
-		CVar->Set(0.5f);
-	}
-
-	{
-		auto CVar = IConsoleManager::Get().FindConsoleVariable(TEXT("r.Shadow.TexelsPerPixel"));
-		CVar->Set(2.0f);
-	}
-
-	return true;
-}
-
-bool FSystemSettings::HandleHighendCommand( const TCHAR* Cmd , FOutputDevice& Ar )
-{
-	{
-		static auto CVar = IConsoleManager::Get().FindConsoleVariable(TEXT("r.OnlyStreamInTextures"));
-		CVar->Set(0);
-	}
-	
-	{
-		auto CVar = IConsoleManager::Get().FindConsoleVariable(TEXT("r.Shadow.MinResolution"));
-		CVar->Set(32);
-	}
-
-	{
-		auto CVar = IConsoleManager::Get().FindConsoleVariable(TEXT("r.Shadow.PreShadowResolutionFactor"));
-		CVar->Set(1.0f);
-	}
-
-	{
-		auto CVar = IConsoleManager::Get().FindConsoleVariable(TEXT("r.Shadow.TexelsPerPixel"));
-		CVar->Set(2.0f);
-	}
-
-	return true;
-}
-
-bool FSystemSettings::HandleResetCommand( const TCHAR* Cmd , FOutputDevice& Ar )
-{
-	// Reset values to defaults from ini.
-	LoadFromIni();
-	return true;
-}
-
-/**
- * Exec handler implementation.
- *
- * @param InWorld World context
- * @param Cmd	Command to parse
- * @param Ar	Output device to log to
- *
- * @return true if command was handled, false otherwise
- */
 bool FSystemSettings::Exec( UWorld* InWorld, const TCHAR* Cmd, FOutputDevice& Ar )
 {
-	FSystemSettingsData OldSystemSettings = *this;
-
-	// Keep track whether the command was handled or not.
-	bool bHandledCommand = false;
-
-	if( FParse::Command(&Cmd,TEXT("SCALE")) )
-	{
-		// Some of these settings are used in both threads so we need to stop the rendering thread before changing them.
-		FlushRenderingCommands();
-
-		if( FParse::Command(&Cmd,TEXT("LOWEND")) )
-		{
-			bHandledCommand	= HandleLowendCommand( Cmd, Ar );
-		}
-		else if( FParse::Command(&Cmd,TEXT("HIGHEND")) )
-		{
-			bHandledCommand	= HandleHighendCommand( Cmd, Ar );	
-		}
-
-		if (!bHandledCommand)
-		{
-			Ar.Logf(TEXT("Unrecognized system setting (note that console variables can be set much easier)"));
-		}
-		else
-		{
-			// Write the new settings to the INI.
-			SaveToIni();
-		}
-	}
-
-	return bHandledCommand;
+	// no longer needed, we have the "Scalability" console command
+	return false;
 }
 
-/**
- * Reads a single entry and parses it into the group array.
- *
- * @param	TextureGroupID		Index/enum of group to parse
- * @param	MinLODSize			Minimum size, in pixels, below which the code won't bias.
- * @param	MaxLODSize			Maximum size, in pixels, above which the code won't bias.
- * @param	LODBias				Group LOD bias.
- */
 void FSystemSettings::SetTextureLODGroup(TextureGroup TextureGroupID, int32 MinLODSize, int32 MaxLODSize, int32 LODBias, TextureMipGenSettings MipGenSettings)
 {
 	TextureLODSettings.GetTextureLODGroup(TextureGroupID).MinLODMipCount	= FMath::CeilLogTwo( MinLODSize );
@@ -608,38 +502,38 @@ void FSystemSettings::ApplyOverrides()
 	{
 		// Modify various system settings to get the best quality regardless of performance impact
 		{
-			auto CVar = IConsoleManager::Get().FindConsoleVariable(TEXT("r.Shadow.MinResolution"));
-			CVar->Set(16);
+			static auto CVar = IConsoleManager::Get().FindConsoleVariable(TEXT("r.Shadow.MinResolution"));
+			CVar->Set(16, ECVF_SetByCode);
 		}
 
 		// Disable shadow fading out over distance
 		{
-			auto CVar = IConsoleManager::Get().FindConsoleVariable(TEXT("r.Shadow.FadeResolution"));
-			CVar->Set(1);
+			static auto CVar = IConsoleManager::Get().FindConsoleVariable(TEXT("r.Shadow.FadeResolution"));
+			CVar->Set(1, ECVF_SetByCode);
 		}
 
 		// Increase minimum preshadow resolution
 		{
-			auto CVar = IConsoleManager::Get().FindConsoleVariable(TEXT("r.Shadow.MinPreShadowResolution"));
-			CVar->Set(16);
+			static auto CVar = IConsoleManager::Get().FindConsoleVariable(TEXT("r.Shadow.MinPreShadowResolution"));
+			CVar->Set(16, ECVF_SetByCode);
 		}
 
 		// Disable preshadow fading out over distance
 		{
-			auto CVar = IConsoleManager::Get().FindConsoleVariable(TEXT("r.Shadow.PreShadowFadeResolution"));
-			CVar->Set(1);
+			static auto CVar = IConsoleManager::Get().FindConsoleVariable(TEXT("r.Shadow.PreShadowFadeResolution"));
+			CVar->Set(1, ECVF_SetByCode);
 		}
 
 		// Increase shadow texel density
 		{
-			auto CVar = IConsoleManager::Get().FindConsoleVariable(TEXT("r.Shadow.TexelsPerPixel"));
-			CVar->Set(4.0f);
+			static auto CVar = IConsoleManager::Get().FindConsoleVariable(TEXT("r.Shadow.TexelsPerPixel"));
+			CVar->Set(4.0f, ECVF_SetByCode);
 		}
 
 		// Don't downsample preshadows
 		{
-			auto CVar = IConsoleManager::Get().FindConsoleVariable(TEXT("r.Shadow.PreShadowResolutionFactor"));
-			CVar->Set(1.0f);
+			static auto CVar = IConsoleManager::Get().FindConsoleVariable(TEXT("r.Shadow.PreShadowResolutionFactor"));
+			CVar->Set(1.0f, ECVF_SetByCode);
 		}
 
 		for (int32 GroupIndex = 0; GroupIndex < TEXTUREGROUP_MAX; GroupIndex++)
diff --git a/Engine/Source/Runtime/Engine/Private/Texture.cpp b/Engine/Source/Runtime/Engine/Private/Texture.cpp
--- a/Engine/Source/Runtime/Engine/Private/Texture.cpp
+++ b/Engine/Source/Runtime/Engine/Private/Texture.cpp
@@ -482,7 +482,7 @@ bool UTexture::ForceUpdateTextureStreaming()
 		// Make sure textures can be streamed out so that we can unload current mips.
 		static auto CVarOnlyStreamInTextures = IConsoleManager::Get().FindConsoleVariable(TEXT("r.OnlyStreamInTextures"));
 		const bool bOldOnlyStreamInTextures = CVarOnlyStreamInTextures->GetInt() != 0;
-		CVarOnlyStreamInTextures->Set(false);
+		CVarOnlyStreamInTextures->Set(false, ECVF_SetByCode);
 
 #if WITH_EDITOR
 		for( TObjectIterator<UTexture2D> It; It; ++It )
@@ -502,7 +502,7 @@ bool UTexture::ForceUpdateTextureStreaming()
 		IStreamingManager::Get().BlockTillAllRequestsFinished();
 
 		// Restore streaming out of textures.
-		CVarOnlyStreamInTextures->Set(bOldOnlyStreamInTextures);
+		CVarOnlyStreamInTextures->Set(bOldOnlyStreamInTextures, ECVF_SetByCode);
 	}
 
 	return true;
diff --git a/Engine/Source/Runtime/Engine/Private/UnrealEngine.cpp b/Engine/Source/Runtime/Engine/Private/UnrealEngine.cpp
--- a/Engine/Source/Runtime/Engine/Private/UnrealEngine.cpp
+++ b/Engine/Source/Runtime/Engine/Private/UnrealEngine.cpp
@@ -147,7 +147,8 @@ static FAutoConsoleVariable CVarSystemResolution(
 	TEXT("Set the display resolution for the current game view. Has no effect in the editor.\n")
 	TEXT("e.g. 1280x720w for windowed\n")
 	TEXT("     1920x1080f for fullscreen\n")
-	TEXT("     1920x1080wf for windowed fullscreen")
+	TEXT("     1920x1080wf for windowed fullscreen\n")
+	TEXT("     1920x1080wm for windowed mirror")
 	);
 
 static TAutoConsoleVariable<float> CVarDepthOfFieldNearBlurSizeThreshold(
@@ -1164,7 +1165,7 @@ void UEngine::ParseCommandline()
 	
 		if (CVarDumpFrames)
 		{
-			CVarDumpFrames->Set(1);
+			CVarDumpFrames->Set(1, ECVF_SetByCommandline);
 		}
 	}
 
@@ -10298,7 +10299,7 @@ void FSystemResolution::RequestResolutionChange(int32 InResX, int32 InResY, EWin
 	}
 
 	FString NewValue = FString::Printf(TEXT("%dx%d%s"), InResX, InResY, *WindowModeSuffix);
-	CVarSystemResolution->Set(*NewValue);
+	CVarSystemResolution->Set(*NewValue, ECVF_SetByCode);
 }
 
 
diff --git a/Engine/Source/Runtime/Engine/Public/SystemSettings.h b/Engine/Source/Runtime/Engine/Public/SystemSettings.h
--- a/Engine/Source/Runtime/Engine/Public/SystemSettings.h
+++ b/Engine/Source/Runtime/Engine/Public/SystemSettings.h
@@ -115,13 +115,6 @@ public:
 	virtual bool Exec( UWorld* InWorld, const TCHAR* Cmd, FOutputDevice& Ar ) override;
 	// End Exec Interface
 
-	/** 
-	 * Exec command handlers
-	 */
-	bool HandleLowendCommand( const TCHAR* Cmd , FOutputDevice& Ar );
-	bool HandleHighendCommand( const TCHAR* Cmd , FOutputDevice& Ar );
-	bool HandleResetCommand( const TCHAR* Cmd , FOutputDevice& Ar );
-
 	/** Mask where 1 bits mean we want to force the engine show flag to be off */
 	const FEngineShowFlags& GetForce0Mask() const { return Force0Mask; }
 	/** Mask where 1 bits mean we want to force the engine show flag to be on */
@@ -131,9 +124,6 @@ private:
 	/** Since System Settings is called into before GIsEditor is set, we must cache this value. */
 	bool bIsEditor;
 
-	/** The initial high quality lightmap settings. Used to stop the user changing it at runtime. */
-	bool bInitialUseHighQualityLightmaps;
-
 	/** Mask where 1 bits mean we want to force the engine show flag to be off */
 	FEngineShowFlags Force0Mask;
 	/** Mask where 1 bits mean we want to force the engine show flag to be on */
diff --git a/Engine/Source/Runtime/Launch/Private/LaunchEngineLoop.cpp b/Engine/Source/Runtime/Launch/Private/LaunchEngineLoop.cpp
--- a/Engine/Source/Runtime/Launch/Private/LaunchEngineLoop.cpp
+++ b/Engine/Source/Runtime/Launch/Private/LaunchEngineLoop.cpp
@@ -1065,7 +1065,7 @@ int32 FEngineLoop::PreInit( const TCHAR* CmdLine )
 	GSystemSettings.Initialize( bHasEditorToken );
 
 	// Apply renderer settings from console variables stored in the INI.
-	ApplyCVarSettingsFromIni(TEXT("/Script/Engine.RendererSettings"),*GEngineIni);
+	ApplyCVarSettingsFromIni(TEXT("/Script/Engine.RendererSettings"),*GEngineIni, ECVF_SetBySystemSettingsIni);
 
 #if !UE_SERVER
 	if (!IsRunningDedicatedServer())
diff --git a/Engine/Source/Runtime/RHI/Private/RHICommandList.cpp b/Engine/Source/Runtime/RHI/Private/RHICommandList.cpp
--- a/Engine/Source/Runtime/RHI/Private/RHICommandList.cpp
+++ b/Engine/Source/Runtime/RHI/Private/RHICommandList.cpp
@@ -385,7 +385,7 @@ void FRHICommandListExecutor::LatchBypass()
 			if (FParse::Param(FCommandLine::Get(),TEXT("parallelrendering")) && CVarRHICmdBypass.GetValueOnRenderThread() >= 1)
 			{
 				IConsoleVariable* BypassVar = IConsoleManager::Get().FindConsoleVariable(TEXT("r.RHICmdBypass"));
-				BypassVar->Set(0);
+				BypassVar->Set(0, ECVF_SetByCommandline);
 			}
 		}
 
