Integrated GitHub 540 : Allow the HUD to be hidden from Matinee actor

"This allows the HUD to be hidden by a Matinee actor by Adding the option to the cinematic properties category.
This brings Matinee actor in line with the "Set Cinematic Mode" blueprint node and the MatineeRecordMovie system."

diff --git a/Engine/Source/Runtime/Engine/Classes/Matinee/MatineeActor.h b/Engine/Source/Runtime/Engine/Classes/Matinee/MatineeActor.h
--- a/Engine/Source/Runtime/Engine/Classes/Matinee/MatineeActor.h
+++ b/Engine/Source/Runtime/Engine/Classes/Matinee/MatineeActor.h
@@ -139,10 +139,14 @@ class AMatineeActor : public AActor
 	UPROPERTY(EditAnywhere, BlueprintReadWrite, Category=Cinematic)
 	uint32 bDisableLookAtInput:1;
 
-	/** Hide Player Pawn during play **/
+	/** Hide Player Pawn during play */
 	UPROPERTY(EditAnywhere, BlueprintReadWrite, Category=Cinematic)
 	uint32 bHidePlayer:1;
 
+	/** Hide HUD during play */
+	UPROPERTY(EditAnywhere, BlueprintReadWrite, Category = Cinematic)
+	uint32 bHideHud:1;
+
 	/** @todo UE4 matinee - shouldnt be directly editable.  Needs a nice interface in matinee */
 	UPROPERTY(replicated)
 	TArray<struct FInterpGroupActorInfo> GroupActorInfos;
diff --git a/Engine/Source/Runtime/Engine/Private/Interpolation.cpp b/Engine/Source/Runtime/Engine/Private/Interpolation.cpp
--- a/Engine/Source/Runtime/Engine/Private/Interpolation.cpp
+++ b/Engine/Source/Runtime/Engine/Private/Interpolation.cpp
@@ -1098,14 +1098,14 @@ void AMatineeActor::DisableRadioFilterIfNeeded()
 void AMatineeActor::EnableCinematicMode(bool bEnable)
 {
 	// iterate through the controller list and set cinematic mode if necessary
-	bool bNeedsCinematicMode = bDisableMovementInput || bDisableLookAtInput || bHidePlayer;
+	bool bNeedsCinematicMode = bDisableMovementInput || bDisableLookAtInput || bHidePlayer || bHideHud;
 
 	if (bNeedsCinematicMode)
 	{
 		for (FConstPlayerControllerIterator Iterator = GetWorld()->GetPlayerControllerIterator(); Iterator; ++Iterator)
 		{
 			APlayerController *PC = *Iterator;
-			PC->SetCinematicMode(bEnable, bHidePlayer, false, bDisableMovementInput, bDisableLookAtInput);
+			PC->SetCinematicMode(bEnable, bHidePlayer, bHideHud, bDisableMovementInput, bDisableLookAtInput);
 		}
 	}
 }
