Fix HitBox hit tests for touchscreens (where the touches are treated as mouse, not touch)

diff --git a/Engine/Source/Runtime/Engine/Classes/GameFramework/HUD.h b/Engine/Source/Runtime/Engine/Classes/GameFramework/HUD.h
--- a/Engine/Source/Runtime/Engine/Classes/GameFramework/HUD.h
+++ b/Engine/Source/Runtime/Engine/Classes/GameFramework/HUD.h
@@ -583,9 +583,10 @@ public:
 	
 	/**
 	 * Update the list of hitboxes and dispatch events for any hits.
-	 * @param	InEventType	Input event that triggered the call.
+	 * @param   ClickLocation	Location of the click event
+	 * @param	InEventType		Type of input event that triggered the call.
 	 */
-	bool UpdateAndDispatchHitBoxClickEvents(const FVector2D ClickLocation, const EInputEvent InEventType, const bool bIsTouchEvent);
+	bool UpdateAndDispatchHitBoxClickEvents(const FVector2D ClickLocation, const EInputEvent InEventType);
 	
 	/**
 	 * Update a the list of hitboxes that have been hit this frame.
diff --git a/Engine/Source/Runtime/Engine/Private/HUD.cpp b/Engine/Source/Runtime/Engine/Private/HUD.cpp
--- a/Engine/Source/Runtime/Engine/Private/HUD.cpp
+++ b/Engine/Source/Runtime/Engine/Private/HUD.cpp
@@ -879,15 +879,15 @@ bool AHUD::AnyCurrentHitBoxHits() const
 	return HitBoxHits.Num() != 0;
 }
 
-bool AHUD::UpdateAndDispatchHitBoxClickEvents(FVector2D ClickLocation, const EInputEvent InEventType, const bool bIsTouchEvent)
+bool AHUD::UpdateAndDispatchHitBoxClickEvents(FVector2D ClickLocation, const EInputEvent InEventType)
 {
 	ClickLocation += GetCoordinateOffset();
 
 	const bool bIsClickEvent = (InEventType == IE_Pressed || InEventType == IE_DoubleClick);
 	bool bHit = false;
 
-	// If this is a touch event we likely don't have the hit box in the hit list yet so we need to check all HitBoxes
-	if (bIsTouchEvent && bIsClickEvent)
+	// If this is a click event we may not have the hit box in the hit list yet (particularly for touch events) so we need to check all HitBoxes
+	if (bIsClickEvent)
 	{
 		for (FHUDHitBox& HitBox : HitBoxMap)
 		{
@@ -912,11 +912,7 @@ bool AHUD::UpdateAndDispatchHitBoxClickEvents(FVector2D ClickLocation, const EIn
 			{
 				bHit = true;
 
-				if (bIsClickEvent)
-				{
-					ReceiveHitBoxClick(HitBoxHit->GetName());
-				}
-				else if (InEventType == IE_Released)
+				if (InEventType == IE_Released)
 				{
 					ReceiveHitBoxRelease(HitBoxHit->GetName());
 				}
diff --git a/Engine/Source/Runtime/Engine/Private/PlayerController.cpp b/Engine/Source/Runtime/Engine/Private/PlayerController.cpp
--- a/Engine/Source/Runtime/Engine/Private/PlayerController.cpp
+++ b/Engine/Source/Runtime/Engine/Private/PlayerController.cpp
@@ -2118,7 +2118,7 @@ bool APlayerController::InputKey(FKey Key, EInputEvent EventType, float AmountDe
 				}
 				if( GetHUD() )
 				{
-					if (GetHUD()->UpdateAndDispatchHitBoxClickEvents(MousePosition, EventType, false))
+					if (GetHUD()->UpdateAndDispatchHitBoxClickEvents(MousePosition, EventType))
 					{
 						ClickedPrimitive = NULL;
 					}
@@ -2181,7 +2181,7 @@ bool APlayerController::InputTouch(uint32 Handle, ETouchType::Type Type, const F
 			{
 				if (Type == ETouchType::Began || Type == ETouchType::Ended)
 				{
-					if (GetHUD()->UpdateAndDispatchHitBoxClickEvents(TouchLocation, (Type == ETouchType::Began ? EInputEvent::IE_Pressed : EInputEvent::IE_Released), true))
+					if (GetHUD()->UpdateAndDispatchHitBoxClickEvents(TouchLocation, (Type == ETouchType::Began ? EInputEvent::IE_Pressed : EInputEvent::IE_Released)))
 					{
 						CurrentComponent = NULL;
 					}
