When not async loading, and a package isn't fully loaded, flush all async loading. This fixes the issue that we end up finding the object in the package, but since it wasn't fully loaded it causes issues

UE-2986

diff --git a/Engine/Source/Runtime/Engine/Private/PackageMapClient.cpp b/Engine/Source/Runtime/Engine/Private/PackageMapClient.cpp
--- a/Engine/Source/Runtime/Engine/Private/PackageMapClient.cpp
+++ b/Engine/Source/Runtime/Engine/Private/PackageMapClient.cpp
@@ -1865,8 +1865,17 @@ UObject* FNetGUIDCache::GetObjectFromNetGUID( const FNetworkGUID& NetGUID, const
 
 		if ( !Package->IsFullyLoaded() )
 		{
-			// If this package isn't fully loaded, don't complain, assume it will fully load at some point
-			return NULL;
+			if ( CVarAllowAsyncLoading.GetValueOnGameThread() > 0 )
+			{
+				// If this package isn't fully loaded (and we are async loading here), don't complain, assume it will fully load at some point
+				return NULL;
+			}
+			else
+			{
+				// If package isn't fully loaded, flush async loading now
+				FlushAsyncLoading(); 
+				check( Package->IsFullyLoaded() );
+			}
 		}
 
 		if ( Package->GetGuid() != CacheObjectPtr->PackageGuid && CVarIgnorePackageMismatch.GetValueOnGameThread() == 0 )
