#ue4 - Perf: only add prerequisite to tick function if that function itself can tick.

diff --git a/Engine/Source/Runtime/Engine/Private/TickTaskManager.cpp b/Engine/Source/Runtime/Engine/Private/TickTaskManager.cpp
--- a/Engine/Source/Runtime/Engine/Private/TickTaskManager.cpp
+++ b/Engine/Source/Runtime/Engine/Private/TickTaskManager.cpp
@@ -1020,7 +1020,10 @@ void FTickFunction::SetTickFunctionEnable(bool bInEnabled)
 **/
 void FTickFunction::AddPrerequisite(UObject* TargetObject, struct FTickFunction& TargetTickFunction)
 {
-	if (TargetTickFunction.bCanEverTick || TargetTickFunction.IsTickFunctionRegistered())
+	const bool bThisCanTick = (bCanEverTick || IsTickFunctionRegistered());
+	const bool bTargetCanTick = (TargetTickFunction.bCanEverTick || TargetTickFunction.IsTickFunctionRegistered());
+	
+	if (bThisCanTick && bTargetCanTick)
 	{
 		Prerequisites.AddUnique(FTickPrerequisite(TargetObject, TargetTickFunction));
 	}
