#ue4 - Fix crash in CharacterMovementComponent if setting UpdatedComponent to NULL.

diff --git a/Engine/Source/Runtime/Engine/Private/CharacterMovementComponent.cpp b/Engine/Source/Runtime/Engine/Private/CharacterMovementComponent.cpp
--- a/Engine/Source/Runtime/Engine/Private/CharacterMovementComponent.cpp
+++ b/Engine/Source/Runtime/Engine/Private/CharacterMovementComponent.cpp
@@ -380,7 +380,7 @@ void UCharacterMovementComponent::SetUpdatedComponent(UPrimitiveComponent* NewUp
 		StopActiveMovement();
 	}
 
-	if (bEnablePhysicsInteraction)
+	if (UpdatedComponent && bEnablePhysicsInteraction)
 	{
 		UpdatedComponent->OnComponentBeginOverlap.AddUniqueDynamic(this, &UCharacterMovementComponent::CapsuleTouched);
 	}
