#ue4 - MovementComponent::SetUpdatedComponent avoids changing physics volume and trying to unbind delegate (causing ensure) if object is pending kill.

diff --git a/Engine/Source/Runtime/Engine/Private/MovementComponent.cpp b/Engine/Source/Runtime/Engine/Private/MovementComponent.cpp
--- a/Engine/Source/Runtime/Engine/Private/MovementComponent.cpp
+++ b/Engine/Source/Runtime/Engine/Private/MovementComponent.cpp
@@ -38,8 +38,11 @@ void UMovementComponent::SetUpdatedComponent(UPrimitiveComponent* NewUpdatedComp
 	if ( UpdatedComponent && UpdatedComponent!=NewUpdatedComponent )
 	{
 		UpdatedComponent->bShouldUpdatePhysicsVolume = false;
-		UpdatedComponent->SetPhysicsVolume(NULL, true);
-		UpdatedComponent->PhysicsVolumeChangedDelegate.RemoveDynamic(this, &UMovementComponent::PhysicsVolumeChanged);
+		if (!UpdatedComponent->IsPendingKill())
+		{
+			UpdatedComponent->SetPhysicsVolume(NULL, true);
+			UpdatedComponent->PhysicsVolumeChangedDelegate.RemoveDynamic(this, &UMovementComponent::PhysicsVolumeChanged);
+		}
 
 		// remove from tick prerequisite
 		UpdatedComponent->PrimaryComponentTick.RemovePrerequisite(this, PrimaryComponentTick); 
