Remove UPROPERTY from internal component state booleans

diff --git a/Engine/Source/Runtime/Engine/Classes/Components/ActorComponent.h b/Engine/Source/Runtime/Engine/Classes/Components/ActorComponent.h
--- a/Engine/Source/Runtime/Engine/Classes/Components/ActorComponent.h
+++ b/Engine/Source/Runtime/Engine/Classes/Components/ActorComponent.h
@@ -38,6 +38,19 @@ protected:
 	 */
 	uint32 bRegistered:1;
 
+	/** If the render state is currently created for this component */
+	uint32 bRenderStateCreated:1;
+
+	/** If the physics state is currently created for this component */
+	uint32 bPhysicsStateCreated:1;
+
+	/** Is this component currently replicating? Should the network code consider it for replication? Owning Actor must be replicating first! */
+	UPROPERTY(EditDefaultsOnly, BlueprintReadOnly, Replicated, Category=ComponentReplication,meta=(DisplayName = "Component Replicates"))
+	uint32 bReplicates:1;
+
+	/** Is this component safe to ID over the network by name?  */
+	uint32 bNetAddressable:1;
+
 private:
 	/** Is this component in need of its whole state being sent to the renderer? */
 	uint32 bRenderStateDirty:1;
@@ -51,19 +64,15 @@ private:
 public:
 
 	/** Does this component automatically register with its owner */
-	UPROPERTY()
 	uint32 bAutoRegister:1;
 
 	/** Should this component be ticked in the editor */
-	UPROPERTY()
 	uint32 bTickInEditor:1;
 
 	/** If true, this component never needs a render update. */
-	UPROPERTY(transient)
 	uint32 bNeverNeedsRenderUpdate:1;
 
 	/** Can we tick this concurrently on other threads? */
-	UPROPERTY()
 	uint32 bAllowConcurrentTick:1;
 
 	/** True if this component was created by a construction script, and will be destroyed by DestroyConstructedComponents */
@@ -79,7 +88,6 @@ public:
 	uint32 bIsActive:1;
 
 	/** If TRUE, we call the virtual InitializeComponent */
-	UPROPERTY()
 	uint32 bWantsInitializeComponent:1;
 
 	/** Indicates that OnCreatedComponent has been called, but OnDestroyedComponent has not yet */
@@ -172,20 +180,6 @@ public:
 	ENetMode GetNetMode() const;
 
 protected:
-	/** If the render state is currently created for this component */
-	uint32 bRenderStateCreated:1;
-
-	/** If the physics state is currently created for this component */
-	uint32 bPhysicsStateCreated:1;
-
-	/** Is this component currently replicating? Should the network code consider it for replication? Owning Actor must be replicating first! */
-	UPROPERTY(EditDefaultsOnly, BlueprintReadOnly, Replicated, Category=ComponentReplication,meta=(DisplayName = "Component Replicates"))
-	uint32 bReplicates:1;
-
-	/** Is this component safe to ID over the network by name?  */
-	UPROPERTY(transient)
-	uint32 bNetAddressable:1;
-
 
 	/** 
 	 * Pointer to the world that this component is currently registered with. 
