SetOcclusionTint and SetMinOcclusion BP functions for skylight components
SetCastShadows is now accessible to skylight components

diff --git a/Engine/Source/Runtime/Engine/Classes/Components/LightComponent.h b/Engine/Source/Runtime/Engine/Classes/Components/LightComponent.h
--- a/Engine/Source/Runtime/Engine/Classes/Components/LightComponent.h
+++ b/Engine/Source/Runtime/Engine/Classes/Components/LightComponent.h
@@ -176,9 +176,6 @@ public:
 	UFUNCTION(BlueprintCallable, Category="Rendering|Components|Light")
 	void SetLightFunctionFadeDistance(float NewLightFunctionFadeDistance);
 
-	UFUNCTION(BlueprintCallable, Category="Rendering|Components|Light")
-	void SetCastShadows(bool bNewValue);
-
 	UFUNCTION(BlueprintCallable, Category="Rendering|Components|Light")
 	void SetAffectDynamicIndirectLighting(bool bNewValue);
 
diff --git a/Engine/Source/Runtime/Engine/Classes/Components/LightComponentBase.h b/Engine/Source/Runtime/Engine/Classes/Components/LightComponentBase.h
--- a/Engine/Source/Runtime/Engine/Classes/Components/LightComponentBase.h
+++ b/Engine/Source/Runtime/Engine/Classes/Components/LightComponentBase.h
@@ -93,6 +93,9 @@ class ENGINE_API ULightComponentBase : public USceneComponent
 	float DynamicEditorTextureScale;
 #endif
 
+	UFUNCTION(BlueprintCallable, Category="Rendering|Components|Light")
+	void SetCastShadows(bool bNewValue);
+
 	virtual void Serialize(FArchive& Ar) override;
 
 	/**
diff --git a/Engine/Source/Runtime/Engine/Classes/Components/SkyLightComponent.h b/Engine/Source/Runtime/Engine/Classes/Components/SkyLightComponent.h
--- a/Engine/Source/Runtime/Engine/Classes/Components/SkyLightComponent.h
+++ b/Engine/Source/Runtime/Engine/Classes/Components/SkyLightComponent.h
@@ -168,6 +168,12 @@ class ENGINE_API USkyLightComponent : public ULightComponentBase
 	UFUNCTION(BlueprintCallable, Category="SkyLight")
 	void SetCubemap(UTextureCube* NewCubemap);
 
+	UFUNCTION(BlueprintCallable, Category="Rendering|Components|SkyLight")
+	void SetOcclusionTint(const FColor& InTint);
+
+	UFUNCTION(BlueprintCallable, Category="Rendering|Components|SkyLight")
+	void SetMinOcclusion(float InMinOcclusion);
+
 	/** Indicates that the capture needs to recapture the scene, adds it to the recapture queue. */
 	void SetCaptureIsDirty();
 
diff --git a/Engine/Source/Runtime/Engine/Private/LightComponent.cpp b/Engine/Source/Runtime/Engine/Private/LightComponent.cpp
--- a/Engine/Source/Runtime/Engine/Private/LightComponent.cpp
+++ b/Engine/Source/Runtime/Engine/Private/LightComponent.cpp
@@ -65,6 +65,16 @@ FArchive& operator<<(FArchive& Ar, FStaticShadowDepthMap& ShadowMap)
 	return Ar;
 }
 
+void ULightComponentBase::SetCastShadows(bool bNewValue)
+{
+	if (!(IsRegistered() && Mobility == EComponentMobility::Static)
+		&& CastShadows != bNewValue)
+	{
+		CastShadows = bNewValue;
+		MarkRenderStateDirty();
+	}
+}
+
 void ULightComponentBase::Serialize(FArchive& Ar)
 {
 	Super::Serialize(Ar);
@@ -721,16 +731,6 @@ void ULightComponent::SetLightFunctionFadeDistance(float NewLightFunctionFadeDis
 	}
 }
 
-void ULightComponent::SetCastShadows(bool bNewValue)
-{
-	if (!(IsRegistered() && Mobility == EComponentMobility::Static)
-		&& CastShadows != bNewValue)
-	{
-		CastShadows = bNewValue;
-		MarkRenderStateDirty();
-	}
-}
-
 void ULightComponent::SetAffectDynamicIndirectLighting(bool bNewValue)
 {
 	if (!(IsRegistered() && Mobility == EComponentMobility::Static)
diff --git a/Engine/Source/Runtime/Engine/Private/SkyLightComponent.cpp b/Engine/Source/Runtime/Engine/Private/SkyLightComponent.cpp
--- a/Engine/Source/Runtime/Engine/Private/SkyLightComponent.cpp
+++ b/Engine/Source/Runtime/Engine/Private/SkyLightComponent.cpp
@@ -466,6 +466,28 @@ void USkyLightComponent::SetCubemap(UTextureCube* NewCubemap)
 	}
 }
 
+void USkyLightComponent::SetOcclusionTint(const FColor& InTint)
+{
+	// Can't set on a static light
+	if (!(IsRegistered() && Mobility == EComponentMobility::Static)
+		&& OcclusionTint != InTint)
+	{
+		OcclusionTint = InTint;
+		MarkRenderStateDirty();
+	}
+}
+
+void USkyLightComponent::SetMinOcclusion(float InMinOcclusion)
+{
+	// Can't set on a static light
+	if (!(IsRegistered() && Mobility == EComponentMobility::Static)
+		&& MinOcclusion != InMinOcclusion)
+	{
+		MinOcclusion = InMinOcclusion;
+		MarkRenderStateDirty();
+	}
+}
+
 void USkyLightComponent::RecaptureSky()
 {
 	SetCaptureIsDirty();
