Materials are no longer limited to 13 unique textures!
* This is achieved by using shared samplers on TextureSample nodes.  The property SamplerSource must be set to 'Shared: Wrap' or 'Shared: Clamp' to not consume a sampler slot.  Materials can use up to 128 unique textures with this setup.

diff --git a/Engine/Shaders/Common.usf b/Engine/Shaders/Common.usf
--- a/Engine/Shaders/Common.usf
+++ b/Engine/Shaders/Common.usf
@@ -152,7 +152,7 @@
 
 // Whether the platform supports independent texture and samplers
 // When enabled, different texture lookups can share samplers to allow more artist samplers in the base pass
-// Ideally this would just be enabled for all SM5 feature level platforms
+// Ideally this would just be enabled for all SM4 and above feature level platforms
 #define SUPPORTS_INDEPENDENT_SAMPLERS (PS4_PROFILE || SM5_PROFILE || SM4_PROFILE)
 
 // This would need to be a #define in GLSL to ignore the SamplerState, however, it is currently a function call in HLSL
diff --git a/Engine/Shaders/MaterialTemplate.usf b/Engine/Shaders/MaterialTemplate.usf
--- a/Engine/Shaders/MaterialTemplate.usf
+++ b/Engine/Shaders/MaterialTemplate.usf
@@ -586,6 +586,18 @@ MaterialFloat ProcessMaterialLinearGreyscaleTextureLookup(MaterialFloat TextureV
 	return TextureValue;
 }
 
+/** Accesses a shared material sampler or falls back if independent samplers are not supported. */
+SamplerState GetMaterialSharedSampler(SamplerState TextureSampler, SamplerState SharedSampler)
+{
+#if SUPPORTS_INDEPENDENT_SAMPLERS
+	return SharedSampler;
+#else
+	// Note: to match behavior on platforms that don't support SUPPORTS_INDEPENDENT_SAMPLERS, 
+	// TextureSampler should have been set to the same sampler.  This is not currently done.
+	return TextureSampler;
+#endif
+}
+
 /** Calculate a reflection vector about the specified world space normal. Optionally normalize this normal **/
 MaterialFloat3 ReflectionAboutCustomWorldNormal(FMaterialPixelParameters Parameters, MaterialFloat3 WorldNormal, bool bNormalizeInputNormal)
 {
diff --git a/Engine/Shaders/ShaderVersion.usf b/Engine/Shaders/ShaderVersion.usf
--- a/Engine/Shaders/ShaderVersion.usf
+++ b/Engine/Shaders/ShaderVersion.usf
@@ -2,4 +2,4 @@
 // Each time the console command is executed it generates a new GUID. As this file is included
 // in common.usf (which should be included in any shader) it allows to invalidate the shader DDC.
 // 
-// GUID = 4C36424F4976EF3266512DACC6E7D30A
\ No newline at end of file
+// GUID = 4C36424F4976EF3266512DACC6E7D40A
\ No newline at end of file
diff --git a/Engine/Source/Runtime/Engine/Classes/Engine/EngineTypes.h b/Engine/Source/Runtime/Engine/Classes/Engine/EngineTypes.h
--- a/Engine/Source/Runtime/Engine/Classes/Engine/EngineTypes.h
+++ b/Engine/Source/Runtime/Engine/Classes/Engine/EngineTypes.h
@@ -88,6 +88,17 @@ enum EBlendMode
 	BLEND_MAX,
 };
 
+UENUM()
+enum ESamplerSourceMode
+{
+	/** Get the sampler from the texture.  Every unique texture will consume a sampler slot, which are limited in number. */
+	SSM_FromTextureAsset UMETA(DisplayName="From texture asset"),
+	/** Shared sampler source that does not consume a sampler slot.  Uses wrap addressing and gets filter mode from the world texture group. */
+	SSM_Wrap_WorldGroupSettings UMETA(DisplayName="Shared: Wrap"),
+	/** Shared sampler source that does not consume a sampler slot.  Uses clamp addressing and gets filter mode from the world texture group. */
+	SSM_Clamp_WorldGroupSettings UMETA(DisplayName="Shared: Clamp")
+};
+
 UENUM()
 enum ETranslucencyLightingMode
 {
diff --git a/Engine/Source/Runtime/Engine/Classes/Materials/MaterialExpressionTextureSample.h b/Engine/Source/Runtime/Engine/Classes/Materials/MaterialExpressionTextureSample.h
--- a/Engine/Source/Runtime/Engine/Classes/Materials/MaterialExpressionTextureSample.h
+++ b/Engine/Source/Runtime/Engine/Classes/Materials/MaterialExpressionTextureSample.h
@@ -39,6 +39,15 @@ class ENGINE_API UMaterialExpressionTextureSample : public UMaterialExpressionTe
 	UPROPERTY(EditAnywhere, Category=MaterialExpressionTextureSample, meta=(DisplayName = "MipValueMode"))
 	TEnumAsByte<enum ETextureMipValueMode> MipValueMode;
 
+	/** 
+	 * Controls where the sampler for this texture lookup will come from.  
+	 * Choose 'from texture asset' to make use of the UTexture addressing settings,
+	 * Otherwise use one of the global samplers, which will not consume a sampler slot.
+	 * This allows materials to use more than 16 unique textures on SM5 platforms.
+	 */
+	UPROPERTY(EditAnywhere, Category=MaterialExpressionTextureSample)
+	TEnumAsByte<enum ESamplerSourceMode> SamplerSource;
+
 	/** only used if Coordinates is not hooked up */
 	UPROPERTY(EditAnywhere, Category = MaterialExpressionTextureSample)
 	uint32 ConstCoordinate;
diff --git a/Engine/Source/Runtime/Engine/Private/Materials/HLSLMaterialTranslator.h b/Engine/Source/Runtime/Engine/Private/Materials/HLSLMaterialTranslator.h
--- a/Engine/Source/Runtime/Engine/Private/Materials/HLSLMaterialTranslator.h
+++ b/Engine/Source/Runtime/Engine/Private/Materials/HLSLMaterialTranslator.h
@@ -2093,7 +2093,14 @@ protected:
 				);
 		}
 
-	virtual int32 TextureSample(int32 TextureIndex,int32 CoordinateIndex,EMaterialSamplerType SamplerType,int32 MipValueIndex=INDEX_NONE,ETextureMipValueMode MipValueMode=TMVM_None) override
+	virtual int32 TextureSample(
+		int32 TextureIndex,
+		int32 CoordinateIndex,
+		EMaterialSamplerType SamplerType,
+		int32 MipValueIndex=INDEX_NONE,
+		ETextureMipValueMode MipValueMode=TMVM_None,
+		ESamplerSourceMode SamplerSource=SSM_FromTextureAsset
+		) override
 	{
 		if(TextureIndex == INDEX_NONE || CoordinateIndex == INDEX_NONE)
 		{
@@ -2131,6 +2138,23 @@ protected:
 			MipValueMode = TMVM_MipLevel;
 		}
 
+		FString SamplerStateCode;
+
+		if (SamplerSource == SSM_FromTextureAsset)
+		{
+			SamplerStateCode = TEXT("%sSampler");
+		}
+		else if (SamplerSource == SSM_Wrap_WorldGroupSettings)
+		{
+			// Use the shared sampler to save sampler slots
+			SamplerStateCode = TEXT("GetMaterialSharedSampler(%sSampler,Material.Wrap_WorldGroupSettings)");
+		}
+		else if (SamplerSource == SSM_Clamp_WorldGroupSettings)
+		{
+			// Use the shared sampler to save sampler slots
+			SamplerStateCode = TEXT("GetMaterialSharedSampler(%sSampler,Material.Clamp_WorldGroupSettings)");
+		}
+
 		FString SampleCode = 
 			(TextureType == MCT_TextureCube)
 			? TEXT("TextureCubeSample")
@@ -2138,7 +2162,7 @@ protected:
 	
 		if(MipValueMode == TMVM_None)
 		{
-			SampleCode += TEXT("(%s,%sSampler,%s)");
+			SampleCode += FString(TEXT("(%s,")) + SamplerStateCode + TEXT(",%s)");
 		}
 		else if(MipValueMode == TMVM_MipLevel)
 		{
@@ -2437,7 +2461,7 @@ protected:
 			);
 	}
 
-	virtual int32 Texture(UTexture* InTexture) override
+	virtual int32 Texture(UTexture* InTexture,ESamplerSourceMode SamplerSource=SSM_FromTextureAsset) override
 	{
 		if (ShaderFrequency != SF_Pixel
 			&& ErrorUnlessFeatureLevelSupported(ERHIFeatureLevel::SM4) == INDEX_NONE)
@@ -2448,10 +2472,10 @@ protected:
 		EMaterialValueType ShaderType = InTexture->GetMaterialType();
 		const int32 TextureReferenceIndex = Material->GetReferencedTextures().Find(InTexture);
 		checkf(TextureReferenceIndex != INDEX_NONE, TEXT("Material expression called Compiler->Texture() without implementing UMaterialExpression::GetReferencedTexture properly"));
-		return AddUniformExpression(new FMaterialUniformExpressionTexture(TextureReferenceIndex),ShaderType,TEXT(""));
+		return AddUniformExpression(new FMaterialUniformExpressionTexture(TextureReferenceIndex, SamplerSource),ShaderType,TEXT(""));
 	}
 
-	virtual int32 TextureParameter(FName ParameterName,UTexture* DefaultValue) override
+	virtual int32 TextureParameter(FName ParameterName,UTexture* DefaultValue,ESamplerSourceMode SamplerSource=SSM_FromTextureAsset) override
 	{
 		if (ShaderFrequency != SF_Pixel
 			&& ErrorUnlessFeatureLevelSupported(ERHIFeatureLevel::SM4) == INDEX_NONE)
@@ -2462,7 +2486,7 @@ protected:
 		EMaterialValueType ShaderType = DefaultValue->GetMaterialType();
 		const int32 TextureReferenceIndex = Material->GetReferencedTextures().Find(DefaultValue);
 		checkf(TextureReferenceIndex != INDEX_NONE, TEXT("Material expression called Compiler->TextureParameter() without implementing UMaterialExpression::GetReferencedTexture properly"));
-		return AddUniformExpression(new FMaterialUniformExpressionTextureParameter(ParameterName,TextureReferenceIndex),ShaderType,TEXT(""));
+		return AddUniformExpression(new FMaterialUniformExpressionTextureParameter(ParameterName, TextureReferenceIndex, SamplerSource),ShaderType,TEXT(""));
 	}
 
 	virtual int32 StaticBool(bool bValue) override
diff --git a/Engine/Source/Runtime/Engine/Private/Materials/MaterialExpressions.cpp b/Engine/Source/Runtime/Engine/Private/Materials/MaterialExpressions.cpp
--- a/Engine/Source/Runtime/Engine/Private/Materials/MaterialExpressions.cpp
+++ b/Engine/Source/Runtime/Engine/Private/Materials/MaterialExpressions.cpp
@@ -1287,7 +1287,7 @@ int32 UMaterialExpressionTextureSample::Compile(class FMaterialCompiler* Compile
 {
 	if (Texture || TextureObject.Expression)
 	{
-		int32 TextureCodeIndex = TextureObject.Expression ? TextureObject.Compile(Compiler) : Compiler->Texture(Texture);
+		int32 TextureCodeIndex = TextureObject.Expression ? TextureObject.Compile(Compiler) : Compiler->Texture(Texture, SamplerSource);
 
 		UTexture* EffectiveTexture = Texture;
 		EMaterialSamplerType EffectiveSamplerType = (EMaterialSamplerType)SamplerType;
@@ -1328,7 +1328,8 @@ int32 UMaterialExpressionTextureSample::Compile(class FMaterialCompiler* Compile
 				Coordinates.Expression ? Coordinates.Compile(Compiler) : Compiler->TextureCoordinate(ConstCoordinate, false, false),
 				(EMaterialSamplerType)EffectiveSamplerType,
 				MipValue.Expression ? MipValue.Compile(Compiler) : Compiler->Constant(ConstMipValue),
-				MipValueMode
+				MipValueMode,
+				SamplerSource
 				);
 		}
 		else
@@ -1465,11 +1466,12 @@ int32 UMaterialExpressionTextureSampleParameter::Compile(class FMaterialCompiler
 	}
 
 	return Compiler->TextureSample(
-					Compiler->TextureParameter(ParameterName, Texture),
+					Compiler->TextureParameter(ParameterName, Texture, SamplerSource),
 					Coordinates.Expression ? Coordinates.Compile(Compiler) : Compiler->TextureCoordinate(ConstCoordinate, false, false),
 					(EMaterialSamplerType)SamplerType,
 					MipValue.Expression ? MipValue.Compile(Compiler) : Compiler->Constant(ConstMipValue),
-					MipValueMode
+					MipValueMode,
+					SamplerSource
 					);
 }
 
diff --git a/Engine/Source/Runtime/Engine/Private/Materials/MaterialShared.cpp b/Engine/Source/Runtime/Engine/Private/Materials/MaterialShared.cpp
--- a/Engine/Source/Runtime/Engine/Private/Materials/MaterialShared.cpp
+++ b/Engine/Source/Runtime/Engine/Private/Materials/MaterialShared.cpp
@@ -1563,24 +1563,6 @@ void FMaterialRenderProxy::EvaluateUniformExpressions(FUniformExpressionCache& O
 	// Create and cache the material's uniform buffer.
 	OutUniformExpressionCache.UniformBuffer = UniformExpressionSet.CreateUniformBuffer(Context, CommandListIfLocalMode, &OutUniformExpressionCache.LocalUniformBuffer);
 
-	// Cache 2D texture uniform expressions.
-	OutUniformExpressionCache.Textures.Empty(UniformExpressionSet.Uniform2DTextureExpressions.Num());
-	for(int32 ExpressionIndex = 0;ExpressionIndex < UniformExpressionSet.Uniform2DTextureExpressions.Num();ExpressionIndex++)
-	{
-		const UTexture* Value;
-		UniformExpressionSet.Uniform2DTextureExpressions[ExpressionIndex]->GetTextureValue(Context,Context.Material,Value);
-		OutUniformExpressionCache.Textures.Add(Value);
-	}
-
-	// Cache cube texture uniform expressions.
-	OutUniformExpressionCache.CubeTextures.Empty(UniformExpressionSet.UniformCubeTextureExpressions.Num());
-	for(int32 ExpressionIndex = 0;ExpressionIndex < UniformExpressionSet.UniformCubeTextureExpressions.Num();ExpressionIndex++)
-	{
-		const UTexture* Value;
-		UniformExpressionSet.UniformCubeTextureExpressions[ExpressionIndex]->GetTextureValue(Context,Context.Material,Value);
-		OutUniformExpressionCache.CubeTextures.Add(Value);
-	}
-
 	OutUniformExpressionCache.ParameterCollections = UniformExpressionSet.ParameterCollections;
 
 	OutUniformExpressionCache.bUpToDate = true;
diff --git a/Engine/Source/Runtime/Engine/Private/Materials/MaterialUniformExpressions.cpp b/Engine/Source/Runtime/Engine/Private/Materials/MaterialUniformExpressions.cpp
--- a/Engine/Source/Runtime/Engine/Private/Materials/MaterialUniformExpressions.cpp
+++ b/Engine/Source/Runtime/Engine/Private/Materials/MaterialUniformExpressions.cpp
@@ -240,6 +240,12 @@ void FUniformExpressionSet::CreateBufferStruct()
 		NextMemberOffset += 8;
 	}
 
+	new(Members) FUniformBufferStruct::FMember(TEXT("Wrap_WorldGroupSettings"),TEXT("SamplerState"),NextMemberOffset,UBMT_SAMPLER,EShaderPrecisionModifier::Float,1,1,1,NULL);
+	NextMemberOffset += 8;
+
+	new(Members) FUniformBufferStruct::FMember(TEXT("Clamp_WorldGroupSettings"),TEXT("SamplerState"),NextMemberOffset,UBMT_SAMPLER,EShaderPrecisionModifier::Float,1,1,1,NULL);
+	NextMemberOffset += 8;
+
 	const uint32 StructSize = Align(NextMemberOffset,UNIFORM_BUFFER_STRUCT_ALIGNMENT);
 	UniformBufferStruct = new FUniformBufferStruct(
 		TEXT("MaterialUniforms"),
@@ -291,11 +297,23 @@ FUniformBufferRHIRef FUniformExpressionSet::CreateUniformBuffer(const FMaterialR
 		for(int32 ExpressionIndex = 0;ExpressionIndex < Uniform2DTextureExpressions.Num();ExpressionIndex++)
 		{
 			const UTexture* Value;
-			Uniform2DTextureExpressions[ExpressionIndex]->GetTextureValue(MaterialRenderContext,MaterialRenderContext.Material,Value);
+			ESamplerSourceMode SourceMode;
+			Uniform2DTextureExpressions[ExpressionIndex]->GetTextureValue(MaterialRenderContext,MaterialRenderContext.Material,Value,SourceMode);
 			if(Value && Value->Resource)
 			{
 				*ResourceTable++ = Value->TextureReference.TextureReferenceRHI;
-				*ResourceTable++ = Value->Resource->SamplerStateRHI;
+				FSamplerStateRHIRef* SamplerSource = &Value->Resource->SamplerStateRHI;
+
+				if (SourceMode == SSM_Wrap_WorldGroupSettings)
+				{
+					SamplerSource = &Wrap_WorldGroupSettings->SamplerStateRHI;
+				}
+				else if (SourceMode == SSM_Clamp_WorldGroupSettings)
+				{
+					SamplerSource = &Clamp_WorldGroupSettings->SamplerStateRHI;
+				}
+
+				*ResourceTable++ = *SamplerSource;
 			}
 			else
 			{
@@ -308,11 +326,23 @@ FUniformBufferRHIRef FUniformExpressionSet::CreateUniformBuffer(const FMaterialR
 		for(int32 ExpressionIndex = 0;ExpressionIndex < UniformCubeTextureExpressions.Num();ExpressionIndex++)
 		{
 			const UTexture* Value;
-			UniformCubeTextureExpressions[ExpressionIndex]->GetTextureValue(MaterialRenderContext,MaterialRenderContext.Material,Value);
+			ESamplerSourceMode SourceMode;
+			UniformCubeTextureExpressions[ExpressionIndex]->GetTextureValue(MaterialRenderContext,MaterialRenderContext.Material,Value,SourceMode);
 			if(Value && Value->Resource)
 			{
 				*ResourceTable++ = Value->TextureReference.TextureReferenceRHI;
-				*ResourceTable++ = Value->Resource->SamplerStateRHI;
+				FSamplerStateRHIRef* SamplerSource = &Value->Resource->SamplerStateRHI;
+
+				if (SourceMode == SSM_Wrap_WorldGroupSettings)
+				{
+					SamplerSource = &Wrap_WorldGroupSettings->SamplerStateRHI;
+				}
+				else if (SourceMode == SSM_Clamp_WorldGroupSettings)
+				{
+					SamplerSource = &Clamp_WorldGroupSettings->SamplerStateRHI;
+				}
+
+				*ResourceTable++ = *SamplerSource;
 			}
 			else
 			{
@@ -320,6 +350,10 @@ FUniformBufferRHIRef FUniformExpressionSet::CreateUniformBuffer(const FMaterialR
 				*ResourceTable++ = GWhiteTexture->SamplerStateRHI;
 			}
 		}
+
+		*ResourceTable++ = Wrap_WorldGroupSettings->SamplerStateRHI;
+		*ResourceTable++ = Clamp_WorldGroupSettings->SamplerStateRHI;
+
 		if (CommandListIfLocalMode)
 		{
 			check(OutLocalUniformBuffer);
@@ -338,12 +372,14 @@ FUniformBufferRHIRef FUniformExpressionSet::CreateUniformBuffer(const FMaterialR
 
 FMaterialUniformExpressionTexture::FMaterialUniformExpressionTexture() :
 	TextureIndex(INDEX_NONE),
+	SamplerSource(SSM_FromTextureAsset),
 	TransientOverrideValue_GameThread(NULL),
 	TransientOverrideValue_RenderThread(NULL)
 {}
 
-FMaterialUniformExpressionTexture::FMaterialUniformExpressionTexture(int32 InTextureIndex) :
+FMaterialUniformExpressionTexture::FMaterialUniformExpressionTexture(int32 InTextureIndex, ESamplerSourceMode InSamplerSource) :
 	TextureIndex(InTextureIndex),
+	SamplerSource(InSamplerSource),
 	TransientOverrideValue_GameThread(NULL),
 	TransientOverrideValue_RenderThread(NULL)
 {
@@ -351,7 +387,9 @@ FMaterialUniformExpressionTexture::FMaterialUniformExpressionTexture(int32 InTex
 
 void FMaterialUniformExpressionTexture::Serialize(FArchive& Ar)
 {
-	Ar << TextureIndex;
+	int32 SamplerSourceInt = (int32)SamplerSource;
+	Ar << TextureIndex << SamplerSourceInt;
+	SamplerSource = (ESamplerSourceMode)SamplerSourceInt;
 }
 
 void FMaterialUniformExpressionTexture::SetTransientOverrideTextureValue( UTexture* InOverrideTexture )
@@ -365,9 +403,10 @@ void FMaterialUniformExpressionTexture::SetTransientOverrideTextureValue( UTextu
 											   });
 }
 
-void FMaterialUniformExpressionTexture::GetTextureValue(const FMaterialRenderContext& Context,const FMaterial& Material,const UTexture*& OutValue) const
+void FMaterialUniformExpressionTexture::GetTextureValue(const FMaterialRenderContext& Context,const FMaterial& Material,const UTexture*& OutValue,ESamplerSourceMode& OutSamplerSource) const
 {
 	check(IsInParallelRenderingThread());
+	OutSamplerSource = SamplerSource;
 	if( TransientOverrideValue_RenderThread != NULL )
 	{
 		OutValue = TransientOverrideValue_RenderThread;
diff --git a/Engine/Source/Runtime/Engine/Private/Materials/MaterialUniformExpressions.h b/Engine/Source/Runtime/Engine/Private/Materials/MaterialUniformExpressions.h
--- a/Engine/Source/Runtime/Engine/Private/Materials/MaterialUniformExpressions.h
+++ b/Engine/Source/Runtime/Engine/Private/Materials/MaterialUniformExpressions.h
@@ -265,8 +265,8 @@ public:
 
 	FMaterialUniformExpressionTextureParameter() {}
 
-	FMaterialUniformExpressionTextureParameter(FName InParameterName, int32 InTextureIndex) :
-		Super(InTextureIndex),
+	FMaterialUniformExpressionTextureParameter(FName InParameterName, int32 InTextureIndex, ESamplerSourceMode InSourceMode) :
+		Super(InTextureIndex, InSourceMode),
 		ParameterName(InParameterName)
 	{}
 
@@ -276,9 +276,10 @@ public:
 		Ar << ParameterName;
 		Super::Serialize(Ar);
 	}
-	virtual void GetTextureValue(const FMaterialRenderContext& Context,const FMaterial& Material,const UTexture*& OutValue) const
+	virtual void GetTextureValue(const FMaterialRenderContext& Context,const FMaterial& Material,const UTexture*& OutValue,ESamplerSourceMode& OutSamplerSource) const
 	{
 		check(IsInParallelRenderingThread());
+		OutSamplerSource = SamplerSource;
 		if( TransientOverrideValue_RenderThread != NULL )
 		{
 			OutValue = TransientOverrideValue_RenderThread;
diff --git a/Engine/Source/Runtime/Engine/Private/SceneManagement.cpp b/Engine/Source/Runtime/Engine/Private/SceneManagement.cpp
--- a/Engine/Source/Runtime/Engine/Private/SceneManagement.cpp
+++ b/Engine/Source/Runtime/Engine/Private/SceneManagement.cpp
@@ -292,3 +292,29 @@ FViewUniformShaderParameters::FViewUniformShaderParameters()
 	, DirectionalLightShadowSampler(TStaticSamplerState<SF_Point, AM_Clamp, AM_Clamp, AM_Clamp>::GetRHI())
 {
 }
+
+void FSharedSamplerState::InitRHI()
+{
+	const float MipMapBias = UTexture2D::GetGlobalMipMapLODBias();
+
+	FSamplerStateInitializerRHI SamplerStateInitializer
+	(
+		GSystemSettings.TextureLODSettings.GetSamplerFilter(TEXTUREGROUP_World),
+		bWrap ? AM_Wrap : AM_Clamp,
+		bWrap ? AM_Wrap : AM_Clamp,
+		bWrap ? AM_Wrap : AM_Clamp,
+		MipMapBias
+	);
+	SamplerStateRHI = RHICreateSamplerState(SamplerStateInitializer);
+}
+
+FSharedSamplerState* Wrap_WorldGroupSettings = NULL;
+FSharedSamplerState* Clamp_WorldGroupSettings = NULL;
+
+void InitializeSharedSamplerStates()
+{
+	Wrap_WorldGroupSettings = new FSharedSamplerState(true);
+	Clamp_WorldGroupSettings = new FSharedSamplerState(false);
+	BeginInitResource(Wrap_WorldGroupSettings);
+	BeginInitResource(Clamp_WorldGroupSettings);
+}
\ No newline at end of file
diff --git a/Engine/Source/Runtime/Engine/Private/TextureLODSettings.cpp b/Engine/Source/Runtime/Engine/Private/TextureLODSettings.cpp
--- a/Engine/Source/Runtime/Engine/Private/TextureLODSettings.cpp
+++ b/Engine/Source/Runtime/Engine/Private/TextureLODSettings.cpp
@@ -318,4 +318,9 @@ ESamplerFilter FTextureLODSettings::GetSamplerFilter( const UTexture* Texture )
 	}
 
 	return Filter;
+}
+
+ESamplerFilter FTextureLODSettings::GetSamplerFilter( int32 InLODGroup) const
+{
+	return TextureLODGroups[InLODGroup].Filter;
 }
\ No newline at end of file
diff --git a/Engine/Source/Runtime/Engine/Public/MaterialCompiler.h b/Engine/Source/Runtime/Engine/Public/MaterialCompiler.h
--- a/Engine/Source/Runtime/Engine/Public/MaterialCompiler.h
+++ b/Engine/Source/Runtime/Engine/Public/MaterialCompiler.h
@@ -100,10 +100,10 @@ public:
 	virtual int32 If(int32 A,int32 B,int32 AGreaterThanB,int32 AEqualsB,int32 ALessThanB,int32 Threshold) = 0;
 
 	virtual int32 TextureCoordinate(uint32 CoordinateIndex, bool UnMirrorU, bool UnMirrorV) = 0;
-	virtual int32 TextureSample(int32 Texture,int32 Coordinate,enum EMaterialSamplerType SamplerType,int32 MipValueIndex=INDEX_NONE,ETextureMipValueMode MipValueMode=TMVM_None) = 0;
+	virtual int32 TextureSample(int32 Texture,int32 Coordinate,enum EMaterialSamplerType SamplerType,int32 MipValueIndex=INDEX_NONE,ETextureMipValueMode MipValueMode=TMVM_None,ESamplerSourceMode SamplerSource=SSM_FromTextureAsset) = 0;
 
-	virtual int32 Texture(UTexture* Texture) = 0;
-	virtual int32 TextureParameter(FName ParameterName,UTexture* DefaultTexture) = 0;
+	virtual int32 Texture(UTexture* Texture,ESamplerSourceMode SamplerSource=SSM_FromTextureAsset) = 0;
+	virtual int32 TextureParameter(FName ParameterName,UTexture* DefaultTexture,ESamplerSourceMode SamplerSource=SSM_FromTextureAsset) = 0;
 
 	virtual	int32 PixelDepth()=0;
 	virtual int32 SceneDepth(int32 Offset, int32 UV, bool bUseOffset) = 0;
@@ -253,11 +253,13 @@ public:
 
 	virtual int32 If(int32 A,int32 B,int32 AGreaterThanB,int32 AEqualsB,int32 ALessThanB,int32 Threshold) override { return Compiler->If(A,B,AGreaterThanB,AEqualsB,ALessThanB,Threshold); }
 
-	virtual int32 TextureSample(int32 InTexture,int32 Coordinate,EMaterialSamplerType SamplerType,int32 MipValueIndex=INDEX_NONE,ETextureMipValueMode MipValueMode=TMVM_None) override { return Compiler->TextureSample(InTexture,Coordinate,SamplerType,MipValueIndex,MipValueMode); }
+	virtual int32 TextureSample(int32 InTexture,int32 Coordinate,EMaterialSamplerType SamplerType,int32 MipValueIndex=INDEX_NONE,ETextureMipValueMode MipValueMode=TMVM_None,ESamplerSourceMode SamplerSource=SSM_FromTextureAsset) override 
+		{ return Compiler->TextureSample(InTexture,Coordinate,SamplerType,MipValueIndex,MipValueMode,SamplerSource); }
+
 	virtual int32 TextureCoordinate(uint32 CoordinateIndex, bool UnMirrorU, bool UnMirrorV) override { return Compiler->TextureCoordinate(CoordinateIndex, UnMirrorU, UnMirrorV); }
 
-	virtual int32 Texture(UTexture* InTexture) override { return Compiler->Texture(InTexture); }
-	virtual int32 TextureParameter(FName ParameterName,UTexture* DefaultValue) override { return Compiler->TextureParameter(ParameterName,DefaultValue); }
+	virtual int32 Texture(UTexture* InTexture,ESamplerSourceMode SamplerSource=SSM_FromTextureAsset) override { return Compiler->Texture(InTexture,SamplerSource); }
+	virtual int32 TextureParameter(FName ParameterName,UTexture* DefaultValue,ESamplerSourceMode SamplerSource=SSM_FromTextureAsset) override { return Compiler->TextureParameter(ParameterName,DefaultValue,SamplerSource); }
 
 	virtual	int32 PixelDepth() override { return Compiler->PixelDepth();	}
 	virtual int32 SceneDepth(int32 Offset, int32 UV, bool bUseOffset) override { return Compiler->SceneDepth(Offset, UV, bUseOffset); }
diff --git a/Engine/Source/Runtime/Engine/Public/MaterialShared.h b/Engine/Source/Runtime/Engine/Public/MaterialShared.h
--- a/Engine/Source/Runtime/Engine/Public/MaterialShared.h
+++ b/Engine/Source/Runtime/Engine/Public/MaterialShared.h
@@ -176,11 +176,11 @@ public:
 
 	FMaterialUniformExpressionTexture();
 
-	FMaterialUniformExpressionTexture(int32 InTextureIndex);
+	FMaterialUniformExpressionTexture(int32 InTextureIndex, ESamplerSourceMode InSamplerSource);
 
 	// FMaterialUniformExpression interface.
 	virtual void Serialize(FArchive& Ar);
-	virtual void GetTextureValue(const FMaterialRenderContext& Context,const FMaterial& Material,const UTexture*& OutValue) const;
+	virtual void GetTextureValue(const FMaterialRenderContext& Context,const FMaterial& Material,const UTexture*& OutValue,ESamplerSourceMode& OutSamplerSource) const;
 	/** Accesses the texture used for rendering this uniform expression. */
 	virtual void GetGameThreadTextureValue(const class UMaterialInterface* MaterialInterface,const FMaterial& Material,UTexture*& OutValue,bool bAllowOverride=true) const;
 	virtual class FMaterialUniformExpressionTexture* GetTextureUniformExpression() { return this; }
@@ -197,6 +197,7 @@ public:
 protected:
 	/** Index into FMaterial::GetReferencedTextures */
 	int32 TextureIndex;
+	ESamplerSourceMode SamplerSource;
 	/** Texture that may be used in the editor for overriding the texture but never saved to disk, accessible only by the game thread! */
 	UTexture* TransientOverrideValue_GameThread;
 	/** Texture that may be used in the editor for overriding the texture but never saved to disk, accessible only by the rendering thread! */
@@ -1190,10 +1191,6 @@ struct FUniformExpressionCache
 	FUniformBufferRHIRef UniformBuffer;
 	/** Material uniform buffer. */
 	FLocalUniformBuffer LocalUniformBuffer;
-	/** Textures referenced by uniform expressions. */
-	TArray<const UTexture*> Textures;
-	/** Cube textures referenced by uniform expressions. */
-	TArray<const UTexture*> CubeTextures;
 	/** Ids of parameter collections needed for rendering. */
 	TArray<FGuid> ParameterCollections;
 	/** True if the cache is up to date. */
diff --git a/Engine/Source/Runtime/Engine/Public/SceneManagement.h b/Engine/Source/Runtime/Engine/Public/SceneManagement.h
--- a/Engine/Source/Runtime/Engine/Public/SceneManagement.h
+++ b/Engine/Source/Runtime/Engine/Public/SceneManagement.h
@@ -1831,3 +1831,30 @@ int8 ENGINE_API ComputeStaticMeshLOD(const FStaticMeshRenderData* RenderData, co
  * @param SphereRadius - Radius of the sphere to use to calculate screen coverage
  */
 int8 ENGINE_API ComputeLODForMeshes(const TIndirectArray<class FStaticMesh>& StaticMeshes, const FSceneView& View, const FVector4& Origin, float SphereRadius, int32 ForcedLODLevel, float ScreenSizeScale = 1.0f);
+
+class FSharedSamplerState : public FRenderResource
+{
+public:
+	FSamplerStateRHIRef SamplerStateRHI;
+	bool bWrap;
+
+	FSharedSamplerState(bool bInWrap) :
+		bWrap(bInWrap)
+	{}
+
+	virtual void InitRHI() override;
+
+	virtual void ReleaseRHI() override
+	{
+		SamplerStateRHI.SafeRelease();
+	}
+};
+
+/** Sampler state using Wrap addressing and taking filter mode from the world texture group. */
+extern ENGINE_API FSharedSamplerState* Wrap_WorldGroupSettings;
+
+/** Sampler state using Clamp addressing and taking filter mode from the world texture group. */
+extern ENGINE_API FSharedSamplerState* Clamp_WorldGroupSettings;
+
+/** Initializes the shared sampler states. */
+extern ENGINE_API void InitializeSharedSamplerStates();
\ No newline at end of file
diff --git a/Engine/Source/Runtime/Engine/Public/TextureLODSettings.h b/Engine/Source/Runtime/Engine/Public/TextureLODSettings.h
--- a/Engine/Source/Runtime/Engine/Public/TextureLODSettings.h
+++ b/Engine/Source/Runtime/Engine/Public/TextureLODSettings.h
@@ -83,6 +83,8 @@ struct FTextureLODSettings
 	 */
 	ENGINE_API ESamplerFilter GetSamplerFilter( const UTexture* Texture ) const;
 
+	ENGINE_API ESamplerFilter GetSamplerFilter( int32 InLODGroup) const;
+
 	/**
 	 * Returns the LODGroup mip gen settings
 	 *
diff --git a/Engine/Source/Runtime/Launch/Private/LaunchEngineLoop.cpp b/Engine/Source/Runtime/Launch/Private/LaunchEngineLoop.cpp
--- a/Engine/Source/Runtime/Launch/Private/LaunchEngineLoop.cpp
+++ b/Engine/Source/Runtime/Launch/Private/LaunchEngineLoop.cpp
@@ -1203,6 +1203,8 @@ int32 FEngineLoop::PreInit( const TCHAR* CmdLine )
 	// Initialize the RHI.
 	RHIInit(bHasEditorToken);
 
+	InitializeSharedSamplerStates();
+
 	if (!FPlatformProperties::RequiresCookedData())
 	{
 		check(!GShaderCompilingManager);
