#ue4 - Fix UpdateOverlaps() call in DetachFromParent() causing incorrect notification when called from AttachTo() and the object does not maintain world position.

Jira: UE-3246

diff --git a/Engine/Source/Runtime/Engine/Classes/Components/SceneComponent.h b/Engine/Source/Runtime/Engine/Classes/Components/SceneComponent.h
--- a/Engine/Source/Runtime/Engine/Classes/Components/SceneComponent.h
+++ b/Engine/Source/Runtime/Engine/Classes/Components/SceneComponent.h
@@ -142,6 +142,9 @@ protected:
 	UPROPERTY(Transient)
 	uint32 bWorldToComponentUpdated:1;
 
+	// Transient flag that temporarily disables UpdateOverlaps within DetachFromParent().
+	uint32 bDisableDetachmentUpdateOverlaps:1;
+
 public:
 
 	/** How often this component is allowed to move, used to make various optimizations. Only safe to set in constructor, use SetMobility() during runtime. */
diff --git a/Engine/Source/Runtime/Engine/Private/SceneComponent.cpp b/Engine/Source/Runtime/Engine/Private/SceneComponent.cpp
--- a/Engine/Source/Runtime/Engine/Private/SceneComponent.cpp
+++ b/Engine/Source/Runtime/Engine/Private/SceneComponent.cpp
@@ -766,9 +766,17 @@ void USceneComponent::AttachTo(class USceneComponent* Parent, FName InSocketName
 			return;
 		}
 
+		// Don't call UpdateOverlaps() when detaching, since we are going to do it anyway after we reattach below.
+		// Aside from a perf benefit this also maintains correct behavior when we don't have KeepWorldPosition set.
+		const bool bSavedDisableDetachmentUpdateOverlaps = bDisableDetachmentUpdateOverlaps;
+		bDisableDetachmentUpdateOverlaps = true;
+
 		// Make sure we are detached
-		bool bMaintainWorldPosition = AttachType == EAttachLocation::KeepWorldPosition;
+		const bool bMaintainWorldPosition = (AttachType == EAttachLocation::KeepWorldPosition);
 		DetachFromParent(bMaintainWorldPosition);
+		
+		// Restore detachment update overlaps flag.
+		bDisableDetachmentUpdateOverlaps = bSavedDisableDetachmentUpdateOverlaps;
 
 		{
 			//This code requires some explaining. Inside the editor we allow user to attach physically simulated objects to other objects. This is done for convenience so that users can group things together in hierarchy.
@@ -952,7 +960,7 @@ void USceneComponent::DetachFromParent(bool bMaintainWorldPosition)
 		UpdateComponentToWorld();
 
 		// Update overlaps, in case location changed or overlap state depends on attachment.
-		if (IsRegistered())
+		if (IsRegistered() && !bDisableDetachmentUpdateOverlaps)
 		{
 			UpdateOverlaps();
 		}
