- Removed unneeded check() in EndRenderFrame() on ES2.

diff --git a/Engine/Source/Runtime/Engine/Private/Slate/SceneViewport.cpp b/Engine/Source/Runtime/Engine/Private/Slate/SceneViewport.cpp
--- a/Engine/Source/Runtime/Engine/Private/Slate/SceneViewport.cpp
+++ b/Engine/Source/Runtime/Engine/Private/Slate/SceneViewport.cpp
@@ -1131,11 +1131,6 @@ void FSceneViewport::EndRenderFrame(FRHICommandListImmediate& RHICmdList, bool b
 	check( IsInRenderingThread() );
 	if (bUseSeparateRenderTarget)
 	{
-		// @todo-mobile
-		if (GMaxRHIShaderPlatform == SP_OPENGL_ES2)
-		{
-			check(0);
-		}
 		if (SlateRenderTargetHandle)
 		{
 			RHICmdList.CopyToResolveTarget( RenderTargetTextureRHI, SlateRenderTargetHandle->GetRHIRef(), false, FResolveParams() );
