#ue4 - Fix mislabeled advanced display params on SetWorldLocation() that were supposed to be on SetWorldRotation().

diff --git a/Engine/Source/Runtime/Engine/Classes/Components/SceneComponent.h b/Engine/Source/Runtime/Engine/Classes/Components/SceneComponent.h
--- a/Engine/Source/Runtime/Engine/Classes/Components/SceneComponent.h
+++ b/Engine/Source/Runtime/Engine/Classes/Components/SceneComponent.h
@@ -281,13 +281,13 @@ public:
 
 
 	/** Set the relative location of this component to put it at the supplied location in world space. */
-	UFUNCTION(BlueprintCallable, Category="Utilities|Transformation", meta=(FriendlyName="SetWorldLocation", AdvancedDisplay="bSweep,SweepHitResult"))
+	UFUNCTION(BlueprintCallable, Category="Utilities|Transformation", meta=(FriendlyName="SetWorldLocation"))
 	void K2_SetWorldLocation(FVector NewLocation, bool bSweep, FHitResult& SweepHitResult);
 	void SetWorldLocation(FVector NewLocation, bool bSweep=false, FHitResult* OutSweepHitResult=nullptr);
 
 
 	/** Set the relative rotation of this component to put it at the supplied orientation in world space. */
-	UFUNCTION(BlueprintCallable, Category="Utilities|Transformation", meta=(FriendlyName="SetWorldRotation"))
+	UFUNCTION(BlueprintCallable, Category="Utilities|Transformation", meta=(FriendlyName="SetWorldRotation", AdvancedDisplay="bSweep,SweepHitResult"))
 	void K2_SetWorldRotation(FRotator NewRotation, bool bSweep, FHitResult& SweepHitResult);
 	void SetWorldRotation(FRotator NewRotation, bool bSweep=false, FHitResult* OutSweepHitResult=nullptr);
 
