#ue4 - Add static FName for "DefaultProjectile" collision profile.

diff --git a/Engine/Source/Runtime/Engine/Classes/Engine/CollisionProfile.h b/Engine/Source/Runtime/Engine/Classes/Engine/CollisionProfile.h
--- a/Engine/Source/Runtime/Engine/Classes/Engine/CollisionProfile.h
+++ b/Engine/Source/Runtime/Engine/Classes/Engine/CollisionProfile.h
@@ -155,6 +155,7 @@ public:
 	static ENGINE_API const FName BlockAllDynamic_ProfileName;
 	static ENGINE_API const FName Pawn_ProfileName;
 	static ENGINE_API const FName Vehicle_ProfileName;
+	static ENGINE_API const FName DefaultProjectile_ProfileName;
 
 	/** Accessor and initializer **/
 	ENGINE_API static UCollisionProfile* Get();
diff --git a/Engine/Source/Runtime/Engine/Private/Collision/CollisionProfile.cpp b/Engine/Source/Runtime/Engine/Private/Collision/CollisionProfile.cpp
--- a/Engine/Source/Runtime/Engine/Private/Collision/CollisionProfile.cpp
+++ b/Engine/Source/Runtime/Engine/Private/Collision/CollisionProfile.cpp
@@ -37,6 +37,8 @@ ENGINE_API const FName UCollisionProfile::PhysicsActor_ProfileName = FName(TEXT(
 ENGINE_API const FName UCollisionProfile::BlockAllDynamic_ProfileName = FName(TEXT("BlockAllDynamic"));
 ENGINE_API const FName UCollisionProfile::Pawn_ProfileName = FName(TEXT("Pawn"));
 ENGINE_API const FName UCollisionProfile::Vehicle_ProfileName = FName(TEXT("Vehicle"));
+ENGINE_API const FName UCollisionProfile::DefaultProjectile_ProfileName = FName(TEXT("DefaultProjectile"));
+
 
 UCollisionProfile::UCollisionProfile(const FObjectInitializer& ObjectInitializer)
 	: Super(ObjectInitializer)
