#ue4 - Added constructor for FWalkableSlopeOverride that takes both editable properties.

diff --git a/Engine/Source/Runtime/Engine/Classes/Engine/EngineTypes.h b/Engine/Source/Runtime/Engine/Classes/Engine/EngineTypes.h
--- a/Engine/Source/Runtime/Engine/Classes/Engine/EngineTypes.h
+++ b/Engine/Source/Runtime/Engine/Classes/Engine/EngineTypes.h
@@ -2363,7 +2363,7 @@ struct FRepAttachment
 UENUM(BlueprintType)
 enum EWalkableSlopeBehavior
 {
-	/** Don't affect the walkable slope. */
+	/** Don't affect the walkable slope. Walkable slope angle will be ignored. */
 	WalkableSlope_Default		UMETA(DisplayName="Unchanged"),
 
 	/**
@@ -2412,6 +2412,12 @@ struct FWalkableSlopeOverride
 	{
 	}
 
+	FWalkableSlopeOverride(EWalkableSlopeBehavior NewSlopeBehavior, float NewSlopeAngle)
+	: WalkableSlopeBehavior(NewSlopeBehavior)
+	, WalkableSlopeAngle(NewSlopeAngle)
+	{
+	}
+
 	// Given a walkable floor normal Z value, either relax or restrict the value if we override such behavior.
 	float ModifyWalkableFloorZ(float InWalkableFloorZ) const
 	{
