#ue4 - Added SetWalkableSlopeOverride to PrimitiveComponent and FBodyInstance to make it easier to set up overrides, and to be able to change it at runtime in blueprints.

diff --git a/Engine/Source/Runtime/Engine/Classes/Components/PrimitiveComponent.h b/Engine/Source/Runtime/Engine/Classes/Components/PrimitiveComponent.h
--- a/Engine/Source/Runtime/Engine/Classes/Components/PrimitiveComponent.h
+++ b/Engine/Source/Runtime/Engine/Classes/Components/PrimitiveComponent.h
@@ -642,6 +642,10 @@ public:
 	UFUNCTION(BlueprintCallable, Category="Physics")
 	const struct FWalkableSlopeOverride& GetWalkableSlopeOverride() const;
 
+	/** Sets a new slope override for this component instance. */
+	UFUNCTION(BlueprintCallable, Category="Physics")
+	void SetWalkableSlopeOverride(const FWalkableSlopeOverride& NewOverride);
+
 	/** 
 	 *	Sets whether or not a single body should use physics simulation, or should be 'fixed' (kinematic).
 	 *
diff --git a/Engine/Source/Runtime/Engine/Classes/PhysicsEngine/BodyInstance.h b/Engine/Source/Runtime/Engine/Classes/PhysicsEngine/BodyInstance.h
--- a/Engine/Source/Runtime/Engine/Classes/PhysicsEngine/BodyInstance.h
+++ b/Engine/Source/Runtime/Engine/Classes/PhysicsEngine/BodyInstance.h
@@ -236,11 +236,14 @@ public:
 
 protected:
 
-	/** Whether this instance of the object has its own custom walkability override setting. */
+	/** Whether this instance of the object has its own custom walkable slope override setting. */
 	UPROPERTY(EditAnywhere, AdvancedDisplay, BlueprintReadOnly, Category=Physics)
 	uint32 bOverrideWalkableSlopeOnInstance:1;
 
-	/** Custom walkability override setting for this instance. @see GetWalkableFloorOverride() */
+	/**
+	 * Custom walkable slope override setting for this instance.
+	 * @see GetWalkableSlopeOverride(), SetWalkableSlopeOverride()
+	 */
 	UPROPERTY(EditAnywhere, AdvancedDisplay, BlueprintReadOnly, Category=Physics, meta=(editcondition="bOverrideWalkableSlopeOnInstance"))
 	struct FWalkableSlopeOverride WalkableSlopeOverride;
 
@@ -387,6 +390,9 @@ public:
 	/** Returns the slope override struct for this instance. If we don't have our own custom setting, it will return the setting from the body setup. */
 	const struct FWalkableSlopeOverride& GetWalkableSlopeOverride() const;
 
+	/** Sets a custom slope override struct for this instance. Implicitly sets bOverrideWalkableSlopeOnInstance to true. */
+	void SetWalkableSlopeOverride(const FWalkableSlopeOverride& NewOverride);
+
 	/** Returns whether this body wants (and can) use the async scene. */
 	bool UseAsyncScene() const;
 
diff --git a/Engine/Source/Runtime/Engine/Private/PhysicsEngine/BodyInstance.cpp b/Engine/Source/Runtime/Engine/Private/PhysicsEngine/BodyInstance.cpp
--- a/Engine/Source/Runtime/Engine/Private/PhysicsEngine/BodyInstance.cpp
+++ b/Engine/Source/Runtime/Engine/Private/PhysicsEngine/BodyInstance.cpp
@@ -2636,6 +2636,12 @@ const FWalkableSlopeOverride& FBodyInstance::GetWalkableSlopeOverride() const
 	}
 }
 
+void FBodyInstance::SetWalkableSlopeOverride(const FWalkableSlopeOverride& NewOverride)
+{
+	bOverrideWalkableSlopeOnInstance = true;
+	WalkableSlopeOverride = NewOverride;
+}
+
 
 /** 
 *	Changes the current PhysMaterialOverride for this body. 
diff --git a/Engine/Source/Runtime/Engine/Private/PrimitiveComponentPhysics.cpp b/Engine/Source/Runtime/Engine/Private/PrimitiveComponentPhysics.cpp
--- a/Engine/Source/Runtime/Engine/Private/PrimitiveComponentPhysics.cpp
+++ b/Engine/Source/Runtime/Engine/Private/PrimitiveComponentPhysics.cpp
@@ -143,6 +143,11 @@ const struct FWalkableSlopeOverride& UPrimitiveComponent::GetWalkableSlopeOverri
 	return BodyInstance.GetWalkableSlopeOverride();
 }
 
+void UPrimitiveComponent::SetWalkableSlopeOverride(const FWalkableSlopeOverride& NewOverride)
+{
+	BodyInstance.SetWalkableSlopeOverride(NewOverride);
+}
+
 void UPrimitiveComponent::SetSimulatePhysics(bool bSimulate)
 {
 	BodyInstance.SetInstanceSimulatePhysics(bSimulate);
