Licensee submission to fix black lightmaps only being tossed when a skylight is absent

diff --git a/Engine/Source/Runtime/Engine/Private/LightMap.cpp b/Engine/Source/Runtime/Engine/Private/LightMap.cpp
--- a/Engine/Source/Runtime/Engine/Private/LightMap.cpp
+++ b/Engine/Source/Runtime/Engine/Private/LightMap.cpp
@@ -1904,11 +1904,14 @@ bool FQuantizedLightmapData::HasNonZeroData() const
 				}
 			}
 
-			for (int32 Index = 0; Index < ARRAY_COUNT(LightmapSample.SkyOcclusion); Index++)
+			if (bHasSkyShadowing)
 			{
-				if (LightmapSample.SkyOcclusion[Index] != 0)
+				for (int32 Index = 0; Index < ARRAY_COUNT(LightmapSample.SkyOcclusion); Index++)
 				{
-					return true;
+					if (LightmapSample.SkyOcclusion[Index] != 0)
+					{
+						return true;
+					}
 				}
 			}
 		}
