#ue4 - Allow Actor.InitialLifeSpan to be edited on instances (EditAnywhere instead of EditDefaultsOnly).

diff --git a/Engine/Source/Runtime/Engine/Classes/GameFramework/Actor.h b/Engine/Source/Runtime/Engine/Classes/GameFramework/Actor.h
--- a/Engine/Source/Runtime/Engine/Classes/GameFramework/Actor.h
+++ b/Engine/Source/Runtime/Engine/Classes/GameFramework/Actor.h
@@ -360,7 +360,7 @@ protected:
 	TArray<class AMatineeActor*> ControllingMatineeActors;
 
 	/** How long this Actor lives before dying, 0=forever. Note this is the INITIAL value and should not be modified once play has begun. */
-	UPROPERTY(EditDefaultsOnly, BlueprintReadOnly, Category=Actor)
+	UPROPERTY(EditAnywhere, BlueprintReadOnly, Category=Actor)
 	float InitialLifeSpan; 
 
 	/**
