More math functions for the kismet library, Atan and DegAtan for completeness, square space length calculators for vector2d and vector. From Pull request 544 by dgrieshofer

diff --git a/Engine/Source/Runtime/Engine/Classes/Kismet/KismetMathLibrary.h b/Engine/Source/Runtime/Engine/Classes/Kismet/KismetMathLibrary.h
--- a/Engine/Source/Runtime/Engine/Classes/Kismet/KismetMathLibrary.h
+++ b/Engine/Source/Runtime/Engine/Classes/Kismet/KismetMathLibrary.h
@@ -320,6 +320,10 @@ class UKismetMathLibrary : public UBlueprintFunctionLibrary
 	UFUNCTION(BlueprintPure, meta=(FriendlyName = "Tan (Radians)", CompactNodeTitle = "TAN"), Category="Math|Trig")
 	static float Tan(float A);
 
+	/* Returns the inverse tan (atan) (result is in Radians)*/
+	UFUNCTION(BlueprintPure, meta=(FriendlyName = "Atan (Radians)"), Category="Math|Trig")
+	static float Atan(float A);
+
 	/* Returns the inverse tan (atan2) of A/B (result is in Radians)*/
 	UFUNCTION(BlueprintPure, meta=(FriendlyName = "Atan2 (Radians)"), Category="Math|Trig")
 	static float Atan2(float A, float B);
@@ -380,6 +384,10 @@ class UKismetMathLibrary : public UBlueprintFunctionLibrary
 	UFUNCTION(BlueprintPure, meta=(FriendlyName = "Tan (Degrees)", CompactNodeTitle = "TANd"), Category="Math|Trig")
 	static float DegTan(float A);
 
+	/* Returns the inverse tan (atan) (result is in Degrees)*/
+	UFUNCTION(BlueprintPure, meta=(FriendlyName = "Atan (Degrees)"), Category="Math|Trig")
+	static float DegAtan(float A);
+
 	/* Returns the inverse tan (atan2) of A/B (result is in Degrees)*/
 	UFUNCTION(BlueprintPure, meta=(FriendlyName = "Atan2 (Degrees)"), Category="Math|Trig")
 	static float DegAtan2(float A, float B);
@@ -545,6 +553,14 @@ class UKismetMathLibrary : public UBlueprintFunctionLibrary
 	UFUNCTION(BlueprintPure, meta=(FriendlyName = "Vector2dLength"), Category="Math|Vector2D")
 	static float VSize2D(FVector2D A);
 
+	/* Returns the squared length of the FVector */
+	UFUNCTION(BlueprintPure, meta=(FriendlyName = "VectorLengthSquared"), Category="Math|Vector")
+	static float VSizeSquared(FVector A);
+
+	/** Returns the squared length of a 2d FVector. */
+	UFUNCTION(BlueprintPure, meta=(FriendlyName = "Vector2dLengthSquared"), Category="Math|Vector2D")
+	static float VSize2DSquared(FVector2D A);
+
 	/* Returns a unit normal version of the FVector A */
 	UFUNCTION(BlueprintPure, meta=(FriendlyName = "Normalize"), Category="Math|Vector")
 	static FVector Normal(FVector A);
diff --git a/Engine/Source/Runtime/Engine/Private/KismetMathLibrary.cpp b/Engine/Source/Runtime/Engine/Private/KismetMathLibrary.cpp
--- a/Engine/Source/Runtime/Engine/Private/KismetMathLibrary.cpp
+++ b/Engine/Source/Runtime/Engine/Private/KismetMathLibrary.cpp
@@ -392,6 +392,11 @@ float UKismetMathLibrary::Tan(float A)
 	return FMath::Tan(A);
 }	
 
+float UKismetMathLibrary::Atan(float A)
+{
+	return FMath::Atan(A);
+}
+
 float UKismetMathLibrary::Atan2(float A, float B)
 {
 	return FMath::Atan2(A, B);
@@ -422,6 +427,11 @@ float UKismetMathLibrary::DegTan(float A)
 	return FMath::Tan(PI/(180.f) * A);
 }
 
+float UKismetMathLibrary::DegAtan(float A)
+{
+	return (180.f)/PI * FMath::Atan(A);
+}
+
 float UKismetMathLibrary::DegAtan2(float A, float B)
 {
 	return (180.f)/PI * FMath::Atan2(A, B);
@@ -713,6 +723,16 @@ float UKismetMathLibrary::VSize2D(FVector2D A)
 	return A.Size();
 }
 
+float UKismetMathLibrary::VSizeSquared(FVector A)
+{
+	return A.SizeSquared();
+}
+
+float UKismetMathLibrary::VSize2DSquared(FVector2D A)
+{
+	return A.SizeSquared();
+}
+
 FVector UKismetMathLibrary::Normal(FVector A)
 {
 	return A.SafeNormal();
