Remove headers from EnginePrivate.h which are included 0 or 1 time in Engine module.

diff --git a/Engine/Source/Runtime/Engine/Classes/Components/SceneCaptureComponent2D.h b/Engine/Source/Runtime/Engine/Classes/Components/SceneCaptureComponent2D.h
--- a/Engine/Source/Runtime/Engine/Classes/Components/SceneCaptureComponent2D.h
+++ b/Engine/Source/Runtime/Engine/Classes/Components/SceneCaptureComponent2D.h
@@ -2,6 +2,7 @@
 
 
 #pragma once
+#include "Components/SceneCaptureComponent.h"
 #include "SceneCaptureComponent2D.generated.h"
 
 // -> will be exported to EngineDecalClasses.h
diff --git a/Engine/Source/Runtime/Engine/Classes/Components/SceneCaptureComponentCube.h b/Engine/Source/Runtime/Engine/Classes/Components/SceneCaptureComponentCube.h
--- a/Engine/Source/Runtime/Engine/Classes/Components/SceneCaptureComponentCube.h
+++ b/Engine/Source/Runtime/Engine/Classes/Components/SceneCaptureComponentCube.h
@@ -2,6 +2,7 @@
 
 
 #pragma once
+#include "Components/SceneCaptureComponent.h"
 #include "SceneCaptureComponentCube.generated.h"
 /**
  *	Used to capture a 'snapshot' of the scene from a 6 planes and feed it to a render target.
diff --git a/Engine/Source/Runtime/Engine/Classes/Engine/DirectionalLight.h b/Engine/Source/Runtime/Engine/Classes/Engine/DirectionalLight.h
--- a/Engine/Source/Runtime/Engine/Classes/Engine/DirectionalLight.h
+++ b/Engine/Source/Runtime/Engine/Classes/Engine/DirectionalLight.h
@@ -1,7 +1,7 @@
 // Copyright 1998-2014 Epic Games, Inc. All Rights Reserved.
 
 #pragma once
-
+#include "Engine/Light.h"
 #include "DirectionalLight.generated.h"
 
 UCLASS(ClassGroup=(Lights, DirectionalLights), MinimalAPI, meta=(ChildCanTick))
diff --git a/Engine/Source/Runtime/Engine/Classes/PhysicsEngine/PhysicsConstraintActor.h b/Engine/Source/Runtime/Engine/Classes/PhysicsEngine/PhysicsConstraintActor.h
--- a/Engine/Source/Runtime/Engine/Classes/PhysicsEngine/PhysicsConstraintActor.h
+++ b/Engine/Source/Runtime/Engine/Classes/PhysicsEngine/PhysicsConstraintActor.h
@@ -4,7 +4,7 @@
 //=============================================================================
 // The Basic constraint actor class.
 //=============================================================================
-
+#include "PhysicsEngine/RigidBodyBase.h"
 #include "PhysicsConstraintActor.generated.h"
 
 UCLASS(ConversionRoot, MinimalAPI)
diff --git a/Engine/Source/Runtime/Engine/Classes/Vehicles/TireType.h b/Engine/Source/Runtime/Engine/Classes/Vehicles/TireType.h
--- a/Engine/Source/Runtime/Engine/Classes/Vehicles/TireType.h
+++ b/Engine/Source/Runtime/Engine/Classes/Vehicles/TireType.h
@@ -5,6 +5,7 @@
  */
 
 #pragma once
+#include "Engine/DataAsset.h"
 #include "TireType.generated.h"
 
 UCLASS()
diff --git a/Engine/Source/Runtime/Engine/Private/Actor.cpp b/Engine/Source/Runtime/Engine/Private/Actor.cpp
--- a/Engine/Source/Runtime/Engine/Private/Actor.cpp
+++ b/Engine/Source/Runtime/Engine/Private/Actor.cpp
@@ -3,6 +3,7 @@
 
 #include "EnginePrivate.h"
 #include "Engine/InputDelegateBinding.h"
+#include "Engine/LevelStreamingPersistent.h"
 #include "PhysicsPublic.h"
 #include "ParticleDefinitions.h"
 #include "Sound/SoundCue.h"
diff --git a/Engine/Source/Runtime/Engine/Private/ActorComponent.cpp b/Engine/Source/Runtime/Engine/Private/ActorComponent.cpp
--- a/Engine/Source/Runtime/Engine/Private/ActorComponent.cpp
+++ b/Engine/Source/Runtime/Engine/Private/ActorComponent.cpp
@@ -2,6 +2,8 @@
 // ActorComponent.cpp: Actor component implementation.
 
 #include "EnginePrivate.h"
+#include "Engine/AssetUserData.h"
+#include "Engine/LevelStreamingPersistent.h"
 #include "Net/UnrealNetwork.h"
 #include "MessageLog.h"
 #include "UObjectToken.h"
diff --git a/Engine/Source/Runtime/Engine/Private/AssetUserData.cpp b/Engine/Source/Runtime/Engine/Private/AssetUserData.cpp
--- a/Engine/Source/Runtime/Engine/Private/AssetUserData.cpp
+++ b/Engine/Source/Runtime/Engine/Private/AssetUserData.cpp
@@ -1,7 +1,7 @@
 // Copyright 1998-2014 Epic Games, Inc. All Rights Reserved.
 
 #include "EnginePrivate.h"
-
+#include "Engine/AssetUserData.h"
 
 
 UAssetUserData::UAssetUserData(const FObjectInitializer& ObjectInitializer)
diff --git a/Engine/Source/Runtime/Engine/Private/Atmosphere/Atmosphere.h b/Engine/Source/Runtime/Engine/Private/Atmosphere/Atmosphere.h
--- a/Engine/Source/Runtime/Engine/Private/Atmosphere/Atmosphere.h
+++ b/Engine/Source/Runtime/Engine/Private/Atmosphere/Atmosphere.h
@@ -1,6 +1,7 @@
 // Copyright 1998-2014 Epic Games, Inc. All Rights Reserved.
 
 #pragma once
+#include "Atmosphere/AtmosphericFogComponent.h"
 
 // Shared by Engine class and Renderer class, need to be core?
 
diff --git a/Engine/Source/Runtime/Engine/Private/BillboardComponent.cpp b/Engine/Source/Runtime/Engine/Private/BillboardComponent.cpp
--- a/Engine/Source/Runtime/Engine/Private/BillboardComponent.cpp
+++ b/Engine/Source/Runtime/Engine/Private/BillboardComponent.cpp
@@ -1,6 +1,7 @@
 // Copyright 1998-2014 Epic Games, Inc. All Rights Reserved.
 
 #include "EnginePrivate.h"
+#include "Engine/Light.h"
 #include "LevelUtils.h"
 
 namespace BillboardConstants
diff --git a/Engine/Source/Runtime/Engine/Private/DataAsset.cpp b/Engine/Source/Runtime/Engine/Private/DataAsset.cpp
--- a/Engine/Source/Runtime/Engine/Private/DataAsset.cpp
+++ b/Engine/Source/Runtime/Engine/Private/DataAsset.cpp
@@ -1,6 +1,7 @@
 // Copyright 1998-2014 Epic Games, Inc. All Rights Reserved.
 
 #include "EnginePrivate.h"
+#include "Engine/DataAsset.h"
 
 UDataAsset::UDataAsset(const FObjectInitializer& ObjectInitializer)
 	: Super(ObjectInitializer)
diff --git a/Engine/Source/Runtime/Engine/Private/DistributionFloatSoundParameter.cpp b/Engine/Source/Runtime/Engine/Private/DistributionFloatSoundParameter.cpp
--- a/Engine/Source/Runtime/Engine/Private/DistributionFloatSoundParameter.cpp
+++ b/Engine/Source/Runtime/Engine/Private/DistributionFloatSoundParameter.cpp
@@ -1,6 +1,7 @@
 // Copyright 1998-2014 Epic Games, Inc. All Rights Reserved.
 
 #include "EnginePrivate.h"
+#include "Distributions/DistributionFloatSoundParameter.h"
 #include "SoundDefinitions.h"
 
 /*-----------------------------------------------------------------------------
diff --git a/Engine/Source/Runtime/Engine/Private/DrawSphereComponent.cpp b/Engine/Source/Runtime/Engine/Private/DrawSphereComponent.cpp
--- a/Engine/Source/Runtime/Engine/Private/DrawSphereComponent.cpp
+++ b/Engine/Source/Runtime/Engine/Private/DrawSphereComponent.cpp
@@ -1,6 +1,7 @@
 // Copyright 1998-2014 Epic Games, Inc. All Rights Reserved.
 
 #include "EnginePrivate.h"
+#include "Components/DrawSphereComponent.h"
 
 UDrawSphereComponent::UDrawSphereComponent(const FObjectInitializer& ObjectInitializer)
 	: Super(ObjectInitializer)
diff --git a/Engine/Source/Runtime/Engine/Private/EditorFramework/ThumbnailInfo.cpp b/Engine/Source/Runtime/Engine/Private/EditorFramework/ThumbnailInfo.cpp
--- a/Engine/Source/Runtime/Engine/Private/EditorFramework/ThumbnailInfo.cpp
+++ b/Engine/Source/Runtime/Engine/Private/EditorFramework/ThumbnailInfo.cpp
@@ -1,6 +1,7 @@
 // Copyright 1998-2014 Epic Games, Inc. All Rights Reserved.
 
 #include "EnginePrivate.h"
+#include "EditorFramework/ThumbnailInfo.h"
 
 UThumbnailInfo::UThumbnailInfo(const FObjectInitializer& ObjectInitializer)
 	: Super(ObjectInitializer)
diff --git a/Engine/Source/Runtime/Engine/Private/Engine.cpp b/Engine/Source/Runtime/Engine/Private/Engine.cpp
--- a/Engine/Source/Runtime/Engine/Private/Engine.cpp
+++ b/Engine/Source/Runtime/Engine/Private/Engine.cpp
@@ -5,6 +5,7 @@
 =============================================================================*/
 
 #include "EnginePrivate.h"
+#include "Materials/MaterialInstanceConstant.h"
 #include "EngineModule.h"
 #include "EditorSupportDelegates.h"
 #include "ComponentReregisterContext.h"
diff --git a/Engine/Source/Runtime/Engine/Private/EngineMessage.cpp b/Engine/Source/Runtime/Engine/Private/EngineMessage.cpp
--- a/Engine/Source/Runtime/Engine/Private/EngineMessage.cpp
+++ b/Engine/Source/Runtime/Engine/Private/EngineMessage.cpp
@@ -2,6 +2,7 @@
 
 #include "EnginePrivate.h"
 #include "Engine/Console.h"
+#include "GameFramework/EngineMessage.h"
 
 UEngineMessage::UEngineMessage(const FObjectInitializer& ObjectInitializer) :
 	Super(ObjectInitializer)
diff --git a/Engine/Source/Runtime/Engine/Private/EnginePrivate.h b/Engine/Source/Runtime/Engine/Private/EnginePrivate.h
--- a/Engine/Source/Runtime/Engine/Private/EnginePrivate.h
+++ b/Engine/Source/Runtime/Engine/Private/EnginePrivate.h
@@ -4,6 +4,212 @@
 	EnginePrivate.h: Unreal engine private header file.
 =============================================================================*/
 
+#pragma once
+#define USE_NEW_PCH
+#ifdef USE_NEW_PCH
+#include "Core.h"
+#include "CoreUObject.h"
+#include "InputCore.h"
+#include "EngineDefines.h"
+#include "EngineSettings.h"
+#include "EngineStats.h"
+#include "EngineLogs.h"
+#include "EngineGlobals.h"
+
+#include "Engine/EngineBaseTypes.h"
+#include "Camera/CameraTypes.h"
+#include "Engine/EngineTypes.h"
+#include "Components/ActorComponent.h"
+#include "GameFramework/Actor.h"
+#include "Engine/Brush.h"
+#include "AI/Navigation/NavAgentInterface.h"
+#include "GameFramework/Controller.h"
+#include "AI/Navigation/NavLinkDefinition.h"
+#include "AI/Navigation/NavigationTypes.h"
+#include "Kismet/BlueprintFunctionLibrary.h"
+#include "AI/Navigation/NavigationSystem.h"
+#include "AI/Navigation/NavCollision.h"
+#include "Camera/PlayerCameraManager.h"
+#include "GameFramework/Info.h"
+#include "GameFramework/GameMode.h"
+#include "Components/InputComponent.h"
+#include "GameFramework/ForceFeedbackEffect.h"
+#include "Engine/SubsurfaceProfile.h"
+#include "GameFramework/PlayerController.h"
+#include "Engine/DebugCameraController.h"
+#include "Engine/Canvas.h"
+#include "GameFramework/GameSession.h"
+#include "GameFramework/GameState.h"
+#include "GameFramework/PlayerState.h"
+#include "GameFramework/WorldSettings.h"
+#include "PhysicsEngine/BodyInstance.h"
+#include "Components/SceneComponent.h"
+#include "Components/LightComponentBase.h"
+#include "Components/LightComponent.h"
+#include "AI/Navigation/NavigationData.h"
+#include "AI/Navigation/NavAreas/NavArea.h"
+#include "AI/Navigation/RecastNavMesh.h"
+#include "Engine/NavigationObjectBase.h"
+#include "GameFramework/PlayerStart.h"
+#include "AI/Navigation/NavPathObserverInterface.h"
+#include "AI/Navigation/NavigationTestingActor.h"
+#include "AI/Navigation/NavLinkHostInterface.h"
+#include "Engine/Note.h"
+#include "Components/PrimitiveComponent.h"
+#include "GameFramework/Pawn.h"
+#include "Interfaces/NetworkPredictionInterface.h"
+#include "PhysicsEngine/PhysicsConstraintActor.h"
+#include "Components/MeshComponent.h"
+#include "Components/SkinnedMeshComponent.h"
+#include "PhysicsEngine/ConstraintInstance.h"
+#include "PhysicsEngine/PhysicsConstraintComponent.h"
+#include "Components/SkeletalMeshComponent.h"
+#include "Components/PawnNoiseEmitterComponent.h"
+#include "Components/AudioComponent.h"
+#include "Camera/CameraComponent.h"
+#include "Components/ChildActorComponent.h"
+#include "Components/DecalComponent.h"
+#include "Components/DirectionalLightComponent.h"
+#include "Components/PointLightComponent.h"
+#include "Components/SpotLightComponent.h"
+#include "Components/SkyLightComponent.h"
+#include "Components/ArrowComponent.h"
+#include "Components/BillboardComponent.h"
+#include "Components/BrushComponent.h"
+#include "Components/DrawFrustumComponent.h"
+#include "Debug/DebugDrawService.h"
+#include "Components/LineBatchComponent.h"
+#include "Components/MaterialBillboardComponent.h"
+#include "Components/DestructibleComponent.h"
+#include "Components/StaticMeshComponent.h"
+#include "Components/InstancedStaticMeshComponent.h"
+#include "Interfaces/Interface_CollisionDataProvider.h"
+#include "Components/ModelComponent.h"
+#include "Components/ShapeComponent.h"
+#include "Components/BoxComponent.h"
+#include "Components/CapsuleComponent.h"
+#include "Components/SphereComponent.h"
+#include "Components/VectorFieldComponent.h"
+#include "Components/ReflectionCaptureComponent.h"
+#include "Components/TimelineComponent.h"
+#include "EditorFramework/AssetImportData.h"
+#include "AI/Navigation/AvoidanceManager.h"
+#include "Engine/BlendableInterface.h"
+#include "Engine/BlueprintCore.h"
+#include "Engine/Blueprint.h"
+#include "Sound/DialogueTypes.h"
+#include "Kismet/GameplayStatics.h"
+#include "Engine/BlueprintGeneratedClass.h"
+#include "PhysicsEngine/BodySetup.h"
+#include "Camera/CameraAnim.h"
+#include "Camera/CameraAnimInst.h"
+#include "Camera/CameraModifier.h"
+#include "Camera/CameraShake.h"
+#include "Camera/CameraModifier_CameraShake.h"
+#include "GameFramework/CheatManager.h"
+#include "Engine/CollisionProfile.h"
+#include "Curves/CurveLinearColor.h"
+#include "Curves/CurveVector.h"
+#include "Engine/CurveTable.h"
+#include "GameFramework/DamageType.h"
+#include "Engine/DataTable.h"
+#include "Sound/DialogueVoice.h"
+#include "Sound/DialogueWave.h"
+#include "Distributions/Distribution.h"
+#include "Distributions/DistributionFloat.h"
+#include "Distributions/DistributionFloatConstant.h"
+#include "Distributions/DistributionFloatParameterBase.h"
+#include "Distributions/DistributionFloatConstantCurve.h"
+#include "Distributions/DistributionFloatUniform.h"
+#include "Distributions/DistributionFloatUniformCurve.h"
+#include "Distributions/DistributionVector.h"
+#include "Distributions/DistributionVectorConstant.h"
+#include "Distributions/DistributionVectorParameterBase.h"
+#include "Distributions/DistributionVectorConstantCurve.h"
+#include "Distributions/DistributionVectorUniform.h"
+#include "Distributions/DistributionVectorUniformCurve.h"
+#include "Engine/Engine.h"
+#include "Engine/GameEngine.h"
+#include "Exporters/Exporter.h"
+#include "Engine/FontImportOptions.h"
+#include "Engine/Font.h"
+#include "GameFramework/GameUserSettings.h"
+#include "GameFramework/PlayerInput.h"
+#include "GameFramework/InputSettings.h"
+#include "Engine/InterpCurveEdSetup.h"
+#include "Engine/Level.h"
+#include "Engine/LevelStreaming.h"
+#include "Engine/LevelStreamingKismet.h"
+#include "GameFramework/LocalMessage.h"
+#include "Materials/MaterialInterface.h"
+#include "Materials/Material.h"
+#include "Materials/MaterialInstance.h"
+#include "Materials/MaterialInstanceDynamic.h"
+#include "Materials/MaterialParameterCollection.h"
+#include "Materials/MaterialParameterCollectionInstance.h"
+#include "Engine/NetDriver.h"
+#include "Engine/PackageMapClient.h"
+#include "Engine/PendingNetGame.h"
+#include "PhysicalMaterials/PhysicalMaterial.h"
+#include "PhysicalMaterials/PhysicalMaterialPropertyBase.h"
+#include "PhysicsEngine/PhysicsAsset.h"
+#include "PhysicsEngine/PhysicsCollisionHandler.h"
+#include "PhysicsEngine/PhysicsConstraintTemplate.h"
+#include "PhysicsEngine/PhysicsSettings.h"
+#include "Engine/PlatformInterfaceBase.h"
+#include "Engine/InGameAdManager.h"
+#include "Engine/Player.h"
+#include "Engine/LocalPlayer.h"
+#include "Engine/NetConnection.h"
+#include "Engine/ChildConnection.h"
+#include "Engine/Polys.h"
+#include "Engine/SCS_Node.h"
+#include "Engine/Selection.h"
+#include "Engine/SimpleConstructionScript.h"
+#include "Engine/SkeletalMesh.h"
+#include "Engine/SkeletalMeshSocket.h"
+#include "Engine/StaticMesh.h"
+#include "Engine/StaticMeshSocket.h"
+#include "Engine/TextureDefines.h"
+#include "Engine/Texture.h"
+#include "Engine/Texture2D.h"
+#include "Engine/LightMapTexture2D.h"
+#include "Engine/ShadowMapTexture2D.h"
+#include "Engine/TextureLightProfile.h"
+#include "Engine/TextureCube.h"
+#include "Engine/TextureRenderTarget.h"
+#include "Engine/TextureRenderTarget2D.h"
+#include "Engine/TimelineTemplate.h"
+#include "GameFramework/TouchInterface.h"
+#include "Engine/UserDefinedStruct.h"
+#include "Engine/World.h"
+
+#include "SystemSettings.h"					// Scalability options.
+#include "SceneManagement.h"				// Scene management.
+
+#include "Engine/Light.h"
+#include "Interfaces/Interface_PostProcessVolume.h"
+#include "NavDataGenerator.h"
+#include "NavigationModifier.h"
+#include "NavigationOctree.h"
+#include "ContentStreaming.h"				// Content streaming class definitions.
+#include "ComponentReregisterContext.h"	
+#include "DrawDebugHelpers.h"
+#include "UnrealEngine.h"					// Unreal engine helpers.
+#include "CanvasTypes.h"							// Canvas.
+#include "EngineUtils.h"
+#include "InstancedFoliage.h"				// Instanced foliage.
+#include "UnrealExporter.h"					// Exporter definition.
+#include "TimerManager.h"					// Game play timers
+#include "NavigationSystemHelpers.h"
+#include "HardwareInfo.h"
+#else
+// Copyright 1998-2014 Epic Games, Inc. All Rights Reserved.
+
+/*=============================================================================
+EnginePrivate.h: Unreal engine private header file.
+=============================================================================*/
+
 #pragma once
 
 
@@ -275,6 +481,4 @@
 #include "EngineService.h"
 #include "NavigationSystemHelpers.h"
 #include "HardwareInfo.h"
-
-
-
+#endif // USE_NEW_PCH
\ No newline at end of file
diff --git a/Engine/Source/Runtime/Engine/Private/EngineService.cpp b/Engine/Source/Runtime/Engine/Private/EngineService.cpp
--- a/Engine/Source/Runtime/Engine/Private/EngineService.cpp
+++ b/Engine/Source/Runtime/Engine/Private/EngineService.cpp
@@ -1,6 +1,7 @@
 // Copyright 1998-2014 Epic Games, Inc. All Rights Reserved.
 
 #include "EnginePrivate.h"
+#include "EngineService.h"
 #include "EngineServiceMessages.h"
 #include "MessageEndpoint.h"
 #include "MessageEndpointBuilder.h"
diff --git a/Engine/Source/Runtime/Engine/Private/ErrorChecking.cpp b/Engine/Source/Runtime/Engine/Private/ErrorChecking.cpp
--- a/Engine/Source/Runtime/Engine/Private/ErrorChecking.cpp
+++ b/Engine/Source/Runtime/Engine/Private/ErrorChecking.cpp
@@ -14,6 +14,7 @@
 #include "GameFramework/PainCausingVolume.h"
 #include "GameFramework/DefaultPhysicsVolume.h"
 #include "Engine/LevelStreamingVolume.h"
+#include "Engine/Light.h"
 
 #define LOCTEXT_NAMESPACE "ErrorChecking"
 
diff --git a/Engine/Source/Runtime/Engine/Private/GameFramework/SpringArmComponent.cpp b/Engine/Source/Runtime/Engine/Private/GameFramework/SpringArmComponent.cpp
--- a/Engine/Source/Runtime/Engine/Private/GameFramework/SpringArmComponent.cpp
+++ b/Engine/Source/Runtime/Engine/Private/GameFramework/SpringArmComponent.cpp
@@ -1,6 +1,7 @@
 // Copyright 1998-2014 Epic Games, Inc. All Rights Reserved.
 
 #include "EnginePrivate.h"
+#include "GameFramework/SpringArmComponent.h"
 
 //////////////////////////////////////////////////////////////////////////
 // USpringArmComponent
diff --git a/Engine/Source/Runtime/Engine/Private/HeadMountedDisplayFunctionLibrary.cpp b/Engine/Source/Runtime/Engine/Private/HeadMountedDisplayFunctionLibrary.cpp
--- a/Engine/Source/Runtime/Engine/Private/HeadMountedDisplayFunctionLibrary.cpp
+++ b/Engine/Source/Runtime/Engine/Private/HeadMountedDisplayFunctionLibrary.cpp
@@ -1,6 +1,7 @@
 // Copyright 1998-2014 Epic Games, Inc. All Rights Reserved.
 
 #include "EnginePrivate.h"
+#include "Kismet/HeadMountedDisplayFunctionLibrary.h"
 #include "HeadMountedDisplay.h"
 
 UHeadMountedDisplayFunctionLibrary::UHeadMountedDisplayFunctionLibrary(const FObjectInitializer& ObjectInitializer)
diff --git a/Engine/Source/Runtime/Engine/Private/HeightFogComponent.cpp b/Engine/Source/Runtime/Engine/Private/HeightFogComponent.cpp
--- a/Engine/Source/Runtime/Engine/Private/HeightFogComponent.cpp
+++ b/Engine/Source/Runtime/Engine/Private/HeightFogComponent.cpp
@@ -5,8 +5,9 @@
 =============================================================================*/
 
 #include "EnginePrivate.h"
-#include "Net/UnrealNetwork.h"
+#include "Components/ExponentialHeightFogComponent.h"
 #include "Engine/ExponentialHeightFog.h"
+#include "Net/UnrealNetwork.h"
 
 UExponentialHeightFogComponent::UExponentialHeightFogComponent(const FObjectInitializer& ObjectInitializer)
 	: Super(ObjectInitializer)
diff --git a/Engine/Source/Runtime/Engine/Private/Interface.cpp b/Engine/Source/Runtime/Engine/Private/Interface.cpp
--- a/Engine/Source/Runtime/Engine/Private/Interface.cpp
+++ b/Engine/Source/Runtime/Engine/Private/Interface.cpp
@@ -1,6 +1,8 @@
 // Copyright 1998-2014 Epic Games, Inc. All Rights Reserved.
 
 #include "EnginePrivate.h"
+#include "Interfaces/Interface_CollisionDataProvider.h"
+#include "Interfaces/Interface_PostProcessVolume.h"
 #include "MovieScene/RuntimeMovieScenePlayerInterface.h"
 
 UInterface_CollisionDataProvider::UInterface_CollisionDataProvider(const FObjectInitializer& ObjectInitializer)
diff --git a/Engine/Source/Runtime/Engine/Private/Layers/Layer.cpp b/Engine/Source/Runtime/Engine/Private/Layers/Layer.cpp
--- a/Engine/Source/Runtime/Engine/Private/Layers/Layer.cpp
+++ b/Engine/Source/Runtime/Engine/Private/Layers/Layer.cpp
@@ -1,5 +1,6 @@
 // Copyright 1998-2014 Epic Games, Inc. All Rights Reserved.
 #include "EnginePrivate.h"
+#include "Layers/Layer.h"
 
 ULayer::ULayer(const FObjectInitializer& ObjectInitializer)
 	: Super(ObjectInitializer)
diff --git a/Engine/Source/Runtime/Engine/Private/Level.cpp b/Engine/Source/Runtime/Engine/Private/Level.cpp
--- a/Engine/Source/Runtime/Engine/Private/Level.cpp
+++ b/Engine/Source/Runtime/Engine/Private/Level.cpp
@@ -5,6 +5,7 @@ Level.cpp: Level-related functions
 =============================================================================*/
 
 #include "EnginePrivate.h"
+#include "Engine/AssetUserData.h"
 #include "Engine/LevelScriptBlueprint.h"
 #include "Engine/LevelScriptActor.h"
 #include "Engine/WorldComposition.h"
diff --git a/Engine/Source/Runtime/Engine/Private/LevelStreaming.cpp b/Engine/Source/Runtime/Engine/Private/LevelStreaming.cpp
--- a/Engine/Source/Runtime/Engine/Private/LevelStreaming.cpp
+++ b/Engine/Source/Runtime/Engine/Private/LevelStreaming.cpp
@@ -1,6 +1,9 @@
 // Copyright 1998-2014 Epic Games, Inc. All Rights Reserved.
 
 #include "EnginePrivate.h"
+#include "Engine/LevelStreamingAlwaysLoaded.h"
+#include "Engine/LevelStreamingBounds.h"
+#include "Engine/LevelStreamingPersistent.h"
 #include "Engine/LevelStreamingVolume.h"
 #include "Engine/LevelBounds.h"
 #include "LevelUtils.h"
diff --git a/Engine/Source/Runtime/Engine/Private/LevelTick.cpp b/Engine/Source/Runtime/Engine/Private/LevelTick.cpp
--- a/Engine/Source/Runtime/Engine/Private/LevelTick.cpp
+++ b/Engine/Source/Runtime/Engine/Private/LevelTick.cpp
@@ -5,6 +5,8 @@
 =============================================================================*/
 
 #include "EnginePrivate.h"
+#include "Components/SceneCaptureComponent2D.h"
+#include "Components/SceneCaptureComponentCube.h"
 #include "Engine/LevelStreamingVolume.h"
 #include "Net/UnrealNetwork.h"
 #include "Collision.h"
diff --git a/Engine/Source/Runtime/Engine/Private/Light.cpp b/Engine/Source/Runtime/Engine/Private/Light.cpp
--- a/Engine/Source/Runtime/Engine/Private/Light.cpp
+++ b/Engine/Source/Runtime/Engine/Private/Light.cpp
@@ -2,8 +2,10 @@
 
 
 #include "EnginePrivate.h"
+#include "Engine/Light.h"
 #include "Engine/PointLight.h"
 #include "Engine/DirectionalLight.h"
+#include "Engine/GeneratedMeshAreaLight.h"
 #include "Components/ArrowComponent.h"
 #include "LightingBuildOptions.h"
 #include "Net/UnrealNetwork.h"
diff --git a/Engine/Source/Runtime/Engine/Private/LightmappedSurfaceCollection.cpp b/Engine/Source/Runtime/Engine/Private/LightmappedSurfaceCollection.cpp
--- a/Engine/Source/Runtime/Engine/Private/LightmappedSurfaceCollection.cpp
+++ b/Engine/Source/Runtime/Engine/Private/LightmappedSurfaceCollection.cpp
@@ -1,7 +1,7 @@
 // Copyright 1998-2014 Epic Games, Inc. All Rights Reserved.
 
 #include "EnginePrivate.h"
-
+#include "Lightmass/LightmappedSurfaceCollection.h"
 ULightmappedSurfaceCollection::ULightmappedSurfaceCollection(const FObjectInitializer& ObjectInitializer)
 	: Super(ObjectInitializer)
 {
diff --git a/Engine/Source/Runtime/Engine/Private/LightmassPrimitiveSettingsObject.cpp b/Engine/Source/Runtime/Engine/Private/LightmassPrimitiveSettingsObject.cpp
--- a/Engine/Source/Runtime/Engine/Private/LightmassPrimitiveSettingsObject.cpp
+++ b/Engine/Source/Runtime/Engine/Private/LightmassPrimitiveSettingsObject.cpp
@@ -1,6 +1,7 @@
 // Copyright 1998-2014 Epic Games, Inc. All Rights Reserved.
 
 #include "EnginePrivate.h"
+#include "Lightmass/LightmassPrimitiveSettingsObject.h"
 
 ULightmassPrimitiveSettingsObject::ULightmassPrimitiveSettingsObject(const FObjectInitializer& ObjectInitializer)
 	: Super(ObjectInitializer)
diff --git a/Engine/Source/Runtime/Engine/Private/Materials/MaterialInstance.cpp b/Engine/Source/Runtime/Engine/Private/Materials/MaterialInstance.cpp
--- a/Engine/Source/Runtime/Engine/Private/Materials/MaterialInstance.cpp
+++ b/Engine/Source/Runtime/Engine/Private/Materials/MaterialInstance.cpp
@@ -7,6 +7,7 @@
 #include "Materials/MaterialExpressionFontSampleParameter.h"
 #include "Materials/MaterialExpressionStaticBoolParameter.h"
 #include "Materials/MaterialExpressionStaticComponentMaskParameter.h"
+#include "Materials/MaterialInstanceConstant.h"
 #include "Materials/MaterialInstanceBasePropertyOverrides.h"
 #include "MaterialUniformExpressions.h"
 #include "MaterialInstanceSupport.h"
diff --git a/Engine/Source/Runtime/Engine/Private/Materials/MaterialInstanceConstant.cpp b/Engine/Source/Runtime/Engine/Private/Materials/MaterialInstanceConstant.cpp
--- a/Engine/Source/Runtime/Engine/Private/Materials/MaterialInstanceConstant.cpp
+++ b/Engine/Source/Runtime/Engine/Private/Materials/MaterialInstanceConstant.cpp
@@ -5,7 +5,8 @@
 =============================================================================*/
 
 #include "EnginePrivate.h"
-#include "MaterialInstanceSupport.h"
+#include "Materials/MaterialInstanceConstant.h"
+#include "Materials/MaterialInstanceSupport.h"
 
 
 UMaterialInstanceConstant::UMaterialInstanceConstant(const FObjectInitializer& ObjectInitializer)
diff --git a/Engine/Source/Runtime/Engine/Private/Particles/ParticleModules_Parameter.cpp b/Engine/Source/Runtime/Engine/Private/Particles/ParticleModules_Parameter.cpp
--- a/Engine/Source/Runtime/Engine/Private/Particles/ParticleModules_Parameter.cpp
+++ b/Engine/Source/Runtime/Engine/Private/Particles/ParticleModules_Parameter.cpp
@@ -6,6 +6,7 @@
 =============================================================================*/
 #include "EnginePrivate.h"
 #include "Materials/MaterialExpressionDynamicParameter.h"
+#include "Materials/MaterialInstanceConstant.h"
 #include "ParticleDefinitions.h"
 #include "../DistributionHelpers.h"
 #include "Particles/Parameter/ParticleModuleParameterBase.h"
diff --git a/Engine/Source/Runtime/Engine/Private/PhysicsEngine/PhysicsHandleComponent.cpp b/Engine/Source/Runtime/Engine/Private/PhysicsEngine/PhysicsHandleComponent.cpp
--- a/Engine/Source/Runtime/Engine/Private/PhysicsEngine/PhysicsHandleComponent.cpp
+++ b/Engine/Source/Runtime/Engine/Private/PhysicsEngine/PhysicsHandleComponent.cpp
@@ -1,6 +1,7 @@
 // Copyright 1998-2014 Epic Games, Inc. All Rights Reserved.
 
 #include "EnginePrivate.h"
+#include "PhysicsEngine/PhysicsHandleComponent.h"
 #include "PhysicsPublic.h"
 #if WITH_PHYSX
 	#include "PhysXSupport.h"
diff --git a/Engine/Source/Runtime/Engine/Private/PhysicsEngine/PhysicsThruster.cpp b/Engine/Source/Runtime/Engine/Private/PhysicsEngine/PhysicsThruster.cpp
--- a/Engine/Source/Runtime/Engine/Private/PhysicsEngine/PhysicsThruster.cpp
+++ b/Engine/Source/Runtime/Engine/Private/PhysicsEngine/PhysicsThruster.cpp
@@ -1,7 +1,9 @@
 // Copyright 1998-2014 Epic Games, Inc. All Rights Reserved.
 
 #include "EnginePrivate.h"
+#include "PhysicsEngine/PhysicsThrusterComponent.h"
 #include "PhysXSupport.h"
+#include "PhysicsEngine/PhysicsThruster.h"
 
 
 UPhysicsThrusterComponent::UPhysicsThrusterComponent(const FObjectInitializer& ObjectInitializer)
diff --git a/Engine/Source/Runtime/Engine/Private/PhysicsEngine/RadialForceComponent.cpp b/Engine/Source/Runtime/Engine/Private/PhysicsEngine/RadialForceComponent.cpp
--- a/Engine/Source/Runtime/Engine/Private/PhysicsEngine/RadialForceComponent.cpp
+++ b/Engine/Source/Runtime/Engine/Private/PhysicsEngine/RadialForceComponent.cpp
@@ -3,6 +3,8 @@
 #include "EnginePrivate.h"
 #include "GameFramework/MovementComponent.h"
 #include "Net/UnrealNetwork.h"
+#include "PhysicsEngine/RadialForceComponent.h"
+#include "PhysicsEngine/RadialForceActor.h"
 
 //////////////////////////////////////////////////////////////////////////
 // RADIALFORCECOMPONENT
diff --git a/Engine/Source/Runtime/Engine/Private/PlayerStartPIE.cpp b/Engine/Source/Runtime/Engine/Private/PlayerStartPIE.cpp
--- a/Engine/Source/Runtime/Engine/Private/PlayerStartPIE.cpp
+++ b/Engine/Source/Runtime/Engine/Private/PlayerStartPIE.cpp
@@ -1,6 +1,7 @@
 // Copyright 1998-2014 Epic Games, Inc. All Rights Reserved.
 
 #include "EnginePrivate.h"
+#include "Engine/PlayerStartPIE.h"
 
 APlayerStartPIE::APlayerStartPIE(const FObjectInitializer& ObjectInitializer)
 	: Super(ObjectInitializer)
diff --git a/Engine/Source/Runtime/Engine/Private/PlayerState.cpp b/Engine/Source/Runtime/Engine/Private/PlayerState.cpp
--- a/Engine/Source/Runtime/Engine/Private/PlayerState.cpp
+++ b/Engine/Source/Runtime/Engine/Private/PlayerState.cpp
@@ -5,6 +5,7 @@
 =============================================================================*/
 
 #include "EnginePrivate.h"
+#include "GameFramework/EngineMessage.h"
 #include "Net/UnrealNetwork.h"
 #include "OnlineSubsystemUtils.h"
 
diff --git a/Engine/Source/Runtime/Engine/Private/PoseableMeshComponent.cpp b/Engine/Source/Runtime/Engine/Private/PoseableMeshComponent.cpp
--- a/Engine/Source/Runtime/Engine/Private/PoseableMeshComponent.cpp
+++ b/Engine/Source/Runtime/Engine/Private/PoseableMeshComponent.cpp
@@ -7,6 +7,7 @@
 #include "EnginePrivate.h"
 #include "AnimTree.h"
 #include "Animation/AnimInstance.h"
+#include "Components/PoseableMeshComponent.h"
 
 UPoseableMeshComponent::UPoseableMeshComponent(const FObjectInitializer& ObjectInitializer)
 	: Super(ObjectInitializer)
diff --git a/Engine/Source/Runtime/Engine/Private/PostProcessComponent.cpp b/Engine/Source/Runtime/Engine/Private/PostProcessComponent.cpp
--- a/Engine/Source/Runtime/Engine/Private/PostProcessComponent.cpp
+++ b/Engine/Source/Runtime/Engine/Private/PostProcessComponent.cpp
@@ -1,6 +1,7 @@
 // Copyright 1998-2013 Epic Games, Inc. All Rights Reserved.
 
 #include "EnginePrivate.h"
+#include "Components/PostProcessComponent.h"
 
 UPostProcessComponent::UPostProcessComponent(const FObjectInitializer& ObjectInitializer)
 	: Super(ObjectInitializer)
diff --git a/Engine/Source/Runtime/Engine/Private/ReflectionCaptureComponent.cpp b/Engine/Source/Runtime/Engine/Private/ReflectionCaptureComponent.cpp
--- a/Engine/Source/Runtime/Engine/Private/ReflectionCaptureComponent.cpp
+++ b/Engine/Source/Runtime/Engine/Private/ReflectionCaptureComponent.cpp
@@ -5,11 +5,19 @@
 =============================================================================*/
 
 #include "EnginePrivate.h"
+#include "Engine/ReflectionCapture.h"
 #include "DerivedDataCacheInterface.h"
 #include "TargetPlatform.h"
 #include "EngineModule.h"
 #include "RendererInterface.h"
 #include "ShaderCompiler.h"
+#include "Engine/SphereReflectionCapture.h"
+#include "Components/SphereReflectionCaptureComponent.h"
+#include "Components/DrawSphereComponent.h"
+#include "Components/BoxReflectionCaptureComponent.h"
+#include "Engine/PlaneReflectionCapture.h"
+#include "Engine/BoxReflectionCapture.h"
+#include "Components/PlaneReflectionCaptureComponent.h"
 
 /** 
  * Size of all reflection captures.
diff --git a/Engine/Source/Runtime/Engine/Private/SceneCaptureComponent.cpp b/Engine/Source/Runtime/Engine/Private/SceneCaptureComponent.cpp
--- a/Engine/Source/Runtime/Engine/Private/SceneCaptureComponent.cpp
+++ b/Engine/Source/Runtime/Engine/Private/SceneCaptureComponent.cpp
@@ -6,6 +6,11 @@
 
 #include "EnginePrivate.h"
 #include "../../Renderer/Private/ScenePrivate.h"
+#include "Engine/SceneCapture.h"
+#include "Engine/SceneCapture2D.h"
+#include "Components/SceneCaptureComponent2D.h"
+#include "Engine/SceneCaptureCube.h"
+#include "Components/SceneCaptureComponentCube.h"
 
 ASceneCapture::ASceneCapture(const FObjectInitializer& ObjectInitializer)
 	: Super(ObjectInitializer)
diff --git a/Engine/Source/Runtime/Engine/Private/SceneView.cpp b/Engine/Source/Runtime/Engine/Private/SceneView.cpp
--- a/Engine/Source/Runtime/Engine/Private/SceneView.cpp
+++ b/Engine/Source/Runtime/Engine/Private/SceneView.cpp
@@ -13,6 +13,7 @@
 #include "Slate/SceneViewport.h"
 #include "RendererInterface.h"
 #include "BufferVisualizationData.h"
+#include "Interfaces/Interface_PostProcessVolume.h"
 
 DEFINE_LOG_CATEGORY(LogBufferVisualization);
 
diff --git a/Engine/Source/Runtime/Engine/Private/ScriptPlatformInterface.cpp b/Engine/Source/Runtime/Engine/Private/ScriptPlatformInterface.cpp
--- a/Engine/Source/Runtime/Engine/Private/ScriptPlatformInterface.cpp
+++ b/Engine/Source/Runtime/Engine/Private/ScriptPlatformInterface.cpp
@@ -5,6 +5,10 @@
 =============================================================================*/ 
 
 #include "EnginePrivate.h"
+#include "Engine/CloudStorageBase.h"
+#include "Engine/MicroTransactionBase.h"
+#include "Engine/PlatformInterfaceWebResponse.h"
+#include "Engine/TwitterIntegrationBase.h"
 
 DEFINE_LOG_CATEGORY_STATIC(LogScriptPlatformInterface, Log, All);
 
diff --git a/Engine/Source/Runtime/Engine/Private/SpotLight.cpp b/Engine/Source/Runtime/Engine/Private/SpotLight.cpp
--- a/Engine/Source/Runtime/Engine/Private/SpotLight.cpp
+++ b/Engine/Source/Runtime/Engine/Private/SpotLight.cpp
@@ -1,6 +1,7 @@
 // Copyright 1998-2014 Epic Games, Inc. All Rights Reserved.
 
 #include "EnginePrivate.h"
+#include "Engine/SpotLight.h"
 
 ASpotLight::ASpotLight(const FObjectInitializer& ObjectInitializer)
 	: Super(ObjectInitializer.SetDefaultSubobjectClass<USpotLightComponent>(TEXT("LightComponent0")))
diff --git a/Engine/Source/Runtime/Engine/Private/StaticMesh.cpp b/Engine/Source/Runtime/Engine/Private/StaticMesh.cpp
--- a/Engine/Source/Runtime/Engine/Private/StaticMesh.cpp
+++ b/Engine/Source/Runtime/Engine/Private/StaticMesh.cpp
@@ -5,6 +5,7 @@
 =============================================================================*/
 
 #include "EnginePrivate.h"
+#include "Engine/AssetUserData.h"
 #include "StaticMeshResources.h"
 #include "MeshBuild.h"
 #include "GenericOctree.h"
diff --git a/Engine/Source/Runtime/Engine/Private/StreamableManager.cpp b/Engine/Source/Runtime/Engine/Private/StreamableManager.cpp
--- a/Engine/Source/Runtime/Engine/Private/StreamableManager.cpp
+++ b/Engine/Source/Runtime/Engine/Private/StreamableManager.cpp
@@ -1,6 +1,7 @@
 // Copyright 1998-2014 Epic Games, Inc. All Rights Reserved.
 
 #include "EnginePrivate.h"
+#include "Engine/StreamableManager.h"
 
 DEFINE_LOG_CATEGORY_STATIC(LogStreamableManager, Log, All);
 
diff --git a/Engine/Source/Runtime/Engine/Private/TargetPoint.cpp b/Engine/Source/Runtime/Engine/Private/TargetPoint.cpp
--- a/Engine/Source/Runtime/Engine/Private/TargetPoint.cpp
+++ b/Engine/Source/Runtime/Engine/Private/TargetPoint.cpp
@@ -1,6 +1,7 @@
 // Copyright 1998-2014 Epic Games, Inc. All Rights Reserved.
 
 #include "EnginePrivate.h"
+#include "Engine/TargetPoint.h"
 
 ATargetPoint::ATargetPoint(const FObjectInitializer& ObjectInitializer)
 	: Super(ObjectInitializer)
diff --git a/Engine/Source/Runtime/Engine/Private/Tests/Internationalization/TextPropertyTest.cpp b/Engine/Source/Runtime/Engine/Private/Tests/Internationalization/TextPropertyTest.cpp
--- a/Engine/Source/Runtime/Engine/Private/Tests/Internationalization/TextPropertyTest.cpp
+++ b/Engine/Source/Runtime/Engine/Private/Tests/Internationalization/TextPropertyTest.cpp
@@ -1,7 +1,8 @@
 // Copyright 1998-2014 Epic Games, Inc. All Rights Reserved.
 
 #include "EnginePrivate.h"
-#include "AutomationTest.h"
+#include "Misc/AutomationTest.h"
+#include "Tests/TextPropertyTestObject.h"
 
 #define LOCTEXT_NAMESPACE "TextPropertyTest"
 
diff --git a/Engine/Source/Runtime/Engine/Private/TextRenderComponent.cpp b/Engine/Source/Runtime/Engine/Private/TextRenderComponent.cpp
--- a/Engine/Source/Runtime/Engine/Private/TextRenderComponent.cpp
+++ b/Engine/Source/Runtime/Engine/Private/TextRenderComponent.cpp
@@ -1,6 +1,7 @@
 // Copyright 1998-2014 Epic Games, Inc. All Rights Reserved.
 
 #include "EnginePrivate.h"
+#include "Components/TextRenderComponent.h"
 #include "DynamicMeshBuilder.h"
 #include "Engine/Font.h"
 #include "Engine/TextRenderActor.h"
diff --git a/Engine/Source/Runtime/Engine/Private/Texture2DDynamic.cpp b/Engine/Source/Runtime/Engine/Private/Texture2DDynamic.cpp
--- a/Engine/Source/Runtime/Engine/Private/Texture2DDynamic.cpp
+++ b/Engine/Source/Runtime/Engine/Private/Texture2DDynamic.cpp
@@ -5,7 +5,7 @@
 =============================================================================*/
 
 #include "EnginePrivate.h"
-
+#include "Engine/Texture2DDynamic.h"
 
 /*-----------------------------------------------------------------------------
 	FTexture2DDynamicResource
diff --git a/Engine/Source/Runtime/Engine/Private/TextureRenderTargetCube.cpp b/Engine/Source/Runtime/Engine/Private/TextureRenderTargetCube.cpp
--- a/Engine/Source/Runtime/Engine/Private/TextureRenderTargetCube.cpp
+++ b/Engine/Source/Runtime/Engine/Private/TextureRenderTargetCube.cpp
@@ -5,6 +5,7 @@
 =============================================================================*/
 
 #include "EnginePrivate.h"
+#include "Engine/TextureRenderTargetCube.h"
 #include "DDSLoader.h"
 
 /*-----------------------------------------------------------------------------
diff --git a/Engine/Source/Runtime/Engine/Private/UnrealClient.cpp b/Engine/Source/Runtime/Engine/Private/UnrealClient.cpp
--- a/Engine/Source/Runtime/Engine/Private/UnrealClient.cpp
+++ b/Engine/Source/Runtime/Engine/Private/UnrealClient.cpp
@@ -2,6 +2,7 @@
 
 
 #include "EnginePrivate.h"
+#include "Components/PostProcessComponent.h"
 #include "Engine/Font.h"
 #include "Matinee/MatineeActor.h"
 #include "EditorSupportDelegates.h"
diff --git a/Engine/Source/Runtime/Engine/Private/UnrealEngine.cpp b/Engine/Source/Runtime/Engine/Private/UnrealEngine.cpp
--- a/Engine/Source/Runtime/Engine/Private/UnrealEngine.cpp
+++ b/Engine/Source/Runtime/Engine/Private/UnrealEngine.cpp
@@ -6,6 +6,7 @@
 
 #include "EnginePrivate.h"
 #include "Engine/Font.h"
+#include "Engine/LevelStreamingPersistent.h"
 #include "Engine/ObjectReferencer.h"
 #include "Net/UnrealNetwork.h"
 #include "Engine/Console.h"
diff --git a/Engine/Source/Runtime/Engine/Private/UserInterfaceSettings.cpp b/Engine/Source/Runtime/Engine/Private/UserInterfaceSettings.cpp
--- a/Engine/Source/Runtime/Engine/Private/UserInterfaceSettings.cpp
+++ b/Engine/Source/Runtime/Engine/Private/UserInterfaceSettings.cpp
@@ -5,6 +5,7 @@
 =============================================================================*/
 
 #include "EnginePrivate.h"
+#include "Engine/RendererSettings.h"
 #include "Engine/UserInterfaceSettings.h"
 
 
diff --git a/Engine/Source/Runtime/Engine/Private/WindDirectionalSource.cpp b/Engine/Source/Runtime/Engine/Private/WindDirectionalSource.cpp
--- a/Engine/Source/Runtime/Engine/Private/WindDirectionalSource.cpp
+++ b/Engine/Source/Runtime/Engine/Private/WindDirectionalSource.cpp
@@ -1,6 +1,7 @@
 // Copyright 1998-2014 Epic Games, Inc. All Rights Reserved.
 
 #include "EnginePrivate.h"
+#include "Components/WindDirectionalSourceComponent.h"
 #include "Engine/WindDirectionalSource.h"
 
 AWindDirectionalSource::AWindDirectionalSource(const FObjectInitializer& ObjectInitializer)
diff --git a/Engine/Source/Runtime/Engine/Private/World.cpp b/Engine/Source/Runtime/Engine/Private/World.cpp
--- a/Engine/Source/Runtime/Engine/Private/World.cpp
+++ b/Engine/Source/Runtime/Engine/Private/World.cpp
@@ -37,6 +37,7 @@
 #include "SceneInterface.h"
 #include "Camera/CameraActor.h"
 #include "Engine/DemoNetDriver.h"
+#include "Layers/Layer.h"
 
 #if WITH_EDITOR
 	#include "Editor/UnrealEd/Public/Kismet2/KismetEditorUtilities.h"
@@ -54,6 +55,7 @@
 #include "RendererInterface.h"
 #include "DataChannel.h"
 #include "ShaderCompiler.h"
+#include "Engine/LevelStreamingPersistent.h"
 
 DEFINE_LOG_CATEGORY_STATIC(LogWorld, Log, All);
 DEFINE_LOG_CATEGORY(LogSpawn);
diff --git a/Engine/Source/Runtime/Engine/Private/WorldSettings.cpp b/Engine/Source/Runtime/Engine/Private/WorldSettings.cpp
--- a/Engine/Source/Runtime/Engine/Private/WorldSettings.cpp
+++ b/Engine/Source/Runtime/Engine/Private/WorldSettings.cpp
@@ -1,6 +1,7 @@
 // Copyright 1998-2014 Epic Games, Inc. All Rights Reserved.
 
 #include "EnginePrivate.h"
+#include "Engine/AssetUserData.h"
 #include "Engine/WorldComposition.h"
 #include "Net/UnrealNetwork.h"
 #include "GameFramework/GameNetworkManager.h"
diff --git a/Engine/Source/Runtime/Renderer/Private/ScenePrivate.h b/Engine/Source/Runtime/Renderer/Private/ScenePrivate.h
--- a/Engine/Source/Runtime/Renderer/Private/ScenePrivate.h
+++ b/Engine/Source/Runtime/Renderer/Private/ScenePrivate.h
@@ -8,6 +8,7 @@
 #define __SCENEPRIVATE_H__
 
 class SceneRenderingAllocator;
+class USceneCaptureComponent;
 
 class SceneRenderingBitArrayAllocator
 	: public TInlineAllocator<4,SceneRenderingAllocator>
