Forward LifeTime parameter to LineBatcher when drawing Cylinders

#github https://github.com/EpicGames/UnrealEngine/pull/527 KazumasaOhashi

#branch UE4

diff --git a/Engine/Source/Runtime/Engine/Private/DrawDebugHelpers.cpp b/Engine/Source/Runtime/Engine/Private/DrawDebugHelpers.cpp
--- a/Engine/Source/Runtime/Engine/Private/DrawDebugHelpers.cpp
+++ b/Engine/Source/Runtime/Engine/Private/DrawDebugHelpers.cpp
@@ -459,15 +459,16 @@ void DrawDebugCylinder(const UWorld* InWorld, FVector const& Start, FVector cons
 		ULineBatchComponent* const LineBatcher = GetDebugLineBatcher( InWorld, bPersistentLines, LifeTime, (DepthPriority == SDPG_Foreground) );
 		if(LineBatcher != NULL)
 		{
+			const float LineLifeTime = (LifeTime > 0.f) ? LifeTime : LineBatcher->DefaultLifeTime;
 			while( Segments-- )
 			{
 				Segment = Perpendicular.RotateAngleAxis(Angle, Axis) * Radius;
 				P2 = Segment + Start;
 				P4 = Segment + End;
 
-				LineBatcher->DrawLine(P2, P4, Color, DepthPriority);
-				LineBatcher->DrawLine(P1, P2, Color, DepthPriority);
-				LineBatcher->DrawLine(P3, P4, Color, DepthPriority);
+				LineBatcher->DrawLine(P2, P4, Color, DepthPriority, 0.f, LineLifeTime);
+				LineBatcher->DrawLine(P1, P2, Color, DepthPriority, 0.f, LineLifeTime);
+				LineBatcher->DrawLine(P3, P4, Color, DepthPriority, 0.f, LineLifeTime);
 
 				P1 = P2;
 				P3 = P4;
