CIS fix.

diff --git a/Engine/Source/Runtime/Engine/Private/AI/Navigation/RecastNavMesh.cpp b/Engine/Source/Runtime/Engine/Private/AI/Navigation/RecastNavMesh.cpp
--- a/Engine/Source/Runtime/Engine/Private/AI/Navigation/RecastNavMesh.cpp
+++ b/Engine/Source/Runtime/Engine/Private/AI/Navigation/RecastNavMesh.cpp
@@ -20,6 +20,7 @@
 #if WITH_EDITOR
 #include "UnrealEd.h"
 #endif
+#include "AI/Navigation/NavMeshRenderingComponent.h"
 
 #if WITH_RECAST
 /// Helper for accessing navigation query from different threads
diff --git a/Engine/Source/Runtime/Engine/Private/GameMode.cpp b/Engine/Source/Runtime/Engine/Private/GameMode.cpp
--- a/Engine/Source/Runtime/Engine/Private/GameMode.cpp
+++ b/Engine/Source/Runtime/Engine/Private/GameMode.cpp
@@ -12,6 +12,7 @@
 #include "GameFramework/HUD.h"
 #include "GameFramework/DefaultPawn.h"
 #include "GameFramework/SpectatorPawn.h"
+#include "Engine/PlayerStartPIE.h"
 
 DEFINE_LOG_CATEGORY_STATIC(LogGameMode, Log, All);
 
diff --git a/Engine/Source/Runtime/Engine/Private/Materials/MaterialExpressions.cpp b/Engine/Source/Runtime/Engine/Private/Materials/MaterialExpressions.cpp
--- a/Engine/Source/Runtime/Engine/Private/Materials/MaterialExpressions.cpp
+++ b/Engine/Source/Runtime/Engine/Private/Materials/MaterialExpressions.cpp
@@ -138,6 +138,8 @@
 #endif //WITH_EDITOR
 
 #include "Engine/Font.h"
+#include "Engine/Texture2DDynamic.h"
+#include "Engine/TextureRenderTargetCube.h"
 
 #define LOCTEXT_NAMESPACE "MaterialExpression"
 
diff --git a/Engine/Source/Runtime/Engine/Private/SoundNodeModulatorContinuous.cpp b/Engine/Source/Runtime/Engine/Private/SoundNodeModulatorContinuous.cpp
--- a/Engine/Source/Runtime/Engine/Private/SoundNodeModulatorContinuous.cpp
+++ b/Engine/Source/Runtime/Engine/Private/SoundNodeModulatorContinuous.cpp
@@ -5,6 +5,7 @@
 #include "SoundDefinitions.h"
 #include "DistributionHelpers.h"
 #include "Sound/SoundNodeModulatorContinuous.h"
+#include "Distributions/DistributionFloatSoundParameter.h"
 
 float FModulatorContinuousParams::GetValue(const FActiveSound& ActiveSound) const
 {
