CIS fix.

diff --git a/Engine/Source/Runtime/Engine/Private/AI/Navigation/NavLinkProxy.cpp b/Engine/Source/Runtime/Engine/Private/AI/Navigation/NavLinkProxy.cpp
--- a/Engine/Source/Runtime/Engine/Private/AI/Navigation/NavLinkProxy.cpp
+++ b/Engine/Source/Runtime/Engine/Private/AI/Navigation/NavLinkProxy.cpp
@@ -8,6 +8,7 @@
 #include "AI/Navigation/NavLinkCustomInterface.h"
 #include "AI/Navigation/NavLinkCustomComponent.h"
 #include "AI/Navigation/NavLinkProxy.h"
+#include "AI/Navigation/NavLinkRenderingComponent.h"
 
 ANavLinkProxy::ANavLinkProxy(const FObjectInitializer& ObjectInitializer) : Super(ObjectInitializer)
 {
