CIS fix.

diff --git a/Engine/Source/Runtime/Engine/Private/AI/Navigation/NavigationTestingActor.cpp b/Engine/Source/Runtime/Engine/Private/AI/Navigation/NavigationTestingActor.cpp
--- a/Engine/Source/Runtime/Engine/Private/AI/Navigation/NavigationTestingActor.cpp
+++ b/Engine/Source/Runtime/Engine/Private/AI/Navigation/NavigationTestingActor.cpp
@@ -5,6 +5,7 @@
 #include "ObjectEditorUtils.h"
 #endif
 #include "AI/Navigation/NavigationTestingActor.h"
+#include "AI/Navigation/NavTestRenderingComponent.h"
 
 void FNavTestTickHelper::Tick(float DeltaTime)
 {
diff --git a/Engine/Source/Runtime/Engine/Private/Animation/SmartName.cpp b/Engine/Source/Runtime/Engine/Private/Animation/SmartName.cpp
--- a/Engine/Source/Runtime/Engine/Private/Animation/SmartName.cpp
+++ b/Engine/Source/Runtime/Engine/Private/Animation/SmartName.cpp
@@ -1,6 +1,7 @@
 // Copyright 1998-2014 Epic Games, Inc. All Rights Reserved.
 
 #include "EnginePrivate.h"
+#include "Animation/SmartName.h"
 
 FSmartNameMapping::FSmartNameMapping()
 : NextUid(0)
diff --git a/Engine/Source/Runtime/Engine/Private/GameMode.cpp b/Engine/Source/Runtime/Engine/Private/GameMode.cpp
--- a/Engine/Source/Runtime/Engine/Private/GameMode.cpp
+++ b/Engine/Source/Runtime/Engine/Private/GameMode.cpp
@@ -13,6 +13,7 @@
 #include "GameFramework/DefaultPawn.h"
 #include "GameFramework/SpectatorPawn.h"
 #include "Engine/PlayerStartPIE.h"
+#include "GameFramework/EngineMessage.h"
 
 DEFINE_LOG_CATEGORY_STATIC(LogGameMode, Log, All);
 
