Fix "Warning AttachTo: '/Game/Maps/UEDPIE_0_Landscapes.Landscapes:PersistentLevel.InstancedFoliageActor_1.RootComponent0' is not static , cannot attach '/Game/Maps/UEDPIE_0_Landscapes.Landscapes:PersistentLevel.InstancedFoliageActor_1.InstancedStaticMeshComponent_6' which is static to it. Aborting."

diff --git a/Engine/Source/Runtime/Engine/Private/InstancedFoliage.cpp b/Engine/Source/Runtime/Engine/Private/InstancedFoliage.cpp
--- a/Engine/Source/Runtime/Engine/Private/InstancedFoliage.cpp
+++ b/Engine/Source/Runtime/Engine/Private/InstancedFoliage.cpp
@@ -721,7 +721,7 @@ AInstancedFoliageActor::AInstancedFoliageActor(const FObjectInitializer& ObjectI
 {
 	USceneComponent* SceneComponent = ObjectInitializer.CreateDefaultSubobject<USceneComponent>(this, TEXT("RootComponent0"));
 	RootComponent = SceneComponent;
-	RootComponent->Mobility = EComponentMobility::Movable;
+	RootComponent->Mobility = EComponentMobility::Static;
 	
 	SetActorEnableCollision(true);
 #if WITH_EDITORONLY_DATA
