Allow non static trimesh objects to participate in simulation. This helps with the case where some seemingly static mesh is made to be mobile (but still doesn't really move a lot) and suddenly there's no collision.

JIRA UE-3240

diff --git a/Engine/Source/Runtime/Engine/Private/PhysicsEngine/BodyInstance.cpp b/Engine/Source/Runtime/Engine/Private/PhysicsEngine/BodyInstance.cpp
--- a/Engine/Source/Runtime/Engine/Private/PhysicsEngine/BodyInstance.cpp
+++ b/Engine/Source/Runtime/Engine/Private/PhysicsEngine/BodyInstance.cpp
@@ -628,12 +628,6 @@ void FBodyInstance::UpdatePhysicsShapeFilterData(uint32 SkelMeshCompID, bool bUs
 		int32 NumSyncShapes = 0;
 		TArray<PxShape*> AllShapes = GetAllShapes(NumSyncShapes);
 
-		// Is the target a static actor
-		const bool bDestStatic = PActor->isRigidStatic() != NULL;
-
-		// Only perform scene queries in the synchronous scene for static shapes
-		const int32 SceneQueryShapeNumMax = bDestStatic ? NumSyncShapes : AllShapes.Num();
-
 		PxScene* AsyncScene = (RigidActorAsync != NULL) ? RigidActorAsync->getScene() : NULL;
 		SCENE_LOCK_WRITE(AsyncScene);
 
@@ -701,7 +695,7 @@ void FBodyInstance::UpdatePhysicsShapeFilterData(uint32 SkelMeshCompID, bool bUs
 					PShape->setQueryFilterData(PComplexQueryData);
 
 					// on dynamic objects and objects which don't use complex as simple, tri mesh not used for sim
-					if (!bSimCollision || !bDestStatic || !bUseComplexAsSimple)
+					if (!bSimCollision || !bUseComplexAsSimple)
 					{
 						PShape->setFlag(PxShapeFlag::eSIMULATION_SHAPE, false);
 					}
