Better looping interpolation - this should fix it.

diff --git a/Engine/Source/Runtime/Engine/Private/Animation/AnimSequence.cpp b/Engine/Source/Runtime/Engine/Private/Animation/AnimSequence.cpp
--- a/Engine/Source/Runtime/Engine/Private/Animation/AnimSequence.cpp
+++ b/Engine/Source/Runtime/Engine/Private/Animation/AnimSequence.cpp
@@ -2407,61 +2407,39 @@ bool UAnimSequence::AddLoopingInterpolation()
 {
 	int32 NumTracks = AnimationTrackNames.Num();
 	int32 NumKeys = NumFrames;
-	int32 SrcKeyIndex = 0; // this is the key index we copy animation from
-	int32 DestKeyIndex = NumFrames; // this is the key index we insert animation to
 	float Interval = (NumKeys)? SequenceLength/NumKeys : 0.f;
 
 	if(Interval > 0.f)
 	{
-		// 2d array of animated time [track index][time key]
-		// the reason I have to query the time is because 
-		// the RawAnimation doesn't have all keys
-		// it removes duplicated keys if not necessary. 
-		TArray< TArray<FTransform> > AnimatedLocalSpaces;
-		AnimatedLocalSpaces.AddZeroed(NumTracks);
-
-		for(int32 TrackIndex=0; TrackIndex<NumTracks; ++TrackIndex)
-		{
-			// add one extra for looping interpolation
-			AnimatedLocalSpaces[TrackIndex].AddUninitialized(NumKeys + 1);
-
-			// fill up keys
-			for(int32 Key=0; Key<NumKeys; ++Key)
-			{
-				GetBoneTransform(AnimatedLocalSpaces[TrackIndex][Key], TrackIndex, Interval*Key, false);
-			}
-
-			AnimatedLocalSpaces[TrackIndex][DestKeyIndex] = AnimatedLocalSpaces[TrackIndex][SrcKeyIndex];
-		}
-
 		// added one more key
 		int32 NewNumKeys = NumKeys +1 ;
 
 		// now I need to calculate back to new animation data
-		TArray<struct FRawAnimSequenceTrack> NewRawAnimationData = RawAnimationData;
 		for(int32 TrackIndex=0; TrackIndex<NumTracks; ++TrackIndex)
 		{
-			auto & RawAnimation = NewRawAnimationData[TrackIndex];
-			RawAnimation.PosKeys.Empty(NewNumKeys);
-			RawAnimation.PosKeys.AddUninitialized(NewNumKeys);
-			RawAnimation.RotKeys.Empty(NewNumKeys);
-			RawAnimation.RotKeys.AddUninitialized(NewNumKeys);
-			RawAnimation.ScaleKeys.Empty(NewNumKeys);
-			RawAnimation.ScaleKeys.AddUninitialized(NewNumKeys);
+			auto & RawAnimation = RawAnimationData[TrackIndex];
+			if (RawAnimation.PosKeys.Num() > 1)
+			{
+				FVector FirstKey = RawAnimation.PosKeys[0];
+				RawAnimation.PosKeys.Add(FirstKey);
+			}
 
-			for(int32 Key=0; Key<NewNumKeys; ++Key)
+			if(RawAnimation.RotKeys.Num() > 1)
 			{
-				RawAnimation.PosKeys[Key] = AnimatedLocalSpaces[TrackIndex][Key].GetLocation();
-				RawAnimation.RotKeys[Key] = AnimatedLocalSpaces[TrackIndex][Key].GetRotation();
-				RawAnimation.ScaleKeys[Key] = AnimatedLocalSpaces[TrackIndex][Key].GetScale3D();
+				FQuat FirstKey = RawAnimation.RotKeys[0];
+				RawAnimation.RotKeys.Add(FirstKey);
+			}
+
+			if(RawAnimation.ScaleKeys.Num() > 1)
+			{
+				FVector FirstKey = RawAnimation.ScaleKeys[0];
+				RawAnimation.ScaleKeys.Add(FirstKey);
 			}
 		}
 
-		RawAnimationData = NewRawAnimationData;
+		SequenceLength += Interval;
 		NumFrames = NewNumKeys;
 
-		// we don't increase sequence length as the original sequence length should stay
-
 		PostProcessSequence();
 		return true;
 	}
