Fix up UTextureDynamic2D creation

diff --git a/Engine/Source/Runtime/Engine/Classes/Engine/Texture2DDynamic.h b/Engine/Source/Runtime/Engine/Classes/Engine/Texture2DDynamic.h
--- a/Engine/Source/Runtime/Engine/Classes/Engine/Texture2DDynamic.h
+++ b/Engine/Source/Runtime/Engine/Classes/Engine/Texture2DDynamic.h
@@ -45,5 +45,5 @@ public:
 	ENGINE_API void Init(int32 InSizeX, int32 InSizeY, EPixelFormat InFormat = PF_B8G8R8A8, bool InIsResolveTarget = false);
 	
 	/** Creates and initializes a new Texture2DDynamic with the requested settings */
-	class UTexture2DDynamic* Create(int32 InSizeX, int32 InSizeY, EPixelFormat InFormat = PF_B8G8R8A8, bool InIsResolveTarget = false);
+	ENGINE_API static UTexture2DDynamic* Create(int32 InSizeX, int32 InSizeY, EPixelFormat InFormat = PF_B8G8R8A8, bool InIsResolveTarget = false);
 };
diff --git a/Engine/Source/Runtime/Engine/Private/Texture2DDynamic.cpp b/Engine/Source/Runtime/Engine/Private/Texture2DDynamic.cpp
--- a/Engine/Source/Runtime/Engine/Private/Texture2DDynamic.cpp
+++ b/Engine/Source/Runtime/Engine/Private/Texture2DDynamic.cpp
@@ -120,7 +120,8 @@ UTexture2DDynamic* UTexture2DDynamic::Create(int32 InSizeX, int32 InSizeY, EPixe
 	EPixelFormat DesiredFormat = EPixelFormat(InFormat);
 	if (InSizeX > 0 && InSizeY > 0 )
 	{
-		UTexture2DDynamic* NewTexture = Cast<UTexture2DDynamic>(StaticConstructObject(GetClass(), GetTransientPackage(), NAME_None, RF_Transient));
+		
+		UTexture2DDynamic* NewTexture = Cast<UTexture2DDynamic>(StaticConstructObject(UTexture2DDynamic::StaticClass(), GetTransientPackage(), NAME_None, RF_Transient));
 		if (NewTexture != NULL)
 		{
 			// Disable compression
