reafactor, for easier debugging

diff --git a/Engine/Source/Runtime/Engine/Private/Materials/HLSLMaterialTranslator.h b/Engine/Source/Runtime/Engine/Private/Materials/HLSLMaterialTranslator.h
--- a/Engine/Source/Runtime/Engine/Private/Materials/HLSLMaterialTranslator.h
+++ b/Engine/Source/Runtime/Engine/Private/Materials/HLSLMaterialTranslator.h
@@ -1139,7 +1139,10 @@ protected:
 	FMaterialUniformExpression* GetParameterUniformExpression(int32 Index) const
 	{
 		check(Index >= 0 && Index < CodeChunks[MaterialProperty][ShaderFrequency].Num());
-		return CodeChunks[MaterialProperty][ShaderFrequency][Index].UniformExpression;
+
+		const FShaderCodeChunk& Chunk = CodeChunks[MaterialProperty][ShaderFrequency][Index];
+
+		return Chunk.UniformExpression;
 	}
 
 	// GetArithmeticResultType
