Now using the proper outer when destroying child actors from a ChildActor component.

diff --git a/Engine/Source/Runtime/Engine/Private/ChildActorComponent.cpp b/Engine/Source/Runtime/Engine/Private/ChildActorComponent.cpp
--- a/Engine/Source/Runtime/Engine/Private/ChildActorComponent.cpp
+++ b/Engine/Source/Runtime/Engine/Private/ChildActorComponent.cpp
@@ -211,7 +211,7 @@ void UChildActorComponent::DestroyChildActor()
 				ChildClass->ClassUnique = FMath::Max(ChildClass->ClassUnique, ChildActor->GetFName().GetNumber());
 
 				const FString ObjectBaseName = FString::Printf(TEXT("DESTROYED_%s_CHILDACTOR"), *ChildClass->GetName());
-				ChildActor->Rename(*MakeUniqueObjectName(GetOuter(), ChildClass, *ObjectBaseName).ToString(), NULL, REN_DoNotDirty);
+				ChildActor->Rename(*MakeUniqueObjectName(ChildActor->GetOuter(), ChildClass, *ObjectBaseName).ToString(), NULL, REN_DoNotDirty);
 				World->DestroyActor(ChildActor);
 			}
 		}
