CIS fix

diff --git a/Engine/Source/Runtime/Engine/Private/UserInterface/CanvasItem.cpp b/Engine/Source/Runtime/Engine/Private/UserInterface/CanvasItem.cpp
--- a/Engine/Source/Runtime/Engine/Private/UserInterface/CanvasItem.cpp
+++ b/Engine/Source/Runtime/Engine/Private/UserInterface/CanvasItem.cpp
@@ -9,6 +9,7 @@
 #include "TileRendering.h"
 #include "RHIStaticStates.h"
 #include "Internationalization.h"
+#include "SlateBasics.h"
 
 
 DECLARE_CYCLE_STAT(TEXT("CanvasTileTextureItem Time"),STAT_Canvas_TileTextureItemTime,STATGROUP_Canvas);
