Niagara: Effect serialization, actor factory, some updates to encapsulate serializeable data from the emitter and separate it from runtime data. First pass only, not fully functional yet.

diff --git a/Engine/Source/Editor/NiagaraEditor/Private/SNiagaraEffectEditorWidget.h b/Engine/Source/Editor/NiagaraEditor/Private/SNiagaraEffectEditorWidget.h
--- a/Engine/Source/Editor/NiagaraEditor/Private/SNiagaraEffectEditorWidget.h
+++ b/Engine/Source/Editor/NiagaraEditor/Private/SNiagaraEffectEditorWidget.h
@@ -51,7 +51,7 @@ public:
 
 		// Set the object found from the asset picker
 		CurUpdateScript = Cast<UNiagaraScript>(Asset);
-		Emitter->SetUpdateScript(CurUpdateScript);
+		Emitter->GetProperties()->UpdateScript = CurUpdateScript;
 	}
 
 	void OnSpawnScriptSelectedFromPicker(UObject *Asset)
@@ -61,7 +61,7 @@ public:
 
 		// Set the object found from the asset picker
 		CurSpawnScript = Cast<UNiagaraScript>(Asset);
-		Emitter->SetSpawnScript(CurSpawnScript);
+		Emitter->GetProperties()->SpawnScript = CurSpawnScript;
 	}
 
 	void OnMaterialSelected(UObject *Asset)
@@ -72,27 +72,28 @@ public:
 		ThumbnailPool->Tick(0);
 		CurMaterial = Cast<UMaterial>(Asset);
 		Emitter->GetEffectRenderer()->SetMaterial(CurMaterial, ERHIFeatureLevel::SM5);
+		Emitter->GetProperties()->Material = CurMaterial;
 	}
 
 	UObject *GetMaterial() const
 	{
-		return CurMaterial;
+		return Emitter->GetProperties()->Material;
 	}
 
 	UObject *GetUpdateScript() const
 	{
-		return CurUpdateScript;
+		return Emitter->GetProperties()->UpdateScript;
 	}
 
 	UObject *GetSpawnScript() const
 	{
-		return CurSpawnScript;
+		return Emitter->GetProperties()->SpawnScript;
 	}
 
 	void OnEmitterEnabledChanged(ESlateCheckBoxState::Type NewCheckedState)
 	{
 		const bool bNewEnabledState = (NewCheckedState == ESlateCheckBoxState::Checked);
-		Emitter->SetEnabled(bNewEnabledState);
+		Emitter->GetProperties()->bIsEnabled = bNewEnabledState;
 	}
 
 	ESlateCheckBoxState::Type IsEmitterEnabled() const
@@ -120,6 +121,7 @@ public:
 			EEmitterRenderModuleType, InType, Type,
 		{
 			InEmitter->SetRenderModuleType(InType, InEffect->GetComponent()->GetWorld()->FeatureLevel);
+			InEmitter->GetProperties()->RenderModuleType = InType;
 			InEffect->RenderModuleupdate();
 		});
 		TComponentReregisterContext<UNiagaraComponent> ComponentReregisterContext;
@@ -127,25 +129,27 @@ public:
 
 	void OnSpawnRateChanged(const FText &InText)
 	{
-		Emitter->SetSpawnRate(FCString::Atof(*InText.ToString()) );
+		float Rate = FCString::Atof(*InText.ToString());
+		Emitter->GetProperties()->SpawnRate = Rate;
 	}
 
 	FText GetSpawnRateText() const
 	{
 		TCHAR Txt[32];
-		FCString::Snprintf(Txt, 32, TEXT("%f"), Emitter->GetSpawnRate() );
+		//FCString::Snprintf(Txt, 32, TEXT("%f"), Emitter->GetSpawnRate() );
+		FCString::Snprintf(Txt, 32, TEXT("%f"), Emitter->GetProperties()->SpawnRate );
 
 		return FText::FromString(Txt);
 	}
 
 	FString GetRenderModuleText() const
 	{
-		if (Emitter->GetRenderModuleType() == RMT_None)
+		if (Emitter->GetProperties()->RenderModuleType == RMT_None)
 		{
 			return FString("None");
 		}
 
-		return *RenderModuleList[Emitter->GetRenderModuleType()-1];
+		return *RenderModuleList[Emitter->GetProperties()->RenderModuleType-1];
 	}
 
 	TSharedRef<SWidget>GenerateRenderModuleComboWidget(TSharedPtr<FString> WidgetString )
diff --git a/Engine/Source/Editor/UnrealEd/Classes/ActorFactories/ActorFactoryNiagara.h b/Engine/Source/Editor/UnrealEd/Classes/ActorFactories/ActorFactoryNiagara.h
new file mode 100644
--- /dev/null
+++ b/Engine/Source/Editor/UnrealEd/Classes/ActorFactories/ActorFactoryNiagara.h
@@ -0,0 +1,22 @@
+// Copyright 1998-2014 Epic Games, Inc. All Rights Reserved.
+
+
+#pragma once
+#include "ActorFactories/ActorFactory.h"
+#include "ActorFactoryNiagara.generated.h"
+
+UCLASS(MinimalAPI, config = Editor, collapsecategories, hidecategories = Object)
+class UActorFactoryNiagara : public UActorFactory
+{
+	GENERATED_UCLASS_BODY()
+
+	// Begin UActorFactory Interface
+	virtual void PostSpawnActor(UObject* Asset, AActor* NewActor) override;
+	virtual void PostCreateBlueprint(UObject* Asset, AActor* CDO) override;
+	virtual bool CanCreateActorFrom(const FAssetData& AssetData, FText& OutErrorMsg) override;
+	virtual UObject* GetAssetFromActorInstance(AActor* ActorInstance) override;
+	// End UActorFactory Interface
+};
+
+
+
diff --git a/Engine/Source/Editor/UnrealEd/Private/Factories/ActorFactory.cpp b/Engine/Source/Editor/UnrealEd/Private/Factories/ActorFactory.cpp
--- a/Engine/Source/Editor/UnrealEd/Private/Factories/ActorFactory.cpp
+++ b/Engine/Source/Editor/UnrealEd/Private/Factories/ActorFactory.cpp
@@ -26,6 +26,9 @@ ActorFactory.cpp:
 #include "Particles/ParticleSystem.h"
 #include "Particles/ParticleSystemComponent.h"
 
+#include "Engine/NiagaraActor.h"
+#include "Engine/NiagaraEffect.h"
+
 #include "VectorField/VectorField.h"
 
 #include "Engine/TriggerBox.h"
@@ -519,6 +522,82 @@ void UActorFactoryEmitter::PostCreateBlueprint( UObject* Asset, AActor* CDO )
 }
 
 
+
+
+
+/*-----------------------------------------------------------------------------
+UActorFactoryNiagara
+-----------------------------------------------------------------------------*/
+UActorFactoryNiagara::UActorFactoryNiagara(const FObjectInitializer& ObjectInitializer)
+: Super(ObjectInitializer)
+{
+	DisplayName = LOCTEXT("EmitterDisplayName", "NiagaraEffect");
+	NewActorClass = ANiagaraActor::StaticClass();
+}
+
+bool UActorFactoryNiagara::CanCreateActorFrom(const FAssetData& AssetData, FText& OutErrorMsg)
+{
+	if (!AssetData.IsValid() || !AssetData.GetClass()->IsChildOf(UNiagaraEffect::StaticClass()))
+	{
+		OutErrorMsg = NSLOCTEXT("CanCreateActor", "NoEffect", "A valid Niagara effect must be specified.");
+		return false;
+	}
+
+	return true;
+}
+
+void UActorFactoryNiagara::PostSpawnActor(UObject* Asset, AActor* NewActor)
+{
+	UNiagaraEffect* Effect = CastChecked<UNiagaraEffect>(Asset);
+	ANiagaraActor* NiagaraActor = CastChecked<ANiagaraActor>(NewActor);
+
+	GEditor->SetActorLabelUnique(NewActor, Effect->GetName());
+
+	// Term Component
+	NiagaraActor->GetNiagaraComponent()->UnregisterComponent();
+
+	// Change properties
+	NiagaraActor->GetNiagaraComponent()->Effect = Effect;
+
+	// if we're created by Kismet on the server during gameplay, we need to replicate the emitter
+	if (NiagaraActor->GetWorld()->HasBegunPlay() && NiagaraActor->GetWorld()->GetNetMode() != NM_Client)
+	{
+		NiagaraActor->SetReplicates(true);
+		NiagaraActor->bAlwaysRelevant = true;
+		NiagaraActor->NetUpdateFrequency = 0.1f; // could also set bNetTemporary but LD might further trigger it or something
+	}
+
+	// Init Component
+	NiagaraActor->GetNiagaraComponent()->RegisterComponent();
+}
+
+UObject* UActorFactoryNiagara::GetAssetFromActorInstance(AActor* Instance)
+{
+	check(Instance->IsA(NewActorClass));
+	ANiagaraActor* NewActor = CastChecked<ANiagaraActor>(Instance);
+	if (NewActor->GetNiagaraComponent())
+	{
+		return NewActor->GetNiagaraComponent()->Effect;
+	}
+	else
+	{
+		return NULL;
+	}
+}
+
+void UActorFactoryNiagara::PostCreateBlueprint(UObject* Asset, AActor* CDO)
+{
+	if (Asset != NULL && CDO != NULL)
+	{
+		UNiagaraEffect* Effect = CastChecked<UNiagaraEffect>(Asset);
+		ANiagaraActor* Actor = CastChecked<ANiagaraActor>(CDO);
+		Actor->GetNiagaraComponent()->Effect = Effect;
+	}
+}
+
+
+
+
 /*-----------------------------------------------------------------------------
 UActorFactoryPlayerStart
 -----------------------------------------------------------------------------*/
diff --git a/Engine/Source/Runtime/Engine/Classes/Engine/NiagaraActor.h b/Engine/Source/Runtime/Engine/Classes/Engine/NiagaraActor.h
--- a/Engine/Source/Runtime/Engine/Classes/Engine/NiagaraActor.h
+++ b/Engine/Source/Runtime/Engine/Classes/Engine/NiagaraActor.h
@@ -3,8 +3,8 @@
 #pragma once
 #include "NiagaraActor.generated.h"
 
-UCLASS(notplaceable)
-class ANiagaraActor : public AActor
+UCLASS()
+class ENGINE_API ANiagaraActor : public AActor
 {
 	GENERATED_UCLASS_BODY()
 
diff --git a/Engine/Source/Runtime/Engine/Classes/Engine/NiagaraEffect.h b/Engine/Source/Runtime/Engine/Classes/Engine/NiagaraEffect.h
--- a/Engine/Source/Runtime/Engine/Classes/Engine/NiagaraEffect.h
+++ b/Engine/Source/Runtime/Engine/Classes/Engine/NiagaraEffect.h
@@ -17,14 +17,6 @@ public:
 	void Init(UNiagaraComponent *InComponent)
 	{
 		Component = InComponent;
-		//ConstantTable.Empty();
-		/*
-		// make room for constants
-		while (ConstantTable.Num() < 10)
-		{
-			ConstantTable.Add(FVector4(0.0f, 0.0f, 0.0f, 0.0f));
-		}
-		*/
 	}
 
 
@@ -51,14 +43,7 @@ public:
 		// TODO: should probably just pass a pointer to the table
 		for (TSharedPtr<FNiagaraSimulation>&it : Emitters)
 		{
-			/*
-			for (int32 i = 0; i < ConstantTable.Num(); i++)
-			{
-				it->SetConstant(i, ConstantTable[i]);
-			}
-			*/
 			it->SetConstants(Constants);
-
 			it->Tick(DeltaSeconds);
 		}
 	}
@@ -74,16 +59,18 @@ public:
 		return Emitters[idx].Get();
 	}
 
-
 	/** Local constant table. */
-	//TArray<FVector4> ConstantTable;
 	FNiagaraConstantMap Constants;
 
 	//UPROPERTY()
 	TArray< TSharedPtr<FNiagaraSimulation> > Emitters;
-
 	UNiagaraComponent *GetComponent() { return Component; }
+
+	virtual void PostLoad() override; 
+
 private:
 	UNiagaraComponent *Component;
+
+	UPROPERTY()
 	TArray<FNiagaraEmitterProperties>EmitterProps;
 };
diff --git a/Engine/Source/Runtime/Engine/Classes/Engine/NiagaraSimulation.h b/Engine/Source/Runtime/Engine/Classes/Engine/NiagaraSimulation.h
--- a/Engine/Source/Runtime/Engine/Classes/Engine/NiagaraSimulation.h
+++ b/Engine/Source/Runtime/Engine/Classes/Engine/NiagaraSimulation.h
@@ -37,16 +37,22 @@ struct FNiagaraEmitterProperties
 {
 	GENERATED_USTRUCT_BODY()
 public:
-	UPROPERTY(EditAnywhere, Category = "Emitter data")
-	float SpawnRate;
+	FNiagaraEmitterProperties() : RenderModuleType(RMT_Sprites)
+	{
+	}
 
-	UPROPERTY(EditAnywhere, Category = "Emitter data")
+	UPROPERTY(EditAnywhere, Category = "Emitter Properties")
+	bool bIsEnabled;
+	UPROPERTY(EditAnywhere, Category = "Emitter Properties")
+	float SpawnRate;
+	UPROPERTY(EditAnywhere, Category = "Emitter Properties")
 	UNiagaraScript *UpdateScript;
-	UPROPERTY(EditAnywhere, Category = "Emitter data")
+	UPROPERTY(EditAnywhere, Category = "Emitter Properties")
 	UNiagaraScript *SpawnScript;
-
-	UPROPERTY(EditAnywhere, Category = "Emitter data")
+	UPROPERTY(EditAnywhere, Category = "Emitter Properties")
 	UMaterial *Material;
+	UPROPERTY(EditAnywhere, Category = "Emitter Properties")
+	TEnumAsByte<EEmitterRenderModuleType> RenderModuleType;
 };
 
 
@@ -54,10 +60,11 @@ public:
 /**
 * A niagara particle simulation.
 */
-class FNiagaraSimulation
+struct FNiagaraSimulation
 {
 public:
-	explicit FNiagaraSimulation();
+	explicit FNiagaraSimulation(FNiagaraEmitterProperties *InProps);
+	FNiagaraSimulation(FNiagaraEmitterProperties &Props, ERHIFeatureLevel::Type InFeatureLeve);
 	virtual ~FNiagaraSimulation()
 	{}
 
@@ -72,30 +79,22 @@ public:
 	void SetConstants(const FNiagaraConstantMap &InMap)	{ Constants = InMap; }
 	NiagaraEffectRenderer *GetEffectRenderer()	{ return EffectRenderer;  }
 	
-	void SetUpdateScript(UNiagaraScript *InScript)	{ UpdateScript = InScript;  }
-	void SetSpawnScript(UNiagaraScript *InScript)	{ SpawnScript = InScript; }
-
 	bool IsEnabled()	{ return bIsEnabled;  }
 	void SetEnabled(bool bInEnabled)	{ bIsEnabled = bInEnabled;  }
 
-	void SetSpawnRate(float InRate)	{ SpawnRate = InRate; }
-	float GetSpawnRate()		{ return SpawnRate;  }
-
-	EEmitterRenderModuleType GetRenderModuleType()	{ return RenderModuleType; }
 	ENGINE_API void SetRenderModuleType(EEmitterRenderModuleType Type, ERHIFeatureLevel::Type FeatureLevel);
 
 	int32 GetNumParticles()	{ return Data.GetNumParticles(); }
+
+	FNiagaraEmitterProperties *GetProperties()	{ return Props; }
+	void SetProperties(FNiagaraEmitterProperties *InProps)	{ Props = InProps; }
+
 private:
+	FNiagaraEmitterProperties *Props;
+
 	float Age;
 	bool bIsEnabled;
-	/** Temporary stuff for the prototype. */
-	float SpawnRate;
-	EEmitterRenderModuleType RenderModuleType;
 
-	/** The particle update script. */
-	UNiagaraScript *UpdateScript;
-	/** The particle spawn script. */
-	UNiagaraScript *SpawnScript;
 	/** Local constant set. */
 	FNiagaraConstantMap Constants;
 	/** particle simulation data */
@@ -112,7 +111,7 @@ private:
 	/** Calc number to spawn */
 	int32 CalcNumToSpawn(float DeltaSeconds)
 	{
-		float FloatNumToSpawn = SpawnRemainder + (DeltaSeconds * SpawnRate);
+		float FloatNumToSpawn = SpawnRemainder + (DeltaSeconds * Props->SpawnRate);
 		int32 NumToSpawn = FMath::FloorToInt(FloatNumToSpawn);
 		SpawnRemainder = FloatNumToSpawn - NumToSpawn;
 		return NumToSpawn;
diff --git a/Engine/Source/Runtime/Engine/Private/NiagaraComponent.cpp b/Engine/Source/Runtime/Engine/Private/NiagaraComponent.cpp
--- a/Engine/Source/Runtime/Engine/Private/NiagaraComponent.cpp
+++ b/Engine/Source/Runtime/Engine/Private/NiagaraComponent.cpp
@@ -46,7 +46,10 @@ void FNiagaraSceneProxy::SetDynamicData_RenderThread(FNiagaraDynamicDataBase* Ne
 {
 	for (NiagaraEffectRenderer *Renderer : EffectRenderers)
 	{
-		Renderer->SetDynamicData_RenderThread(NewDynamicData);
+		if (Renderer)
+		{
+			Renderer->SetDynamicData_RenderThread(NewDynamicData);
+		}
 	}
 	return;
 }
@@ -56,7 +59,10 @@ void FNiagaraSceneProxy::ReleaseRenderThreadResources()
 {
 	for (NiagaraEffectRenderer *Renderer : EffectRenderers)
 	{
-		Renderer->ReleaseRenderThreadResources();
+		if (Renderer)
+		{
+			Renderer->ReleaseRenderThreadResources();
+		}
 	}
 	return;
 }
@@ -66,7 +72,10 @@ void FNiagaraSceneProxy::CreateRenderThreadResources()
 {
 	for (NiagaraEffectRenderer *Renderer : EffectRenderers)
 	{
-		Renderer->CreateRenderThreadResources();
+		if (Renderer)
+		{
+			Renderer->CreateRenderThreadResources();
+		}
 	}
 	return;
 }
@@ -85,7 +94,10 @@ void FNiagaraSceneProxy::PreRenderView(const FSceneViewFamily* ViewFamily, const
 {
 	for (NiagaraEffectRenderer *Renderer : EffectRenderers)
 	{
-		Renderer->PreRenderView(ViewFamily, VisibilityMap, FrameNumber, this);
+		if (Renderer)
+		{
+			Renderer->PreRenderView(ViewFamily, VisibilityMap, FrameNumber, this);
+		}
 	}
 	return;
 }
@@ -94,7 +106,10 @@ void FNiagaraSceneProxy::DrawDynamicElements(FPrimitiveDrawInterface* PDI, const
 {
 	for (NiagaraEffectRenderer *Renderer : EffectRenderers)
 	{
-		Renderer->DrawDynamicElements(PDI, View, this);
+		if (Renderer)
+		{
+			Renderer->DrawDynamicElements(PDI, View, this);
+		}
 	}
 	return;
 }
@@ -106,7 +121,10 @@ FPrimitiveViewRelevance FNiagaraSceneProxy::GetViewRelevance(const FSceneView* V
 
 	for (NiagaraEffectRenderer *Renderer : EffectRenderers)
 	{
-		Relevance |= EffectRenderers[0]->GetViewRelevance(View, this);
+		if (Renderer)
+		{
+			Relevance |= Renderer->GetViewRelevance(View, this);
+		}
 	}
 	return Relevance;
 }
@@ -122,7 +140,10 @@ uint32 FNiagaraSceneProxy::GetAllocatedSize() const
 	uint32 DynamicDataSize = 0;
 	for (NiagaraEffectRenderer *Renderer : EffectRenderers)
 	{
-		DynamicDataSize += Renderer->GetDynamicDataSize();
+		if (Renderer)
+		{
+			DynamicDataSize += Renderer->GetDynamicDataSize();
+		}
 	}
 	return FPrimitiveSceneProxy::GetAllocatedSize() + DynamicDataSize;
 }
@@ -169,7 +190,7 @@ void UNiagaraComponent::TickComponent(float DeltaSeconds, enum ELevelTick TickTy
 //	EmitterAge += DeltaSeconds;
 
 	if (Effect)
-	{
+	{ 
 
 		//Todo, open this up to the UI and setting via code and BPs.
 		static FName Const_Zero(TEXT("ZERO"));
@@ -201,6 +222,19 @@ void UNiagaraComponent::OnRegister()
 	if (Effect)
 	{
 		Effect->Init(this);
+
+		// initialize all render modules
+		ENQUEUE_UNIQUE_RENDER_COMMAND_TWOPARAMETER(
+			FChangeNiagaraRenderModule,
+			UNiagaraEffect*, InEffect, this->Effect,
+			UNiagaraComponent*, InComponent, this,
+			{
+				for (TSharedPtr<FNiagaraSimulation> Emitter : InEffect->Emitters)
+				{
+					Emitter->SetRenderModuleType(Emitter->GetProperties()->RenderModuleType, InComponent->GetWorld()->FeatureLevel);
+				}
+			}
+		);
 	}
 	VectorVM::Init();
 }
@@ -262,6 +296,7 @@ FBoxSphereBounds UNiagaraComponent::CalcBounds(const FTransform& LocalToWorld) c
 FPrimitiveSceneProxy* UNiagaraComponent::CreateSceneProxy()
 {
 	FNiagaraSceneProxy *Proxy = new FNiagaraSceneProxy(this);
+	Proxy->UpdateEffectRenderers(Effect);
 	return Proxy;
 }
 
diff --git a/Engine/Source/Runtime/Engine/Private/NiagaraEffect.cpp b/Engine/Source/Runtime/Engine/Private/NiagaraEffect.cpp
--- a/Engine/Source/Runtime/Engine/Private/NiagaraEffect.cpp
+++ b/Engine/Source/Runtime/Engine/Private/NiagaraEffect.cpp
@@ -15,7 +15,10 @@ UNiagaraEffect::UNiagaraEffect(const FObjectInitializer& ObjectInitializer)
 
 TSharedPtr<FNiagaraSimulation> UNiagaraEffect::AddEmitter()
 {
-	TSharedPtr<FNiagaraSimulation> Sim = MakeShareable(new FNiagaraSimulation());
+	FNiagaraEmitterProperties Props;
+	EmitterProps.Add(Props);
+
+	TSharedPtr<FNiagaraSimulation> Sim = MakeShareable(new FNiagaraSimulation(&EmitterProps.Top()));
 
 	Sim->SetRenderModuleType(RMT_Sprites, Component->GetWorld()->FeatureLevel);
 	Emitters.Add(Sim);
@@ -23,8 +26,21 @@ TSharedPtr<FNiagaraSimulation> UNiagaraEffect::AddEmitter()
 	FNiagaraSceneProxy *SceneProxy = static_cast<FNiagaraSceneProxy*>(Component->SceneProxy);
 	SceneProxy->UpdateEffectRenderers(this);
 
-	FNiagaraEmitterProperties Props;
-	EmitterProps.Add(Props);
-
 	return Sim;
 }
+
+
+
+
+void UNiagaraEffect::PostLoad()
+{
+	Super::PostLoad();
+
+	Emitters.Empty();
+
+	for (FNiagaraEmitterProperties &Props : EmitterProps)
+	{
+		FNiagaraSimulation *Sim = new FNiagaraSimulation(&Props);
+		Emitters.Add(MakeShareable(Sim));
+	}
+}
diff --git a/Engine/Source/Runtime/Engine/Private/NiagaraSimulation.cpp b/Engine/Source/Runtime/Engine/Private/NiagaraSimulation.cpp
--- a/Engine/Source/Runtime/Engine/Private/NiagaraSimulation.cpp
+++ b/Engine/Source/Runtime/Engine/Private/NiagaraSimulation.cpp
@@ -5,16 +5,13 @@
 #include "VectorVM.h"
 
 
-FNiagaraSimulation::FNiagaraSimulation() : Age(0.0f)
+FNiagaraSimulation::FNiagaraSimulation(FNiagaraEmitterProperties *InProps) : Age(0.0f)
 , bIsEnabled(true)
-, SpawnRate(0)
-, RenderModuleType(RMT_None)
-, UpdateScript(nullptr)
-, SpawnScript(nullptr)
 , SpawnRemainder(0.0f)
 , CachedBounds(ForceInit)
 , EffectRenderer(nullptr)
 {
+	Props = InProps;
 }
 
 
@@ -33,14 +30,14 @@ void FNiagaraSimulation::Tick(float DeltaSeconds)
 	int32 NumToSpawn = 0;
 
 	// Figure out how many we will spawn.
-	if (SpawnScript)
+	if (Props->SpawnScript)
 	{
 		NumToSpawn = CalcNumToSpawn(DeltaSeconds);
 		MaxNewParticles = Data.GetNumParticles() + NumToSpawn;
 		Data.Allocate(MaxNewParticles);
 	}
 
-	if (UpdateScript)
+	if (Props->UpdateScript)
 	{
 		// Simulate particles forward by DeltaSeconds.
 		UpdateParticles(
@@ -53,7 +50,7 @@ void FNiagaraSimulation::Tick(float DeltaSeconds)
 			);
 	}
 
-	if (SpawnScript)
+	if (Props->SpawnScript)
 	{
 		SpawnAndKillParticles(NumToSpawn);
 	}
@@ -79,7 +76,7 @@ void FNiagaraSimulation::UpdateParticles(
 
 	VectorRegister* InputRegisters[VectorVM::MaxInputRegisters] = { 0 };
 	VectorRegister* OutputRegisters[VectorVM::MaxOutputRegisters] = { 0 };
-	const int32 NumAttr = UpdateScript->Attributes.Num();
+	const int32 NumAttr = Props->UpdateScript->Attributes.Num();
 
 	check(NumAttr < VectorVM::MaxInputRegisters);
 	check(NumAttr < VectorVM::MaxOutputRegisters);
@@ -93,10 +90,10 @@ void FNiagaraSimulation::UpdateParticles(
 
 	//Fill constant table with required emitter constants and internal script constants.
 	TArray<FVector4> ConstantTable;
-	UpdateScript->ConstantData.FillConstantTable(Constants, ConstantTable);
+	Props->UpdateScript->ConstantData.FillConstantTable(Constants, ConstantTable);
 
 	VectorVM::Exec(
-		UpdateScript->ByteCode.GetData(),
+		Props->UpdateScript->ByteCode.GetData(),
 		InputRegisters,
 		NumAttr,
 		OutputRegisters,
@@ -118,11 +115,11 @@ int32 FNiagaraSimulation::SpawnAndKillParticles(int32 NumToSpawn)
 	CurNumParticles = SpawnParticles(NumToSpawn);
 
 	// run the spawn graph over all new particles
-	if (SpawnScript && SpawnScript->ByteCode.Num())
+	if (Props->SpawnScript && Props->SpawnScript->ByteCode.Num())
 	{
 		VectorRegister* InputRegisters[VectorVM::MaxInputRegisters] = { 0 };
 		VectorRegister* OutputRegisters[VectorVM::MaxOutputRegisters] = { 0 };
-		const int32 NumAttr = SpawnScript->Attributes.Num();
+		const int32 NumAttr = Props->SpawnScript->Attributes.Num();
 		const int32 NumVectors = NumToSpawn;
 
 		check(NumAttr < VectorVM::MaxInputRegisters);
@@ -139,10 +136,10 @@ int32 FNiagaraSimulation::SpawnAndKillParticles(int32 NumToSpawn)
 
 		//Fill constant table with required emitter constants and internal script constants.
 		TArray<FVector4> ConstantTable;
-		SpawnScript->ConstantData.FillConstantTable(Constants, ConstantTable);
+		Props->SpawnScript->ConstantData.FillConstantTable(Constants, ConstantTable);
 
 		VectorVM::Exec(
-			SpawnScript->ByteCode.GetData(),
+			Props->SpawnScript->ByteCode.GetData(),
 			InputRegisters,
 			NumAttr,
 			OutputRegisters,
@@ -175,7 +172,7 @@ int32 FNiagaraSimulation::SpawnAndKillParticles(int32 NumToSpawn)
  */
 void FNiagaraSimulation::SetRenderModuleType(EEmitterRenderModuleType Type, ERHIFeatureLevel::Type FeatureLevel)
 {
-	if (Type != RenderModuleType)
+	if (Type != Props->RenderModuleType || EffectRenderer==nullptr)
 	{
 		UMaterial *Material = UMaterial::GetDefaultMaterial(MD_Surface);
 
@@ -184,14 +181,24 @@ void FNiagaraSimulation::SetRenderModuleType(EEmitterRenderModuleType Type, ERHI
 			Material = EffectRenderer->GetMaterial();
 			delete EffectRenderer;
 		}
+		else
+		{
+			if (Props->Material)
+			{
+				Material = Props->Material;
+			}
+		}
 
-		RenderModuleType = Type;
+		Props->RenderModuleType = Type;
 		switch (Type)
 		{
 		case RMT_Sprites: EffectRenderer = new NiagaraEffectRendererSprites(FeatureLevel);
 			break;
 		case RMT_Ribbon: EffectRenderer = new NiagaraEffectRendererRibbon(FeatureLevel);
 			break;
+		default:	EffectRenderer = new NiagaraEffectRendererSprites(FeatureLevel);
+					Props->RenderModuleType = RMT_Sprites;
+					break;
 		}
 
 		EffectRenderer->SetMaterial(Material, FeatureLevel);
