Delete long deprecated class UDEPRECATED_SaveGameSummary

diff --git a/Engine/Source/Runtime/Core/Public/UObject/ObjectVersion.h b/Engine/Source/Runtime/Core/Public/UObject/ObjectVersion.h
--- a/Engine/Source/Runtime/Core/Public/UObject/ObjectVersion.h
+++ b/Engine/Source/Runtime/Core/Public/UObject/ObjectVersion.h
@@ -630,6 +630,8 @@ enum EUnrealEngineObjectUE4Version
 	VER_UE4_FOLIAGE_STATIC_LIGHTING_SUPPORT,
 	// Added composite fonts to Slate font info
 	VER_UE4_SLATE_COMPOSITE_FONTS,
+	// Remove UDEPRECATED_SaveGameSummary, required for UWorld::Serialize
+	VER_UE4_REMOVE_SAVEGAMESUMMARY,
 
 	// -----<new versions can be added before this line>-------------------------------------------------
 	// - this needs to be the last line (see note below)
diff --git a/Engine/Source/Runtime/Engine/Classes/Engine/SaveGameSummary.h b/Engine/Source/Runtime/Engine/Classes/Engine/SaveGameSummary.h
deleted file mode 100644
diff --git a/Engine/Source/Runtime/Engine/Classes/Engine/World.h b/Engine/Source/Runtime/Engine/Classes/Engine/World.h
--- a/Engine/Source/Runtime/Engine/Classes/Engine/World.h
+++ b/Engine/Source/Runtime/Engine/Classes/Engine/World.h
@@ -496,10 +496,6 @@ class ENGINE_API UWorld : public UObject, public FNetworkNotify
 	UPROPERTY(Transient)
 	class ULevel*								PersistentLevel;
 
-	/** Reference to last save game info used for serialization. The only time this is non NULL is during UEngine::SaveGame(..) */
-	UPROPERTY(Transient)
-	class UDEPRECATED_SaveGameSummary*			SaveGameSummary_DEPRECATED;
-
 	/** The NAME_GameNetDriver game connection(s) for client/server communication */
 	UPROPERTY(Transient)
 	class UNetDriver*							NetDriver;
diff --git a/Engine/Source/Runtime/Engine/Private/EnginePrivate.h b/Engine/Source/Runtime/Engine/Private/EnginePrivate.h
--- a/Engine/Source/Runtime/Engine/Private/EnginePrivate.h
+++ b/Engine/Source/Runtime/Engine/Private/EnginePrivate.h
@@ -432,7 +432,6 @@ EnginePrivate.h: Unreal engine private header file.
 #include "Engine/RendererSettings.h"
 #include "MovieScene/RuntimeMovieScenePlayerInterface.h"
 #include "GameFramework/SaveGame.h"
-#include "Engine/SaveGameSummary.h"
 #include "Engine/SCS_Node.h"
 #include "Engine/Selection.h"
 #include "Engine/SimpleConstructionScript.h"
diff --git a/Engine/Source/Runtime/Engine/Private/SaveGameSummary.cpp b/Engine/Source/Runtime/Engine/Private/SaveGameSummary.cpp
deleted file mode 100644
diff --git a/Engine/Source/Runtime/Engine/Private/World.cpp b/Engine/Source/Runtime/Engine/Private/World.cpp
--- a/Engine/Source/Runtime/Engine/Private/World.cpp
+++ b/Engine/Source/Runtime/Engine/Private/World.cpp
@@ -9,7 +9,6 @@
 #include "Engine/LevelScriptBlueprint.h"
 #include "Engine/LevelScriptActor.h"
 #include "Engine/CullDistanceVolume.h"
-#include "Engine/SaveGameSummary.h"
 #include "GameFramework/DefaultPhysicsVolume.h"
 #include "Engine/Console.h"
 #include "Engine/WorldComposition.h"
@@ -126,8 +125,13 @@ void UWorld::Serialize( FArchive& Ar )
 	Ar << EditorViews[1];
 	Ar << EditorViews[2];
 	Ar << EditorViews[3];
-	Ar << SaveGameSummary_DEPRECATED;
 
+	if (Ar.UE4Ver() < VER_UE4_REMOVE_SAVEGAMESUMMARY)
+	{
+		UObject* DummyObject;
+		Ar << DummyObject;
+		int i = 2;
+	}
 
 	if( !Ar.IsLoading() && !Ar.IsSaving() )
 	{
@@ -182,12 +186,11 @@ void UWorld::Serialize( FArchive& Ar )
 
 void UWorld::AddReferencedObjects(UObject* InThis, FReferenceCollector& Collector)
 {	
-	UWorld* This = CastChecked<UWorld>(InThis);
 #if WITH_EDITOR
+	UWorld* This = CastChecked<UWorld>(InThis);
 	if( GIsEditor )
 	{
 		Collector.AddReferencedObject( This->PersistentLevel, This );
-		Collector.AddReferencedObject( This->SaveGameSummary_DEPRECATED, This );
 
 		for( int32 Index = 0; Index < This->Levels.Num(); Index++ )
 		{
@@ -222,7 +225,7 @@ void UWorld::AddReferencedObjects(UObject* InThis, FReferenceCollector& Collecto
 	}
 #endif
 
-	Super::AddReferencedObjects( This, Collector );
+	Super::AddReferencedObjects( InThis, Collector );
 }
 
 #if WITH_EDITOR
diff --git a/Engine/Source/Runtime/Engine/Public/Engine.h b/Engine/Source/Runtime/Engine/Public/Engine.h
--- a/Engine/Source/Runtime/Engine/Public/Engine.h
+++ b/Engine/Source/Runtime/Engine/Public/Engine.h
@@ -277,7 +277,6 @@
 #include "Sound/ReverbEffect.h"
 #include "MovieScene/RuntimeMovieScenePlayerInterface.h"
 #include "GameFramework/SaveGame.h"
-#include "Engine/SaveGameSummary.h"
 #include "Engine/SCS_Node.h"
 #include "Engine/Selection.h"
 #include "Engine/SimpleConstructionScript.h"
