CIS fix

diff --git a/Engine/Source/Runtime/Engine/Classes/Engine/Font.h b/Engine/Source/Runtime/Engine/Classes/Engine/Font.h
--- a/Engine/Source/Runtime/Engine/Classes/Engine/Font.h
+++ b/Engine/Source/Runtime/Engine/Classes/Engine/Font.h
@@ -3,6 +3,7 @@
 #pragma once
 #include "Engine/FontImportOptions.h"
 #include "FontProviderInterface.h"
+#include "SlateFontInfo.h"
 #include "Font.generated.h"
 
 /** What kind of caching is this font using? */
