Remove debug code left in by accident

diff --git a/Engine/Source/Runtime/Engine/Private/World.cpp b/Engine/Source/Runtime/Engine/Private/World.cpp
--- a/Engine/Source/Runtime/Engine/Private/World.cpp
+++ b/Engine/Source/Runtime/Engine/Private/World.cpp
@@ -130,7 +130,6 @@ void UWorld::Serialize( FArchive& Ar )
 	{
 		UObject* DummyObject;
 		Ar << DummyObject;
-		int i = 2;
 	}
 
 	if( !Ar.IsLoading() && !Ar.IsSaving() )
