Fixes to PSA_Velocity and PSA_AwayFromCenter, both of which were potentially broken with initial rotation configurations

diff --git a/Engine/Source/Runtime/Engine/Private/Particles/ParticleEmitterInstances.cpp b/Engine/Source/Runtime/Engine/Private/Particles/ParticleEmitterInstances.cpp
--- a/Engine/Source/Runtime/Engine/Private/Particles/ParticleEmitterInstances.cpp
+++ b/Engine/Source/Runtime/Engine/Private/Particles/ParticleEmitterInstances.cpp
@@ -2814,7 +2814,8 @@ void FParticleMeshEmitterInstance::InitParameters(UParticleEmitter* InTemplate,
 
 	// Grab the MeshRotationRate module offset, if there is one...
 	MeshRotationActive = false;
-	if (LODLevel->RequiredModule->ScreenAlignment == PSA_Velocity)
+	if (LODLevel->RequiredModule->ScreenAlignment == PSA_Velocity
+		|| LODLevel->RequiredModule->ScreenAlignment == PSA_AwayFromCenter)
 	{
 		MeshRotationActive = true;
 	}
@@ -2923,7 +2924,6 @@ void FParticleMeshEmitterInstance::Tick(float DeltaTime, bool bSuppressSpawning)
 				}
 				else if (LODLevel->RequiredModule->ScreenAlignment == PSA_AwayFromCenter)
 				{
-
 					NewDirection = Particle.Location;
 				}
 
@@ -3149,10 +3149,16 @@ void FParticleMeshEmitterInstance::PostSpawn(FBaseParticle* Particle, float Inte
 
 	FMeshRotationPayloadData* PayloadData = (FMeshRotationPayloadData*)((uint8*)Particle + MeshRotationOffset);
 
-	if (LODLevel->RequiredModule->ScreenAlignment == PSA_Velocity)
+	if (LODLevel->RequiredModule->ScreenAlignment == PSA_Velocity
+		|| LODLevel->RequiredModule->ScreenAlignment == PSA_AwayFromCenter)
 	{
 		// Determine the rotation to the velocity vector and apply it to the mesh
-		FVector	NewDirection	= Particle->Velocity;
+		FVector	NewDirection = Particle->Velocity;
+		if (LODLevel->RequiredModule->ScreenAlignment == PSA_AwayFromCenter)
+		{
+			NewDirection = Particle->Location;
+		}
+
 		NewDirection.Normalize();
 		FVector	OldDirection(1.0f, 0.0f, 0.0f);
 
diff --git a/Engine/Source/Runtime/Engine/Private/Particles/ParticleSystemRender.cpp b/Engine/Source/Runtime/Engine/Private/Particles/ParticleSystemRender.cpp
--- a/Engine/Source/Runtime/Engine/Private/Particles/ParticleSystemRender.cpp
+++ b/Engine/Source/Runtime/Engine/Private/Particles/ParticleSystemRender.cpp
@@ -2534,9 +2534,7 @@ void FDynamicMeshEmitterData::GetParticleTransform(const FBaseParticle& InPartic
 		}
 		else
 		{
-			FRotator OtherRot = FRotator::MakeFromEuler(PayloadData->InitRotation+PayloadData->CurContinuousRotation);
-			FRotationMatrix RotMat(OtherRot);
-			OutTransformMat = (OrientMat*kScaleMat) * RotMat * kTransMat;
+			OutTransformMat = (OrientMat*kScaleMat) * kRotMat * kTransMat;
 		}
 	}
 	else if ((bUseCameraFacing == true) || (bUseMeshLockedAxis == true))
