Fix raycast not passing response params into the query

diff --git a/Engine/Source/Runtime/Engine/Private/Collision/WorldCollision.cpp b/Engine/Source/Runtime/Engine/Private/Collision/WorldCollision.cpp
--- a/Engine/Source/Runtime/Engine/Private/Collision/WorldCollision.cpp
+++ b/Engine/Source/Runtime/Engine/Private/Collision/WorldCollision.cpp
@@ -103,7 +103,7 @@ bool UWorld::SweepTest(const FVector& Start, const FVector& End, const FQuat& Ro
 	if (CollisionShape.IsNearlyZero())
 	{
 		// if extent is 0, we'll just do linetrace instead
-		return LineTraceTest(Start, End, TraceChannel, Params);
+		return LineTraceTest(Start, End, TraceChannel, Params, ResponseParam);
 	}
 	else
 	{
@@ -119,7 +119,7 @@ bool UWorld::SweepSingle(struct FHitResult& OutHit, const FVector& Start, const
 {
 	if (CollisionShape.IsNearlyZero())
 	{
-		return LineTraceSingle(OutHit, Start, End, TraceChannel, Params);
+		return LineTraceSingle(OutHit, Start, End, TraceChannel, Params, ResponseParam);
 	}
 	else
 	{
@@ -137,7 +137,7 @@ bool UWorld::SweepMulti(TArray<FHitResult>& OutHits, const FVector& Start, const
 {
 	if (CollisionShape.IsNearlyZero())
 	{
-		return LineTraceMulti(OutHits, Start, End, TraceChannel, Params);
+		return LineTraceMulti(OutHits, Start, End, TraceChannel, Params, ResponseParam);
 	}
 	else
 	{
