Removed some code from UFont as it was causing a crash when cooking

Slate application is null in this case, so we need a better way to clear unloaded fonts out of the Slate font cache

diff --git a/Engine/Source/Runtime/Engine/Private/Font.cpp b/Engine/Source/Runtime/Engine/Private/Font.cpp
--- a/Engine/Source/Runtime/Engine/Private/Font.cpp
+++ b/Engine/Source/Runtime/Engine/Private/Font.cpp
@@ -23,6 +23,8 @@ UFont::UFont(const FObjectInitializer& ObjectInitializer)
 
 UFont::~UFont()
 {
+	// todo: jdale - this is causing a crash when cooking due to the Slate application being null
+	/*
 	if(FontCacheType == EFontCacheType::Runtime)
 	{
 		TSharedPtr<FSlateRenderer> SlateRenderer = FSlateApplication::Get().GetRenderer();
@@ -32,6 +34,7 @@ UFont::~UFont()
 			FontCache->FlushObject(this);
 		}
 	}
+	*/
 }
 
 void UFont::Serialize( FArchive& Ar )
