Convert more FRepLayout code to use FRepLayoutCmdIterator

diff --git a/Engine/Source/Runtime/Engine/Private/RepLayout.cpp b/Engine/Source/Runtime/Engine/Private/RepLayout.cpp
--- a/Engine/Source/Runtime/Engine/Private/RepLayout.cpp
+++ b/Engine/Source/Runtime/Engine/Private/RepLayout.cpp
@@ -1831,108 +1831,95 @@ void FRepLayout::SanityCheckChangeList( const uint8* RESTRICT Data, TArray< uint
 	check( Changed[ChangedIndex] == 0 );
 }
 
-void FRepLayout::DiffProperties_DynamicArray_r( 
-	FRepState *				RepState,
-	const int32				CmdIndex, 
-	const uint8* RESTRICT	StoredData, 
-	const uint8* RESTRICT	Data,
-	const bool				bSync,
-	bool &					bOutDifferent ) const
+class FDiffPropertiesImpl : public FRepLayoutCmdIterator< FDiffPropertiesImpl, FCmdIteratorBaseStackState >
 {
-	const FRepLayoutCmd& Cmd = Cmds[ CmdIndex ];
+public:
+	FDiffPropertiesImpl( const bool bInSync, TArray< UProperty * >&	InRepNotifies, const TArray< FRepParentCmd >& InParents, const TArray< FRepLayoutCmd >& InCmds ) : 
+		FRepLayoutCmdIterator( InParents, InCmds ),
+		bSync( bInSync ),
+		RepNotifies( InRepNotifies ),
+		bDifferent( false )
+	{}
 
-	FScriptArray * Array = (FScriptArray *)Data;
-	FScriptArray * StoredArray = (FScriptArray *)StoredData;
+	INIT_STACK( FCmdIteratorBaseStackState ) { }
+
+	SHOULD_PROCESS_NEXT_CMD() 
+	{ 
+		return true;
+	}
 
-	if ( Array->Num() != StoredArray->Num() )
+	PROCESS_ARRAY_CMD( FCmdIteratorBaseStackState ) 
 	{
-		bOutDifferent = true;
+		if ( StackState.DataArray->Num() != StackState.ShadowArray->Num() )
+		{
+			bDifferent = true;
 
-		if ( !bSync )
-		{			
-			UE_LOG( LogNet, Warning, TEXT( "DiffProperties_DynamicArray_r: Array sizes different: %s %i / %i" ), *Cmd.Property->GetFullName(), Array->Num(), StoredArray->Num() );
-			return;
-		}
+			if ( !bSync )
+			{			
+				UE_LOG( LogNet, Warning, TEXT( "FDiffPropertiesImpl: Array sizes different: %s %i / %i" ), *Cmd.Property->GetFullName(), StackState.DataArray->Num(), StackState.ShadowArray->Num() );
+				return;
+			}
 
-		if ( !( Parents[Cmd.ParentIndex].Flags & PARENT_IsLifetime ) )
-		{
-			// Currently, only lifetime properties init from their defaults
-			return;
-		}
+			if ( !( Parents[Cmd.ParentIndex].Flags & PARENT_IsLifetime ) )
+			{
+				// Currently, only lifetime properties init from their defaults
+				return;
+			}
 
-		// Make the shadow state match the actual state
-		FScriptArrayHelper StoredArrayHelper( (UArrayProperty *)Cmd.Property, StoredData );
-		StoredArrayHelper.Resize( Array->Num() );
-	}
+			// Make the shadow state match the actual state
+			FScriptArrayHelper ShadowArrayHelper( (UArrayProperty *)Cmd.Property, ShadowData );
+			ShadowArrayHelper.Resize( StackState.DataArray->Num() );
+		}
 
-	Data		= (uint8*)Array->GetData();
-	StoredData	= (uint8*)StoredArray->GetData();
+		StackState.BaseData			= (uint8*)StackState.DataArray->GetData();
+		StackState.ShadowBaseData	= (uint8*)StackState.ShadowArray->GetData();
 
-	for ( int32 i = 0; i < Array->Num(); i++ )
-	{
-		const int32 ElementOffset = i * Cmd.ElementSize;
-		DiffProperties_r( RepState, CmdIndex + 1, Cmd.EndCmd - 1, StoredData + ElementOffset, Data + ElementOffset, bSync, bOutDifferent );
+		// Loop over array
+		ProcessDataArrayElements_r( StackState, Cmd );
 	}
-}
 
-void FRepLayout::DiffProperties_r( 
-	FRepState *				RepState,
-	const int32				CmdStart, 
-	const int32				CmdEnd, 
-	const uint8* RESTRICT	StoredData, 
-	const uint8* RESTRICT	Data,
-	const bool				bSync,
-	bool &					bOutDifferent ) const
-{
-	for ( int32 CmdIndex = CmdStart; CmdIndex < CmdEnd; CmdIndex++ )
+	PROCESS_CMD( FCmdIteratorBaseStackState ) 
 	{
-		const FRepLayoutCmd& Cmd = Cmds[ CmdIndex ];
-
-		check( Cmd.Type != REPCMD_Return );
-
-		if ( Cmd.Type == REPCMD_DynamicArray )
-		{
-			DiffProperties_DynamicArray_r( RepState, CmdIndex, StoredData + Cmd.Offset, Data + Cmd.Offset, bSync, bOutDifferent );
-			CmdIndex = Cmd.EndCmd - 1;	// Jump past children of this array (-1 for the ++ in the for loop)
-			continue;
-		}
-
 		const FRepParentCmd& Parent = Parents[Cmd.ParentIndex];
 
-		const FRepLayoutCmd& SwappedCmd = Cmd;//Parent.RoleSwapIndex != -1 ? Cmds[Parents[Parent.RoleSwapIndex].CmdStart] : Cmd;
-
 		// Make the shadow state match the actual state at the time of send
-		if ( Parent.RepNotifyCondition == REPNOTIFY_Always || !PropertiesAreIdentical( Cmd, (const void*)( Data + SwappedCmd.Offset ), (const void*)( StoredData + Cmd.Offset ) ) )
+		if ( Parent.RepNotifyCondition == REPNOTIFY_Always || !PropertiesAreIdentical( Cmd, (const void*)( Data + Cmd.Offset ), (const void*)( ShadowData + Cmd.Offset ) ) )
 		{
-			bOutDifferent = true;
+			bDifferent = true;
 
 			if ( !bSync )
 			{			
-				UE_LOG( LogNet, Warning, TEXT( "DiffProperties_r: Property different: %s" ), *Cmd.Property->GetFullName() );
-				continue;
+				UE_LOG( LogNet, Warning, TEXT( "FDiffPropertiesImpl: Property different: %s" ), *Cmd.Property->GetFullName() );
+				return;
 			}
 
-			if ( !( Parents[Cmd.ParentIndex].Flags & PARENT_IsLifetime ) )
+			if ( !( Parent.Flags & PARENT_IsLifetime ) )
 			{
 				// Currently, only lifetime properties init from their defaults
-				continue;
+				return;
 			}
 
-			StoreProperty( Cmd, (void*)( Data + SwappedCmd.Offset ), (const void*)( StoredData + Cmd.Offset ) );
+			StoreProperty( Cmd, (void*)( Data + Cmd.Offset ), (const void*)( ShadowData + Cmd.Offset ) );
 
 			if ( Parent.Property->HasAnyPropertyFlags( CPF_RepNotify ) )
 			{
-				RepState->RepNotifies.AddUnique( Parent.Property );
+				RepNotifies.AddUnique( Parent.Property );
 			}
 		}
 	}
-}
+
+	bool					bSync;
+	TArray< UProperty * >&	RepNotifies;
+	bool					bDifferent;
+};
 
 bool FRepLayout::DiffProperties( FRepState * RepState, const void* RESTRICT Data, const bool bSync ) const
-{
-	bool bDifferent = false;
-	DiffProperties_r( RepState, 0, Cmds.Num() - 1, RepState->StaticBuffer.GetData(), (uint8*)Data, bSync, bDifferent );
-	return bDifferent;
+{	
+	FDiffPropertiesImpl DiffPropertiesImpl( bSync, RepState->RepNotifies, Parents, Cmds );
+
+	DiffPropertiesImpl.ProcessCmds( RepState, (uint8*)Data );
+
+	return DiffPropertiesImpl.bDifferent;
 }
 
 void FRepLayout::AddPropertyCmd( UProperty * Property, int32 Offset, int32 RelativeHandle, int32 ParentIndex )
diff --git a/Engine/Source/Runtime/Engine/Public/Net/RepLayout.h b/Engine/Source/Runtime/Engine/Public/Net/RepLayout.h
--- a/Engine/Source/Runtime/Engine/Public/Net/RepLayout.h
+++ b/Engine/Source/Runtime/Engine/Public/Net/RepLayout.h
@@ -403,23 +403,6 @@ private:
 
 	void SanityCheckChangeList( const uint8* RESTRICT Data, TArray< uint16 > & Changed ) const;
 
-	void DiffProperties_DynamicArray_r( 
-		FRepState *				RepState,
-		const int32				CmdIndex, 
-		const uint8* RESTRICT	StoredData, 
-		const uint8* RESTRICT	Data,
-		const bool				bSync,
-		bool &					bOutDifferent ) const;
-
-	void DiffProperties_r( 
-		FRepState *				RepState,
-		const int32				CmdStart, 
-		const int32				CmdEnd, 
-		const uint8* RESTRICT	StoredData, 
-		const uint8* RESTRICT	Data,
-		const bool				bSync,
-		bool &					bOutDifferent ) const;
-
 	uint16 AddParentProperty( UProperty * Property, int32 ArrayIndex );
 
 	int32 InitFromProperty_r( UProperty * Property, int32 Offset, int32 RelativeHandle, int32 ParentIndex );
