Fixed memory leak in ApexClothVertexFactory

diff --git a/Engine/Source/Runtime/Engine/Private/GPUSkinVertexFactory.cpp b/Engine/Source/Runtime/Engine/Private/GPUSkinVertexFactory.cpp
--- a/Engine/Source/Runtime/Engine/Private/GPUSkinVertexFactory.cpp
+++ b/Engine/Source/Runtime/Engine/Private/GPUSkinVertexFactory.cpp
@@ -331,7 +331,7 @@ template <bool bExtraBoneInfluencesT>
 void TGPUSkinAPEXClothVertexFactory<bExtraBoneInfluencesT>::ReleaseDynamicRHI()
 {
 	Super::ReleaseDynamicRHI();
-	ClothShaderData.ReleaseClothUniformBuffer();
+	ClothShaderData.ReleaseClothSimulData();
 }
 
 /*-----------------------------------------------------------------------------
diff --git a/Engine/Source/Runtime/Engine/Public/GPUSkinVertexFactory.h b/Engine/Source/Runtime/Engine/Public/GPUSkinVertexFactory.h
--- a/Engine/Source/Runtime/Engine/Public/GPUSkinVertexFactory.h
+++ b/Engine/Source/Runtime/Engine/Public/GPUSkinVertexFactory.h
@@ -641,16 +641,12 @@ public:
 
 		void UpdateClothUniformBuffer(const TArray<FVector4>& InSimulPositions, const TArray<FVector4>& InSimulNormals);
 
-		void ReleaseClothUniformBuffer()
-		{
-			APEXClothUniformBuffer.SafeRelease();
-		}
-
-
 		void UpdateClothSimulData(const TArray<FVector4>& InSimulPositions, const TArray<FVector4>& InSimulNormals, ERHIFeatureLevel::Type FeatureLevel);
 
 		void ReleaseClothSimulData()
 		{
+			APEXClothUniformBuffer.SafeRelease();
+
 			if(IsValidRef(ClothSimulPositionBuffer))
 			{
 				ClothSimulDataBufferPool.ReleasePooledResource(ClothSimulPositionBuffer);
