Missing headers from Physics and AI, causing errors in Mac nonunity builds. Fixed.

diff --git a/Engine/Source/Runtime/Engine/Public/AI/NavigationSystemHelpers.h b/Engine/Source/Runtime/Engine/Public/AI/NavigationSystemHelpers.h
--- a/Engine/Source/Runtime/Engine/Public/AI/NavigationSystemHelpers.h
+++ b/Engine/Source/Runtime/Engine/Public/AI/NavigationSystemHelpers.h
@@ -4,6 +4,7 @@
 
 #include "AI/NavigationModifier.h"
 #include "AI/Navigation/NavLinkDefinition.h"
+#include "AI/Navigation/NavigationTypes.h"
 
 #if WITH_PHYSX
 namespace physx
diff --git a/Engine/Source/Runtime/Engine/Public/DynamicMeshBuilder.h b/Engine/Source/Runtime/Engine/Public/DynamicMeshBuilder.h
--- a/Engine/Source/Runtime/Engine/Public/DynamicMeshBuilder.h
+++ b/Engine/Source/Runtime/Engine/Public/DynamicMeshBuilder.h
@@ -5,7 +5,9 @@
 =============================================================================*/
 
 #pragma once
+
 #include "RenderUtils.h"
+#include "SceneManagement.h"
 
 class FPrimitiveDrawInterface;
 class FMaterialRenderProxy;
diff --git a/Engine/Source/Runtime/Engine/Public/MeshBatch.h b/Engine/Source/Runtime/Engine/Public/MeshBatch.h
--- a/Engine/Source/Runtime/Engine/Public/MeshBatch.h
+++ b/Engine/Source/Runtime/Engine/Public/MeshBatch.h
@@ -1,7 +1,9 @@
 // Copyright 1998-2014 Epic Games, Inc. All Rights Reserved.
 
 #pragma once
+
 #include "MaterialShared.h"
+#include "PrimitiveUniformShaderParameters.h"
 
 class FLightCacheInterface;
 /**
diff --git a/Engine/Source/Runtime/Engine/Public/PhysicsPublic.h b/Engine/Source/Runtime/Engine/Public/PhysicsPublic.h
--- a/Engine/Source/Runtime/Engine/Public/PhysicsPublic.h
+++ b/Engine/Source/Runtime/Engine/Public/PhysicsPublic.h
@@ -8,6 +8,7 @@
 #pragma once
 
 #include "PhysxUserData.h"
+#include "DynamicMeshBuilder.h"
 #include "LocalVertexFactory.h"
 /**
  * Physics stats
