Cleaned up HUD header. Moved some structs into their own file.
Also revised check in AddDebugText_Implementation that was meant to draw white if no color specified.

diff --git a/Engine/Source/Runtime/Engine/Classes/GameFramework/DebugTextInfo.h b/Engine/Source/Runtime/Engine/Classes/GameFramework/DebugTextInfo.h
new file mode 100644
--- /dev/null
+++ b/Engine/Source/Runtime/Engine/Classes/GameFramework/DebugTextInfo.h
@@ -0,0 +1,84 @@
+// Copyright 1998-2014 Epic Games, Inc. All Rights Reserved.
+
+#pragma once
+
+#include "DebugTextInfo.generated.h"
+
+class AActor;
+class UFont;
+
+/*
+ * Single entry of a debug text item to render. 
+ *
+ * @see AHud
+ * @see AddDebugText(), RemoveDebugText() and DrawDebugTextList()
+ */
+USTRUCT()
+struct FDebugTextInfo
+{
+	GENERATED_USTRUCT_BODY()
+
+	/**  AActor related to text item */
+	UPROPERTY()
+	AActor* SrcActor;
+
+	/** Offset from SrcActor.Location to apply */
+	UPROPERTY()
+	FVector SrcActorOffset;
+
+	/** Desired offset to interpolate to */
+	UPROPERTY()
+	FVector SrcActorDesiredOffset;
+
+	/** Text to display */
+	UPROPERTY()
+	FString DebugText;
+
+	/** Time remaining for the debug text, -1.f == infinite */
+	UPROPERTY(transient)
+	float TimeRemaining;
+
+	/** Duration used to lerp desired offset */
+	UPROPERTY()
+	float Duration;
+
+	/** Text color */
+	UPROPERTY()
+	FColor TextColor;
+
+	/** whether the offset should be treated as absolute world location of the string */
+	UPROPERTY()
+	uint32 bAbsoluteLocation:1;
+
+	/** If the actor moves does the text also move with it? */
+	UPROPERTY()
+	uint32 bKeepAttachedToActor:1;
+
+	/** When we first spawn store off the original actor location for use with bKeepAttachedToActor */
+	UPROPERTY()
+	FVector OrigActorLocation;
+
+	/** The Font which to display this as.  Will Default to GetSmallFont()**/
+	UPROPERTY()
+	UFont* Font;
+
+	/** Scale to apply to font when rendering */
+	UPROPERTY()
+	float FontScale;
+
+	FDebugTextInfo()
+		: SrcActor(NULL)
+		, SrcActorOffset(ForceInit)
+		, SrcActorDesiredOffset(ForceInit)
+		, TimeRemaining(0)
+		, Duration(0)
+		, TextColor(ForceInit)
+		, bAbsoluteLocation(false)
+		, bKeepAttachedToActor(false)
+		, OrigActorLocation(ForceInit)
+		, Font(NULL)
+		, FontScale(1.0f)
+	{
+	}
+
+};
diff --git a/Engine/Source/Runtime/Engine/Classes/GameFramework/HUD.h b/Engine/Source/Runtime/Engine/Classes/GameFramework/HUD.h
--- a/Engine/Source/Runtime/Engine/Classes/GameFramework/HUD.h
+++ b/Engine/Source/Runtime/Engine/Classes/GameFramework/HUD.h
@@ -2,159 +2,26 @@
 
 #pragma once
 #include "GameFramework/Actor.h"
+#include "DebugTextInfo.h"
+#include "HUDHitBox.h"
 #include "HUD.generated.h"
 
 class FCanvasTextItem;
 class UTexture;
-
-/** List of actors and debug text to draw, @see AddDebugText(), RemoveDebugText(), and DrawDebugTextList() */
-USTRUCT()
-struct FDebugTextInfo
-{
-	GENERATED_USTRUCT_BODY()
-
-	/**  AActor to draw DebugText over */
-	UPROPERTY()
-	class AActor* SrcActor;
-
-	/** Offset from SrcActor.Location to apply */
-	UPROPERTY()
-	FVector SrcActorOffset;
-
-	/** Desired offset to interpolate to */
-	UPROPERTY()
-	FVector SrcActorDesiredOffset;
-
-	/** Text to display */
-	UPROPERTY()
-	FString DebugText;
-
-	/** Time remaining for the debug text, -1.f == infinite */
-	UPROPERTY(transient)
-	float TimeRemaining;
-
-	/** Duration used to lerp desired offset */
-	UPROPERTY()
-	float Duration;
-
-	/** Text color */
-	UPROPERTY()
-	FColor TextColor;
-
-	/** whether the offset should be treated as absolute world location of the string */
-	UPROPERTY()
-	uint32 bAbsoluteLocation:1;
-
-	/** If the actor moves does the text also move with it? */
-	UPROPERTY()
-	uint32 bKeepAttachedToActor:1;
-
-	/** When we first spawn store off the original actor location for use with bKeepAttachedToActor */
-	UPROPERTY()
-	FVector OrigActorLocation;
-
-	/** The Font which to display this as.  Will Default to GetSmallFont()**/
-	UPROPERTY()
-	class UFont* Font;
-
-	/** Scale to apply to font when rendering */
-	UPROPERTY()
-	float FontScale;
-
-	FDebugTextInfo()
-		: SrcActor(NULL)
-		, SrcActorOffset(ForceInit)
-		, SrcActorDesiredOffset(ForceInit)
-		, TimeRemaining(0)
-		, Duration(0)
-		, TextColor(ForceInit)
-		, bAbsoluteLocation(false)
-		, bKeepAttachedToActor(false)
-		, OrigActorLocation(ForceInit)
-		, Font(NULL)
-		, FontScale(1.0f)
-	{
-	}
-
-};
-
-class ENGINE_API FSimpleReticle
-{
-public:
-	FSimpleReticle()
-	{
-		SetupReticle( 8.0f, 20.0f );
-	}
-
-	void SetupReticle( const float Length, const float InnerSize )
-	{
-		HorizontalOffsetMin.Set( InnerSize, 0.0f );
-		HorizontalOffsetMax.Set( InnerSize + Length, 0.0f );
-		VerticalOffsetMin.Set( 0.0f, InnerSize);
-		VerticalOffsetMax.Set( 0.0f, InnerSize + Length );
-	}
-	void Draw( class UCanvas* InCanvas, FLinearColor InColor );
-private:
-	FVector2D HorizontalOffsetMin;
-	FVector2D HorizontalOffsetMax;
-	FVector2D VerticalOffsetMin;
-	FVector2D VerticalOffsetMax;	
-};
-
-class ENGINE_API FHUDHitBox
-{
-public:
-	/** 
-	 * Constructor for a hitbox.
-	 * @param	InCoords		Coordinates of top left of hitbox.
-	 * @param	InSize			Size of the box.
-	 */
-	FHUDHitBox( FVector2D InCoords, FVector2D InSize, const FName& InName, bool bInConsumesInput, int32 InPriority );
-
-	/** 
-	 * Are the given coordinates within this hitbox.
-	 * @param	InCoords		Coordinates to check.
-	 * @returns true if coordinates are within this hitbox.
-	 */
-	bool Contains( FVector2D InCoords ) const;
-
-	/** 
-	 * Debug render for this hitbox.
-	 * @param	InCanvas		Canvas on which to render.
-	 * @param	InColor			Color to render the box.
-	 */
-	void Draw( class FCanvas* InCanvas, const FLinearColor& InColor ) const;
-	
-	/** Get the name of this hitbox.  */
-	const FName& GetName() const { return Name;};
-
-	/** Should other boxes be processed if this box is clicked.  */
-	const bool ConsumesInput() const { return bConsumesInput; };
-
-	/** Get the priority of this hitbox.  */
-	const int32 GetPriority() const { return Priority; };
-
-private:
-	/** Coordinates of top left of hitbox.  */
-	FVector2D	Coords;
-
-	/** Size of this hitbox.  */
-	FVector2D	Size;
-	
-	/** The name of this hitbox.  */
-	FName		Name;	
-
-	/** Wether or not this hitbox should prevent hit checks to other hitboxes.  */
-	bool bConsumesInput;
-
-	/** The priority of this hitbox. Higher boxes are given priority. */
-	int32 Priority;
-};
-
-//=============================================================================
-// Base class of the heads-up display.
-//
-//=============================================================================
+class AActor;
+class UCanvas;
+class APlayerController;
+class UFont;
+
+/** 
+ * Base class of the heads-up display. This has a canvas and a debug canvas on which primitives can be drawn.
+ * It also contains a list of simple hit boxes that can be used for simple item click detection.
+ * A method of rendering debug text is also included.
+ * Provides some simple methods for rendering text, textures, rectangles and materials which can also be accessed from blueprints.
+ * @see UCanvas
+ * @see FHUDHitBox
+ * @see FDebugTextInfo
+ */
 UCLASS(config=Game, hidecategories=(Rendering,Actor,Input,Replication), showcategories=("Input|MouseInput", "Input|TouchInput"), notplaceable, transient, BlueprintType, Blueprintable)
 class ENGINE_API AHUD : public AActor
 {
@@ -172,7 +39,7 @@ class ENGINE_API AHUD : public AActor
 
 	/** PlayerController which owns this HUD. */
 	UPROPERTY()
-	class APlayerController* PlayerOwner;    
+	APlayerController* PlayerOwner;    
 
 	/** Tells whether the game was paused due to lost focus */
 	UPROPERTY(transient)
@@ -200,7 +67,7 @@ class ENGINE_API AHUD : public AActor
 
 	/** Holds a list of Actors that need PostRender() calls. */
 	UPROPERTY()
-	TArray<class AActor*> PostRenderedActors;
+	TArray<AActor*> PostRenderedActors;
 
 	/** Used to calculate delta time between HUD rendering. */
 	UPROPERTY(transient)
@@ -221,11 +88,11 @@ class ENGINE_API AHUD : public AActor
 protected:
 	/** Canvas to Draw HUD on.  Only valid during PostRender() event.  */
 	UPROPERTY()
-	class UCanvas* Canvas;
+	UCanvas* Canvas;
 
 	/** 'Foreground' debug canvas, will draw in front of Slate UI. */
 	UPROPERTY()
-	class UCanvas* DebugCanvas;
+	UCanvas* DebugCanvas;
 
 	UPROPERTY()
 	TArray<struct FDebugTextInfo> DebugTextList;
@@ -251,14 +118,39 @@ public:
 	UFUNCTION(exec)
 	void ShowDebugToggleSubCategory(FName Category);
 
+	/**
+	 * Add debug text for a specific actor to be displayed via DrawDebugTextList().  If the debug text is invalid then it will
+	 * attempt to remove any previous entries via RemoveDebugText().
+	 * 
+	 * @param DebugText				Text to draw
+	 * @param SrcActor				Actor to which this relates
+	 * @param Duration				Duration to display the string
+	 * @param Offset 				Initial offset to render text
+	 * @param DesiredOffset 		Desired offset to render text - the text will move to this location over the given duration
+	 * @param TextColor 			Color of text to render
+	 * @param bSkipOverwriteCheck 	skips the check to see if there is already debug text for the given actor
+	 * @param bAbsoluteLocation 	use an absolute world location
+	 * @param bKeepAttachedToActor 	if this is true the text will follow the actor, otherwise it will be drawn at the location when the call was made
+	 * @param InFont 				font to use
+	 * @param FontScale 			scale
+	 */
 	UFUNCTION(reliable, client, SealedEvent)
-	void AddDebugText(const FString& DebugText, class AActor* SrcActor = NULL, float Duration = 0, FVector Offset = FVector(ForceInit), FVector DesiredOffset = FVector(ForceInit), FColor TextColor = FColor(ForceInit), bool bSkipOverwriteCheck = false, bool bAbsoluteLocation = false, bool bKeepAttachedToActor = false, class UFont* InFont = NULL, float FontScale = 1.0);
+	void AddDebugText(const FString& DebugText, AActor* SrcActor = NULL, float Duration = 0, FVector Offset = FVector(ForceInit), FVector DesiredOffset = FVector(ForceInit), FColor TextColor = FColor(ForceInit), bool bSkipOverwriteCheck = false, bool bAbsoluteLocation = false, bool bKeepAttachedToActor = false, UFont* InFont = NULL, float FontScale = 1.0);
 
+	/**
+	 * Remove all debug strings added via AddDebugText
+	 */
 	UFUNCTION(reliable, client, SealedEvent)
 	void RemoveAllDebugStrings();
 
+	/**
+	 * Remove debug strings for the given actor
+	 *
+	 * @param	SrcActor			Actor whose string you wish to remove
+	 * @param	bLeaveDurationText	when true text that has a finite duration will be removed, otherwise all will be removed for given actor
+	 */
 	UFUNCTION(reliable, client, SealedEvent)
-	void RemoveDebugText(class AActor* SrcActor, bool bLeaveDurationText = false);
+	void RemoveDebugText(AActor* SrcActor, bool bLeaveDurationText = false);
 
 	/** Hook to allow blueprints to do custom HUD drawing. @see bSuppressNativeHUD to control HUD drawing in base class. */
 	UFUNCTION(BlueprintImplementableEvent, BlueprintCosmetic)
@@ -292,7 +184,7 @@ public:
 	 * @param Scale				Scale multiplier to control size of the text.
 	 */
 	UFUNCTION(BlueprintCallable, Category=HUD)
-	void GetTextSize(const FString& Text, float& OutWidth, float& OutHeight, class UFont* Font=NULL, float Scale=1.f) const;
+	void GetTextSize(const FString& Text, float& OutWidth, float& OutHeight, UFont* Font=NULL, float Scale=1.f) const;
 
 	/**
 	 * Draws a string on the HUD.
@@ -305,7 +197,7 @@ public:
 	 * @param bScalePosition	Whether the "Scale" parameter should also scale the position of this draw call.
 	 */
 	UFUNCTION(BlueprintCallable, Category=HUD)
-	void DrawText(const FString& Text, FLinearColor TextColor, float ScreenX, float ScreenY, class UFont* Font=NULL, float Scale=1.f, bool bScalePosition=false);
+	void DrawText(const FString& Text, FLinearColor TextColor, float ScreenX, float ScreenY, UFont* Font=NULL, float Scale=1.f, bool bScalePosition=false);
 
 	/**
 	 * Draws a 2D line on the HUD.
@@ -401,7 +293,6 @@ public:
 	UFUNCTION(BlueprintCallable, Category = HUD)
 	void Deproject(float ScreenX, float ScreenY, FVector& WorldPosition, FVector& WorldDirection) const;
 
-
 	/**
 	 * Returns the array of actors inside a selection rectangle, with a class filter.
 	 *
@@ -458,8 +349,7 @@ public:
 	 *
 	 */
 	UFUNCTION(BlueprintPure, Category=HUD)
-	void GetActorsInSelectionRectangle(TSubclassOf<class AActor> ClassFilter, const FVector2D& FirstPoint, const FVector2D& SecondPoint, TArray<AActor*>& OutActors, bool bIncludeNonCollidingComponents = true, bool bActorMustBeFullyEnclosed = false);
-
+	void GetActorsInSelectionRectangle(TSubclassOf<AActor> ClassFilter, const FVector2D& FirstPoint, const FVector2D& SecondPoint, TArray<AActor*>& OutActors, bool bIncludeNonCollidingComponents = true, bool bActorMustBeFullyEnclosed = false);
 
 	/**
 	 * Add a hitbox to the hud
@@ -501,7 +391,7 @@ public:
 	void Draw2DLine(int32 X1, int32 Y1, int32 X2, int32 Y2, FColor LineColor);
 
 	/** Set the canvas and debug canvas to use during drawing */
-	void SetCanvas(class UCanvas* InCanvas, class UCanvas* InDebugCanvas);
+	void SetCanvas(UCanvas* InCanvas, UCanvas* InDebugCanvas);
 
 	virtual void PostInitializeComponents() override;
 
@@ -517,10 +407,10 @@ public:
 	************************************************************************************************************/
 
 	/** remove an actor from the PostRenderedActors array */
-	virtual void RemovePostRenderedActor(class AActor* A);
+	virtual void RemovePostRenderedActor(AActor* A);
 
 	/** add an actor to the PostRenderedActors array */
-	virtual void AddPostRenderedActor(class AActor* A);
+	virtual void AddPostRenderedActor(AActor* A);
 
 	/**
 	 * check if we should be display debug information for particular types of debug messages
@@ -546,10 +436,10 @@ public:
 	//=============================================================================
 
 	/** Display current messages */
-	virtual void DrawText(const FString& Text, FVector2D Position, class UFont* TextFont, FVector2D FontScale, FColor TextColor);
+	virtual void DrawText(const FString& Text, FVector2D Position, UFont* TextFont, FVector2D FontScale, FColor TextColor);
 	
 	/** @return UFont* from given FontSize index */
-	virtual class UFont* GetFontFromSizeIndex(int32 FontSize) const;
+	virtual  UFont* GetFontFromSizeIndex(int32 FontSize) const;
 
 	/**
 	 *	Pauses or unpauses the game due to main window's focus being lost.
@@ -570,7 +460,7 @@ public:
 	TArray< FHUDHitBox >	HitBoxMap;
 	
 	/** Array of hitboxes that have been hit for this frame. */
-	TArray< class FHUDHitBox* >	HitBoxHits;
+	TArray< FHUDHitBox* >	HitBoxHits;
 
 	/** Set of hitbox (by name) that are currently moused over or have a touch contacting them */
 	TSet< FName > HitBoxesOver;
diff --git a/Engine/Source/Runtime/Engine/Classes/GameFramework/HUDHitBox.h b/Engine/Source/Runtime/Engine/Classes/GameFramework/HUDHitBox.h
new file mode 100644
--- /dev/null
+++ b/Engine/Source/Runtime/Engine/Classes/GameFramework/HUDHitBox.h
@@ -0,0 +1,60 @@
+// Copyright 1998-2014 Epic Games, Inc. All Rights Reserved.
+
+#pragma once
+
+/**
+ * Hitbox use for clickable hud elements.
+ * These are named regions that receive various events when clicked upon.
+ * @see AHud
+ * @see AddHitBox(), ReceiveHitBoxClick(), ReceiveHitBoxRelease(), ReceiveHitBoxBeginCursorOver(), ReceiveHitBoxEndCursorOver()
+ * @see  GetHitBoxAtCoordinates(), GetHitBoxesAtCoordinates(), GetHitBoxWithName()
+ */
+class ENGINE_API FHUDHitBox
+{
+public:
+	/** 
+	 * Constructor for a hitbox.
+	 * @param	InCoords		Coordinates of top left of hitbox.
+	 * @param	InSize			Size of the box.
+	 */
+	FHUDHitBox( FVector2D InCoords, FVector2D InSize, const FName& InName, bool bInConsumesInput, int32 InPriority );
+
+	/** 
+	 * Are the given coordinates within this hitbox.
+	 * @param	InCoords		Coordinates to check.
+	 * @returns true if coordinates are within this hitbox.
+	 */
+	bool Contains( FVector2D InCoords ) const;
+
+	/** 
+	 * Debug render for this hitbox.
+	 * @param	InCanvas		Canvas on which to render.
+	 * @param	InColor			Color to render the box.
+	 */
+	void Draw( class FCanvas* InCanvas, const FLinearColor& InColor ) const;
+	
+	/** Get the name of this hitbox.  */
+	const FName& GetName() const { return Name;};
+
+	/** Should other boxes be processed if this box is clicked.  */
+	const bool ConsumesInput() const { return bConsumesInput; };
+
+	/** Get the priority of this hitbox.  */
+	const int32 GetPriority() const { return Priority; };
+
+private:
+	/** Coordinates of top left of hitbox.  */
+	FVector2D	Coords;
+
+	/** Size of this hitbox.  */
+	FVector2D	Size;
+	
+	/** The name of this hitbox.  */
+	FName		Name;	
+
+	/** Wether or not this hitbox should prevent hit checks to other hitboxes.  */
+	bool bConsumesInput;
+
+	/** The priority of this hitbox. Higher boxes are given priority. */
+	int32 Priority;
+};
diff --git a/Engine/Source/Runtime/Engine/Classes/GameFramework/SimpleReticle.h b/Engine/Source/Runtime/Engine/Classes/GameFramework/SimpleReticle.h
new file mode 100644
--- /dev/null
+++ b/Engine/Source/Runtime/Engine/Classes/GameFramework/SimpleReticle.h
@@ -0,0 +1,31 @@
+// Copyright 1998-2014 Epic Games, Inc. All Rights Reserved.
+
+#pragma once
+
+/**
+ * Simple reticle. A very simple crosshair reticle that can be draw on the HUD canvas. 
+ *
+ * @see AHUD
+ */
+class ENGINE_API FSimpleReticle
+{
+public:
+	FSimpleReticle()
+	{
+		SetupReticle( 8.0f, 20.0f );
+	}
+
+	void SetupReticle( const float Length, const float InnerSize )
+	{
+		HorizontalOffsetMin.Set( InnerSize, 0.0f );
+		HorizontalOffsetMax.Set( InnerSize + Length, 0.0f );
+		VerticalOffsetMin.Set( 0.0f, InnerSize);
+		VerticalOffsetMax.Set( 0.0f, InnerSize + Length );
+	}
+	void Draw( class UCanvas* InCanvas, FLinearColor InColor );
+private:
+	FVector2D HorizontalOffsetMin;
+	FVector2D HorizontalOffsetMax;
+	FVector2D VerticalOffsetMin;
+	FVector2D VerticalOffsetMax;	
+};
\ No newline at end of file
diff --git a/Engine/Source/Runtime/Engine/Private/GameFramework/HUDHitBox.cpp b/Engine/Source/Runtime/Engine/Private/GameFramework/HUDHitBox.cpp
new file mode 100644
--- /dev/null
+++ b/Engine/Source/Runtime/Engine/Private/GameFramework/HUDHitBox.cpp
@@ -0,0 +1,35 @@
+// Copyright 1998-2014 Epic Games, Inc. All Rights Reserved.
+
+#include "EnginePrivate.h"
+#include "GameFramework/HUDHitBox.h"
+
+FHUDHitBox::FHUDHitBox(FVector2D InCoords, FVector2D InSize, const FName& InName, bool bInConsumesInput, int32 InPriority)
+	: Coords(InCoords)
+	, Size(InSize)
+	, Name(InName)
+	, bConsumesInput(bInConsumesInput)
+	, Priority(InPriority)
+{
+}
+
+bool FHUDHitBox::Contains(FVector2D InCoords) const
+{
+	bool bResult = false;
+	if ((InCoords.X >= Coords.X) && (InCoords.X <= (Coords.X + Size.X)))
+	{
+		if ((InCoords.Y >= Coords.Y) && (InCoords.Y <= (Coords.Y + Size.Y)))
+		{
+			bResult = true;
+		}
+	}
+	return bResult;
+}
+
+void FHUDHitBox::Draw(FCanvas* InCanvas, const FLinearColor& InColor) const
+{
+	FCanvasBoxItem	BoxItem(Coords, Size);
+	BoxItem.SetColor(InColor);
+	InCanvas->DrawItem(BoxItem);
+	FCanvasTextItem	TextItem(Coords, FText::FromName(Name), GEngine->GetSmallFont(), InColor);
+	InCanvas->DrawItem(TextItem);
+}
diff --git a/Engine/Source/Runtime/Engine/Private/GameFramework/SimpleReticle.cpp b/Engine/Source/Runtime/Engine/Private/GameFramework/SimpleReticle.cpp
new file mode 100644
--- /dev/null
+++ b/Engine/Source/Runtime/Engine/Private/GameFramework/SimpleReticle.cpp
@@ -0,0 +1,15 @@
+// Copyright 1998-2014 Epic Games, Inc. All Rights Reserved.
+
+#include "EnginePrivate.h"
+#include "GameFramework/SimpleReticle.h"
+
+void FSimpleReticle::Draw( UCanvas* InCanvas, FLinearColor InColor )
+{
+	FVector2D CanvasCenter( InCanvas->OrgX + ((InCanvas->ClipX - InCanvas->OrgX) / 2.0f), InCanvas->OrgX + ((InCanvas->ClipY - InCanvas->OrgY) / 2.0f) );
+	FCanvasLineItem LineItem( CanvasCenter, FVector2D(0.0f, 0.0f) );
+	LineItem.SetColor( InColor );
+	LineItem.Draw( InCanvas->Canvas, CanvasCenter - HorizontalOffsetMin, CanvasCenter - HorizontalOffsetMax );
+	LineItem.Draw( InCanvas->Canvas, CanvasCenter + HorizontalOffsetMin, CanvasCenter + HorizontalOffsetMax );
+	LineItem.Draw( InCanvas->Canvas, CanvasCenter - VerticalOffsetMin, CanvasCenter - VerticalOffsetMax );
+	LineItem.Draw( InCanvas->Canvas, CanvasCenter + VerticalOffsetMin, CanvasCenter + VerticalOffsetMax );
+}
diff --git a/Engine/Source/Runtime/Engine/Private/HUD.cpp b/Engine/Source/Runtime/Engine/Private/HUD.cpp
--- a/Engine/Source/Runtime/Engine/Private/HUD.cpp
+++ b/Engine/Source/Runtime/Engine/Private/HUD.cpp
@@ -455,10 +455,6 @@ void AHUD::DrawDebugTextList()
 	}
 }
 
-/**
- * Add debug text for a specific actor to be displayed via DrawDebugTextList().  If the debug text is invalid then it will
- * attempt to remove any previous entries via RemoveDebugText().
- */
 void AHUD::AddDebugText_Implementation(const FString& DebugText,
 										 AActor* SrcActor,
 										 float Duration,
@@ -473,7 +469,7 @@ void AHUD::AddDebugText_Implementation(const FString& DebugText,
 										 )
 {
 	// set a default color
-	if (TextColor == FColor::Black)
+	if (TextColor == FLinearColor::Transparent)
 	{
 		TextColor = FColor::White;
 	}
@@ -965,48 +961,5 @@ bool AHUD::IsCanvasValid_WarnIfNot() const
 	return bIsValid;
 }
 
-/////////////////
-
-FHUDHitBox::FHUDHitBox( FVector2D InCoords, FVector2D InSize, const FName& InName, bool bInConsumesInput, int32 InPriority )
-	: Coords(InCoords)
-	, Size(InSize)
-	, Name(InName)
-	, bConsumesInput(bInConsumesInput)
-	, Priority(InPriority)
-{
-}
-
-bool FHUDHitBox::Contains( FVector2D InCoords ) const
-{
-	bool bResult = false;
-	if( ( InCoords.X >= Coords.X ) && (InCoords.X <= ( Coords.X + Size.X ) ) )
-	{
-		if( ( InCoords.Y >= Coords.Y ) && (InCoords.Y <= ( Coords.Y + Size.Y ) ) )
-		{
-			bResult = true;
-		}
-	}
-	return bResult;
-}
-
-void FHUDHitBox::Draw( FCanvas* InCanvas, const FLinearColor& InColor ) const
-{
-	FCanvasBoxItem	BoxItem( Coords, Size );
-	BoxItem.SetColor( InColor );
-	InCanvas->DrawItem( BoxItem );
-	FCanvasTextItem	TextItem( Coords, FText::FromName( Name ), GEngine->GetSmallFont(), InColor );
-	InCanvas->DrawItem( TextItem );
-}
-
-void FSimpleReticle::Draw( UCanvas* InCanvas, FLinearColor InColor )
-{
-	FVector2D CanvasCenter( InCanvas->OrgX + ((InCanvas->ClipX - InCanvas->OrgX) / 2.0f), InCanvas->OrgX + ((InCanvas->ClipY - InCanvas->OrgY) / 2.0f) );
-	FCanvasLineItem LineItem( CanvasCenter, FVector2D(0.0f, 0.0f) );
-	LineItem.SetColor( InColor );
-	LineItem.Draw( InCanvas->Canvas, CanvasCenter - HorizontalOffsetMin, CanvasCenter - HorizontalOffsetMax );
-	LineItem.Draw( InCanvas->Canvas, CanvasCenter + HorizontalOffsetMin, CanvasCenter + HorizontalOffsetMax );
-	LineItem.Draw( InCanvas->Canvas, CanvasCenter - VerticalOffsetMin, CanvasCenter - VerticalOffsetMax );
-	LineItem.Draw( InCanvas->Canvas, CanvasCenter + VerticalOffsetMin, CanvasCenter + VerticalOffsetMax );
-}
 
 #undef LOCTEXT_NAMESPACE
