greatly reduced amount of navmesh rendering invalidations (was causing fps drop)

diff --git a/Engine/Source/Runtime/Engine/Private/AI/Navigation/RecastNavMeshGenerator.cpp b/Engine/Source/Runtime/Engine/Private/AI/Navigation/RecastNavMeshGenerator.cpp
--- a/Engine/Source/Runtime/Engine/Private/AI/Navigation/RecastNavMeshGenerator.cpp
+++ b/Engine/Source/Runtime/Engine/Private/AI/Navigation/RecastNavMeshGenerator.cpp
@@ -2963,11 +2963,18 @@ void FRecastNavMeshGenerator::CancelBuild()
 
 void FRecastNavMeshGenerator::TickAsyncBuild(float DeltaSeconds)
 {
+	bool bRequestDrawingUpdate = false;
+
 #if	WITH_EDITOR
 	// Remove expired tiles
 	{
 		const double Timestamp = FPlatformTime::Seconds();
+		const int32 NumPreRemove = RecentlyBuiltTiles.Num();
+		
 		RecentlyBuiltTiles.RemoveAllSwap([&](const FTileTimestamp& Tile) { return (Timestamp - Tile.Timestamp) > 0.5; });
+
+		const int32 NumPostRemove = RecentlyBuiltTiles.Num();
+		bRequestDrawingUpdate = (NumPreRemove != NumPostRemove);
 	}
 #endif//WITH_EDITOR
 
@@ -2982,7 +2989,7 @@ void FRecastNavMeshGenerator::TickAsyncBuild(float DeltaSeconds)
 	{
 		// Invalidate active paths that go through regenerated tiles
 		DestNavMesh->InvalidateAffectedPaths(UpdatedTileIndices);
-		DestNavMesh->RequestDrawingUpdate();
+		bRequestDrawingUpdate = true;
 
 #if	WITH_EDITOR
 		// Store completed tiles with timestamps to have ability to distinguish during debug draw
@@ -2997,6 +3004,11 @@ void FRecastNavMeshGenerator::TickAsyncBuild(float DeltaSeconds)
 		}
 #endif//WITH_EDITOR
 	}
+
+	if (bRequestDrawingUpdate)
+	{
+		DestNavMesh->RequestDrawingUpdate();
+	}
 }
 
 void FRecastNavMeshGenerator::OnNavigationBoundsChanged()
diff --git a/Engine/Source/Runtime/Engine/Private/NavMeshRenderingComponent.cpp b/Engine/Source/Runtime/Engine/Private/NavMeshRenderingComponent.cpp
--- a/Engine/Source/Runtime/Engine/Private/NavMeshRenderingComponent.cpp
+++ b/Engine/Source/Runtime/Engine/Private/NavMeshRenderingComponent.cpp
@@ -439,15 +439,7 @@ void UNavMeshRenderingComponent::GatherData(struct FNavMeshSceneProxyData* Curre
 			DebugMeshData.ClusterColor = NavMeshRenderColor_RecastTileBeingRebuilt;
 			CurrentData->MeshBuilders.Add(DebugMeshData);
 
-			// if we got here it means there's something being built, or is fresh so we better redraw in some time
-			DECLARE_CYCLE_STAT(TEXT("FSimpleDelegateGraphTask.Requesting navmesh redraw"),
-				STAT_FSimpleDelegateGraphTask_RequestingNavmeshRedraw,
-				STATGROUP_TaskGraphTasks);
-
-			FSimpleDelegateGraphTask::CreateAndDispatchWhenReady(
-				FSimpleDelegateGraphTask::FDelegate::CreateUObject(NavMesh, &ARecastNavMesh::MarkComponentsRenderStateDirty),
-				GET_STATID(STAT_FSimpleDelegateGraphTask_RequestingNavmeshRedraw), NULL, ENamedThreads::GameThread
-			);
+			// updates should be requested by FRecastNavMeshGenerator::TickAsyncBuild after tiles were refreshed
 		}
 
 		if (NavMesh->bDrawClusters)
