Fixed bug in DrawDebugCylinder which was using an int32 for a float calculation

diff --git a/Engine/Source/Runtime/Engine/Private/DrawDebugHelpers.cpp b/Engine/Source/Runtime/Engine/Private/DrawDebugHelpers.cpp
--- a/Engine/Source/Runtime/Engine/Private/DrawDebugHelpers.cpp
+++ b/Engine/Source/Runtime/Engine/Private/DrawDebugHelpers.cpp
@@ -435,8 +435,8 @@ void DrawDebugCylinder(const UWorld* InWorld, FVector const& Start, FVector cons
 		// Rotate a point around axis to form cylinder segments
 		FVector Segment;
 		FVector P1, P2, P3, P4;
-		const int32 AngleInc = 360.f / Segments;
-		int32 Angle = AngleInc;
+		const float AngleInc = 360.f / Segments;
+		float Angle = AngleInc;
 
 		// Default for Axis is up
 		FVector Axis = (End - Start).SafeNormal();
@@ -449,8 +449,7 @@ void DrawDebugCylinder(const UWorld* InWorld, FVector const& Start, FVector cons
 		FVector Dummy;
 
 		Axis.FindBestAxisVectors(Perpendicular, Dummy);
-
-
+		
 		Segment = Perpendicular.RotateAngleAxis(0, Axis) * Radius;
 		P1 = Segment + Start;
 		P3 = Segment + End;
