Remove unused UWorld::bDebugStepExecution

diff --git a/Engine/Source/Editor/UnrealEd/Private/Kismet2/KismetDebugUtilities.cpp b/Engine/Source/Editor/UnrealEd/Private/Kismet2/KismetDebugUtilities.cpp
--- a/Engine/Source/Editor/UnrealEd/Private/Kismet2/KismetDebugUtilities.cpp
+++ b/Engine/Source/Editor/UnrealEd/Private/Kismet2/KismetDebugUtilities.cpp
@@ -395,7 +395,6 @@ void FKismetDebugUtilities::AttemptToBreakExecution(UBlueprint* BlueprintObj, co
 	{
 		// Pause the simulation
 		GUnrealEd->PlayWorld->bDebugPauseExecution = true;
-		GUnrealEd->PlayWorld->bDebugStepExecution = false;
 		GUnrealEd->PlayWorld->bDebugFrameStepExecution = false;
 	}
 	else
diff --git a/Engine/Source/Runtime/Engine/Classes/Engine/World.h b/Engine/Source/Runtime/Engine/Classes/Engine/World.h
--- a/Engine/Source/Runtime/Engine/Classes/Engine/World.h
+++ b/Engine/Source/Runtime/Engine/Classes/Engine/World.h
@@ -961,9 +961,6 @@ public:
 	// Kismet debugging flags - they can be only editor only, but they're uint32, so it doens't make much difference
 	uint32 bDebugPauseExecution:1;
 
-	/** @todo document */
-	uint32 bDebugStepExecution:1;
-
 	/** Indicates this scene always allows audio playback. */
 	uint32 bAllowAudioPlayback:1;
 
diff --git a/Engine/Source/Runtime/Engine/Private/LevelTick.cpp b/Engine/Source/Runtime/Engine/Private/LevelTick.cpp
--- a/Engine/Source/Runtime/Engine/Private/LevelTick.cpp
+++ b/Engine/Source/Runtime/Engine/Private/LevelTick.cpp
@@ -388,7 +388,7 @@ bool UWorld::IsPaused()
 	return ( (Info->Pauser != NULL && TimeSeconds >= PauseDelay) ||
 				(bRequestedBlockOnAsyncLoading && GetNetMode() == NM_Client) ||
 				(GEngine->ShouldCommitPendingMapChange(this)) ||
-				(IsPlayInEditor() && bDebugPauseExecution && !bDebugStepExecution) );
+				(IsPlayInEditor() && bDebugPauseExecution) );
 }
 
 
@@ -1261,8 +1261,6 @@ void UWorld::Tick( ELevelTick TickType, float DeltaSeconds )
 	}
 
 	// Finish up.
-	bDebugStepExecution = false; // If debugging, clean out step check
-
 	if(bDebugFrameStepExecution)
 	{
 		bDebugPauseExecution = true;
