Document AudioComponent.h
Move FSoundParseParameters out of AudioComponent.h and in to ActiveSound.h

diff --git a/Engine/Source/Runtime/Engine/Classes/Components/AudioComponent.h b/Engine/Source/Runtime/Engine/Classes/Components/AudioComponent.h
--- a/Engine/Source/Runtime/Engine/Classes/Components/AudioComponent.h
+++ b/Engine/Source/Runtime/Engine/Classes/Components/AudioComponent.h
@@ -23,21 +23,26 @@ struct FAudioComponentParam
 {
 	GENERATED_USTRUCT_BODY()
 
+	// Name of the parameter
 	UPROPERTY(EditAnywhere, BlueprintReadWrite, Category=AudioComponentParam)
 	FName ParamName;
 
+	// Value of the parameter when used as a float
 	UPROPERTY(EditAnywhere, BlueprintReadWrite, Category=AudioComponentParam)
 	float FloatParam;
 
+	// Value of the parameter when used as a boolean
 	UPROPERTY(EditAnywhere, BlueprintReadWrite, Category=AudioComponentParam)
 	bool BoolParam;
 	
+	// Value of the parameter when used as an integer
 	UPROPERTY(EditAnywhere, BlueprintReadWrite, Category=AudioComponentParam)
 	int32 IntParam;
 
 	UPROPERTY()
 	class UDEPRECATED_SoundNodeWave* WaveParam_DEPRECATED;
 
+	// Value of the parameter when used as a sound wave
 	UPROPERTY(EditAnywhere, BlueprintReadWrite, Category=AudioComponentParam)
 	class USoundWave* SoundWaveParam;
 
@@ -53,50 +58,18 @@ struct FAudioComponentParam
 
 };
 
-struct FSoundParseParameters
-{
-	FNotifyBufferFinishedHooks NotifyBufferFinishedHooks;
-
-	class USoundClass* SoundClass;
-	
-	FTransform Transform;
-	FVector Velocity;
-	
-	float Volume;
-	float VolumeMultiplier;
-
-	float Pitch;
-	float HighFrequencyGain;
-
-	float StartTime;
-	float OmniRadius;
-
-	uint32 bUseSpatialization:1;
-	uint32 bLooping:1;
-
-	FSoundParseParameters()
-		: SoundClass(NULL)
-		, Velocity(ForceInit)
-		, Volume(1.f)
-		, VolumeMultiplier(1.f)
-		, Pitch(1.f)
-		, HighFrequencyGain(1.f)
-		, StartTime(-1.f)
-		, OmniRadius(0.0f)
-		, bUseSpatialization(false)
-		, bLooping(false)
-	{
-	}
-};
-
 /**
- * Used to provide objects with audio.
+ * AudioComponent is used to play a Sound
+ *
+ * @see https://docs.unrealengine.com/latest/INT/Audio/Overview/index.html
+ * @see USoundBase
  */
 UCLASS(ClassGroup=(Audio, Common), hidecategories=(Object, ActorComponent, Physics, Rendering, Mobility, LOD), ShowCategories=Trigger, meta=(BlueprintSpawnableComponent))
 class ENGINE_API UAudioComponent : public USceneComponent
 {
 	GENERATED_UCLASS_BODY()
 
+	/** The sound to be played */
 	UPROPERTY(EditAnywhere, BlueprintReadWrite, Category=Sound)
 	class USoundBase* Sound;
 
@@ -128,6 +101,7 @@ class ENGINE_API UAudioComponent : public USceneComponent
 	UPROPERTY(EditAnywhere, BlueprintReadWrite, Category=Attenuation)
 	uint32 bAllowSpatialization:1;
 
+	/** Should the Attenuation Settings asset be used (false) or should the properties set directly on the component be used for attenuation properties */
 	UPROPERTY(EditAnywhere, BlueprintReadWrite, Category=Attenuation)
 	uint32 bOverrideAttenuation:1;
 
@@ -135,30 +109,39 @@ class ENGINE_API UAudioComponent : public USceneComponent
 	UPROPERTY(EditAnywhere, BlueprintReadWrite, Category = Sound)
 	uint32 bIsUISound : 1;
 
+	/** The lower bound to use when randomly determining a pitch multiplier */
 	UPROPERTY(EditAnywhere, BlueprintReadWrite, Category=Modulation)
 	float PitchModulationMin;
 
+	/** The upper bound to use when randomly determining a pitch multiplier */
 	UPROPERTY(EditAnywhere, BlueprintReadWrite, Category=Modulation)
 	float PitchModulationMax;
 
+	/** The lower bound to use when randomly determining a volume multiplier */
 	UPROPERTY(EditAnywhere, BlueprintReadWrite, Category=Modulation)
 	float VolumeModulationMin;
 
+	/** The upper bound to use when randomly determining a volume multiplier */
 	UPROPERTY(EditAnywhere, BlueprintReadWrite, Category=Modulation)
 	float VolumeModulationMax;
 
+	/** A volume multiplier to apply to sounds generated by this component */
 	UPROPERTY(EditAnywhere, BlueprintReadWrite, Category=Sound)
 	float VolumeMultiplier;
 
+	/** A pitch multiplier to apply to sounds generated by this component */
 	UPROPERTY(EditAnywhere, BlueprintReadWrite, Category=Sound)
 	float PitchMultiplier;
 
+	/** A multiplier to apply to the high frequency gain for sounds generated by this component */
 	UPROPERTY(EditAnywhere, BlueprintReadWrite, Category=Modulation)
 	float HighFrequencyGainMultiplier;
 
+	/** If bOverrideSettings is false, the asset to use to determine attenuation properties for sounds generated by this component */
 	UPROPERTY(EditAnywhere, BlueprintReadWrite, Category=Attenuation, meta=(EditCondition="!bOverrideAttenuation"))
 	class USoundAttenuation* AttenuationSettings;
 
+	/** If bOverrideSettings is true, the attenuation properties to use for sounds generated by this component */
 	UPROPERTY(EditAnywhere, BlueprintReadWrite, Category=Attenuation, meta=(EditCondition="bOverrideAttenuation"))
 	struct FAttenuationSettings AttenuationOverrides;
 
@@ -168,6 +151,7 @@ class ENGINE_API UAudioComponent : public USceneComponent
 	/** called when we finish playing audio, either because it played to completion or because a Stop() call turned it off early */
 	UPROPERTY(BlueprintAssignable)
 	FOnAudioFinished OnAudioFinished;
+
 	/** shadow delegate for non UObject subscribers */
 	FOnAudioFinishedNative OnAudioFinishedNative;
 
@@ -175,6 +159,7 @@ class ENGINE_API UAudioComponent : public USceneComponent
 	UPROPERTY()
 	FOnQueueSubtitles OnQueueSubtitles;
 
+	// Set what sound is played by this component
 	UFUNCTION(BlueprintCallable, Category="Audio|Components|Audio")
 	void SetSound( USoundBase* NewSound );
 
@@ -201,7 +186,6 @@ class ENGINE_API UAudioComponent : public USceneComponent
 	 * @param FadeVolumeLevel the percentage of the AudioComponents's calculated volume in which to fade to
 	 */
 	UFUNCTION(BlueprintCallable, Category="Audio|Components|Audio")
-
 	void FadeOut(float FadeOutDuration, float FadeVolumeLevel);
 
 	/** Start a sound playing on an audio component */
@@ -212,20 +196,28 @@ class ENGINE_API UAudioComponent : public USceneComponent
 	UFUNCTION(BlueprintCallable, Category="Audio|Components|Audio")
 	void Stop();
 
+	/** @return true if this component is currently playing a SoundCue. */
+	UFUNCTION(BlueprintCallable, Category="Audio|Components|Audio")
+	bool IsPlaying();
+
 	/** This will allow one to adjust the volume of an AudioComponent on the fly */
 	UFUNCTION(BlueprintCallable, Category="Audio|Components|Audio")
 	void AdjustVolume(float AdjustVolumeDuration, float AdjustVolumeLevel);
 
+	/**  Set a float instance parameter for use in sound cues played by this audio component */
 	UFUNCTION(BlueprintCallable, Category="Audio|Components|Audio")
 	void SetFloatParameter(FName InName, float InFloat);
 
+	/**  Set a sound wave instance parameter for use in sound cues played by this audio component */
 	UFUNCTION(BlueprintCallable, Category="Audio|Components|Audio")
 	void SetWaveParameter(FName InName, class USoundWave* InWave);
 
-	UFUNCTION(BlueprintCallable, Category="Audio|Components|Audio")
+	/** Set a boolean instance parameter for use in sound cues played by this audio component */
+	UFUNCTION(BlueprintCallable, Category="Audio|Components|Audio", meta=(FriendlyName="Set Boolean Parameter"))
 	void SetBoolParameter(FName InName, bool InBool);
 	
-	UFUNCTION(BlueprintCallable, Category="Audio|Components|Audio")
+	/** Set an integer instance parameter for use in sound cues played by this audio component */
+	UFUNCTION(BlueprintCallable, Category="Audio|Components|Audio", meta=(FriendlyName="Set Integer Parameter"))
 	void SetIntParameter(FName InName, int32 InInt);
 
 	/** Set a new volume multiplier */
@@ -236,9 +228,11 @@ class ENGINE_API UAudioComponent : public USceneComponent
 	UFUNCTION(BlueprintCallable, Category="Audio|Components|Audio")
 	void SetPitchMultiplier(float NewPitchMultiplier);
 
+	/** Set whether sounds generated by this audio component should be considered UI sounds */
 	UFUNCTION(BlueprintCallable, Category="Audio|Components|Audio")
 	void SetUISound(bool bInUISound);
 
+	/** Modify the attenuation settings of the audio component */
 	UFUNCTION(BlueprintCallable, Category="Audio|Components|Audio")
 	void AdjustAttenuation(const FAttenuationSettings& InAttenuationSettings);
 
@@ -256,10 +250,6 @@ public:
 	/** If true, this sound will not be stopped when flushing the audio device. */
 	uint32 bIgnoreForFlushing:1;
 
-	/**
-	 * Properties of the audio component set by its owning sound class
-	 */
-	
 	/** Whether audio effects are applied */
 	uint32 bEQFilterApplied:1;
 
@@ -295,11 +285,6 @@ public:
 	virtual void OnUpdateTransform(bool bSkipPhysicsMove) override;
 	// End USceneComponent Interface
 
-	/** @return true if this component is currently playing a SoundCue. */
-	UFUNCTION(BlueprintCallable, Category="Audio|Components|Audio")
-	bool IsPlaying();
-
-
 	// Begin ActorComponent interface.
 #if WITH_EDITORONLY_DATA
 	virtual void OnRegister() override;
@@ -307,18 +292,24 @@ public:
 	virtual void OnUnregister() override;
 	// End ActorComponent interface.
 
+	/** Returns a pointer to the attenuation settings to be used (if any) for this audio component dependent on the SoundAttenuation asset or overrides set. */
 	const FAttenuationSettings* GetAttenuationSettingsToApply() const;
+
+	/** Collects the various attenuation shapes that may be applied to the sound played by the audio component for visualization in the editor or via the in game debug visualization. */
 	void CollectAttenuationShapesForVisualization(TMultiMap<EAttenuationShape::Type, FAttenuationSettings::AttenuationShapeDetails>& ShapeDetailsMap) const;
 
 private:
 	
 #if WITH_EDITORONLY_DATA
+	/** Editor only component used to display the sprite so as to be able to see the location of the Audio Component  */
 	UPROPERTY(transient)
 	class UBillboardComponent* SpriteComponent;
 
+	/** Utility function that updates which texture is displayed on the sprite dependent on the properties of the Audio Component. */
 	void UpdateSpriteTexture();
 #endif
 
+	/** Utility function called by Play and FadeIn to start a sound playing. */
 	void PlayInternal(const float StartTime = 0.f, const float FadeInDuration = 0.f, const float FadeVolumeLevel = 1.f);
 };
 
diff --git a/Engine/Source/Runtime/Engine/Classes/Sound/SoundNode.h b/Engine/Source/Runtime/Engine/Classes/Sound/SoundNode.h
--- a/Engine/Source/Runtime/Engine/Classes/Sound/SoundNode.h
+++ b/Engine/Source/Runtime/Engine/Classes/Sound/SoundNode.h
@@ -103,7 +103,7 @@ public:
 	 */
 	virtual float GetDuration( void );
 
-	virtual void ParseNodes( FAudioDevice* AudioDevice, const UPTRINT NodeWaveInstanceHash, FActiveSound& ActiveSound, const FSoundParseParameters& ParseParams, TArray<FWaveInstance*>& WaveInstances );
+	virtual void ParseNodes( FAudioDevice* AudioDevice, const UPTRINT NodeWaveInstanceHash, FActiveSound& ActiveSound, const struct FSoundParseParameters& ParseParams, TArray<FWaveInstance*>& WaveInstances );
 
 	/**
 	 * Returns an array of all (not just active) nodes.
diff --git a/Engine/Source/Runtime/Engine/Public/ActiveSound.h b/Engine/Source/Runtime/Engine/Public/ActiveSound.h
--- a/Engine/Source/Runtime/Engine/Public/ActiveSound.h
+++ b/Engine/Source/Runtime/Engine/Public/ActiveSound.h
@@ -5,6 +5,56 @@
 #include "Sound/AudioVolume.h"
 #include "Sound/SoundAttenuation.h"
 
+/**
+ *	Struct used for gathering the final parameters to apply to a wave instance
+ */
+struct FSoundParseParameters
+{
+	// A collection of 
+	FNotifyBufferFinishedHooks NotifyBufferFinishedHooks;
+
+	// The Sound Class to use the settings of
+	class USoundClass* SoundClass;
+	
+	// The transform of the sound (scale is not used)
+	FTransform Transform;
+
+	// The speed that the sound is moving relative to the listener
+	FVector Velocity;
+	
+	float Volume;
+	float VolumeMultiplier;
+
+	float Pitch;
+	float HighFrequencyGain;
+
+	// How far in to the sound the
+	float StartTime;
+
+	// At what distance from the source of the sound should spatialization begin
+	float OmniRadius;
+
+	// Whether the sound should be spatialized
+	uint32 bUseSpatialization:1;
+
+	// Whether the sound should be seamlessly looped
+	uint32 bLooping:1;
+
+	FSoundParseParameters()
+		: SoundClass(NULL)
+		, Velocity(ForceInit)
+		, Volume(1.f)
+		, VolumeMultiplier(1.f)
+		, Pitch(1.f)
+		, HighFrequencyGain(1.f)
+		, StartTime(-1.f)
+		, OmniRadius(0.0f)
+		, bUseSpatialization(false)
+		, bLooping(false)
+	{
+	}
+};
+
 struct ENGINE_API FActiveSound
 {
 public:
