Document PlayerInput.h and InputComponent.h
Rename RelationshipType to ERelationshipType

diff --git a/Engine/Source/Runtime/Engine/Classes/Components/InputComponent.h b/Engine/Source/Runtime/Engine/Classes/Components/InputComponent.h
--- a/Engine/Source/Runtime/Engine/Classes/Components/InputComponent.h
+++ b/Engine/Source/Runtime/Engine/Classes/Components/InputComponent.h
@@ -5,7 +5,7 @@
 #include "InputComponent.generated.h"
 
 
-// Utility delegate class to allow binding to either a C++ function or a blueprint script delegate
+/** Utility delegate class to allow binding to either a C++ function or a blueprint script delegate */
 template<class DelegateType, class DynamicDelegateType>
 struct TInputUnifiedDelegate
 {
@@ -13,11 +13,13 @@ struct TInputUnifiedDelegate
 	TInputUnifiedDelegate(DelegateType const& D) : FuncDelegate(D) {};
 	TInputUnifiedDelegate(DynamicDelegateType const& D) : FuncDynDelegate(D) {};
 
+	/** Returns if either the native or dynamic delegate is bound */
 	inline bool IsBound() const
 	{
 		return ( FuncDelegate.IsBound() || FuncDynDelegate.IsBound() );
 	}
 
+	/** Returns if either the native or dynamic delegate is bound to an object */
 	inline bool IsBoundToObject(void const* Object) const
 	{
 		if (FuncDelegate.IsBound())
@@ -32,6 +34,7 @@ struct TInputUnifiedDelegate
 		return false;
 	}
 
+	/** Binds a native delegate and unbinds any bound dynamic delegate */
 	template< class UserClass >
 	inline void BindDelegate(UserClass* Object, typename DelegateType::template TUObjectMethodDelegate< UserClass >::FMethodPtr Func)
 	{
@@ -39,18 +42,21 @@ struct TInputUnifiedDelegate
 		FuncDelegate.BindUObject(Object, Func);
 	}
 
+	/** Binds a dynamic delegate and unbinds any bound native delegate */
 	inline void BindDelegate(UObject* Object, const FName FuncName)
 	{
 		FuncDelegate.Unbind();
 		FuncDynDelegate.BindUFunction(Object, FuncName);
 	}
 
-	DelegateType & GetDelegateForManualSet()
+	/** Returns a reference to the native delegate and unbinds any bound dynamic delegate */
+	DelegateType& GetDelegateForManualSet()
 	{
 		FuncDynDelegate.Unbind();
 		return FuncDelegate;
 	}
 
+	/** Unbinds any bound delegates */
 	inline void Unbind()
 	{
 		FuncDelegate.Unbind();
@@ -70,28 +76,37 @@ protected:
 };
 
 
-// An Input Chord is a key and the modifier keys that are to be held with it
+/** An Input Chord is a key and the modifier keys that are to be held with it. */
 USTRUCT()
 struct FInputChord
 {
 	GENERATED_USTRUCT_BODY()
 
+	/** The Key is the core of the chord. */
 	UPROPERTY()
 	FKey Key;
 
+	/** Whether the shift key is part of the chord.  */
 	UPROPERTY()
 	uint32 bShift:1;
 
+	/** Whether the control key is part of the chord.  */
 	UPROPERTY()
 	uint32 bCtrl:1;
 
+	/** Whether the alt key is part of the chord.  */
 	UPROPERTY()
 	uint32 bAlt:1;
 
+	/** Whether the command key is part of the chord.  */
 	UPROPERTY()
 	uint32 bCmd:1;
 
-	enum RelationshipType
+	/** 
+	  * The ways two chords can be related to each other. A chord is considered masking 
+	  * when it has all the same modifier keys as another chord plus more.
+	  */
+	enum ERelationshipType
 	{
 		None,
 		Same,
@@ -99,9 +114,8 @@ struct FInputChord
 		Masks
 	};
 
-	// Returns the relationship between this chord and another.  A chord is considered masking
-	// when it has all the same modifier keys as another chord plus more
-	RelationshipType GetRelationship(const FInputChord& OtherChord) const;
+	/** Returns the relationship between this chord and another. */
+	ERelationshipType GetRelationship(const FInputChord& OtherChord) const;
 
 	FInputChord()
 		: bShift(false)
@@ -128,14 +142,29 @@ struct FInputChord
 	}
 };
 
+/** Base class for the different binding types. */
+struct FInputBinding
+{
+	/** Whether the binding should consume the input or allow it to pass to another component */
+	uint32 bConsumeInput:1;
+
+	/** Whether the binding should execute while paused */
+	uint32 bExecuteWhenPaused:1;
+
+	FInputBinding()
+		: bConsumeInput(true)
+		, bExecuteWhenPaused(false)
+	{}
+};
 
-/** Delegate signature for action handlers. */
+/** Delegate signature for action events. */
 DECLARE_DELEGATE( FInputActionHandlerSignature );
 DECLARE_DYNAMIC_DELEGATE( FInputActionHandlerDynamicSignature );
 
-
+/** Unified delegate specialization for Input Actions. */
 struct FInputActionUnifiedDelegate : public TInputUnifiedDelegate<FInputActionHandlerSignature, FInputActionHandlerDynamicSignature>
 {
+	/** Execute function for the action unified delegate. */
 	inline void Execute() const
 	{
 		if (FuncDelegate.IsBound())
@@ -149,34 +178,19 @@ struct FInputActionUnifiedDelegate : public TInputUnifiedDelegate<FInputActionHa
 	}
 };
 
-
-struct FInputBinding
-{
-	/** Whether the binding should consume the input or allow it to pass to another component */
-	uint32 bConsumeInput:1;
-
-	/** Whether the binding should execute while paused */
-	uint32 bExecuteWhenPaused:1;
-
-	FInputBinding()
-		: bConsumeInput(true)
-		, bExecuteWhenPaused(false)
-	{}
-};
-
-
-/** Binds a game-defined action. */
+/** Binds a delegate to an action. */
 struct FInputActionBinding : public FInputBinding
 {
 	/** Friendly name of action, e.g "jump" */
 	FName ActionName;
 
-	/** Key event to bind it to, e.g. pressed, released, dblclick */
+	/** Key event to bind it to, e.g. pressed, released, double click */
 	TEnumAsByte<EInputEvent> KeyEvent;
 
 	/** Whether the binding is part of a paired (both pressed and released events bound) action */
 	uint32 bPaired:1;
 
+	/** The delegate bound to the action */
 	FInputActionUnifiedDelegate ActionDelegate;
 
 	FInputActionBinding()
@@ -194,15 +208,16 @@ struct FInputActionBinding : public FInputBinding
 	{ }
 };
 
-
+/** Binds a delegate to a key chord. */
 struct FInputKeyBinding : public FInputBinding
 {
-	/** Raw key to bind to */
+	/** Input Chord to bind to */
 	FInputChord Chord;
 
-	/** Key event to bind it to, e.g. pressed, released, double click */
+	/** Key event to bind it to (e.g. pressed, released, double click) */
 	TEnumAsByte<EInputEvent> KeyEvent;
 
+	/** The delegate bound to the key chord */
 	FInputActionUnifiedDelegate KeyDelegate;
 
 	FInputKeyBinding()
@@ -220,15 +235,16 @@ struct FInputKeyBinding : public FInputBinding
 
 /** 
  * Delegate signature for touch handlers. 
- * @FigerIndex: Which finger touched
+ * @FingerIndex: Which finger touched
  * @Location: The 2D screen location that was touched
  */
 DECLARE_DELEGATE_TwoParams( FInputTouchHandlerSignature, ETouchIndex::Type, FVector );
 DECLARE_DYNAMIC_DELEGATE_TwoParams( FInputTouchHandlerDynamicSignature, ETouchIndex::Type, FingerIndex, FVector, Location );
 
-
+/** Unified delegate specialization for Touch events. */
 struct FInputTouchUnifiedDelegate : public TInputUnifiedDelegate<FInputTouchHandlerSignature, FInputTouchHandlerDynamicSignature>
 {
+	/** Execute function for the touch unified delegate. */
 	inline void Execute(const ETouchIndex::Type FingerIndex, const FVector Location) const
 	{
 		if (FuncDelegate.IsBound())
@@ -242,12 +258,13 @@ struct FInputTouchUnifiedDelegate : public TInputUnifiedDelegate<FInputTouchHand
 	}
 };
 
-
+/** Binds a delegate to touch input. */
 struct FInputTouchBinding : public FInputBinding
 {
-	/** Key event to bind it to, e.g. pressed, released, double click */
+	/** Key event to bind it to (e.g. pressed, released, double click) */
 	TEnumAsByte<EInputEvent> KeyEvent;
 
+	/** The delegate bound to the touch events */
 	FInputTouchUnifiedDelegate TouchDelegate;
 
 	FInputTouchBinding()
@@ -271,9 +288,10 @@ struct FInputTouchBinding : public FInputBinding
 DECLARE_DELEGATE_OneParam( FInputAxisHandlerSignature, float );
 DECLARE_DYNAMIC_DELEGATE_OneParam( FInputAxisHandlerDynamicSignature, float, AxisValue );
 
-
+/** Unified delegate specialization for float axis events. */
 struct FInputAxisUnifiedDelegate : public TInputUnifiedDelegate<FInputAxisHandlerSignature, FInputAxisHandlerDynamicSignature>
 {
+	/** Execute function for the axis unified delegate. */
 	inline void Execute(const float AxisValue) const
 	{
 		if (FuncDelegate.IsBound())
@@ -288,13 +306,23 @@ struct FInputAxisUnifiedDelegate : public TInputUnifiedDelegate<FInputAxisHandle
 };
 
 
-/** Binds a game-defined axis to a function */
+/** Binds a delegate to an axis mapping. */
 struct FInputAxisBinding : public FInputBinding
 {
+	/** The axis mapping being bound to. */
 	FName AxisName;
 
+	/** 
+	 * The delegate bound to the axis. 
+	 * It will be called each frame that the input component is in the input stack 
+	 * regardless of whether the value is non-zero or has changed.
+	 */
 	FInputAxisUnifiedDelegate AxisDelegate;
 
+	/** 
+	 * The value of the axis as calculated during the most recent UPlayerInput::ProcessInputStack
+	 * if the InputComponent was in the stack, otherwise all values should be 0.
+	 */
 	float AxisValue;
 
 	FInputAxisBinding()
@@ -311,13 +339,23 @@ struct FInputAxisBinding : public FInputBinding
 };
 
 
-/** Binds a raw axis to a function */
+/** Binds a delegate to a raw float axis mapping. */
 struct FInputAxisKeyBinding : public FInputBinding
 {
+	/** The axis being bound to. */
 	FKey AxisKey;
 
+	/** 
+	 * The delegate bound to the axis. 
+	 * It will be called each frame that the input component is in the input stack 
+	 * regardless of whether the value is non-zero or has changed.
+	 */
 	FInputAxisUnifiedDelegate AxisDelegate;
 
+	/** 
+	 * The value of the axis as calculated during the most recent UPlayerInput::ProcessInputStack
+	 * if the InputComponent containing the binding was in the stack, otherwise the value will be 0.
+	 */
 	float AxisValue;
 
 	FInputAxisKeyBinding()
@@ -342,9 +380,10 @@ struct FInputAxisKeyBinding : public FInputBinding
 DECLARE_DELEGATE_OneParam(FInputVectorAxisHandlerSignature, FVector);
 DECLARE_DYNAMIC_DELEGATE_OneParam(FInputVectorAxisHandlerDynamicSignature, FVector, AxisValue);
 
-
+/** Unified delegate specialization for vector axis events. */
 struct FInputVectorAxisUnifiedDelegate : public TInputUnifiedDelegate<FInputVectorAxisHandlerSignature, FInputVectorAxisHandlerDynamicSignature>
 {
+	/** Execute function for the axis unified delegate. */
 	inline void Execute(const FVector AxisValue) const
 	{
 		if (FuncDelegate.IsBound())
@@ -358,13 +397,23 @@ struct FInputVectorAxisUnifiedDelegate : public TInputUnifiedDelegate<FInputVect
 	}
 };
 
-
+/** Binds a delegate to a raw vector axis mapping. */
 struct FInputVectorAxisBinding : public FInputBinding
 {
+	/** The axis being bound to. */
 	FKey AxisKey;
 
+	/** 
+	 * The delegate bound to the axis. 
+	 * It will be called each frame that the input component is in the input stack 
+	 * regardless of whether the value is non-zero or has changed.
+	 */
 	FInputVectorAxisUnifiedDelegate AxisDelegate;
 
+	/** 
+	 * The value of the axis as calculated during the most recent UPlayerInput::ProcessInputStack
+	 * if the InputComponent containing the binding was in the stack, otherwise the value will be (0,0,0).
+	 */
 	FVector AxisValue;
 
 	FInputVectorAxisBinding()
@@ -388,9 +437,10 @@ struct FInputVectorAxisBinding : public FInputBinding
 DECLARE_DELEGATE_OneParam( FInputGestureHandlerSignature, float );
 DECLARE_DYNAMIC_DELEGATE_OneParam( FInputGestureHandlerDynamicSignature, float, Value );
 
-
+/** Unified delegate specialization for gestureevents. */
 struct FInputGestureUnifiedDelegate : public TInputUnifiedDelegate<FInputGestureHandlerSignature, FInputGestureHandlerDynamicSignature>
 {
+	/** Execute function for the gesture unified delegate. */
 	inline void Execute(const float Value) const
 	{
 		if (FuncDelegate.IsBound())
@@ -408,8 +458,10 @@ struct FInputGestureUnifiedDelegate : public TInputUnifiedDelegate<FInputGesture
 /** Binds a gesture to a function. */
 struct FInputGestureBinding : public FInputBinding
 {
+	/** The gesture being bound to. */
 	FKey GestureKey;
 
+	/** The delegate bound to the gesture events */
 	FInputGestureUnifiedDelegate GestureDelegate;
 
 	/** Value parameter, meaning is dependent on the gesture. */
@@ -468,9 +520,11 @@ namespace EControllerAnalogStick
 /**
  * Implement an Actor component for input bindings.
  *
- * Input Component is a component that can be added to an Actor that allows the Actor to bind various kinds of input events
- * to delegate functions.  Input components are processed from a stack managed by the PlayerController and each bind can 
- * consume the input event preventing other components on the input stack from processing the input.
+ * An Input Component is a transient component that is to an Actor to bind various forms of input events to delegate functions.  
+ * Input components are processed from a stack managed by the PlayerController and processed by the PlayerInput.
+ * Each binding can consume the input event preventing other components on the input stack from processing the input.
+ *
+ * @see https://docs.unrealengine.com/latest/INT/Gameplay/Input/index.html
  */
 UCLASS(transient, config=Input, hidecategories=(Activation, "Components|Activation"))
 class ENGINE_API UInputComponent
diff --git a/Engine/Source/Runtime/Engine/Classes/GameFramework/PlayerInput.h b/Engine/Source/Runtime/Engine/Classes/GameFramework/PlayerInput.h
--- a/Engine/Source/Runtime/Engine/Classes/GameFramework/PlayerInput.h
+++ b/Engine/Source/Runtime/Engine/Classes/GameFramework/PlayerInput.h
@@ -15,44 +15,49 @@
 
 ENGINE_API DECLARE_LOG_CATEGORY_EXTERN(LogInput, Log, All);
 
-//
-// Forward declarations.
-//
-struct FDelegateDispatchDetails;
-
+/** Struct containing mappings for legacy method of binding keys to exec commands. */
 USTRUCT()
 struct FKeyBind
 {
 	GENERATED_USTRUCT_BODY()
 
+	/** The key to be bound to the command */
 	UPROPERTY(config)
 	FKey Key;
 
+	/** The command to execute when the key is pressed/released */
 	UPROPERTY(config)
 	FString Command;
 
+	/** Whether the control key needs to be held when the key event occurs */
 	UPROPERTY(config)
 	uint32 Control:1;
 
+	/** Whether the shift key needs to be held when the key event occurs */
 	UPROPERTY(config)
 	uint32 Shift:1;
 
+	/** Whether the alt key needs to be held when the key event occurs */
 	UPROPERTY(config)
 	uint32 Alt:1;
 
+	/** Whether the command key needs to be held when the key event occurs */
 	UPROPERTY(config)
 	uint32 Cmd:1;
 
-	/** if true, the bind will not be activated if the corresponding key is held down */
+	/** Whether the control key must not be held when the key event occurs */
 	UPROPERTY(config)
 	uint32 bIgnoreCtrl:1;
 
+	/** Whether the shift key must not be held when the key event occurs */
 	UPROPERTY(config)
 	uint32 bIgnoreShift:1;
 
+	/** Whether the alt key must not be held when the key event occurs */
 	UPROPERTY(config)
 	uint32 bIgnoreAlt:1;
 
+	/** Whether the command key must not be held when the key event occurs */
 	UPROPERTY(config)
 	uint32 bIgnoreCmd:1;
 
@@ -61,17 +66,17 @@ struct FKeyBind
 	}
 };
 
-/** Configurable properties for control axes, used to turn raw input value into game-useful values. */
+/** Configurable properties for control axes, used to transform raw input into game ready values. */
 USTRUCT()
 struct FInputAxisProperties
 {
 	GENERATED_USTRUCT_BODY()
 
-	/** For control axes such as analog sticks. */
+	/** What the dead zone of the axis is.  For control axes such as analog sticks. */
 	UPROPERTY(EditAnywhere, Category="Input")
 	float DeadZone;
 
-	/** Scaling factor. */
+	/** Scaling factor to multiply raw value by. */
 	UPROPERTY(EditAnywhere, Category="Input")
 	float Sensitivity;
 
@@ -92,20 +97,26 @@ struct FInputAxisProperties
 	
 };
 
-/** Data structure exposed to defaultproperties/inis for defining axis properties. */
+/** Configurable properties for control axes. */
 USTRUCT()
 struct FInputAxisConfigEntry
 {
 	GENERATED_USTRUCT_BODY()
 
+	/** Axis Key these properties apply to */
 	UPROPERTY(VisibleAnywhere, Category="Input")
 	FName AxisKeyName;
 
+	/** Properties for the Axis Key */
 	UPROPERTY(EditAnywhere, Category="Input")
 	struct FInputAxisProperties AxisProperties;
 };
 
-/** Defines a mapping between an action and key */
+/** 
+ * Defines a mapping between an action and key 
+ *
+ * @see https://docs.unrealengine.com/latest/INT/Gameplay/Input/index.html
+ */
 USTRUCT()
 struct FInputActionKeyMapping
 {
@@ -169,7 +180,11 @@ struct FInputActionKeyMapping
 	{}
 };
 
-/** Defines a mapping between an action and key */
+/** 
+ * Defines a mapping between an axis and key 
+ * 
+ * @see https://docs.unrealengine.com/latest/INT/Gameplay/Input/index.html
+**/
 USTRUCT()
 struct FInputAxisKeyMapping
 {
@@ -223,26 +238,11 @@ struct FInputAxisKeyMapping
 	{}
 };
 
-struct FActionKeyDetails
-{
-	TArray<FInputActionKeyMapping> Actions;
-	FInputChord CapturingChord;
-};
-
-struct FAxisKeyDetails
-{
-	TArray<FInputAxisKeyMapping> KeyMappings;
-	uint32 bInverted:1;
-
-	FAxisKeyDetails()
-		: bInverted(false)
-	{
-	}
-};
-
 /**
- * Object within PlayerController that manages player input.
- * Only spawned on the client in network games.
+ * Object within PlayerController that processes player input.
+ * Only exists on the client in network games.
+ *
+ * @see https://docs.unrealengine.com/latest/INT/Gameplay/Input/index.html
  */
 UCLASS(Within=PlayerController, config=Input, transient)
 class ENGINE_API UPlayerInput : public UObject
@@ -279,121 +279,139 @@ public:
 	/** This player's version of Axis Mappings */
 	TArray<struct FInputAxisKeyMapping> AxisMappings;
 
+	/** List of Axis Mappings that have been inverted */
 	UPROPERTY(config)
 	TArray<FName> InvertedAxis;
 
-	//=============================================================================
-	// Input related functions.
+	/** Exec function to change the mouse sensitivity */
 	UFUNCTION(exec)
 	void SetMouseSensitivity(const float Sensitivity);
 
+	/** Exec function to return the mouse sensitivity to its default value */
 	UFUNCTION(exec)
 	void SetMouseSensitivityToDefault();
 
+	/** Exec function to add a debug exec command */
 	UFUNCTION(exec)
 	void SetBind(FName BindName, const FString& Command);
 
 	/** Returns the mouse sensitivity along the X-axis, or the Y-axis, or 1.0 if none are known. */
 	float GetMouseSensitivity();
 
+	/** Returns whether an Axis Key is inverted */
 	bool GetInvertAxisKey(const FKey AxisKey);
+
+	/** Returns whether an Axis Mapping is inverted */
 	bool GetInvertAxis(const FName AxisName);
 
+	/** Exec function to invert an axis key */
 	UFUNCTION(exec)
 	void InvertAxisKey(const FKey AxisKey);
 
-	// Backwards compatibility exec function for people used to it instead of using InvertAxisKey
+	/** Backwards compatibility exec function for people used to it instead of using InvertAxisKey */
 	UFUNCTION(exec)
 	void InvertMouse();
 
+	/** Exec function to invert an axis mapping */
 	UFUNCTION(exec)
 	void InvertAxis(const FName AxisName);
 
-	//*************************************************************************************
-	// Normal gameplay execs
-	// Type the name of the exec function at the console to execute it
-
-	// Mouse smoothing
+	/** Exec function to reset mouse smoothing values */
 	UFUNCTION(exec)
 	void ClearSmoothing();
 
+	/** Add a player specific action mapping. */
 	void AddActionMapping(const FInputActionKeyMapping& KeyMapping);
+
+	/** Remove a player specific action mapping. */
 	void RemoveActionMapping(const FInputActionKeyMapping& KeyMapping);
 
+	/** Add a player specific axis mapping. */
 	void AddAxisMapping(const FInputAxisKeyMapping& KeyMapping);
+
+	/** Remove a player specific action mapping. */
 	void RemoveAxisMapping(const FInputAxisKeyMapping& KeyMapping);
 
+	/** Add an engine defined action mapping that cannot be remapped. */
 	static void AddEngineDefinedActionMapping(const FInputActionKeyMapping& ActionMapping);
+
+	/** Add an engine defined axis mapping that cannot be remapped. */
 	static void AddEngineDefinedAxisMapping(const FInputAxisKeyMapping& AxisMapping);
 
+	/** Clear the current cached key maps and rebuild from the source arrays. */
 	void ForceRebuildingKeyMaps(const bool bRestoreDefaults = false);
 
 protected:
+
+	/** Runtime struct that caches the list of mappings for a given Action Name and the capturing chord if applicable */
+	struct FActionKeyDetails
+	{
+		/** List of all action key mappings that correspond to the action name in the containing map */
+		TArray<FInputActionKeyMapping> Actions;
+
+		/** For paired actions only, this represents the chord that is currently held and when it is released will represent the release event */
+		FInputChord CapturingChord;
+	};
+
+	/** Runtime struct that caches the list of mappings for a given Axis Name and whether that axis is currently inverted */
+	struct FAxisKeyDetails
+	{
+		/** List of all axis key mappings that correspond to the axis name in the containing map */
+		TArray<FInputAxisKeyMapping> KeyMappings;
+
+		/** Whether this axis should invert its outputs */
+		uint32 bInverted:1;
+
+		FAxisKeyDetails()
+			: bInverted(false)
+		{
+		}
+	};
+
 	/** Internal structure for storing axis config data. */
 	TMap<FKey,FInputAxisProperties> AxisProperties;
+
+	/** Map of Action Name to details about the keys mapped to that action */
 	TMap<FName,FActionKeyDetails> ActionKeyMap;
+
+	/** Map of Axis Name to details about the keys mapped to that axis */
 	TMap<FName,FAxisKeyDetails> AxisKeyMap;
 
+	/** The current game view of each key */
 	TMap<FKey,FKeyState> KeyStateMap;
 
-	struct FDelegateDispatchDetails
-	{
-		uint32 EventIndex;
-		uint32 FoundIndex;
-
-		FInputActionUnifiedDelegate ActionDelegate;
-		const FInputActionBinding* SourceAction;
-		FInputChord Chord;
-		TEnumAsByte<EInputEvent> KeyEvent;
-
-		FInputTouchUnifiedDelegate TouchDelegate;
-		FVector TouchLocation;
-		ETouchIndex::Type FingerIndex;
-
-		FInputGestureUnifiedDelegate GestureDelegate;
-		float GestureValue;
-
-		FDelegateDispatchDetails(const uint32 InEventIndex, const uint32 InFoundIndex, const FInputChord& InChord, const FInputActionUnifiedDelegate& InDelegate, const EInputEvent InKeyEvent, const FInputActionBinding* InSourceAction = NULL)
-			: EventIndex(InEventIndex)
-			, FoundIndex(InFoundIndex)
-			, ActionDelegate(InDelegate)
-			, SourceAction(InSourceAction)
-			, Chord(InChord)
-			, KeyEvent(InKeyEvent)
-		{}
-
-		FDelegateDispatchDetails(const uint32 InEventIndex, const uint32 InFoundIndex, const FInputTouchUnifiedDelegate& InDelegate, const FVector InLocation, const ETouchIndex::Type InFingerIndex)
-			: EventIndex(InEventIndex)
-			, FoundIndex(InFoundIndex)
-			, TouchDelegate(InDelegate)
-			, TouchLocation(InLocation)
-			, FingerIndex(InFingerIndex)
-		{}
-
-		FDelegateDispatchDetails(const uint32 InEventIndex, const uint32 InFoundIndex, const FInputGestureUnifiedDelegate& InDelegate, const float InValue)
-			: EventIndex(InEventIndex)
-			, FoundIndex(InFoundIndex)
-			, GestureDelegate(InDelegate)
-			, GestureValue(InValue)
-		{}
-	};
-
 public:
 	
+	// Begin UObject Interface
 	virtual void PostInitProperties() override;
+	virtual UWorld* GetWorld() const override;
+	// End UObject Interface
 
+	/** Flushes the current key state. */
 	void FlushPressedKeys();
+
+	/** Handles a key input event.  Returns true if there is an action that handles the specified key. */
 	bool InputKey(FKey Key, enum EInputEvent Event, float AmountDepressed, bool bGamepad);
+
+	/** Handles an axis input event.  Returns true if a legacy key bind handled the input, otherwise false. */
 	bool InputAxis(FKey Key, float Delta, float DeltaTime, int32 NumSamples, bool bGamepad);
+
+	/** Handles a touch input event.  Returns true. */
 	bool InputTouch(uint32 Handle, ETouchType::Type Type, const FVector2D& TouchLocation, FDateTime DeviceTimestamp, uint32 TouchpadIndex);
+
+	/** Handles a motion input event.  Returns true. */
 	bool InputMotion(const FVector& Tilt, const FVector& RotationRate, const FVector& Gravity, const FVector& Acceleration);
+
+	/** Handles a gesture input event.  Returns true. */
 	bool InputGesture(const FKey Gesture, const EInputEvent Event, const float Value);
+
+	/** Per frame tick function. Primarily for gesture recognition */
 	void Tick(float DeltaTime);
 
-	/** Process the player's input. */
+	/** Process the frame's input events given the current input component stack. */
 	void ProcessInputStack(const TArray<UInputComponent*>& InputComponentStack, const float DeltaTime, const bool bGamePaused);
 
-	/** Rather than processing input, consume it and discard without doing anything useful with it.  Like calling PlayerInput() and ignoring all results. */
+	/** Rather than processing input, consume it and discard without doing anything useful with it.  Like calling ProcessInputStack() and ignoring all results. */
 	void DiscardPlayerInput();
 
 	/** 
@@ -421,6 +439,7 @@ public:
 	 */
 	void DisplayDebug(class UCanvas* Canvas, const FDebugDisplayInfo& DebugDisplay, float& YL, float& YPos);
 
+	/** @return true if InKey is currently held */
 	bool IsPressed( FKey InKey ) const;
 	
 	/** @return true if InKey went from up to down since player input was last processed. */
@@ -456,34 +475,22 @@ public:
 #if !UE_BUILD_SHIPPING
 	/**
 	 * Exec handler
-	 *
-	 * @param	InWorld	World context
-	 * @param	Cmd		Command to parse
-	 * @param	Ar		Output device to log to
-	 *
-	 * @return	true if command was handled, false otherwise
 	 */
 	bool Exec(UWorld* UInWorld, const TCHAR* Cmd,FOutputDevice& Ar);
 
-	/** @todo document */
+	/** Returns the command for a given key in the legacy binding system */
 	FString GetBind(FKey Key);
 
-	/** 
-	* Get the Exec key binding for the given command. 
-	* @param	ExecCommand		Command to find the key binding for
-	*
-	* @return	The key binding for the command.
-	*/
+	/** Get the legacy Exec key binding for the given command. */
 	FKeyBind GetExecBind(FString const& ExecCommand);
 
-	/** Execute input commands. */
-	void ExecInputCommands( UWorld* InWorld, const TCHAR* Cmd,class FOutputDevice& Ar);
+	/** Execute input commands within the legacy key binding system. */
+	void ExecInputCommands( UWorld* InWorld, const TCHAR* Cmd, class FOutputDevice& Ar);
 #endif
 
+	/** Returns the list of keys mapped to the specified Action Name */
 	const TArray<FInputActionKeyMapping>& GetKeysForAction(const FName ActionName);
 
-	class UWorld* GetWorld() const;
-
 protected:
 	/** 
 	 * Given raw keystate value, returns the "massaged" value. Override for any custom behavior,
@@ -504,24 +511,30 @@ protected:
 	 */
 	bool KeyEventOccurred(FKey Key, EInputEvent Event, TArray<uint32>& EventIndices) const;
 
-	/* Collects the chords and the delegates they invoke
+	/* Collects the chords and the delegates they invoke for an action binding
 	 * @param ActionBinding - the action to determine whether it occurred
 	 * @param FoundChords - the list of chord/delegate pairs to add to
-	 * @param KeysToConsume - optional array to collect the keys associated with this binding if the action occurred
+	 * @param KeysToConsume - array to collect the keys associated with this binding that should be consumed
 	 */
-	void GetChordsForAction(const FInputActionBinding& ActionBinding, const bool bGamePaused, TArray<FDelegateDispatchDetails>& FoundChords, TArray<FKey>& KeysToConsume);
+	void GetChordsForAction(const FInputActionBinding& ActionBinding, const bool bGamePaused, TArray<struct FDelegateDispatchDetails>& FoundChords, TArray<FKey>& KeysToConsume);
 
-	void GetChordForKey(const FInputKeyBinding& KeyBinding, const bool bGamePaused, TArray<FDelegateDispatchDetails>& FoundChords, TArray<FKey>& KeysToConsume);
+	/* Collects the chords and the delegates they invoke for a key binding
+	 * @param KeyBinding - the key to determine whether it occurred
+	 * @param FoundChords - the list of chord/delegate pairs to add to
+	 * @param KeysToConsume - array to collect the keys associated with this binding that should be consumed
+	 */
+	void GetChordForKey(const FInputKeyBinding& KeyBinding, const bool bGamePaused, TArray<struct FDelegateDispatchDetails>& FoundChords, TArray<FKey>& KeysToConsume);
 
 	/* Returns the summed values of all the components of this axis this frame
 	 * @param AxisBinding - the action to determine if it ocurred
-	 * @param KeysToConsume - optional array to collect the keys associated with this binding
+	 * @param KeysToConsume - array to collect the keys associated with this binding that should be consumed
 	 */
 	float DetermineAxisValue(const FInputAxisBinding& AxisBinding, const bool bGamePaused, TArray<FKey>& KeysToConsume);
 
+	/** Utility function to ensure the key mapping cache maps are built */
 	void ConditionalBuildKeyMappings();
 
-	/** Set Key to be ignored */
+	/** Set the Key consumed for the frame so that subsequent input components will not be notified they were pressed */
 	void ConsumeKey(FKey Key);
 
 	/** @return true if InKey is being consumed */
@@ -539,14 +552,18 @@ protected:
 	friend FGestureRecognizer;
 
 private:
+	/** Static empty array to be able to return from GetKeysFromAction when there are no keys mapped to the requested action name */
 	static const TArray<FInputActionKeyMapping> NoKeyMappings;
 
+	/** Action Mappings defined by engine systems that cannot be remapped by users */
 	static TArray<FInputActionKeyMapping> EngineDefinedActionMappings;
+
+	/** Axis Mappings defined by engine systems that cannot be remapped by users */
 	static TArray<FInputAxisKeyMapping> EngineDefinedAxisMappings;
 
-	// A counter used to track the order in which events occurred since the last time the input stack was processed
+	/** A counter used to track the order in which events occurred since the last time the input stack was processed */
 	uint32 EventCount;
 
-	// Cache the last time dilation so as to be able to clear smoothing when it changes
+	/** Cache the last time dilation so as to be able to clear smoothing when it changes */
 	float LastTimeDilation;
 };
diff --git a/Engine/Source/Runtime/Engine/Private/InputComponent.cpp b/Engine/Source/Runtime/Engine/Private/InputComponent.cpp
--- a/Engine/Source/Runtime/Engine/Private/InputComponent.cpp
+++ b/Engine/Source/Runtime/Engine/Private/InputComponent.cpp
@@ -6,9 +6,9 @@
 /* FInputChord interface
  *****************************************************************************/
 
-FInputChord::RelationshipType FInputChord::GetRelationship( const FInputChord& OtherChord ) const
+FInputChord::ERelationshipType FInputChord::GetRelationship( const FInputChord& OtherChord ) const
 {
-	RelationshipType Relationship = None;
+	ERelationshipType Relationship = None;
 
 	if (Key == OtherChord.Key)
 	{
diff --git a/Engine/Source/Runtime/Engine/Private/UserInterface/PlayerInput.cpp b/Engine/Source/Runtime/Engine/Private/UserInterface/PlayerInput.cpp
--- a/Engine/Source/Runtime/Engine/Private/UserInterface/PlayerInput.cpp
+++ b/Engine/Source/Runtime/Engine/Private/UserInterface/PlayerInput.cpp
@@ -21,6 +21,49 @@ const TArray<FInputActionKeyMapping> UPlayerInput::NoKeyMappings;
 TArray<FInputActionKeyMapping> UPlayerInput::EngineDefinedActionMappings;
 TArray<FInputAxisKeyMapping> UPlayerInput::EngineDefinedAxisMappings;
 
+/** Runtime struct that gathers up the different kinds of delegates that might be issued */
+struct FDelegateDispatchDetails
+{
+	uint32 EventIndex;
+	uint32 FoundIndex;
+
+	FInputActionUnifiedDelegate ActionDelegate;
+	const FInputActionBinding* SourceAction;
+	FInputChord Chord;
+	TEnumAsByte<EInputEvent> KeyEvent;
+
+	FInputTouchUnifiedDelegate TouchDelegate;
+	FVector TouchLocation;
+	ETouchIndex::Type FingerIndex;
+
+	FInputGestureUnifiedDelegate GestureDelegate;
+	float GestureValue;
+
+	FDelegateDispatchDetails(const uint32 InEventIndex, const uint32 InFoundIndex, const FInputChord& InChord, const FInputActionUnifiedDelegate& InDelegate, const EInputEvent InKeyEvent, const FInputActionBinding* InSourceAction = NULL)
+		: EventIndex(InEventIndex)
+		, FoundIndex(InFoundIndex)
+		, ActionDelegate(InDelegate)
+		, SourceAction(InSourceAction)
+		, Chord(InChord)
+		, KeyEvent(InKeyEvent)
+	{}
+
+	FDelegateDispatchDetails(const uint32 InEventIndex, const uint32 InFoundIndex, const FInputTouchUnifiedDelegate& InDelegate, const FVector InLocation, const ETouchIndex::Type InFingerIndex)
+		: EventIndex(InEventIndex)
+		, FoundIndex(InFoundIndex)
+		, TouchDelegate(InDelegate)
+		, TouchLocation(InLocation)
+		, FingerIndex(InFingerIndex)
+	{}
+
+	FDelegateDispatchDetails(const uint32 InEventIndex, const uint32 InFoundIndex, const FInputGestureUnifiedDelegate& InDelegate, const float InValue)
+		: EventIndex(InEventIndex)
+		, FoundIndex(InFoundIndex)
+		, GestureDelegate(InDelegate)
+		, GestureValue(InValue)
+	{}
+};
+
 UPlayerInput::UPlayerInput(const FObjectInitializer& ObjectInitializer)
 	: Super(ObjectInitializer)
 {
@@ -561,7 +604,7 @@ void UPlayerInput::GetChordsForAction(const FInputActionBinding& ActionBinding,
 					const FInputChord Chord(KeyMapping.Key, KeyMapping.bShift, KeyMapping.bCtrl, KeyMapping.bAlt, KeyMapping.bCmd);
 					for (int32 ChordIndex = FoundChords.Num() - 1; ChordIndex >= 0; --ChordIndex)
 					{
-						FInputChord::RelationshipType ChordRelationship = Chord.GetRelationship(FoundChords[ChordIndex].Chord);
+						FInputChord::ERelationshipType ChordRelationship = Chord.GetRelationship(FoundChords[ChordIndex].Chord);
 
 						if (ChordRelationship == FInputChord::Masks)
 						{
@@ -623,7 +666,7 @@ void UPlayerInput::GetChordForKey(const FInputKeyBinding& KeyBinding, const bool
 			// look through the found chords and determine if this is masked (or masks) anything in the array
 			for (int32 ChordIndex = FoundChords.Num() - 1; ChordIndex >= 0; --ChordIndex)
 			{
-				FInputChord::RelationshipType ChordRelationship = KeyBinding.Chord.GetRelationship(FoundChords[ChordIndex].Chord);
+				FInputChord::ERelationshipType ChordRelationship = KeyBinding.Chord.GetRelationship(FoundChords[ChordIndex].Chord);
 
 				if (ChordRelationship == FInputChord::Masks)
 				{
