Documentation of GameMode.h, InputSettings.h, and World.h

diff --git a/Engine/Source/Runtime/Engine/Classes/Engine/World.h b/Engine/Source/Runtime/Engine/Classes/Engine/World.h
--- a/Engine/Source/Runtime/Engine/Classes/Engine/World.h
+++ b/Engine/Source/Runtime/Engine/Classes/Engine/World.h
@@ -181,11 +181,13 @@ public:
 		return IsInTransition() && bSwitchedToDefaultMap;
 	}
 
+	/** @return the destination map that is being travelled to via seamless travel */
 	inline FString GetDestinationMapName() const
 	{
 		return (IsInTransition() ? PendingTravelURL.Map : TEXT(""));
 	}
 
+	/** @return the destination world that has been loaded asynchronously by the seamless travel handler. */
 	inline const UWorld* GetLoadedWorld() const
 	{
 		return LoadedWorld;
@@ -259,43 +261,48 @@ private:
 	static TArray<FName> LevelPackageNames;
 };
 
+/** Saved editor viewport state information */
 struct ENGINE_API FLevelViewportInfo
 {
+	/** Where the camera is positioned within the viewport. */
 	FVector CamPosition;
+
+	/** The camera's position within the viewport. */
 	FRotator CamRotation;
+
+	/** The zoom value  for orthographic mode. */
 	float CamOrthoZoom;
+
+	/** Whether camera settings have been systematically changed since the last level viewport update. */
 	bool CamUpdated;
 
 	FLevelViewportInfo()
+		: CamPosition(FVector::ZeroVector);
+		, CamRotation(FRotator::ZeroRotator);
+		, CamOrthoZoom(DEFAULT_ORTHOZOOM);
+		, CamUpdated(false);
 	{
-		CamPosition = FVector::ZeroVector;
-		CamRotation = FRotator::ZeroRotator;
-		CamOrthoZoom = DEFAULT_ORTHOZOOM;
-		CamUpdated = false;
 	}
 
 	FLevelViewportInfo(const FVector& InCamPosition, const FRotator& InCamRotation, float InCamOrthoZoom)
+		: CamPosition(InCamPosition);
+		, CamRotation(InCamRotationr);
+		, CamOrthoZoom(InCamOrthoZoom);
+		, CamUpdated(false);
 	{
-		CamPosition = InCamPosition;
-		CamRotation = InCamRotation;
-		CamOrthoZoom = InCamOrthoZoom;
-		CamUpdated = false;
 	}
 
 	friend FArchive& operator<<( FArchive& Ar, FLevelViewportInfo& I )
 	{
-		if ( Ar.IsLoading() )
-		{
-			I.CamUpdated = true;
-		}
-
 		Ar << I.CamPosition;
 		Ar << I.CamRotation;
 		Ar << I.CamOrthoZoom;
 
 		if ( Ar.IsLoading() )
 		{
-			if ( I.CamOrthoZoom == 0 )
+			I.CamUpdated = true;
+
+			if ( I.CamOrthoZoom == 0.f )
 			{
 				I.CamOrthoZoom = DEFAULT_ORTHOZOOM;
 			}
@@ -397,6 +404,7 @@ struct FEndClothSimulationFunction : public FTickFunction
 	virtual FString DiagnosticMessage();
 };
 
+/* Struct of optional parameters passed to SpawnActor function(s). */
 struct ENGINE_API FActorSpawnParameters
 {
 	FActorSpawnParameters()
@@ -469,6 +477,17 @@ struct ENGINE_API FWorldAsyncTraceState
 	int32 NextAvailableOverlapIndex;
 };
 
+/** 
+ * The World is the top level object representing a map or a sandbox in which Actors and Components will exist and be rendered.  
+ *
+ * A World can be a single Persistent Level with an optional list of streaming levels that are loaded and unloaded via volumes and blueprint functions
+ * or it can be a collection of levels organized with a World Composition.
+ *
+ * In a standalone game, generally only a single World exists except during seamless area transitions when both a destination and current world exists.
+ * In the editor many Worlds exist: The level being edited, each PIE instance, each editor tool which has an interactive rendered viewport, and many more.
+ *
+ */
+
 UCLASS(customConstructor, config=Engine)
 class ENGINE_API UWorld : public UObject, public FNetworkNotify
 {
@@ -516,11 +535,11 @@ class ENGINE_API UWorld : public UObject, public FNetworkNotify
 	UPROPERTY(Transient)
 	class AGameState*							GameState;
 
-	/** @todo document */
+	/** Instance of this world's game-specific networking management */
 	UPROPERTY(Transient)
 	class AGameNetworkManager*					NetworkManager;
 
-	/** Instance of UPhysicsCollisionHandler */
+	/** Instance of this world's game-specific physics collision handler */
 	UPROPERTY(Transient)
 	class UPhysicsCollisionHandler*				PhysicsCollisionHandler;
 
@@ -566,13 +585,11 @@ class ENGINE_API UWorld : public UObject, public FNetworkNotify
 	/** set for one tick after completely loading and initializing a new world
 	 * (regardless of whether it's LoadMap() or seamless travel)
 	 */
-	UPROPERTY(transient)
 	uint32 bWorldWasLoadedThisTick:1;
 
 	/**
 	 * Triggers a call to PostLoadMap() the next Tick, turns off loading movie if LoadMap() has been called.
 	 */
-	UPROPERTY(transient)
 	uint32 bTriggerPostLoadMap:1;
 
 private:
@@ -682,6 +699,7 @@ private:
 	/**	Objects currently being debugged in Kismet	*/
 	FBlueprintToDebuggedObjectMap BlueprintObjectsBeingDebugged;
 
+	/** Whether the render scene for this World should be created with HitProxies or not */
 	bool bRequiresHitProxies;
 
 	/** a delegate that broadcasts a notification whenever an actor is spawned */
@@ -763,6 +781,7 @@ private:
 	uint32 bBroadcastSelectionChange:1;
 #endif //WITH_EDITORONLY_DATA
 public:
+	/** The URL that was used when loading this World.																			*/
 	FURL										URL;
 
 	/** Interface to the FX system managing particles and related effects for this world.										*/
@@ -800,7 +819,7 @@ public:
 	bool										bPostTickComponentUpdate;
 
 	/** Counter for allocating game- unique controller player numbers															*/
-	int32											PlayerNum;
+	int32										PlayerNum;
 
 	/** Time in seconds (game time so we respect time dilation) since the last time we purged references to pending kill objects */
 	float										TimeSinceLastPendingKillPurge;
@@ -815,10 +834,10 @@ public:
 	bool										bIsWorldInitialized;
 	
 	/** Override, forcing level load requests to be allowed. < 0 == not allowed, 0 == have code choose, > 1 == force allow.		 */
-	int32											AllowLevelLoadOverride;
+	int32										AllowLevelLoadOverride;
 
 	/** Number of frames to delay Streaming Volume updating, useful if you preload a bunch of levels but the camera hasn't caught up yet (INDEX_NONE for infinite) */
-	int32											StreamingVolumeUpdateDelay;
+	int32										StreamingVolumeUpdateDelay;
 
 	/** Is level streaming currently frozen?																					*/
 	bool										bIsLevelStreamingFrozen;
@@ -852,10 +871,6 @@ public:
 	/** When non-'None', all line traces where the TraceTag match this will be drawn */
 	FName    DebugDrawTraceTag;
 
-	/*****************************************************************************************************/
-	/** Moved from WorldSettings properties - START 														**/
-	/*****************************************************************************************************/
-		
 	/** An array of post processing volumes, sorted in ascending order of priority.					*/
 	TArray< IInterface_PostProcessVolume * > PostProcessVolumes;
 
@@ -900,7 +915,7 @@ public:
 	/** The type of travel to perform next when doing a server travel */
 	ETravelType			NextTravelType;
 	
-	/** @todo document */
+	/** The URL to be used for the upcoming server travel */
 	FString NextURL;
 
 	/** Amount of time to wait before traveling to next map, gives clients time to receive final RPCs @see ServerTravelPause */
@@ -909,7 +924,7 @@ public:
 	/** array of levels that were loaded into this map via PrepareMapChange() / CommitMapChange() (to inform newly joining clients) */
 	TArray<FName> PreparingLevelNames;
 
-	/** @todo document */
+	/** Name of persistent level if we've loaded levels via CommitMapChange() that aren't normally in the StreamingLevels array (to inform newly joining clients) */
 	FName CommittedPersistentLevelName;
 
 	/**
@@ -940,13 +955,13 @@ public:
 	/** Whether the match has been started */
 	uint32 bMatchStarted:1;
 
-	/**  Only update players. */
+	/**  When ticking the world, only update players. */
 	uint32 bPlayersOnly:1;
 
-	/** Only update players.  Next frame will set bPlayersOnly */
+	/** Indicates that at the end the frame bPlayersOnly will be set to true. */
 	uint32 bPlayersOnlyPending:1;
 
-	/** Starting gameplay. */
+	/** Is the world in its actor initialization phase. */
 	uint32 bStartup:1;
 
 	/** Is the world being torn down */
@@ -971,10 +986,6 @@ public:
 	UPROPERTY(transient)
 	uint32 bAreConstraintsDirty:1;
 
-	/*****************************************************************************************************/
-	/** Moved from WorldSettings - END 																		**/
-	/*****************************************************************************************************/
-
 	/**
 	 * UWorld default constructor
 	 */
@@ -1452,6 +1463,7 @@ public:
 	/** Get an iterator for the list of CameraActors that auto-activate for PlayerControllers. */
 	FConstCameraActorIterator GetAutoActivateCameraIterator() const;
 	
+	/** Returns a reference to the game viewport displaying this world if one exists. */
 	UGameViewportClient* GetGameViewport() const;
 
 	DEPRECATED(4.3, "GetBrush is deprecated use GetDefaultBrush instead.")
@@ -1502,6 +1514,7 @@ public:
 	/** Creates a new physics scene for this world. */
 	void CreatePhysicsScene();
 
+	/** Returns a pointer to the physics scene for this world. */
 	FPhysScene* GetPhysicsScene() const { return PhysicsScene; }
 
 	/** Set the physics scene to use by this world */
@@ -1642,7 +1655,6 @@ public:
 	virtual bool Rename(const TCHAR* NewName = NULL, UObject* NewOuter = NULL, ERenameFlags Flags = REN_None) override;
 #endif
 	virtual void PostDuplicate(bool bDuplicateForPIE) override;
-
 	// End UObject Interface
 	
 	/**
@@ -1791,6 +1803,7 @@ public:
 	/** Updates this world's scene with the list of instances, and optionally updates each instance's uniform buffer. */
 	void UpdateParameterCollectionInstances(bool bUpdateInstanceUniformBuffers);
 
+	/** Struct containing a collection of optional parameters for initialization of a World. */
 	struct InitializationValues
 	{
 		InitializationValues()
@@ -1807,17 +1820,35 @@ public:
 		{
 		}
 
+		/** Should the scenes (physics, rendering) be created. */
 		uint32 bInitializeScenes:1;
+
+		/** Are sounds allowed to be generated from this world. */
 		uint32 bAllowAudioPlayback:1;
+
+		/** Should the render scene create hit proxies. */
 		uint32 bRequiresHitProxies:1;
+
+		/** Should the physics scene be created. bInitializeScenes must be true for this to be considered. */
 		uint32 bCreatePhysicsScene:1;
+
+		/** Should the navigation system be created for this world. */
 		uint32 bCreateNavigation:1;
+
+		/** Should the AI system be created for this world. */
 		uint32 bCreateAISystem:1;
+
+		/** Should physics be simulated in this world. */
 		uint32 bShouldSimulatePhysics:1;
+
+		/** Are collision trace calls valid within this world. */
 		uint32 bEnableTraceCollision:1;
+
+		/** Should actions performed to objects in this world be saved to the transaction buffer. */
 		uint32 bTransactional:1;
-		uint32 bCreateFXSystem:1;
 
+		/** Should the FX system be created for this world. */
+		uint32 bCreateFXSystem:1;
 
 		InitializationValues& InitializeScenes(const bool bInitialize) { bInitializeScenes = bInitialize; return *this; }
 		InitializationValues& AllowAudioPlayback(const bool bAllow) { bAllowAudioPlayback = bAllow; return *this; }
@@ -1910,7 +1941,7 @@ public:
 	void SendAllEndOfFrameUpdates(FGraphEventArray* OutCompletion = NULL);
 
 
-	/** @todo document */
+	/** Do per frame tick behaviors related to the network driver */
 	void TickNetClient( float DeltaSeconds );
 
 	/**
@@ -1972,7 +2003,7 @@ public:
 	 */
 	TArray<class ULevel*>& GetSelectedLevels();
 
-	/** @todo document */
+	/** Shrink level elements to their minimum size. */
 	void ShrinkLevel();
 #endif // WITH_EDITOR
 	
@@ -2011,19 +2042,30 @@ public:
 	 */
 	bool RemoveLevel( ULevel* InLevel );
 
-	/** @todo document */
+	/** Handle Exec/Console Commands related to the World */
 	bool Exec( UWorld* InWorld, const TCHAR* Cmd, FOutputDevice& Ar=*GLog );
 
-	/**
-	 * Exec command handlers
-	 */
+private:
+	/** Utility function to handle Exec/Console Commands related to the Trace Tags */
 	bool HandleTraceTagCommand( const TCHAR* Cmd, FOutputDevice& Ar );
+
+	/** Utility function to handle Exec/Console Commands related to persistent debug lines */
 	bool HandleFlushPersistentDebugLinesCommand( const TCHAR* Cmd, FOutputDevice& Ar );
+
+	/** Utility function to handle Exec/Console Commands related to logging actor counts */
 	bool HandleLogActorCountsCommand( const TCHAR* Cmd, FOutputDevice& Ar, UWorld* InWorld );
+
+	/** Utility function to handle Exec/Console Commands related to demo recording */
 	bool HandleDemoRecordCommand( const TCHAR* Cmd, FOutputDevice& Ar, UWorld* InWorld );
+
+	/** Utility function to handle Exec/Console Commands related to playing a demo recording*/
 	bool HandleDemoPlayCommand( const TCHAR* Cmd, FOutputDevice& Ar, UWorld* InWorld );
+
+	/** Utility function to handle Exec/Console Commands related to stopping demo playback */
 	bool HandleDemoStopCommand( const TCHAR* Cmd, FOutputDevice& Ar, UWorld* InWorld );
 
+public:
+
 	// Destroys the current demo net driver
 	void DestroyDemoNetDriver();
 
@@ -2036,7 +2078,7 @@ public:
 	/** @return true if this level is a server */
 	bool IsServer();
 
-	/** @todo document */
+	/** @return true if the world is in the paused state */
 	bool IsPaused();
 
 	/**
@@ -2070,32 +2112,37 @@ public:
 
 	AActor* SpawnActor( UClass* Class, FVector const* Location=NULL, FRotator const* Rotation=NULL, const FActorSpawnParameters& SpawnParameters = FActorSpawnParameters() );
 
-	// Templated version of SpawnActor that allows you to specify a class type
+	/** Templated version of SpawnActor that allows you to specify a class type via the template type */
 	template< class T >
 	T* SpawnActor( const FActorSpawnParameters& SpawnParameters = FActorSpawnParameters() )
 	{
 		return CastChecked<T>(SpawnActor(T::StaticClass(), NULL, NULL, SpawnParameters),ECastCheckedType::NullAllowed);
 	}
 
+	/** Templated version of SpawnActor that allows you to specify location and rotation in addition to class type via the template type */
 	template< class T >
 	T* SpawnActor( FVector const& Location, FRotator const& Rotation, const FActorSpawnParameters& SpawnParameters = FActorSpawnParameters() )
 	{
 		return CastChecked<T>(SpawnActor(T::StaticClass(), &Location, &Rotation, SpawnParameters),ECastCheckedType::NullAllowed);
 	}
 	
+	/** Templated version of SpawnActor that allows you to specify the class type via parameter while the return type is a parent class of that type */
 	template< class T >
 	T* SpawnActor( UClass* Class, const FActorSpawnParameters& SpawnParameters = FActorSpawnParameters() )
 	{
 		return CastChecked<T>(SpawnActor(Class, NULL, NULL, SpawnParameters),ECastCheckedType::NullAllowed);
 	}
 
+	/** 
+	 *  Templated version of SpawnActor that allows you to specify the rotation and location in addition
+	 *  class type via parameter while the return type is a parent class of that type 
+	 */
 	template< class T >
 	T* SpawnActor( UClass* Class, FVector const& Location, FRotator const& Rotation, const FActorSpawnParameters& SpawnParameters = FActorSpawnParameters() )
 	{
 		return CastChecked<T>(SpawnActor(Class, &Location, &Rotation, SpawnParameters),ECastCheckedType::NullAllowed);
 	}
-
-
+	
 	/**
 	* Spawns given class and returns class T pointer, forcibly sets world position. WILL NOT run Construction Script of Blueprints 
 	* to give caller an opportunity to set parameters beforehand.  Caller is responsible for invoking construction
@@ -2123,25 +2170,37 @@ public:
 		return (Class != NULL) ? Cast<T>(SpawnActor(Class, &Location, &Rotation, SpawnInfo )) : NULL;
 	}
 
-	/** Returns the current GameMode instance, valid only on server. */
+	/** 
+	 *  Returns the current GameMode instance cast to the template type.
+	 *  This can only return a valid pointer on the server. Will always return null on a client
+	 */
 	template< class T >
 	T* GetAuthGameMode() const
 	{
 		return Cast<T>(AuthorityGameMode);
 	}
-	AGameMode* GetAuthGameMode() const { return AuthorityGameMode; };
+
+	/** 
+	 *  Returns the current GameMode instance.
+	 *  This can only return a valid pointer on the server. Will always return null on a client
+	 */
+	AGameMode* GetAuthGameMode() const { return AuthorityGameMode; }
 	
-	/** Returns the current GameState. */
+	/** Returns the current GameState instance cast to the template type. */
 	template< class T >
 	T* GetGameState() const
 	{
 		return Cast<T>(GameState);
 	}
 
+	/** Returns the current GameState instance. */
+	AGameState* GetGameState() const { return GameState; }
+
+	/** Copies GameState properties from the GameMode. */
 	void CopyGameState(AGameMode* FromGameMode, AGameState* FromGameState);
 
 
-	/** @todo document */
+	/** Spawns a Brush Actor in the World */
 	ABrush*	SpawnBrush();
 
 	/** 
@@ -2197,8 +2256,6 @@ public:
 
 	// Begin FNetworkNotify interface
 	virtual EAcceptConnection::Type NotifyAcceptingConnection() override;
-
-	/** @todo document */
 	virtual void NotifyAcceptedConnection( class UNetConnection* Connection ) override;
 	virtual bool NotifyAcceptingChannel( class UChannel* Channel ) override;
 	virtual void NotifyControlMessage(UNetConnection* Connection, uint8 MessageType, class FInBunch& Bunch) override;
@@ -2217,6 +2274,7 @@ public:
 		return NetDriver;
 	}
 
+	/** Returns the net mode this world is running under */
 	ENetMode GetNetMode() const;
 
 	/**
@@ -2228,8 +2286,6 @@ public:
 		NetDriver = NewDriver;
 	}
 
-	void AssignActorNewNetGUID(AActor* Actor, bool bIsStatic);
-
 	/**
 	 * Sets the number of frames to delay Streaming Volume updating, 
 	 * useful if you preload a bunch of levels but the camera hasn't caught up yet 
@@ -2269,10 +2325,6 @@ public:
 
 public:
 
-	/*****************************************************************************************************/
-	/** Moved from WorldSettings properties - START 														**/
-	/*****************************************************************************************************/
-
 	/**
 	 * Finds the audio settings to use for a given view location, taking into account the world's default
 	 * settings and the audio volumes in the world.
@@ -2368,7 +2420,7 @@ public:
 	/** @return the current detail mode, like EDetailMode but can be outside of the range */
 	int32 GetDetailMode();
 
-	/** @todo document */
+	/** Updates the timer between garbage collection such that at the next opportunity garbage collection will be run. */
 	void ForceGarbageCollection( bool bFullPurge = false );
 
 	/** asynchronously loads the given levels in preparation for a streaming map transition.
@@ -2412,10 +2464,12 @@ public:
 		return LatentActionManager;
 	}
 
+	/** Sets the owning game instance for this world */
 	inline void SetGameInstance(UGameInstance* NewGI)
 	{
 		OwningGameInstance = NewGI;
 	}
+	/** Returns the owning game instance for this world */
 	inline UGameInstance* GetGameInstance() const
 	{
 		return OwningGameInstance;
@@ -2443,19 +2497,32 @@ public:
 	void GetLightMapsAndShadowMaps(ULevel* Level, TArray<UTexture2D*>& OutLightMapsAndShadowMaps);
 
 public:
+	/** Rename this world such that it has the prefix on names for the given PIE Instance ID */
 	void RenameToPIEWorld(int32 PIEInstanceID);
+
+	/** Given a PackageName and a PIE Instance ID return the name of that Package when being run as a PIE world */
 	static FString ConvertToPIEPackageName(const FString& PackageName, int32 PIEInstanceID);
+
+	/** Given a PackageName and a prefix type, get back to the original package name (i.e. the saved map name) */
 	static FString StripPIEPrefixFromPackageName(const FString& PackageName, const FString& Prefix);
+
+	/** Return the prefix for PIE packages given a PIE Instance ID */
 	static FString BuildPIEPackagePrefix(int32 PIEInstanceID);
+
+	/** Given a loaded editor UWorld, duplicate it for play in editor purposes with OwningWorld as the world with the persistent level. */
 	static UWorld* DuplicateWorldForPIE(const FString& PackageName, UWorld* OwningWorld);
+
+	/** Given a string, return that string with any PIE prefix removed */
 	static FString RemovePIEPrefix(const FString &Source);
+
+	/** Given a package, locate the UWorld contained within if one exists */
 	static UWorld* FindWorldInPackage(UPackage* Package);
 
 	/** If the specified package contains a redirector to a UWorld, that UWorld is returned. Otherwise, nullptr is returned. */
 	static UWorld* FollowWorldRedirectorInPackage(UPackage* Package, UObjectRedirector** OptionalOutRedirector = nullptr);
 };
 
-/** Global UWorld pointer */
+/** Global UWorld pointer. Use of this pointer should be avoided whenever possible. */
 extern ENGINE_API class UWorldProxy GWorld;
 
 /** World delegates */
diff --git a/Engine/Source/Runtime/Engine/Classes/GameFramework/GameMode.h b/Engine/Source/Runtime/Engine/Classes/GameFramework/GameMode.h
--- a/Engine/Source/Runtime/Engine/Classes/GameFramework/GameMode.h
+++ b/Engine/Source/Runtime/Engine/Classes/GameFramework/GameMode.h
@@ -34,24 +34,30 @@ struct FGameClassShortName
 {
 	GENERATED_USTRUCT_BODY()
 
+	/** Abbreviation that can be used as an alias for the class name */
 	UPROPERTY()
 	FString ShortName;
 
+	/** The class name to use when the alias is specified in a URL */
 	UPROPERTY()
 	FString GameClassName;
 };
 
-//=============================================================================
-// The GameMode defines the game being played. It governs the game rules, scoring, what actors
-// are allowed to exist in this game type, and who may enter the game.
-// A GameMode actor is instantiated when the level is initialized for gameplay in
-// C++ UGameEngine::LoadMap().  The class of this GameMode actor is determined by
-// (in order) either the URL ?game=xxx, or the
-// DefaultGameMode entry in the game's .ini file (in the /Script/Engine.Engine section).
-// The GameMode used can be overridden in the GameMode function GetGameModeClass(), called
-// on the game class picked by the above process.
-//
-//=============================================================================
+/**
+ * The GameMode defines the game being played. It governs the game rules, scoring, what actors
+ * are allowed to exist in this game type, and who may enter the game.
+ *
+ * It is only instanced on the server and will never exist on the client. 
+ *
+ * A GameMode actor is instantiated when the level is initialized for gameplay in
+ * C++ UGameEngine::LoadMap().  
+ * 
+ * The class of this GameMode actor is determined by (in order) either the URL ?game=xxx, 
+ * the GameMode Override value set in the World Settings, or the DefaultGameMode entry set 
+ * in the game's Project Settings.
+ *
+ * @see https://docs.unrealengine.com/latest/INT/Gameplay/Framework/GameMode/index.html
+ */
 UCLASS(config=Game, notplaceable, BlueprintType, Blueprintable, hidecategories=(Info, Rendering, MovementReplication, Replication, Actor))
 class ENGINE_API AGameMode : public AInfo
 {
diff --git a/Engine/Source/Runtime/Engine/Classes/GameFramework/InputSettings.h b/Engine/Source/Runtime/Engine/Classes/GameFramework/InputSettings.h
--- a/Engine/Source/Runtime/Engine/Classes/GameFramework/InputSettings.h
+++ b/Engine/Source/Runtime/Engine/Classes/GameFramework/InputSettings.h
@@ -5,13 +5,18 @@
 #include "InputSettings.generated.h"
 
 
+/**
+ * Project wide settings for input handling
+ * 
+ * @see https://docs.unrealengine.com/latest/INT/Gameplay/Input/index.html
+ */
 UCLASS(config=Input, defaultconfig)
 class ENGINE_API UInputSettings
 	: public UObject
 {
 	GENERATED_UCLASS_BODY()
 
-	/** List of Axis to be defined in ini file. These are put into AxisProperties for game use. */
+	/** Properties of Axis controls */
 	UPROPERTY(config, EditAnywhere, EditFixedSize, Category="Bindings", meta=(ToolTip="List of Axis Properties"), AdvancedDisplay)
 	TArray<struct FInputAxisConfigEntry> AxisConfig;
 
@@ -31,7 +36,7 @@ class ENGINE_API UInputSettings
 	UPROPERTY(config, EditAnywhere, Category="MouseProperties", AdvancedDisplay, meta=(editcondition="bEnableFOVScaling"))
 	float FOVScale;
 
-	/** If a key is pressed twice in this amount of time, it's considered a "double click" */
+	/** If a key is pressed twice in this amount of time it is considered a "double click" */
 	UPROPERTY(config, EditAnywhere, Category="MouseProperties", AdvancedDisplay)
 	float DoubleClickTime;
 
@@ -71,16 +76,28 @@ class ENGINE_API UInputSettings
 	virtual void PostInitProperties() override;
 	// End of UObject interface
 
+	/** Programmatically add an action mapping to the project defaults */
 	void AddActionMapping(const FInputActionKeyMapping& KeyMapping);
+
+	/** Programmatically remove an action mapping to the project defaults */
 	void RemoveActionMapping(const FInputActionKeyMapping& KeyMapping);
 
+	/** Programmatically add an axis mapping to the project defaults */
 	void AddAxisMapping(const FInputAxisKeyMapping& KeyMapping);
+
+	/** Programmatically remove an axis mapping to the project defaults */
 	void RemoveAxisMapping(const FInputAxisKeyMapping& KeyMapping);
 
+	/** Flush the current mapping values to the config file */
 	void SaveKeyMappings();
 
+	/** Populate a list of all defined action names */
 	void GetActionNames(TArray<FName>& ActionNames) const;
+
+	/** Populate a list of all defined axis names */
 	void GetAxisNames(TArray<FName>& AxisNames) const;
+
 private:
+	/** When changes are made to the default mappings, push those changes out to PlayerInput key maps */
 	void ForceRebuildKeymaps();
 };
