Make UGameUserSettings::ApplyNonResolutionSettings virtual

diff --git a/Engine/Source/Runtime/Engine/Classes/GameFramework/GameUserSettings.h b/Engine/Source/Runtime/Engine/Classes/GameFramework/GameUserSettings.h
--- a/Engine/Source/Runtime/Engine/Classes/GameFramework/GameUserSettings.h
+++ b/Engine/Source/Runtime/Engine/Classes/GameFramework/GameUserSettings.h
@@ -18,7 +18,7 @@ public:
 	DEPRECATED(4.4, "ApplySettings() is deprecated, use ApplySettings(bool) instead.")
 	virtual void ApplySettings();
 	
-	void ApplyNonResolutionSettings();
+	virtual void ApplyNonResolutionSettings();
 	void ApplyResolutionSettings(bool bCheckForCommandLineOverrides);
 	void ConditionallyOverrideResolutionSettings();
 
