Fix World compile errors

diff --git a/Engine/Source/Runtime/Engine/Classes/Engine/World.h b/Engine/Source/Runtime/Engine/Classes/Engine/World.h
--- a/Engine/Source/Runtime/Engine/Classes/Engine/World.h
+++ b/Engine/Source/Runtime/Engine/Classes/Engine/World.h
@@ -277,18 +277,18 @@ struct ENGINE_API FLevelViewportInfo
 	bool CamUpdated;
 
 	FLevelViewportInfo()
-		: CamPosition(FVector::ZeroVector);
-		, CamRotation(FRotator::ZeroRotator);
-		, CamOrthoZoom(DEFAULT_ORTHOZOOM);
-		, CamUpdated(false);
+		: CamPosition(FVector::ZeroVector)
+		, CamRotation(FRotator::ZeroRotator)
+		, CamOrthoZoom(DEFAULT_ORTHOZOOM)
+		, CamUpdated(false)
 	{
 	}
 
 	FLevelViewportInfo(const FVector& InCamPosition, const FRotator& InCamRotation, float InCamOrthoZoom)
-		: CamPosition(InCamPosition);
-		, CamRotation(InCamRotationr);
-		, CamOrthoZoom(InCamOrthoZoom);
-		, CamUpdated(false);
+		: CamPosition(InCamPosition)
+		, CamRotation(InCamRotation)
+		, CamOrthoZoom(InCamOrthoZoom)
+		, CamUpdated(false)
 	{
 	}
 
