Engine: Fixed non-unity build

diff --git a/Engine/Source/Runtime/Engine/Private/EnginePrivate.h b/Engine/Source/Runtime/Engine/Private/EnginePrivate.h
--- a/Engine/Source/Runtime/Engine/Private/EnginePrivate.h
+++ b/Engine/Source/Runtime/Engine/Private/EnginePrivate.h
@@ -183,6 +183,7 @@
 #include "GameFramework/TouchInterface.h"
 #include "Engine/UserDefinedStruct.h"
 #include "Engine/World.h"
+#include "ActiveSound.h"
 
 #include "SystemSettings.h"					// Scalability options.
 #include "SceneManagement.h"				// Scene management.
