Engine: fix non-unity build.

diff --git a/Engine/Source/Runtime/Engine/Private/SoundNodeDialoguePlayer.cpp b/Engine/Source/Runtime/Engine/Private/SoundNodeDialoguePlayer.cpp
--- a/Engine/Source/Runtime/Engine/Private/SoundNodeDialoguePlayer.cpp
+++ b/Engine/Source/Runtime/Engine/Private/SoundNodeDialoguePlayer.cpp
@@ -4,6 +4,7 @@
 #include "Sound/DialogueWave.h"
 #include "Sound/SoundBase.h"
 #include "Sound/SoundNodeDialoguePlayer.h"
+#include "Components/AudioComponent.h"
 
 USoundNodeDialoguePlayer::USoundNodeDialoguePlayer(const FObjectInitializer& ObjectInitializer)
 	: Super(ObjectInitializer)
diff --git a/Engine/Source/Runtime/Engine/Private/SoundNodeWavePlayer.cpp b/Engine/Source/Runtime/Engine/Private/SoundNodeWavePlayer.cpp
--- a/Engine/Source/Runtime/Engine/Private/SoundNodeWavePlayer.cpp
+++ b/Engine/Source/Runtime/Engine/Private/SoundNodeWavePlayer.cpp
@@ -4,6 +4,7 @@
 #include "ActiveSound.h"
 #include "Sound/SoundNodeWavePlayer.h"
 #include "Sound/SoundWave.h"
+#include "Components/AudioComponent.h"
 
 #define LOCTEXT_NAMESPACE "SoundNodeWavePlayer"
 
