Roll out incorrect/unnecessary fixes in CL#2344052 and properly fix SoundNodeDialoguePlayer.cpp

diff --git a/Engine/Source/Runtime/Engine/Private/SoundNodeDialoguePlayer.cpp b/Engine/Source/Runtime/Engine/Private/SoundNodeDialoguePlayer.cpp
--- a/Engine/Source/Runtime/Engine/Private/SoundNodeDialoguePlayer.cpp
+++ b/Engine/Source/Runtime/Engine/Private/SoundNodeDialoguePlayer.cpp
@@ -1,10 +1,10 @@
 // Copyright 1998-2014 Epic Games, Inc. All Rights Reserved.
 
 #include "EnginePrivate.h"
+#include "ActiveSound.h"
 #include "Sound/DialogueWave.h"
 #include "Sound/SoundBase.h"
 #include "Sound/SoundNodeDialoguePlayer.h"
-#include "Components/AudioComponent.h"
 
 USoundNodeDialoguePlayer::USoundNodeDialoguePlayer(const FObjectInitializer& ObjectInitializer)
 	: Super(ObjectInitializer)
diff --git a/Engine/Source/Runtime/Engine/Private/SoundNodeWavePlayer.cpp b/Engine/Source/Runtime/Engine/Private/SoundNodeWavePlayer.cpp
--- a/Engine/Source/Runtime/Engine/Private/SoundNodeWavePlayer.cpp
+++ b/Engine/Source/Runtime/Engine/Private/SoundNodeWavePlayer.cpp
@@ -4,7 +4,6 @@
 #include "ActiveSound.h"
 #include "Sound/SoundNodeWavePlayer.h"
 #include "Sound/SoundWave.h"
-#include "Components/AudioComponent.h"
 
 #define LOCTEXT_NAMESPACE "SoundNodeWavePlayer"
 
