Don't use TSharedPtr for FComponentInstanceData

diff --git a/Engine/Source/Runtime/Engine/Classes/Atmosphere/AtmosphericFogComponent.h b/Engine/Source/Runtime/Engine/Classes/Atmosphere/AtmosphericFogComponent.h
--- a/Engine/Source/Runtime/Engine/Classes/Atmosphere/AtmosphericFogComponent.h
+++ b/Engine/Source/Runtime/Engine/Classes/Atmosphere/AtmosphericFogComponent.h
@@ -258,9 +258,9 @@ public:
 	ENGINE_API void ReleaseResource();
 
 	// Begin UActorComponent interface.
-	virtual TSharedPtr<FComponentInstanceDataBase> GetComponentInstanceData() const override;
+	virtual FComponentInstanceDataBase* GetComponentInstanceData() const override;
 	virtual FName GetComponentInstanceDataType() const override;
-	virtual void ApplyComponentInstanceData(TSharedPtr<FComponentInstanceDataBase> ComponentInstanceData) override;
+	virtual void ApplyComponentInstanceData(FComponentInstanceDataBase* ComponentInstanceData) override;
 	// End UActorComponent interface.
 
 	const FAtmospherePrecomputeParameters& GetPrecomputeParameters() const { return PrecomputeParams;  }
diff --git a/Engine/Source/Runtime/Engine/Classes/Components/ActorComponent.h b/Engine/Source/Runtime/Engine/Classes/Components/ActorComponent.h
--- a/Engine/Source/Runtime/Engine/Classes/Components/ActorComponent.h
+++ b/Engine/Source/Runtime/Engine/Classes/Components/ActorComponent.h
@@ -429,13 +429,13 @@ public:
 	virtual void PostNetReceive() { }
 
 	/** Called before we throw away components during RerunConstructionScripts, to cache any data we wish to persist across that operation */
-	virtual TSharedPtr<class FComponentInstanceDataBase> GetComponentInstanceData() const { return NULL; }
+	virtual class FComponentInstanceDataBase* GetComponentInstanceData() const { return NULL; }
 
 	/** The type of the component instance data that this component is interested in */
 	virtual FName GetComponentInstanceDataType() const { return NAME_None; }
 
 	/** Called after we create new components during RerunConstructionScripts, to optionally apply any data backed up during GetComponentInstanceData */
-	virtual void ApplyComponentInstanceData(TSharedPtr<class FComponentInstanceDataBase> ComponentInstanceData ) {}
+	virtual void ApplyComponentInstanceData(class FComponentInstanceDataBase* ComponentInstanceData ) {}
 
 	// Begin UObject interface.
 	virtual void BeginDestroy() override;
diff --git a/Engine/Source/Runtime/Engine/Classes/Components/ChildActorComponent.h b/Engine/Source/Runtime/Engine/Classes/Components/ChildActorComponent.h
--- a/Engine/Source/Runtime/Engine/Classes/Components/ChildActorComponent.h
+++ b/Engine/Source/Runtime/Engine/Classes/Components/ChildActorComponent.h
@@ -17,22 +17,22 @@ class UChildActorComponent : public USceneComponent
 	TSubclassOf<AActor>	ChildActorClass;
 
 	/** The actor that we spawned and own */
-	UPROPERTY(BlueprintReadOnly, Category = ChildActorComponent, TextExportTransient)
+	UPROPERTY(BlueprintReadOnly, Category=ChildActorComponent, TextExportTransient)
 	AActor*	ChildActor;
 
 	/** We try to keep the child actor's name as best we can, so we store it off here when destroying */
 	FName ChildActorName;
 
 	/** Cached copy of the instance data when the ChildActor is destroyed to be available when needed */
-	TSharedPtr<FChildActorComponentInstanceData> CachedInstanceData;
+	mutable FChildActorComponentInstanceData* CachedInstanceData;
 
 	// Begin ActorComponent interface.
 	virtual void OnComponentCreated() override;
 	virtual void OnComponentDestroyed() override;
 	virtual void OnRegister() override;
-	virtual TSharedPtr<FComponentInstanceDataBase> GetComponentInstanceData() const override;
+	virtual FComponentInstanceDataBase* GetComponentInstanceData() const override;
 	virtual FName GetComponentInstanceDataType() const override;
-	virtual void ApplyComponentInstanceData(TSharedPtr<FComponentInstanceDataBase> ComponentInstanceData) override;
+	virtual void ApplyComponentInstanceData(FComponentInstanceDataBase* ComponentInstanceData) override;
 
 	// End ActorComponent interface.
 
diff --git a/Engine/Source/Runtime/Engine/Classes/Components/InstancedStaticMeshComponent.h b/Engine/Source/Runtime/Engine/Classes/Components/InstancedStaticMeshComponent.h
--- a/Engine/Source/Runtime/Engine/Classes/Components/InstancedStaticMeshComponent.h
+++ b/Engine/Source/Runtime/Engine/Classes/Components/InstancedStaticMeshComponent.h
@@ -129,9 +129,9 @@ public:
 	TArray<FBodyInstance*> InstanceBodies;
 
 	// Begin UActorComponent interface 
-	virtual TSharedPtr<FComponentInstanceDataBase> GetComponentInstanceData() const override;
+	virtual FComponentInstanceDataBase* GetComponentInstanceData() const override;
 	virtual FName GetComponentInstanceDataType() const override;
-	virtual void ApplyComponentInstanceData(TSharedPtr<FComponentInstanceDataBase> ComponentInstanceData) override;
+	virtual void ApplyComponentInstanceData(FComponentInstanceDataBase* ComponentInstanceData) override;
 	// End UActorComponent interface 
 
 	// Begin UPrimitiveComponent Interface
diff --git a/Engine/Source/Runtime/Engine/Classes/Components/LightComponent.h b/Engine/Source/Runtime/Engine/Classes/Components/LightComponent.h
--- a/Engine/Source/Runtime/Engine/Classes/Components/LightComponent.h
+++ b/Engine/Source/Runtime/Engine/Classes/Components/LightComponent.h
@@ -278,9 +278,9 @@ public:
 	virtual bool IsReadyForFinishDestroy() override;
 	// End UObject interface.
 
-	virtual TSharedPtr<FComponentInstanceDataBase> GetComponentInstanceData() const override;
+	virtual FComponentInstanceDataBase* GetComponentInstanceData() const override;
 	virtual FName GetComponentInstanceDataType() const override;
-	virtual void ApplyComponentInstanceData(TSharedPtr<FComponentInstanceDataBase> ComponentInstanceData) override;
+	virtual void ApplyComponentInstanceData(FComponentInstanceDataBase* ComponentInstanceData) override;
 
 	/** @return number of material elements in this primitive */
 	virtual int32 GetNumMaterials() const;
diff --git a/Engine/Source/Runtime/Engine/Classes/Components/SkyLightComponent.h b/Engine/Source/Runtime/Engine/Classes/Components/SkyLightComponent.h
--- a/Engine/Source/Runtime/Engine/Classes/Components/SkyLightComponent.h
+++ b/Engine/Source/Runtime/Engine/Classes/Components/SkyLightComponent.h
@@ -147,9 +147,9 @@ class ENGINE_API USkyLightComponent : public ULightComponentBase
 	virtual bool IsReadyForFinishDestroy() override;
 	// End UObject Interface
 
-	virtual TSharedPtr<FComponentInstanceDataBase> GetComponentInstanceData() const override;
+	virtual FComponentInstanceDataBase* GetComponentInstanceData() const override;
 	virtual FName GetComponentInstanceDataType() const override;
-	virtual void ApplyComponentInstanceData(TSharedPtr<FComponentInstanceDataBase> ComponentInstanceData) override;
+	virtual void ApplyComponentInstanceData(FComponentInstanceDataBase* ComponentInstanceData) override;
 
 	/** Called each tick to recapture and queued sky captures. */
 	static void UpdateSkyCaptureContents(UWorld* WorldToUpdate);
diff --git a/Engine/Source/Runtime/Engine/Classes/Components/SplineComponent.h b/Engine/Source/Runtime/Engine/Classes/Components/SplineComponent.h
--- a/Engine/Source/Runtime/Engine/Classes/Components/SplineComponent.h
+++ b/Engine/Source/Runtime/Engine/Classes/Components/SplineComponent.h
@@ -45,9 +45,9 @@ private:
 
 public:
 	// Begin UActorComponent interface.
-	virtual TSharedPtr<FComponentInstanceDataBase> GetComponentInstanceData() const override;
+	virtual FComponentInstanceDataBase* GetComponentInstanceData() const override;
 	virtual FName GetComponentInstanceDataType() const override;
-	virtual void ApplyComponentInstanceData(TSharedPtr<FComponentInstanceDataBase> ComponentInstanceData) override;
+	virtual void ApplyComponentInstanceData(FComponentInstanceDataBase* ComponentInstanceData) override;
 	// End UActorComponent interface.
 
 	/** Update the spline tangents and SplineReparamTable */
diff --git a/Engine/Source/Runtime/Engine/Classes/Components/StaticMeshComponent.h b/Engine/Source/Runtime/Engine/Classes/Components/StaticMeshComponent.h
--- a/Engine/Source/Runtime/Engine/Classes/Components/StaticMeshComponent.h
+++ b/Engine/Source/Runtime/Engine/Classes/Components/StaticMeshComponent.h
@@ -253,9 +253,9 @@ public:
 #if WITH_EDITOR
 	virtual void CheckForErrors() override;
 #endif
-	virtual TSharedPtr<FComponentInstanceDataBase> GetComponentInstanceData() const override;
+	virtual FComponentInstanceDataBase* GetComponentInstanceData() const override;
 	virtual FName GetComponentInstanceDataType() const override;
-	virtual void ApplyComponentInstanceData(TSharedPtr<FComponentInstanceDataBase> ComponentInstanceData) override;
+	virtual void ApplyComponentInstanceData(FComponentInstanceDataBase* ComponentInstanceData) override;
 	// End UActorComponent interface.
 
 
diff --git a/Engine/Source/Runtime/Engine/Private/Atmosphere/Atmosphere.cpp b/Engine/Source/Runtime/Engine/Private/Atmosphere/Atmosphere.cpp
--- a/Engine/Source/Runtime/Engine/Private/Atmosphere/Atmosphere.cpp
+++ b/Engine/Source/Runtime/Engine/Private/Atmosphere/Atmosphere.cpp
@@ -727,14 +727,14 @@ FName UAtmosphericFogComponent::GetComponentInstanceDataType() const
 }
 
 // Backup the precomputed data before re-running Blueprint construction script
-TSharedPtr<FComponentInstanceDataBase> UAtmosphericFogComponent::GetComponentInstanceData() const
+FComponentInstanceDataBase* UAtmosphericFogComponent::GetComponentInstanceData() const
 {
-	TSharedPtr<FAtmospherePrecomputeInstanceData> PrecomputedData;
+	FAtmospherePrecomputeInstanceData* PrecomputedData = nullptr;
 
 	if (TransmittanceData.GetElementCount() && IrradianceData.GetElementCount() && InscatterData.GetElementCount() && PrecomputeCounter.GetValue() == EValid)
 	{
 		// Allocate new struct for holding light map data
-		 PrecomputedData = MakeShareable(new FAtmospherePrecomputeInstanceData(this));
+		 PrecomputedData = new FAtmospherePrecomputeInstanceData(this);
 
 		// Fill in info
 		PrecomputedData->PrecomputeParameter = PrecomputeParams;
@@ -767,10 +767,10 @@ TSharedPtr<FComponentInstanceDataBase> UAtmosphericFogComponent::GetComponentIns
 }
 
 // Restore the precomputed data after re-running Blueprint construction script
-void UAtmosphericFogComponent::ApplyComponentInstanceData(TSharedPtr<FComponentInstanceDataBase> ComponentInstanceData)
+void UAtmosphericFogComponent::ApplyComponentInstanceData(FComponentInstanceDataBase* ComponentInstanceData)
 {
-	check(ComponentInstanceData.IsValid());
-	TSharedPtr<FAtmospherePrecomputeInstanceData> PrecomputedData = StaticCastSharedPtr<FAtmospherePrecomputeInstanceData>(ComponentInstanceData);
+	check(ComponentInstanceData);
+	FAtmospherePrecomputeInstanceData* PrecomputedData = static_cast<FAtmospherePrecomputeInstanceData*>(const_cast<FComponentInstanceDataBase*>(ComponentInstanceData));
 
 	FComponentReregisterContext ReregisterContext(this);
 	ReleaseResource();
diff --git a/Engine/Source/Runtime/Engine/Private/ChildActorComponent.cpp b/Engine/Source/Runtime/Engine/Private/ChildActorComponent.cpp
--- a/Engine/Source/Runtime/Engine/Private/ChildActorComponent.cpp
+++ b/Engine/Source/Runtime/Engine/Private/ChildActorComponent.cpp
@@ -88,15 +88,27 @@ FName UChildActorComponent::GetComponentInstanceDataType() const
 	return ChildActorComponentInstanceDataName;
 }
 
-TSharedPtr<FComponentInstanceDataBase> UChildActorComponent::GetComponentInstanceData() const
+FComponentInstanceDataBase* UChildActorComponent::GetComponentInstanceData() const
 {
-	return (CachedInstanceData.IsValid() ? CachedInstanceData : MakeShareable(new FChildActorComponentInstanceData(this)));
+	FChildActorComponentInstanceData* InstanceData = CachedInstanceData;
+
+	if (CachedInstanceData)
+	{
+		// We've handed over ownership of the pointer to the instance cache, so drop our reference
+		CachedInstanceData = nullptr;
+	}
+	else
+	{
+		InstanceData = new FChildActorComponentInstanceData(this);
+	}
+	
+	return InstanceData;
 }
 
-void UChildActorComponent::ApplyComponentInstanceData(TSharedPtr<FComponentInstanceDataBase> ComponentInstanceData)
+void UChildActorComponent::ApplyComponentInstanceData(FComponentInstanceDataBase* ComponentInstanceData)
 {
-	check(ComponentInstanceData.IsValid());
-	FChildActorComponentInstanceData* ChildActorInstanceData = StaticCastSharedPtr<FChildActorComponentInstanceData>(ComponentInstanceData).Get();
+	check(ComponentInstanceData);
+	FChildActorComponentInstanceData* ChildActorInstanceData  = static_cast<FChildActorComponentInstanceData*>(ComponentInstanceData);
 	ChildActorName = ChildActorInstanceData->ChildActorName;
 	if (ChildActor)
 	{
@@ -104,9 +116,9 @@ void UChildActorComponent::ApplyComponentInstanceData(TSharedPtr<FComponentInsta
 		if(ChildActorName != NAME_None)
 		{
 			const FString ChildActorNameString = ChildActorName.ToString();
-			if (ChildActor->Rename(*ChildActorNameString, NULL, REN_Test))
+			if (ChildActor->Rename(*ChildActorNameString, nullptr, REN_Test))
 			{
-				ChildActor->Rename(*ChildActorNameString, NULL, REN_DoNotDirty);
+				ChildActor->Rename(*ChildActorNameString, nullptr, REN_DoNotDirty);
 			}
 		}
 
@@ -138,13 +150,17 @@ void UChildActorComponent::CreateChildActor()
 	DestroyChildActor();
 
 	// This is no longer needed
-	CachedInstanceData = NULL;
+	if (CachedInstanceData)
+	{
+		delete CachedInstanceData;
+		CachedInstanceData = nullptr;
+	}
 
 	// If we have a class to spawn.
-	if(ChildActorClass != NULL)
+	if(ChildActorClass != nullptr)
 	{
 		UWorld* World = GetWorld();
-		if(World != NULL)
+		if(World != nullptr)
 		{
 			// Before we spawn let's try and prevent cyclic disaster
 			bool bSpawn = true;
@@ -175,7 +191,7 @@ void UChildActorComponent::CreateChildActor()
 				ChildActor = World->SpawnActor(ChildActorClass, &Location, &Rotation, Params);
 
 				// If spawn was successful, 
-				if(ChildActor != NULL)
+				if(ChildActor != nullptr)
 				{
 					ChildActorName = ChildActor->GetFName();
 
@@ -193,16 +209,17 @@ void UChildActorComponent::CreateChildActor()
 void UChildActorComponent::DestroyChildActor()
 {
 	// If we own an Actor, kill it now
-	if(ChildActor != NULL)
+	if(ChildActor != nullptr)
 	{
 		// if still alive, destroy, otherwise just clear the pointer
 		if(!ChildActor->IsPendingKill())
 		{
-			CachedInstanceData = MakeShareable(new FChildActorComponentInstanceData(this));
+			check(!CachedInstanceData);
+			CachedInstanceData = new FChildActorComponentInstanceData(this);
 
 			UWorld* World = ChildActor->GetWorld();
-			// World may be NULL during shutdown
-			if(World != NULL)
+			// World may be nullptr during shutdown
+			if(World != nullptr)
 			{
 				UClass* ChildClass = ChildActor->GetClass();
 
@@ -211,11 +228,11 @@ void UChildActorComponent::DestroyChildActor()
 				ChildClass->ClassUnique = FMath::Max(ChildClass->ClassUnique, ChildActor->GetFName().GetNumber());
 
 				const FString ObjectBaseName = FString::Printf(TEXT("DESTROYED_%s_CHILDACTOR"), *ChildClass->GetName());
-				ChildActor->Rename(*MakeUniqueObjectName(ChildActor->GetOuter(), ChildClass, *ObjectBaseName).ToString(), NULL, REN_DoNotDirty);
+				ChildActor->Rename(*MakeUniqueObjectName(ChildActor->GetOuter(), ChildClass, *ObjectBaseName).ToString(), nullptr, REN_DoNotDirty);
 				World->DestroyActor(ChildActor);
 			}
 		}
 
-		ChildActor = NULL;
+		ChildActor = nullptr;
 	}
 }
diff --git a/Engine/Source/Runtime/Engine/Private/ComponentInstanceDataCache.cpp b/Engine/Source/Runtime/Engine/Private/ComponentInstanceDataCache.cpp
--- a/Engine/Source/Runtime/Engine/Private/ComponentInstanceDataCache.cpp
+++ b/Engine/Source/Runtime/Engine/Private/ComponentInstanceDataCache.cpp
@@ -79,8 +79,8 @@ FComponentInstanceDataCache::FComponentInstanceDataCache(const AActor* Actor)
 		{
 			if(Component->bCreatedByConstructionScript) // Only cache data from 'created by construction script' components
 			{
-				TSharedPtr<FComponentInstanceDataBase> ComponentInstanceData = Component->GetComponentInstanceData();
-				if (ComponentInstanceData.IsValid())
+				FComponentInstanceDataBase* ComponentInstanceData = Component->GetComponentInstanceData();
+				if (ComponentInstanceData)
 				{
 					check(!Component->GetComponentInstanceDataType().IsNone());
 					TypeToDataMap.Add(Component->GetComponentInstanceDataType(), ComponentInstanceData);
@@ -90,6 +90,14 @@ FComponentInstanceDataCache::FComponentInstanceDataCache(const AActor* Actor)
 	}
 }
 
+FComponentInstanceDataCache::~FComponentInstanceDataCache()
+{
+	for (auto InstanceDataPair : TypeToDataMap)
+	{
+		delete InstanceDataPair.Value;
+	}
+}
+
 void FComponentInstanceDataCache::ApplyToActor(AActor* Actor) const
 {
 	if(Actor != NULL)
@@ -106,12 +114,12 @@ void FComponentInstanceDataCache::ApplyToActor(AActor* Actor) const
 
 				if (!ComponentInstanceDataType.IsNone())
 				{
-					TArray< TSharedPtr<FComponentInstanceDataBase> > CachedData;
+					TArray< FComponentInstanceDataBase* > CachedData;
 					TypeToDataMap.MultiFind(ComponentInstanceDataType, CachedData);
 
-					for (TSharedPtr<FComponentInstanceDataBase> ComponentInstanceData : CachedData)
+					for (FComponentInstanceDataBase* ComponentInstanceData : CachedData)
 					{
-						if (ComponentInstanceData.IsValid() && ComponentInstanceData->MatchesComponent(Component))
+						if (ComponentInstanceData && ComponentInstanceData->MatchesComponent(Component))
 						{
 							Component->ApplyComponentInstanceData(ComponentInstanceData);
 							break;
diff --git a/Engine/Source/Runtime/Engine/Private/InstancedStaticMesh.cpp b/Engine/Source/Runtime/Engine/Private/InstancedStaticMesh.cpp
--- a/Engine/Source/Runtime/Engine/Private/InstancedStaticMesh.cpp
+++ b/Engine/Source/Runtime/Engine/Private/InstancedStaticMesh.cpp
@@ -1273,15 +1273,15 @@ FName UInstancedStaticMeshComponent::GetComponentInstanceDataType() const
 	return InstancedStaticMeshComponentInstanceDataName;
 }
 
-TSharedPtr<FComponentInstanceDataBase> UInstancedStaticMeshComponent::GetComponentInstanceData() const
+FComponentInstanceDataBase* UInstancedStaticMeshComponent::GetComponentInstanceData() const
 {
 #if WITH_EDITOR
-	TSharedPtr<FInstancedStaticMeshComponentInstanceData> InstanceData;
+	FInstancedStaticMeshComponentInstanceData* InstanceData = nullptr;
 
 	// Don't back up static lighting if there isn't any
 	if (bHasCachedStaticLighting || SelectedInstances.Num() > 0)
 	{
-		InstanceData = MakeShareable(new FInstancedStaticMeshComponentInstanceData(*this));
+		InstanceData = new FInstancedStaticMeshComponentInstanceData(*this);
 	}
 
 	// Don't back up static lighting if there isn't any
@@ -1310,16 +1310,16 @@ TSharedPtr<FComponentInstanceDataBase> UInstancedStaticMeshComponent::GetCompone
 
 	return InstanceData;
 #else
-	return TSharedPtr<FComponentInstanceDataBase>();
+	return nullptr;
 #endif
 }
 
-void UInstancedStaticMeshComponent::ApplyComponentInstanceData(TSharedPtr<FComponentInstanceDataBase> ComponentInstanceData)
+void UInstancedStaticMeshComponent::ApplyComponentInstanceData(FComponentInstanceDataBase* ComponentInstanceData)
 {
 #if WITH_EDITOR
-	check(ComponentInstanceData.IsValid());
+	check(ComponentInstanceData);
 
-	auto* InstancedMeshData = static_cast<FInstancedStaticMeshComponentInstanceData*>(ComponentInstanceData.Get());
+	FInstancedStaticMeshComponentInstanceData* InstancedMeshData  = static_cast<FInstancedStaticMeshComponentInstanceData*>(ComponentInstanceData);
 
 	// See if data matches current state
 	if (InstancedMeshData->bHasCachedStaticLighting)
diff --git a/Engine/Source/Runtime/Engine/Private/LightComponent.cpp b/Engine/Source/Runtime/Engine/Private/LightComponent.cpp
--- a/Engine/Source/Runtime/Engine/Private/LightComponent.cpp
+++ b/Engine/Source/Runtime/Engine/Private/LightComponent.cpp
@@ -901,16 +901,16 @@ FName ULightComponent::GetComponentInstanceDataType() const
 	return PrecomputedLightInstanceDataTypeName;
 }
 
-TSharedPtr<FComponentInstanceDataBase> ULightComponent::GetComponentInstanceData() const
+FComponentInstanceDataBase* ULightComponent::GetComponentInstanceData() const
 {
 	// Allocate new struct for holding light map data
-	return MakeShareable(new FPrecomputedLightInstanceData(this));
+	return new FPrecomputedLightInstanceData(this);
 }
 
-void ULightComponent::ApplyComponentInstanceData(TSharedPtr<FComponentInstanceDataBase> ComponentInstanceData)
+void ULightComponent::ApplyComponentInstanceData(FComponentInstanceDataBase* ComponentInstanceData)
 {
-	check(ComponentInstanceData.IsValid());
-	TSharedPtr<FPrecomputedLightInstanceData> LightMapData = StaticCastSharedPtr<FPrecomputedLightInstanceData>(ComponentInstanceData);
+	check(ComponentInstanceData);
+	FPrecomputedLightInstanceData* LightMapData  = static_cast<FPrecomputedLightInstanceData*>(ComponentInstanceData);
 
 	LightGuid = LightMapData->LightGuid;
 	ShadowMapChannel = LightMapData->ShadowMapChannel;
diff --git a/Engine/Source/Runtime/Engine/Private/SkyLightComponent.cpp b/Engine/Source/Runtime/Engine/Private/SkyLightComponent.cpp
--- a/Engine/Source/Runtime/Engine/Private/SkyLightComponent.cpp
+++ b/Engine/Source/Runtime/Engine/Private/SkyLightComponent.cpp
@@ -345,9 +345,9 @@ FName USkyLightComponent::GetComponentInstanceDataType() const
 	return PrecomputedSkyLightInstanceDataTypeName;
 }
 
-TSharedPtr<FComponentInstanceDataBase> USkyLightComponent::GetComponentInstanceData() const
+FComponentInstanceDataBase* USkyLightComponent::GetComponentInstanceData() const
 {
-	TSharedRef<FPrecomputedSkyLightInstanceData> InstanceData = MakeShareable(new FPrecomputedSkyLightInstanceData(this));
+	FPrecomputedSkyLightInstanceData* InstanceData = new FPrecomputedSkyLightInstanceData(this);
 	InstanceData->LightGuid = LightGuid;
 	InstanceData->bPrecomputedLightingIsValid = bPrecomputedLightingIsValid;
 	InstanceData->ProcessedSkyTexture = ProcessedSkyTexture;
@@ -356,10 +356,10 @@ TSharedPtr<FComponentInstanceDataBase> USkyLightComponent::GetComponentInstanceD
 	return InstanceData;
 }
 
-void USkyLightComponent::ApplyComponentInstanceData(TSharedPtr<FComponentInstanceDataBase> ComponentInstanceData)
+void USkyLightComponent::ApplyComponentInstanceData(FComponentInstanceDataBase* ComponentInstanceData)
 {
-	check(ComponentInstanceData.IsValid());
-	TSharedPtr<FPrecomputedSkyLightInstanceData> LightMapData = StaticCastSharedPtr<FPrecomputedSkyLightInstanceData>(ComponentInstanceData);
+	check(ComponentInstanceData);
+	FPrecomputedSkyLightInstanceData* LightMapData  = static_cast<FPrecomputedSkyLightInstanceData*>(ComponentInstanceData);
 
 	LightGuid = LightMapData->LightGuid;
 	bPrecomputedLightingIsValid = LightMapData->bPrecomputedLightingIsValid;
diff --git a/Engine/Source/Runtime/Engine/Private/SplineComponent.cpp b/Engine/Source/Runtime/Engine/Private/SplineComponent.cpp
--- a/Engine/Source/Runtime/Engine/Private/SplineComponent.cpp
+++ b/Engine/Source/Runtime/Engine/Private/SplineComponent.cpp
@@ -526,23 +526,23 @@ FName USplineComponent::GetComponentInstanceDataType() const
 	return SplineInstanceDataTypeName;
 }
 
-TSharedPtr<FComponentInstanceDataBase> USplineComponent::GetComponentInstanceData() const
+FComponentInstanceDataBase* USplineComponent::GetComponentInstanceData() const
 {
-	TSharedPtr<FSplineInstanceData> SplineInstanceData;
+	FSplineInstanceData* SplineInstanceData = nullptr;
 	if (bAllowSplineEditingPerInstance)
 	{
-		SplineInstanceData = MakeShareable(new FSplineInstanceData(this));
+		SplineInstanceData = new FSplineInstanceData(this);
 		SplineInstanceData->SplineInfo = SplineInfo;
 		SplineInstanceData->bClosedLoop = bClosedLoop;
 	}
 	return SplineInstanceData;
 }
 
-void USplineComponent::ApplyComponentInstanceData(TSharedPtr<FComponentInstanceDataBase> ComponentInstanceData)
+void USplineComponent::ApplyComponentInstanceData(FComponentInstanceDataBase* ComponentInstanceData)
 {
-	if (ComponentInstanceData.IsValid())
+	if (ComponentInstanceData)
 	{
-		const FSplineInstanceData* SplineInstanceData = StaticCastSharedPtr<FSplineInstanceData>(ComponentInstanceData).Get();
+		FSplineInstanceData* SplineInstanceData  = static_cast<FSplineInstanceData*>(ComponentInstanceData);
 		if (bAllowSplineEditingPerInstance)
 		{
 			SplineInfo = SplineInstanceData->SplineInfo;
diff --git a/Engine/Source/Runtime/Engine/Private/StaticMeshComponent.cpp b/Engine/Source/Runtime/Engine/Private/StaticMeshComponent.cpp
--- a/Engine/Source/Runtime/Engine/Private/StaticMeshComponent.cpp
+++ b/Engine/Source/Runtime/Engine/Private/StaticMeshComponent.cpp
@@ -43,7 +43,7 @@ public:
 	}
 
 	/** Re-apply vertex color data after RerunConstructionScripts is called */
-	bool ApplyVertexColorData(UStaticMeshComponent* StaticMeshComponent)
+	bool ApplyVertexColorData(UStaticMeshComponent* StaticMeshComponent) const
 	{
 		bool bAppliedAnyData = false;
 
@@ -1512,14 +1512,14 @@ FName UStaticMeshComponent::GetComponentInstanceDataType() const
 	return StaticMeshComponentInstanceDataName;
 }
 
-TSharedPtr<FComponentInstanceDataBase> UStaticMeshComponent::GetComponentInstanceData() const
+FComponentInstanceDataBase* UStaticMeshComponent::GetComponentInstanceData() const
 {
-	TSharedPtr<FStaticMeshComponentInstanceData> InstanceData;
+	FStaticMeshComponentInstanceData* InstanceData = nullptr;
 
 	// Don't back up static lighting if there isn't any
 	if(bHasCachedStaticLighting)
 	{
-		InstanceData = MakeShareable(new FStaticMeshComponentInstanceData(this));
+		InstanceData = new FStaticMeshComponentInstanceData(this);
 
 		// Fill in info
 		InstanceData->bHasCachedStaticLighting = true;
@@ -1541,9 +1541,9 @@ TSharedPtr<FComponentInstanceDataBase> UStaticMeshComponent::GetComponentInstanc
 
 		if ( LODInfo.OverrideVertexColors && LODInfo.OverrideVertexColors->GetNumVertices() > 0 && LODInfo.PaintedVertices.Num() > 0 )
 		{
-			if (!InstanceData.IsValid())
+			if (!InstanceData)
 			{
-				InstanceData = MakeShareable(new FStaticMeshComponentInstanceData(this));
+				InstanceData = new FStaticMeshComponentInstanceData(this);
 			}
 
 			InstanceData->AddVertexColorData(LODInfo, LODIndex);
@@ -1553,13 +1553,13 @@ TSharedPtr<FComponentInstanceDataBase> UStaticMeshComponent::GetComponentInstanc
 	return InstanceData;
 }
 
-void UStaticMeshComponent::ApplyComponentInstanceData(TSharedPtr<FComponentInstanceDataBase> ComponentInstanceData)
+void UStaticMeshComponent::ApplyComponentInstanceData(FComponentInstanceDataBase* ComponentInstanceData)
 {
-	check(ComponentInstanceData.IsValid());
+	check(ComponentInstanceData);
 
 	// Note: ApplyComponentInstanceData is called while the component is registered so the rendering thread is already using this component
 	// That means all component state that is modified here must be mirrored on the scene proxy, which will be recreated to receive the changes later due to MarkRenderStateDirty.
-	TSharedPtr<FStaticMeshComponentInstanceData> StaticMeshInstanceData = StaticCastSharedPtr<FStaticMeshComponentInstanceData>(ComponentInstanceData);
+	FStaticMeshComponentInstanceData* StaticMeshInstanceData  = static_cast<FStaticMeshComponentInstanceData*>(ComponentInstanceData);
 
 	// See if data matches current state
 	if(	StaticMeshInstanceData->bHasCachedStaticLighting && StaticMeshInstanceData->CachedStaticLighting.Transform.Equals(ComponentToWorld) )
diff --git a/Engine/Source/Runtime/Engine/Public/ComponentInstanceDataCache.h b/Engine/Source/Runtime/Engine/Public/ComponentInstanceDataCache.h
--- a/Engine/Source/Runtime/Engine/Public/ComponentInstanceDataCache.h
+++ b/Engine/Source/Runtime/Engine/Public/ComponentInstanceDataCache.h
@@ -38,11 +38,10 @@ class ENGINE_API FComponentInstanceDataCache
 public:
 	FComponentInstanceDataCache() {}
 
-	/** Constructor that also populatees cache from Actor */
+	/** Constructor that also populates cache from Actor */
 	FComponentInstanceDataCache(const AActor* InActor);
 
-	/** Add some new instance data to the cache */
-	void AddInstanceData(TSharedPtr<FComponentInstanceDataBase> NewData);
+	~FComponentInstanceDataCache();
 
 	/** Util to iterate over components and apply data to each */
 	void ApplyToActor(AActor* Actor) const;
@@ -51,5 +50,5 @@ public:
 
 private:
 	/** Map of data type name to data of that type */
-	TMultiMap< FName, TSharedPtr<FComponentInstanceDataBase> >	TypeToDataMap;
+	TMultiMap< FName, FComponentInstanceDataBase* >	TypeToDataMap;
 };
\ No newline at end of file
diff --git a/Engine/Source/Runtime/UMG/Private/Components/WidgetComponent.cpp b/Engine/Source/Runtime/UMG/Private/Components/WidgetComponent.cpp
--- a/Engine/Source/Runtime/UMG/Private/Components/WidgetComponent.cpp
+++ b/Engine/Source/Runtime/UMG/Private/Components/WidgetComponent.cpp
@@ -520,18 +520,18 @@ FName UWidgetComponent::GetComponentInstanceDataType() const
 	return InstanceDataName;
 }
 
-TSharedPtr<FComponentInstanceDataBase> UWidgetComponent::GetComponentInstanceData() const
+FComponentInstanceDataBase* UWidgetComponent::GetComponentInstanceData() const
 {
-	return MakeShareable( new FWidgetComponentInstanceData( this ) );
+	return new FWidgetComponentInstanceData( this );
 }
 
-void UWidgetComponent::ApplyComponentInstanceData(TSharedPtr<FComponentInstanceDataBase> ComponentInstanceData)
+void UWidgetComponent::ApplyComponentInstanceData(FComponentInstanceDataBase* ComponentInstanceData)
 {
-	check(ComponentInstanceData.IsValid());
+	check(ComponentInstanceData);
 
 	// Note: ApplyComponentInstanceData is called while the component is registered so the rendering thread is already using this component
 	// That means all component state that is modified here must be mirrored on the scene proxy, which will be recreated to receive the changes later due to MarkRenderStateDirty.
-	TSharedPtr<FWidgetComponentInstanceData> WidgetInstanceData = StaticCastSharedPtr<FWidgetComponentInstanceData>(ComponentInstanceData);
+	FWidgetComponentInstanceData* WidgetInstanceData  = static_cast<FWidgetComponentInstanceData*>(ComponentInstanceData);
 
 	RenderTarget = WidgetInstanceData->RenderTarget;
 	MaterialInstance->SetTextureParameterValue("SlateUI", RenderTarget);
diff --git a/Engine/Source/Runtime/UMG/Public/Components/WidgetComponent.h b/Engine/Source/Runtime/UMG/Public/Components/WidgetComponent.h
--- a/Engine/Source/Runtime/UMG/Public/Components/WidgetComponent.h
+++ b/Engine/Source/Runtime/UMG/Public/Components/WidgetComponent.h
@@ -26,9 +26,9 @@ public:
 	virtual void DestroyComponent() override;
 	virtual void TickComponent(float DeltaTime, enum ELevelTick TickType, FActorComponentTickFunction *ThisTickFunction) override;
 
-	virtual TSharedPtr<FComponentInstanceDataBase> GetComponentInstanceData() const override;
+	virtual FComponentInstanceDataBase* GetComponentInstanceData() const override;
 	virtual FName GetComponentInstanceDataType() const override;
-	virtual void ApplyComponentInstanceData(TSharedPtr<class FComponentInstanceDataBase> ComponentInstanceData ) override;
+	virtual void ApplyComponentInstanceData(FComponentInstanceDataBase* ComponentInstanceData) override;
 
 #if WITH_EDITORONLY_DATA
 	virtual void PostEditChangeProperty(FPropertyChangedEvent& PropertyChangedEvent);
