Get rid of check whether delegate or handle is valid

diff --git a/Engine/Source/Runtime/Engine/Private/TimerManager.cpp b/Engine/Source/Runtime/Engine/Private/TimerManager.cpp
--- a/Engine/Source/Runtime/Engine/Private/TimerManager.cpp
+++ b/Engine/Source/Runtime/Engine/Private/TimerManager.cpp
@@ -181,25 +181,26 @@ void FTimerManager::InternalSetTimer(FTimerHandle& InOutHandle, FTimerUnifiedDel
 
 void FTimerManager::InternalSetTimer(FTimerData& NewTimerData, float InRate, bool InbLoop, float InFirstDelay)
 {
-	check(NewTimerData.TimerDelegate.IsBound() || NewTimerData.TimerHandle.IsValid());
-
-	NewTimerData.Rate = InRate;
-	NewTimerData.bLoop = InbLoop;
+	if (NewTimerData.TimerDelegate.IsBound() || NewTimerData.TimerHandle.IsValid())
+	{
+		NewTimerData.Rate = InRate;
+		NewTimerData.bLoop = InbLoop;
 
-	float FirstDelay = (InFirstDelay >= 0.f) ? InFirstDelay : InRate;
+		const float FirstDelay = (InFirstDelay >= 0.f) ? InFirstDelay : InRate;
 
-	if (HasBeenTickedThisFrame())
-	{
-		NewTimerData.ExpireTime = InternalTime + FirstDelay;
-		NewTimerData.Status = ETimerStatus::Active;
-		ActiveTimerHeap.HeapPush(NewTimerData);
-	}
-	else
-	{
-		// Store time remaining in ExpireTime while pending
-		NewTimerData.ExpireTime = FirstDelay;
-		NewTimerData.Status = ETimerStatus::Pending;
-		PendingTimerList.Add(NewTimerData);
+		if (HasBeenTickedThisFrame())
+		{
+			NewTimerData.ExpireTime = InternalTime + FirstDelay;
+			NewTimerData.Status = ETimerStatus::Active;
+			ActiveTimerHeap.HeapPush(NewTimerData);
+		}
+		else
+		{
+			// Store time remaining in ExpireTime while pending
+			NewTimerData.ExpireTime = FirstDelay;
+			NewTimerData.Status = ETimerStatus::Pending;
+			PendingTimerList.Add(NewTimerData);
+		}
 	}
 }
 
