Expose MinRespawnDelay to defaults and to blueprints as read/write
Expose NumPlayers, NumTravellingPlayers, and NumSpectators as BlueprintRead
Some category rearrangement

diff --git a/Engine/Source/Runtime/Engine/Classes/GameFramework/GameMode.h b/Engine/Source/Runtime/Engine/Classes/GameFramework/GameMode.h
--- a/Engine/Source/Runtime/Engine/Classes/GameFramework/GameMode.h
+++ b/Engine/Source/Runtime/Engine/Classes/GameFramework/GameMode.h
@@ -162,19 +162,19 @@ public:
 	FString OptionsString;
 
 	/** The default pawn class used by players. */
-	UPROPERTY(EditAnywhere, BlueprintReadWrite, Category=GameMode)
+	UPROPERTY(EditAnywhere, BlueprintReadWrite, Category=Classes)
 	TSubclassOf<class APawn>  DefaultPawnClass;
 
 	/** HUD class this game uses. */
-	UPROPERTY(EditAnywhere, noclear, BlueprintReadWrite, Category=GameMode, meta=(DisplayName="HUD Class"))
+	UPROPERTY(EditAnywhere, noclear, BlueprintReadWrite, Category=Classes, meta=(DisplayName="HUD Class"))
 	TSubclassOf<class AHUD>  HUDClass;    
 
 	/** Current number of spectators. */
-	UPROPERTY()
+	UPROPERTY(BlueprintRead, Category=GameMode)
 	int32 NumSpectators;    
 
 	/** Current number of human players. */
-	UPROPERTY()
+	UPROPERTY(BlueprintRead, Category=GameMode)
 	int32 NumPlayers;    
 
 	/** number of non-human players (AI controlled but participating as a player). */
@@ -182,7 +182,7 @@ public:
 	int32 NumBots;    
 
 	/** Minimum time before player can respawn after dying. */
-	UPROPERTY()
+	UPROPERTY(EditAnywhere, BlueprintReadWrite, Category=GameMode, meta=(DisplayName="Minimum Respawn Delay"))
 	float MinRespawnDelay;
 
 	/** Game Session handles login approval, arbitration, online game interface */
@@ -190,7 +190,7 @@ public:
 	class AGameSession* GameSession;
 
 	/** Number of players that are still traveling from a previous map */
-	UPROPERTY()
+	UPROPERTY(BlueprintRead, Category=GameMode)
 	int32 NumTravellingPlayers;
 
 	/** Used to assign unique PlayerIDs to each PlayerState. */
@@ -210,19 +210,19 @@ public:
 	TSubclassOf<class ULocalMessage> EngineMessageClass;
 
 	/** The class of PlayerController to spawn for players logging in. */
-	UPROPERTY(EditAnywhere, noclear, BlueprintReadWrite, Category=GameMode)
+	UPROPERTY(EditAnywhere, noclear, BlueprintReadWrite, Category=Classes)
 	TSubclassOf<class APlayerController> PlayerControllerClass;
 
 	/** The pawn class used by the PlayerController for players when spectating. */
-	UPROPERTY(EditAnywhere, noclear, BlueprintReadWrite, Category=GameMode)
+	UPROPERTY(EditAnywhere, noclear, BlueprintReadWrite, Category=Classes)
 	TSubclassOf<class ASpectatorPawn> SpectatorClass;
 
 	/** A PlayerState of this class will be associated with every player to replicate relevant player information to all clients. */
-	UPROPERTY(EditAnywhere, noclear, BlueprintReadWrite, Category=GameMode, AdvancedDisplay)
+	UPROPERTY(EditAnywhere, noclear, BlueprintReadWrite, Category=Classes, AdvancedDisplay)
 	TSubclassOf<class APlayerState> PlayerStateClass;
 
 	/** Class of GameState associated with this GameMode. */
-	UPROPERTY(EditAnywhere, noclear, BlueprintReadWrite, Category=GameMode)
+	UPROPERTY(EditAnywhere, noclear, BlueprintReadWrite, Category=Classes)
 	TSubclassOf<class AGameState> GameStateClass;
 
 	/** GameState is used to replicate game state relevant properties to all clients. */
