Expose UPrimitiveComponent::SetCullDistance (as Set Max Draw Distance) as a UFUNCTION to blueprints

diff --git a/Engine/Source/Runtime/Engine/Classes/Components/PrimitiveComponent.h b/Engine/Source/Runtime/Engine/Classes/Components/PrimitiveComponent.h
--- a/Engine/Source/Runtime/Engine/Classes/Components/PrimitiveComponent.h
+++ b/Engine/Source/Runtime/Engine/Classes/Components/PrimitiveComponent.h
@@ -1426,6 +1426,7 @@ public:
 	 * Changes the value of CullDistance.
 	 * @param NewCullDistance - The value to assign to CullDistance.
 	 */
+	UFUNCTION(BlueprintCallable, Category="LOD", meta=(FriendlyName="Set Max Draw Distance"))
 	void SetCullDistance(float NewCullDistance);
 	
 	/**
