[UE-2007] STREAMING NAVIMESH: need a solution for streaming chunks of the navimesh in large worlds

Each sub-level can hold set of tiles for a static navmesh (with disabled runtime rebuild). Those tiles will be attached to main navmesh when sub-level is streamed in and detached when streamed out. Whether to save set of tiles into sub-level controlled by navigation volume placement. When user creates navigation volume inside sub-level, tiles that intersect that volume will be saved into sub-level. So in case all volumes are in the persistent level whole navmesh will be saved into persistent level and will work without streaming.

diff --git a/Engine/Source/Runtime/Engine/Classes/AI/Navigation/NavigationData.h b/Engine/Source/Runtime/Engine/Classes/AI/Navigation/NavigationData.h
--- a/Engine/Source/Runtime/Engine/Classes/AI/Navigation/NavigationData.h
+++ b/Engine/Source/Runtime/Engine/Classes/AI/Navigation/NavigationData.h
@@ -412,6 +412,8 @@ class ENGINE_API ANavigationData : public AActor
 	virtual bool NeedsRebuild() const { return false; }
 	virtual bool SupportsRuntimeGeneration() const;
 	virtual void OnNavigationBoundsChanged();
+	virtual void OnStreamingLevelAdded(ULevel* InLevel) {};
+	virtual void OnStreamingLevelRemoved(ULevel* InLevel) {};
 	
 	//----------------------------------------------------------------------//
 	// Generation & data access                                                      
@@ -429,9 +431,21 @@ protected:
 	virtual void FillConfig(FNavDataConfig& Dest) { Dest = NavDataConfig; }
 
 public:
+	/** Creates new generator in case navigation supports it */
 	virtual void ConstructGenerator();
+	
+	/** Triggers rebuild in case navigation supports it */
 	virtual void RebuildAll();
+
+	/** Blocks until navigation build is complete  */
+	virtual void EnsureBuildCompletion();
+
+	/** Cancels current build  */
+	virtual void CancelBuild();
+
+	/** Ticks navigation build  */
 	virtual void TickAsyncBuild(float DeltaSeconds);
+	
 	/** Retrieves navmesh's generator */
 	FNavDataGenerator* GetGenerator() { return NavDataGenerator.Get(); }
 	const FNavDataGenerator* GetGenerator() const { return NavDataGenerator.Get(); }
@@ -494,6 +508,12 @@ public:
 	/** Returns bounding box for the navmesh. */
 	virtual FBox GetBounds() const PURE_VIRTUAL(ANavigationData::GetBounds,return FBox(););
 	
+	/** Returns list of navigable bounds. */
+	TArray<FBox> GetNavigableBounds() const;
+	
+	/** Returns list of navigable bounds that belongs to specific level */
+	TArray<FBox> GetNavigableBoundsInLevel(const FName& InLevelPackageName) const;
+	
 	//----------------------------------------------------------------------//
 	// Debug                                                                
 	//----------------------------------------------------------------------//
diff --git a/Engine/Source/Runtime/Engine/Classes/AI/Navigation/NavigationDataChunk.h b/Engine/Source/Runtime/Engine/Classes/AI/Navigation/NavigationDataChunk.h
new file mode 100644
--- /dev/null
+++ b/Engine/Source/Runtime/Engine/Classes/AI/Navigation/NavigationDataChunk.h
@@ -0,0 +1,19 @@
+// Copyright 1998-2014 Epic Games, Inc. All Rights Reserved.
+
+#pragma once
+#include "AI/Navigation/NavigationTypes.h"
+#include "NavigationDataChunk.generated.h"
+
+/** 
+ * 
+ */
+UCLASS(NotBlueprintable, abstract)
+class ENGINE_API UNavigationDataChunk : public UObject
+{
+	GENERATED_UCLASS_BODY()
+
+public:
+	/** Name of NavigationData actor that owns this chunk */
+	UPROPERTY()
+	FName NavigationDataName;
+};
diff --git a/Engine/Source/Runtime/Engine/Classes/AI/Navigation/NavigationTypes.h b/Engine/Source/Runtime/Engine/Classes/AI/Navigation/NavigationTypes.h
--- a/Engine/Source/Runtime/Engine/Classes/AI/Navigation/NavigationTypes.h
+++ b/Engine/Source/Runtime/Engine/Classes/AI/Navigation/NavigationTypes.h
@@ -88,6 +88,7 @@ struct FNavigationBounds
 {
 	uint32	UniqueID;
 	FBox	AreaBox;	
+	FName	PackageName;
 
 	bool operator==(const FNavigationBounds& Other) const 
 	{ 
diff --git a/Engine/Source/Runtime/Engine/Classes/AI/Navigation/RecastNavMesh.h b/Engine/Source/Runtime/Engine/Classes/AI/Navigation/RecastNavMesh.h
--- a/Engine/Source/Runtime/Engine/Classes/AI/Navigation/RecastNavMesh.h
+++ b/Engine/Source/Runtime/Engine/Classes/AI/Navigation/RecastNavMesh.h
@@ -37,6 +37,7 @@ class FRecastQueryFilter;
 class INavLinkCustomInterface;
 class UNavArea;
 class UNavigationSystem;
+class URecastNavMeshDataChunk;
 struct FAreaNavModifier;
 
 UENUM()
@@ -369,7 +370,8 @@ struct FNavMeshTileData
 	}
 
 	FORCEINLINE bool IsValid() const { return NavData.IsValid() && GetData() != nullptr && DataSize > 0; }
-	FORCEINLINE void Release() { if (NavData.IsValid()) { NavData->RawNavData = nullptr; } DataSize = 0; LayerIndex = 0; }
+
+	uint8* Release();
 };
 
 DECLARE_MULTICAST_DELEGATE(FOnNavMeshUpdate);
@@ -657,6 +659,9 @@ public:
 	virtual FBox GetBounds() const override { return GetNavMeshBounds(); }
 	/** Called on world origin changes **/
 	virtual void ApplyWorldOffset(const FVector& InOffset, bool bWorldShift) override;
+
+	virtual void OnStreamingLevelAdded(ULevel* InLevel) override;
+	virtual void OnStreamingLevelRemoved(ULevel* InLevel) override;
 	// End ANavigationData Interface
 
 protected:
@@ -698,6 +703,9 @@ public:
 	/** Invalidates active paths that go through changed tiles  */
 	void InvalidateAffectedPaths(const TArray<uint32>& ChangedTiles);
 
+	/** Event from generator that navmesh build has finished */
+	void OnNavMeshGenerationFinished();
+
 	virtual void SetConfig(const FNavDataConfig& Src) override;
 protected:
 	virtual void FillConfig(FNavDataConfig& Dest) override;
@@ -868,6 +876,9 @@ private:
 	void UpdateNavVersion();
 	void UpdateNavObject();
 
+	/** @return Navmesh data chunk that belongs to this actor */
+	URecastNavMeshDataChunk* GetNavigationDataChunk(ULevel* InLevel) const;
+
 protected:
 	// retrieves RecastNavMeshImpl
 	FPImplRecastNavMesh* GetRecastNavMeshImpl() { return RecastNavMeshImpl; }
diff --git a/Engine/Source/Runtime/Engine/Classes/AI/Navigation/RecastNavMeshDataChunk.h b/Engine/Source/Runtime/Engine/Classes/AI/Navigation/RecastNavMeshDataChunk.h
new file mode 100644
--- /dev/null
+++ b/Engine/Source/Runtime/Engine/Classes/AI/Navigation/RecastNavMeshDataChunk.h
@@ -0,0 +1,61 @@
+// Copyright 1998-2014 Epic Games, Inc. All Rights Reserved.
+
+#pragma once
+#include "NavigationDataChunk.h"
+#include "RecastNavMeshDataChunk.generated.h"
+
+struct FRecastTileData
+{
+	struct FRawData
+	{
+		FRawData(uint8* InData);
+		~FRawData();
+
+		uint8* RawData;
+	};
+
+	FRecastTileData();
+	FRecastTileData(NavNodeRef Ref, int32 DataSize, uint8* RawData);
+	
+	NavNodeRef				TileRef;
+	int32					TileDataSize;
+	TSharedPtr<FRawData>	TileRawData;
+};
+
+class dtNavMesh;
+
+/** 
+ * 
+ */
+UCLASS()
+class ENGINE_API URecastNavMeshDataChunk : public UNavigationDataChunk
+{
+	GENERATED_UCLASS_BODY()
+
+	// Begin UObject Interface
+	virtual void Serialize(FArchive& Ar) override;
+	// End UObject Interface
+
+	/** Attaches tiles to specified navmesh, transferring tile ownership to navmesh */
+	TArray<uint32> AttachTiles(dtNavMesh* NavMesh);
+	
+	/** Detaches tiles from specified navmesh, taking tile ownership */
+	TArray<uint32> DetachTiles(dtNavMesh* NavMesh);
+	
+	/** Number of tiles in this chunk */
+	int32 GetNumTiles() const;
+
+	/** Releases all tiles that this chunk holds */
+	void ReleaseTiles();
+
+public:
+#if WITH_EDITOR
+	void GatherTiles(const dtNavMesh* NavMesh, const TArray<int32>& TileIndices);
+#endif// WITH_EDITOR
+
+private:
+	void SerializeRecastData(FArchive& Ar);
+
+private:
+	TArray<FRecastTileData> Tiles;
+};
diff --git a/Engine/Source/Runtime/Engine/Classes/Engine/Level.h b/Engine/Source/Runtime/Engine/Classes/Engine/Level.h
--- a/Engine/Source/Runtime/Engine/Classes/Engine/Level.h
+++ b/Engine/Source/Runtime/Engine/Classes/Engine/Level.h
@@ -6,6 +6,7 @@
 
 class ALevelBounds;
 class UTexture2D;
+class UNavigationDataChunk;
 
 /**
  * Structure containing all information needed for determining the screen space
@@ -291,7 +292,11 @@ public:
 	class ANavigationObjectBase *NavListStart;
 	UPROPERTY()
 	class ANavigationObjectBase	*NavListEnd;
-
+	
+	/** Navigation related data that can be stored per level */
+	UPROPERTY()
+	TArray<UNavigationDataChunk*> NavDataChunks;
+	
 	/** Total number of KB used for lightmap textures in the level. */
 	UPROPERTY(VisibleAnywhere, Category=Level)
 	float LightmapTotalSize;
diff --git a/Engine/Source/Runtime/Engine/Private/AI/Navigation/NavigationData.cpp b/Engine/Source/Runtime/Engine/Private/AI/Navigation/NavigationData.cpp
--- a/Engine/Source/Runtime/Engine/Private/AI/Navigation/NavigationData.cpp
+++ b/Engine/Source/Runtime/Engine/Private/AI/Navigation/NavigationData.cpp
@@ -366,6 +366,22 @@ void ANavigationData::RebuildAll()
 	}
 }
 
+void ANavigationData::EnsureBuildCompletion()
+{
+	if (NavDataGenerator)
+	{
+		NavDataGenerator->EnsureBuildCompletion();
+	}
+}
+
+void ANavigationData::CancelBuild()
+{
+	if (NavDataGenerator)
+	{
+		NavDataGenerator->CancelBuild();
+	}
+}
+
 void ANavigationData::OnNavigationBoundsChanged()
 {
 	// Create generator if it wasn't yet
@@ -396,6 +412,41 @@ void ANavigationData::RebuildDirtyAreas(const TArray<FNavigationDirtyArea>& Dirt
 	}
 }
 
+TArray<FBox> ANavigationData::GetNavigableBounds() const
+{
+	TArray<FBox> Result;
+	const UNavigationSystem* NavSys = GetWorld()->GetNavigationSystem();
+	
+	const auto& NavigationBounds = NavSys->GetNavigationBounds();
+	Result.Reserve(NavigationBounds.Num());
+
+	for (const auto& Bounds : NavigationBounds)
+	{
+		Result.Add(Bounds.AreaBox);
+	}
+	
+	return Result;
+}
+	
+TArray<FBox> ANavigationData::GetNavigableBoundsInLevel(const FName& InLevelPackageName) const
+{
+	TArray<FBox> Result;
+	const UNavigationSystem* NavSys = GetWorld()->GetNavigationSystem();
+	
+	const auto& NavigationBounds = NavSys->GetNavigationBounds();
+	Result.Reserve(NavigationBounds.Num());
+
+	for (const auto& Bounds : NavigationBounds)
+	{
+		if (Bounds.PackageName == InLevelPackageName)
+		{
+			Result.Add(Bounds.AreaBox);
+		}
+	}
+	
+	return Result;
+}
+
 void ANavigationData::DrawDebugPath(FNavigationPath* Path, FColor PathColor, UCanvas* Canvas, bool bPersistent, const uint32 NextPathPointIndex) const
 {
 	Path->DebugDraw(this, PathColor, Canvas, bPersistent, NextPathPointIndex);
diff --git a/Engine/Source/Runtime/Engine/Private/AI/Navigation/NavigationDataChunk.cpp b/Engine/Source/Runtime/Engine/Private/AI/Navigation/NavigationDataChunk.cpp
new file mode 100644
--- /dev/null
+++ b/Engine/Source/Runtime/Engine/Private/AI/Navigation/NavigationDataChunk.cpp
@@ -0,0 +1,13 @@
+// Copyright 1998-2014 Epic Games, Inc. All Rights Reserved.
+
+#include "EnginePrivate.h"
+#include "AI/Navigation/NavigationDataChunk.h"
+
+//----------------------------------------------------------------------//
+// UNavigationDataChunk                                                                
+//----------------------------------------------------------------------//
+UNavigationDataChunk::UNavigationDataChunk(const FObjectInitializer& ObjectInitializer)
+	: Super(ObjectInitializer)
+{
+}
+
diff --git a/Engine/Source/Runtime/Engine/Private/AI/Navigation/NavigationSystem.cpp b/Engine/Source/Runtime/Engine/Private/AI/Navigation/NavigationSystem.cpp
--- a/Engine/Source/Runtime/Engine/Private/AI/Navigation/NavigationSystem.cpp
+++ b/Engine/Source/Runtime/Engine/Private/AI/Navigation/NavigationSystem.cpp
@@ -52,7 +52,7 @@ DEFINE_STAT(STAT_Navigation_QueriesTimeSync);
 DEFINE_STAT(STAT_Navigation_RequestingAsyncPathfinding);
 DEFINE_STAT(STAT_Navigation_PathfindingSync);
 DEFINE_STAT(STAT_Navigation_PathfindingAsync);
-DEFINE_STAT(STAT_Navigation_AddTile);
+DEFINE_STAT(STAT_Navigation_AddGeneratedTiles);
 DEFINE_STAT(STAT_Navigation_TileNavAreaSorting);
 DEFINE_STAT(STAT_Navigation_TileGeometryExportToObjAsync);
 DEFINE_STAT(STAT_Navigation_TileVoxelFilteringAsync);
@@ -410,6 +410,23 @@ void UNavigationSystem::OnWorldInitDone(FNavigationSystem::EMode Mode)
 				}
 			}
 		}
+
+		// All navigation actors are registered
+		// Add NavMesh parts from all sub-levels that were streamed in prior NavMesh registration
+		if (World->IsGameWorld())
+		{
+			const auto& Levels = World->GetLevels();
+			for (ULevel* Level : Levels)
+			{
+				if (!Level->IsPersistentLevel() && Level->bIsVisible)
+				{
+					for (ANavigationData* NavData : NavDataSet)
+					{
+						NavData->OnStreamingLevelAdded(Level);
+					}
+				}
+			}
+		}
 	}
 }
 
@@ -438,6 +455,8 @@ void UNavigationSystem::SetCrowdManager(UCrowdManager* NewCrowdManager)
 
 void UNavigationSystem::Tick(float DeltaSeconds)
 {
+	const bool bIsGame = (GetWorld() && GetWorld()->IsGameWorld());
+	
 	// Register any pending nav areas
 	if (PendingNavAreaRegistration.Num() > 0)
 	{
@@ -467,13 +486,13 @@ void UNavigationSystem::Tick(float DeltaSeconds)
 		}
 		INC_FLOAT_STAT_BY(STAT_Navigation_CumulativeBuildTime,(float)ThisTime*1000);
 	}
-
+		
 	{
 		SCOPE_CYCLE_COUNTER(STAT_Navigation_TickMarkDirty);
 
 		DirtyAreasUpdateTime += DeltaSeconds;
 		const float DirtyAreasUpdateDeltaTime = 1.0f / DirtyAreasUpdateFreq;
-		const bool bCanRebuildNow = (DirtyAreasUpdateTime >= DirtyAreasUpdateDeltaTime) || (GetWorld() && !GetWorld()->IsGameWorld());
+		const bool bCanRebuildNow = (DirtyAreasUpdateTime >= DirtyAreasUpdateDeltaTime) || !bIsGame;
 
 		if (DirtyAreas.Num() > 0 && bCanRebuildNow)
 		{
@@ -492,7 +511,7 @@ void UNavigationSystem::Tick(float DeltaSeconds)
 	}
 
 	// Tick navigation mesh async builders
-	if (!bAsyncBuildPaused)
+	if (!bAsyncBuildPaused && (bNavigationAutoUpdateEnabled || bIsGame))
 	{
 		SCOPE_CYCLE_COUNTER(STAT_Navigation_TickAsyncBuild);
 		for (ANavigationData* NavData : NavDataSet)
@@ -1737,6 +1756,17 @@ void UNavigationSystem::InitializeForWorld(UWorld* World, FNavigationSystem::EMo
 			NavSys = CreateNavigationSystem(World);
 		}
 
+		// Remove old chunk data from all levels
+		// In case navigation system will be created chunks will be regenerated anyway
+		if (Mode == FNavigationSystem::EditorMode)
+		{
+			const auto& Levels = World->GetLevels();
+			for (ULevel* Level : Levels)
+			{
+				Level->NavDataChunks.Empty();
+			}
+		}
+
 		if (NavSys)
 		{
 			NavSys->OnWorldInitDone(Mode);
@@ -1909,6 +1939,12 @@ void UNavigationSystem::AddElementToNavOctree(const FNavigationDirtyElement& Dir
 
 	const FBox BBox = GeneratedData.Bounds.GetBox();
 	const bool bValidBBox = BBox.IsValid && !BBox.GetSize().IsNearlyZero();
+
+	if (BBox.GetExtent().X > 400000)
+	{
+		volatile int32 i = 0;
+	}
+
 	if (bValidBBox && !GeneratedData.IsEmpty())
 	{
 		const int32 DirtyFlag = DirtyElement.FlagsOverride ? DirtyElement.FlagsOverride : GeneratedData.Data.GetDirtyFlag();
@@ -2303,17 +2339,13 @@ void UNavigationSystem::InitializeLevelCollisions()
 	UWorld* World = GetWorld();
 	if (!bInitialLevelsAdded && UNavigationSystem::GetCurrent(World) == this)
 	{
-		if (World->PersistentLevel)
-		{
-			OnLevelAddedToWorld(World->PersistentLevel, World);
-		}
-
-		for (int32 Idx = 0; Idx < World->StreamingLevels.Num(); Idx++)
+		// Process all visible levels
+		const auto& Levels = World->GetLevels();
+		for (ULevel* Level : Levels)
 		{
-			ULevelStreaming* StreamingInfo = World->StreamingLevels[Idx];
-			if (StreamingInfo && StreamingInfo->GetLoadedLevel())
+			if (Level->bIsVisible)
 			{
-				OnLevelAddedToWorld(StreamingInfo->GetLoadedLevel(), World);
+				AddLevelCollisionToOctree(Level);
 			}
 		}
 
@@ -2365,6 +2397,7 @@ void UNavigationSystem::OnNavigationBoundsUpdated(ANavMeshBoundsVolume* NavVolum
 	FNavigationBoundsUpdateRequest UpdateRequest;
 	UpdateRequest.NavBounds.UniqueID	= NavVolume->GetUniqueID();
 	UpdateRequest.NavBounds.AreaBox		= NavVolume->GetComponentsBoundingBox(true);
+	UpdateRequest.NavBounds.PackageName	= NavVolume->GetOutermost()->GetFName();
 	UpdateRequest.UpdateRequest			= FNavigationBoundsUpdateRequest::Updated;
 	AddNavigationBoundsUpdateRequest(UpdateRequest);
 }
@@ -2379,6 +2412,7 @@ void UNavigationSystem::OnNavigationBoundsAdded(ANavMeshBoundsVolume* NavVolume)
 	FNavigationBoundsUpdateRequest UpdateRequest;
 	UpdateRequest.NavBounds.UniqueID	= NavVolume->GetUniqueID();
 	UpdateRequest.NavBounds.AreaBox		= NavVolume->GetComponentsBoundingBox(true);
+	UpdateRequest.NavBounds.PackageName	= NavVolume->GetOutermost()->GetFName();
 	UpdateRequest.UpdateRequest	= FNavigationBoundsUpdateRequest::Added;
 	AddNavigationBoundsUpdateRequest(UpdateRequest);
 }
@@ -2393,6 +2427,7 @@ void UNavigationSystem::OnNavigationBoundsRemoved(ANavMeshBoundsVolume* NavVolum
 	FNavigationBoundsUpdateRequest UpdateRequest;
 	UpdateRequest.NavBounds.UniqueID	= NavVolume->GetUniqueID();
 	UpdateRequest.NavBounds.AreaBox		= NavVolume->GetComponentsBoundingBox(true);
+	UpdateRequest.NavBounds.PackageName	= NavVolume->GetOutermost()->GetFName();
 	UpdateRequest.UpdateRequest	= FNavigationBoundsUpdateRequest::Removed;
 	AddNavigationBoundsUpdateRequest(UpdateRequest);
 }
@@ -2480,7 +2515,7 @@ void UNavigationSystem::PerformNavigationBoundsUpdate(const TArray<FNavigationBo
 	}
 
 #if WITH_RECAST
-	if (IsNavigationBuildingLocked() == false)
+	if (!IsNavigationBuildingLocked())
 	{
 		if (UpdatedAreas.Num())
 		{
@@ -2509,8 +2544,9 @@ void UNavigationSystem::GatherNavigationBounds()
 		if (V != nullptr && !V->IsPendingKill())
 		{
 			FNavigationBounds NavBounds;
-			NavBounds.UniqueID	= V->GetUniqueID();
-			NavBounds.AreaBox	= V->GetComponentsBoundingBox(true);
+			NavBounds.UniqueID		= V->GetUniqueID();
+			NavBounds.AreaBox		= V->GetComponentsBoundingBox(true);
+			NavBounds.PackageName	= V->GetOutermost()->GetFName();
 			RegisteredNavBounds.Add(NavBounds);
 		}
 	}
@@ -2535,9 +2571,18 @@ void UNavigationSystem::Build()
 	// make sure freshly created navigation instances are registered before we try to build them
 	ProcessRegistrationCandidates();
 
-	// and now iterate through all registered and just build them
+	// and now iterate through all registered and just start building them
 	RebuildAll();
 
+	// Block until build is finished
+	for (ANavigationData* NavData : NavDataSet)
+	{
+		if (NavData)
+		{
+			NavData->EnsureBuildCompletion();
+		}
+	}
+
 	UE_LOG(LogNavigation, Display, TEXT("UNavigationSystem::Build total execution time: %.5f"), float(FPlatformTime::Seconds() - BuildStartTime));
 }
 
@@ -2693,7 +2738,10 @@ void UNavigationSystem::NavigationBuildingUnlock(bool bForce)
 		bNavigationBuildingLocked = false;
 		bInitialBuildingLockActive = false;
 		
-		RebuildAll();
+		if (bNavigationAutoUpdateEnabled)
+		{
+			RebuildAll();
+		}
 	}
 	else if (bInitialBuildingLockActive == true)
 	{
@@ -2710,6 +2758,18 @@ void UNavigationSystem::RebuildAll()
 	
 	GatherNavigationBounds();
 
+	// make sure that octree is up to date
+	for (TSet<FNavigationDirtyElement>::TIterator It(PendingOctreeUpdates); It; ++It)
+	{
+		AddElementToNavOctree(*It);
+	}
+	PendingOctreeUpdates.Empty(32);
+
+	// discard all pending dirty areas, we are going to rebuild navmesh anyway 
+	DirtyAreas.Reset();
+	PendingNavBoundsUpdates.Reset();
+	
+	// 
 	for (int32 NavDataIndex = 0; NavDataIndex < NavDataSet.Num(); ++NavDataIndex)
 	{
 		ANavigationData* NavData = NavDataSet[NavDataIndex];
@@ -2780,6 +2840,14 @@ void UNavigationSystem::OnLevelAddedToWorld(ULevel* InLevel, UWorld* InWorld)
 	if (InWorld == GetWorld())
 	{
 		AddLevelCollisionToOctree(InLevel);
+
+		if (!InLevel->IsPersistentLevel())
+		{
+			for (ANavigationData* NavData : NavDataSet)
+			{
+				NavData->OnStreamingLevelAdded(InLevel);
+			}
+		}
 	}
 }
 
@@ -2788,6 +2856,14 @@ void UNavigationSystem::OnLevelRemovedFromWorld(ULevel* InLevel, UWorld* InWorld
 	if (InWorld == GetWorld())
 	{
 		RemoveLevelCollisionFromOctree(InLevel);
+
+		if (!InLevel->IsPersistentLevel())
+		{
+			for (ANavigationData* NavData : NavDataSet)
+			{
+				NavData->OnStreamingLevelRemoved(InLevel);
+			}
+		}
 	}
 }
 
diff --git a/Engine/Source/Runtime/Engine/Private/AI/Navigation/PImplRecastNavMesh.cpp b/Engine/Source/Runtime/Engine/Private/AI/Navigation/PImplRecastNavMesh.cpp
--- a/Engine/Source/Runtime/Engine/Private/AI/Navigation/PImplRecastNavMesh.cpp
+++ b/Engine/Source/Runtime/Engine/Private/AI/Navigation/PImplRecastNavMesh.cpp
@@ -293,18 +293,33 @@ void FPImplRecastNavMesh::Serialize( FArchive& Ar )
 	}
 
 	int32 NumTiles = 0;
+	TArray<int32> TilesToSave;
 
 	if (Ar.IsSaving())
 	{
-		dtNavMesh const* ConstNavMesh = DetourNavMesh;
-		for (int i = 0; i < ConstNavMesh->getMaxTiles(); ++i)
+		TilesToSave.Reserve(DetourNavMesh->getMaxTiles());
+
+		if (!NavMeshOwner->bRebuildAtRuntime && !IsRunningCommandlet())
 		{
-			const dtMeshTile* Tile = ConstNavMesh->getTile(i);
-			if (Tile != NULL && Tile->header != NULL && Tile->dataSize > 0)
+			// For static navmeshes we save only tiles that belongs to this level
+			FName LevelPackageName = NavMeshOwner->GetOutermost()->GetFName();
+			GetNavMeshTilesIn(NavMeshOwner->GetNavigableBoundsInLevel(LevelPackageName), TilesToSave);
+		}
+		else
+		{
+			// Otherwise all valid tiles
+			dtNavMesh const* ConstNavMesh = DetourNavMesh;
+			for (int i = 0; i < ConstNavMesh->getMaxTiles(); ++i)
 			{
-				++NumTiles;
+				const dtMeshTile* Tile = ConstNavMesh->getTile(i);
+				if (Tile != NULL && Tile->header != NULL && Tile->dataSize > 0)
+				{
+					TilesToSave.Add(i);
+				}
 			}
 		}
+		
+		NumTiles = TilesToSave.Num();
 	}
 
 	Ar << NumTiles;
@@ -385,18 +400,15 @@ void FPImplRecastNavMesh::Serialize( FArchive& Ar )
 	else if (Ar.IsSaving())
 	{
 		dtNavMesh const* ConstNavMesh = DetourNavMesh;
-		for (int i = 0; i < ConstNavMesh->getMaxTiles(); ++i)
+		for (int TileIndex : TilesToSave)
 		{
-			const dtMeshTile* Tile = ConstNavMesh->getTile(i);
-			if (Tile != NULL && Tile->header != NULL && Tile->dataSize > 0)
-			{
-				dtTileRef TileRef = ConstNavMesh->getTileRef(Tile);
-				int32 TileDataSize = Tile->dataSize;
-				Ar << TileRef << TileDataSize;
+			const dtMeshTile* Tile = ConstNavMesh->getTile(TileIndex);
+			dtTileRef TileRef = ConstNavMesh->getTileRef(Tile);
+			int32 TileDataSize = Tile->dataSize;
+			Ar << TileRef << TileDataSize;
 
-				unsigned char* TileData = Tile->data;
-				SerializeRecastMeshTile(Ar, TileData, TileDataSize);
-			}
+			unsigned char* TileData = Tile->data;
+			SerializeRecastMeshTile(Ar, TileData, TileDataSize);
 		}
 	}
 }
@@ -2247,6 +2259,54 @@ void FPImplRecastNavMesh::GetNavMeshTilesAt(int32 TileX, int32 TileY, TArray<int
 	}
 }
 
+void FPImplRecastNavMesh::GetNavMeshTilesIn(const TArray<FBox>& InclusionBounds, TArray<int32>& Indices) const
+{
+	if (DetourNavMesh)
+	{
+		const float* NavMeshOrigin = DetourNavMesh->getParams()->orig;
+		const float TileSize = DetourNavMesh->getParams()->tileWidth;
+
+		TSet<FIntPoint>	TileCoords;	
+		for (const FBox& Bounds : InclusionBounds)
+		{
+			const FBox RcBounds = Unreal2RecastBox(Bounds);
+			const int32 XMin = FMath::FloorToInt((RcBounds.Min.X - NavMeshOrigin[0]) / TileSize);
+			const int32 XMax = FMath::FloorToInt((RcBounds.Max.X - NavMeshOrigin[0]) / TileSize);
+			const int32 YMin = FMath::FloorToInt((RcBounds.Min.Z - NavMeshOrigin[2]) / TileSize);
+			const int32 YMax = FMath::FloorToInt((RcBounds.Max.Z - NavMeshOrigin[2]) / TileSize);
+
+			for (int32 y = YMin; y <= YMax; ++y)
+			{
+				for (int32 x = XMin; x <= XMax; ++x)
+				{
+					TileCoords.Add(FIntPoint(x, y));
+				}
+			}
+		}
+
+		// We guess that each tile has 3 layers in average
+		Indices.Reserve(TileCoords.Num()*3);
+
+		for (FIntPoint TileCoord : TileCoords)
+		{
+			int32 MaxTiles = DetourNavMesh->getTileCountAt(TileCoord.X, TileCoord.Y);
+			TArray<const dtMeshTile*> Tiles;
+			Tiles.AddZeroed(MaxTiles);
+
+			const int32 NumTiles = DetourNavMesh->getTilesAt(TileCoord.X, TileCoord.Y, Tiles.GetData(), MaxTiles);
+			for (int32 i = 0; i < NumTiles; ++i)
+			{
+				dtTileRef TileRef = DetourNavMesh->getTileRef(Tiles[i]);
+				if (TileRef)
+				{
+					const int32 TileIndex = (int32)DetourNavMesh->decodePolyIdTile(TileRef);
+					Indices.Add(TileIndex);
+				}
+			}
+		}
+	}
+}
+
 float FPImplRecastNavMesh::GetTotalDataSize() const
 {
 	float TotalBytes = sizeof(*this);
diff --git a/Engine/Source/Runtime/Engine/Private/AI/Navigation/RecastNavMesh.cpp b/Engine/Source/Runtime/Engine/Private/AI/Navigation/RecastNavMesh.cpp
--- a/Engine/Source/Runtime/Engine/Private/AI/Navigation/RecastNavMesh.cpp
+++ b/Engine/Source/Runtime/Engine/Private/AI/Navigation/RecastNavMesh.cpp
@@ -7,6 +7,7 @@
 #include "AI/Navigation/NavAreas/NavArea_LowHeight.h"
 #include "AI/Navigation/NavLinkCustomInterface.h"
 #include "AI/Navigation/RecastNavMesh.h"
+#include "AI/Navigation/RecastNavMeshDataChunk.h"
 #include "VisualLogger/VisualLogger.h"
 
 #if WITH_EDITOR
@@ -45,6 +46,21 @@ FNavMeshTileData::FNavData::~FNavData()
 #endif
 }
 
+uint8* FNavMeshTileData::Release()
+{
+	uint8* RawData = nullptr;
+
+	if (NavData.IsValid()) 
+	{ 
+		RawData = NavData->RawNavData;
+		NavData->RawNavData = nullptr; 
+	} 
+ 
+	DataSize = 0; 
+	LayerIndex = 0; 
+	return RawData;
+}
+
 float ARecastNavMesh::DrawDistanceSq = 0.0f;
 #if !WITH_RECAST
 
@@ -607,14 +623,6 @@ void ARecastNavMesh::SerializeRecastNavMesh(FArchive& Ar, FPImplRecastNavMesh*&
 		}
 	}
 	
-	if (Ar.IsSaving())
-	{
-		if (NavDataGenerator)
-		{
-			NavDataGenerator->EnsureBuildCompletion();
-		}
-	}
-
 	if (RecastNavMeshImpl)
 	{
 		RecastNavMeshImpl->Serialize(Ar);
@@ -1218,6 +1226,77 @@ void ARecastNavMesh::InvalidateAffectedPaths(const TArray<uint32>& ChangedTiles)
 	}
 }
 
+URecastNavMeshDataChunk* ARecastNavMesh::GetNavigationDataChunk(ULevel* InLevel) const
+{
+	FName ThisName = GetFName();
+	int32 ChunkIndex = InLevel->NavDataChunks.IndexOfByPredicate([&](UNavigationDataChunk* Chunk) 
+	{
+		return Chunk->NavigationDataName == ThisName;
+	});
+	
+	URecastNavMeshDataChunk* RcNavDataChunk = nullptr;
+	if (ChunkIndex != INDEX_NONE)
+	{
+		RcNavDataChunk = Cast<URecastNavMeshDataChunk>(InLevel->NavDataChunks[ChunkIndex]);
+	}
+		
+	return RcNavDataChunk;
+}
+
+void ARecastNavMesh::OnNavMeshGenerationFinished()
+{
+#if WITH_EDITOR
+	UWorld* World = GetWorld();
+	// For static navmeshes create navigation data holders in each streaming level
+	// so parts of navmesh can be streamed in/out with those levels
+	if (!World->IsGameWorld())
+	{
+		const auto& Levels = World->GetLevels();
+		for (auto Level : Levels)
+		{
+			if (Level->IsPersistentLevel())
+			{
+				continue;
+			}
+			
+			URecastNavMeshDataChunk* NavDataChunk = GetNavigationDataChunk(Level);
+				
+			if (!bRebuildAtRuntime)
+			{
+				// We use nav volumes that belongs to this streaming level to find tiles we want to save
+				TArray<int32> LevelTiles;
+				TArray<FBox> LevelNavBounds = GetNavigableBoundsInLevel(Level->GetOutermost()->GetFName());
+				RecastNavMeshImpl->GetNavMeshTilesIn(LevelNavBounds, LevelTiles);
+				
+				if (LevelTiles.Num())
+				{
+					// Create new chunk only if we have something to save in it			
+					if (NavDataChunk == nullptr)
+					{
+						NavDataChunk = NewObject<URecastNavMeshDataChunk>(Level);
+						NavDataChunk->NavigationDataName = GetFName();
+						Level->NavDataChunks.Add(NavDataChunk);
+					}
+					
+					NavDataChunk->GatherTiles(RecastNavMeshImpl->DetourNavMesh, LevelTiles);
+					NavDataChunk->MarkPackageDirty();
+					continue;
+				}
+			}
+			
+			// stale data that is left in the level
+			if (NavDataChunk)
+			{
+				// clear it
+				NavDataChunk->ReleaseTiles();
+				NavDataChunk->MarkPackageDirty();
+				Level->NavDataChunks.Remove(NavDataChunk);
+			}
+		}
+	}
+#endif// WITH_EDITOR
+}
+
 #if !UE_BUILD_SHIPPING
 uint32 ARecastNavMesh::LogMemUsed() const 
 {
@@ -1268,6 +1347,44 @@ void ARecastNavMesh::ApplyWorldOffset(const FVector& InOffset, bool bWorldShift)
 	RequestDrawingUpdate();
 }
 
+void ARecastNavMesh::OnStreamingLevelAdded(ULevel* InLevel)
+{
+	QUICK_SCOPE_CYCLE_COUNTER(STAT_RecastNavMesh_OnStreamingLevelAdded);
+	
+	if (!bRebuildAtRuntime && GetWorld()->IsGameWorld())
+	{
+		URecastNavMeshDataChunk* NavDataChunk = GetNavigationDataChunk(InLevel);
+		if (NavDataChunk)
+		{
+			TArray<uint32> AttachedIndices = NavDataChunk->AttachTiles(RecastNavMeshImpl->DetourNavMesh);
+			if (AttachedIndices.Num() > 0)
+			{
+				InvalidateAffectedPaths(AttachedIndices);
+				RequestDrawingUpdate();
+			}
+		}
+	}
+}
+
+void ARecastNavMesh::OnStreamingLevelRemoved(ULevel* InLevel)
+{
+	QUICK_SCOPE_CYCLE_COUNTER(STAT_RecastNavMesh_OnStreamingLevelRemoved);
+	
+	if (!bRebuildAtRuntime && GetWorld()->IsGameWorld())
+	{
+		URecastNavMeshDataChunk* NavDataChunk = GetNavigationDataChunk(InLevel);
+		if (NavDataChunk)
+		{
+			TArray<uint32> DetachedIndices = NavDataChunk->DetachTiles(RecastNavMeshImpl->DetourNavMesh);
+			if (DetachedIndices.Num() > 0)
+			{
+				InvalidateAffectedPaths(DetachedIndices);
+				RequestDrawingUpdate();
+			}
+		}
+	}
+}
+
 bool ARecastNavMesh::AdjustLocationWithFilter(const FVector& StartLoc, FVector& OutAdjustedLocation, const FNavigationQueryFilter& Filter, const UObject* QueryOwner) const
 {
 	INITIALIZE_NAVQUERY(NavQuery, Filter.GetMaxSearchNodes());
diff --git a/Engine/Source/Runtime/Engine/Private/AI/Navigation/RecastNavMeshDataChunk.cpp b/Engine/Source/Runtime/Engine/Private/AI/Navigation/RecastNavMeshDataChunk.cpp
new file mode 100644
--- /dev/null
+++ b/Engine/Source/Runtime/Engine/Private/AI/Navigation/RecastNavMeshDataChunk.cpp
@@ -0,0 +1,198 @@
+// Copyright 1998-2014 Epic Games, Inc. All Rights Reserved.
+
+#include "EnginePrivate.h"
+#include "AI/Navigation/RecastNavMeshDataChunk.h"
+#include "AI/Navigation/PImplRecastNavMesh.h"
+
+//----------------------------------------------------------------------//
+// FRecastTileData                                                                
+//----------------------------------------------------------------------//
+FRecastTileData::FRawData::FRawData(uint8* InData)
+	: RawData(InData)
+{
+}
+
+FRecastTileData::FRawData::~FRawData()
+{
+#if WITH_RECAST
+	dtFree(RawData);
+#else
+	delete RawData;
+#endif
+}
+
+FRecastTileData::FRecastTileData()
+	: TileRef(0)
+	, TileDataSize(0)
+{
+}
+	
+FRecastTileData::FRecastTileData(NavNodeRef Ref, int32 DataSize, uint8* RawData)
+	: TileRef(Ref)
+	, TileDataSize(DataSize)
+{
+	TileRawData = MakeShareable(new FRawData(RawData));
+}
+
+//----------------------------------------------------------------------//
+// URecastNavMeshDataChunk                                                                
+//----------------------------------------------------------------------//
+URecastNavMeshDataChunk::URecastNavMeshDataChunk(const FObjectInitializer& ObjectInitializer)
+	: Super(ObjectInitializer)
+{
+}
+
+void URecastNavMeshDataChunk::Serialize(FArchive& Ar)
+{
+	Super::Serialize(Ar);
+
+	int32 NavMeshVersion = NAVMESHVER_LATEST;
+	Ar << NavMeshVersion;
+
+	// when writing, write a zero here for now.  will come back and fill it in later.
+	int64 RecastNavMeshSizeBytes = 0;
+	int64 RecastNavMeshSizePos = Ar.Tell();
+	Ar << RecastNavMeshSizeBytes;
+
+	if (Ar.IsLoading())
+	{
+		if (NavMeshVersion < NAVMESHVER_MIN_COMPATIBLE)
+		{
+			// incompatible, just skip over this data.  navmesh needs rebuilt.
+			Ar.Seek(RecastNavMeshSizePos + RecastNavMeshSizeBytes);
+		}
+		else if (RecastNavMeshSizeBytes > 4)
+		{
+			SerializeRecastData(Ar);
+		}
+		else
+		{
+			// empty, just skip over this data
+			Ar.Seek(RecastNavMeshSizePos + RecastNavMeshSizeBytes);
+		}
+	}
+	else
+	{
+		SerializeRecastData(Ar);
+
+		if (Ar.IsSaving())
+		{
+			int64 CurPos = Ar.Tell();
+			RecastNavMeshSizeBytes = CurPos - RecastNavMeshSizePos;
+			Ar.Seek(RecastNavMeshSizePos);
+			Ar << RecastNavMeshSizeBytes;
+			Ar.Seek(CurPos);
+		}
+	}
+}
+
+void URecastNavMeshDataChunk::SerializeRecastData(FArchive& Ar)
+{
+	int32 TileNum = Tiles.Num();
+	Ar << TileNum;
+
+	if (Ar.IsLoading())
+	{
+		Tiles.Empty(TileNum);
+		for (int32 TileIdx = 0; TileIdx < TileNum; TileIdx++)
+		{
+			int32 TileDataSize = 0;
+			Ar << TileDataSize;
+
+			uint8* TileRawData = nullptr;
+			FPImplRecastNavMesh::SerializeRecastMeshTile(Ar, TileRawData, TileDataSize);
+			
+			if (TileRawData != nullptr)
+			{
+				FRecastTileData TileData(0, TileDataSize, TileRawData);
+				Tiles.Add(TileData);
+			}
+		}
+	}
+	else if (Ar.IsSaving())
+	{
+		for (FRecastTileData& TileData : Tiles)
+		{
+			if (TileData.TileRawData.IsValid())
+			{
+				Ar << TileData.TileDataSize;
+				FPImplRecastNavMesh::SerializeRecastMeshTile(Ar, TileData.TileRawData->RawData, TileData.TileDataSize);
+			}
+		}
+	}
+}
+
+TArray<uint32> URecastNavMeshDataChunk::AttachTiles(dtNavMesh* NavMesh)
+{
+	TArray<uint32> Result;
+	Result.Reserve(Tiles.Num());
+	
+	for (FRecastTileData& TileData : Tiles)
+	{
+		if (TileData.TileRawData.IsValid())
+		{
+			dtStatus status = NavMesh->addTile(TileData.TileRawData->RawData, TileData.TileDataSize, DT_TILE_FREE_DATA, 0, &TileData.TileRef);
+			if (dtStatusFailed(status))
+			{
+				TileData.TileRef = 0;
+				continue;
+			}
+			
+			// We don't own data anymore
+			TileData.TileDataSize = 0;
+			TileData.TileRawData->RawData = nullptr;
+
+			Result.Add(NavMesh->decodePolyIdTile(TileData.TileRef));
+		}
+	}
+
+	UE_LOG(LogNavigation, Log, TEXT("Attached %d tiles to NavMesh - %s"), Result.Num(), *NavigationDataName.ToString());
+	return Result;
+}
+
+TArray<uint32> URecastNavMeshDataChunk::DetachTiles(dtNavMesh* NavMesh)
+{
+	TArray<uint32> Result;
+	Result.Reserve(Tiles.Num());
+
+	for (FRecastTileData& TileData : Tiles)
+	{
+		if (TileData.TileRef != 0)
+		{
+			NavMesh->removeTile(TileData.TileRef, &TileData.TileRawData->RawData, &TileData.TileDataSize);
+			Result.Add(NavMesh->decodePolyIdTile(TileData.TileRef));
+		}
+	}
+
+	UE_LOG(LogNavigation, Log, TEXT("Detached %d tiles from NavMesh - %s"), Result.Num(), *NavigationDataName.ToString());
+	return Result;
+}
+
+int32 URecastNavMeshDataChunk::GetNumTiles() const
+{
+	return Tiles.Num();
+}
+
+void URecastNavMeshDataChunk::ReleaseTiles()
+{
+	Tiles.Reset();
+}
+
+#if WITH_EDITOR
+void URecastNavMeshDataChunk::GatherTiles(const dtNavMesh* NavMesh, const TArray<int32>& TileIndices)
+{
+	Tiles.Empty(TileIndices.Num());
+	
+	for (int32 TileIdx : TileIndices)
+	{
+		const dtMeshTile* Tile = NavMesh->getTile(TileIdx);
+		if (Tile && Tile->header)
+		{
+			uint8* RawTileData = (uint8*)dtAlloc(Tile->dataSize, DT_ALLOC_PERM);
+			FMemory::Memcpy(RawTileData, Tile->data, Tile->dataSize);
+			FRecastTileData RecastTileData(NavMesh->getTileRef(Tile), Tile->dataSize, RawTileData);
+			Tiles.Add(RecastTileData);
+		}
+	}
+}
+#endif// WITH_EDITOR
\ No newline at end of file
diff --git a/Engine/Source/Runtime/Engine/Private/AI/Navigation/RecastNavMeshGenerator.cpp b/Engine/Source/Runtime/Engine/Private/AI/Navigation/RecastNavMeshGenerator.cpp
--- a/Engine/Source/Runtime/Engine/Private/AI/Navigation/RecastNavMeshGenerator.cpp
+++ b/Engine/Source/Runtime/Engine/Private/AI/Navigation/RecastNavMeshGenerator.cpp
@@ -2926,7 +2926,7 @@ bool FRecastNavMeshGenerator::RebuildAll()
 	}
 	else
 	{
-		// There are no navigation bounds to build, probably navmesh was resided and we just need to update debug draw
+		// There are no navigation bounds to build, probably navmesh was resized and we just need to update debug draw
 		DestNavMesh->RequestDrawingUpdate();
 	}
 
@@ -2935,6 +2935,8 @@ bool FRecastNavMeshGenerator::RebuildAll()
 
 void FRecastNavMeshGenerator::EnsureBuildCompletion()
 {
+	const bool bHasTasks = GetNumRemaningBuildTasks() > 0;
+	
 	do 
 	{
 		ProcessTileTasks(16);
@@ -2945,7 +2947,13 @@ void FRecastNavMeshGenerator::EnsureBuildCompletion()
 			Element.AsyncTask->EnsureCompletion();
 		}
 	}
-	while (PendingDirtyTiles.Num() || RunningDirtyTiles.Num());
+	while (GetNumRemaningBuildTasks() > 0);
+
+	// Update navmesh drawing only if we had something to build
+	if (bHasTasks)
+	{
+		DestNavMesh->RequestDrawingUpdate();
+	}
 }
 
 void FRecastNavMeshGenerator::CancelBuild()
@@ -3047,19 +3055,6 @@ void FRecastNavMeshGenerator::OnAreaAdded(const UClass* AreaClass, int32 AreaID)
 	AdditionalCachedData.OnAreaAdded(AreaClass, AreaID);
 }
 
-bool FRecastNavMeshGenerator::IntercestsInclusionBounds(const FBox& Box) const
-{
-	for (const FBox& NavBox : InclusionBounds)
-	{
-		if (NavBox.Intersect(Box))
-		{
-			return true;
-		}
-	}
-
-	return false;
-}
-
 TArray<uint32> FRecastNavMeshGenerator::RemoveTileLayers(const int32 TileX, const int32 TileY)
 {
 	TArray<uint32> ResultTileIndices;
@@ -3081,8 +3076,7 @@ TArray<uint32> FRecastNavMeshGenerator::RemoveTileLayers(const int32 TileX, cons
 			UE_LOG(LogNavigation, Log, TEXT("%s> Tile (%d,%d:%d), removing TileRef: 0x%X (active:%d)"),
 				*DestNavMesh->GetName(), TileX, TileY, LayerIndex, TileRef, NumActiveTiles);
 
-			uint8* RawNavData = NULL;
-			DetourMesh->removeTile(TileRef, &RawNavData, NULL);
+			DetourMesh->removeTile(TileRef, nullptr, nullptr);
 
 			ResultTileIndices.AddUnique(DetourMesh->decodePolyIdTile(TileRef));
 		}
@@ -3094,9 +3088,9 @@ TArray<uint32> FRecastNavMeshGenerator::RemoveTileLayers(const int32 TileX, cons
 	return ResultTileIndices;
 }
 
-TArray<uint32> FRecastNavMeshGenerator::AddTile(const FRecastTileGenerator& TileGenerator)
+TArray<uint32> FRecastNavMeshGenerator::AddGeneratedTiles(const FRecastTileGenerator& TileGenerator)
 {
-	SCOPE_CYCLE_COUNTER(STAT_Navigation_AddTile);
+	SCOPE_CYCLE_COUNTER(STAT_Navigation_AddGeneratedTiles);
 	
 	TArray<uint32> ResultTileIndices;
 	const int32 TileX = TileGenerator.GetTileX();
@@ -3132,14 +3126,7 @@ TArray<uint32> FRecastNavMeshGenerator::AddTile(const FRecastTileGenerator& Tile
 				UE_LOG(LogNavigation, Log, TEXT("%s> Tile (%d,%d:%d), removing TileRef: 0x%X (active:%d)"),
 					*DestNavMesh->GetName(), TileX, TileY, LayerIndex, OldTileRef, NumActiveTiles);
 
-				// this call will fill OldRawNavData with address of previously added data
-				// if OldRawNavData is NULL it means either this tile was empty, or the memory
-				// allocated has been owned by navmesh itself (i.e. has been serialized-in along
-				// with the level)
-				uint8* OldRawNavData = nullptr;
-				DetourMesh->removeTile(OldTileRef, &OldRawNavData, nullptr);
-				// data should always be owned by navmesh itself
-				check(OldRawNavData == nullptr);
+				DetourMesh->removeTile(OldTileRef, nullptr, nullptr);
 			
 				ResultTileIndices.AddUnique(DetourMesh->decodePolyIdTile(OldTileRef));
 			}
@@ -3171,7 +3158,7 @@ TArray<uint32> FRecastNavMeshGenerator::AddTile(const FRecastTileGenerator& Tile
 						*DestNavMesh->GetName(), TileX, TileY, LayerIndex, ResultTileRef, NumActiveTiles);
 				
 					// NavMesh took the ownership of generated data, so we don't need to deallocate it
-					TileLayers[i].Release();
+					uint8* ReleasedData = TileLayers[i].Release();
 				}
 			}
 		}
@@ -3182,6 +3169,8 @@ TArray<uint32> FRecastNavMeshGenerator::AddTile(const FRecastTileGenerator& Tile
 
 void FRecastNavMeshGenerator::DiscardCurrentBuildingTasks()
 {
+	PendingDirtyTiles.Empty();
+	
 	for (FRunningTileElement& Element : RunningDirtyTiles)
 	{
 		if (Element.AsyncTask)
@@ -3221,6 +3210,19 @@ TArray<FNavMeshTileData> FRecastNavMeshGenerator::GetIntermediateLayersData(FInt
 	}
 }
 
+static bool IntercestBounds(const FBox& TestBox, const TNavStatArray<FBox>& Bounds)
+{
+	for (const FBox& Box : Bounds)
+	{
+		if (Box.Intersect(TestBox))
+		{
+			return true;
+		}
+	}
+
+	return false;
+}
+
 void FRecastNavMeshGenerator::MarkDirtyTiles(const TArray<FNavigationDirtyArea>& DirtyAreas)
 {
 	QUICK_SCOPE_CYCLE_COUNTER(STAT_RecastNavMeshGenerator_MarkDirtyTiles);
@@ -3237,7 +3239,7 @@ void FRecastNavMeshGenerator::MarkDirtyTiles(const TArray<FNavigationDirtyArea>&
 	{
 		const FBox AdjustedAreaBounds = GrowBoundingBox(DirtyArea.Bounds, DirtyArea.HasFlag(ENavigationDirtyFlag::UseAgentHeight));
 		if (!DirtyArea.HasFlag(ENavigationDirtyFlag::NavigationBounds) && 
-			!IntercestsInclusionBounds(AdjustedAreaBounds))
+			!IntercestBounds(AdjustedAreaBounds, InclusionBounds))
 		{
 			// Skip whole area
 			continue;
@@ -3255,7 +3257,7 @@ void FRecastNavMeshGenerator::MarkDirtyTiles(const TArray<FNavigationDirtyArea>&
 			{
 				const FBox TileBox = CalculateTileBounds(x, y, NavMeshOrigin, TotalNavBounds, TileSizeInWorldUnits);
 				if (!DirtyArea.HasFlag(ENavigationDirtyFlag::NavigationBounds) && 
-					!IntercestsInclusionBounds(TileBox))
+					!IntercestBounds(TileBox, InclusionBounds))
 				{
 					// Skip this tile
 					continue;
@@ -3384,6 +3386,7 @@ TArray<uint32> FRecastNavMeshGenerator::ProcessTileTasks(const int32 NumTasksToS
 	QUICK_SCOPE_CYCLE_COUNTER(STAT_RecastNavMeshGenerator_ProcessTileTasks);
 	
 	TArray<uint32> UpdatedTiles;
+	const bool bHasTasksAtStart = GetNumRemaningBuildTasks() > 0;
 	
 	int32 NumSubmittedTasks = 0;
 	// Submit pending tile elements
@@ -3446,7 +3449,7 @@ TArray<uint32> FRecastNavMeshGenerator::ProcessTileTasks(const int32 NumTasksToS
 			if (!Element.bShouldDiscard)
 			{
 				const FRecastTileGenerator& TileGenerator = Element.AsyncTask->GetTask();
-				TArray<uint32> UpdatedTileIndices = AddTile(TileGenerator);
+				TArray<uint32> UpdatedTileIndices = AddGeneratedTiles(TileGenerator);
 				UpdatedTiles.Append(UpdatedTileIndices);
 			
 				// Store intermediate layers data, so it can be reused later
@@ -3470,6 +3473,13 @@ TArray<uint32> FRecastNavMeshGenerator::ProcessTileTasks(const int32 NumTasksToS
 		}
 	}
 
+	// Notify owner in case all tasks has been completed
+	const bool bHasTasksAtEnd = GetNumRemaningBuildTasks() > 0;
+	if (bHasTasksAtStart && !bHasTasksAtEnd)
+	{
+		DestNavMesh->OnNavMeshGenerationFinished();
+	}
+	
 	return UpdatedTiles;
 }
 
diff --git a/Engine/Source/Runtime/Engine/Private/EnginePrivate.h b/Engine/Source/Runtime/Engine/Private/EnginePrivate.h
--- a/Engine/Source/Runtime/Engine/Private/EnginePrivate.h
+++ b/Engine/Source/Runtime/Engine/Private/EnginePrivate.h
@@ -259,6 +259,7 @@ EnginePrivate.h: Unreal engine private header file.
 #include "AI/Navigation/NavAreas/NavArea.h"
 #include "AI/Navigation/RecastNavMesh.h"
 #include "AI/Navigation/NavigationGraphNode.h"
+#include "AI/Navigation/NavigationDataChunk.h"
 #include "Engine/NavigationObjectBase.h"
 #include "GameFramework/PlayerStart.h"
 #include "Engine/PlayerStartPIE.h"
diff --git a/Engine/Source/Runtime/Engine/Public/AI/Navigation/PImplRecastNavMesh.h b/Engine/Source/Runtime/Engine/Public/AI/Navigation/PImplRecastNavMesh.h
--- a/Engine/Source/Runtime/Engine/Public/AI/Navigation/PImplRecastNavMesh.h
+++ b/Engine/Source/Runtime/Engine/Public/AI/Navigation/PImplRecastNavMesh.h
@@ -88,6 +88,9 @@ public:
 	/** Retrieves all tile indices at matching XY coordinates */
 	void GetNavMeshTilesAt(int32 TileX, int32 TileY, TArray<int32>& Indices) const;
 
+	/** Retrieves list of tiles that intersect specified bounds */
+	void GetNavMeshTilesIn(const TArray<FBox>& InclusionBounds, TArray<int32>& Indices) const;
+
 	/** Retrieves number of tiles in this navmesh */
 	FORCEINLINE int32 GetNavMeshTilesCount() const { return DetourNavMesh->getMaxTiles(); }
 
diff --git a/Engine/Source/Runtime/Engine/Public/AI/Navigation/RecastNavMeshGenerator.h b/Engine/Source/Runtime/Engine/Public/AI/Navigation/RecastNavMeshGenerator.h
--- a/Engine/Source/Runtime/Engine/Public/AI/Navigation/RecastNavMeshGenerator.h
+++ b/Engine/Source/Runtime/Engine/Public/AI/Navigation/RecastNavMeshGenerator.h
@@ -333,15 +333,11 @@ public:
 
 	const TNavStatArray<FBox>& GetInclusionBounds() const { return InclusionBounds; }
 	
-	/** Whether specified box in inside or intersects with nav bounds  */
-	bool IntercestsInclusionBounds(const FBox& Box) const;
-
 	/** Total navigable area box, sum of all navigation volumes bounding boxes */
 	FBox GetTotalBounds() const { return TotalNavBounds; }
 
 	const FRecastNavMeshCachedData& GetAdditionalCachedData() const { return AdditionalCachedData; }
 
-	TArray<uint32> AddTile(const FRecastTileGenerator& TileGenerator);
 	bool HasDirtyTiles() const;
 
 	FBox GrowBoundingBox(const FBox& BBox, bool bIncludeAgentHeight) const;
@@ -381,6 +377,9 @@ private:
 	/** Processes pending tile generattion tasks */
 	TArray<uint32> ProcessTileTasks(const int32 NumTasksToSubmit);
 
+	/** Adds generated tiles to NavMesh, replacing old ones */
+	TArray<uint32> AddGeneratedTiles(const FRecastTileGenerator& TileGenerator);
+
 	/** Removes all tiles at specified grid location */
 	TArray<uint32> RemoveTileLayers(const int32 TileX, const int32 TileY);
 	
diff --git a/Engine/Source/Runtime/Engine/Public/EngineStats.h b/Engine/Source/Runtime/Engine/Public/EngineStats.h
--- a/Engine/Source/Runtime/Engine/Public/EngineStats.h
+++ b/Engine/Source/Runtime/Engine/Public/EngineStats.h
@@ -136,7 +136,7 @@ DECLARE_DWORD_COUNTER_STAT_EXTERN(TEXT("Num Destructible shapes exported"),STAT_
 DECLARE_DWORD_COUNTER_STAT_EXTERN(TEXT("Num update nav octree"),STAT_Navigation_UpdateNavOctree,STATGROUP_Navigation, );
 DECLARE_CYCLE_STAT_EXTERN(TEXT("Adding actors to navoctree"),STAT_Navigation_AddingActorsToNavOctree,STATGROUP_Navigation, );
 DECLARE_CYCLE_STAT_EXTERN(TEXT("    Adjusting nav links"),STAT_Navigation_AdjustingNavLinks,STATGROUP_Navigation, );
-DECLARE_CYCLE_STAT_EXTERN(TEXT("    Async AddTile"),STAT_Navigation_AddTile,STATGROUP_Navigation, );
+DECLARE_CYCLE_STAT_EXTERN(TEXT("    Async AddTile"),STAT_Navigation_AddGeneratedTiles,STATGROUP_Navigation, );
 DECLARE_CYCLE_STAT_EXTERN(TEXT("Recast tick"),STAT_Navigation_RecastTick,STATGROUP_Navigation, );
 DECLARE_CYCLE_STAT_EXTERN(TEXT("Recast: build compressed layers"),STAT_Navigation_RecastBuildCompressedLayers,STATGROUP_Navigation, );
 DECLARE_CYCLE_STAT_EXTERN(TEXT("Recast: build navmesh"),STAT_Navigation_RecastBuildNavigation,STATGROUP_Navigation, );
diff --git a/Engine/Source/Runtime/Navmesh/Private/Detour/DetourNavMesh.cpp b/Engine/Source/Runtime/Navmesh/Private/Detour/DetourNavMesh.cpp
--- a/Engine/Source/Runtime/Navmesh/Private/Detour/DetourNavMesh.cpp
+++ b/Engine/Source/Runtime/Navmesh/Private/Detour/DetourNavMesh.cpp
@@ -1965,9 +1965,12 @@ dtStatus dtNavMesh::removeTile(dtTileRef ref, unsigned char** data, int* dataSiz
 		for (int j = 0; j < nneis; ++j)
 			unconnectExtLinks(neis[j], tile);
 	}
-		
+
+	// Whether caller wants to own tile data
+	bool callerOwnsData = (data && dataSize);
+
 	// Reset tile.
-	if (tile->flags & DT_TILE_FREE_DATA)
+	if ((tile->flags & DT_TILE_FREE_DATA) && !callerOwnsData)
 	{
 		// Owns data
 		dtFree(tile->data);
