Update Reflections and Skylight captures when updating reflections.

diff --git a/Engine/Source/Runtime/Engine/Private/ReflectionCaptureComponent.cpp b/Engine/Source/Runtime/Engine/Private/ReflectionCaptureComponent.cpp
--- a/Engine/Source/Runtime/Engine/Private/ReflectionCaptureComponent.cpp
+++ b/Engine/Source/Runtime/Engine/Private/ReflectionCaptureComponent.cpp
@@ -33,9 +33,7 @@ void UWorld::UpdateAllReflectionCaptures()
 		UE_LOG(LogMaterial, Warning, TEXT("Update reflection captures only works with an active feature level of SM4 or greater."));
 		return;
 	}
-
-	TArray<UReflectionCaptureComponent*> UpdatedComponents;
-
+	
 	for (TObjectIterator<UReflectionCaptureComponent> It; It; ++It)
 	{
 		UReflectionCaptureComponent* CaptureComponent = *It;
@@ -43,12 +41,20 @@ void UWorld::UpdateAllReflectionCaptures()
 		if (ContainsActor(CaptureComponent->GetOwner()) && !CaptureComponent->IsPendingKill())
 		{
 			// Purge cached derived data and force an update
-			CaptureComponent->SetCaptureIsDirty();
-			UpdatedComponents.Add(CaptureComponent);
+			CaptureComponent->SetCaptureIsDirty();			
 		}
 	}
-
 	UReflectionCaptureComponent::UpdateReflectionCaptureContents(this);
+	
+	for (TObjectIterator<USkyLightComponent> It; It; ++It)
+	{
+		USkyLightComponent* SkylightComponent = *It;
+		if (ContainsActor(SkylightComponent->GetOwner()) && !SkylightComponent->IsPendingKill())
+		{			
+			SkylightComponent->SetCaptureIsDirty();			
+		}
+	}
+	USkyLightComponent::UpdateSkyCaptureContents(this);
 }
 
 AReflectionCapture::AReflectionCapture(const FObjectInitializer& ObjectInitializer)
