Increase WORLD_MAX by a factor of 4. Maximum world dimensions are approx 20km, in other words +-10km either side of the origin.

diff --git a/Engine/Source/Programs/UnrealLightmass/Private/Lighting/Lighting.h b/Engine/Source/Programs/UnrealLightmass/Private/Lighting/Lighting.h
--- a/Engine/Source/Programs/UnrealLightmass/Private/Lighting/Lighting.h
+++ b/Engine/Source/Programs/UnrealLightmass/Private/Lighting/Lighting.h
@@ -30,7 +30,7 @@ namespace Lightmass
 	extern double GStartupTime;
 }
 
-#define WORLD_MAX			524288.0	/* Maximum size of the world */
-#define HALF_WORLD_MAX		262144.0	/* Half the maximum size of the world */
+#define WORLD_MAX			2097152.0			/* Maximum size of the world */
+#define HALF_WORLD_MAX		( WORLD_MAX*0.5f )	/* Half the maximum size of the world */
 
 
diff --git a/Engine/Source/Runtime/Engine/Public/EngineDefines.h b/Engine/Source/Runtime/Engine/Public/EngineDefines.h
--- a/Engine/Source/Runtime/Engine/Public/EngineDefines.h
+++ b/Engine/Source/Runtime/Engine/Public/EngineDefines.h
@@ -72,7 +72,7 @@
 	Size of the world.
 -----------------------------------------------------------------------------*/
 
-#define WORLD_MAX					524288.0				/* Maximum size of the world */
+#define WORLD_MAX					2097152.0				/* Maximum size of the world */
 #define HALF_WORLD_MAX				(WORLD_MAX * 0.5)		/* Half the maximum size of the world */
 #define HALF_WORLD_MAX1				(HALF_WORLD_MAX - 1.0)	/* Half the maximum size of the world minus one */
 
