Fix landscape painting being interrupts by a slow task dialog showing up for <1 frame and breaking focus

diff --git a/Engine/Source/Runtime/Engine/Private/TextureDerivedData.cpp b/Engine/Source/Runtime/Engine/Private/TextureDerivedData.cpp
--- a/Engine/Source/Runtime/Engine/Private/TextureDerivedData.cpp
+++ b/Engine/Source/Runtime/Engine/Private/TextureDerivedData.cpp
@@ -471,7 +471,6 @@ public:
 		: FScopedSlowTask(0, InMessage, IsInGameThread())
 	{
 		UE_LOG(LogTexture,Display,TEXT("%s"),*InMessage.ToString());
-		MakeDialog();
 	}
 };
 
