Fix property flagging

diff --git a/Engine/Source/Runtime/Engine/Classes/GameFramework/GameMode.h b/Engine/Source/Runtime/Engine/Classes/GameFramework/GameMode.h
--- a/Engine/Source/Runtime/Engine/Classes/GameFramework/GameMode.h
+++ b/Engine/Source/Runtime/Engine/Classes/GameFramework/GameMode.h
@@ -170,11 +170,11 @@ public:
 	TSubclassOf<class AHUD>  HUDClass;    
 
 	/** Current number of spectators. */
-	UPROPERTY(BlueprintRead, Category=GameMode)
+	UPROPERTY(BlueprintReadOnly, Category=GameMode)
 	int32 NumSpectators;    
 
 	/** Current number of human players. */
-	UPROPERTY(BlueprintRead, Category=GameMode)
+	UPROPERTY(BlueprintReadOnly, Category=GameMode)
 	int32 NumPlayers;    
 
 	/** number of non-human players (AI controlled but participating as a player). */
@@ -190,7 +190,7 @@ public:
 	class AGameSession* GameSession;
 
 	/** Number of players that are still traveling from a previous map */
-	UPROPERTY(BlueprintRead, Category=GameMode)
+	UPROPERTY(BlueprintReadOnly, Category=GameMode)
 	int32 NumTravellingPlayers;
 
 	/** Used to assign unique PlayerIDs to each PlayerState. */
