Fix for Bug UE-2997, mesh particles continue to render in wireframe if switching to wireframe and back in the viewport.

diff --git a/Engine/Source/Runtime/Engine/Private/Particles/ParticleSystemRender.cpp b/Engine/Source/Runtime/Engine/Private/Particles/ParticleSystemRender.cpp
--- a/Engine/Source/Runtime/Engine/Private/Particles/ParticleSystemRender.cpp
+++ b/Engine/Source/Runtime/Engine/Private/Particles/ParticleSystemRender.cpp
@@ -2231,6 +2231,7 @@ void FDynamicMeshEmitterData::PreRenderView(FParticleSystemSceneProxy* Proxy, co
 					{
 						Mesh.Type = PT_TriangleList;
 						Mesh.MaterialRenderProxy = MeshMaterials[SectionIndex]->GetRenderProxy(bSelected);
+						Mesh.bWireframe = false;
 						BatchElement.IndexBuffer = &LODModel.IndexBuffer;
 						BatchElement.FirstIndex = Section.FirstIndex;
 						BatchElement.NumPrimitives = Section.NumTriangles;
