Marking dirty NavMesh tiles minor optimization #UE4

- also fixed a bug with uninitialized FCompositeNavModifier::bAdjustHeight that resulted in occasional undesired behavior in non-debug builds

diff --git a/Engine/Source/Runtime/Engine/Private/AI/Navigation/RecastNavMeshGenerator.cpp b/Engine/Source/Runtime/Engine/Private/AI/Navigation/RecastNavMeshGenerator.cpp
--- a/Engine/Source/Runtime/Engine/Private/AI/Navigation/RecastNavMeshGenerator.cpp
+++ b/Engine/Source/Runtime/Engine/Private/AI/Navigation/RecastNavMeshGenerator.cpp
@@ -3055,6 +3055,19 @@ void FRecastNavMeshGenerator::OnAreaAdded(const UClass* AreaClass, int32 AreaID)
 	AdditionalCachedData.OnAreaAdded(AreaClass, AreaID);
 }
 
+int32 FRecastNavMeshGenerator::FindInclusionBoundEncapsulatingBox(const FBox& Box) const
+{
+	for (int32 Index = 0; Index < InclusionBounds.Num(); ++Index)
+	{
+		if (DoesBoxContainBox(InclusionBounds[Index], Box))
+		{
+			return Index;
+		}
+	}
+
+	return INDEX_NONE;
+}
+
 TArray<uint32> FRecastNavMeshGenerator::RemoveTileLayers(const int32 TileX, const int32 TileY)
 {
 	TArray<uint32> ResultTileIndices;
@@ -3223,6 +3236,22 @@ static bool IntercestBounds(const FBox& TestBox, const TNavStatArray<FBox>& Boun
 	return false;
 }
 
+namespace 
+{
+	FBox CalculateBoxIntercetion(const FBox& BoxA, const FBox& BoxB)
+	{
+		// assumes boxes overlap
+		ensure(BoxA.Intersect(BoxB));
+		return FBox(FVector(FMath::Max(BoxA.Min.X, BoxB.Min.X)
+							, FMath::Max(BoxA.Min.Y, BoxB.Min.Y)
+							, FMath::Max(BoxA.Min.Z, BoxB.Min.Z))
+					, FVector(FMath::Min(BoxA.Max.X, BoxB.Max.X)
+							, FMath::Min(BoxA.Max.Y, BoxB.Max.Y)
+							, FMath::Min(BoxA.Max.Z, BoxB.Max.Z))
+					);
+	}
+}
+
 void FRecastNavMeshGenerator::MarkDirtyTiles(const TArray<FNavigationDirtyArea>& DirtyAreas)
 {
 	QUICK_SCOPE_CYCLE_COUNTER(STAT_RecastNavMeshGenerator_MarkDirtyTiles);
@@ -3237,14 +3266,34 @@ void FRecastNavMeshGenerator::MarkDirtyTiles(const TArray<FNavigationDirtyArea>&
 	TSet<FPendingTileElement> DirtyTiles;
 	for (const FNavigationDirtyArea& DirtyArea : DirtyAreas)
 	{
-		const FBox AdjustedAreaBounds = GrowBoundingBox(DirtyArea.Bounds, DirtyArea.HasFlag(ENavigationDirtyFlag::UseAgentHeight));
-		if (!DirtyArea.HasFlag(ENavigationDirtyFlag::NavigationBounds) && 
-			!IntercestBounds(AdjustedAreaBounds, InclusionBounds))
+		bool bDoTileInclusionTest = false;
+		FBox AdjustedAreaBounds = DirtyArea.Bounds;
+		
+		// if it's not expanding the navigatble area
+		if (DirtyArea.HasFlag(ENavigationDirtyFlag::NavigationBounds) == false)
 		{
-			// Skip whole area
-			continue;
+			// and is outside of current bounds
+			if (GetTotalBounds().Intersect(DirtyArea.Bounds) == false)
+			{
+				// skip it
+				continue;
+			}
+
+			const FBox CutDownArea = CalculateBoxIntercetion(GetTotalBounds(), DirtyArea.Bounds);
+			AdjustedAreaBounds = GrowBoundingBox(CutDownArea, DirtyArea.HasFlag(ENavigationDirtyFlag::UseAgentHeight));
+
+			// @todo this and the following test share some work in common
+			if (IntercestBounds(AdjustedAreaBounds, InclusionBounds) == false)
+			{
+				continue;
+			}
+
+			// check if any of inclusion volumes encapsulates this box
+			// using CutDownArea not AdjustedAreaBounds since if the area is on the border of navigable space
+			// then FindInclusionBoundEncapsulatingBox can produce false negative
+			bDoTileInclusionTest = FindInclusionBoundEncapsulatingBox(CutDownArea) == INDEX_NONE;
 		}
-		
+				
 		const FBox RcAreaBounds = Unreal2RecastBox(AdjustedAreaBounds);
 		const int32 XMin = FMath::FloorToInt((RcAreaBounds.Min.X - NavMeshOrigin.X) / TileSizeInWorldUnits);
 		const int32 XMax = FMath::FloorToInt((RcAreaBounds.Max.X - NavMeshOrigin.X) / TileSizeInWorldUnits);
@@ -3255,12 +3304,16 @@ void FRecastNavMeshGenerator::MarkDirtyTiles(const TArray<FNavigationDirtyArea>&
 		{
 			for (int32 x = XMin; x <= XMax; ++x)
 			{
-				const FBox TileBox = CalculateTileBounds(x, y, NavMeshOrigin, TotalNavBounds, TileSizeInWorldUnits);
-				if (!DirtyArea.HasFlag(ENavigationDirtyFlag::NavigationBounds) && 
-					!IntercestBounds(TileBox, InclusionBounds))
+				if (DirtyArea.HasFlag(ENavigationDirtyFlag::NavigationBounds) == false && bDoTileInclusionTest == true)
 				{
-					// Skip this tile
-					continue;
+					const FBox TileBox = CalculateTileBounds(x, y, NavMeshOrigin, TotalNavBounds, TileSizeInWorldUnits);
+
+					// do per tile check since we can have lots of tiles inbetween navigable bounds volumes
+					if (IntercestBounds(TileBox, InclusionBounds) == false)
+					{
+						// Skip this tile
+						continue;
+					}
 				}
 												
 				FPendingTileElement Element;
diff --git a/Engine/Source/Runtime/Engine/Public/AI/Navigation/RecastNavMeshGenerator.h b/Engine/Source/Runtime/Engine/Public/AI/Navigation/RecastNavMeshGenerator.h
--- a/Engine/Source/Runtime/Engine/Public/AI/Navigation/RecastNavMeshGenerator.h
+++ b/Engine/Source/Runtime/Engine/Public/AI/Navigation/RecastNavMeshGenerator.h
@@ -333,6 +333,10 @@ public:
 
 	const TNavStatArray<FBox>& GetInclusionBounds() const { return InclusionBounds; }
 	
+	/** checks if any on InclusionBounds encapsulates given box.
+	 *	@return index to first item in InclusionBounds that meets expectations */
+	int32 FindInclusionBoundEncapsulatingBox(const FBox& Box) const;
+
 	/** Total navigable area box, sum of all navigation volumes bounding boxes */
 	FBox GetTotalBounds() const { return TotalNavBounds; }
 
diff --git a/Engine/Source/Runtime/Engine/Public/AI/NavigationModifier.h b/Engine/Source/Runtime/Engine/Public/AI/NavigationModifier.h
--- a/Engine/Source/Runtime/Engine/Public/AI/NavigationModifier.h
+++ b/Engine/Source/Runtime/Engine/Public/AI/NavigationModifier.h
@@ -194,7 +194,7 @@ protected:
 
 struct ENGINE_API FCompositeNavModifier : public FNavigationModifier
 {
-	FCompositeNavModifier() : bHasPotentialLinks(false) {}
+	FCompositeNavModifier() : bHasPotentialLinks(false), bAdjustHeight(false) {}
 
 	void Shrink();
 	void Reset();
