File missing in previous commit #UE4

diff --git a/Engine/Source/Runtime/Engine/Private/AI/Navigation/NavigationModifier.cpp b/Engine/Source/Runtime/Engine/Private/AI/Navigation/NavigationModifier.cpp
--- a/Engine/Source/Runtime/Engine/Private/AI/Navigation/NavigationModifier.cpp
+++ b/Engine/Source/Runtime/Engine/Private/AI/Navigation/NavigationModifier.cpp
@@ -454,6 +454,8 @@ void FCompositeNavModifier::Reset()
 	Areas.Reset();
 	SimpleLinks.Reset();
 	CustomLinks.Reset();
+	bHasPotentialLinks = false;
+	bAdjustHeight = false;
 }
 
 void FCompositeNavModifier::Empty()
@@ -462,6 +464,7 @@ void FCompositeNavModifier::Empty()
 	SimpleLinks.Empty();
 	CustomLinks.Empty();
 	bHasPotentialLinks = false;
+	bAdjustHeight = false;
 }
 
 FCompositeNavModifier FCompositeNavModifier::GetInstantiatedMetaModifier(const FNavAgentProperties* NavAgent, TWeakObjectPtr<UObject> WeakOwnerPtr) const
