Added header comment to UModelComponent.

diff --git a/Engine/Source/Runtime/Engine/Classes/Components/ModelComponent.h b/Engine/Source/Runtime/Engine/Classes/Components/ModelComponent.h
--- a/Engine/Source/Runtime/Engine/Classes/Components/ModelComponent.h
+++ b/Engine/Source/Runtime/Engine/Classes/Components/ModelComponent.h
@@ -9,6 +9,12 @@
 //
 class FModelElement;
 
+/**
+ * ModelComponents are PrimitiveComponents that represent elements of BSP geometry in a ULevel object.
+ * They are used exclusively by ULevel and are not intended as general-purpose components.
+ *
+ * @see ULevel
+ */
 UCLASS(MinimalAPI)
 class UModelComponent : public UPrimitiveComponent, public IInterface_CollisionDataProvider
 {
