Fixed generated filename when capturing high res screenshots in-game.  Added error message when HiResShot used with incorrect parameters.

diff --git a/Engine/Source/Runtime/Engine/Private/GameViewportClient.cpp b/Engine/Source/Runtime/Engine/Private/GameViewportClient.cpp
--- a/Engine/Source/Runtime/Engine/Private/GameViewportClient.cpp
+++ b/Engine/Source/Runtime/Engine/Private/GameViewportClient.cpp
@@ -1185,6 +1185,12 @@ void UGameViewportClient::ProcessScreenShots(FViewport* InViewport)
 						bWriteAlpha = GetHighResScreenshotConfig().MergeMaskIntoAlpha(Bitmap);
 						SourceRect = GetHighResScreenshotConfig().CaptureRegion;
 					}
+
+					if (!FPaths::GetExtension(ScreenShotName).IsEmpty())
+					{
+						ScreenShotName = FPaths::GetBaseFilename(ScreenShotName, false);
+						ScreenShotName += TEXT(".bmp");
+					}
 					FFileHelper::CreateBitmap(*ScreenShotName, InViewport->GetSizeXY().X, InViewport->GetSizeXY().Y, Bitmap.GetData(), &SourceRect, &IFileManager::Get(), NULL, bWriteAlpha);
 				}
 			}
diff --git a/Engine/Source/Runtime/Engine/Private/UnrealClient.cpp b/Engine/Source/Runtime/Engine/Private/UnrealClient.cpp
--- a/Engine/Source/Runtime/Engine/Private/UnrealClient.cpp
+++ b/Engine/Source/Runtime/Engine/Private/UnrealClient.cpp
@@ -1600,6 +1600,10 @@ ENGINE_API bool GetHighResScreenShotInput(const TCHAR* Cmd, FOutputDevice& Ar, u
 
 		return true;
 	}
+	else
+	{
+		Ar.Logf(TEXT("Error: Bad input. Input should be in either the form \"HighResShot 1920x1080\" or \"HighResShot 2\""));
+	}
 
 	return false;
 }
