#vr Cleaning up / renaming / commenting on SWindow and viewport independent sizing functions

diff --git a/Engine/Plugins/Runtime/OculusRift/Source/OculusRift/Private/OculusRiftHMD.cpp b/Engine/Plugins/Runtime/OculusRift/Source/OculusRift/Private/OculusRiftHMD.cpp
--- a/Engine/Plugins/Runtime/OculusRift/Source/OculusRift/Private/OculusRiftHMD.cpp
+++ b/Engine/Plugins/Runtime/OculusRift/Source/OculusRift/Private/OculusRiftHMD.cpp
@@ -1211,7 +1211,7 @@ bool FOculusRiftHMD::EnableStereo(bool stereo)
 			{
 				FVector2D size = Window->GetSizeInScreen();
 				GameEngine->SceneViewport->SetViewportSize(size.X, size.Y);
-				Window->SetViewportAutoResize(true);
+				Window->SetViewportSizeDrivenByWindow(true);
 			}
 		}
 	}
diff --git a/Engine/Plugins/Runtime/OculusRift/Source/OculusRift/Private/OculusRiftRender.cpp b/Engine/Plugins/Runtime/OculusRift/Source/OculusRift/Private/OculusRiftRender.cpp
--- a/Engine/Plugins/Runtime/OculusRift/Source/OculusRift/Private/OculusRiftRender.cpp
+++ b/Engine/Plugins/Runtime/OculusRift/Source/OculusRift/Private/OculusRiftRender.cpp
@@ -579,7 +579,7 @@ void FOculusRiftHMD::UpdateViewport(bool bUseSeparateRenderTarget, const FViewpo
 				TSharedPtr<SWindow> Window = FSlateApplication::Get().FindWidgetWindow(Widget, WidgetPath);
 				if (Window.IsValid())
 				{
-					Window->SetViewportAutoResize(true);
+					Window->SetViewportSizeDrivenByWindow(true);
 				}
 			}
 		}
diff --git a/Engine/Source/Runtime/Engine/Private/Slate/SceneViewport.cpp b/Engine/Source/Runtime/Engine/Private/Slate/SceneViewport.cpp
--- a/Engine/Source/Runtime/Engine/Private/Slate/SceneViewport.cpp
+++ b/Engine/Source/Runtime/Engine/Private/Slate/SceneViewport.cpp
@@ -277,7 +277,7 @@ void FSceneViewport::OnDrawViewport( const FGeometry& AllottedGeometry, const FS
 		TSharedPtr<SWindow> Window = FSlateApplication::Get().FindWidgetWindow( ViewportWidget.Pin().ToSharedRef() );
 
 		check(Window.IsValid());
-		if (Window->IsViewportAutoResizeable())
+		if (Window->IsViewportSizeDrivenByWindow())
 		{
 			ResizeViewport(FMath::Max(0, DrawSize.X), FMath::Max(0, DrawSize.Y), Window->GetWindowMode(), 0, 0);
 		}
@@ -856,7 +856,7 @@ void FSceneViewport::ResizeFrame(uint32 NewSizeX, uint32 NewSizeY, EWindowMode::
 
 			if (GEngine->HMDDevice.IsValid() && GEngine->HMDDevice->IsHMDConnected())
 			{
-				WindowToResize->SetViewportAutoResize(true);
+				WindowToResize->SetViewportSizeDrivenByWindow(true);
 				// Resize & move only if moving to a fullscreen mode
 				if (NewWindowMode != EWindowMode::Windowed)
 				{
@@ -878,7 +878,7 @@ void FSceneViewport::ResizeFrame(uint32 NewSizeX, uint32 NewSizeY, EWindowMode::
 								NewSizeY = MonitorInfo.WindowSizeY;
 							}
 							NewWindowMode = DesiredWindowMode = EWindowMode::WindowedMirror;
-							WindowToResize->SetViewportSize(FVector2D(ViewportSizeX, ViewportSizeY));
+							WindowToResize->SetIndependentViewportSize(FVector2D(ViewportSizeX, ViewportSizeY));
 						}
 					}
 				}
@@ -926,7 +926,7 @@ void FSceneViewport::ResizeFrame(uint32 NewSizeX, uint32 NewSizeY, EWindowMode::
 						}
 						ViewportSizeX = NewSizeX;
 						ViewportSizeY = NewSizeY;
-						WindowToResize->SetViewportAutoResize(true);
+						WindowToResize->SetViewportSizeDrivenByWindow(true);
 					}
 
 					if (NewWindowMode != WindowMode)
@@ -952,7 +952,7 @@ void FSceneViewport::SetViewportSize(uint32 NewViewportSizeX, uint32 NewViewport
 	TSharedPtr<SWindow> Window = FSlateApplication::Get().FindWidgetWindow(ViewportWidget.Pin().ToSharedRef());
 	if (Window.IsValid())
 	{
-		Window->SetViewportSize(FVector2D(NewViewportSizeX, NewViewportSizeY));
+		Window->SetIndependentViewportSize(FVector2D(NewViewportSizeX, NewViewportSizeY));
 		const FVector2D vp = Window->GetViewportSize();
 		FSlateApplicationBase::Get().GetRenderer()->UpdateFullscreenState(Window.ToSharedRef(), vp.X, vp.Y);
 		ResizeViewport(NewViewportSizeX, NewViewportSizeY, Window->GetWindowMode(), 0, 0);
diff --git a/Engine/Source/Runtime/SlateCore/Public/Widgets/SWindow.h b/Engine/Source/Runtime/SlateCore/Public/Widgets/SWindow.h
--- a/Engine/Source/Runtime/SlateCore/Public/Widgets/SWindow.h
+++ b/Engine/Source/Runtime/SlateCore/Public/Widgets/SWindow.h
@@ -629,9 +629,12 @@ public:
 	/** Are any of our child windows active? */
 	bool HasActiveChildren() const;
 
-	inline void SetViewportAutoResize(bool InAutoResize)
+	/**
+	 * Sets whether or not the viewport size should be driven by the window's size.  If true, the two will be the same.  If false, an independent viewport size can be specified with SetIndependentViewportSize
+	 */
+	inline void SetViewportSizeDrivenByWindow(bool bDrivenByWindow)
 	{
-		if (InAutoResize)
+		if (bDrivenByWindow)
 		{
 			ViewportSize = FVector2D::ZeroVector;
 		}
@@ -640,12 +643,27 @@ public:
 			ViewportSize = Size;
 		}
 	}
-	inline bool IsViewportAutoResizeable() const
+	
+	/**
+	 * Returns whether or not the viewport and window size should be linked together.  If false, the two can be independent in cases where it is needed (e.g. mirror mode window drawing)
+	 */
+	inline bool IsViewportSizeDrivenByWindow() const
 	{
 		return (ViewportSize.X == 0);
 	}
-	inline FVector2D GetViewportSize() const { return (ViewportSize.X) ? ViewportSize : Size; }
-	inline void SetViewportSize(const FVector2D& VP) 
+
+	/**
+	 * Returns the viewport size, taking into consideration if the window size should drive the viewport size
+	 */
+	inline FVector2D GetViewportSize() const
+	{
+		return (ViewportSize.X) ? ViewportSize : Size;
+	}
+	
+	/**
+	 * Sets the viewport size independently of the window size, if non-zero.
+	 */
+	inline void SetIndependentViewportSize(const FVector2D& VP) 
 	{
 		ViewportSize = VP;
 	}
diff --git a/Engine/Source/Runtime/SlateRHIRenderer/Private/SlateRHIRenderer.cpp b/Engine/Source/Runtime/SlateRHIRenderer/Private/SlateRHIRenderer.cpp
--- a/Engine/Source/Runtime/SlateRHIRenderer/Private/SlateRHIRenderer.cpp
+++ b/Engine/Source/Runtime/SlateRHIRenderer/Private/SlateRHIRenderer.cpp
@@ -598,7 +598,7 @@ void FSlateRHIRenderer::DrawWindows_Private( FSlateDrawBuffer& WindowDrawBuffer
 				// The viewport had better exist at this point  
 				FViewportInfo* ViewInfo = WindowToViewportInfo.FindChecked( Window.Get() );
 
-				if (Window->IsViewportAutoResizeable())
+				if (Window->IsViewportSizeDrivenByWindow())
 				{
 					// Resize the viewport if needed
 					ConditionalResizeViewport(ViewInfo, ViewInfo->DesiredWidth, ViewInfo->DesiredHeight, IsViewportFullscreen(*Window));
