Re-added the code to flush stale UFont data from the font cache

Now with an added check that the Slate application is valid to fix the previous crash when cooking.

diff --git a/Engine/Source/Runtime/Engine/Private/Font.cpp b/Engine/Source/Runtime/Engine/Private/Font.cpp
--- a/Engine/Source/Runtime/Engine/Private/Font.cpp
+++ b/Engine/Source/Runtime/Engine/Private/Font.cpp
@@ -23,9 +23,7 @@ UFont::UFont(const FObjectInitializer& ObjectInitializer)
 
 UFont::~UFont()
 {
-	// todo: jdale - this is causing a crash when cooking due to the Slate application being null
-	/*
-	if(FontCacheType == EFontCacheType::Runtime)
+	if(FontCacheType == EFontCacheType::Runtime && FSlateApplication::IsInitialized())
 	{
 		TSharedPtr<FSlateRenderer> SlateRenderer = FSlateApplication::Get().GetRenderer();
 		if(SlateRenderer.IsValid())
@@ -34,7 +32,6 @@ UFont::~UFont()
 			FontCache->FlushObject(this);
 		}
 	}
-	*/
 }
 
 void UFont::Serialize( FArchive& Ar )
