#ue4 - ProjectileMovementComponent: friction when bouncing has a new optional behavior. Fixed sliding/rolling code and exposed more related variables.

- New projectile friction mode when bouncing modified friction based on the normal force from impact (higher when perpendicular, lower when glancing). Default behavior remains the same as before.
- Added bIsSliding to indicate if the projectile is sliding/rolling, and fixed behavior with multiple impacts between frames.

diff --git a/Engine/Source/Runtime/Core/Public/UObject/ObjectVersion.h b/Engine/Source/Runtime/Core/Public/UObject/ObjectVersion.h
--- a/Engine/Source/Runtime/Core/Public/UObject/ObjectVersion.h
+++ b/Engine/Source/Runtime/Core/Public/UObject/ObjectVersion.h
@@ -636,6 +636,8 @@ enum EUnrealEngineObjectUE4Version
 	VER_UE4_REMOVE_SKELETALMESH_COMPONENT_BODYSETUP_SERIALIZATION,
 	// Made Slate font data use bulk data to store the embedded font data
 	VER_UE4_SLATE_BULK_FONT_DATA,
+	// Add new friction behavior in ProjectileMovementComponent.
+	VER_UE4_ADD_PROJECTILE_FRICTION_BEHAVIOR,
 
 	// -----<new versions can be added before this line>-------------------------------------------------
 	// - this needs to be the last line (see note below)
diff --git a/Engine/Source/Runtime/Engine/Classes/GameFramework/ProjectileMovementComponent.h b/Engine/Source/Runtime/Engine/Classes/GameFramework/ProjectileMovementComponent.h
--- a/Engine/Source/Runtime/Engine/Classes/GameFramework/ProjectileMovementComponent.h
+++ b/Engine/Source/Runtime/Engine/Classes/GameFramework/ProjectileMovementComponent.h
@@ -62,6 +62,28 @@ class ENGINE_API UProjectileMovementComponent : public UMovementComponent
 	UPROPERTY(EditAnywhere, BlueprintReadWrite, Category=Homing)
 	uint32 bIsHomingProjectile:1;
 
+	/**
+	 * Controls the effects of friction on velocity parallel to the impact surface when bouncing.
+	 * If true, friction will be modified based on the angle of impact, making friction higher for perpendicular impacts and lower for glancing impacts.
+	 * If false, a bounce will retain a proportion of tangential velocity equal to (1.0 - Friction), acting as a "horizontal restitution".
+	 */
+	UPROPERTY(EditAnywhere, BlueprintReadWrite, Category=ProjectileBounces)
+	uint32 bBounceAngleAffectsFriction:1;
+
+	/**
+	 * If true, projectile is sliding / rolling along a surface.
+	 */
+	UPROPERTY(VisibleInstanceOnly, BlueprintReadOnly, Category=ProjectileBounces)
+	uint32 bIsSliding:1;
+
+	/** Saved HitResult Time (0 to 1) from previous simulation step. Equal to 1.0 when there was no impact. */
+	UPROPERTY(VisibleInstanceOnly, BlueprintReadOnly, Category=ProjectileBounces)
+	float PreviousHitTime;
+
+	/** Saved HitResult Normal from previous simulation step that resulted in an impact. If PreviousHitTime is 1.0, then the hit was not in the last step. */
+	UPROPERTY(VisibleInstanceOnly, BlueprintReadOnly, Category=ProjectileBounces)
+	FVector PreviousHitNormal;
+
 	/** Custom gravity scale for this projectile. Set to 0 for no gravity. */
 	UPROPERTY(EditAnywhere, BlueprintReadWrite, Category=Projectile)
 	float ProjectileGravityScale;
@@ -82,6 +104,7 @@ class ENGINE_API UProjectileMovementComponent : public UMovementComponent
 	 * Normal range is [0,1] : 0.0 = no friction, 1.0+ = very high friction.
 	 * Also affects the percentage of velocity maintained after the bounce in the direction tangent to the normal of impact.
 	 * Ignored if bShouldBounce is false.
+	 * @see bBounceAngleAffectsFriction
 	 */
 	UPROPERTY(EditAnywhere, BlueprintReadWrite, Category=ProjectileBounces, meta=(ClampMin="0", UIMin="0"))
 	float Friction;
@@ -118,7 +141,7 @@ class ENGINE_API UProjectileMovementComponent : public UMovementComponent
 
 	//Begin UActorComponent Interface
 	virtual void TickComponent(float DeltaTime, enum ELevelTick TickType, FActorComponentTickFunction *ThisTickFunction) override;
-	virtual void Serialize( FArchive& Ar) override;
+	virtual void PostLoad() override;
 	//End UActorComponent Interface
 
 	//Begin UMovementComponent Interface
@@ -199,7 +222,7 @@ class ENGINE_API UProjectileMovementComponent : public UMovementComponent
 
 protected:
 
-	/** @return true if the wall hit has been fully processed, and no further movement is needed */
+	/** @return true if the simulation should stop. */
 	bool HandleHitWall(const FHitResult& Hit, float TimeTick, const FVector& MoveDelta);
 
 	/** Applies bounce logic (if enabled) to affect velocity upon impact, or stops the projectile if bounces are not enabled (or velocity is below threshold). Fires applicable events. */
diff --git a/Engine/Source/Runtime/Engine/Private/ProjectileMovementComponent.cpp b/Engine/Source/Runtime/Engine/Private/ProjectileMovementComponent.cpp
--- a/Engine/Source/Runtime/Engine/Private/ProjectileMovementComponent.cpp
+++ b/Engine/Source/Runtime/Engine/Private/ProjectileMovementComponent.cpp
@@ -28,13 +28,20 @@ UProjectileMovementComponent::UProjectileMovementComponent(const FObjectInitiali
 
 	MaxSimulationTimeStep = 0.05f;
 	MaxSimulationIterations = 8;
+
+	bBounceAngleAffectsFriction = false;
+	bIsSliding = false;
+	PreviousHitTime = 1.f;
+	PreviousHitNormal = FVector::UpVector;
 }
 
-void UProjectileMovementComponent::Serialize( FArchive& Ar)
+void UProjectileMovementComponent::PostLoad()
 {
-	Super::Serialize(Ar);
+	Super::PostLoad();
+
+	const int32 LinkerUE4Ver = GetLinkerUE4Version();
 
-	if (Ar.IsLoading() && Ar.UE4Ver() < VER_UE4_REFACTOR_PROJECTILE_MOVEMENT)
+	if (LinkerUE4Ver < VER_UE4_REFACTOR_PROJECTILE_MOVEMENT)
 	{
 		// Old code used to treat Bounciness as Friction as well.
 		Friction = FMath::Clamp(1.f - Bounciness, 0.f, 1.f);
@@ -101,11 +108,9 @@ void UProjectileMovementComponent::TickComponent(float DeltaTime, enum ELevelTic
 	}
 
 	float RemainingTime	= DeltaTime;
-	int32	NumBounces = 0;
+	int32 NumBounces = 0;
 	int32 Iterations = 0;
 	FHitResult Hit(1.f);
-	FVector OldHitNormal(0.f,0.f,1.f);
-	float PreviousHitTime = 1.f;
 
 	while( RemainingTime >= MIN_TICK_TIME && (Iterations < MaxSimulationIterations) && !ActorOwner->IsPendingKill() && UpdatedComponent )
 	{
@@ -144,6 +149,7 @@ void UProjectileMovementComponent::TickComponent(float DeltaTime, enum ELevelTic
 		if( !Hit.bBlockingHit )
 		{
 			PreviousHitTime = 1.f;
+			bIsSliding = false;
 
 			// Only calculate new velocity if events didn't change it during the movement update.
 			if (Velocity == OldVelocity)
@@ -169,19 +175,19 @@ void UProjectileMovementComponent::TickComponent(float DeltaTime, enum ELevelTic
 			float SubTickTimeRemaining = TimeTick * (1.f - Hit.Time);
 
 			// Multiple hits within very short time period?
-			const bool bMultiHit = (PreviousHitTime <= KINDA_SMALL_NUMBER && Hit.Time <= KINDA_SMALL_NUMBER);
-			
+			const bool bMultiHit = (PreviousHitTime < 1.f && Hit.Time <= KINDA_SMALL_NUMBER);
+
 			// if velocity still into wall (after HandleHitWall() had a chance to adjust), slide along wall
 			const float DotTolerance = 0.01f;
-			const bool bSliding = (bMultiHit && FVector::Coincident(OldHitNormal, Hit.Normal)) ||
-								  ((Velocity.SafeNormal() | Hit.Normal) <= DotTolerance);
+			bIsSliding = (bMultiHit && FVector::Coincident(PreviousHitNormal, Hit.Normal)) ||
+						 ((Velocity.SafeNormal() | Hit.Normal) <= DotTolerance);
 			
-			if (bSliding)
+			if (bIsSliding)
 			{
-				if (bMultiHit && (OldHitNormal | Hit.Normal) <= 0.f)
+				if (bMultiHit && (PreviousHitNormal | Hit.Normal) <= 0.f)
 				{
 					//90 degree or less corner, so use cross product for direction
-					FVector NewDir = (Hit.Normal ^ OldHitNormal);
+					FVector NewDir = (Hit.Normal ^ PreviousHitNormal);
 					NewDir = NewDir.SafeNormal();
 					Velocity = Velocity.ProjectOnToNormal(NewDir);
 					if ((OldVelocity | Velocity) < 0.f)
@@ -213,7 +219,7 @@ void UProjectileMovementComponent::TickComponent(float DeltaTime, enum ELevelTic
 			}
 
 			PreviousHitTime = Hit.Time;
-			OldHitNormal = Hit.Normal;
+			PreviousHitNormal = Hit.Normal;
 			
 			// A few initial bounces should add more time and iterations to complete most of the simulation.
 			if (NumBounces <= 2 && SubTickTimeRemaining >= MIN_TICK_TIME)
@@ -398,7 +404,7 @@ FVector UProjectileMovementComponent::ComputeBounceResult(const FHitResult& Hit,
 		TempVelocity += ProjectedNormal;
 
 		// Only tangential velocity should be affected by friction.
-		const float ScaledFriction = FMath::Clamp(-VDotNormal / TempVelocity.Size(), 0.f, 1.f) * Friction;
+		const float ScaledFriction = (bBounceAngleAffectsFriction || bIsSliding) ? FMath::Clamp(-VDotNormal / TempVelocity.Size(), 0.f, 1.f) * Friction : Friction;
 		TempVelocity *= FMath::Clamp(1.f - ScaledFriction, 0.f, 1.f);
 
 		// Coefficient of restitution only applies perpendicular to impact.
