[UE-3597] Documenting level blueprint source files.

diff --git a/Engine/Source/Runtime/Engine/Classes/Engine/LevelScriptActor.h b/Engine/Source/Runtime/Engine/Classes/Engine/LevelScriptActor.h
--- a/Engine/Source/Runtime/Engine/Classes/Engine/LevelScriptActor.h
+++ b/Engine/Source/Runtime/Engine/Classes/Engine/LevelScriptActor.h
@@ -4,6 +4,20 @@
 
 #include "LevelScriptActor.generated.h"
 
+/**
+ * ALevelScriptActor is the base class for classes generated by 
+ * ULevelScriptBlueprints. ALevelScriptActor instances are hidden actors that 
+ * exist within a level, and can execute level-wide logic (operating on specific
+ * actor instances within the level). The level-script's functionality is defined
+ * inside the ULevelScriptBlueprint itself (using the blueprint's node-based 
+ * interface).
+ *
+ * @see AActor
+ * @see https://docs.unrealengine.com/latest/INT/Engine/Blueprints/UserGuide/Types/LevelBlueprint/index.html
+ * @see ULevelScriptBlueprint
+ * @see https://docs.unrealengine.com/latest/INT/Engine/Blueprints/index.html
+ * @see UBlueprint
+ */
 UCLASS(notplaceable, meta=(KismetHideOverrides = "ReceiveAnyDamage,ReceivePointDamage,ReceiveRadialDamage,ReceiveActorBeginOverlap,ReceiveActorEndOverlap,ReceiveHit,ReceiveDestroyed,ReceiveActorBeginCursorOver,ReceiveActorEndCursorOver,ReceiveActorOnClicked,ReceiveActorOnReleased,ReceiveActorOnInputTouchBegin,ReceiveActorOnInputTouchEnd,ReceiveActorOnInputTouchEnter,ReceiveActorOnInputTouchLeave"), HideCategories=(Collision,Rendering,"Utilities|Transformation"))
 class ENGINE_API ALevelScriptActor : public AActor
 {
diff --git a/Engine/Source/Runtime/Engine/Classes/Engine/LevelScriptBlueprint.h b/Engine/Source/Runtime/Engine/Classes/Engine/LevelScriptBlueprint.h
--- a/Engine/Source/Runtime/Engine/Classes/Engine/LevelScriptBlueprint.h
+++ b/Engine/Source/Runtime/Engine/Classes/Engine/LevelScriptBlueprint.h
@@ -6,6 +6,17 @@
 
 #include "LevelScriptBlueprint.generated.h"
 
+/**
+ * A level blueprint is a specialized type of blueprint. It is used to house
+ * global, level-wide logic. In a level blueprint, you can operate on specific 
+ * level-actor instances through blueprint's node-based interface. UE3 users 
+ * should be familiar with this concept, as it is very similar to Kismet.
+ *
+ * @see https://docs.unrealengine.com/latest/INT/Engine/Blueprints/UserGuide/Types/LevelBlueprint/index.html
+ * @see https://docs.unrealengine.com/latest/INT/Engine/Blueprints/index.html
+ * @see UBlueprint
+ * @see ALevelScriptActor
+ */
 UCLASS(MinimalAPI, NotBlueprintType)
 class ULevelScriptBlueprint : public UBlueprint
 {
