#ue4 - Only copy the "advanced view" state of pins when reconstructing nodes if users can edit the advanced state. Otherwise changes to metadata advanced pin state don't get propagated.

diff --git a/Engine/Source/Editor/BlueprintGraph/Classes/K2Node_SwitchEnum.h b/Engine/Source/Editor/BlueprintGraph/Classes/K2Node_SwitchEnum.h
--- a/Engine/Source/Editor/BlueprintGraph/Classes/K2Node_SwitchEnum.h
+++ b/Engine/Source/Editor/BlueprintGraph/Classes/K2Node_SwitchEnum.h
@@ -38,6 +38,7 @@ class UK2Node_SwitchEnum : public UK2Node_Switch, public INodeDependingOnEnumInt
 	virtual bool IsConnectionDisallowed(const UEdGraphPin* MyPin, const UEdGraphPin* OtherPin, FString& OutReason) const override;
 	virtual void GetMenuActions(FBlueprintActionDatabaseRegistrar& ActionRegistrar) const override;
 	virtual bool CanEverRemoveExecutionPin() const override { return false; }
+	virtual bool CanUserEditPinAdvancedViewFlag() const override { return true; }
 	// End of UK2Node interface
 
 	// UK2Node_Switch Interface
diff --git a/Engine/Source/Runtime/Engine/Classes/EdGraph/EdGraphNode.h b/Engine/Source/Runtime/Engine/Classes/EdGraph/EdGraphNode.h
--- a/Engine/Source/Runtime/Engine/Classes/EdGraph/EdGraphNode.h
+++ b/Engine/Source/Runtime/Engine/Classes/EdGraph/EdGraphNode.h
@@ -353,6 +353,9 @@ public:
 	/** Whether or not this node can be deleted by user action */
 	virtual bool CanUserDeleteNode() const;
 
+	/** Whether or not this node allows users to edit the advanced view flag of pins (actually edit the property, not the same as show/hide advanced pins). */
+	virtual bool CanUserEditPinAdvancedViewFlag() const { return false; }
+
 	/** Tries to come up with a descriptive name for the compiled output */
 	virtual FString GetDescriptiveCompiledName() const;
 
diff --git a/Engine/Source/Runtime/Engine/Private/EdGraph/EdGraphPin.cpp b/Engine/Source/Runtime/Engine/Private/EdGraph/EdGraphPin.cpp
--- a/Engine/Source/Runtime/Engine/Private/EdGraph/EdGraphPin.cpp
+++ b/Engine/Source/Runtime/Engine/Private/EdGraph/EdGraphPin.cpp
@@ -159,7 +159,8 @@ void UEdGraphPin::CopyPersistentDataFromOldPin(const UEdGraphPin& SourcePin)
 
 	// In K2 schemas the wildcard pins need to have their type copied before we get to pin splitting
 	// TODO: Better less hacky way of this?
-	if (PinType.PinCategory == TEXT("wildcard"))
+	static const FString WildCardText(TEXT("wildcard"));
+	if (PinType.PinCategory == WildCardText)
 	{
 		PinType = SourcePin.PinType;
 	}
@@ -197,8 +198,13 @@ void UEdGraphPin::CopyPersistentDataFromOldPin(const UEdGraphPin& SourcePin)
 	}
 
 #if WITH_EDITORONLY_DATA
-	// Copy advanced visibility property, since it can be changed by user
-	bAdvancedView = SourcePin.bAdvancedView;
+	// Copy advanced visibility property, if it can be changed by user.
+	// Otherwise we don't want to copy this, or we'd be ignoring new metadata that tries to hide old pins.
+	UEdGraphNode* OuterNode = Cast<UEdGraphNode>(GetOuter());
+	if (OuterNode != nullptr && OuterNode->CanUserEditPinAdvancedViewFlag())
+	{
+		bAdvancedView = SourcePin.bAdvancedView;
+	}
 #endif // WITH_EDITORONLY_DATA
 }
 
