Remove physics spring component from blueprint. Added comment about it being early days

diff --git a/Engine/Source/Runtime/Engine/Classes/PhysicsEngine/PhysicsSpringComponent.h b/Engine/Source/Runtime/Engine/Classes/PhysicsEngine/PhysicsSpringComponent.h
--- a/Engine/Source/Runtime/Engine/Classes/PhysicsEngine/PhysicsSpringComponent.h
+++ b/Engine/Source/Runtime/Engine/Classes/PhysicsEngine/PhysicsSpringComponent.h
@@ -5,10 +5,11 @@
 #include "PhysicsSpringComponent.generated.h"
 
 /** 
+ *  Note: this component is still work in progress. Uses raycast springs for simple vehicle forces
  *	Used with objects that have physics to create a spring down the X direction
  *	ie. point X in the direction you want generate spring.
  */
-UCLASS(hidecategories=(Object, Mobility, LOD), ClassGroup=Physics, showcategories=Trigger, meta=(BlueprintSpawnableComponent))
+UCLASS(hidecategories=(Object, Mobility, LOD), ClassGroup=Physics, showcategories=Trigger)
 class ENGINE_API UPhysicsSpringComponent : public USceneComponent
 {
 	GENERATED_UCLASS_BODY()
