In TermArticulated don't delete the FBodyInstance pointers within the initial loop to call TermBody on them. Doing so will cause a use-after-free error on Bodies[1] when clicking Build in some projects. The UPrimitiveComponent::UnWeldFromParent code will crash on NewRootBI->bWelded == false because NewRootBI is the pointer for Bodies[0] which was just deleted.

diff --git a/Engine/Source/Runtime/Engine/Private/SkeletalMeshComponentPhysics.cpp b/Engine/Source/Runtime/Engine/Private/SkeletalMeshComponentPhysics.cpp
--- a/Engine/Source/Runtime/Engine/Private/SkeletalMeshComponentPhysics.cpp
+++ b/Engine/Source/Runtime/Engine/Private/SkeletalMeshComponentPhysics.cpp
@@ -861,6 +861,11 @@ void USkeletalMeshComponent::TermArticulated()
 	{
 		check( Bodies[i] );
 		Bodies[i]->TermBody();
+	}
+	
+	for(int32 i=0; i<Bodies.Num(); i++)
+	{
+		check( Bodies[i] );
 		delete Bodies[i];
 	}
 
