Add comment for FWorldAsyncTraceState  as well as all Async data struct

diff --git a/Engine/Source/Runtime/Engine/Classes/Engine/World.h b/Engine/Source/Runtime/Engine/Classes/Engine/World.h
--- a/Engine/Source/Runtime/Engine/Classes/Engine/World.h
+++ b/Engine/Source/Runtime/Engine/Classes/Engine/World.h
@@ -456,20 +456,30 @@ struct ENGINE_API FActorSpawnParameters
 	EObjectFlags ObjectFlags;		
 };
 
+/**
+ *  This encapsulate World's async trace functionality. This contains two buffers of trace data buffer and alternates it for each tick. 
+ *  
+ *	You can use async trace using following APIs : AsyncLineTrace, AsyncSweep, AsyncOverlap
+ *	When you use those APIs, it will be saved to AsyncTraceData
+ *	FWorldAsyncTraceState contains two buffers to rotate each frame as you might need the result in the next frame
+ *	However, if you do not get the result by next frame, the result will be discarded. 
+ *	Use Delegate if you would like to get the result right away when available. 
+ */
 struct ENGINE_API FWorldAsyncTraceState
 {
 	FWorldAsyncTraceState();
 
+	/** Get the Buffer for input Frame */
 	FORCEINLINE AsyncTraceData& GetBufferForFrame        (int32 Frame) { return DataBuffer[ Frame             % 2]; }
+	/** Get the Buffer for Current Frame */
 	FORCEINLINE AsyncTraceData& GetBufferForCurrentFrame ()            { return DataBuffer[ CurrentFrame      % 2]; }
+	/** Get the Buffer for Previous Frame */
 	FORCEINLINE AsyncTraceData& GetBufferForPreviousFrame()            { return DataBuffer[(CurrentFrame + 1) % 2]; }
 
-	/** Async Trace Data Buffer **/
+	/** Async Trace Data Buffer Array. For now we only saves 2 frames. **/
 	AsyncTraceData DataBuffer[2];
 
-	/** Used as counter for Buffer swap for DataBuffer
-	 *	Right now it's only 2, but it can change
-	 */
+	/** Used as counter for Buffer swap for DataBuffer. Right now it's only 2, but it can change. */
 	int32 CurrentFrame;
 
 	/** Next available index for each pool - used as ID for each trace query **/
diff --git a/Engine/Source/Runtime/Engine/Public/WorldCollision.h b/Engine/Source/Runtime/Engine/Public/WorldCollision.h
--- a/Engine/Source/Runtime/Engine/Public/WorldCollision.h
+++ b/Engine/Source/Runtime/Engine/Public/WorldCollision.h
@@ -1,8 +1,7 @@
 // Copyright 1998-2014 Epic Games, Inc. All Rights Reserved.
 
-/*=============================================================================
-	WorldCollision.h: WorldCollision related data structures/types 
-=============================================================================*/
+// Structs that are used for Async Trace functionality
+// Mostly used by a batch of traces that you don't need a result right away
 
 #pragma once 
 
@@ -10,13 +9,12 @@
 #include "Runtime/Core/Public/Async/TaskGraphInterfaces.h"
 #include "CollisionQueryParams.h"
 
-/**
- * Async Trace Data Structures 
- */
+/** Trace Data Structs that are used for Async Trace */
 
-/** Trace Handle : unique ID that is returned once trace is requested **/
+/** Trace Handle - unique ID struct that is returned once trace is requested for tracking purpose **/
 struct FTraceHandle
 {
+	/** Handle is created by FrameNumber + Index of the request */
 	union
 	{
 		uint64	_Handle;
@@ -35,7 +33,7 @@ struct FTraceHandle
 	}
 };
 
-/** Trace Shape Types **/
+/** Types of Collision Shapes that are used by Trace **/
 namespace ECollisionShape
 {
 	enum Type
@@ -47,10 +45,12 @@ namespace ECollisionShape
 	};
 };
 
+/** Collision Shapes that supports Sphere, Capsule, Box, or Line **/
 struct FCollisionShape
 {
 	ECollisionShape::Type ShapeType;
 
+	/** Union that supports upto 3 floats **/
 	union
 	{
 		struct 
@@ -77,6 +77,7 @@ struct FCollisionShape
 		ShapeType = ECollisionShape::Line;
 	}
 
+	/** Utility function to Set Box and dimension */
 	void SetBox(const FVector& HalfExtent)
 	{
 		ShapeType = ECollisionShape::Box;
@@ -85,12 +86,14 @@ struct FCollisionShape
 		Box.HalfExtentZ = HalfExtent.Z;		
 	}
 
+	/** Utility function to set Sphere with Radius */
 	void SetSphere(const float Radius)
 	{
 		ShapeType = ECollisionShape::Sphere;
 		Sphere.Radius = Radius;
 	}
 
+	/** Utility function to set Capsule with Radius and Half Height */
 	void SetCapsule(const float Radius, const float HalfHeight)
 	{
 		ShapeType = ECollisionShape::Capsule;
@@ -98,6 +101,7 @@ struct FCollisionShape
 		Capsule.HalfHeight = HalfHeight;
 	}
 
+	/** Utility function to set Capsule from Extent data */
 	void SetCapsule(const FVector& Extent)
 	{
 		ShapeType = ECollisionShape::Capsule;
@@ -105,6 +109,7 @@ struct FCollisionShape
 		Capsule.HalfHeight = Extent.Z;
 	}
 	
+	/** Return true if nearly zero. If so, it will back out and use line trace instead */
 	bool IsNearlyZero() const
 	{
 		switch (ShapeType)
@@ -127,6 +132,7 @@ struct FCollisionShape
 		return true;
 	}
 
+	/** Utility function to return Extent of the shape */
 	FVector GetExtent() const
 	{
 		switch(ShapeType)
@@ -156,28 +162,34 @@ struct FCollisionShape
 		return FMath::Max<float>(Capsule.HalfHeight - Capsule.Radius, 1.f);
 	}
 
+	/** Utility function to get Box Extention */
 	FVector GetBox() const
 	{
 		return FVector(Box.HalfExtentX, Box.HalfExtentY, Box.HalfExtentZ);
 	}
 
+	/** Utility function to get Sphere Radius */
 	const float GetSphereRadius() const
 	{
 		return Sphere.Radius;
 	}
 
+	/** Utility function to get Capsule Radius */
 	const float GetCapsuleRadius() const
 	{
 		return Capsule.Radius;
 	}
 
+	/** Utility function to get Capsule Half Height */
 	const float GetCapsuleHalfHeight() const
 	{
 		return Capsule.HalfHeight;
 	}
 
+	/** Used by engine in multiple places. Since LineShape doesn't need any dimension, declare once and used by all codes. */
 	static struct FCollisionShape LineShape;
 
+	/** Static utility function to make a box */
 	static FCollisionShape MakeBox(const FVector& BoxHalfExtent)
 	{
 		FCollisionShape BoxShape;
@@ -185,6 +197,7 @@ struct FCollisionShape
 		return BoxShape;
 	}
 
+	/** Static utility function to make a sphere */
 	static FCollisionShape MakeSphere(const float SphereRadius)
 	{
 		FCollisionShape SphereShape;
@@ -192,6 +205,7 @@ struct FCollisionShape
 		return SphereShape;
 	}
 
+	/** Static utility function to make a capsule */
 	static FCollisionShape MakeCapsule(const float CapsuleRadius, const float CapsuleHalfHeight)
 	{
 		FCollisionShape CapsuleShape;
@@ -199,6 +213,7 @@ struct FCollisionShape
 		return CapsuleShape;
 	}
 
+	/** Static utility function to make a capsule */
 	static FCollisionShape MakeCapsule(const FVector& Extent)
 	{
 		FCollisionShape CapsuleShape;
@@ -207,42 +222,52 @@ struct FCollisionShape
 	}
 };
 
-/** Trace Shape Shapes - Params packs all possible shape types **/
+/** 
+ * Sets of Collision Parameters to run the async trace
+ * 
+ * It includes basic Query Parameter, Response Parameter, Object Query Parameter as well as Shape of collision testing 
+ *
+ */
 struct FCollisionParameters
 {
-	/** Collision trance params **/
+	/** Collision Trace Parameters **/
 	struct FCollisionQueryParams		CollisionQueryParam;
 	struct FCollisionResponseParams		ResponseParam;
 	struct FCollisionObjectQueryParams	ObjectQueryParam;
 
-	/** Shape data **/
+	/** Contains Collision Shape data including dimension of the shape **/
 	struct FCollisionShape CollisionShape;
 };
 
-/** Base Trace Data Struct **/
+/** 
+ * Base Async Trace Data Struct for both overlap and trace 
+ * 
+ * Contains basic data that will need for handling trace 
+ * such as World, Collision parameters and so on. 
+ */
 struct FBaseTraceDatum
 {
-	/* Physics World this trace will run in **/
+	/** Physics World this trace will run in ***/
 	TWeakObjectPtr<UWorld> PhysWorld;
 
-	/* Shape Data */
+	/** Collection of collision parameters */
 	FCollisionParameters	CollisionParams;
 
-	/* Collision trace info*/
+	/** Collsion Trace Channel that this trace is running*/
 	ECollisionChannel TraceChannel;
 
-	/* Framecount when requested is made*/
+	/** Framecount when requested is made*/
 	uint32	FrameNumber; 
 
-	/* User data*/
+	/** User data */
 	uint32	UserData;
 
-	/* single or multi - changed to uint32 if we get more booleans*/
+	/** single or multi trace. */
 	bool 	bIsMultiTrace;
 
 	FBaseTraceDatum() {};
 
-	/* Set functions for each Shape type */
+	/** Set functions for each Shape type */
 	void Set(UWorld * World, const FCollisionShape& InCollisionShape, const FCollisionQueryParams& Param, const struct FCollisionResponseParams &InResponseParam, const struct FCollisionObjectQueryParams& InObjectQueryParam, 
 		ECollisionChannel Channel, uint32 InUserData, bool bInIsMultiTrace, int32 FrameCounter);
 };
@@ -251,30 +276,42 @@ struct FTraceDatum;
 struct FOverlapDatum;
 
 /**
- * Delegate for Trace 
+ * This is Trace/Sweep Delegate that can be used if you'd like to get notified whenever available
+ * Otherwise, you'll have to query manually using your TraceHandle 
+ * 
  * @param	FTraceHandle	TraceHandle that is returned when requested
  * @param	FTraceDatum		TraceDatum that includes input/output
  */
 DECLARE_DELEGATE_TwoParams( FTraceDelegate, const FTraceHandle&, FTraceDatum &);
 /**
- * Delegate for Overlap
+ * This is Overlap Delegate that can be used if you'd like to get notified whenever available
+ * Otherwise, you'll have to query manually using your TraceHandle
+ * 
  * @param	FTHandle		TraceHandle that is returned when requestsed
  * @param	FOverlapDatum	OverlapDatum that includes input/output
  */
 DECLARE_DELEGATE_TwoParams( FOverlapDelegate, const FTraceHandle&, FOverlapDatum &);
 
 
-/** Trace Data Struct **/
+/**
+ * Trace/Sweep Data structure for async trace
+ * 
+ * This saves request information by main thread and result will be filled up by worker thread
+ */
 struct FTraceDatum : public FBaseTraceDatum
 {
-	/** Start of the trace **/
+	/** 
+	 * Input of the trace/sweep request. Filled up by main thread 
+	 * 
+	 * Start/End of the trace
+	 * The Shape is defined in FBaseTraceDatum
+	 */
 	FVector Start;
-	/** End of the trace **/
 	FVector End;
-	/** Delegate - optional - you can query **/
+	/** Delegate to be set if you want Delegate to be called when the output is available. Filled up by requester (main thread) **/
 	FTraceDelegate Delegate;
 
-	/** Output of the hits if made **/
+	/** Output of the overlap request. Filled up by worker thread */
 	TArray<struct FHitResult> OutHits;
 
 	FTraceDatum() {}
@@ -298,15 +335,25 @@ struct FTraceDatum : public FBaseTraceDatum
 	}
 };
 
-// overlap data struct
+/**
+ * Overlap Data structure for async trace
+ * 
+ * This saves request information by main thread and result will be filled up by worker thread
+ */
 struct FOverlapDatum : FBaseTraceDatum
 {
-	//input
+	/** 
+	 * Input of the overlap request. Filled up by main thread 
+	 *
+	 * Position/Rotation data of overlap request. 
+	 * The Shape is defined in FBaseTraceDatum
+	 */
 	FVector Pos;
 	FQuat	Rot;
+	/** Delegate to be set if you want Delegate to be called when the output is available. Filled up by requester (main thread) **/
 	FOverlapDelegate Delegate;
 
-	//output 
+	/** Output of the overlap request. Filled up by worker thread */
 	TArray<struct FOverlapResult> OutOverlaps;
 
 	FOverlapDatum() {}
@@ -332,33 +379,46 @@ struct FOverlapDatum : FBaseTraceDatum
 
 #define ASYNC_TRACE_BUFFER_SIZE 64
 
+/**
+ * Trace Data that one Thread can handle per type. For now ASYNC_TRACE_BUFFER_SIZE is 64. 
+ */
 template <typename T>
 struct TTraceThreadData
 {
 	T Buffer[ASYNC_TRACE_BUFFER_SIZE];
 };
 
-/** We use double buffer for trace data pool. FrameNumber % 2 = is going to be the one collecting NEW data **/
+/** 
+ * Contains all Async Trace Result for one frame. 
+ * 
+ * We use double buffer for trace data pool. FrameNumber % 2 = is going to be the one collecting NEW data 
+ * Check FWorldAsyncTraceState to see how this is used. For now it is only double buffer, but it can be more. 
+ */
 struct AsyncTraceData : FNoncopyable
 {
-	// Data Buffer for each trace
-
-	// FTraceThreadData is one atomic data size for thread
-	// so once that's filled up, this will be sent to thread
-	// if you need more, it will allocate new FTraceThreadData and add to TArray
-
-	// however when we calculate index of buffer, we calculate continuously, 
-	// meaning if TraceData(1).Buffer[50] will have 1*ASYNC_TRACE_BUFFER_SIZE + 50 as index
-	// in order to give UNIQUE INDEX to every data in this buffer
+	/**
+	 * Data Buffer for each trace type - one for Trace/Sweep and one for Overlap
+	 *
+	 * FTraceThreadData is one atomic data size for thread
+	 * so once that's filled up, this will be sent to thread
+	 * if you need more, it will allocate new FTraceThreadData and add to TArray
+	 *
+	 * however when we calculate index of buffer, we calculate continuously, 
+	 * meaning if TraceData(1).Buffer[50] will have 1*ASYNC_TRACE_BUFFER_SIZE + 50 as index
+	 * in order to give UNIQUE INDEX to every data in this buffer
+	 */
 	TArray<TUniquePtr<TTraceThreadData<FTraceDatum>>>			TraceData;
 	TArray<TUniquePtr<TTraceThreadData<FOverlapDatum>>>			OverlapData;
 
-	// if Execution is all done, set this to be true
-	// when reinitialize bAsyncAllowed is true, once execution is done, this will set to be false
-	// this is to find cases where execution is already done but another request is made again within the same frame. 
-	bool					bAsyncAllowed; // are async calls allowed this frame
+	/**
+	 * if Execution is all done, set this to be true
+	 * 
+	 * when reinitialize bAsyncAllowed is true, once execution is done, this will set to be false
+	 * this is to find cases where execution is already done but another request is made again within the same frame. 
+	 */
+	bool					bAsyncAllowed; 
 
-	// Thread completion event for batch
+	/**  Thread completion event for batch **/
 	FGraphEventArray		AsyncTraceCompletionEvent;
 
 	AsyncTraceData() : bAsyncAllowed(false) {}
