Make MinRespawnTime blueprint read only instead of read/write
Make GameMode transient
Expose NumBots as Blueprint read

diff --git a/Engine/Source/Runtime/Engine/Classes/GameFramework/GameMode.h b/Engine/Source/Runtime/Engine/Classes/GameFramework/GameMode.h
--- a/Engine/Source/Runtime/Engine/Classes/GameFramework/GameMode.h
+++ b/Engine/Source/Runtime/Engine/Classes/GameFramework/GameMode.h
@@ -58,7 +58,7 @@ struct FGameClassShortName
  *
  * @see https://docs.unrealengine.com/latest/INT/Gameplay/Framework/GameMode/index.html
  */
-UCLASS(config=Game, notplaceable, BlueprintType, Blueprintable, hidecategories=(Info, Rendering, MovementReplication, Replication, Actor))
+UCLASS(config=Game, notplaceable, BlueprintType, Blueprintable, Transient, hidecategories=(Info, Rendering, MovementReplication, Replication, Actor))
 class ENGINE_API AGameMode : public AInfo
 {
 	GENERATED_UCLASS_BODY()
@@ -178,11 +178,11 @@ public:
 	int32 NumPlayers;    
 
 	/** number of non-human players (AI controlled but participating as a player). */
-	UPROPERTY()
+	UPROPERTY(BlueprintReadOnly, Category=GameMode)
 	int32 NumBots;    
 
 	/** Minimum time before player can respawn after dying. */
-	UPROPERTY(EditAnywhere, BlueprintReadWrite, Category=GameMode, meta=(DisplayName="Minimum Respawn Delay"))
+	UPROPERTY(EditAnywhere, BlueprintReadOnly, Category=GameMode, meta=(DisplayName="Minimum Respawn Delay"))
 	float MinRespawnDelay;
 
 	/** Game Session handles login approval, arbitration, online game interface */
