updated docs - r.ScreenPercentage can be used for SuperSampling

diff --git a/Engine/Source/Runtime/Engine/Classes/Engine/Scene.h b/Engine/Source/Runtime/Engine/Classes/Engine/Scene.h
--- a/Engine/Source/Runtime/Engine/Classes/Engine/Scene.h
+++ b/Engine/Source/Runtime/Engine/Classes/Engine/Scene.h
@@ -743,7 +743,13 @@ struct FPostProcessSettings
 	UPROPERTY(interp, BlueprintReadWrite, Category=MotionBlur, AdvancedDisplay, meta=(ClampMin = "0.0", UIMax = "100.0", editcondition = "bOverride_MotionBlurPerObjectSize", DisplayName = "Per Object Size"))
 	float MotionBlurPerObjectSize;
 
-	/** to render with lower resolution and upscale, controlled by console variable, 100:off, needs to be <99 to see effect, only applied in game  */
+	/**
+	 * To render with lower or high resolution than it is presented, 
+	 * controlled by console variable, 
+	 * 100:off, needs to be <99 to get upsampling and lower to get performance,
+	 * >100 for super sampling (slower but higher quality), 
+	 * only applied in game
+	 */
 	UPROPERTY(interp, BlueprintReadWrite, Category=Misc, AdvancedDisplay, meta=(ClampMin = "0.0", ClampMax = "400.0", editcondition = "bOverride_ScreenPercentage"))
 	float ScreenPercentage;
 
diff --git a/Engine/Source/Runtime/Engine/Private/LocalPlayer.cpp b/Engine/Source/Runtime/Engine/Private/LocalPlayer.cpp
--- a/Engine/Source/Runtime/Engine/Private/LocalPlayer.cpp
+++ b/Engine/Source/Runtime/Engine/Private/LocalPlayer.cpp
@@ -757,7 +757,7 @@ FSceneView* ULocalPlayer::CalcSceneView( class FSceneViewFamily* ViewFamily,
 		View->EndFinalPostprocessSettings();
 	}
 
-	// Upscaling
+	// Upscaling or Super sampling
 	{
 		float LocalScreenPercentage = View->FinalPostProcessSettings.ScreenPercentage;
 
diff --git a/Engine/Source/Runtime/Engine/Public/SceneView.h b/Engine/Source/Runtime/Engine/Public/SceneView.h
--- a/Engine/Source/Runtime/Engine/Public/SceneView.h
+++ b/Engine/Source/Runtime/Engine/Public/SceneView.h
@@ -213,7 +213,7 @@ public:
 	/** An interaction which draws the view's interaction elements. */
 	FViewElementDrawer* Drawer;
 
-	/* Final position of the view in the final render target (in pixels), potentially scaled down by ScreenPercentage */
+	/* Final position of the view in the final render target (in pixels), potentially scaled by ScreenPercentage */
 	FIntRect ViewRect;
 
 	/* Final position of the view in the final render target (in pixels), potentially constrained by an aspect ratio requirement (black bars) */
