Improved comments

diff --git a/Engine/Source/Runtime/Engine/Private/NetworkDriver.cpp b/Engine/Source/Runtime/Engine/Private/NetworkDriver.cpp
--- a/Engine/Source/Runtime/Engine/Private/NetworkDriver.cpp
+++ b/Engine/Source/Runtime/Engine/Private/NetworkDriver.cpp
@@ -1893,12 +1893,14 @@ int32 UNetDriver::ServerReplicateActors(float DeltaSeconds)
 				continue;
 			}
 
-			// Don't send actors that may still be streaming in
+			// This actor may belong to a different net driver, make sure this is the correct one
+			// (this can happen when using beacon net drivers for example)
 			if ( Actor->NetDriverName != NetDriverName )
 			{
 				continue;
 			}
 
+			// Don't send actors that may still be streaming in
 			ULevel* Level = Actor->GetLevel();
 			if ( Level->HasVisibilityRequestPending() || Level->bIsAssociatingLevel )
 			{
