Removed outdated comment

diff --git a/Engine/Source/Runtime/Engine/Private/LevelTick.cpp b/Engine/Source/Runtime/Engine/Private/LevelTick.cpp
--- a/Engine/Source/Runtime/Engine/Private/LevelTick.cpp
+++ b/Engine/Source/Runtime/Engine/Private/LevelTick.cpp
@@ -1273,7 +1273,6 @@ void UWorld::Tick( ELevelTick TickType, float DeltaSeconds )
 
 	if ( FullPurgeTriggered )
 	{
-		// NOTE: This will block until all current Async I/O requests have been completed before doing the GC.
 		CollectGarbage( GARBAGE_COLLECTION_KEEPFLAGS, true );
 		CleanupActors();
 		FullPurgeTriggered = false;
