TSubobjectPtr removal, phase II. Making sure we're not breaking licensee projects by making subobjects TEMPORARILY public and  deprecated (but only for game modules).

#change Made subobjects public again
#change Added temporary privat_subobject macro (defaults to public keyword) to help with changing eveything back later
#change TEMPORARILY moved subobject accessors to cpp fiels so that we don't get deprecation warnings when compiling the engine headers
#change Added DEPRECATED_FORGAME macro which is empty for the engine modules and reverts back to DEPRECATED for game modules

diff --git a/Engine/Source/Developer/FunctionalTesting/Classes/FunctionalTest.h b/Engine/Source/Developer/FunctionalTesting/Classes/FunctionalTest.h
--- a/Engine/Source/Developer/FunctionalTesting/Classes/FunctionalTest.h
+++ b/Engine/Source/Developer/FunctionalTesting/Classes/FunctionalTest.h
@@ -29,7 +29,8 @@ class AFunctionalTest : public AActor
 
 	static const uint32 DefaultTimeLimit = 60;	// seconds
 
-private:
+private_subobject:
+	DEPRECATED_FORGAME(4.6, "SpriteComponent should not be accessed directly, please use GetSpriteComponent() function instead. SpriteComponent will soon be private and your code will not compile.")
 	UPROPERTY()
 	UBillboardComponent* SpriteComponent;
 public:
@@ -127,5 +128,5 @@ protected:
 
 public:
 	/** Returns SpriteComponent subobject **/
-	FORCEINLINE UBillboardComponent* GetSpriteComponent() { return SpriteComponent; }
+	UBillboardComponent* GetSpriteComponent();
 };
diff --git a/Engine/Source/Developer/FunctionalTesting/Private/FunctionalTest.cpp b/Engine/Source/Developer/FunctionalTesting/Private/FunctionalTest.cpp
--- a/Engine/Source/Developer/FunctionalTesting/Private/FunctionalTest.cpp
+++ b/Engine/Source/Developer/FunctionalTesting/Private/FunctionalTest.cpp
@@ -254,4 +254,7 @@ void AFunctionalTest::GoToObservationPoint()
 			PC->SetControlRotation(ObservationPoint->GetActorRotation());
 		}
 	}
-}
\ No newline at end of file
+}
+
+/** Returns SpriteComponent subobject **/
+UBillboardComponent* AFunctionalTest::GetSpriteComponent() { return SpriteComponent; }
diff --git a/Engine/Source/Programs/UnrealBuildTool/Configuration/UEBuildModule.cs b/Engine/Source/Programs/UnrealBuildTool/Configuration/UEBuildModule.cs
--- a/Engine/Source/Programs/UnrealBuildTool/Configuration/UEBuildModule.cs
+++ b/Engine/Source/Programs/UnrealBuildTool/Configuration/UEBuildModule.cs
@@ -656,6 +656,11 @@ namespace UnrealBuildTool
 		{
 			var VisitedModules = new Dictionary<UEBuildModule, bool>();
 
+			if (this.Type == UEBuildModuleType.Game)
+			{
+				Definitions.Add("DEPRECATED_FORGAME=DEPRECATED");
+			}
+
 			// Add this module's private include paths and definitions.
 			AddIncludePathsWithChecks(IncludePaths, PrivateIncludePaths);
 
diff --git a/Engine/Source/Runtime/AIModule/Classes/AIController.h b/Engine/Source/Runtime/AIModule/Classes/AIController.h
--- a/Engine/Source/Runtime/AIModule/Classes/AIController.h
+++ b/Engine/Source/Runtime/AIModule/Classes/AIController.h
@@ -90,12 +90,14 @@ public:
 	UPROPERTY()
 	uint32 bWantsPlayerState:1;
 
-private:
+private_subobject:
 	/** Component used for pathfinding and querying environment's navigation. */
+	DEPRECATED_FORGAME(4.6, "NavComponent should not be accessed directly, please use GetNavComponent() function instead. NavComponent will soon be private and your code will not compile.")
 	UPROPERTY()
 	UNavigationComponent* NavComponent;
 
 	/** Component used for moving along a path. */
+	DEPRECATED_FORGAME(4.6, "PathFollowingComponent should not be accessed directly, please use GetPathFollowingComponent() function instead. PathFollowingComponent will soon be private and your code will not compile.")
 	UPROPERTY()
 	UPathFollowingComponent* PathFollowingComponent;
 
@@ -108,7 +110,8 @@ public:
 	UPROPERTY()
 	UAIPerceptionComponent* PerceptionComponent;
 
-private:
+private_subobject:
+	DEPRECATED_FORGAME(4.6, "ActionsComp should not be accessed directly, please use GetActionsComp() function instead. ActionsComp will soon be private and your code will not compile.")
 	UPROPERTY(BlueprintReadOnly, Category = AI, meta = (AllowPrivateAccess = "true"))
 	UPawnActionsComponent* ActionsComp;
 
@@ -165,7 +168,7 @@ public:
 	virtual void OnMoveCompleted(FAIRequestID RequestID, EPathFollowingResult::Type Result);
 
 	/** Returns the Move Request ID for the current move */
-	FORCEINLINE FAIRequestID GetCurrentMoveRequestID() const { return PathFollowingComponent ? PathFollowingComponent->GetCurrentRequestId() : FAIRequestID::InvalidRequest; }
+	FORCEINLINE FAIRequestID GetCurrentMoveRequestID() const { return GetPathFollowingComponent() ? GetPathFollowingComponent()->GetCurrentRequestId() : FAIRequestID::InvalidRequest; }
 
 	/** Blueprint notification that we've completed the current movement request */
 	UPROPERTY(BlueprintAssignable, meta=(DisplayName="MoveCompleted"))
@@ -312,11 +315,11 @@ public:
 
 public:
 	/** Returns NavComponent subobject **/
-	FORCEINLINE UNavigationComponent* GetNavComponent() const { return NavComponent; }
+	UNavigationComponent* GetNavComponent() const;
 	/** Returns PathFollowingComponent subobject **/
-	FORCEINLINE UPathFollowingComponent* GetPathFollowingComponent() const { return PathFollowingComponent; }
+	UPathFollowingComponent* GetPathFollowingComponent() const;
 	/** Returns ActionsComp subobject **/
-	FORCEINLINE UPawnActionsComponent* GetActionsComp() const { return ActionsComp; }
+	UPawnActionsComponent* GetActionsComp() const;
 };
 
 
diff --git a/Engine/Source/Runtime/AIModule/Classes/EnvironmentQuery/EQSTestingPawn.h b/Engine/Source/Runtime/AIModule/Classes/EnvironmentQuery/EQSTestingPawn.h
--- a/Engine/Source/Runtime/AIModule/Classes/EnvironmentQuery/EQSTestingPawn.h
+++ b/Engine/Source/Runtime/AIModule/Classes/EnvironmentQuery/EQSTestingPawn.h
@@ -45,8 +45,9 @@ class AIMODULE_API AEQSTestingPawn : public ACharacter, public IEQSQueryResultSo
 	
 
 #if WITH_EDITORONLY_DATA
-private:
+private_subobject:
 	/** Editor Preview */
+	DEPRECATED_FORGAME(4.6, "EdRenderComp should not be accessed directly, please use GetEdRenderComp() function instead. EdRenderComp will soon be private and your code will not compile.")
 	UPROPERTY()
 	class UEQSRenderingComponent* EdRenderComp;
 #endif // WITH_EDITORONLY_DATA
@@ -87,6 +88,6 @@ protected:
 public:
 #if WITH_EDITORONLY_DATA
 	/** Returns EdRenderComp subobject **/
-	FORCEINLINE class UEQSRenderingComponent* GetEdRenderComp() { return EdRenderComp; }
+	class UEQSRenderingComponent* GetEdRenderComp();
 #endif
 };
diff --git a/Engine/Source/Runtime/AIModule/Private/AIController.cpp b/Engine/Source/Runtime/AIModule/Private/AIController.cpp
--- a/Engine/Source/Runtime/AIModule/Private/AIController.cpp
+++ b/Engine/Source/Runtime/AIModule/Private/AIController.cpp
@@ -796,3 +796,10 @@ FString AAIController::GetDebugIcon() const
 	return TEXT("/Engine/EngineResources/AICON-Green.AICON-Green");
 }
 
+
+/** Returns NavComponent subobject **/
+UNavigationComponent* AAIController::GetNavComponent() const { return NavComponent; }
+/** Returns PathFollowingComponent subobject **/
+UPathFollowingComponent* AAIController::GetPathFollowingComponent() const { return PathFollowingComponent; }
+/** Returns ActionsComp subobject **/
+UPawnActionsComponent* AAIController::GetActionsComp() const { return ActionsComp; }
diff --git a/Engine/Source/Runtime/AIModule/Private/EnvironmentQuery/EQSTestingPawn.cpp b/Engine/Source/Runtime/AIModule/Private/EnvironmentQuery/EQSTestingPawn.cpp
--- a/Engine/Source/Runtime/AIModule/Private/EnvironmentQuery/EQSTestingPawn.cpp
+++ b/Engine/Source/Runtime/AIModule/Private/EnvironmentQuery/EQSTestingPawn.cpp
@@ -289,3 +289,8 @@ void AEQSTestingPawn::PostEditMove(bool bFinished)
 }
 
 #endif // WITH_EDITOR
+
+#if WITH_EDITORONLY_DATA
+/** Returns EdRenderComp subobject **/
+UEQSRenderingComponent* AEQSTestingPawn::GetEdRenderComp() { return EdRenderComp; }
+#endif
diff --git a/Engine/Source/Runtime/Core/Public/HAL/Platform.h b/Engine/Source/Runtime/Core/Public/HAL/Platform.h
--- a/Engine/Source/Runtime/Core/Public/HAL/Platform.h
+++ b/Engine/Source/Runtime/Core/Public/HAL/Platform.h
@@ -427,6 +427,15 @@
 	#define DLLIMPORT
 #endif
 
+
+#ifndef DEPRECATED_FORGAME
+	#define DEPRECATED_FORGAME(...)
+#endif
+
+#ifndef private_subobject
+#define private_subobject public
+#endif
+
 // explicit bool support
 namespace FHasOperatorImpl
 {
diff --git a/Engine/Source/Runtime/Engine/Classes/AI/Navigation/NavLinkProxy.h b/Engine/Source/Runtime/Engine/Classes/AI/Navigation/NavLinkProxy.h
--- a/Engine/Source/Runtime/Engine/Classes/AI/Navigation/NavLinkProxy.h
+++ b/Engine/Source/Runtime/Engine/Classes/AI/Navigation/NavLinkProxy.h
@@ -23,8 +23,9 @@ class ENGINE_API ANavLinkProxy : public AActor, public INavLinkHostInterface, pu
 	UPROPERTY()
 	TArray<FNavigationSegmentLink> SegmentLinks;
 
-private:
+private_subobject:
 	/** Smart link: can affect path following */
+	DEPRECATED_FORGAME(4.6, "SmartLinkComp should not be accessed directly, please use GetSmartLinkComp() function instead. SmartLinkComp will soon be private and your code will not compile.")
 	UPROPERTY(VisibleAnywhere, Category=SmartLink)
 	class UNavLinkCustomComponent* SmartLinkComp;
 public:
@@ -34,11 +35,13 @@ public:
 	bool bSmartLinkIsRelevant;
 
 #if WITH_EDITORONLY_DATA
-private:
+private_subobject:
 	/** Editor Preview */
+	DEPRECATED_FORGAME(4.6, "EdRenderComp should not be accessed directly, please use GetEdRenderComp() function instead. EdRenderComp will soon be private and your code will not compile.")
 	UPROPERTY()
 	class UNavLinkRenderingComponent* EdRenderComp;
 
+	DEPRECATED_FORGAME(4.6, "SpriteComponent should not be accessed directly, please use GetSpriteComponent() function instead. SpriteComponent will soon be private and your code will not compile.")
 	UPROPERTY()
 	class UBillboardComponent* SpriteComponent;
 public:
@@ -93,11 +96,11 @@ protected:
 
 public:
 	/** Returns SmartLinkComp subobject **/
-	FORCEINLINE class UNavLinkCustomComponent* GetSmartLinkComp() const { return SmartLinkComp; }
+	class UNavLinkCustomComponent* GetSmartLinkComp() const;
 #if WITH_EDITORONLY_DATA
 	/** Returns EdRenderComp subobject **/
-	FORCEINLINE class UNavLinkRenderingComponent* GetEdRenderComp() const { return EdRenderComp; }
+	class UNavLinkRenderingComponent* GetEdRenderComp() const;
 	/** Returns SpriteComponent subobject **/
-	FORCEINLINE class UBillboardComponent* GetSpriteComponent() const { return SpriteComponent; }
+	class UBillboardComponent* GetSpriteComponent() const;
 #endif
 };
diff --git a/Engine/Source/Runtime/Engine/Classes/AI/Navigation/NavigationTestingActor.h b/Engine/Source/Runtime/Engine/Classes/AI/Navigation/NavigationTestingActor.h
--- a/Engine/Source/Runtime/Engine/Classes/AI/Navigation/NavigationTestingActor.h
+++ b/Engine/Source/Runtime/Engine/Classes/AI/Navigation/NavigationTestingActor.h
@@ -31,12 +31,14 @@ class ENGINE_API ANavigationTestingActor : public AActor, public INavAgentInterf
 {
 	GENERATED_UCLASS_BODY()
 
-private:
+private_subobject:
+	DEPRECATED_FORGAME(4.6, "CapsuleComponent should not be accessed directly, please use GetCapsuleComponent() function instead. CapsuleComponent will soon be private and your code will not compile.")
 	UPROPERTY()
 	class UCapsuleComponent* CapsuleComponent;
 
 #if WITH_EDITORONLY_DATA
 	/** Editor Preview */
+	DEPRECATED_FORGAME(4.6, "EdRenderComp should not be accessed directly, please use GetEdRenderComp() function instead. EdRenderComp will soon be private and your code will not compile.")
 	UPROPERTY()
 	class UNavTestRenderingComponent* EdRenderComp;
 #endif // WITH_EDITORONLY_DATA
@@ -174,9 +176,9 @@ public:
 
 public:
 	/** Returns CapsuleComponent subobject **/
-	FORCEINLINE class UCapsuleComponent* GetCapsuleComponent() const { return CapsuleComponent; }
+	class UCapsuleComponent* GetCapsuleComponent() const;
 #if WITH_EDITORONLY_DATA
 	/** Returns EdRenderComp subobject **/
-	FORCEINLINE class UNavTestRenderingComponent* GetEdRenderComp() const { return EdRenderComp; }
+	class UNavTestRenderingComponent* GetEdRenderComp() const;
 #endif
 };
diff --git a/Engine/Source/Runtime/Engine/Classes/Animation/SkeletalMeshActor.h b/Engine/Source/Runtime/Engine/Classes/Animation/SkeletalMeshActor.h
--- a/Engine/Source/Runtime/Engine/Classes/Animation/SkeletalMeshActor.h
+++ b/Engine/Source/Runtime/Engine/Classes/Animation/SkeletalMeshActor.h
@@ -24,7 +24,8 @@ class ENGINE_API ASkeletalMeshActor : public AActor, public IMatineeAnimInterfac
 	UPROPERTY()
 	uint32 bWakeOnLevelStart_DEPRECATED:1;
 
-private:
+private_subobject:
+	DEPRECATED_FORGAME(4.6, "SkeletalMeshComponent should not be accessed directly, please use GetSkeletalMeshComponent() function instead. SkeletalMeshComponent will soon be private and your code will not compile.")
 	UPROPERTY(Category = SkeletalMeshActor, VisibleAnywhere, BlueprintReadOnly, meta = (ExposeFunctionCategories = "Mesh,Components|SkeletalMesh,Animation,Physics", AllowPrivateAccess = "true"))
 	class USkeletalMeshComponent* SkeletalMeshComponent;
 public:
@@ -93,7 +94,7 @@ private:
 
 public:
 	/** Returns SkeletalMeshComponent subobject **/
-	FORCEINLINE class USkeletalMeshComponent* GetSkeletalMeshComponent() { return SkeletalMeshComponent; }
+	class USkeletalMeshComponent* GetSkeletalMeshComponent();
 };
 
 
diff --git a/Engine/Source/Runtime/Engine/Classes/Atmosphere/AtmosphericFog.h b/Engine/Source/Runtime/Engine/Classes/Atmosphere/AtmosphericFog.h
--- a/Engine/Source/Runtime/Engine/Classes/Atmosphere/AtmosphericFog.h
+++ b/Engine/Source/Runtime/Engine/Classes/Atmosphere/AtmosphericFog.h
@@ -13,13 +13,15 @@ class AAtmosphericFog : public AInfo
 {
 	GENERATED_UCLASS_BODY()
 
-private:
+private_subobject:
 	/** Main fog component */
+	DEPRECATED_FORGAME(4.6, "AtmosphericFogComponent should not be accessed directly, please use GetAtmosphericFogComponent() function instead. AtmosphericFogComponent will soon be private and your code will not compile.")
 	UPROPERTY(BlueprintReadOnly, VisibleAnywhere, Category = Atmosphere, meta = (AllowPrivateAccess = "true"))
 	class UAtmosphericFogComponent* AtmosphericFogComponent;
 
 #if WITH_EDITORONLY_DATA
 	/** Arrow component to indicate default sun rotation */
+	DEPRECATED_FORGAME(4.6, "ArrowComponent should not be accessed directly, please use GetArrowComponent() function instead. ArrowComponent will soon be private and your code will not compile.")
 	UPROPERTY()
 	class UArrowComponent* ArrowComponent;
 #endif
@@ -31,10 +33,10 @@ public:
 #endif
 
 	/** Returns AtmosphericFogComponent subobject **/
-	FORCEINLINE class UAtmosphericFogComponent* GetAtmosphericFogComponent() { return AtmosphericFogComponent; }
+	ENGINE_API class UAtmosphericFogComponent* GetAtmosphericFogComponent();
 #if WITH_EDITORONLY_DATA
 	/** Returns ArrowComponent subobject **/
-	FORCEINLINE class UArrowComponent* GetArrowComponent() { return ArrowComponent; }
+	ENGINE_API class UArrowComponent* GetArrowComponent();
 #endif
 };
 
diff --git a/Engine/Source/Runtime/Engine/Classes/Camera/CameraActor.h b/Engine/Source/Runtime/Engine/Classes/Camera/CameraActor.h
--- a/Engine/Source/Runtime/Engine/Classes/Camera/CameraActor.h
+++ b/Engine/Source/Runtime/Engine/Classes/Camera/CameraActor.h
@@ -18,7 +18,10 @@ private:
 	UPROPERTY(Category="AutoPlayerActivation", EditAnywhere)
 	TEnumAsByte<EAutoReceiveInput::Type> AutoActivateForPlayer;
 
+private_subobject:
+
 	/** The camera component for this camera */
+	DEPRECATED_FORGAME(4.6, "CameraComponent should not be accessed directly, please use GetCameraComponent() function instead. CameraComponent will soon be private and your code will not compile.")
 	UPROPERTY(Category = CameraActor, VisibleAnywhere, BlueprintReadOnly, meta = (AllowPrivateAccess = "true"))
 	class UCameraComponent* CameraComponent;
 
@@ -63,5 +66,5 @@ public:
 	// End AActor interface
 
 	/** Returns CameraComponent subobject **/
-	FORCEINLINE class UCameraComponent* GetCameraComponent() const { return CameraComponent; }
+	ENGINE_API class UCameraComponent* GetCameraComponent() const;
 };
diff --git a/Engine/Source/Runtime/Engine/Classes/Camera/CameraAnimInst.h b/Engine/Source/Runtime/Engine/Classes/Camera/CameraAnimInst.h
--- a/Engine/Source/Runtime/Engine/Classes/Camera/CameraAnimInst.h
+++ b/Engine/Source/Runtime/Engine/Classes/Camera/CameraAnimInst.h
@@ -16,8 +16,9 @@ class ENGINE_API UCameraAnimInst : public UObject
 	UPROPERTY()
 	class UCameraAnim* CamAnim;
 
-private:
+private_subobject:
 	/** the UInterpGroupInst used to do the interpolation */
+	DEPRECATED_FORGAME(4.6, "InterpGroupInst should not be accessed directly, please use GetInterpGroupInst() function instead. InterpGroupInst will soon be private and your code will not compile.")
 	UPROPERTY(instanced)
 	class UInterpGroupInst* InterpGroupInst;
 
@@ -136,7 +137,7 @@ public:
 
 protected:
 	/** Returns InterpGroupInst subobject **/
-	FORCEINLINE class UInterpGroupInst* GetInterpGroupInst() const { return InterpGroupInst; }
+	class UInterpGroupInst* GetInterpGroupInst() const;
 };
 
 
diff --git a/Engine/Source/Runtime/Engine/Classes/Camera/PlayerCameraManager.h b/Engine/Source/Runtime/Engine/Classes/Camera/PlayerCameraManager.h
--- a/Engine/Source/Runtime/Engine/Classes/Camera/PlayerCameraManager.h
+++ b/Engine/Source/Runtime/Engine/Classes/Camera/PlayerCameraManager.h
@@ -144,8 +144,9 @@ class ENGINE_API APlayerCameraManager : public AActor
 	UPROPERTY()
 	class APlayerController* PCOwner;
 
-private:
+private_subobject:
 	/** Dummy component we can use to attach things to the camera. */
+	DEPRECATED_FORGAME(4.6, "TransformComponent should not be accessed directly, please use GetTransformComponent() function instead. TransformComponent will soon be private and your code will not compile.")
 	UPROPERTY(Category = PlayerCameraManager, VisibleAnywhere, BlueprintReadOnly, meta = (AllowPrivateAccess = "true"))
 	class USceneComponent* TransformComponent;
 public:
@@ -644,5 +645,5 @@ private:
 
 public:
 	/** Returns TransformComponent subobject **/
-	FORCEINLINE class USceneComponent* GetTransformComponent() const { return TransformComponent; }
+	class USceneComponent* GetTransformComponent() const;
 };
diff --git a/Engine/Source/Runtime/Engine/Classes/Engine/Brush.h b/Engine/Source/Runtime/Engine/Classes/Engine/Brush.h
--- a/Engine/Source/Runtime/Engine/Classes/Engine/Brush.h
+++ b/Engine/Source/Runtime/Engine/Classes/Engine/Brush.h
@@ -109,7 +109,8 @@ class ENGINE_API ABrush : public AActor
 	UPROPERTY(export)
 	class UModel* Brush;
 
-private:
+private_subobject:
+	DEPRECATED_FORGAME(4.6, "BrushComponent should not be accessed directly, please use GetBrushComponent() function instead. BrushComponent will soon be private and your code will not compile.")
 	UPROPERTY(Category = Collision, VisibleAnywhere, BlueprintReadOnly, meta = (AllowPrivateAccess = "true"))
 	class UBrushComponent* BrushComponent;
 public:
@@ -256,7 +257,7 @@ public:
 
 public:
 	/** Returns BrushComponent subobject **/
-	FORCEINLINE class UBrushComponent* GetBrushComponent() const { return BrushComponent; }
+	class UBrushComponent* GetBrushComponent() const;
 };
 
 
diff --git a/Engine/Source/Runtime/Engine/Classes/Engine/DecalActor.h b/Engine/Source/Runtime/Engine/Classes/Engine/DecalActor.h
--- a/Engine/Source/Runtime/Engine/Classes/Engine/DecalActor.h
+++ b/Engine/Source/Runtime/Engine/Classes/Engine/DecalActor.h
@@ -18,21 +18,25 @@ class ADecalActor : public AActor
 {
 	GENERATED_UCLASS_BODY()
 
-private:
+private_subobject:
 	/** The decal component for this decal actor */
+	DEPRECATED_FORGAME(4.6, "Decal should not be accessed directly, please use GetDecal() function instead. Decal will soon be private and your code will not compile.")
 	UPROPERTY(Category = Decal, VisibleAnywhere, BlueprintReadOnly, meta = (ExposeFunctionCategories = "Decal,Rendering|Components|Decal", AllowPrivateAccess = "true"))
 	class UDecalComponent* Decal;
 
 #if WITH_EDITORONLY_DATA
 	/* Reference to the editor only arrow visualization component */
+	DEPRECATED_FORGAME(4.6, "ArrowComponent should not be accessed directly, please use GetArrowComponent() function instead. ArrowComponent will soon be private and your code will not compile.")
 	UPROPERTY()
 	class UArrowComponent* ArrowComponent;
 
 	/* Reference to the billboard component */
+	DEPRECATED_FORGAME(4.6, "SpriteComponent should not be accessed directly, please use GetSpriteComponent() function instead. SpriteComponent will soon be private and your code will not compile.")
 	UPROPERTY()
 	UBillboardComponent* SpriteComponent;
 
 	/* Reference to the selected visualization box component */
+	DEPRECATED_FORGAME(4.6, "BoxComponent should not be accessed directly, please use GetBoxComponent() function instead. BoxComponent will soon be private and your code will not compile.")
 	UPROPERTY()
 	UBoxComponent* BoxComponent;
 #endif
@@ -62,14 +66,14 @@ public:
 
 public:
 	/** Returns Decal subobject **/
-	FORCEINLINE class UDecalComponent* GetDecal() const { return Decal; }
+	ENGINE_API class UDecalComponent* GetDecal() const;
 #if WITH_EDITORONLY_DATA
 	/** Returns ArrowComponent subobject **/
-	FORCEINLINE class UArrowComponent* GetArrowComponent() const { return ArrowComponent; }
+	ENGINE_API class UArrowComponent* GetArrowComponent() const;
 	/** Returns SpriteComponent subobject **/
-	FORCEINLINE UBillboardComponent* GetSpriteComponent() const { return SpriteComponent; }
+	ENGINE_API UBillboardComponent* GetSpriteComponent() const;
 	/** Returns BoxComponent subobject **/
-	FORCEINLINE UBoxComponent* GetBoxComponent() const { return BoxComponent; }
+	ENGINE_API UBoxComponent* GetBoxComponent() const;
 #endif
 };
 
diff --git a/Engine/Source/Runtime/Engine/Classes/Engine/DirectionalLight.h b/Engine/Source/Runtime/Engine/Classes/Engine/DirectionalLight.h
--- a/Engine/Source/Runtime/Engine/Classes/Engine/DirectionalLight.h
+++ b/Engine/Source/Runtime/Engine/Classes/Engine/DirectionalLight.h
@@ -11,7 +11,8 @@ class ADirectionalLight : public ALight
 
 #if WITH_EDITORONLY_DATA
 	// Reference to editor visualization arrow
-private:
+private_subobject:
+	DEPRECATED_FORGAME(4.6, "ArrowComponent should not be accessed directly, please use GetArrowComponent() function instead. ArrowComponent will soon be private and your code will not compile.")
 	UPROPERTY()
 	UArrowComponent* ArrowComponent;
 #endif
@@ -29,7 +30,7 @@ public:
 public:
 #if WITH_EDITORONLY_DATA
 	/** Returns ArrowComponent subobject **/
-	FORCEINLINE UArrowComponent* GetArrowComponent() const { return ArrowComponent; }
+	ENGINE_API UArrowComponent* GetArrowComponent() const;
 #endif
 };
 
diff --git a/Engine/Source/Runtime/Engine/Classes/Engine/DocumentationActor.h b/Engine/Source/Runtime/Engine/Classes/Engine/DocumentationActor.h
--- a/Engine/Source/Runtime/Engine/Classes/Engine/DocumentationActor.h
+++ b/Engine/Source/Runtime/Engine/Classes/Engine/DocumentationActor.h
@@ -45,7 +45,8 @@ class ENGINE_API ADocumentationActor : public AActor
 	UPROPERTY(Category = HelpDocumentation, EditAnywhere, BlueprintReadWrite, AdvancedDisplay)
 	FString DocumentLink; 
 	
-private:
+private_subobject:
+	DEPRECATED_FORGAME(4.6, "Billboard should not be accessed directly, please use GetBillboard() function instead. Billboard will soon be private and your code will not compile.")
 	UPROPERTY(Category = Sprite, VisibleAnywhere, BlueprintReadOnly, meta = (AllowPrivateAccess = "true"))
 	class UMaterialBillboardComponent* Billboard;
 #endif
@@ -65,7 +66,7 @@ private:
 public:
 #if WITH_EDITORONLY_DATA
 	/** Returns Billboard subobject **/
-	FORCEINLINE class UMaterialBillboardComponent* GetBillboard() const { return Billboard; }
+	class UMaterialBillboardComponent* GetBillboard() const;
 #endif
 };
 
diff --git a/Engine/Source/Runtime/Engine/Classes/Engine/ExponentialHeightFog.h b/Engine/Source/Runtime/Engine/Classes/Engine/ExponentialHeightFog.h
--- a/Engine/Source/Runtime/Engine/Classes/Engine/ExponentialHeightFog.h
+++ b/Engine/Source/Runtime/Engine/Classes/Engine/ExponentialHeightFog.h
@@ -9,8 +9,9 @@ class AExponentialHeightFog : public AInfo
 {
 	GENERATED_UCLASS_BODY()
 
-private:
+private_subobject:
 	/** @todo document */
+	DEPRECATED_FORGAME(4.6, "Component should not be accessed directly, please use GetComponent() function instead. Component will soon be private and your code will not compile.")
 	UPROPERTY(Category = ExponentialHeightFog, VisibleAnywhere, BlueprintReadOnly, meta = (AllowPrivateAccess = "true"))
 	class UExponentialHeightFogComponent* Component;
 public:
@@ -29,7 +30,7 @@ public:
 	//End AActor Interface
 
 	/** Returns Component subobject **/
-	FORCEINLINE class UExponentialHeightFogComponent* GetComponent() const { return Component; }
+	ENGINE_API class UExponentialHeightFogComponent* GetComponent() const;
 };
 
 
diff --git a/Engine/Source/Runtime/Engine/Classes/Engine/Light.h b/Engine/Source/Runtime/Engine/Classes/Engine/Light.h
--- a/Engine/Source/Runtime/Engine/Classes/Engine/Light.h
+++ b/Engine/Source/Runtime/Engine/Classes/Engine/Light.h
@@ -9,8 +9,9 @@ class ENGINE_API ALight : public AActor
 {
 	GENERATED_UCLASS_BODY()
 
-private:
+private_subobject:
 	/** @todo document */
+	DEPRECATED_FORGAME(4.6, "LightComponent should not be accessed directly, please use GetLightComponent() function instead. LightComponent will soon be private and your code will not compile.")
 	UPROPERTY(Category = Light, VisibleAnywhere, BlueprintReadOnly, meta = (ExposeFunctionCategories = "Light,Rendering,Rendering|Components|Light", AllowPrivateAccess = "true"))
 	class ULightComponent* LightComponent;
 public:
@@ -63,10 +64,7 @@ public:
 	/**
 	 * Return whether the light supports being toggled off and on on-the-fly.
 	 */
-	bool IsToggleable() const
-	{
-		return !LightComponent->HasStaticLighting();
-	}
+	bool IsToggleable() const;
 
 	// Begin AActor interface.
 	void Destroyed();
@@ -74,7 +72,7 @@ public:
 	// End AActor interface.
 
 	/** Returns LightComponent subobject **/
-	FORCEINLINE class ULightComponent* GetLightComponent() const { return LightComponent; }
+	class ULightComponent* GetLightComponent() const;
 };
 
 
diff --git a/Engine/Source/Runtime/Engine/Classes/Engine/NavigationObjectBase.h b/Engine/Source/Runtime/Engine/Classes/Engine/NavigationObjectBase.h
--- a/Engine/Source/Runtime/Engine/Classes/Engine/NavigationObjectBase.h
+++ b/Engine/Source/Runtime/Engine/Classes/Engine/NavigationObjectBase.h
@@ -9,15 +9,18 @@ class ENGINE_API ANavigationObjectBase : public AActor, public INavAgentInterfac
 {
 	GENERATED_UCLASS_BODY()
 
-private:
+private_subobject:
+	DEPRECATED_FORGAME(4.6, "CapsuleComponent should not be accessed directly, please use GetCapsuleComponent() function instead. CapsuleComponent will soon be private and your code will not compile.")
 	UPROPERTY()
 	class UCapsuleComponent* CapsuleComponent;
 
 	/** Normal editor sprite. */
+	DEPRECATED_FORGAME(4.6, "GoodSprite should not be accessed directly, please use GetGoodSprite() function instead. GoodSprite will soon be private and your code will not compile.")
 	UPROPERTY()
 	class UBillboardComponent* GoodSprite;
 
 	/** Used to draw bad collision intersection in editor. */
+	DEPRECATED_FORGAME(4.6, "BadSprite should not be accessed directly, please use GetBadSprite() function instead. BadSprite will soon be private and your code will not compile.")
 	UPROPERTY()
 	class UBillboardComponent* BadSprite;
 public:
@@ -58,11 +61,11 @@ public:
 
 public:
 	/** Returns CapsuleComponent subobject **/
-	FORCEINLINE class UCapsuleComponent* GetCapsuleComponent() const { return CapsuleComponent; }
+	class UCapsuleComponent* GetCapsuleComponent() const;
 	/** Returns GoodSprite subobject **/
-	FORCEINLINE class UBillboardComponent* GetGoodSprite() const { return GoodSprite; }
+	class UBillboardComponent* GetGoodSprite() const;
 	/** Returns BadSprite subobject **/
-	FORCEINLINE class UBillboardComponent* GetBadSprite() const { return BadSprite; }
+	class UBillboardComponent* GetBadSprite() const;
 };
 
 
diff --git a/Engine/Source/Runtime/Engine/Classes/Engine/NiagaraActor.h b/Engine/Source/Runtime/Engine/Classes/Engine/NiagaraActor.h
--- a/Engine/Source/Runtime/Engine/Classes/Engine/NiagaraActor.h
+++ b/Engine/Source/Runtime/Engine/Classes/Engine/NiagaraActor.h
@@ -8,17 +8,20 @@ class ENGINE_API ANiagaraActor : public AActor
 {
 	GENERATED_UCLASS_BODY()
 
-private:
+private_subobject:
 	/** Pointer to effect component */
+	DEPRECATED_FORGAME(4.6, "NiagaraComponent should not be accessed directly, please use GetNiagaraComponent() function instead. NiagaraComponent will soon be private and your code will not compile.")
 	UPROPERTY(VisibleAnywhere, Category=NiagaraActor)
 	class UNiagaraComponent* NiagaraComponent;
 
 #if WITH_EDITORONLY_DATA
 	// Reference to sprite visualization component
+	DEPRECATED_FORGAME(4.6, "SpriteComponent should not be accessed directly, please use GetSpriteComponent() function instead. SpriteComponent will soon be private and your code will not compile.")
 	UPROPERTY()
 	class UBillboardComponent* SpriteComponent;
 
 	// Reference to arrow visualization component
+	DEPRECATED_FORGAME(4.6, "ArrowComponent should not be accessed directly, please use GetArrowComponent() function instead. ArrowComponent will soon be private and your code will not compile.")
 	UPROPERTY()
 	class UArrowComponent* ArrowComponent;
 
@@ -26,11 +29,11 @@ private:
 
 public:
 	/** Returns NiagaraComponent subobject **/
-	FORCEINLINE class UNiagaraComponent* GetNiagaraComponent() const { return NiagaraComponent; }
+	class UNiagaraComponent* GetNiagaraComponent() const;
 #if WITH_EDITORONLY_DATA
 	/** Returns SpriteComponent subobject **/
-	FORCEINLINE class UBillboardComponent* GetSpriteComponent() const { return SpriteComponent; }
+	class UBillboardComponent* GetSpriteComponent() const;
 	/** Returns ArrowComponent subobject **/
-	FORCEINLINE class UArrowComponent* GetArrowComponent() const { return ArrowComponent; }
+	class UArrowComponent* GetArrowComponent() const;
 #endif
 };
diff --git a/Engine/Source/Runtime/Engine/Classes/Engine/Note.h b/Engine/Source/Runtime/Engine/Classes/Engine/Note.h
--- a/Engine/Source/Runtime/Engine/Classes/Engine/Note.h
+++ b/Engine/Source/Runtime/Engine/Classes/Engine/Note.h
@@ -18,11 +18,13 @@ class ANote : public AActor
 	FString Text;
 
 	// Reference to sprite visualization component
-private:
+private_subobject:
+	DEPRECATED_FORGAME(4.6, "SpriteComponent should not be accessed directly, please use GetSpriteComponent() function instead. SpriteComponent will soon be private and your code will not compile.")
 	UPROPERTY()
 	class UBillboardComponent* SpriteComponent;
 
 	// Reference to arrow visualization component
+	DEPRECATED_FORGAME(4.6, "ArrowComponent should not be accessed directly, please use GetArrowComponent() function instead. ArrowComponent will soon be private and your code will not compile.")
 	UPROPERTY()
 	class UArrowComponent* ArrowComponent;
 public:
@@ -37,9 +39,9 @@ public:
 
 #if WITH_EDITORONLY_DATA
 	/** Returns SpriteComponent subobject **/
-	FORCEINLINE class UBillboardComponent* GetSpriteComponent() const { return SpriteComponent; }
+	ENGINE_API class UBillboardComponent* GetSpriteComponent() const;
 	/** Returns ArrowComponent subobject **/
-	FORCEINLINE class UArrowComponent* GetArrowComponent() const { return ArrowComponent; }
+	ENGINE_API class UArrowComponent* GetArrowComponent() const;
 #endif
 };
 
diff --git a/Engine/Source/Runtime/Engine/Classes/Engine/ReflectionCapture.h b/Engine/Source/Runtime/Engine/Classes/Engine/ReflectionCapture.h
--- a/Engine/Source/Runtime/Engine/Classes/Engine/ReflectionCapture.h
+++ b/Engine/Source/Runtime/Engine/Classes/Engine/ReflectionCapture.h
@@ -8,12 +8,14 @@ class AReflectionCapture : public AActor
 {
 	GENERATED_UCLASS_BODY()
 
-private:
+private_subobject:
 	/** Reflection capture component. */
+	DEPRECATED_FORGAME(4.6, "CaptureComponent should not be accessed directly, please use GetCaptureComponent() function instead. CaptureComponent will soon be private and your code will not compile.")
 	UPROPERTY(Category = DecalActor, VisibleAnywhere, BlueprintReadOnly, meta = (AllowPrivateAccess = "true"))
 	class UReflectionCaptureComponent* CaptureComponent;
 
 #if WITH_EDITORONLY_DATA
+	DEPRECATED_FORGAME(4.6, "SpriteComponent should not be accessed directly, please use GetSpriteComponent() function instead. SpriteComponent will soon be private and your code will not compile.")
 	UPROPERTY()
 	UBillboardComponent* SpriteComponent;
 #endif // WITH_EDITORONLY_DATA
@@ -27,10 +29,10 @@ public:
 #endif // WITH_EDITOR
 
 	/** Returns CaptureComponent subobject **/
-	FORCEINLINE class UReflectionCaptureComponent* GetCaptureComponent() const { return CaptureComponent; }
+	ENGINE_API class UReflectionCaptureComponent* GetCaptureComponent() const;
 #if WITH_EDITORONLY_DATA
 	/** Returns SpriteComponent subobject **/
-	FORCEINLINE UBillboardComponent* GetSpriteComponent() const { return SpriteComponent; }
+	ENGINE_API UBillboardComponent* GetSpriteComponent() const;
 #endif
 };
 
diff --git a/Engine/Source/Runtime/Engine/Classes/Engine/SceneCapture.h b/Engine/Source/Runtime/Engine/Classes/Engine/SceneCapture.h
--- a/Engine/Source/Runtime/Engine/Classes/Engine/SceneCapture.h
+++ b/Engine/Source/Runtime/Engine/Classes/Engine/SceneCapture.h
@@ -13,14 +13,15 @@ class ASceneCapture : public AActor
 {
 	GENERATED_UCLASS_BODY()
 
-private:
+private_subobject:
 	/** To display the 3d camera in the editor. */
+	DEPRECATED_FORGAME(4.6, "MeshComp should not be accessed directly, please use GetMeshComp() function instead. MeshComp will soon be private and your code will not compile.")
 	UPROPERTY()
 	class UStaticMeshComponent* MeshComp;
 
 public:
 	/** Returns MeshComp subobject **/
-	FORCEINLINE class UStaticMeshComponent* GetMeshComp() const { return MeshComp; }
+	ENGINE_API class UStaticMeshComponent* GetMeshComp() const;
 };
 
 
diff --git a/Engine/Source/Runtime/Engine/Classes/Engine/SceneCapture2D.h b/Engine/Source/Runtime/Engine/Classes/Engine/SceneCapture2D.h
--- a/Engine/Source/Runtime/Engine/Classes/Engine/SceneCapture2D.h
+++ b/Engine/Source/Runtime/Engine/Classes/Engine/SceneCapture2D.h
@@ -13,11 +13,13 @@ class ASceneCapture2D : public ASceneCapture
 {
 	GENERATED_UCLASS_BODY()
 
-private:
+private_subobject:
 	/** Scene capture component. */
+	DEPRECATED_FORGAME(4.6, "CaptureComponent2D should not be accessed directly, please use GetCaptureComponent2D() function instead. CaptureComponent2D will soon be private and your code will not compile.")
 	UPROPERTY(Category = DecalActor, VisibleAnywhere, BlueprintReadOnly, meta=(AllowPrivateAccess = "true"))
 	class USceneCaptureComponent2D* CaptureComponent2D;
 	/** To allow drawing the camera frustum in the editor. */
+	DEPRECATED_FORGAME(4.6, "DrawFrustum should not be accessed directly, please use GetDrawFrustum() function instead. DrawFrustum will soon be private and your code will not compile.")
 	UPROPERTY()
 	class UDrawFrustumComponent* DrawFrustum;
 
@@ -34,9 +36,9 @@ public:
 	void OnInterpToggle(bool bEnable);
 
 	/** Returns CaptureComponent2D subobject **/
-	FORCEINLINE class USceneCaptureComponent2D* GetCaptureComponent2D() const { return CaptureComponent2D; }
+	ENGINE_API class USceneCaptureComponent2D* GetCaptureComponent2D() const;
 	/** Returns DrawFrustum subobject **/
-	FORCEINLINE class UDrawFrustumComponent* GetDrawFrustum() const { return DrawFrustum; }
+	ENGINE_API class UDrawFrustumComponent* GetDrawFrustum() const;
 };
 
 
diff --git a/Engine/Source/Runtime/Engine/Classes/Engine/SceneCaptureCube.h b/Engine/Source/Runtime/Engine/Classes/Engine/SceneCaptureCube.h
--- a/Engine/Source/Runtime/Engine/Classes/Engine/SceneCaptureCube.h
+++ b/Engine/Source/Runtime/Engine/Classes/Engine/SceneCaptureCube.h
@@ -13,12 +13,14 @@ class ASceneCaptureCube : public ASceneCapture
 {
 	GENERATED_UCLASS_BODY()
 
-private:
+private_subobject:
 	/** Scene capture component. */
+	DEPRECATED_FORGAME(4.6, "CaptureComponentCube should not be accessed directly, please use GetCaptureComponentCube() function instead. CaptureComponentCube will soon be private and your code will not compile.")
 	UPROPERTY(Category = DecalActor, VisibleAnywhere, BlueprintReadOnly, meta=(AllowPrivateAccess = "true"))
 	class USceneCaptureComponentCube* CaptureComponentCube;
 
 	/** To allow drawing the camera frustum in the editor. */
+	DEPRECATED_FORGAME(4.6, "DrawFrustum should not be accessed directly, please use GetDrawFrustum() function instead. DrawFrustum will soon be private and your code will not compile.")
 	UPROPERTY()
 	class UDrawFrustumComponent* DrawFrustum;
 
@@ -38,9 +40,9 @@ public:
 	void OnInterpToggle(bool bEnable);
 
 	/** Returns CaptureComponentCube subobject **/
-	FORCEINLINE class USceneCaptureComponentCube* GetCaptureComponentCube() const { return CaptureComponentCube; }
+	ENGINE_API class USceneCaptureComponentCube* GetCaptureComponentCube() const;
 	/** Returns DrawFrustum subobject **/
-	FORCEINLINE class UDrawFrustumComponent* GetDrawFrustum() const { return DrawFrustum; }
+	ENGINE_API class UDrawFrustumComponent* GetDrawFrustum() const;
 };
 
 
diff --git a/Engine/Source/Runtime/Engine/Classes/Engine/SkyLight.h b/Engine/Source/Runtime/Engine/Classes/Engine/SkyLight.h
--- a/Engine/Source/Runtime/Engine/Classes/Engine/SkyLight.h
+++ b/Engine/Source/Runtime/Engine/Classes/Engine/SkyLight.h
@@ -9,8 +9,9 @@ class ENGINE_API ASkyLight : public AInfo
 {
 	GENERATED_UCLASS_BODY()
 
-private:
+private_subobject:
 	/** @todo document */
+	DEPRECATED_FORGAME(4.6, "LightComponent should not be accessed directly, please use GetLightComponent() function instead. LightComponent will soon be private and your code will not compile.")
 	UPROPERTY(Category = Light, VisibleAnywhere, BlueprintReadOnly, meta = (ExposeFunctionCategories = "Light,Rendering,Rendering|Components|SkyLight", AllowPrivateAccess = "true"))
 	class USkyLightComponent* LightComponent;
 public:
@@ -24,7 +25,7 @@ public:
 	virtual void OnRep_bEnabled();
 
 	/** Returns LightComponent subobject **/
-	FORCEINLINE class USkyLightComponent* GetLightComponent() const { return LightComponent; }
+	class USkyLightComponent* GetLightComponent() const;
 };
 
 
diff --git a/Engine/Source/Runtime/Engine/Classes/Engine/SphereReflectionCapture.h b/Engine/Source/Runtime/Engine/Classes/Engine/SphereReflectionCapture.h
--- a/Engine/Source/Runtime/Engine/Classes/Engine/SphereReflectionCapture.h
+++ b/Engine/Source/Runtime/Engine/Classes/Engine/SphereReflectionCapture.h
@@ -14,11 +14,14 @@ class ASphereReflectionCapture : public AReflectionCapture
 {
 	GENERATED_UCLASS_BODY()
 
-private:
+private_subobject:
 	/** Sphere component used to visualize the capture radius */
+	DEPRECATED_FORGAME(4.6, "DrawCaptureRadius should not be accessed directly, please use GetDrawCaptureRadius() function instead. DrawCaptureRadius will soon be private and your code will not compile.")
 	UPROPERTY()
 	UDrawSphereComponent* DrawCaptureRadius;
 
+public:
+
 #if WITH_EDITOR
 	// Begin AActor interface.
 	virtual void EditorApplyScale(const FVector& DeltaScale, const FVector* PivotLocation, bool bAltDown, bool bShiftDown, bool bCtrlDown) override;
@@ -26,7 +29,7 @@ private:
 #endif
 
 	/** Returns DrawCaptureRadius subobject **/
-	FORCEINLINE UDrawSphereComponent* GetDrawCaptureRadius() const { return DrawCaptureRadius; }
+	ENGINE_API UDrawSphereComponent* GetDrawCaptureRadius() const;
 };
 
 
diff --git a/Engine/Source/Runtime/Engine/Classes/Engine/SpotLight.h b/Engine/Source/Runtime/Engine/Classes/Engine/SpotLight.h
--- a/Engine/Source/Runtime/Engine/Classes/Engine/SpotLight.h
+++ b/Engine/Source/Runtime/Engine/Classes/Engine/SpotLight.h
@@ -15,7 +15,8 @@ class ASpotLight : public ALight
 
 #if WITH_EDITORONLY_DATA
 	// Reference to editor arrow component visualization 
-private:
+private_subobject:
+	DEPRECATED_FORGAME(4.6, "ArrowComponent should not be accessed directly, please use GetArrowComponent() function instead. ArrowComponent will soon be private and your code will not compile.")
 	UPROPERTY()
 	class UArrowComponent* ArrowComponent;
 public:
@@ -50,7 +51,7 @@ public:
 public:
 #if WITH_EDITORONLY_DATA
 	/** Returns ArrowComponent subobject **/
-	FORCEINLINE class UArrowComponent* GetArrowComponent() const { return ArrowComponent; }
+	ENGINE_API class UArrowComponent* GetArrowComponent() const;
 #endif
 };
 
diff --git a/Engine/Source/Runtime/Engine/Classes/Engine/StaticMeshActor.h b/Engine/Source/Runtime/Engine/Classes/Engine/StaticMeshActor.h
--- a/Engine/Source/Runtime/Engine/Classes/Engine/StaticMeshActor.h
+++ b/Engine/Source/Runtime/Engine/Classes/Engine/StaticMeshActor.h
@@ -16,7 +16,8 @@ class ENGINE_API AStaticMeshActor : public AActor
 {
 	GENERATED_UCLASS_BODY()
 
-private:
+private_subobject:
+	DEPRECATED_FORGAME(4.6, "StaticMeshComponent should not be accessed directly, please use GetStaticMeshComponent() function instead. StaticMeshComponent will soon be private and your code will not compile.")
 	UPROPERTY(Category = StaticMeshActor, VisibleAnywhere, BlueprintReadOnly, meta = (ExposeFunctionCategories = "Mesh,Rendering,Physics,Components|StaticMesh", AllowPrivateAccess = "true"))
 	class UStaticMeshComponent* StaticMeshComponent;
 public:
@@ -48,7 +49,7 @@ protected:
 
 public:
 	/** Returns StaticMeshComponent subobject **/
-	FORCEINLINE class UStaticMeshComponent* GetStaticMeshComponent() const { return StaticMeshComponent; }
+	class UStaticMeshComponent* GetStaticMeshComponent() const;
 };
 
 
diff --git a/Engine/Source/Runtime/Engine/Classes/Engine/TargetPoint.h b/Engine/Source/Runtime/Engine/Classes/Engine/TargetPoint.h
--- a/Engine/Source/Runtime/Engine/Classes/Engine/TargetPoint.h
+++ b/Engine/Source/Runtime/Engine/Classes/Engine/TargetPoint.h
@@ -13,18 +13,20 @@ class ATargetPoint : public AActor
 	GENERATED_UCLASS_BODY()
 
 #if WITH_EDITORONLY_DATA
-private:
+private_subobject:
+	DEPRECATED_FORGAME(4.6, "SpriteComponent should not be accessed directly, please use GetSpriteComponent() function instead. SpriteComponent will soon be private and your code will not compile.")
 	UPROPERTY(VisibleAnywhere, BlueprintReadOnly, Category = Display, meta = (AllowPrivateAccess = "true"))
 	class UBillboardComponent* SpriteComponent;
 
+	DEPRECATED_FORGAME(4.6, "ArrowComponent should not be accessed directly, please use GetArrowComponent() function instead. ArrowComponent will soon be private and your code will not compile.")
 	UPROPERTY()
 	class UArrowComponent* ArrowComponent;
 
 public:
 	/** Returns SpriteComponent subobject **/
-	FORCEINLINE class UBillboardComponent* GetSpriteComponent() const { return SpriteComponent; }
+	ENGINE_API class UBillboardComponent* GetSpriteComponent() const;
 	/** Returns ArrowComponent subobject **/
-	FORCEINLINE class UArrowComponent* GetArrowComponent() const { return ArrowComponent; }
+	ENGINE_API class UArrowComponent* GetArrowComponent() const;
 #endif
 };
 
diff --git a/Engine/Source/Runtime/Engine/Classes/Engine/TextRenderActor.h b/Engine/Source/Runtime/Engine/Classes/Engine/TextRenderActor.h
--- a/Engine/Source/Runtime/Engine/Classes/Engine/TextRenderActor.h
+++ b/Engine/Source/Runtime/Engine/Classes/Engine/TextRenderActor.h
@@ -15,23 +15,25 @@ class ATextRenderActor : public AActor
 
 	friend class UActorFactoryTextRender;
 
-private:
+private_subobject:
 	/** Component to render a text in 3d with a font */
+	DEPRECATED_FORGAME(4.6, "TextRender should not be accessed directly, please use GetTextRender() function instead. TextRender will soon be private and your code will not compile.")
 	UPROPERTY(Category = TextRenderActor, VisibleAnywhere, BlueprintReadOnly, meta = (ExposeFunctionCategories = "Rendering|Components|TextRender", AllowPrivateAccess = "true"))
 	class UTextRenderComponent* TextRender;
 
 #if WITH_EDITORONLY_DATA
 	// Reference to the billboard component
+	DEPRECATED_FORGAME(4.6, "SpriteComponent should not be accessed directly, please use GetSpriteComponent() function instead. SpriteComponent will soon be private and your code will not compile.")
 	UPROPERTY()
 	UBillboardComponent* SpriteComponent;
 #endif
 
 public:
 	/** Returns TextRender subobject **/
-	FORCEINLINE class UTextRenderComponent* GetTextRender() const { return TextRender; }
+	ENGINE_API class UTextRenderComponent* GetTextRender() const;
 #if WITH_EDITORONLY_DATA
 	/** Returns SpriteComponent subobject **/
-	FORCEINLINE UBillboardComponent* GetSpriteComponent() const { return SpriteComponent; }
+	ENGINE_API UBillboardComponent* GetSpriteComponent() const;
 #endif
 };
 
diff --git a/Engine/Source/Runtime/Engine/Classes/Engine/TriggerBase.h b/Engine/Source/Runtime/Engine/Classes/Engine/TriggerBase.h
--- a/Engine/Source/Runtime/Engine/Classes/Engine/TriggerBase.h
+++ b/Engine/Source/Runtime/Engine/Classes/Engine/TriggerBase.h
@@ -13,20 +13,22 @@ class ATriggerBase : public AActor
 {
 	GENERATED_UCLASS_BODY()
 
-private:
+private_subobject:
 	/** Shape component used for collision */
+	DEPRECATED_FORGAME(4.6, "CollisionComponent should not be accessed directly, please use GetCollisionComponent() function instead. CollisionComponent will soon be private and your code will not compile.")
 	UPROPERTY(Category = TriggerBase, VisibleAnywhere, BlueprintReadOnly, meta = (AllowPrivateAccess = "true"))
 	UShapeComponent* CollisionComponent;
 
 	/** Billboard used to see the trigger in the editor */
+	DEPRECATED_FORGAME(4.6, "SpriteComponent should not be accessed directly, please use GetSpriteComponent() function instead. SpriteComponent will soon be private and your code will not compile.")
 	UPROPERTY()
 	UBillboardComponent* SpriteComponent;
 
 public:
 	/** Returns CollisionComponent subobject **/
-	FORCEINLINE UShapeComponent* GetCollisionComponent() const { return CollisionComponent; }
+	ENGINE_API UShapeComponent* GetCollisionComponent() const;
 	/** Returns SpriteComponent subobject **/
-	FORCEINLINE UBillboardComponent* GetSpriteComponent() const { return SpriteComponent; }
+	ENGINE_API UBillboardComponent* GetSpriteComponent() const;
 };
 
 
diff --git a/Engine/Source/Runtime/Engine/Classes/Engine/WindDirectionalSource.h b/Engine/Source/Runtime/Engine/Classes/Engine/WindDirectionalSource.h
--- a/Engine/Source/Runtime/Engine/Classes/Engine/WindDirectionalSource.h
+++ b/Engine/Source/Runtime/Engine/Classes/Engine/WindDirectionalSource.h
@@ -8,22 +8,24 @@ class ENGINE_API AWindDirectionalSource : public AInfo
 {
 	GENERATED_UCLASS_BODY()
 
-private:
+private_subobject:
+	DEPRECATED_FORGAME(4.6, "Component should not be accessed directly, please use GetComponent() function instead. Component will soon be private and your code will not compile.")
 	UPROPERTY(Category = WindDirectionalSource, VisibleAnywhere, BlueprintReadOnly, meta = (AllowPrivateAccess = "true"))
 	class UWindDirectionalSourceComponent* Component;
 
 #if WITH_EDITORONLY_DATA
+	DEPRECATED_FORGAME(4.6, "ArrowComponent should not be accessed directly, please use GetArrowComponent() function instead. ArrowComponent will soon be private and your code will not compile.")
 	UPROPERTY()
 	class UArrowComponent* ArrowComponent;
 #endif
 
 public:
 	/** Returns Component subobject **/
-	FORCEINLINE class UWindDirectionalSourceComponent* GetComponent() const { return Component; }
+	class UWindDirectionalSourceComponent* GetComponent() const;
 
 #if WITH_EDITORONLY_DATA
 	/** Returns ArrowComponent subobject **/
-	FORCEINLINE class UArrowComponent* GetArrowComponent() const { return ArrowComponent; }
+	class UArrowComponent* GetArrowComponent() const;
 #endif
 };
 
diff --git a/Engine/Source/Runtime/Engine/Classes/GameFramework/Character.h b/Engine/Source/Runtime/Engine/Classes/GameFramework/Character.h
--- a/Engine/Source/Runtime/Engine/Classes/GameFramework/Character.h
+++ b/Engine/Source/Runtime/Engine/Classes/GameFramework/Character.h
@@ -207,21 +207,25 @@ public:
 	 */
 	ACharacter(const FObjectInitializer& ObjectInitializer);
 
-private:
+private_subobject:
 	/** The main skeletal mesh associated with this Character (optional sub-object). */
+	DEPRECATED_FORGAME(4.6, "Mesh should not be accessed directly, please use GetMesh() function instead. Mesh will soon be private and your code will not compile.")
 	UPROPERTY(Category = Character, VisibleAnywhere, BlueprintReadOnly, meta = (AllowPrivateAccess = "true"))
 	class USkeletalMeshComponent* Mesh;
 
 #if WITH_EDITORONLY_DATA
+	DEPRECATED_FORGAME(4.6, "ArrowComponent should not be accessed directly, please use GetArrowComponent() function instead. ArrowComponent will soon be private and your code will not compile.")
 	UPROPERTY()
 	class UArrowComponent* ArrowComponent;
 #endif
 
 	/** Movement component used for movement logic in various movement modes (walking, falling, etc), containing relevant settings and functions to control movement. */
+	DEPRECATED_FORGAME(4.6, "CharacterMovement should not be accessed directly, please use GetCharacterMovement() function instead. CharacterMovement will soon be private and your code will not compile.")
 	UPROPERTY(Category = Character, VisibleAnywhere, BlueprintReadOnly, meta = (AllowPrivateAccess = "true"))
 	class UCharacterMovementComponent* CharacterMovement;
 
 	/** The CapsuleComponent being used for movement collision (by CharacterMovement). Always treated as being vertically aligned in simple collision check functions. */
+	DEPRECATED_FORGAME(4.6, "CapsuleComponent should not be accessed directly, please use GetCapsuleComponent() function instead. CapsuleComponent will soon be private and your code will not compile.")
 	UPROPERTY(Category = Character, VisibleAnywhere, BlueprintReadOnly, meta = (AllowPrivateAccess = "true"))
 	class UCapsuleComponent* CapsuleComponent;
 public:
@@ -684,13 +688,13 @@ public:
 
 public:
 	/** Returns Mesh subobject **/
-	FORCEINLINE class USkeletalMeshComponent* GetMesh() const { return Mesh; }
+	class USkeletalMeshComponent* GetMesh() const;
 #if WITH_EDITORONLY_DATA
 	/** Returns ArrowComponent subobject **/
-	FORCEINLINE class UArrowComponent* GetArrowComponent() const { return ArrowComponent; }
+	class UArrowComponent* GetArrowComponent() const;
 #endif
 	/** Returns CharacterMovement subobject **/
-	FORCEINLINE class UCharacterMovementComponent* GetCharacterMovement() const { return CharacterMovement; }
+	class UCharacterMovementComponent* GetCharacterMovement() const;
 	/** Returns CapsuleComponent subobject **/
-	FORCEINLINE class UCapsuleComponent* GetCapsuleComponent() const { return CapsuleComponent; }
+	class UCapsuleComponent* GetCapsuleComponent() const;
 };
diff --git a/Engine/Source/Runtime/Engine/Classes/GameFramework/Controller.h b/Engine/Source/Runtime/Engine/Classes/GameFramework/Controller.h
--- a/Engine/Source/Runtime/Engine/Classes/GameFramework/Controller.h
+++ b/Engine/Source/Runtime/Engine/Classes/GameFramework/Controller.h
@@ -54,8 +54,9 @@ public:
 	UPROPERTY(replicatedUsing=OnRep_PlayerState, BlueprintReadOnly, Category="Controller")
 	class APlayerState* PlayerState;
 
-private:
+private_subobject:
 	/** Component to give controllers a transform and enable attachment if desired. */
+	DEPRECATED_FORGAME(4.6, "TransformComponent should not be accessed directly, please use GetTransformComponent() function instead. TransformComponent will soon be private and your code will not compile.")
 	UPROPERTY()
 	class USceneComponent* TransformComponent;
 
@@ -307,5 +308,5 @@ private:
 
 protected:
 	/** Returns TransformComponent subobject **/
-	FORCEINLINE class USceneComponent* GetTransformComponent() const { return TransformComponent; }
+	class USceneComponent* GetTransformComponent() const;
 };
diff --git a/Engine/Source/Runtime/Engine/Classes/GameFramework/DefaultPawn.h b/Engine/Source/Runtime/Engine/Classes/GameFramework/DefaultPawn.h
--- a/Engine/Source/Runtime/Engine/Classes/GameFramework/DefaultPawn.h
+++ b/Engine/Source/Runtime/Engine/Classes/GameFramework/DefaultPawn.h
@@ -96,8 +96,9 @@ public:
 	/** Name of the CollisionComponent. */
 	static FName CollisionComponentName;
 
-private:
+private_subobject:
 	/** DefaultPawn collision component */
+	DEPRECATED_FORGAME(4.6, "CollisionComponent should not be accessed directly, please use GetCollisionComponent() function instead. CollisionComponent will soon be private and your code will not compile.")
 	UPROPERTY(Category = Pawn, VisibleAnywhere, BlueprintReadOnly, meta = (AllowPrivateAccess = "true"))
 	USphereComponent* CollisionComponent;
 public:
@@ -105,8 +106,9 @@ public:
 	/** Name of the MeshComponent. Use this name if you want to prevent creation of the component (with ObjectInitializer.DoNotCreateDefaultSubobject). */
 	static FName MeshComponentName;
 
-private:
+private_subobject:
 	/** The mesh associated with this Pawn. */
+	DEPRECATED_FORGAME(4.6, "MeshComponent should not be accessed directly, please use GetMeshComponent() function instead. MeshComponent will soon be private and your code will not compile.")
 	UPROPERTY(Category = Pawn, VisibleAnywhere, BlueprintReadOnly, meta = (AllowPrivateAccess = "true"))
 	UStaticMeshComponent* MeshComponent;
 public:
@@ -116,8 +118,8 @@ public:
 	uint32 bAddDefaultMovementBindings:1;
 
 	/** Returns CollisionComponent subobject **/
-	FORCEINLINE USphereComponent* GetCollisionComponent() const { return CollisionComponent; }
+	USphereComponent* GetCollisionComponent() const;
 	/** Returns MeshComponent subobject **/
-	FORCEINLINE UStaticMeshComponent* GetMeshComponent() const { return MeshComponent; }
+	UStaticMeshComponent* GetMeshComponent() const;
 };
 
diff --git a/Engine/Source/Runtime/Engine/Classes/GameFramework/Info.h b/Engine/Source/Runtime/Engine/Classes/GameFramework/Info.h
--- a/Engine/Source/Runtime/Engine/Classes/GameFramework/Info.h
+++ b/Engine/Source/Runtime/Engine/Classes/GameFramework/Info.h
@@ -15,8 +15,9 @@ class AInfo : public AActor
 	GENERATED_UCLASS_BODY()
 
 #if WITH_EDITORONLY_DATA
-private:
+private_subobject:
 	/** Billboard Component displayed in editor */
+	DEPRECATED_FORGAME(4.6, "SpriteComponent should not be accessed directly, please use GetSpriteComponent() function instead. SpriteComponent will soon be private and your code will not compile.")
 	UPROPERTY()
 	class UBillboardComponent* SpriteComponent;
 public:
@@ -28,7 +29,7 @@ public:
 public:
 #if WITH_EDITORONLY_DATA
 	/** Returns SpriteComponent subobject **/
-	FORCEINLINE class UBillboardComponent* GetSpriteComponent() const { return SpriteComponent; }
+	ENGINE_API class UBillboardComponent* GetSpriteComponent() const;
 #endif
 };
 
diff --git a/Engine/Source/Runtime/Engine/Classes/GameFramework/PlayerStart.h b/Engine/Source/Runtime/Engine/Classes/GameFramework/PlayerStart.h
--- a/Engine/Source/Runtime/Engine/Classes/GameFramework/PlayerStart.h
+++ b/Engine/Source/Runtime/Engine/Classes/GameFramework/PlayerStart.h
@@ -17,7 +17,8 @@ class ENGINE_API APlayerStart : public ANavigationObjectBase
 	FName PlayerStartTag;
 
 #if WITH_EDITORONLY_DATA
-private:
+private_subobject:
+	DEPRECATED_FORGAME(4.6, "ArrowComponent should not be accessed directly, please use GetArrowComponent() function instead. ArrowComponent will soon be private and your code will not compile.")
 	UPROPERTY()
 	class UArrowComponent* ArrowComponent;
 public:
@@ -29,7 +30,7 @@ public:
 
 #if WITH_EDITORONLY_DATA
 	/** Returns ArrowComponent subobject **/
-	FORCEINLINE class UArrowComponent* GetArrowComponent() const { return ArrowComponent; }
+	class UArrowComponent* GetArrowComponent() const;
 #endif
 };
 
diff --git a/Engine/Source/Runtime/Engine/Classes/GameFramework/WheeledVehicle.h b/Engine/Source/Runtime/Engine/Classes/GameFramework/WheeledVehicle.h
--- a/Engine/Source/Runtime/Engine/Classes/GameFramework/WheeledVehicle.h
+++ b/Engine/Source/Runtime/Engine/Classes/GameFramework/WheeledVehicle.h
@@ -11,12 +11,14 @@ class ENGINE_API AWheeledVehicle : public APawn
 {
 	GENERATED_UCLASS_BODY()
 
-private:
+private_subobject:
 	/**  The main skeletal mesh associated with this Vehicle */
+	DEPRECATED_FORGAME(4.6, "Mesh should not be accessed directly, please use GetMesh() function instead. Mesh will soon be private and your code will not compile.")
 	UPROPERTY(Category = Vehicle, VisibleDefaultsOnly, BlueprintReadOnly, meta = (AllowPrivateAccess = "true"))
 	class USkeletalMeshComponent* Mesh;
 
 	/** vehicle simulation component */
+	DEPRECATED_FORGAME(4.6, "VehicleMovement should not be accessed directly, please use GetVehicleMovement() function instead. VehicleMovement will soon be private and your code will not compile.")
 	UPROPERTY(Category = Vehicle, VisibleDefaultsOnly, BlueprintReadOnly, meta = (AllowPrivateAccess = "true"))
 	class UWheeledVehicleMovementComponent* VehicleMovement;
 public:
@@ -28,17 +30,14 @@ public:
 	static FName VehicleMovementComponentName;
 
 	/** Util to get the wheeled vehicle movement component */
-	class UWheeledVehicleMovementComponent* GetVehicleMovementComponent() const 
-	{ 
-		return VehicleMovement; 
-	}
+	class UWheeledVehicleMovementComponent* GetVehicleMovementComponent() const;
 
 	// Begin AActor interface
 	virtual void DisplayDebug(class UCanvas* Canvas, const FDebugDisplayInfo& DebugDisplay, float& YL, float& YPos) override;
 	// End Actor interface
 
 	/** Returns Mesh subobject **/
-	FORCEINLINE class USkeletalMeshComponent* GetMesh() const { return Mesh; }
+	class USkeletalMeshComponent* GetMesh() const;
 	/** Returns VehicleMovement subobject **/
-	FORCEINLINE class UWheeledVehicleMovementComponent* GetVehicleMovement() const { return VehicleMovement; }
+	class UWheeledVehicleMovementComponent* GetVehicleMovement() const;
 };
diff --git a/Engine/Source/Runtime/Engine/Classes/Materials/MaterialInstanceActor.h b/Engine/Source/Runtime/Engine/Classes/Materials/MaterialInstanceActor.h
--- a/Engine/Source/Runtime/Engine/Classes/Materials/MaterialInstanceActor.h
+++ b/Engine/Source/Runtime/Engine/Classes/Materials/MaterialInstanceActor.h
@@ -18,7 +18,8 @@ class AMaterialInstanceActor : public AActor
 	TArray<class AActor*> TargetActors;
 
 #if WITH_EDITORONLY_DATA
-private:
+private_subobject:
+	DEPRECATED_FORGAME(4.6, "SpriteComponent should not be accessed directly, please use GetSpriteComponent() function instead. SpriteComponent will soon be private and your code will not compile.")
 	// Reference to actor sprite
 	UBillboardComponent* SpriteComponent;
 #endif
@@ -34,7 +35,7 @@ public:
 
 #if WITH_EDITORONLY_DATA
 	/** Returns SpriteComponent subobject **/
-	FORCEINLINE UBillboardComponent* GetSpriteComponent() const { return SpriteComponent; }
+	ENGINE_API UBillboardComponent* GetSpriteComponent() const;
 #endif
 };
 
diff --git a/Engine/Source/Runtime/Engine/Classes/Matinee/MatineeActor.h b/Engine/Source/Runtime/Engine/Classes/Matinee/MatineeActor.h
--- a/Engine/Source/Runtime/Engine/Classes/Matinee/MatineeActor.h
+++ b/Engine/Source/Runtime/Engine/Classes/Matinee/MatineeActor.h
@@ -166,7 +166,8 @@ class AMatineeActor : public AActor
 #if WITH_EDITORONLY_DATA
 
 	// Reference to the actor sprite
-private:
+private_subobject:
+	DEPRECATED_FORGAME(4.6, "SpriteComponent should not be accessed directly, please use GetSpriteComponent() function instead. SpriteComponent will soon be private and your code will not compile.")
 	UPROPERTY()
 	class UBillboardComponent* SpriteComponent;
 public:
@@ -467,7 +468,7 @@ public:
 public:
 #if WITH_EDITORONLY_DATA
 	/** Returns SpriteComponent subobject **/
-	FORCEINLINE class UBillboardComponent* GetSpriteComponent() const { return SpriteComponent; }
+	ENGINE_API class UBillboardComponent* GetSpriteComponent() const;
 #endif
 };
 
diff --git a/Engine/Source/Runtime/Engine/Classes/Particles/Emitter.h b/Engine/Source/Runtime/Engine/Classes/Particles/Emitter.h
--- a/Engine/Source/Runtime/Engine/Classes/Particles/Emitter.h
+++ b/Engine/Source/Runtime/Engine/Classes/Particles/Emitter.h
@@ -49,7 +49,8 @@ class ENGINE_API AEmitter : public AActor
 {
 	GENERATED_UCLASS_BODY()
 
-private:
+private_subobject:
+	DEPRECATED_FORGAME(4.6, "ParticleSystemComponent should not be accessed directly, please use GetParticleSystemComponent() function instead. ParticleSystemComponent will soon be private and your code will not compile.")
 	UPROPERTY(Category = Emitter, VisibleAnywhere, BlueprintReadOnly, meta = (ExposeFunctionCategories = "Particles|Beam,Particles|Parameters,Particles,Effects|Components|ParticleSystem,Rendering,Activation,Components|Activation", AllowPrivateAccess = "true"))
 	class UParticleSystemComponent* ParticleSystemComponent;
 public:
@@ -78,10 +79,12 @@ public:
 
 #if WITH_EDITORONLY_DATA
 
-private:
+private_subobject:
+	DEPRECATED_FORGAME(4.6, "SpriteComponent should not be accessed directly, please use GetSpriteComponent() function instead. SpriteComponent will soon be private and your code will not compile.")
 	UPROPERTY()
 	class UBillboardComponent* SpriteComponent;
 
+	DEPRECATED_FORGAME(4.6, "ArrowComponent should not be accessed directly, please use GetArrowComponent() function instead. ArrowComponent will soon be private and your code will not compile.")
 	UPROPERTY()
 	class UArrowComponent* ArrowComponent;
 public:
@@ -143,12 +146,12 @@ public:
 
 public:
 	/** Returns ParticleSystemComponent subobject **/
-	FORCEINLINE class UParticleSystemComponent* GetParticleSystemComponent() { return ParticleSystemComponent; }
+	class UParticleSystemComponent* GetParticleSystemComponent();
 #if WITH_EDITORONLY_DATA
 	/** Returns SpriteComponent subobject **/
-	FORCEINLINE class UBillboardComponent* GetSpriteComponent() const { return SpriteComponent; }
+	class UBillboardComponent* GetSpriteComponent() const;
 	/** Returns ArrowComponent subobject **/
-	FORCEINLINE class UArrowComponent* GetArrowComponent() const { return ArrowComponent; }
+	class UArrowComponent* GetArrowComponent() const;
 #endif
 };
 
diff --git a/Engine/Source/Runtime/Engine/Classes/PhysicsEngine/DestructibleActor.h b/Engine/Source/Runtime/Engine/Classes/PhysicsEngine/DestructibleActor.h
--- a/Engine/Source/Runtime/Engine/Classes/PhysicsEngine/DestructibleActor.h
+++ b/Engine/Source/Runtime/Engine/Classes/PhysicsEngine/DestructibleActor.h
@@ -19,7 +19,8 @@ class ADestructibleActor : public AActor
 	/**
 	 * The component which holds the skinned mesh and physics data for this actor.
 	 */
-private:
+private_subobject:
+	DEPRECATED_FORGAME(4.6, "DestructibleComponent should not be accessed directly, please use GetDestructibleComponent() function instead. DestructibleComponent will soon be private and your code will not compile.")
 	UPROPERTY(VisibleAnywhere, BlueprintReadOnly, Category = Destruction, meta = (ExposeFunctionCategories = "Destruction,Components|Destructible", AllowPrivateAccess = "true"))
 	UDestructibleComponent* DestructibleComponent;
 public:
@@ -40,7 +41,7 @@ public:
 
 public:
 	/** Returns DestructibleComponent subobject **/
-	FORCEINLINE UDestructibleComponent* GetDestructibleComponent() const { return DestructibleComponent; }
+	ENGINE_API UDestructibleComponent* GetDestructibleComponent() const;
 };
 
 
diff --git a/Engine/Source/Runtime/Engine/Classes/PhysicsEngine/PhysicsConstraintActor.h b/Engine/Source/Runtime/Engine/Classes/PhysicsEngine/PhysicsConstraintActor.h
--- a/Engine/Source/Runtime/Engine/Classes/PhysicsEngine/PhysicsConstraintActor.h
+++ b/Engine/Source/Runtime/Engine/Classes/PhysicsEngine/PhysicsConstraintActor.h
@@ -13,7 +13,8 @@ class APhysicsConstraintActor : public ARigidBodyBase
 	GENERATED_UCLASS_BODY()
 
 	// Cached reference to constraint component
-private:
+private_subobject:
+	DEPRECATED_FORGAME(4.6, "ConstraintComp should not be accessed directly, please use GetConstraintComp() function instead. ConstraintComp will soon be private and your code will not compile.")
 	UPROPERTY(Category = ConstraintActor, VisibleAnywhere, BlueprintReadOnly, meta = (ExposeFunctionCategories = "JointDrive,Physics|Components|PhysicsConstraint", AllowPrivateAccess = "true"))
 	class UPhysicsConstraintComponent* ConstraintComp;
 public:
@@ -34,7 +35,7 @@ public:
 
 public:
 	/** Returns ConstraintComp subobject **/
-	FORCEINLINE class UPhysicsConstraintComponent* GetConstraintComp() const { return ConstraintComp; }
+	ENGINE_API class UPhysicsConstraintComponent* GetConstraintComp() const;
 };
 
 
diff --git a/Engine/Source/Runtime/Engine/Classes/PhysicsEngine/PhysicsThruster.h b/Engine/Source/Runtime/Engine/Classes/PhysicsEngine/PhysicsThruster.h
--- a/Engine/Source/Runtime/Engine/Classes/PhysicsEngine/PhysicsThruster.h
+++ b/Engine/Source/Runtime/Engine/Classes/PhysicsEngine/PhysicsThruster.h
@@ -14,28 +14,31 @@ class APhysicsThruster : public ARigidBodyBase
 {
 	GENERATED_UCLASS_BODY()
 
-private:
+private_subobject:
 	/** Thruster component */
+	DEPRECATED_FORGAME(4.6, "ThrusterComponent should not be accessed directly, please use GetThrusterComponent() function instead. ThrusterComponent will soon be private and your code will not compile.")
 	UPROPERTY(Category = Physics, VisibleAnywhere, BlueprintReadOnly, meta = (ExposeFunctionCategories = "Activation,Components|Activation", AllowPrivateAccess = "true"))
 	class UPhysicsThrusterComponent* ThrusterComponent;
 
 #if WITH_EDITORONLY_DATA
 
+	DEPRECATED_FORGAME(4.6, "ArrowComponent should not be accessed directly, please use GetArrowComponent() function instead. ArrowComponent will soon be private and your code will not compile.")
 	UPROPERTY()
 	class UArrowComponent* ArrowComponent;
 
+	DEPRECATED_FORGAME(4.6, "SpriteComponent should not be accessed directly, please use GetSpriteComponent() function instead. SpriteComponent will soon be private and your code will not compile.")
 	UPROPERTY()
 	class UBillboardComponent* SpriteComponent;
 #endif
 
 public:
 	/** Returns ThrusterComponent subobject **/
-	FORCEINLINE class UPhysicsThrusterComponent* GetThrusterComponent() const { return ThrusterComponent; }
+	ENGINE_API class UPhysicsThrusterComponent* GetThrusterComponent() const;
 #if WITH_EDITORONLY_DATA
 	/** Returns ArrowComponent subobject **/
-	FORCEINLINE class UArrowComponent* GetArrowComponent() const { return ArrowComponent; }
+	ENGINE_API class UArrowComponent* GetArrowComponent() const;
 	/** Returns SpriteComponent subobject **/
-	FORCEINLINE class UBillboardComponent* GetSpriteComponent() const { return SpriteComponent; }
+	ENGINE_API class UBillboardComponent* GetSpriteComponent() const;
 #endif
 };
 
diff --git a/Engine/Source/Runtime/Engine/Classes/PhysicsEngine/RadialForceActor.h b/Engine/Source/Runtime/Engine/Classes/PhysicsEngine/RadialForceActor.h
--- a/Engine/Source/Runtime/Engine/Classes/PhysicsEngine/RadialForceActor.h
+++ b/Engine/Source/Runtime/Engine/Classes/PhysicsEngine/RadialForceActor.h
@@ -9,12 +9,14 @@ class ARadialForceActor : public ARigidBodyBase
 {
 	GENERATED_UCLASS_BODY()
 
-private:
+private_subobject:
 	/** Force component */
+	DEPRECATED_FORGAME(4.6, "ForceComponent should not be accessed directly, please use GetForceComponent() function instead. ForceComponent will soon be private and your code will not compile.")
 	UPROPERTY(Category = RadialForceActor, VisibleAnywhere, BlueprintReadOnly, meta = (ExposeFunctionCategories = "Activation,Components|Activation,Physics,Physics|Components|RadialForce", AllowPrivateAccess = "true"))
 	class URadialForceComponent* ForceComponent;
 
 #if WITH_EDITORONLY_DATA
+	DEPRECATED_FORGAME(4.6, "SpriteComponent should not be accessed directly, please use GetSpriteComponent() function instead. SpriteComponent will soon be private and your code will not compile.")
 	UPROPERTY()
 	UBillboardComponent* SpriteComponent;
 #endif
@@ -40,10 +42,10 @@ public:
 
 public:
 	/** Returns ForceComponent subobject **/
-	FORCEINLINE class URadialForceComponent* GetForceComponent() const { return ForceComponent; }
+	ENGINE_API class URadialForceComponent* GetForceComponent() const;
 #if WITH_EDITORONLY_DATA
 	/** Returns SpriteComponent subobject **/
-	FORCEINLINE UBillboardComponent* GetSpriteComponent() const { return SpriteComponent; }
+	ENGINE_API UBillboardComponent* GetSpriteComponent() const;
 #endif
 };
 
diff --git a/Engine/Source/Runtime/Engine/Classes/Sound/AmbientSound.h b/Engine/Source/Runtime/Engine/Classes/Sound/AmbientSound.h
--- a/Engine/Source/Runtime/Engine/Classes/Sound/AmbientSound.h
+++ b/Engine/Source/Runtime/Engine/Classes/Sound/AmbientSound.h
@@ -12,8 +12,9 @@ class AAmbientSound : public AActor
 	UPROPERTY()
 	uint32 bAutoPlay_DEPRECATED:1;
 
-private:
+private_subobject:
 	/** Audio component that handles sound playing */
+	DEPRECATED_FORGAME(4.6, "AudioComponent should not be accessed directly, please use GetAudioComponent() function instead. AudioComponent will soon be private and your code will not compile.")
 	UPROPERTY(Category = Sound, VisibleAnywhere, BlueprintReadOnly, meta = (ExposeFunctionCategories = "Sound,Audio,Audio|Components|Audio", AllowPrivateAccess = "true"))
 	class UAudioComponent* AudioComponent;
 public:
@@ -49,7 +50,7 @@ public:
 
 public:
 	/** Returns AudioComponent subobject **/
-	FORCEINLINE class UAudioComponent* GetAudioComponent() const { return AudioComponent; }
+	ENGINE_API class UAudioComponent* GetAudioComponent() const;
 };
 
 
diff --git a/Engine/Source/Runtime/Engine/Classes/VectorField/VectorFieldVolume.h b/Engine/Source/Runtime/Engine/Classes/VectorField/VectorFieldVolume.h
--- a/Engine/Source/Runtime/Engine/Classes/VectorField/VectorFieldVolume.h
+++ b/Engine/Source/Runtime/Engine/Classes/VectorField/VectorFieldVolume.h
@@ -13,21 +13,23 @@ class AVectorFieldVolume : public AActor
 {
 	GENERATED_UCLASS_BODY()
 
-private:
+private_subobject:
+	DEPRECATED_FORGAME(4.6, "VectorFieldComponent should not be accessed directly, please use GetVectorFieldComponent() function instead. VectorFieldComponent will soon be private and your code will not compile.")
 	UPROPERTY(VisibleAnywhere, BlueprintReadOnly, Category = VectorFieldVolume, meta = (AllowPrivateAccess = "true"))
 	class UVectorFieldComponent* VectorFieldComponent;
 
 #if WITH_EDITORONLY_DATA
+	DEPRECATED_FORGAME(4.6, "SpriteComponent should not be accessed directly, please use GetSpriteComponent() function instead. SpriteComponent will soon be private and your code will not compile.")
 	UPROPERTY()
 	UBillboardComponent* SpriteComponent;
 #endif
 
 public:
 	/** Returns VectorFieldComponent subobject **/
-	FORCEINLINE class UVectorFieldComponent* GetVectorFieldComponent() const { return VectorFieldComponent; }
+	ENGINE_API class UVectorFieldComponent* GetVectorFieldComponent() const;
 #if WITH_EDITORONLY_DATA
 	/** Returns SpriteComponent subobject **/
-	FORCEINLINE UBillboardComponent* GetSpriteComponent() const { return SpriteComponent; }
+	ENGINE_API UBillboardComponent* GetSpriteComponent() const;
 #endif
 };
 
diff --git a/Engine/Source/Runtime/Engine/Private/AI/Navigation/NavLinkProxy.cpp b/Engine/Source/Runtime/Engine/Private/AI/Navigation/NavLinkProxy.cpp
--- a/Engine/Source/Runtime/Engine/Private/AI/Navigation/NavLinkProxy.cpp
+++ b/Engine/Source/Runtime/Engine/Private/AI/Navigation/NavLinkProxy.cpp
@@ -183,3 +183,12 @@ bool ANavLinkProxy::HasMovingAgents() const
 {
 	return SmartLinkComp->HasMovingAgents();
 }
+
+/** Returns SmartLinkComp subobject **/
+UNavLinkCustomComponent* ANavLinkProxy::GetSmartLinkComp() const { return SmartLinkComp; }
+#if WITH_EDITORONLY_DATA
+/** Returns EdRenderComp subobject **/
+UNavLinkRenderingComponent* ANavLinkProxy::GetEdRenderComp() const { return EdRenderComp; }
+/** Returns SpriteComponent subobject **/
+UBillboardComponent* ANavLinkProxy::GetSpriteComponent() const { return SpriteComponent; }
+#endif
diff --git a/Engine/Source/Runtime/Engine/Private/AI/Navigation/NavigationTestingActor.cpp b/Engine/Source/Runtime/Engine/Private/AI/Navigation/NavigationTestingActor.cpp
--- a/Engine/Source/Runtime/Engine/Private/AI/Navigation/NavigationTestingActor.cpp
+++ b/Engine/Source/Runtime/Engine/Private/AI/Navigation/NavigationTestingActor.cpp
@@ -395,3 +395,10 @@ FPathFindingQuery ANavigationTestingActor::BuildPathFindingQuery(const ANavigati
 	return FPathFindingQuery(this, MyNavData, GetNavAgentLocation(), Goal->GetNavAgentLocation(), UNavigationQueryFilter::GetQueryFilter(MyNavData, FilterClass));
 }
 
+
+/** Returns CapsuleComponent subobject **/
+UCapsuleComponent* ANavigationTestingActor::GetCapsuleComponent() const { return CapsuleComponent; }
+#if WITH_EDITORONLY_DATA
+/** Returns EdRenderComp subobject **/
+UNavTestRenderingComponent* ANavigationTestingActor::GetEdRenderComp() const { return EdRenderComp; }
+#endif
diff --git a/Engine/Source/Runtime/Engine/Private/AmbientSound.cpp b/Engine/Source/Runtime/Engine/Private/AmbientSound.cpp
--- a/Engine/Source/Runtime/Engine/Private/AmbientSound.cpp
+++ b/Engine/Source/Runtime/Engine/Private/AmbientSound.cpp
@@ -323,4 +323,7 @@ void AAmbientSound::Stop()
 	}
 }
 
-#undef LOCTEXT_NAMESPACE
\ No newline at end of file
+#undef LOCTEXT_NAMESPACE
+
+/** Returns AudioComponent subobject **/
+UAudioComponent* AAmbientSound::GetAudioComponent() const { return AudioComponent; }
diff --git a/Engine/Source/Runtime/Engine/Private/Atmosphere/Atmosphere.cpp b/Engine/Source/Runtime/Engine/Private/Atmosphere/Atmosphere.cpp
--- a/Engine/Source/Runtime/Engine/Private/Atmosphere/Atmosphere.cpp
+++ b/Engine/Source/Runtime/Engine/Private/Atmosphere/Atmosphere.cpp
@@ -836,3 +836,10 @@ static void AtmosphereRenderSinkFunction()
 }
 
 FAutoConsoleVariableSink CVarAtmosphereRenderSink(FConsoleCommandDelegate::CreateStatic(&AtmosphereRenderSinkFunction));
+
+/** Returns AtmosphericFogComponent subobject **/
+UAtmosphericFogComponent* AAtmosphericFog::GetAtmosphericFogComponent() { return AtmosphericFogComponent; }
+#if WITH_EDITORONLY_DATA
+/** Returns ArrowComponent subobject **/
+UArrowComponent* AAtmosphericFog::GetArrowComponent() { return ArrowComponent; }
+#endif
diff --git a/Engine/Source/Runtime/Engine/Private/Brush.cpp b/Engine/Source/Runtime/Engine/Private/Brush.cpp
--- a/Engine/Source/Runtime/Engine/Private/Brush.cpp
+++ b/Engine/Source/Runtime/Engine/Private/Brush.cpp
@@ -283,4 +283,7 @@ bool ABrush::Modify(bool bAlwaysMarkDirty)
 		bSavedToTransactionBuffer = Brush->Modify(bAlwaysMarkDirty) || bSavedToTransactionBuffer;
 	}
 	return bSavedToTransactionBuffer;
-}
\ No newline at end of file
+}
+
+/** Returns BrushComponent subobject **/
+UBrushComponent* ABrush::GetBrushComponent() const { return BrushComponent; }
diff --git a/Engine/Source/Runtime/Engine/Private/Camera/CameraActor.cpp b/Engine/Source/Runtime/Engine/Private/Camera/CameraActor.cpp
--- a/Engine/Source/Runtime/Engine/Private/Camera/CameraActor.cpp
+++ b/Engine/Source/Runtime/Engine/Private/Camera/CameraActor.cpp
@@ -117,3 +117,6 @@ void ACameraActor::BeginPlay()
 }
 
 #undef LOCTEXT_NAMESPACE
+
+/** Returns CameraComponent subobject **/
+UCameraComponent* ACameraActor::GetCameraComponent() const { return CameraComponent; }
diff --git a/Engine/Source/Runtime/Engine/Private/Camera/CameraAnimInst.cpp b/Engine/Source/Runtime/Engine/Private/Camera/CameraAnimInst.cpp
--- a/Engine/Source/Runtime/Engine/Private/Camera/CameraAnimInst.cpp
+++ b/Engine/Source/Runtime/Engine/Private/Camera/CameraAnimInst.cpp
@@ -256,3 +256,6 @@ void UCameraAnimInst::SetPlaySpace(ECameraAnimPlaySpace::Type NewSpace, FRotator
 	UserPlaySpaceMatrix = (PlaySpace == ECameraAnimPlaySpace::UserDefined) ? FRotationMatrix(UserPlaySpace) : FMatrix::Identity;
 }
 
+
+/** Returns InterpGroupInst subobject **/
+UInterpGroupInst* UCameraAnimInst::GetInterpGroupInst() const { return InterpGroupInst; }
diff --git a/Engine/Source/Runtime/Engine/Private/Character.cpp b/Engine/Source/Runtime/Engine/Private/Character.cpp
--- a/Engine/Source/Runtime/Engine/Private/Character.cpp
+++ b/Engine/Source/Runtime/Engine/Private/Character.cpp
@@ -1338,4 +1338,15 @@ void ACharacter::ClientCheatGhost_Implementation()
 		CharacterMovement->bCheatFlying = true;
 		CharacterMovement->SetMovementMode(MOVE_Flying);
 	}
-}
\ No newline at end of file
+}
+
+/** Returns Mesh subobject **/
+USkeletalMeshComponent* ACharacter::GetMesh() const { return Mesh; }
+#if WITH_EDITORONLY_DATA
+/** Returns ArrowComponent subobject **/
+UArrowComponent* ACharacter::GetArrowComponent() const { return ArrowComponent; }
+#endif
+/** Returns CharacterMovement subobject **/
+UCharacterMovementComponent* ACharacter::GetCharacterMovement() const { return CharacterMovement; }
+/** Returns CapsuleComponent subobject **/
+UCapsuleComponent* ACharacter::GetCapsuleComponent() const { return CapsuleComponent; }
diff --git a/Engine/Source/Runtime/Engine/Private/Controller.cpp b/Engine/Source/Runtime/Engine/Private/Controller.cpp
--- a/Engine/Source/Runtime/Engine/Private/Controller.cpp
+++ b/Engine/Source/Runtime/Engine/Private/Controller.cpp
@@ -601,4 +601,7 @@ void AController::GetLifetimeReplicatedProps( TArray< FLifetimeProperty > & OutL
 
 	DOREPLIFETIME( AController, PlayerState );
 	DOREPLIFETIME( AController, Pawn );
-}
\ No newline at end of file
+}
+
+/** Returns TransformComponent subobject **/
+USceneComponent* AController::GetTransformComponent() const { return TransformComponent; }
diff --git a/Engine/Source/Runtime/Engine/Private/DecalActor.cpp b/Engine/Source/Runtime/Engine/Private/DecalActor.cpp
--- a/Engine/Source/Runtime/Engine/Private/DecalActor.cpp
+++ b/Engine/Source/Runtime/Engine/Private/DecalActor.cpp
@@ -133,3 +133,14 @@ class UMaterialInstanceDynamic* ADecalActor::CreateDynamicMaterialInstance()
 {
 	return Decal ? Decal->CreateDynamicMaterialInstance() : NULL;
 }
+
+/** Returns Decal subobject **/
+UDecalComponent* ADecalActor::GetDecal() const { return Decal; }
+#if WITH_EDITORONLY_DATA
+/** Returns ArrowComponent subobject **/
+UArrowComponent* ADecalActor::GetArrowComponent() const { return ArrowComponent; }
+/** Returns SpriteComponent subobject **/
+UBillboardComponent* ADecalActor::GetSpriteComponent() const { return SpriteComponent; }
+/** Returns BoxComponent subobject **/
+UBoxComponent* ADecalActor::GetBoxComponent() const { return BoxComponent; }
+#endif
diff --git a/Engine/Source/Runtime/Engine/Private/DefaultPawn.cpp b/Engine/Source/Runtime/Engine/Private/DefaultPawn.cpp
--- a/Engine/Source/Runtime/Engine/Private/DefaultPawn.cpp
+++ b/Engine/Source/Runtime/Engine/Private/DefaultPawn.cpp
@@ -211,3 +211,8 @@ UPawnMovementComponent* ADefaultPawn::GetMovementComponent() const
 {
 	return MovementComponent;
 }
+
+/** Returns CollisionComponent subobject **/
+USphereComponent* ADefaultPawn::GetCollisionComponent() const { return CollisionComponent; }
+/** Returns MeshComponent subobject **/
+UStaticMeshComponent* ADefaultPawn::GetMeshComponent() const { return MeshComponent; }
diff --git a/Engine/Source/Runtime/Engine/Private/DestructibleActor.cpp b/Engine/Source/Runtime/Engine/Private/DestructibleActor.cpp
--- a/Engine/Source/Runtime/Engine/Private/DestructibleActor.cpp
+++ b/Engine/Source/Runtime/Engine/Private/DestructibleActor.cpp
@@ -29,3 +29,6 @@ bool ADestructibleActor::GetReferencedContentObjects( TArray<UObject*>& Objects
 #endif // WITH_EDITOR
 
 
+
+/** Returns DestructibleComponent subobject **/
+UDestructibleComponent* ADestructibleActor::GetDestructibleComponent() const { return DestructibleComponent; }
diff --git a/Engine/Source/Runtime/Engine/Private/DocumentationActor.cpp b/Engine/Source/Runtime/Engine/Private/DocumentationActor.cpp
--- a/Engine/Source/Runtime/Engine/Private/DocumentationActor.cpp
+++ b/Engine/Source/Runtime/Engine/Private/DocumentationActor.cpp
@@ -91,3 +91,7 @@ void ADocumentationActor::UpdateLinkType()
 	
 }
 
+#if WITH_EDITORONLY_DATA
+/** Returns Billboard subobject **/
+UMaterialBillboardComponent* ADocumentationActor::GetBillboard() const { return Billboard; }
+#endif
diff --git a/Engine/Source/Runtime/Engine/Private/HeightFogComponent.cpp b/Engine/Source/Runtime/Engine/Private/HeightFogComponent.cpp
--- a/Engine/Source/Runtime/Engine/Private/HeightFogComponent.cpp
+++ b/Engine/Source/Runtime/Engine/Private/HeightFogComponent.cpp
@@ -211,3 +211,6 @@ void AExponentialHeightFog::OnRep_bEnabled()
 {
 	Component->SetVisibility(bEnabled);
 }
+
+/** Returns Component subobject **/
+UExponentialHeightFogComponent* AExponentialHeightFog::GetComponent() const { return Component; }
diff --git a/Engine/Source/Runtime/Engine/Private/Info.cpp b/Engine/Source/Runtime/Engine/Private/Info.cpp
--- a/Engine/Source/Runtime/Engine/Private/Info.cpp
+++ b/Engine/Source/Runtime/Engine/Private/Info.cpp
@@ -38,3 +38,8 @@ AInfo::AInfo(const FObjectInitializer& ObjectInitializer)
 	bReplicateMovement = false;
 	bCanBeDamaged = false;
 }
+
+#if WITH_EDITORONLY_DATA
+/** Returns SpriteComponent subobject **/
+UBillboardComponent* AInfo::GetSpriteComponent() const { return SpriteComponent; }
+#endif
diff --git a/Engine/Source/Runtime/Engine/Private/Interpolation.cpp b/Engine/Source/Runtime/Engine/Private/Interpolation.cpp
--- a/Engine/Source/Runtime/Engine/Private/Interpolation.cpp
+++ b/Engine/Source/Runtime/Engine/Private/Interpolation.cpp
@@ -9679,3 +9679,10 @@ void UInterpGroupCamera::PostEditChangeProperty(FPropertyChangedEvent& PropertyC
 	Super::PostEditChangeProperty(PropertyChangedEvent);
 }
 #endif // WITH_EDITOR
+
+#if WITH_EDITORONLY_DATA
+/** Returns SpriteComponent subobject **/
+UBillboardComponent* AMatineeActor::GetSpriteComponent() const { return SpriteComponent; }
+/** Returns SpriteComponent subobject **/
+UBillboardComponent* AMaterialInstanceActor::GetSpriteComponent() const { return SpriteComponent; }
+#endif
diff --git a/Engine/Source/Runtime/Engine/Private/Light.cpp b/Engine/Source/Runtime/Engine/Private/Light.cpp
--- a/Engine/Source/Runtime/Engine/Private/Light.cpp
+++ b/Engine/Source/Runtime/Engine/Private/Light.cpp
@@ -336,5 +336,14 @@ AGeneratedMeshAreaLight::AGeneratedMeshAreaLight(const FObjectInitializer& Objec
 	GetLightComponent()->CastStaticShadows = false;
 }
 
-// EOF
+bool ALight::IsToggleable() const
+{
+	return !LightComponent->HasStaticLighting();
+}
 
+/** Returns LightComponent subobject **/
+ULightComponent* ALight::GetLightComponent() const { return LightComponent; }
+#if WITH_EDITORONLY_DATA
+/** Returns ArrowComponent subobject **/
+UArrowComponent* ADirectionalLight::GetArrowComponent() const { return ArrowComponent; }
+#endif
diff --git a/Engine/Source/Runtime/Engine/Private/NavigationObjectBase.cpp b/Engine/Source/Runtime/Engine/Private/NavigationObjectBase.cpp
--- a/Engine/Source/Runtime/Engine/Private/NavigationObjectBase.cpp
+++ b/Engine/Source/Runtime/Engine/Private/NavigationObjectBase.cpp
@@ -229,3 +229,10 @@ void ANavigationObjectBase::GetMoveGoalReachTest(class AActor* MovingActor, cons
 	GoalOffset = FVector::ZeroVector;
 	GetSimpleCollisionCylinder(GoalRadius, GoalHalfHeight);	
 }
+
+/** Returns CapsuleComponent subobject **/
+UCapsuleComponent* ANavigationObjectBase::GetCapsuleComponent() const { return CapsuleComponent; }
+/** Returns GoodSprite subobject **/
+UBillboardComponent* ANavigationObjectBase::GetGoodSprite() const { return GoodSprite; }
+/** Returns BadSprite subobject **/
+UBillboardComponent* ANavigationObjectBase::GetBadSprite() const { return BadSprite; }
diff --git a/Engine/Source/Runtime/Engine/Private/NiagaraActor.cpp b/Engine/Source/Runtime/Engine/Private/NiagaraActor.cpp
--- a/Engine/Source/Runtime/Engine/Private/NiagaraActor.cpp
+++ b/Engine/Source/Runtime/Engine/Private/NiagaraActor.cpp
@@ -59,3 +59,12 @@ ANiagaraActor::ANiagaraActor(const FObjectInitializer& ObjectInitializer)
 	}
 #endif // WITH_EDITORONLY_DATA
 }
+
+/** Returns NiagaraComponent subobject **/
+UNiagaraComponent* ANiagaraActor::GetNiagaraComponent() const { return NiagaraComponent; }
+#if WITH_EDITORONLY_DATA
+/** Returns SpriteComponent subobject **/
+UBillboardComponent* ANiagaraActor::GetSpriteComponent() const { return SpriteComponent; }
+/** Returns ArrowComponent subobject **/
+UArrowComponent* ANiagaraActor::GetArrowComponent() const { return ArrowComponent; }
+#endif
diff --git a/Engine/Source/Runtime/Engine/Private/Note.cpp b/Engine/Source/Runtime/Engine/Private/Note.cpp
--- a/Engine/Source/Runtime/Engine/Private/Note.cpp
+++ b/Engine/Source/Runtime/Engine/Private/Note.cpp
@@ -57,3 +57,10 @@ ANote::ANote(const FObjectInitializer& ObjectInitializer)
 	bHidden = true;
 	bCanBeDamaged = false;
 }
+
+#if WITH_EDITORONLY_DATA
+/** Returns SpriteComponent subobject **/
+UBillboardComponent* ANote::GetSpriteComponent() const { return SpriteComponent; }
+/** Returns ArrowComponent subobject **/
+UArrowComponent* ANote::GetArrowComponent() const { return ArrowComponent; }
+#endif
diff --git a/Engine/Source/Runtime/Engine/Private/Particles/ParticleComponents.cpp b/Engine/Source/Runtime/Engine/Private/Particles/ParticleComponents.cpp
--- a/Engine/Source/Runtime/Engine/Private/Particles/ParticleComponents.cpp
+++ b/Engine/Source/Runtime/Engine/Private/Particles/ParticleComponents.cpp
@@ -6079,3 +6079,12 @@ void AEmitterCameraLensEffectBase::ActivateLensEffect()
 }
 
 #undef LOCTEXT_NAMESPACE
+
+/** Returns ParticleSystemComponent subobject **/
+UParticleSystemComponent* AEmitter::GetParticleSystemComponent() { return ParticleSystemComponent; }
+#if WITH_EDITORONLY_DATA
+/** Returns SpriteComponent subobject **/
+UBillboardComponent* AEmitter::GetSpriteComponent() const { return SpriteComponent; }
+/** Returns ArrowComponent subobject **/
+UArrowComponent* AEmitter::GetArrowComponent() const { return ArrowComponent; }
+#endif
diff --git a/Engine/Source/Runtime/Engine/Private/PhysicsEngine/PhysicsConstraintActor.cpp b/Engine/Source/Runtime/Engine/Private/PhysicsEngine/PhysicsConstraintActor.cpp
--- a/Engine/Source/Runtime/Engine/Private/PhysicsEngine/PhysicsConstraintActor.cpp
+++ b/Engine/Source/Runtime/Engine/Private/PhysicsEngine/PhysicsConstraintActor.cpp
@@ -62,3 +62,6 @@ void APhysicsConstraintActor::LoadedFromAnotherClass( const FName& OldClassName
 	ConstraintComp->UpdateSpriteTexture();
 }
 #endif // WITH_EDITOR
+
+/** Returns ConstraintComp subobject **/
+UPhysicsConstraintComponent* APhysicsConstraintActor::GetConstraintComp() const { return ConstraintComp; }
diff --git a/Engine/Source/Runtime/Engine/Private/PhysicsEngine/PhysicsThruster.cpp b/Engine/Source/Runtime/Engine/Private/PhysicsEngine/PhysicsThruster.cpp
--- a/Engine/Source/Runtime/Engine/Private/PhysicsEngine/PhysicsThruster.cpp
+++ b/Engine/Source/Runtime/Engine/Private/PhysicsEngine/PhysicsThruster.cpp
@@ -93,4 +93,13 @@ APhysicsThruster::APhysicsThruster(const FObjectInitializer& ObjectInitializer)
 		}
 	}
 #endif // WITH_EDITORONLY_DATA
-}
\ No newline at end of file
+}
+
+/** Returns ThrusterComponent subobject **/
+UPhysicsThrusterComponent* APhysicsThruster::GetThrusterComponent() const { return ThrusterComponent; }
+#if WITH_EDITORONLY_DATA
+/** Returns ArrowComponent subobject **/
+UArrowComponent* APhysicsThruster::GetArrowComponent() const { return ArrowComponent; }
+/** Returns SpriteComponent subobject **/
+UBillboardComponent* APhysicsThruster::GetSpriteComponent() const { return SpriteComponent; }
+#endif
diff --git a/Engine/Source/Runtime/Engine/Private/PhysicsEngine/RadialForceComponent.cpp b/Engine/Source/Runtime/Engine/Private/PhysicsEngine/RadialForceComponent.cpp
--- a/Engine/Source/Runtime/Engine/Private/PhysicsEngine/RadialForceComponent.cpp
+++ b/Engine/Source/Runtime/Engine/Private/PhysicsEngine/RadialForceComponent.cpp
@@ -269,3 +269,10 @@ void ARadialForceActor::ToggleForce()
 	}
 }
 
+
+/** Returns ForceComponent subobject **/
+URadialForceComponent* ARadialForceActor::GetForceComponent() const { return ForceComponent; }
+#if WITH_EDITORONLY_DATA
+/** Returns SpriteComponent subobject **/
+UBillboardComponent* ARadialForceActor::GetSpriteComponent() const { return SpriteComponent; }
+#endif
diff --git a/Engine/Source/Runtime/Engine/Private/PlayerCameraManager.cpp b/Engine/Source/Runtime/Engine/Private/PlayerCameraManager.cpp
--- a/Engine/Source/Runtime/Engine/Private/PlayerCameraManager.cpp
+++ b/Engine/Source/Runtime/Engine/Private/PlayerCameraManager.cpp
@@ -1229,3 +1229,6 @@ void FTViewTarget::CheckViewTarget(APlayerController* OwningController)
 		}
 	}
 }
+
+/** Returns TransformComponent subobject **/
+USceneComponent* APlayerCameraManager::GetTransformComponent() const { return TransformComponent; }
diff --git a/Engine/Source/Runtime/Engine/Private/PlayerStart.cpp b/Engine/Source/Runtime/Engine/Private/PlayerStart.cpp
--- a/Engine/Source/Runtime/Engine/Private/PlayerStart.cpp
+++ b/Engine/Source/Runtime/Engine/Private/PlayerStart.cpp
@@ -78,3 +78,7 @@ void APlayerStart::PostUnregisterAllComponents()
 	}
 }
 
+#if WITH_EDITORONLY_DATA
+/** Returns ArrowComponent subobject **/
+UArrowComponent* APlayerStart::GetArrowComponent() const { return ArrowComponent; }
+#endif
diff --git a/Engine/Source/Runtime/Engine/Private/ReflectionCaptureComponent.cpp b/Engine/Source/Runtime/Engine/Private/ReflectionCaptureComponent.cpp
--- a/Engine/Source/Runtime/Engine/Private/ReflectionCaptureComponent.cpp
+++ b/Engine/Source/Runtime/Engine/Private/ReflectionCaptureComponent.cpp
@@ -1555,3 +1555,12 @@ void FReflectionCaptureProxy::SetTransform(const FMatrix& InTransform)
 	ReflectionXAxisAndYScale = ReflectionXAxis.SafeNormal() * ScaleVector.Y;
 	ReflectionXAxisAndYScale.W = ScaleVector.Y / ScaleVector.Z;
 }
+
+/** Returns CaptureComponent subobject **/
+UReflectionCaptureComponent* AReflectionCapture::GetCaptureComponent() const { return CaptureComponent; }
+#if WITH_EDITORONLY_DATA
+/** Returns SpriteComponent subobject **/
+UBillboardComponent* AReflectionCapture::GetSpriteComponent() const { return SpriteComponent; }
+#endif
+/** Returns DrawCaptureRadius subobject **/
+UDrawSphereComponent* ASphereReflectionCapture::GetDrawCaptureRadius() const { return DrawCaptureRadius; }
diff --git a/Engine/Source/Runtime/Engine/Private/SceneCaptureComponent.cpp b/Engine/Source/Runtime/Engine/Private/SceneCaptureComponent.cpp
--- a/Engine/Source/Runtime/Engine/Private/SceneCaptureComponent.cpp
+++ b/Engine/Source/Runtime/Engine/Private/SceneCaptureComponent.cpp
@@ -370,3 +370,14 @@ void USceneCaptureComponentCube::PostEditChangeProperty(FPropertyChangedEvent& P
 	UpdateContent();
 }
 #endif // WITH_EDITOR
+
+/** Returns MeshComp subobject **/
+UStaticMeshComponent* ASceneCapture::GetMeshComp() const { return MeshComp; }
+/** Returns CaptureComponent2D subobject **/
+USceneCaptureComponent2D* ASceneCapture2D::GetCaptureComponent2D() const { return CaptureComponent2D; }
+/** Returns DrawFrustum subobject **/
+UDrawFrustumComponent* ASceneCapture2D::GetDrawFrustum() const { return DrawFrustum; }
+/** Returns CaptureComponentCube subobject **/
+USceneCaptureComponentCube* ASceneCaptureCube::GetCaptureComponentCube() const { return CaptureComponentCube; }
+/** Returns DrawFrustum subobject **/
+UDrawFrustumComponent* ASceneCaptureCube::GetDrawFrustum() const { return DrawFrustum; }
diff --git a/Engine/Source/Runtime/Engine/Private/SkeletalMesh.cpp b/Engine/Source/Runtime/Engine/Private/SkeletalMesh.cpp
--- a/Engine/Source/Runtime/Engine/Private/SkeletalMesh.cpp
+++ b/Engine/Source/Runtime/Engine/Private/SkeletalMesh.cpp
@@ -4951,3 +4951,6 @@ FPrimitiveSceneProxy* USkinnedMeshComponent::CreateSceneProxy()
 }
 
 #undef LOCTEXT_NAMESPACE
+
+/** Returns SkeletalMeshComponent subobject **/
+USkeletalMeshComponent* ASkeletalMeshActor::GetSkeletalMeshComponent() { return SkeletalMeshComponent; }
diff --git a/Engine/Source/Runtime/Engine/Private/SkyLightComponent.cpp b/Engine/Source/Runtime/Engine/Private/SkyLightComponent.cpp
--- a/Engine/Source/Runtime/Engine/Private/SkyLightComponent.cpp
+++ b/Engine/Source/Runtime/Engine/Private/SkyLightComponent.cpp
@@ -542,3 +542,6 @@ void ASkyLight::OnRep_bEnabled()
 }
 
 #undef LOCTEXT_NAMESPACE
+
+/** Returns LightComponent subobject **/
+USkyLightComponent* ASkyLight::GetLightComponent() const { return LightComponent; }
diff --git a/Engine/Source/Runtime/Engine/Private/SpotLight.cpp b/Engine/Source/Runtime/Engine/Private/SpotLight.cpp
--- a/Engine/Source/Runtime/Engine/Private/SpotLight.cpp
+++ b/Engine/Source/Runtime/Engine/Private/SpotLight.cpp
@@ -137,4 +137,9 @@ void ASpotLight::PostEditChangeProperty(FPropertyChangedEvent& PropertyChangedEv
 		ArrowComponent->ArrowColor = GetLightColor();
 	}
 }
-#endif
\ No newline at end of file
+#endif
+
+#if WITH_EDITORONLY_DATA
+/** Returns ArrowComponent subobject **/
+UArrowComponent* ASpotLight::GetArrowComponent() const { return ArrowComponent; }
+#endif
diff --git a/Engine/Source/Runtime/Engine/Private/StaticMeshActor.cpp b/Engine/Source/Runtime/Engine/Private/StaticMeshActor.cpp
--- a/Engine/Source/Runtime/Engine/Private/StaticMeshActor.cpp
+++ b/Engine/Source/Runtime/Engine/Private/StaticMeshActor.cpp
@@ -203,4 +203,7 @@ void AStaticMeshActor::CheckForErrors()
 
 #endif // WITH_EDITOR
 
-#undef LOCTEXT_NAMESPACE
\ No newline at end of file
+#undef LOCTEXT_NAMESPACE
+
+/** Returns StaticMeshComponent subobject **/
+UStaticMeshComponent* AStaticMeshActor::GetStaticMeshComponent() const { return StaticMeshComponent; }
diff --git a/Engine/Source/Runtime/Engine/Private/TargetPoint.cpp b/Engine/Source/Runtime/Engine/Private/TargetPoint.cpp
--- a/Engine/Source/Runtime/Engine/Private/TargetPoint.cpp
+++ b/Engine/Source/Runtime/Engine/Private/TargetPoint.cpp
@@ -60,4 +60,11 @@ ATargetPoint::ATargetPoint(const FObjectInitializer& ObjectInitializer)
 
 	bHidden = true;
 	bCanBeDamaged = false;
-}
\ No newline at end of file
+}
+
+#if WITH_EDITORONLY_DATA
+/** Returns SpriteComponent subobject **/
+UBillboardComponent* ATargetPoint::GetSpriteComponent() const { return SpriteComponent; }
+/** Returns ArrowComponent subobject **/
+UArrowComponent* ATargetPoint::GetArrowComponent() const { return ArrowComponent; }
+#endif
diff --git a/Engine/Source/Runtime/Engine/Private/TextRenderComponent.cpp b/Engine/Source/Runtime/Engine/Private/TextRenderComponent.cpp
--- a/Engine/Source/Runtime/Engine/Private/TextRenderComponent.cpp
+++ b/Engine/Source/Runtime/Engine/Private/TextRenderComponent.cpp
@@ -994,3 +994,10 @@ void UTextRenderComponent::PostLoad()
 
 	Super::PostLoad();
 }
+
+/** Returns TextRender subobject **/
+UTextRenderComponent* ATextRenderActor::GetTextRender() const { return TextRender; }
+#if WITH_EDITORONLY_DATA
+/** Returns SpriteComponent subobject **/
+UBillboardComponent* ATextRenderActor::GetSpriteComponent() const { return SpriteComponent; }
+#endif
diff --git a/Engine/Source/Runtime/Engine/Private/TriggerBase.cpp b/Engine/Source/Runtime/Engine/Private/TriggerBase.cpp
--- a/Engine/Source/Runtime/Engine/Private/TriggerBase.cpp
+++ b/Engine/Source/Runtime/Engine/Private/TriggerBase.cpp
@@ -48,3 +48,8 @@ ATriggerBase::ATriggerBase(const FObjectInitializer& ObjectInitializer)
 		SpriteComponent->bIsScreenSizeScaled = true;
 	}
 }
+
+/** Returns CollisionComponent subobject **/
+UShapeComponent* ATriggerBase::GetCollisionComponent() const { return CollisionComponent; }
+/** Returns SpriteComponent subobject **/
+UBillboardComponent* ATriggerBase::GetSpriteComponent() const { return SpriteComponent; }
diff --git a/Engine/Source/Runtime/Engine/Private/VectorFieldVolume.cpp b/Engine/Source/Runtime/Engine/Private/VectorFieldVolume.cpp
--- a/Engine/Source/Runtime/Engine/Private/VectorFieldVolume.cpp
+++ b/Engine/Source/Runtime/Engine/Private/VectorFieldVolume.cpp
@@ -40,3 +40,10 @@ AVectorFieldVolume::AVectorFieldVolume(const FObjectInitializer& ObjectInitializ
 #endif // WITH_EDITORONLY_DATA
 }
 
+
+/** Returns VectorFieldComponent subobject **/
+UVectorFieldComponent* AVectorFieldVolume::GetVectorFieldComponent() const { return VectorFieldComponent; }
+#if WITH_EDITORONLY_DATA
+/** Returns SpriteComponent subobject **/
+UBillboardComponent* AVectorFieldVolume::GetSpriteComponent() const { return SpriteComponent; }
+#endif
diff --git a/Engine/Source/Runtime/Engine/Private/WheeledVehicle.cpp b/Engine/Source/Runtime/Engine/Private/WheeledVehicle.cpp
--- a/Engine/Source/Runtime/Engine/Private/WheeledVehicle.cpp
+++ b/Engine/Source/Runtime/Engine/Private/WheeledVehicle.cpp
@@ -43,4 +43,14 @@ void AWheeledVehicle::DisplayDebug(UCanvas* Canvas, const FDebugDisplayInfo& Deb
 		GetVehicleMovementComponent()->DrawDebug(Canvas, YL, YPos);
 	}
 #endif
-}
\ No newline at end of file
+}
+
+class UWheeledVehicleMovementComponent* AWheeledVehicle::GetVehicleMovementComponent() const
+{
+	return VehicleMovement;
+}
+
+/** Returns Mesh subobject **/
+USkeletalMeshComponent* AWheeledVehicle::GetMesh() const { return Mesh; }
+/** Returns VehicleMovement subobject **/
+UWheeledVehicleMovementComponent* AWheeledVehicle::GetVehicleMovement() const { return VehicleMovement; }
diff --git a/Engine/Source/Runtime/Engine/Private/WindDirectionalSource.cpp b/Engine/Source/Runtime/Engine/Private/WindDirectionalSource.cpp
--- a/Engine/Source/Runtime/Engine/Private/WindDirectionalSource.cpp
+++ b/Engine/Source/Runtime/Engine/Private/WindDirectionalSource.cpp
@@ -136,3 +136,10 @@ FWindSourceSceneProxy* UWindDirectionalSourceComponent::CreateSceneProxy() const
 		);
 }
 
+
+/** Returns Component subobject **/
+UWindDirectionalSourceComponent* AWindDirectionalSource::GetComponent() const { return Component; }
+#if WITH_EDITORONLY_DATA
+/** Returns ArrowComponent subobject **/
+UArrowComponent* AWindDirectionalSource::GetArrowComponent() const { return ArrowComponent; }
+#endif
diff --git a/Engine/Source/Runtime/GameplayAbilities/Private/Abilities/GameplayAbilityWorldReticle_ActorVisualization.cpp b/Engine/Source/Runtime/GameplayAbilities/Private/Abilities/GameplayAbilityWorldReticle_ActorVisualization.cpp
--- a/Engine/Source/Runtime/GameplayAbilities/Private/Abilities/GameplayAbilityWorldReticle_ActorVisualization.cpp
+++ b/Engine/Source/Runtime/GameplayAbilities/Private/Abilities/GameplayAbilityWorldReticle_ActorVisualization.cpp
@@ -70,3 +70,6 @@ void AGameplayAbilityWorldReticle_ActorVisualization::EndPlay(const EEndPlayReas
 
 	Super::EndPlay(EndPlayReason);
 }
+
+/** Returns CollisionComponent subobject **/
+UCapsuleComponent* AGameplayAbilityWorldReticle_ActorVisualization::GetCollisionComponent() { return CollisionComponent; }
diff --git a/Engine/Source/Runtime/GameplayAbilities/Private/AbilitySystemTestPawn.cpp b/Engine/Source/Runtime/GameplayAbilities/Private/AbilitySystemTestPawn.cpp
--- a/Engine/Source/Runtime/GameplayAbilities/Private/AbilitySystemTestPawn.cpp
+++ b/Engine/Source/Runtime/GameplayAbilities/Private/AbilitySystemTestPawn.cpp
@@ -35,4 +35,7 @@ void AAbilitySystemTestPawn::PostInitializeComponents()
 UAbilitySystemComponent* AAbilitySystemTestPawn::GetAbilitySystemComponent() const
 {
 	return FindComponentByClass<UAbilitySystemComponent>();
-}
\ No newline at end of file
+}
+
+/** Returns AbilitySystemComponent subobject **/
+UAbilitySystemComponent* AAbilitySystemTestPawn::GetAbilitySystemComponent() { return AbilitySystemComponent; }
diff --git a/Engine/Source/Runtime/GameplayAbilities/Public/Abilities/GameplayAbilityWorldReticle_ActorVisualization.h b/Engine/Source/Runtime/GameplayAbilities/Public/Abilities/GameplayAbilityWorldReticle_ActorVisualization.h
--- a/Engine/Source/Runtime/GameplayAbilities/Public/Abilities/GameplayAbilityWorldReticle_ActorVisualization.h
+++ b/Engine/Source/Runtime/GameplayAbilities/Public/Abilities/GameplayAbilityWorldReticle_ActorVisualization.h
@@ -19,6 +19,7 @@ public:
 
 private:
 	/** Hardcoded collision component, so other objects don't think they can collide with the visualization actor */
+	DEPRECATED_FORGAME(4.6, "CollisionComponent should not be accessed directly, please use GetCollisionComponent() function instead. CollisionComponent will soon be private and your code will not compile.")
 	UPROPERTY()
 	class UCapsuleComponent* CollisionComponent;
 public:
@@ -30,5 +31,5 @@ public:
 	virtual void EndPlay(const EEndPlayReason::Type EndPlayReason);
 
 	/** Returns CollisionComponent subobject **/
-	FORCEINLINE class UCapsuleComponent* GetCollisionComponent() { return CollisionComponent; }
+	GAMEPLAYABILITIES_API class UCapsuleComponent* GetCollisionComponent();
 };
diff --git a/Engine/Source/Runtime/GameplayAbilities/Public/AbilitySystemTestPawn.h b/Engine/Source/Runtime/GameplayAbilities/Public/AbilitySystemTestPawn.h
--- a/Engine/Source/Runtime/GameplayAbilities/Public/AbilitySystemTestPawn.h
+++ b/Engine/Source/Runtime/GameplayAbilities/Public/AbilitySystemTestPawn.h
@@ -15,8 +15,9 @@ class GAMEPLAYABILITIES_API AAbilitySystemTestPawn : public ADefaultPawn, public
 
 	virtual UAbilitySystemComponent* GetAbilitySystemComponent() const override;
 	
-private:
+private_subobject:
 	/** DefaultPawn collision component */
+	DEPRECATED_FORGAME(4.6, "AbilitySystemComponent should not be accessed directly, please use GetAbilitySystemComponent() function instead. AbilitySystemComponent will soon be private and your code will not compile.")
 	UPROPERTY(Category = AbilitySystem, VisibleAnywhere, BlueprintReadOnly, meta = (AllowPrivateAccess = "true"))
 	class UAbilitySystemComponent* AbilitySystemComponent;
 public:
@@ -30,5 +31,5 @@ public:
 	static FName AbilitySystemComponentName;
 
 	/** Returns AbilitySystemComponent subobject **/
-	FORCEINLINE class UAbilitySystemComponent* GetAbilitySystemComponent() { return AbilitySystemComponent; }
+	class UAbilitySystemComponent* GetAbilitySystemComponent();
 };
diff --git a/Engine/Source/Runtime/Landscape/Classes/LandscapeGizmoActor.h b/Engine/Source/Runtime/Landscape/Classes/LandscapeGizmoActor.h
--- a/Engine/Source/Runtime/Landscape/Classes/LandscapeGizmoActor.h
+++ b/Engine/Source/Runtime/Landscape/Classes/LandscapeGizmoActor.h
@@ -33,7 +33,8 @@ class ALandscapeGizmoActor : public AActor
 	UPROPERTY(EditAnywhere, transient, Category=Gizmo)
 	class ULandscapeInfo* TargetLandscapeInfo;
 
-private:
+private_subobject:
+	DEPRECATED_FORGAME(4.6, "SpriteComponent should not be accessed directly, please use GetSpriteComponent() function instead. SpriteComponent will soon be private and your code will not compile.")
 	UPROPERTY()
 	UBillboardComponent* SpriteComponent;
 #endif // WITH_EDITORONLY_DATA
@@ -58,7 +59,7 @@ public:
 public:
 #if WITH_EDITORONLY_DATA
 	/** Returns SpriteComponent subobject **/
-	FORCEINLINE UBillboardComponent* GetSpriteComponent() const { return SpriteComponent; }
+	LANDSCAPE_API UBillboardComponent* GetSpriteComponent() const;
 #endif
 };
 
diff --git a/Engine/Source/Runtime/Landscape/Private/LandscapeGizmoActor.cpp b/Engine/Source/Runtime/Landscape/Private/LandscapeGizmoActor.cpp
--- a/Engine/Source/Runtime/Landscape/Private/LandscapeGizmoActor.cpp
+++ b/Engine/Source/Runtime/Landscape/Private/LandscapeGizmoActor.cpp
@@ -1286,3 +1286,8 @@ void ALandscapeGizmoActiveActor::ImportFromClipboard()
 	}
 }
 #endif	//WITH_EDITOR
+
+#if WITH_EDITORONLY_DATA
+/** Returns SpriteComponent subobject **/
+UBillboardComponent* ALandscapeGizmoActor::GetSpriteComponent() const { return SpriteComponent; }
+#endif
\ No newline at end of file
