Engine: Add ENGINE_API to ACameraActor

diff --git a/Engine/Source/Runtime/Engine/Classes/Camera/CameraActor.h b/Engine/Source/Runtime/Engine/Classes/Camera/CameraActor.h
--- a/Engine/Source/Runtime/Engine/Classes/Camera/CameraActor.h
+++ b/Engine/Source/Runtime/Engine/Classes/Camera/CameraActor.h
@@ -7,8 +7,8 @@
 /** 
  * A camera that can be placed in a level.
  */
-UCLASS(ClassGroup=Common, hideCategories=(Input, Rendering), showcategories=("Input|MouseInput", "Input|TouchInput"), MinimalAPI, Blueprintable)
-class ACameraActor : public AActor
+UCLASS(ClassGroup=Common, hideCategories=(Input, Rendering), showcategories=("Input|MouseInput", "Input|TouchInput"), Blueprintable)
+class ENGINE_API ACameraActor : public AActor
 {
 	GENERATED_UCLASS_BODY()
 
@@ -32,7 +32,7 @@ public:
 
 	/** Returns index of the player for whom we auto-activate, or INDEX_NONE (-1) if disabled. */
 	UFUNCTION(BlueprintCallable, Category="AutoPlayerActivation")
-	ENGINE_API int32 GetAutoActivatePlayerIndex() const;
+	int32 GetAutoActivatePlayerIndex() const;
 
 private:
 
@@ -54,15 +54,15 @@ private:
 public:
 	// Begin UObject interface
 	virtual void Serialize(FArchive& Ar) override;
-	ENGINE_API virtual void PostLoadSubobjects(FObjectInstancingGraph* OuterInstanceGraph) override;
+	virtual void PostLoadSubobjects(FObjectInstancingGraph* OuterInstanceGraph) override;
 
 #if WITH_EDITOR
-	ENGINE_API virtual void PostEditChangeProperty(FPropertyChangedEvent& PropertyChangedEvent) override;
+	virtual void PostEditChangeProperty(FPropertyChangedEvent& PropertyChangedEvent) override;
 #endif
 	// End UObject interface
 	
 	// Begin AActor interface
-	ENGINE_API virtual void BeginPlay() override;
+	virtual void BeginPlay() override;
 	// End AActor interface
 
 	/** Returns CameraComponent subobject **/
