Fixed crash when anim notify state BP is compiled while active.

https://jira.ol.epicgames.net/browse/UE-3651?filter=12278

diff --git a/Engine/Source/Runtime/Engine/Classes/Animation/AnimInstance.h b/Engine/Source/Runtime/Engine/Classes/Animation/AnimInstance.h
--- a/Engine/Source/Runtime/Engine/Classes/Animation/AnimInstance.h
+++ b/Engine/Source/Runtime/Engine/Classes/Animation/AnimInstance.h
@@ -577,6 +577,7 @@ public:
 	/** Currently Active AnimNotifyState, stored as a copy of the event as we need to
 		call NotifyEnd on the event after a deletion in the editor. After this the event
 		is removed correctly. */
+	UPROPERTY(transient)
 	TArray<FAnimNotifyEvent> ActiveAnimNotifyState;
 
 	/** Curve Values that are added to trigger in event**/
