Fix CIS

diff --git a/Engine/Source/Runtime/Engine/Classes/Camera/CameraActor.h b/Engine/Source/Runtime/Engine/Classes/Camera/CameraActor.h
--- a/Engine/Source/Runtime/Engine/Classes/Camera/CameraActor.h
+++ b/Engine/Source/Runtime/Engine/Classes/Camera/CameraActor.h
@@ -66,5 +66,5 @@ public:
 	// End AActor interface
 
 	/** Returns CameraComponent subobject **/
-	ENGINE_API class UCameraComponent* GetCameraComponent() const;
+	class UCameraComponent* GetCameraComponent() const;
 };
