GitHub pull request #545 - Allow blueprints to set cinematic mode from player controller - from user dgrieshofer.

diff --git a/Engine/Source/Runtime/Engine/Classes/GameFramework/PlayerController.h b/Engine/Source/Runtime/Engine/Classes/GameFramework/PlayerController.h
--- a/Engine/Source/Runtime/Engine/Classes/GameFramework/PlayerController.h
+++ b/Engine/Source/Runtime/Engine/Classes/GameFramework/PlayerController.h
@@ -1307,6 +1307,7 @@ public:
 	 * @param	bAffectsMovement	specify true to disable movement in cinematic mode, enable it when leaving
 	 * @param	bAffectsTurning		specify true to disable turning in cinematic mode or enable it when leaving
 	 */
+	UFUNCTION(BlueprintCallable, Category="Game|Cinematic", meta=(bHidePlayer="true", bAffectsHUD="true"))
 	virtual void SetCinematicMode(bool bInCinematicMode, bool bHidePlayer, bool bAffectsHUD, bool bAffectsMovement, bool bAffectsTurning);
 
 	/**
