Fix crash when null level is passed to OnLevelRemovedFromWorld

diff --git a/Engine/Source/Runtime/Engine/Private/AI/Navigation/NavigationSystem.cpp b/Engine/Source/Runtime/Engine/Private/AI/Navigation/NavigationSystem.cpp
--- a/Engine/Source/Runtime/Engine/Private/AI/Navigation/NavigationSystem.cpp
+++ b/Engine/Source/Runtime/Engine/Private/AI/Navigation/NavigationSystem.cpp
@@ -2857,7 +2857,7 @@ void UNavigationSystem::OnLevelRemovedFromWorld(ULevel* InLevel, UWorld* InWorld
 	{
 		RemoveLevelCollisionFromOctree(InLevel);
 
-		if (!InLevel->IsPersistentLevel())
+		if (InLevel && !InLevel->IsPersistentLevel())
 		{
 			for (ANavigationData* NavData : NavDataSet)
 			{
