#ue4 - Add enum controlling MovementComponent plane constraint / axis-locking behavior. Modes are: custom, X, Y, Z, and an option to use the global physics axis setting.

diff --git a/Engine/Source/Runtime/Core/Public/UObject/ObjectVersion.h b/Engine/Source/Runtime/Core/Public/UObject/ObjectVersion.h
--- a/Engine/Source/Runtime/Core/Public/UObject/ObjectVersion.h
+++ b/Engine/Source/Runtime/Core/Public/UObject/ObjectVersion.h
@@ -638,6 +638,8 @@ enum EUnrealEngineObjectUE4Version
 	VER_UE4_SLATE_BULK_FONT_DATA,
 	// Add new friction behavior in ProjectileMovementComponent.
 	VER_UE4_ADD_PROJECTILE_FRICTION_BEHAVIOR,
+	// Add axis settings enum to MovementComponent.
+	VER_UE4_MOVEMENTCOMPONENT_AXIS_SETTINGS,
 
 	// -----<new versions can be added before this line>-------------------------------------------------
 	// - this needs to be the last line (see note below)
diff --git a/Engine/Source/Runtime/Engine/Classes/GameFramework/MovementComponent.h b/Engine/Source/Runtime/Engine/Classes/GameFramework/MovementComponent.h
--- a/Engine/Source/Runtime/Engine/Classes/GameFramework/MovementComponent.h
+++ b/Engine/Source/Runtime/Engine/Classes/GameFramework/MovementComponent.h
@@ -13,6 +13,25 @@ struct FCollisionQueryParams;
 struct FCollisionResponseParams;
 struct FCollisionShape;
 
+/**
+ * Setting that controls behavior when movement is restricted to a 2D plane defined by a specific axis/normal,
+ * so that movement along the locked axis is not be possible.
+ */
+UENUM(BlueprintType)
+enum class EPlaneConstraintAxisSetting : uint8
+{
+	/** Lock movement to a user-defined axis. */
+	Custom,
+	/** Lock movement in the X axis. */
+	X,
+	/** Lock movement in the Y axis. */
+	Y,
+	/** Lock movement in the Z axis. */
+	Z,
+	/** Use the global physics project setting. */
+	UseGlobalPhysicsSetting
+};
+
 
 /**
  * MovementComponent is an abstract component class that defines functionality for moving a PrimitiveComponent (our UpdatedComponent) each tick.
@@ -46,34 +65,55 @@ class ENGINE_API UMovementComponent : public UActorComponent
 	UPROPERTY(EditAnywhere, BlueprintReadWrite, Category=Velocity)
 	FVector Velocity;
 
-protected:
 	/**
-	 * The normal of the plane that constrains movement, if plane constraint is enabled.
-	 * If for example you wanted to constrain movement to the Y-Z plane (so that X cannot change), the normal would be set to X=1 Y=0 Z=0.
-	 * @see SetPlaneConstraintNormal(), SetPlaneConstraintFromVectors()
+	 * If true, movement will be constrained to a plane.
+	 * @see PlaneConstraintNormal, PlaneConstraintOrigin, PlaneConstraintAxisSetting
 	 */
 	UPROPERTY(EditDefaultsOnly, BlueprintReadOnly, Category=PlanarMovement)
+	uint32 bConstrainToPlane:1;
+
+	/** If true and plane constraints are enabled, then the updated component will be snapped to the plane when first attached. */
+	UPROPERTY(EditDefaultsOnly, BlueprintReadOnly, Category=PlanarMovement, meta=(editcondition=bConstrainToPlane))
+	uint32 bSnapToPlaneAtStart:1;
+
+private:
+
+	/**
+	 * Setting that controls behavior when movement is restricted to a 2D plane defined by a specific axis/normal,
+	 * so that movement along the locked axis is not be possible.
+	 * @see SetPlaneConstraintAxisSetting
+	 */
+	UPROPERTY(EditDefaultsOnly, Category=PlanarMovement, meta=(editcondition=bConstrainToPlane))
+	EPlaneConstraintAxisSetting PlaneConstraintAxisSetting;
+
+protected:
+
+	/**
+	 * The normal or axis of the plane that constrains movement, if bConstrainToPlane is enabled.
+	 * If for example you wanted to constrain movement to the X-Z plane (so that Y cannot change), the normal would be set to X=0 Y=1 Z=0.
+	 * This is recalculated whenever PlaneConstraintAxisSetting changes.
+	 * @see bConstrainToPlane, SetPlaneConstraintNormal(), SetPlaneConstraintFromVectors()
+	 */
+	UPROPERTY(EditDefaultsOnly, BlueprintReadOnly, Category=PlanarMovement, meta=(editcondition=bConstrainToPlane))
 	FVector PlaneConstraintNormal;
 
 	/**
 	 * The origin of the plane that constrains movement, if plane constraint is enabled. 
 	 * This defines the behavior of snapping a position to the plane, such as by SnapUpdatedComponentToPlane().
-	 * @see SetPlaneConstraintOrigin().
+	 * @see bConstrainToPlane, SetPlaneConstraintOrigin().
 	 */
-	UPROPERTY(EditDefaultsOnly, BlueprintReadOnly, Category=PlanarMovement)
+	UPROPERTY(EditDefaultsOnly, BlueprintReadOnly, Category=PlanarMovement, meta=(editcondition=bConstrainToPlane))
 	FVector PlaneConstraintOrigin;
 
-public:
 	/**
-	 * If true, movement will be constrained to a plane.
-	 * @see PlaneConstraintNormal, PlaneConstraintOrigin
+	 * Helper to compute the plane constraint axis from the current setting.
+	 * 
+	 * @param  AxisSetting Setting to use when computing the axis.
+	 * @return Plane constraint axis/normal.
 	 */
-	UPROPERTY(EditDefaultsOnly, BlueprintReadOnly, Category=PlanarMovement)
-	uint32 bConstrainToPlane:1;
+	FVector GetPlaneConstraintNormalFromAxisSetting(EPlaneConstraintAxisSetting AxisSetting) const;
 
-	/** If true and plane constraints are enabled, then the updated component will be snapped to the plane when first attached. */
-	UPROPERTY(EditDefaultsOnly, BlueprintReadOnly, Category=PlanarMovement)
-	uint32 bSnapToPlaneAtStart:1;
+public:
 
 	/** If true, skips TickComponent() if UpdatedComponent was not recently rendered. */
 	UPROPERTY(EditAnywhere, BlueprintReadWrite, Category=MovementComponent)
@@ -92,9 +132,17 @@ public:
 
 	// Begin ActorComponent interface 
 	virtual void RegisterComponentTickFunctions(bool bRegister) override;
+	virtual void PostLoad() override;
 
 	/** Overridden to auto-register the updated component if it starts NULL, and we can find a root component on our owner. */
 	virtual void OnRegister() override;
+
+
+#if WITH_EDITOR
+	virtual void PostEditChangeProperty(FPropertyChangedEvent& PropertyChangedEvent) override;
+	static void PhysicsLockedAxisSettingChanged();
+#endif // WITH_EDITOR
+
 	// End ActorComponent interface
 
 	/** @return gravity that affects this component */
@@ -270,8 +318,25 @@ public:
 	 */
 	virtual void AddRadialImpulse(const FVector& Origin, float Radius, float Strength, ERadialImpulseFalloff Falloff, bool bVelChange);
 
+	/**
+	 * Set the plane constraint axis setting.
+	 * Changing this setting will modify the current value of PlaneConstraintNormal.
+	 * 
+	 * @param  NewAxisSetting New plane constraint axis setting.
+	 */
+	UFUNCTION(BlueprintCallable, Category="Components|Movement|Planar")
+	virtual void SetPlaneConstraintAxisSetting(EPlaneConstraintAxisSetting NewAxisSetting);
+
+	/**
+	 * Get the plane constraint axis setting.
+	 */
+	UFUNCTION(BlueprintCallable, Category="Components|Movement|Planar")
+	EPlaneConstraintAxisSetting GetPlaneConstraintAxisSetting() const;
+
 	/**
 	 * Sets the normal of the plane that constrains movement, enforced if the plane constraint is enabled.
+	 * Changing the normal automatically sets PlaneConstraintAxisSetting to "Custom".
+	 *
 	 * @param PlaneNormal	The normal of the plane. If non-zero in length, it will be normalized.
 	 */
 	UFUNCTION(BlueprintCallable, Category="Components|Movement|Planar")
@@ -331,3 +396,8 @@ inline void UMovementComponent::StopMovementImmediately()
 	Velocity = FVector::ZeroVector;
 	UpdateComponentVelocity();
 }
+
+inline EPlaneConstraintAxisSetting UMovementComponent::GetPlaneConstraintAxisSetting() const
+{
+	return PlaneConstraintAxisSetting;
+}
diff --git a/Engine/Source/Runtime/Engine/Private/MovementComponent.cpp b/Engine/Source/Runtime/Engine/Private/MovementComponent.cpp
--- a/Engine/Source/Runtime/Engine/Private/MovementComponent.cpp
+++ b/Engine/Source/Runtime/Engine/Private/MovementComponent.cpp
@@ -10,6 +10,8 @@
 #define LOCTEXT_NAMESPACE "MovementComponent"
 DEFINE_LOG_CATEGORY_STATIC(LogMovement, Log, All);
 
+PRAGMA_DISABLE_OPTIMIZATION
+
 //----------------------------------------------------------------------//
 // UMovementComponent
 //----------------------------------------------------------------------//
@@ -24,6 +26,9 @@ UMovementComponent::UMovementComponent(const FObjectInitializer& ObjectInitializ
 	bUpdateOnlyIfRendered = false;
 	bAutoUpdateTickRegistration = true;
 	bAutoRegisterUpdatedComponent = true;
+
+	PlaneConstraintNormal = FVector::ZeroVector;
+	PlaneConstraintAxisSetting = EPlaneConstraintAxisSetting::Custom;
 	bConstrainToPlane = false;
 	bSnapToPlaneAtStart = false;
 
@@ -101,6 +106,11 @@ void UMovementComponent::OnRegister()
 		}
 	}
 
+	if (PlaneConstraintAxisSetting != EPlaneConstraintAxisSetting::Custom)
+	{
+		SetPlaneConstraintAxisSetting(PlaneConstraintAxisSetting);
+	}
+
 	PlaneConstraintNormal = PlaneConstraintNormal.SafeNormal();
 	SetUpdatedComponent(NewUpdatedComponent);
 }
@@ -130,6 +140,53 @@ void UMovementComponent::UpdateTickRegistration()
 }
 
 
+void UMovementComponent::PostLoad()
+{
+	Super::PostLoad();
+
+	if (PlaneConstraintAxisSetting == EPlaneConstraintAxisSetting::UseGlobalPhysicsSetting)
+	{
+		// Make sure to use the most up-to-date project setting in case it has changed.
+		PlaneConstraintNormal = GetPlaneConstraintNormalFromAxisSetting(PlaneConstraintAxisSetting);
+	}
+}
+
+
+#if WITH_EDITOR
+void UMovementComponent::PostEditChangeProperty(FPropertyChangedEvent& PropertyChangedEvent)
+{
+	Super::PostEditChangeProperty(PropertyChangedEvent);
+
+	const UProperty* PropertyThatChanged = PropertyChangedEvent.MemberProperty;
+	if (PropertyThatChanged)
+	{
+		if (PropertyThatChanged->GetFName() == GET_MEMBER_NAME_CHECKED(UMovementComponent, PlaneConstraintAxisSetting))
+		{
+			PlaneConstraintNormal = GetPlaneConstraintNormalFromAxisSetting(PlaneConstraintAxisSetting);
+		}
+		else if (PropertyThatChanged->GetFName() == GET_MEMBER_NAME_CHECKED(UMovementComponent, PlaneConstraintNormal))
+		{
+			PlaneConstraintAxisSetting = EPlaneConstraintAxisSetting::Custom;
+		}
+	}
+}
+
+void UMovementComponent::PhysicsLockedAxisSettingChanged()
+{
+	
+	for (TObjectIterator<UMovementComponent> Iter; Iter; ++Iter)
+	{
+		UMovementComponent* MovementComponent = *Iter;
+		if (MovementComponent->PlaneConstraintAxisSetting == EPlaneConstraintAxisSetting::UseGlobalPhysicsSetting)
+		{
+			MovementComponent->PlaneConstraintNormal = MovementComponent->GetPlaneConstraintNormalFromAxisSetting(MovementComponent->PlaneConstraintAxisSetting);
+		}
+	}
+}
+
+#endif // WITH_EDITOR
+
+
 void UMovementComponent::PhysicsVolumeChanged(APhysicsVolume* NewVolume)
 {
 	//UE_LOG(LogMovement, Log, TEXT("PhysicsVolumeChanged : No implementation"));
@@ -240,9 +297,46 @@ bool UMovementComponent::IsExceedingMaxSpeed(float MaxSpeed) const
 	return (Velocity.SizeSquared() > MaxSpeedSquared * OverVelocityPercent);
 }
 
+
+FVector UMovementComponent::GetPlaneConstraintNormalFromAxisSetting(EPlaneConstraintAxisSetting AxisSetting) const
+{
+	if (AxisSetting == EPlaneConstraintAxisSetting::UseGlobalPhysicsSetting)
+	{
+		ESettingsLockedAxis::Type GlobalSetting = UPhysicsSettings::Get()->LockedAxis;
+		switch (GlobalSetting)
+		{
+		case ESettingsLockedAxis::None:	return FVector::ZeroVector;
+		case ESettingsLockedAxis::X:	return FVector(1.f, 0.f, 0.f);
+		case ESettingsLockedAxis::Y:	return FVector(0.f, 1.f, 0.f);
+		case ESettingsLockedAxis::Z:	return FVector(0.f, 0.f, 1.f);
+		default:
+			checkf(false, TEXT("GetPlaneConstraintNormalFromAxisSetting: Unknown global axis setting %d for %s"), int32(GlobalSetting), *GetNameSafe(GetOwner()));
+			return FVector::ZeroVector;
+		}
+	}
+
+	switch(AxisSetting)
+	{
+	case EPlaneConstraintAxisSetting::Custom:	return PlaneConstraintNormal;
+	case EPlaneConstraintAxisSetting::X:		return FVector(1.f, 0.f, 0.f);
+	case EPlaneConstraintAxisSetting::Y:		return FVector(0.f, 1.f, 0.f);
+	case EPlaneConstraintAxisSetting::Z:		return FVector(0.f, 0.f, 1.f);
+	default:
+		checkf(false, TEXT("GetPlaneConstraintNormalFromAxisSetting: Unknown axis %d for %s"), int32(AxisSetting), *GetNameSafe(GetOwner()));
+		return FVector::ZeroVector;
+	}
+}
+
+void UMovementComponent::SetPlaneConstraintAxisSetting(EPlaneConstraintAxisSetting NewAxisSetting)
+{
+	PlaneConstraintAxisSetting = NewAxisSetting;
+	PlaneConstraintNormal = GetPlaneConstraintNormalFromAxisSetting(PlaneConstraintAxisSetting);
+}
+
 void UMovementComponent::SetPlaneConstraintNormal(FVector PlaneNormal)
 {
 	PlaneConstraintNormal = PlaneNormal.SafeNormal();
+	PlaneConstraintAxisSetting = EPlaneConstraintAxisSetting::Custom;
 }
 
 void UMovementComponent::SetPlaneConstraintFromVectors(FVector Forward, FVector Up)
diff --git a/Engine/Source/Runtime/Engine/Private/PhysicsEngine/PhysicsSettings.cpp b/Engine/Source/Runtime/Engine/Private/PhysicsEngine/PhysicsSettings.cpp
--- a/Engine/Source/Runtime/Engine/Private/PhysicsEngine/PhysicsSettings.cpp
+++ b/Engine/Source/Runtime/Engine/Private/PhysicsEngine/PhysicsSettings.cpp
@@ -1,6 +1,7 @@
 // Copyright 1998-2014 Epic Games, Inc. All Rights Reserved.
 
 #include "EnginePrivate.h"
+#include "GameFramework/MovementComponent.h"
 
 UPhysicsSettings::UPhysicsSettings(const FObjectInitializer& ObjectInitializer)
 	: Super(ObjectInitializer)
@@ -50,10 +51,14 @@ void UPhysicsSettings::PostEditChangeProperty(struct FPropertyChangedEvent& Prop
 	Super::PostEditChangeProperty(PropertyChangedEvent);
 
 	const FName PropertyName = PropertyChangedEvent.Property ? PropertyChangedEvent.Property->GetFName() : NAME_None;
-	if (PropertyName == FName(TEXT("FrictionCombineMode")))
+	if (PropertyName == GET_MEMBER_NAME_CHECKED(UPhysicsSettings, FrictionCombineMode))
 	{
 		UPhysicalMaterial::RebuildPhysicalMaterials();
 	}
+	else if (PropertyName == GET_MEMBER_NAME_CHECKED(UPhysicsSettings, LockedAxis))
+	{
+		UMovementComponent::PhysicsLockedAxisSettingChanged();
+	}
 }
 
 void UPhysicsSettings::LoadSurfaceType()
