#ue4 - Remove accidental PRAGMA_DISABLE_OPTIMIZATION checkin.

diff --git a/Engine/Source/Runtime/Engine/Private/MovementComponent.cpp b/Engine/Source/Runtime/Engine/Private/MovementComponent.cpp
--- a/Engine/Source/Runtime/Engine/Private/MovementComponent.cpp
+++ b/Engine/Source/Runtime/Engine/Private/MovementComponent.cpp
@@ -10,8 +10,6 @@
 #define LOCTEXT_NAMESPACE "MovementComponent"
 DEFINE_LOG_CATEGORY_STATIC(LogMovement, Log, All);
 
-PRAGMA_DISABLE_OPTIMIZATION
-
 //----------------------------------------------------------------------//
 // UMovementComponent
 //----------------------------------------------------------------------//
