#ue4 - Projectiles behave correctly when the plane constraint is enabled.

- Added UMovementComponent::ConstrainNormalToPlane(), which normalizes the vector after constraining its direction.

diff --git a/Engine/Source/Runtime/Engine/Classes/GameFramework/MovementComponent.h b/Engine/Source/Runtime/Engine/Classes/GameFramework/MovementComponent.h
--- a/Engine/Source/Runtime/Engine/Classes/GameFramework/MovementComponent.h
+++ b/Engine/Source/Runtime/Engine/Classes/GameFramework/MovementComponent.h
@@ -372,6 +372,10 @@ public:
 	UFUNCTION(BlueprintCallable, Category="Components|Movement|Planar")
 	virtual FVector ConstrainLocationToPlane(FVector Location) const;
 
+	/** Constrain a normal vector (of unit length) to the plane constraint, if enabled. */
+	UFUNCTION(BlueprintCallable, Category="Components|Movement|Planar")
+	virtual FVector ConstrainNormalToPlane(FVector Normal) const;
+
 	/** Snap the updated component to the plane constraint, if enabled. */
 	UFUNCTION(BlueprintCallable, Category="Components|Movement|Planar")
 	virtual void SnapUpdatedComponentToPlane();
diff --git a/Engine/Source/Runtime/Engine/Private/MovementComponent.cpp b/Engine/Source/Runtime/Engine/Private/MovementComponent.cpp
--- a/Engine/Source/Runtime/Engine/Private/MovementComponent.cpp
+++ b/Engine/Source/Runtime/Engine/Private/MovementComponent.cpp
@@ -379,6 +379,17 @@ FVector UMovementComponent::ConstrainLocationToPlane(FVector Location) const
 }
 
 
+FVector UMovementComponent::ConstrainNormalToPlane(FVector Normal) const
+{
+	if (bConstrainToPlane)
+	{
+		Normal = FVector::VectorPlaneProject(Normal, PlaneConstraintNormal).SafeNormal();
+	}
+
+	return Normal;
+}
+
+
 void UMovementComponent::SnapUpdatedComponentToPlane()
 {
 	if (UpdatedComponent && bConstrainToPlane)
@@ -527,7 +538,7 @@ FVector UMovementComponent::ComputeSlideVector(const FVector& Delta, const float
 	}
 	else
 	{
-		const FVector ProjectedNormal = ConstrainDirectionToPlane(Normal).SafeNormal();
+		const FVector ProjectedNormal = ConstrainNormalToPlane(Normal);
 		return FVector::VectorPlaneProject(Delta, ProjectedNormal) * Time;
 	}
 }
diff --git a/Engine/Source/Runtime/Engine/Private/ProjectileMovementComponent.cpp b/Engine/Source/Runtime/Engine/Private/ProjectileMovementComponent.cpp
--- a/Engine/Source/Runtime/Engine/Private/ProjectileMovementComponent.cpp
+++ b/Engine/Source/Runtime/Engine/Private/ProjectileMovementComponent.cpp
@@ -173,32 +173,34 @@ void UProjectileMovementComponent::TickComponent(float DeltaTime, enum ELevelTic
 
 			NumBounces++;
 			float SubTickTimeRemaining = TimeTick * (1.f - Hit.Time);
+			FVector Normal = ConstrainNormalToPlane(Hit.Normal);
 
 			// Multiple hits within very short time period?
 			const bool bMultiHit = (PreviousHitTime < 1.f && Hit.Time <= KINDA_SMALL_NUMBER);
 
 			// if velocity still into wall (after HandleHitWall() had a chance to adjust), slide along wall
 			const float DotTolerance = 0.01f;
-			bIsSliding = (bMultiHit && FVector::Coincident(PreviousHitNormal, Hit.Normal)) ||
-						 ((Velocity.SafeNormal() | Hit.Normal) <= DotTolerance);
+			bIsSliding = (bMultiHit && FVector::Coincident(PreviousHitNormal, Normal)) ||
+						 ((Velocity.SafeNormal() | Normal) <= DotTolerance);
 			
 			if (bIsSliding)
 			{
-				if (bMultiHit && (PreviousHitNormal | Hit.Normal) <= 0.f)
+				if (bMultiHit && (PreviousHitNormal | Normal) <= 0.f)
 				{
 					//90 degree or less corner, so use cross product for direction
-					FVector NewDir = (Hit.Normal ^ PreviousHitNormal);
+					FVector NewDir = (Normal ^ PreviousHitNormal);
 					NewDir = NewDir.SafeNormal();
 					Velocity = Velocity.ProjectOnToNormal(NewDir);
 					if ((OldVelocity | Velocity) < 0.f)
 					{
 						Velocity *= -1.f;
 					}
+					Velocity = ConstrainDirectionToPlane(Velocity);
 				}
 				else 
 				{
 					//adjust to move along new wall
-					Velocity = FVector::VectorPlaneProject(Velocity, Hit.Normal);
+					Velocity = ComputeSlideVector(Velocity, 1.f, Normal, Hit);
 				}
 
 				// Check min velocity.
@@ -215,11 +217,12 @@ void UProjectileMovementComponent::TickComponent(float DeltaTime, enum ELevelTic
 					{
 						break;
 					}
+					Normal = ConstrainNormalToPlane(Hit.Normal);
 				}
 			}
 
 			PreviousHitTime = Hit.Time;
-			PreviousHitNormal = Hit.Normal;
+			PreviousHitNormal = Normal;
 			
 			// A few initial bounces should add more time and iterations to complete most of the simulation.
 			if (NumBounces <= 2 && SubTickTimeRemaining >= MIN_TICK_TIME)
@@ -237,7 +240,7 @@ void UProjectileMovementComponent::TickComponent(float DeltaTime, enum ELevelTic
 bool UProjectileMovementComponent::HandleSliding(FHitResult& Hit, float& SubTickTimeRemaining)
 {
 	FHitResult InitialHit(Hit);
-	const FVector OldHitNormal = Hit.Normal;
+	const FVector OldHitNormal = ConstrainDirectionToPlane(Hit.Normal);
 
 	// Velocity is now parallel to the impact surface.
 	// Perform the move now, before adding gravity/accel again, so we don't just keep hitting the surface.
@@ -265,7 +268,7 @@ bool UProjectileMovementComponent::HandleSliding(FHitResult& Hit, float& SubTick
 			const FVector NewVelocity = Velocity + ProjectedForce;
 
 			const FVector FrictionForce = -NewVelocity.SafeNormal() * FMath::Min(-ForceDotN * Friction, NewVelocity.Size());
-			Velocity = NewVelocity + FrictionForce;
+			Velocity = ConstrainDirectionToPlane(NewVelocity + FrictionForce);
 		}
 		else
 		{
@@ -294,6 +297,7 @@ void UProjectileMovementComponent::SetVelocityInLocalSpace(FVector NewVelocity)
 }
 
 
+// Deprecated
 FVector UProjectileMovementComponent::CalculateVelocity(FVector OldVelocity, float DeltaTime, bool bGravityEnabled_UNUSED) const
 {
 	return ComputeVelocity(OldVelocity, DeltaTime);
@@ -317,7 +321,7 @@ FVector UProjectileMovementComponent::LimitVelocity(FVector NewVelocity) const
 		NewVelocity = NewVelocity.ClampMaxSize(CurrentMaxSpeed);
 	}
 
-	return NewVelocity;
+	return ConstrainDirectionToPlane(NewVelocity);
 }
 
 FVector UProjectileMovementComponent::ComputeMoveDelta(const FVector& InVelocity, float DeltaTime) const
@@ -392,13 +396,14 @@ bool UProjectileMovementComponent::HandleHitWall(const FHitResult& Hit, float Ti
 FVector UProjectileMovementComponent::ComputeBounceResult(const FHitResult& Hit, float TimeSlice, const FVector& MoveDelta)
 {
 	FVector TempVelocity = Velocity;
-	const float VDotNormal = (TempVelocity | Hit.Normal);
+	const FVector Normal = ConstrainNormalToPlane(Hit.Normal);
+	const float VDotNormal = (TempVelocity | Normal);
 
 	// Only if velocity is opposed by normal
 	if (VDotNormal < 0.f)
 	{
 		// Project velocity onto normal in reflected direction.
-		const FVector ProjectedNormal = Hit.Normal * -VDotNormal;
+		const FVector ProjectedNormal = Normal * -VDotNormal;
 
 		// Point velocity in direction parallel to surface
 		TempVelocity += ProjectedNormal;
