Fixed Retargeting issue when node/parent ordering mismatched.

diff --git a/Engine/Source/Runtime/Engine/Private/Animation/AnimSequence.cpp b/Engine/Source/Runtime/Engine/Private/Animation/AnimSequence.cpp
--- a/Engine/Source/Runtime/Engine/Private/Animation/AnimSequence.cpp
+++ b/Engine/Source/Runtime/Engine/Private/Animation/AnimSequence.cpp
@@ -2499,99 +2499,145 @@ int32 UAnimSequence::GetSpaceBasedAnimationData(TArray< TArray<FTransform> > & A
 
 		check(Rig);
 
-		// @Todo sort rig control by parent -> child, so that we can calculate parent first before child
-		// but that depends on each skeleton
-		// first put all the component space
+		// to fix the issue where parent of rig doesn't correspond to parent of this skeleton
+		// we do this in multiple iteration if needed. 
+		// this flag will be used to evaluate all of them until done
+		TArray<bool> BoneEvaluated;
+		BoneEvaluated.AddZeroed(NumBones);
 
-		for(int32 BoneIndex = 0; BoneIndex < NumBones; ++BoneIndex)
+		bool bCompleted = false;
+		do
 		{
-			const FName& BoneName = RefSkeleton.GetBoneName(BoneIndex);
-			const FName& NodeName = Skeleton->GetRigNodeNameFromBoneName(BoneName);
-			const FTransformBase* TransformBase = Rig->GetTransformBaseByNodeName(NodeName);
-			const int32 NodeIndex = RiggingAnimationData->TrackNames.Find(NodeName);
-			if (NodeIndex != INDEX_NONE)
+			for(int32 BoneIndex = 0; BoneIndex < NumBones; ++BoneIndex)
 			{
-				check (TransformBase);
-
-				// now calculate the component space
-				const TArray<FRigTransformConstraint>	& RotTransformConstraints = TransformBase->Constraints[EControlConstraint::Type::Orientation].TransformConstraints;
+				if ( !BoneEvaluated[BoneIndex] )
+				{
+					const FName& BoneName = RefSkeleton.GetBoneName(BoneIndex);
+					const FName& NodeName = Skeleton->GetRigNodeNameFromBoneName(BoneName);
+					const FTransformBase* TransformBase = Rig->GetTransformBaseByNodeName(NodeName);
+					const int32 NodeIndex = RiggingAnimationData->TrackNames.Find(NodeName);
+					if(NodeIndex != INDEX_NONE)
+					{
+						check(TransformBase);
 
-				FQuat ComponentRotation;
-				FTransform ComponentTranslation;
-				FVector ComponentScale = FVector(1.f);
+						// now calculate the component space
+						const TArray<FRigTransformConstraint>	& RotTransformConstraints = TransformBase->Constraints[EControlConstraint::Type::Orientation].TransformConstraints;
 
-				//if (RotTransformConstraints.Num() > 0)
-				{
-					const FName& ParentNodeName = RotTransformConstraints[0].ParentSpace;
-					const FName& ParentBoneName = Skeleton->GetRigBoneMapping(ParentNodeName);
-					const int32& ParentBoneIndex = RefSkeleton.FindBoneIndex(ParentBoneName);
+						FQuat ComponentRotation;
+						FTransform ComponentTranslation;
+						FVector ComponentScale = FVector(1.f);
 
-					if(ParentBoneIndex != INDEX_NONE)
-					{
-						for(int32 Key = 0; Key < NumKeys; ++Key)
+						// rotation first
+						// this is easy since we just make sure it's evaluated or not
 						{
-							ComponentRotation = AnimationDataInComponentSpace[ParentBoneIndex][Key].GetRotation() * RiggingAnimationData->AnimationTracks[NodeIndex].RotKeys[Key];
-							AnimationDataInComponentSpace[BoneIndex][Key].SetRotation(ComponentRotation);
-						}
-					}
-					else
-					{
-						for(int32 Key = 0; Key < NumKeys; ++Key)
-						{
-							ComponentRotation = RiggingAnimationData->AnimationTracks[NodeIndex].RotKeys[Key];
-							AnimationDataInComponentSpace[BoneIndex][Key].SetRotation(ComponentRotation);
-						}
-					}
-				}
+							const FName& ParentNodeName = RotTransformConstraints[0].ParentSpace;
+							const FName& ParentBoneName = Skeleton->GetRigBoneMapping(ParentNodeName);
+							const int32& ParentBoneIndex = RefSkeleton.FindBoneIndex(ParentBoneName);
 
-				const TArray<FRigTransformConstraint>	& PosTransformConstraints = TransformBase->Constraints[EControlConstraint::Type::Translation].TransformConstraints;
+							if(ParentBoneIndex != INDEX_NONE)
+							{
+								if (BoneEvaluated[ParentBoneIndex])
+								{
+									for(int32 Key = 0; Key < NumKeys; ++Key)
+									{
+										ComponentRotation = AnimationDataInComponentSpace[ParentBoneIndex][Key].GetRotation() * RiggingAnimationData->AnimationTracks[NodeIndex].RotKeys[Key];
+										AnimationDataInComponentSpace[BoneIndex][Key].SetRotation(ComponentRotation);
+									}
 
-				//if (PosTransformConstraints.Num() > 0)
-				{
-					const FName& ParentNodeName = PosTransformConstraints[0].ParentSpace;
-					const FName& ParentBoneName = Skeleton->GetRigBoneMapping(ParentNodeName);
-					const int32& ParentBoneIndex = RefSkeleton.FindBoneIndex(ParentBoneName);
+									BoneEvaluated[BoneIndex] = true;
+								}
+							}
+							else
+							{
+								for(int32 Key = 0; Key < NumKeys; ++Key)
+								{
+									ComponentRotation = RiggingAnimationData->AnimationTracks[NodeIndex].RotKeys[Key];
+									AnimationDataInComponentSpace[BoneIndex][Key].SetRotation(ComponentRotation);
+								}
 
-					if(ParentBoneIndex != INDEX_NONE)
-					{
-						for(int32 Key = 0; Key < NumKeys; ++Key)
+								BoneEvaluated[BoneIndex] = true;
+							}
+						}
+
+						const TArray<FRigTransformConstraint>	& PosTransformConstraints = TransformBase->Constraints[EControlConstraint::Type::Translation].TransformConstraints;
+
+						// now time to check translation
+						// this is a bit more complicated
+						// since we have to make sure if it's true to start with
+						// did we succeed on getting rotation?
+						if (BoneEvaluated[BoneIndex])
 						{
-							ComponentTranslation = FTransform(RiggingAnimationData->AnimationTracks[NodeIndex].PosKeys[Key]) * AnimationDataInComponentSpace[ParentBoneIndex][Key];
-							AnimationDataInComponentSpace[BoneIndex][Key].SetTranslation(ComponentTranslation.GetTranslation());
-							AnimationDataInComponentSpace[BoneIndex][Key].SetScale3D(ComponentScale);
+							const FName& ParentNodeName = PosTransformConstraints[0].ParentSpace;
+							const FName& ParentBoneName = Skeleton->GetRigBoneMapping(ParentNodeName);
+							const int32& ParentBoneIndex = RefSkeleton.FindBoneIndex(ParentBoneName);
+
+							if(ParentBoneIndex != INDEX_NONE)
+							{
+								// this has to be check
+								if (BoneEvaluated[ParentBoneIndex])
+								{
+									for(int32 Key = 0; Key < NumKeys; ++Key)
+									{
+										ComponentTranslation = FTransform(RiggingAnimationData->AnimationTracks[NodeIndex].PosKeys[Key]) * AnimationDataInComponentSpace[ParentBoneIndex][Key];
+										AnimationDataInComponentSpace[BoneIndex][Key].SetTranslation(ComponentTranslation.GetTranslation());
+										AnimationDataInComponentSpace[BoneIndex][Key].SetScale3D(ComponentScale);
+									}
+								}
+								else
+								{
+									// if we failed to get parent clear the flag
+									// because if translation has been calculated, BoneEvaluated[BoneIndex] might be true
+									BoneEvaluated[BoneIndex] = false;
+								}
+							}
+							else
+							{
+								for(int32 Key = 0; Key < NumKeys; ++Key)
+								{
+									ComponentTranslation = FTransform(RiggingAnimationData->AnimationTracks[NodeIndex].PosKeys[Key]);
+									AnimationDataInComponentSpace[BoneIndex][Key].SetTranslation(ComponentTranslation.GetTranslation());
+									AnimationDataInComponentSpace[BoneIndex][Key].SetScale3D(ComponentScale);
+								}
+							}
 						}
 					}
 					else
 					{
-						for(int32 Key = 0; Key < NumKeys; ++Key)
+						int32 ParentIndex = RefSkeleton.GetParentIndex(BoneIndex);
+						const FTransform& LocalSpace = RefSkeleton.GetRefBonePose()[BoneIndex];
+						if(ParentIndex != INDEX_NONE)
+						{
+							// if parent is evaluated, do it
+							if (BoneEvaluated[ParentIndex])
+							{
+								for(int32 Key = 0; Key < NumKeys; ++Key)
+								{
+									AnimationDataInComponentSpace[BoneIndex][Key] = LocalSpace * AnimationDataInComponentSpace[ParentIndex][Key];
+								}
+
+								BoneEvaluated[BoneIndex] = true;
+							}
+						}
+						else
 						{
-							ComponentTranslation = FTransform(RiggingAnimationData->AnimationTracks[NodeIndex].PosKeys[Key]);
-							AnimationDataInComponentSpace[BoneIndex][Key].SetTranslation(ComponentTranslation.GetTranslation());
-							AnimationDataInComponentSpace[BoneIndex][Key].SetScale3D(ComponentScale);
+							BoneEvaluated[BoneIndex] = true;
+
+							for(int32 Key = 0; Key < NumKeys; ++Key)
+							{
+								AnimationDataInComponentSpace[BoneIndex][Key] = LocalSpace;
+							}
 						}
 					}
 				}
 			}
-			else
+
+			bCompleted = true;
+			// see if we can get out, brute force for now
+			for(int32 BoneIndex = 0; BoneIndex < NumBones && bCompleted; ++BoneIndex)
 			{
-				int32 ParentIndex = RefSkeleton.GetParentIndex(BoneIndex);
-				const FTransform& LocalSpace = RefSkeleton.GetRefBonePose()[BoneIndex];
-				if(ParentIndex != INDEX_NONE)
-				{
-					for(int32 Key = 0; Key < NumKeys; ++Key)
-					{
-						AnimationDataInComponentSpace[BoneIndex][Key] = LocalSpace * AnimationDataInComponentSpace[ParentIndex][Key];
-					}
-				}
-				else
-				{
-					for(int32 Key = 0; Key < NumKeys; ++Key)
-					{
-						AnimationDataInComponentSpace[BoneIndex][Key] = LocalSpace;
-					}
-				}
+				bCompleted &= !!BoneEvaluated[BoneIndex];
 			}
-		}
+		} while (bCompleted == false);
 	}
 	else
 	{
