Physics - Adding a check for the destructible mesh actor to prevent a crash.

diff --git a/Engine/Source/Runtime/Engine/Private/DestructibleComponent.cpp b/Engine/Source/Runtime/Engine/Private/DestructibleComponent.cpp
--- a/Engine/Source/Runtime/Engine/Private/DestructibleComponent.cpp
+++ b/Engine/Source/Runtime/Engine/Private/DestructibleComponent.cpp
@@ -134,7 +134,7 @@ void UDestructibleComponent::CreatePhysicsState()
 {
 	// to avoid calling PrimitiveComponent, I'm just calling ActorComponent::CreatePhysicsState
 	// @todo lh - fix me based on the discussion with Bryan G
- 	UActorComponent::CreatePhysicsState();
+	UActorComponent::CreatePhysicsState();
 	bPhysicsStateCreated = true;
 
 	// What we want to do with BodySetup is simply use it to store a PhysicalMaterial, and possibly some other relevant fields.  Set up pointers from the BodyInstance to the BodySetup and this component
@@ -374,7 +374,6 @@ void UDestructibleComponent::DestroyPhysicsState()
 		//In theory anyone using BodyInstance on a PrimitiveComponent should be using functions like GetBodyInstance - in which case we properly fix up the dangling pointer
 		BodyInstance.RigidActorSync = NULL;
 		BodyInstance.RigidActorAsync = NULL;
-
 	}
 #endif	// #if WITH_APEX
 	Super::DestroyPhysicsState();
@@ -404,14 +403,17 @@ bool UDestructibleComponent::CanEditSimulatePhysics()
 void UDestructibleComponent::AddImpulse( FVector Impulse, FName BoneName /*= NAME_None*/, bool bVelChange /*= false*/ )
 {
 #if WITH_APEX
-	int32 ChunkIdx = BoneIdxToChunkIdx(GetBoneIndex(BoneName));
-	PxRigidDynamic* PActor = ApexDestructibleActor->getChunkPhysXActor(ChunkIdx);
-
-	if (PActor != NULL)
+	if ( ApexDestructibleActor )
 	{
-		SCOPED_SCENE_WRITE_LOCK(PActor->getScene());
+		int32 ChunkIdx = BoneIdxToChunkIdx(GetBoneIndex(BoneName));
+		PxRigidDynamic* PActor = ApexDestructibleActor->getChunkPhysXActor(ChunkIdx);
 
-		PActor->addForce(U2PVector(Impulse), bVelChange ? PxForceMode::eVELOCITY_CHANGE : PxForceMode::eIMPULSE);
+		if ( PActor != NULL )
+		{
+			SCOPED_SCENE_WRITE_LOCK(PActor->getScene());
+
+			PActor->addForce(U2PVector(Impulse), bVelChange ? PxForceMode::eVELOCITY_CHANGE : PxForceMode::eIMPULSE);
+		}
 	}
 #endif
 }
