Particles - Changing an InverseFast to an Inverse to prevent a crash when it has a null matrix.

diff --git a/Engine/Source/Runtime/Engine/Public/ParticleHelper.h b/Engine/Source/Runtime/Engine/Public/ParticleHelper.h
--- a/Engine/Source/Runtime/Engine/Public/ParticleHelper.h
+++ b/Engine/Source/Runtime/Engine/Public/ParticleHelper.h
@@ -2595,7 +2595,7 @@ public:
 
 	// While this isn't good OO design, access to everything is made public.
 	// This is to allow custom emitter instances to easily be written when extending the engine.
-	FMatrix GetWorldToLocal() const		{	return GetLocalToWorld().InverseFast();	}
+	FMatrix GetWorldToLocal() const		{	return GetLocalToWorld().Inverse();	}
 	bool GetCastShadow() const			{	return bCastShadow;				}
 	const FMaterialRelevance& GetMaterialRelevance() const
 	{
