Fixing build with no WITH_RECAST

diff --git a/Engine/Source/Runtime/Engine/Classes/AI/Navigation/RecastNavMeshDataChunk.h b/Engine/Source/Runtime/Engine/Classes/AI/Navigation/RecastNavMeshDataChunk.h
--- a/Engine/Source/Runtime/Engine/Classes/AI/Navigation/RecastNavMeshDataChunk.h
+++ b/Engine/Source/Runtime/Engine/Classes/AI/Navigation/RecastNavMeshDataChunk.h
@@ -54,7 +54,9 @@ public:
 #endif// WITH_EDITOR
 
 private:
+#if WITH_RECAST
 	void SerializeRecastData(FArchive& Ar);
+#endif//WITH_RECAST
 
 private:
 	TArray<FRecastTileData> Tiles;
diff --git a/Engine/Source/Runtime/Engine/Private/AI/Navigation/RecastNavMeshDataChunk.cpp b/Engine/Source/Runtime/Engine/Private/AI/Navigation/RecastNavMeshDataChunk.cpp
--- a/Engine/Source/Runtime/Engine/Private/AI/Navigation/RecastNavMeshDataChunk.cpp
+++ b/Engine/Source/Runtime/Engine/Private/AI/Navigation/RecastNavMeshDataChunk.cpp
@@ -61,10 +61,12 @@ void URecastNavMeshDataChunk::Serialize(FArchive& Ar)
 			// incompatible, just skip over this data.  navmesh needs rebuilt.
 			Ar.Seek(RecastNavMeshSizePos + RecastNavMeshSizeBytes);
 		}
+#if WITH_RECAST
 		else if (RecastNavMeshSizeBytes > 4)
 		{
 			SerializeRecastData(Ar);
 		}
+#endif// WITH_RECAST
 		else
 		{
 			// empty, just skip over this data
@@ -73,7 +75,9 @@ void URecastNavMeshDataChunk::Serialize(FArchive& Ar)
 	}
 	else
 	{
+#if WITH_RECAST
 		SerializeRecastData(Ar);
+#endif// WITH_RECAST
 
 		if (Ar.IsSaving())
 		{
@@ -86,6 +90,7 @@ void URecastNavMeshDataChunk::Serialize(FArchive& Ar)
 	}
 }
 
+#if WITH_RECAST
 void URecastNavMeshDataChunk::SerializeRecastData(FArchive& Ar)
 {
 	int32 TileNum = Tiles.Num();
@@ -121,12 +126,14 @@ void URecastNavMeshDataChunk::SerializeRecastData(FArchive& Ar)
 		}
 	}
 }
+#endif// WITH_RECAST
 
 TArray<uint32> URecastNavMeshDataChunk::AttachTiles(dtNavMesh* NavMesh)
 {
 	TArray<uint32> Result;
 	Result.Reserve(Tiles.Num());
-	
+
+#if WITH_RECAST	
 	for (FRecastTileData& TileData : Tiles)
 	{
 		if (TileData.TileRawData.IsValid())
@@ -145,6 +152,7 @@ TArray<uint32> URecastNavMeshDataChunk::AttachTiles(dtNavMesh* NavMesh)
 			Result.Add(NavMesh->decodePolyIdTile(TileData.TileRef));
 		}
 	}
+#endif// WITH_RECAST
 
 	UE_LOG(LogNavigation, Log, TEXT("Attached %d tiles to NavMesh - %s"), Result.Num(), *NavigationDataName.ToString());
 	return Result;
@@ -155,6 +163,7 @@ TArray<uint32> URecastNavMeshDataChunk::DetachTiles(dtNavMesh* NavMesh)
 	TArray<uint32> Result;
 	Result.Reserve(Tiles.Num());
 
+#if WITH_RECAST	
 	for (FRecastTileData& TileData : Tiles)
 	{
 		if (TileData.TileRef != 0)
@@ -163,6 +172,7 @@ TArray<uint32> URecastNavMeshDataChunk::DetachTiles(dtNavMesh* NavMesh)
 			Result.Add(NavMesh->decodePolyIdTile(TileData.TileRef));
 		}
 	}
+#endif// WITH_RECAST
 
 	UE_LOG(LogNavigation, Log, TEXT("Detached %d tiles from NavMesh - %s"), Result.Num(), *NavigationDataName.ToString());
 	return Result;
