Fixed crash when using a runtime cached font with a text actor

These type of fonts can't currently be measured on the render thread, as they use the default Slate font cache which can only be accessed via Slate application on the game thread.

For now I've just disabled generating any render resources when using a runtime cached font.

diff --git a/Engine/Source/Runtime/Engine/Private/TextRenderComponent.cpp b/Engine/Source/Runtime/Engine/Private/TextRenderComponent.cpp
--- a/Engine/Source/Runtime/Engine/Private/TextRenderComponent.cpp
+++ b/Engine/Source/Runtime/Engine/Private/TextRenderComponent.cpp
@@ -432,6 +432,13 @@ FTextRenderSceneProxy::~FTextRenderSceneProxy()
 
 void FTextRenderSceneProxy::CreateRenderThreadResources()
 {
+	if(Font && Font->FontCacheType == EFontCacheType::Runtime)
+	{
+		// Runtime fonts can't currently be used here as they use the font cache from Slate application
+		// which can only be used on the game thread
+		return;
+	}
+
 	if(BuildStringMesh(VertexBuffer.Vertices, IndexBuffer.Indices))
 	{
 		// Init vertex factory
