UMG - Re-enabling a permanant link to UMG even in server/ program builds.  Currently prevents cooking to leave it out, will have to make future efforts to prevent linking too much into dedicated server builds.

diff --git a/Engine/Source/Runtime/Engine/Engine.Build.cs b/Engine/Source/Runtime/Engine/Engine.Build.cs
--- a/Engine/Source/Runtime/Engine/Engine.Build.cs
+++ b/Engine/Source/Runtime/Engine/Engine.Build.cs
@@ -80,20 +80,10 @@ public class Engine : ModuleRules
 				"MovieSceneCoreTypes",
 				"HeadMountedDisplay",
 				"StreamingPauseRendering",
+                "UMG",
 			}
 		);
 
-        // Prevent adding UMG as a dependency in dedicated server builds.
-        if (!(Target.Type == TargetRules.TargetType.Server || Target.Type == TargetRules.TargetType.Program))
-        {
-            DynamicallyLoadedModuleNames.AddRange(
-                new string[]
-			    {
-				    "UMG",
-			    }
-            );
-        }
-
 		if (UEBuildConfiguration.bBuildDeveloperTools)
 		{
 			// Add "BlankModule" so that it gets compiled as an example and will be maintained and tested.  This can be removed
