Tooltip fix for github pr# 552 by user stephenwittle

diff --git a/Engine/Source/Runtime/Engine/Classes/Kismet/KismetMathLibrary.h b/Engine/Source/Runtime/Engine/Classes/Kismet/KismetMathLibrary.h
--- a/Engine/Source/Runtime/Engine/Classes/Kismet/KismetMathLibrary.h
+++ b/Engine/Source/Runtime/Engine/Classes/Kismet/KismetMathLibrary.h
@@ -458,7 +458,7 @@ class UKismetMathLibrary : public UBlueprintFunctionLibrary
 	UFUNCTION(BlueprintPure, Category="Math|Float")
 	static int32 FCeil(float A);
 
-	/* Returns the number of times A will go into B, as well as the remainder */
+	/* Returns the number of times Divisor will go into Dividend (i.e., Dividend divided by Divisor), as well as the remainder */
 	UFUNCTION(BlueprintPure, meta=(FriendlyName = "Division (whole and remainder)"), Category="Math|Float")
 	static int32 FMod(float Dividend, float Divisor, float& Remainder);
 
