Added gamma corrected Canvas rendering for runtime cached fonts

diff --git a/Engine/Shaders/SimpleElementPixelShader.usf b/Engine/Shaders/SimpleElementPixelShader.usf
--- a/Engine/Shaders/SimpleElementPixelShader.usf
+++ b/Engine/Shaders/SimpleElementPixelShader.usf
@@ -153,6 +153,41 @@ void GammaMain(
 #endif
 }
 
+void GammaAlphaOnlyMain(
+	in float2 TextureCoordinate : TEXCOORD0,
+	in float4 Color : TEXCOORD1,
+	in float4 InHitProxyId : TEXCOORD2,			// needed to have d3ddebug warning supressed (VS and PS signatures with SV_Position
+	out float4 OutColor : SV_Target0
+#if FEATURE_LEVEL >= FEATURE_LEVEL_SM4
+	,in float4 SvPosition : SV_Position
+	,out float4 OutWorldNormal : SV_Target1
+#endif
+	)
+{
+#if FEATURE_LEVEL >= FEATURE_LEVEL_SM4
+	// depth test manually if compositing editor primitives
+	FLATTEN
+	if (bEnableEditorPrimitiveDepthTest)
+	{
+		PerformEditorCompositingDepthTest(SvPosition);
+	}
+#endif
+
+	float4 TmpColor = Color;
+	TmpColor.a *= Texture2DSample_A8(InTexture, InTextureSampler, TextureCoordinate);
+	if( Gamma != 1.0 )
+	{
+		// Gamma correct the output color.
+		TmpColor.rgb = pow(saturate(TmpColor.rgb),Gamma);	
+	}
+	OutColor = RETURN_COLOR(TmpColor);
+
+#if FEATURE_LEVEL >= FEATURE_LEVEL_SM4
+	// Set the G buffer bits that indicate unlit
+	OutWorldNormal = 0;
+#endif
+}
+
 float ClipRef;
 
 void GammaMaskedMain(
diff --git a/Engine/Source/Runtime/Engine/Private/BatchedElements.cpp b/Engine/Source/Runtime/Engine/Private/BatchedElements.cpp
--- a/Engine/Source/Runtime/Engine/Private/BatchedElements.cpp
+++ b/Engine/Source/Runtime/Engine/Private/BatchedElements.cpp
@@ -365,6 +365,7 @@ FGlobalBoundShaderState FBatchedElements::DistanceFieldBoundShaderState;
 FGlobalBoundShaderState FBatchedElements::HitTestingBoundShaderState;
 FGlobalBoundShaderState FBatchedElements::ColorChannelMaskShaderState;
 FGlobalBoundShaderState FBatchedElements::AlphaOnlyShaderState;
+FGlobalBoundShaderState FBatchedElements::GammaAlphaOnlyShaderState;
 
 /** Global alpha ref test value for rendering masked batched elements */
 float GBatchedElementAlphaRefVal = 128.f;
@@ -552,12 +553,24 @@ void FBatchedElements::PrepareShaders(
 			{
 				SetBlendState(RHICmdList, BlendMode);
 
-				TShaderMapRef<FSimpleElementAlphaOnlyPS> AlphaOnlyPixelShader(GetGlobalShaderMap(FeatureLevel));
-				SetGlobalBoundShaderState(RHICmdList, FeatureLevel, AlphaOnlyShaderState, GSimpleElementVertexDeclaration.VertexDeclarationRHI,
-					*VertexShader, *AlphaOnlyPixelShader);
+				if (FMath::Abs(Gamma - 1.0f) < KINDA_SMALL_NUMBER)
+			    {
+					TShaderMapRef<FSimpleElementAlphaOnlyPS> AlphaOnlyPixelShader(GetGlobalShaderMap(FeatureLevel));
+					SetGlobalBoundShaderState(RHICmdList, FeatureLevel, AlphaOnlyShaderState, GSimpleElementVertexDeclaration.VertexDeclarationRHI,
+						*VertexShader, *AlphaOnlyPixelShader);
+
+					AlphaOnlyPixelShader->SetParameters(RHICmdList, Texture);
+					AlphaOnlyPixelShader->SetEditorCompositingParameters(RHICmdList, View, DepthTexture);
+			    }
+			    else
+			    {
+					TShaderMapRef<FSimpleElementGammaAlphaOnlyPS> GammaAlphaOnlyPixelShader(GetGlobalShaderMap(FeatureLevel));
+					SetGlobalBoundShaderState(RHICmdList, FeatureLevel, GammaAlphaOnlyShaderState, GSimpleElementVertexDeclaration.VertexDeclarationRHI,
+						*VertexShader, *GammaAlphaOnlyPixelShader);
 
-				AlphaOnlyPixelShader->SetParameters(RHICmdList, Texture);
-				AlphaOnlyPixelShader->SetEditorCompositingParameters(RHICmdList, View, DepthTexture);
+					GammaAlphaOnlyPixelShader->SetParameters(RHICmdList, Texture, Gamma, BlendMode);
+					GammaAlphaOnlyPixelShader->SetEditorCompositingParameters(RHICmdList, View, DepthTexture);
+			    }
 			}
 			else if(BlendMode >= SE_BLEND_RGBA_MASK_START && BlendMode <= SE_BLEND_RGBA_MASK_END)
 			{
diff --git a/Engine/Source/Runtime/Engine/Private/SimpleElementShaders.cpp b/Engine/Source/Runtime/Engine/Private/SimpleElementShaders.cpp
--- a/Engine/Source/Runtime/Engine/Private/SimpleElementShaders.cpp
+++ b/Engine/Source/Runtime/Engine/Private/SimpleElementShaders.cpp
@@ -300,6 +300,7 @@ IMPLEMENT_SHADER_TYPE(,FSimpleElementPS, TEXT("SimpleElementPixelShader"), TEXT(
 IMPLEMENT_SHADER_TYPE(,FSimpleElementAlphaOnlyPS, TEXT("SimpleElementPixelShader"), TEXT("AlphaOnlyMain"), SF_Pixel);
 IMPLEMENT_SHADER_TYPE(template<>, FSimpleElementGammaPS_SRGB, TEXT("SimpleElementPixelShader"), TEXT("GammaMain"), SF_Pixel);
 IMPLEMENT_SHADER_TYPE(template<>, FSimpleElementGammaPS_Linear, TEXT("SimpleElementPixelShader"), TEXT("GammaMain"), SF_Pixel);
+IMPLEMENT_SHADER_TYPE(, FSimpleElementGammaAlphaOnlyPS, TEXT("SimpleElementPixelShader"), TEXT("GammaAlphaOnlyMain"), SF_Pixel);
 IMPLEMENT_SHADER_TYPE(template<>, FSimpleElementMaskedGammaPS_SRGB, TEXT("SimpleElementPixelShader"), TEXT("GammaMaskedMain"), SF_Pixel);
 IMPLEMENT_SHADER_TYPE(template<>, FSimpleElementMaskedGammaPS_Linear, TEXT("SimpleElementPixelShader"), TEXT("GammaMaskedMain"), SF_Pixel);
 IMPLEMENT_SHADER_TYPE(,FSimpleElementDistanceFieldGammaPS,TEXT("SimpleElementPixelShader"),TEXT("GammaDistanceFieldMain"),SF_Pixel);
diff --git a/Engine/Source/Runtime/Engine/Private/UserInterface/CanvasItem.cpp b/Engine/Source/Runtime/Engine/Private/UserInterface/CanvasItem.cpp
--- a/Engine/Source/Runtime/Engine/Private/UserInterface/CanvasItem.cpp
+++ b/Engine/Source/Runtime/Engine/Private/UserInterface/CanvasItem.cpp
@@ -801,6 +801,7 @@ void FCanvasTextItem::Draw( class FCanvas* InCanvas )
 	if (GetFontCacheType() == EFontCacheType::Runtime)
 	{
 		// The runtime font cache uses an alpha-only texture, so we have to force this blend mode so we use the correct shader
+		check(BlendMode == SE_BLEND_Translucent || BlendMode == SE_BLEND_TranslucentAlphaOnly);
 		BlendMode = SE_BLEND_TranslucentAlphaOnly;
 	}
 
diff --git a/Engine/Source/Runtime/Engine/Public/BatchedElements.h b/Engine/Source/Runtime/Engine/Public/BatchedElements.h
--- a/Engine/Source/Runtime/Engine/Public/BatchedElements.h
+++ b/Engine/Source/Runtime/Engine/Public/BatchedElements.h
@@ -311,6 +311,8 @@ private:
 	static FGlobalBoundShaderState ColorChannelMaskShaderState;
 	/** bound shader state for alpha only fonts */
 	static FGlobalBoundShaderState AlphaOnlyShaderState;
+	/** bound shader state for gamma corrected alpha only fonts */
+	static FGlobalBoundShaderState GammaAlphaOnlyShaderState;
 
 	/**
 	 * Sets the appropriate vertex and pixel shader.
diff --git a/Engine/Source/Runtime/Engine/Public/SimpleElementShaders.h b/Engine/Source/Runtime/Engine/Public/SimpleElementShaders.h
--- a/Engine/Source/Runtime/Engine/Public/SimpleElementShaders.h
+++ b/Engine/Source/Runtime/Engine/Public/SimpleElementShaders.h
@@ -118,6 +118,20 @@ public:
 	}
 };
 
+/**
+ * Simple pixel shader that just reads from an alpha-only texture and gamma corrects the output
+ */
+class FSimpleElementGammaAlphaOnlyPS : public FSimpleElementGammaBasePS
+{
+	DECLARE_SHADER_TYPE(FSimpleElementGammaAlphaOnlyPS, Global);
+public:
+
+	FSimpleElementGammaAlphaOnlyPS(const ShaderMetaType::CompiledShaderInitializerType& Initializer) : FSimpleElementGammaBasePS(Initializer) {}
+	FSimpleElementGammaAlphaOnlyPS() {}
+
+	static bool ShouldCache(EShaderPlatform Platform) { return true; }
+};
+
 /**
  * A pixel shader for rendering a masked texture on a simple element.
  */
