Clearing tile cache on navmesh reallocation and fixed tile cache memory stats

diff --git a/Engine/Source/Runtime/Engine/Classes/AI/Navigation/RecastNavMesh.h b/Engine/Source/Runtime/Engine/Classes/AI/Navigation/RecastNavMesh.h
--- a/Engine/Source/Runtime/Engine/Classes/AI/Navigation/RecastNavMesh.h
+++ b/Engine/Source/Runtime/Engine/Classes/AI/Navigation/RecastNavMesh.h
@@ -336,17 +336,11 @@ struct FNavMeshTileData
 	// actual tile data
 	TSharedPtr<FNavData> NavData;
 	
-
 	FNavMeshTileData() : LayerIndex(0), DataSize(0) { }
-
-	explicit FNavMeshTileData(uint8* RawData, int32 RawDataSize, int32 LayerIdx = 0, FBox LayerBounds = FBox(0))
-		: LayerIndex(LayerIdx)
-		, LayerBBox(LayerBounds)
-		, DataSize(RawDataSize)
-	{
-		NavData = MakeShareable(new FNavData(RawData));
-	}
+	~FNavMeshTileData();
 	
+	explicit FNavMeshTileData(uint8* RawData, int32 RawDataSize, int32 LayerIdx = 0, FBox LayerBounds = FBox(0));
+		
 	FORCEINLINE uint8* GetData()
 	{
 		check(NavData.IsValid());
diff --git a/Engine/Source/Runtime/Engine/Private/AI/Navigation/RecastNavMesh.cpp b/Engine/Source/Runtime/Engine/Private/AI/Navigation/RecastNavMesh.cpp
--- a/Engine/Source/Runtime/Engine/Private/AI/Navigation/RecastNavMesh.cpp
+++ b/Engine/Source/Runtime/Engine/Private/AI/Navigation/RecastNavMesh.cpp
@@ -46,14 +46,32 @@ FNavMeshTileData::FNavData::~FNavData()
 #endif
 }
 
+FNavMeshTileData::FNavMeshTileData(uint8* RawData, int32 RawDataSize, int32 LayerIdx, FBox LayerBounds)
+	: LayerIndex(LayerIdx)
+	, LayerBBox(LayerBounds)
+	, DataSize(RawDataSize)
+{
+	INC_MEMORY_STAT_BY(STAT_Navigation_TileCacheMemory, DataSize);
+	NavData = MakeShareable(new FNavData(RawData));
+}
+
+FNavMeshTileData::~FNavMeshTileData()
+{
+	if (NavData.IsUnique() && NavData->RawNavData)
+	{
+		DEC_MEMORY_STAT_BY(STAT_Navigation_TileCacheMemory, DataSize);
+	}
+}
+
 uint8* FNavMeshTileData::Release()
 {
 	uint8* RawData = nullptr;
 
-	if (NavData.IsValid()) 
+	if (NavData.IsValid() && NavData->RawNavData) 
 	{ 
 		RawData = NavData->RawNavData;
-		NavData->RawNavData = nullptr; 
+		NavData->RawNavData = nullptr;
+		DEC_MEMORY_STAT_BY(STAT_Navigation_TileCacheMemory, DataSize);
 	} 
  
 	DataSize = 0; 
diff --git a/Engine/Source/Runtime/Engine/Private/AI/Navigation/RecastNavMeshGenerator.cpp b/Engine/Source/Runtime/Engine/Private/AI/Navigation/RecastNavMeshGenerator.cpp
--- a/Engine/Source/Runtime/Engine/Private/AI/Navigation/RecastNavMeshGenerator.cpp
+++ b/Engine/Source/Runtime/Engine/Private/AI/Navigation/RecastNavMeshGenerator.cpp
@@ -1829,11 +1829,9 @@ bool FRecastTileGenerator::GenerateTile()
 
 	if (bRegenerateCompressedLayers)
 	{
-		DEC_MEMORY_STAT_BY(STAT_Navigation_TileCacheMemory, GetTileCacheSizeHelper(CompressedLayers));
 		CompressedLayers.Reset();
 		
 		bSuccess = GenerateCompressedLayers(BuildContext);
-		INC_MEMORY_STAT_BY(STAT_Navigation_TileCacheMemory, GetTileCacheSizeHelper(CompressedLayers));
 
 		if (bSuccess)
 		{
@@ -2719,7 +2717,6 @@ FRecastNavMeshGenerator::FRecastNavMeshGenerator(class ARecastNavMesh* InDestNav
 	, Version(0)
 {
 	INC_DWORD_STAT_BY( STAT_NavigationMemory, sizeof(*this) );
-	SET_MEMORY_STAT(STAT_Navigation_TileCacheMemory, 0);
 
 	check(InDestNavMesh);
 	Init();
@@ -2731,7 +2728,6 @@ FRecastNavMeshGenerator::~FRecastNavMeshGenerator()
 	DiscardCurrentBuildingTasks();
 	
 	DEC_DWORD_STAT_BY( STAT_NavigationMemory, sizeof(*this) );
-	SET_MEMORY_STAT(STAT_Navigation_TileCacheMemory, 0);
 }
 
 void FRecastNavMeshGenerator::Init()
@@ -2840,7 +2836,10 @@ void FRecastNavMeshGenerator::UpdateNavigationBounds()
 bool FRecastNavMeshGenerator::ConstructTiledNavMesh() 
 {
 	bool bSuccess = false;
-	
+
+	// There is should not be any active build tasks
+	CancelBuild();
+
 	// create new Detour navmesh instance
 	dtNavMesh* DetourMesh = dtAllocNavMesh();	
 	if (DetourMesh)
@@ -2962,7 +2961,6 @@ void FRecastNavMeshGenerator::CancelBuild()
 	RunningDirtyTiles.Empty();
 	PendingDirtyTiles.Empty();
 	IntermediateLayerDataMap.Empty();
-	SET_MEMORY_STAT(STAT_Navigation_TileCacheMemory, 0);
 
 #if	WITH_EDITOR	
 	RecentlyBuiltTiles.Empty();
