Fix incorrect ShapeBoundToBody logic in FBodyInstance::OverlapPhysX

diff --git a/Engine/Source/Runtime/Engine/Private/PhysicsEngine/BodyInstance.cpp b/Engine/Source/Runtime/Engine/Private/PhysicsEngine/BodyInstance.cpp
--- a/Engine/Source/Runtime/Engine/Private/PhysicsEngine/BodyInstance.cpp
+++ b/Engine/Source/Runtime/Engine/Private/PhysicsEngine/BodyInstance.cpp
@@ -3709,7 +3709,7 @@ bool FBodyInstance::OverlapPhysX(const PxGeometry& PGeom, const PxTransform& Sha
 		const PxShape* PShape = PShapes[ShapeIdx];
 		check(PShape);
 
-		if (ShapeBoundToBody(PShape, this) == false)
+		if (ShapeBoundToBody(PShape, this) == true)
 		{
 			if (PxGeometryQuery::overlap(PShape->getGeometry().any(), PxShapeExt::getGlobalPose(*PShape, *RigidBody), PGeom, ShapePose))
 			{
