Fix when end delegate is called, it calls with null montage.

diff --git a/Engine/Source/Runtime/Engine/Private/Animation/AnimMontage.cpp b/Engine/Source/Runtime/Engine/Private/Animation/AnimMontage.cpp
--- a/Engine/Source/Runtime/Engine/Private/Animation/AnimMontage.cpp
+++ b/Engine/Source/Runtime/Engine/Private/Animation/AnimMontage.cpp
@@ -887,7 +887,7 @@ void FAnimMontageInstance::Terminate()
 	// terminating, trigger end
 	if (UAnimInstance* Inst = AnimInstance.Get())
 	{
-		Inst->OnMontageEnded.Broadcast(Montage, bInterrupted);
+		Inst->OnMontageEnded.Broadcast(OldMontage, bInterrupted);
 	}
 	OnMontageEnded.ExecuteIfBound(OldMontage, bInterrupted);
 }
