Make const

diff --git a/Engine/Source/Runtime/Engine/Private/Collision/CollisionConversions.cpp b/Engine/Source/Runtime/Engine/Private/Collision/CollisionConversions.cpp
--- a/Engine/Source/Runtime/Engine/Private/Collision/CollisionConversions.cpp
+++ b/Engine/Source/Runtime/Engine/Private/Collision/CollisionConversions.cpp
@@ -423,7 +423,7 @@ void ConvertRaycastResults(int32 NumHits, PxRaycastHit* Hits, float CheckLength,
 	OutHits.Sort( FCompareFHitResultTime() );
 }
 
-bool AddSweepResults(int32 NumHits, PxSweepHit* Hits, float CheckLength, const PxFilterData& QueryFilter, TArray<FHitResult>& OutHits, const FVector& StartLoc, const FVector& EndLoc, const PxGeometry& Geom, const PxTransform& QueryTM, float MaxDistance, bool bReturnPhysMat)
+bool AddSweepResults(int32 NumHits, const PxSweepHit* Hits, float CheckLength, const PxFilterData& QueryFilter, TArray<FHitResult>& OutHits, const FVector& StartLoc, const FVector& EndLoc, const PxGeometry& Geom, const PxTransform& QueryTM, float MaxDistance, bool bReturnPhysMat)
 {
 	OutHits.Reserve(NumHits);
 	bool bHadBlockingHit = false;
diff --git a/Engine/Source/Runtime/Engine/Private/Collision/CollisionConversions.h b/Engine/Source/Runtime/Engine/Private/Collision/CollisionConversions.h
--- a/Engine/Source/Runtime/Engine/Private/Collision/CollisionConversions.h
+++ b/Engine/Source/Runtime/Engine/Private/Collision/CollisionConversions.h
@@ -48,7 +48,7 @@ void ConvertRaycastResults(int32 NumHits, PxRaycastHit* Hits, float CheckLength,
  * @param	Geom
  * @return	true if any blocking hit was found within MaxDistance.
  */
-bool AddSweepResults(int32 NumHits, PxSweepHit* Hits, float CheckLength, const PxFilterData& QueryFilter, TArray<FHitResult>& OutHits, const FVector& StartLoc, const FVector& EndLoc, const PxGeometry& Geom, const PxTransform& QueryTM, float MaxDistance, bool bReturnPhysMat);
+bool AddSweepResults(int32 NumHits, const PxSweepHit* Hits, float CheckLength, const PxFilterData& QueryFilter, TArray<FHitResult>& OutHits, const FVector& StartLoc, const FVector& EndLoc, const PxGeometry& Geom, const PxTransform& QueryTM, float MaxDistance, bool bReturnPhysMat);
 
 /** 
  * Util to convert physX overlap query to our overlap result
