Temporary fix for CCD issue when going from simulated to kinematic

diff --git a/Engine/Source/Runtime/Engine/Private/PhysicsEngine/BodyInstance.cpp b/Engine/Source/Runtime/Engine/Private/PhysicsEngine/BodyInstance.cpp
--- a/Engine/Source/Runtime/Engine/Private/PhysicsEngine/BodyInstance.cpp
+++ b/Engine/Source/Runtime/Engine/Private/PhysicsEngine/BodyInstance.cpp
@@ -652,7 +652,7 @@ void FBodyInstance::UpdatePhysicsShapeFilterData(uint32 SkelMeshCompID, bool bUs
 			PxFilterData PComplexQueryData;
 			if (UseCollisionEnabled != ECollisionEnabled::NoCollision)
 			{
-				CreateShapeFilterData(BI->ObjectType, OwnerID, UseResponse, SkelMeshCompID, BI->InstanceBodyIndex, PSimpleQueryData, PSimFilterData, BI->bUseCCD && !bPhysicsStatic, bUseNotify, bPhysicsStatic);
+				CreateShapeFilterData(BI->ObjectType, OwnerID, UseResponse, SkelMeshCompID, BI->InstanceBodyIndex, PSimpleQueryData, PSimFilterData, bUseCCD && !bPhysicsStatic, bUseNotify, bPhysicsStatic);	//CCD is determined by root body in case of welding
 				PComplexQueryData = PSimpleQueryData;
 
 				// Build filterdata variations for complex and simple
@@ -2102,6 +2102,7 @@ void FBodyInstance::UpdateInstanceSimulatePhysics()
 		{
 			SCOPED_SCENE_WRITE_LOCK(PRigidDynamic->getScene());
 			PRigidDynamic->setRigidDynamicFlag(PxRigidDynamicFlag::eKINEMATIC, bNewKinematic);
+			PRigidDynamic->setRigidBodyFlag(PxRigidBodyFlag::eENABLE_CCD, !bNewKinematic && bUseCCD);
 
 			//if wake when level starts is true, calling this function automatically wakes body up
 			if (bSimulatePhysics && bStartAwake)
