fixed BoundsScale for TextRenderComponent

diff --git a/Engine/Source/Runtime/Engine/Private/TextRenderComponent.cpp b/Engine/Source/Runtime/Engine/Private/TextRenderComponent.cpp
--- a/Engine/Source/Runtime/Engine/Private/TextRenderComponent.cpp
+++ b/Engine/Source/Runtime/Engine/Private/TextRenderComponent.cpp
@@ -845,7 +845,12 @@ FBoxSphereBounds UTextRenderComponent::CalcBounds(const FTransform& LocalToWorld
 		LeftTop.Y = ComputeVerticalAlignmentOffset(Size.Y, VerticalAlignment, FirstLineHeight);
 		FBox LocalBox(FVector(0, -LeftTop.X, -LeftTop.Y), FVector(0, -(LeftTop.X + Size.X), -(LeftTop.Y + Size.Y)));
 
-		return FBoxSphereBounds(LocalBox.TransformBy(LocalToWorld));
+		FBoxSphereBounds Ret(LocalBox.TransformBy(LocalToWorld));
+
+		Ret.BoxExtent *= BoundsScale;
+		Ret.SphereRadius *= BoundsScale;
+
+		return Ret;
 	}
 	else
 	{
