Niagara: effect instance class, proper instantiation via ActorFactory, surrounding infrastructure and corresponding editor changes, various bug fixes.

diff --git a/Engine/Source/Editor/NiagaraEditor/Private/NiagaraCompiler_VectorVM.cpp b/Engine/Source/Editor/NiagaraEditor/Private/NiagaraCompiler_VectorVM.cpp
--- a/Engine/Source/Editor/NiagaraEditor/Private/NiagaraCompiler_VectorVM.cpp
+++ b/Engine/Source/Editor/NiagaraEditor/Private/NiagaraCompiler_VectorVM.cpp
@@ -180,7 +180,7 @@ void FNiagaraCompiler_VectorVM::CompileScript(UNiagaraScript* InScript)
 	check(InScript);
 
 	Script = InScript;
-	UNiagaraScriptSource* Source = CastChecked<UNiagaraScriptSource>(InScript->Source);
+	Source = CastChecked<UNiagaraScriptSource>(InScript->Source);
 
 	// Clone the source graph so we can modify it as needed; merging in the child graphs
 	UEdGraph* UpdateGraph = FEdGraphUtilities::CloneGraph(Source->UpdateGraph, Source, &MessageLog, true);
diff --git a/Engine/Source/Editor/NiagaraEditor/Private/NiagaraEffectEditor.cpp b/Engine/Source/Editor/NiagaraEditor/Private/NiagaraEffectEditor.cpp
--- a/Engine/Source/Editor/NiagaraEditor/Private/NiagaraEffectEditor.cpp
+++ b/Engine/Source/Editor/NiagaraEditor/Private/NiagaraEffectEditor.cpp
@@ -97,7 +97,7 @@ FLinearColor FNiagaraEffectEditor::GetWorldCentricTabColorScale() const
 TSharedRef<SNiagaraEffectEditorWidget> FNiagaraEffectEditor::CreateEditorWidget(UNiagaraEffect* InEffect)
 {
 	check(InEffect != NULL);
-	
+	EffectInstance = new FNiagaraEffectInstance(InEffect);
 	
 	if (!EditorCommands.IsValid())
 	{
@@ -108,8 +108,6 @@ TSharedRef<SNiagaraEffectEditorWidget> FNiagaraEffectEditor::CreateEditorWidget(
 	FGraphAppearanceInfo AppearanceInfo;
 	AppearanceInfo.CornerText = LOCTEXT("AppearanceCornerText", "NIAGARA").ToString();
 
-	//SGraphEditor::FGraphEditorEvents InEvents;
-
 	// Make title bar
 	TSharedRef<SWidget> TitleBarWidget =
 		SNew(SBorder)
@@ -127,12 +125,10 @@ TSharedRef<SNiagaraEffectEditorWidget> FNiagaraEffectEditor::CreateEditorWidget(
 			]
 		];
 
-	// make preview pane
-	
 		
 		
 	// Make the effect editor widget
-	return SNew(SNiagaraEffectEditorWidget).TitleBar(TitleBarWidget).EffectObj(InEffect);
+	return SNew(SNiagaraEffectEditorWidget).TitleBar(TitleBarWidget).EffectObj(InEffect).EffectInstance(EffectInstance);
 }
 
 
@@ -212,8 +208,9 @@ FReply FNiagaraEffectEditor::OnAddEmitterClicked()
 {
 	FNiagaraEditorModule& NiagaraEditorModule = FModuleManager::LoadModuleChecked<FNiagaraEditorModule>("NiagaraEditor");
 
-	Effect->AddEmitter();
-	UpdateEditorPtr.Pin()->GetViewport()->SetPreviewEffect(Effect);
+	FNiagaraEmitterProperties &Props = Effect->AddEmitterProperties();
+	EffectInstance->AddEmitter(&Props);
+	UpdateEditorPtr.Pin()->GetViewport()->SetPreviewEffect(EffectInstance);
 	UpdateEditorPtr.Pin()->UpdateList();
 	return FReply::Handled();
 }
diff --git a/Engine/Source/Editor/NiagaraEditor/Private/NiagaraEffectEditor.h b/Engine/Source/Editor/NiagaraEditor/Private/NiagaraEffectEditor.h
--- a/Engine/Source/Editor/NiagaraEditor/Private/NiagaraEffectEditor.h
+++ b/Engine/Source/Editor/NiagaraEditor/Private/NiagaraEffectEditor.h
@@ -46,7 +46,9 @@ private:
 
 	/* The Effect being edited */
 	UNiagaraEffect	*Effect;
-	
+	class FNiagaraEffectInstance *EffectInstance;
+
+
 	/** The command list for this editor */
 	TSharedPtr<FUICommandList> EditorCommands;
 
diff --git a/Engine/Source/Editor/NiagaraEditor/Private/SNiagaraEffectEditorViewport.cpp b/Engine/Source/Editor/NiagaraEditor/Private/SNiagaraEffectEditorViewport.cpp
--- a/Engine/Source/Editor/NiagaraEditor/Private/SNiagaraEffectEditorViewport.cpp
+++ b/Engine/Source/Editor/NiagaraEditor/Private/SNiagaraEffectEditorViewport.cpp
@@ -141,7 +141,7 @@ void SNiagaraEffectEditorViewport::Tick(const FGeometry& AllottedGeometry, const
 
 
 
-void SNiagaraEffectEditorViewport::SetPreviewEffect(UNiagaraEffect *InPreviewEffect)
+void SNiagaraEffectEditorViewport::SetPreviewEffect(FNiagaraEffectInstance *InPreviewEffect)
 {
 	check( PreviewComponent );
 
@@ -150,8 +150,8 @@ void SNiagaraEffectEditorViewport::SetPreviewEffect(UNiagaraEffect *InPreviewEff
 	PreviewScene.AddComponent(PreviewComponent, Transform);
 
 	FComponentReregisterContext ReregisterContext(PreviewComponent);
-	PreviewComponent->Effect = InPreviewEffect;
-	PreviewComponent->Effect->Init(PreviewComponent);
+	PreviewComponent->SetEffectInstance(InPreviewEffect);
+	PreviewComponent->GetEffectInstance()->Init(PreviewComponent);
 }
 
 
diff --git a/Engine/Source/Editor/NiagaraEditor/Private/SNiagaraEffectEditorViewport.h b/Engine/Source/Editor/NiagaraEditor/Private/SNiagaraEffectEditorViewport.h
--- a/Engine/Source/Editor/NiagaraEditor/Private/SNiagaraEffectEditorViewport.h
+++ b/Engine/Source/Editor/NiagaraEditor/Private/SNiagaraEffectEditorViewport.h
@@ -24,7 +24,7 @@ public:
 	virtual void Tick(const FGeometry& AllottedGeometry, const double InCurrentTime, const float InDeltaTime) override;
 	void RefreshViewport();
 	
-	void SetPreviewEffect(UNiagaraEffect* InEffect);
+	void SetPreviewEffect(FNiagaraEffectInstance* InEffect);
 	
 	void ToggleRealtime();
 	
diff --git a/Engine/Source/Editor/NiagaraEditor/Private/SNiagaraEffectEditorWidget.cpp b/Engine/Source/Editor/NiagaraEditor/Private/SNiagaraEffectEditorWidget.cpp
--- a/Engine/Source/Editor/NiagaraEditor/Private/SNiagaraEffectEditorWidget.cpp
+++ b/Engine/Source/Editor/NiagaraEditor/Private/SNiagaraEffectEditorWidget.cpp
@@ -52,6 +52,15 @@ void SNiagaraEffectEditorWidget::Construct(const FArguments& InArgs)
 	*/
 
 	EffectObj = InArgs._EffectObj;
+	if (InArgs._EffectInstance)
+	{
+		EffectInstance = InArgs._EffectInstance;
+	}
+	else
+	{
+		EffectInstance = new FNiagaraEffectInstance(EffectObj);
+	}
+
 	//bNeedsRefresh = false;
 
 
@@ -112,7 +121,7 @@ void SNiagaraEffectEditorWidget::Construct(const FArguments& InArgs)
 					[
 						SAssignNew(ListView, SListView<TSharedPtr<FNiagaraSimulation> >)
 						.ItemHeight(256)
-						.ListItemsSource(&(this->EffectObj->Emitters))
+						.ListItemsSource(&(this->EffectInstance->GetEmitters()))
 						.OnGenerateRow(this, &SNiagaraEffectEditorWidget::OnGenerateWidgetForList)
 					]
 
@@ -152,7 +161,7 @@ void SNiagaraEffectEditorWidget::Construct(const FArguments& InArgs)
 				]
 		];
 
-	Viewport->SetPreviewEffect(EffectObj);
+	Viewport->SetPreviewEffect(EffectInstance);
 }
 
 
@@ -164,7 +173,7 @@ void SEmitterWidget::Construct(const FArguments& InArgs)
 	CurSpawnScript = nullptr;
 	CurUpdateScript = nullptr;
 	Emitter = InArgs._Emitter;
-	Effect = InArgs._Effect;
+	EffectInstance = InArgs._Effect;
 	CurMaterial = nullptr;
 
 	ThumbnailPool = MakeShareable(new FAssetThumbnailPool(1, /*InAreRealTileThumbnailsAllowed=*/false));
diff --git a/Engine/Source/Editor/NiagaraEditor/Private/SNiagaraEffectEditorWidget.h b/Engine/Source/Editor/NiagaraEditor/Private/SNiagaraEffectEditorWidget.h
--- a/Engine/Source/Editor/NiagaraEditor/Private/SNiagaraEffectEditorWidget.h
+++ b/Engine/Source/Editor/NiagaraEditor/Private/SNiagaraEffectEditorWidget.h
@@ -23,7 +23,7 @@ class SEmitterWidget : public SCompoundWidget, public FNotifyHook
 {
 private:
 	TSharedPtr<FNiagaraSimulation> Emitter;
-	UNiagaraEffect *Effect;
+	FNiagaraEffectInstance *EffectInstance;
 	TSharedPtr<STextBlock> UpdateComboText, SpawnComboText;
 	TSharedPtr<SComboButton> UpdateCombo, SpawnCombo;
 	TSharedPtr<SContentReference> UpdateScriptSelector, SpawnScriptSelector;
@@ -41,7 +41,7 @@ public:
 	{
 	}
 	SLATE_ARGUMENT(TSharedPtr<FNiagaraSimulation>, Emitter)
-	SLATE_ARGUMENT(UNiagaraEffect*, Effect)
+	SLATE_ARGUMENT(FNiagaraEffectInstance*, Effect)
 	SLATE_END_ARGS()
 
 	void OnUpdateScriptSelectedFromPicker(UObject *Asset)
@@ -116,7 +116,7 @@ public:
 
 		ENQUEUE_UNIQUE_RENDER_COMMAND_THREEPARAMETER(
 			FChangeNiagaraRenderModule,
-			UNiagaraEffect*, InEffect, this->Effect,
+			FNiagaraEffectInstance*, InEffect, this->EffectInstance,
 			TSharedPtr<FNiagaraSimulation>, InEmitter, this->Emitter,
 			EEmitterRenderModuleType, InType, Type,
 		{
@@ -180,6 +180,7 @@ class SNiagaraEffectEditorWidget : public SCompoundWidget, public FNotifyHook
 {
 private:
 	UNiagaraEffect *EffectObj;
+	FNiagaraEffectInstance *EffectInstance;
 
 	/** Notification list to pass messages to editor users  */
 	TSharedPtr<SNotificationList> NotificationListPtr;
@@ -194,7 +195,8 @@ public:
 	{
 	}
 	SLATE_ARGUMENT(UNiagaraEffect*, EffectObj)
-		SLATE_ARGUMENT(TSharedPtr<SWidget>, TitleBar)
+	SLATE_ARGUMENT(TSharedPtr<SWidget>, TitleBar)
+	SLATE_ARGUMENT(FNiagaraEffectInstance*, EffectInstance)
 	SLATE_END_ARGS()
 
 
@@ -220,7 +222,7 @@ public:
 				//.BorderImage(FEditorStyle::GetBrush("ToolPanel.GroupBorder"))
 				.Padding(2.0f)
 				[
-					SNew(SEmitterWidget).Emitter(InItem).Effect(EffectObj)
+					SNew(SEmitterWidget).Emitter(InItem).Effect(EffectInstance)
 				]
 			];
 	}
diff --git a/Engine/Source/Editor/UnrealEd/Private/Factories/ActorFactory.cpp b/Engine/Source/Editor/UnrealEd/Private/Factories/ActorFactory.cpp
--- a/Engine/Source/Editor/UnrealEd/Private/Factories/ActorFactory.cpp
+++ b/Engine/Source/Editor/UnrealEd/Private/Factories/ActorFactory.cpp
@@ -557,7 +557,7 @@ void UActorFactoryNiagara::PostSpawnActor(UObject* Asset, AActor* NewActor)
 	NiagaraActor->GetNiagaraComponent()->UnregisterComponent();
 
 	// Change properties
-	NiagaraActor->GetNiagaraComponent()->Effect = Effect;
+	NiagaraActor->GetNiagaraComponent()->SetAsset(Effect);
 
 	// if we're created by Kismet on the server during gameplay, we need to replicate the emitter
 	if (NiagaraActor->GetWorld()->HasBegunPlay() && NiagaraActor->GetWorld()->GetNetMode() != NM_Client)
@@ -577,7 +577,7 @@ UObject* UActorFactoryNiagara::GetAssetFromActorInstance(AActor* Instance)
 	ANiagaraActor* NewActor = CastChecked<ANiagaraActor>(Instance);
 	if (NewActor->GetNiagaraComponent())
 	{
-		return NewActor->GetNiagaraComponent()->Effect;
+		return NewActor->GetNiagaraComponent()->GetAsset();
 	}
 	else
 	{
@@ -591,7 +591,7 @@ void UActorFactoryNiagara::PostCreateBlueprint(UObject* Asset, AActor* CDO)
 	{
 		UNiagaraEffect* Effect = CastChecked<UNiagaraEffect>(Asset);
 		ANiagaraActor* Actor = CastChecked<ANiagaraActor>(CDO);
-		Actor->GetNiagaraComponent()->Effect = Effect;
+		Actor->GetNiagaraComponent()->SetAsset(Effect);
 	}
 }
 
diff --git a/Engine/Source/Runtime/Engine/Classes/Components/NiagaraComponent.h b/Engine/Source/Runtime/Engine/Classes/Components/NiagaraComponent.h
--- a/Engine/Source/Runtime/Engine/Classes/Components/NiagaraComponent.h
+++ b/Engine/Source/Runtime/Engine/Classes/Components/NiagaraComponent.h
@@ -89,17 +89,19 @@ UCLASS()
 class ENGINE_API UNiagaraComponent : public UPrimitiveComponent
 {
 	GENERATED_UCLASS_BODY()
-
-	UPROPERTY(EditAnywhere, Category = NiagaraComponent)
-	class UNiagaraEffect *Effect;
+private:
+	UPROPERTY()
+	//class UNiagaraEffect *Effect;
+	class UNiagaraEffect *Asset;
+	class FNiagaraEffectInstance *EffectInstance;
 
 	// Begin UActorComponent interface.
-	virtual void TickComponent(float DeltaTime, enum ELevelTick TickType, FActorComponentTickFunction *ThisTickFunction) override;
 protected:
 	virtual void OnRegister() override;
 	virtual void OnUnregister()  override;
 	virtual void SendRenderDynamicData_Concurrent() override;
 public:
+	virtual void TickComponent(float DeltaTime, enum ELevelTick TickType, FActorComponentTickFunction *ThisTickFunction) override;
 	// End UActorComponent interface.
 
 	// Begin UPrimitiveComponent Interface
@@ -108,6 +110,12 @@ public:
 	virtual FPrimitiveSceneProxy* CreateSceneProxy() override;
 	// End UPrimitiveComponent Interface
 
+	void SetAsset(UNiagaraEffect *InAsset);
+	UNiagaraEffect *GetAsset() const { return Asset; }
+
+	FNiagaraEffectInstance *GetEffectInstance()	const { return EffectInstance; }
+	void SetEffectInstance(FNiagaraEffectInstance *InInstance)	{ EffectInstance = InInstance; }
+
 	// Begin UObject interface.
 #if WITH_EDITOR
 	virtual void PostEditChangeProperty(FPropertyChangedEvent& PropertyChangedEvent) override;
@@ -134,7 +142,7 @@ public:
 	void SetDynamicData_RenderThread(struct FNiagaraDynamicDataBase* NewDynamicData);
 	TArray<class NiagaraEffectRenderer*> &GetEffectRenderers() { return EffectRenderers; }
 	void AddEffectRenderer(NiagaraEffectRenderer *Renderer)	{ EffectRenderers.Add(Renderer); }
-	ENGINE_API void UpdateEffectRenderers(UNiagaraEffect *InEffect);
+	ENGINE_API void UpdateEffectRenderers(FNiagaraEffectInstance *InEffect);
 
 private:
 	void ReleaseRenderThreadResources();
diff --git a/Engine/Source/Runtime/Engine/Classes/Engine/NiagaraEffect.h b/Engine/Source/Runtime/Engine/Classes/Engine/NiagaraEffect.h
--- a/Engine/Source/Runtime/Engine/Classes/Engine/NiagaraEffect.h
+++ b/Engine/Source/Runtime/Engine/Classes/Engine/NiagaraEffect.h
@@ -11,7 +11,7 @@ class UNiagaraEffect : public UObject
 	GENERATED_UCLASS_BODY()
 
 public:
-	ENGINE_API TSharedPtr<FNiagaraSimulation> AddEmitter();
+	ENGINE_API FNiagaraEmitterProperties &AddEmitterProperties();
 
 
 	void Init(UNiagaraComponent *InComponent)
@@ -19,6 +19,64 @@ public:
 		Component = InComponent;
 	}
 
+	TArray<FNiagaraEmitterProperties> &GetEmitterProperties() {
+		return EmitterProps;
+	};
+
+	virtual void PostLoad() override; 
+
+private:
+	UNiagaraComponent *Component;
+
+	UPROPERTY()
+	TArray<FNiagaraEmitterProperties>EmitterProps;
+};
+
+
+
+
+
+
+class FNiagaraEffectInstance
+{
+public:
+	explicit FNiagaraEffectInstance(UNiagaraEffect *InAsset, UNiagaraComponent *InComponent)
+	{
+		Component = InComponent;
+		for (int32 i = 0; i < InAsset->GetEmitterProperties().Num(); i++)
+		{
+			FNiagaraSimulation *Sim = new FNiagaraSimulation(&InAsset->GetEmitterProperties()[i], Component->GetWorld()->FeatureLevel);
+			Emitters.Add(MakeShareable(Sim));
+		}
+		InitRenderModules(Component->GetWorld()->FeatureLevel);
+	}
+
+	explicit FNiagaraEffectInstance(UNiagaraEffect *InAsset)
+		: Component(nullptr)
+	{
+		for (FNiagaraEmitterProperties &Props : InAsset->GetEmitterProperties())
+		{
+			FNiagaraSimulation *Sim = new FNiagaraSimulation(&Props);
+			Emitters.Add(MakeShareable(Sim));
+		}
+	}
+
+
+	void InitRenderModules(ERHIFeatureLevel::Type InFeatureLevel)
+	{
+		for (TSharedPtr<FNiagaraSimulation> Sim : Emitters)
+		{
+			Sim->SetRenderModuleType(Sim->GetProperties()->RenderModuleType, InFeatureLevel);
+		}
+	}
+
+	void Init(UNiagaraComponent *InComponent)
+	{
+		Component = InComponent;
+		//InitRenderModules(Component->GetWorld()->FeatureLevel);
+	}
+
+	ENGINE_API void AddEmitter(FNiagaraEmitterProperties *Properties);
 
 
 	void SetConstant(FName ConstantName, const float Value)
@@ -59,18 +117,18 @@ public:
 		return Emitters[idx].Get();
 	}
 
-	/** Local constant table. */
-	FNiagaraConstantMap Constants;
 
-	//UPROPERTY()
-	TArray< TSharedPtr<FNiagaraSimulation> > Emitters;
 	UNiagaraComponent *GetComponent() { return Component; }
 
-	virtual void PostLoad() override; 
+	TArray< TSharedPtr<FNiagaraSimulation> > &GetEmitters()	{ return Emitters; }
 
 private:
 	UNiagaraComponent *Component;
 
-	UPROPERTY()
-	TArray<FNiagaraEmitterProperties>EmitterProps;
-};
+	/** Local constant table. */
+	FNiagaraConstantMap Constants;
+
+	TArray< TSharedPtr<FNiagaraSimulation> > Emitters;
+
+
+};
\ No newline at end of file
diff --git a/Engine/Source/Runtime/Engine/Classes/Engine/NiagaraSimulation.h b/Engine/Source/Runtime/Engine/Classes/Engine/NiagaraSimulation.h
--- a/Engine/Source/Runtime/Engine/Classes/Engine/NiagaraSimulation.h
+++ b/Engine/Source/Runtime/Engine/Classes/Engine/NiagaraSimulation.h
@@ -37,7 +37,7 @@ struct FNiagaraEmitterProperties
 {
 	GENERATED_USTRUCT_BODY()
 public:
-	FNiagaraEmitterProperties() : RenderModuleType(RMT_Sprites)
+	FNiagaraEmitterProperties() : RenderModuleType(RMT_Sprites), Material(nullptr), SpawnRate(50), UpdateScript(nullptr), SpawnScript(nullptr), bIsEnabled(true)
 	{
 	}
 
@@ -63,8 +63,8 @@ public:
 struct FNiagaraSimulation
 {
 public:
-	explicit FNiagaraSimulation(FNiagaraEmitterProperties *InProps);
-	FNiagaraSimulation(FNiagaraEmitterProperties &Props, ERHIFeatureLevel::Type InFeatureLeve);
+	explicit ENGINE_API FNiagaraSimulation(FNiagaraEmitterProperties *InProps);
+	ENGINE_API FNiagaraSimulation(FNiagaraEmitterProperties *Props, ERHIFeatureLevel::Type InFeatureLevel);
 	virtual ~FNiagaraSimulation()
 	{}
 
diff --git a/Engine/Source/Runtime/Engine/Private/NiagaraComponent.cpp b/Engine/Source/Runtime/Engine/Private/NiagaraComponent.cpp
--- a/Engine/Source/Runtime/Engine/Private/NiagaraComponent.cpp
+++ b/Engine/Source/Runtime/Engine/Private/NiagaraComponent.cpp
@@ -21,15 +21,15 @@ DECLARE_DWORD_COUNTER_STAT(TEXT("NumParticles"),STAT_NiagaraNumParticles,STATGRO
 FNiagaraSceneProxy::FNiagaraSceneProxy(const UNiagaraComponent* InComponent)
 		:	FPrimitiveSceneProxy(InComponent)
 {
-	UpdateEffectRenderers(InComponent->Effect);
+	UpdateEffectRenderers(InComponent->GetEffectInstance());
 }
 
-void FNiagaraSceneProxy::UpdateEffectRenderers(UNiagaraEffect *InEffect)
+void FNiagaraSceneProxy::UpdateEffectRenderers(FNiagaraEffectInstance *InEffect)
 {
 	EffectRenderers.Empty();
 	if (InEffect)
 	{
-		for (TSharedPtr<FNiagaraSimulation>Emitter : InEffect->Emitters)
+		for (TSharedPtr<FNiagaraSimulation>Emitter : InEffect->GetEmitters())
 		{
 			AddEffectRenderer(Emitter->GetEffectRenderer());
 		}
@@ -189,27 +189,25 @@ void UNiagaraComponent::TickComponent(float DeltaSeconds, enum ELevelTick TickTy
 {
 //	EmitterAge += DeltaSeconds;
 
-	if (Effect)
+	if (EffectInstance)
 	{ 
 
 		//Todo, open this up to the UI and setting via code and BPs.
 		static FName Const_Zero(TEXT("ZERO"));
-		Effect->SetConstant(Const_Zero, FVector4(0.0f, 0.0f, 0.0f, 0.0f));	// zero constant
+		EffectInstance->SetConstant(Const_Zero, FVector4(0.0f, 0.0f, 0.0f, 0.0f));	// zero constant
 		static FName Const_DeltaTime(TEXT("Delta Time"));
-		Effect->SetConstant(Const_DeltaTime, FVector4(DeltaSeconds, DeltaSeconds, DeltaSeconds, DeltaSeconds));
+		EffectInstance->SetConstant(Const_DeltaTime, FVector4(DeltaSeconds, DeltaSeconds, DeltaSeconds, DeltaSeconds));
 		static FName Const_EmitterPos(TEXT("Emitter Position"));
-		Effect->SetConstant(Const_EmitterPos, FVector4(ComponentToWorld.GetTranslation()));
-		//static FName Const_Age(TEXT("Emitter Age"));
-		//Effect->SetConstant(Const_Age, FVector4(EmitterAge, EmitterAge, EmitterAge, EmitterAge));
+		EffectInstance->SetConstant(Const_EmitterPos, FVector4(ComponentToWorld.GetTranslation()));
 		static FName Const_EmitterX(TEXT("Emitter X Axis"));
-		Effect->SetConstant(Const_EmitterX, FVector4(ComponentToWorld.GetUnitAxis(EAxis::X)));
+		EffectInstance->SetConstant(Const_EmitterX, FVector4(ComponentToWorld.GetUnitAxis(EAxis::X)));
 		static FName Const_EmitterY(TEXT("Emitter Y Axis"));
-		Effect->SetConstant(Const_EmitterY, FVector4(ComponentToWorld.GetUnitAxis(EAxis::Y)));
+		EffectInstance->SetConstant(Const_EmitterY, FVector4(ComponentToWorld.GetUnitAxis(EAxis::Y)));
 		static FName Const_EmitterZ(TEXT("Emitter Z Axis"));
-		Effect->SetConstant(Const_EmitterZ, FVector4(ComponentToWorld.GetUnitAxis(EAxis::Z)));
+		EffectInstance->SetConstant(Const_EmitterZ, FVector4(ComponentToWorld.GetUnitAxis(EAxis::Z)));
 		static FName Const_EmitterTransform(TEXT("Emitter Transform"));
-		Effect->SetConstant(Const_EmitterTransform, ComponentToWorld.ToMatrixWithScale());
-		Effect->Tick(DeltaSeconds);
+		EffectInstance->SetConstant(Const_EmitterTransform, ComponentToWorld.ToMatrixWithScale());
+		EffectInstance->Tick(DeltaSeconds);
 	}
 
 	UpdateComponentToWorld();
@@ -219,24 +217,30 @@ void UNiagaraComponent::TickComponent(float DeltaSeconds, enum ELevelTick TickTy
 void UNiagaraComponent::OnRegister()
 {
 	Super::OnRegister();
-	if (Effect)
+	if (Asset)
 	{
-		Effect->Init(this);
+		if (!EffectInstance)
+		{
+			EffectInstance = new FNiagaraEffectInstance(Asset, this);
+		}
+		{
+			EffectInstance->Init(this);
 
-		// initialize all render modules
-		ENQUEUE_UNIQUE_RENDER_COMMAND_TWOPARAMETER(
-			FChangeNiagaraRenderModule,
-			UNiagaraEffect*, InEffect, this->Effect,
-			UNiagaraComponent*, InComponent, this,
-			{
-				for (TSharedPtr<FNiagaraSimulation> Emitter : InEffect->Emitters)
+			// initialize all render modules
+			ENQUEUE_UNIQUE_RENDER_COMMAND_TWOPARAMETER(
+				FChangeNiagaraRenderModule,
+				FNiagaraEffectInstance*, InEffect, this->EffectInstance,
+				UNiagaraComponent*, InComponent, this,
 				{
-					Emitter->SetRenderModuleType(Emitter->GetProperties()->RenderModuleType, InComponent->GetWorld()->FeatureLevel);
+					for (TSharedPtr<FNiagaraSimulation> Emitter : InEffect->GetEmitters())
+					{
+						Emitter->SetRenderModuleType(Emitter->GetProperties()->RenderModuleType, InComponent->GetWorld()->FeatureLevel);
+					}
 				}
-			}
-		);
+			);
+		}
+		VectorVM::Init();
 	}
-	VectorVM::Init();
 }
 
 
@@ -247,21 +251,25 @@ void UNiagaraComponent::OnUnregister()
 
 void UNiagaraComponent::SendRenderDynamicData_Concurrent()
 {
-	if (Effect && SceneProxy)
+	if (EffectInstance && SceneProxy)
 	{
 		FNiagaraSceneProxy *NiagaraProxy = static_cast<FNiagaraSceneProxy*>(SceneProxy);
-		for (int32 i = 0; i < Effect->Emitters.Num(); i++)
+		for (int32 i = 0; i < EffectInstance->GetEmitters().Num(); i++)
 		{
-			FNiagaraSimulation *Emitter = Effect->GetEmitter(i);
-			FNiagaraDynamicDataBase* DynamicData = Emitter->GetEffectRenderer()->GenerateVertexData(Effect->Emitters[i]->GetData());
-
-			ENQUEUE_UNIQUE_RENDER_COMMAND_TWOPARAMETER(
-				FSendNiagaraDynamicData,
-				NiagaraEffectRenderer*, EffectRenderer, Emitter->GetEffectRenderer(),
-				FNiagaraDynamicDataBase*, DynamicData, DynamicData,
-				{
-				EffectRenderer->SetDynamicData_RenderThread(DynamicData);
-			});
+			TSharedPtr<FNiagaraSimulation>Emitter = EffectInstance->GetEmitters()[i];
+			NiagaraEffectRenderer *Renderer = Emitter->GetEffectRenderer();
+			if (Renderer)
+			{
+				FNiagaraDynamicDataBase* DynamicData = Renderer->GenerateVertexData(Emitter->GetData());
+
+				ENQUEUE_UNIQUE_RENDER_COMMAND_TWOPARAMETER(
+					FSendNiagaraDynamicData,
+					NiagaraEffectRenderer*, EffectRenderer, Emitter->GetEffectRenderer(),
+					FNiagaraDynamicDataBase*, DynamicData, DynamicData,
+					{
+					EffectRenderer->SetDynamicData_RenderThread(DynamicData);
+				});
+			}
 		}
 	}
 
@@ -296,7 +304,7 @@ FBoxSphereBounds UNiagaraComponent::CalcBounds(const FTransform& LocalToWorld) c
 FPrimitiveSceneProxy* UNiagaraComponent::CreateSceneProxy()
 {
 	FNiagaraSceneProxy *Proxy = new FNiagaraSceneProxy(this);
-	Proxy->UpdateEffectRenderers(Effect);
+	Proxy->UpdateEffectRenderers(EffectInstance);
 	return Proxy;
 }
 
@@ -308,3 +316,12 @@ void UNiagaraComponent::PostEditChangeProperty(FPropertyChangedEvent& PropertyCh
 	FComponentReregisterContext ReregisterContext(this);
 }
 #endif // WITH_EDITOR
+
+
+
+void UNiagaraComponent::SetAsset(UNiagaraEffect *InAsset)
+{
+	Asset = InAsset;
+
+	EffectInstance = new FNiagaraEffectInstance(Asset, this);
+}
diff --git a/Engine/Source/Runtime/Engine/Private/NiagaraEffect.cpp b/Engine/Source/Runtime/Engine/Private/NiagaraEffect.cpp
--- a/Engine/Source/Runtime/Engine/Private/NiagaraEffect.cpp
+++ b/Engine/Source/Runtime/Engine/Private/NiagaraEffect.cpp
@@ -13,34 +13,32 @@ UNiagaraEffect::UNiagaraEffect(const FObjectInitializer& ObjectInitializer)
 
 
 
-TSharedPtr<FNiagaraSimulation> UNiagaraEffect::AddEmitter()
+FNiagaraEmitterProperties &UNiagaraEffect::AddEmitterProperties()
 {
 	FNiagaraEmitterProperties Props;
 	EmitterProps.Add(Props);
+	return EmitterProps.Top();
+}
 
-	TSharedPtr<FNiagaraSimulation> Sim = MakeShareable(new FNiagaraSimulation(&EmitterProps.Top()));
 
-	Sim->SetRenderModuleType(RMT_Sprites, Component->GetWorld()->FeatureLevel);
-	Emitters.Add(Sim);
 
-	FNiagaraSceneProxy *SceneProxy = static_cast<FNiagaraSceneProxy*>(Component->SceneProxy);
-	SceneProxy->UpdateEffectRenderers(this);
 
-	return Sim;
+void UNiagaraEffect::PostLoad()
+{
+	Super::PostLoad();
 }
 
 
 
 
-void UNiagaraEffect::PostLoad()
+
+void FNiagaraEffectInstance::AddEmitter(FNiagaraEmitterProperties *Properties)
 {
-	Super::PostLoad();
+	TSharedPtr<FNiagaraSimulation> Sim = MakeShareable(new FNiagaraSimulation(Properties));
 
-	Emitters.Empty();
+	Sim->SetRenderModuleType(RMT_Sprites, Component->GetWorld()->FeatureLevel);
+	Emitters.Add(Sim);
+	FNiagaraSceneProxy *SceneProxy = static_cast<FNiagaraSceneProxy*>(Component->SceneProxy);
+	SceneProxy->UpdateEffectRenderers(this);
 
-	for (FNiagaraEmitterProperties &Props : EmitterProps)
-	{
-		FNiagaraSimulation *Sim = new FNiagaraSimulation(&Props);
-		Emitters.Add(MakeShareable(Sim));
-	}
-}
+}
\ No newline at end of file
diff --git a/Engine/Source/Runtime/Engine/Private/NiagaraEffectRenderer.cpp b/Engine/Source/Runtime/Engine/Private/NiagaraEffectRenderer.cpp
--- a/Engine/Source/Runtime/Engine/Private/NiagaraEffectRenderer.cpp
+++ b/Engine/Source/Runtime/Engine/Private/NiagaraEffectRenderer.cpp
@@ -343,7 +343,7 @@ FNiagaraDynamicDataBase *NiagaraEffectRendererSprites::GenerateVertexData(const
 		NewVertex.ParticleId = ParticleId;
 		ParticleId += IdInc;
 		NewVertex.RelativeTime = AgePtr[ParticleIndex].X;
-		NewVertex.Size = FVector2D(1.0f, 1.0f);
+		NewVertex.Size = FVector2D(2.0f, 2.0f);
 		NewVertex.Rotation = RotPtr[ParticleIndex].X;
 		NewVertex.SubImageIndex = 0.f;
 
