Merging displayalllocation and displayallrotation improvements

//depot/UE4-UT/...

to //depot/UE4/...

diff --git a/Engine/Source/Runtime/Engine/Private/UnrealEngine.cpp b/Engine/Source/Runtime/Engine/Private/UnrealEngine.cpp
--- a/Engine/Source/Runtime/Engine/Private/UnrealEngine.cpp
+++ b/Engine/Source/Runtime/Engine/Private/UnrealEngine.cpp
@@ -6520,12 +6520,36 @@ static void DrawProperty(UCanvas* CanvasObject, UObject* Obj, const FDebugDispla
 			if (PropData.PropertyName == NAME_Location)
 			{
 				AActor *Actor = Cast<AActor>(Obj);
-				ValueText = FString::Printf(TEXT("%s"), Actor != NULL ? *Actor->GetActorLocation().ToString() : TEXT("None"));
+				USceneComponent *Component = Cast<USceneComponent>(Obj);
+				if (Actor != nullptr)
+				{
+					ValueText = Actor->GetActorLocation().ToString();
+				}
+				else if (Component != nullptr)
+				{
+					ValueText = Component->GetComponentLocation().ToString();
+				}
+				else
+				{
+					ValueText = TEXT("Unsupported for this type");
+				}
 			}
 			else if (PropData.PropertyName == NAME_Rotation)
 			{
 				AActor *Actor = Cast<AActor>(Obj);
-				ValueText = FString::Printf(TEXT("%s"), Actor != NULL ? *Actor->GetActorRotation().ToString() : TEXT("None"));
+				USceneComponent *Component = Cast<USceneComponent>(Obj);
+				if (Actor != nullptr)
+				{
+					ValueText = Actor->GetActorRotation().ToString();
+				}
+				else if (Component != nullptr)
+				{
+					ValueText = Component->GetComponentRotation().ToString();
+				}
+				else
+				{
+					ValueText = TEXT("Unsupported for this type");
+				}
 			}
 		}
 	}
