CIS fix - Fixed crash in postload for montages with empty animation tracks.

diff --git a/Engine/Source/Runtime/Engine/Private/Animation/AnimCompositeBase.cpp b/Engine/Source/Runtime/Engine/Private/Animation/AnimCompositeBase.cpp
--- a/Engine/Source/Runtime/Engine/Private/Animation/AnimCompositeBase.cpp
+++ b/Engine/Source/Runtime/Engine/Private/Animation/AnimCompositeBase.cpp
@@ -321,18 +321,21 @@ int32 FAnimTrack::GetTrackAdditiveType() const
 void FAnimTrack::ValidateSegmentTimes()
 {
 	// rearrange, make sure there are no gaps between and all start times are correctly set
-	AnimSegments[0].StartPos = 0.0f;
-	for(int32 J = 0; J < AnimSegments.Num(); J++)
+	if(AnimSegments.Num() > 0)
 	{
-		FAnimSegment& Segment = AnimSegments[J];
-		if(J > 0)
+		AnimSegments[0].StartPos = 0.0f;
+		for(int32 J = 0; J < AnimSegments.Num(); J++)
 		{
-			Segment.StartPos = AnimSegments[J - 1].StartPos + AnimSegments[J - 1].GetLength();
-		}
+			FAnimSegment& Segment = AnimSegments[J];
+			if(J > 0)
+			{
+				Segment.StartPos = AnimSegments[J - 1].StartPos + AnimSegments[J - 1].GetLength();
+			}
 
-		if(Segment.AnimReference && Segment.AnimEndTime > Segment.AnimReference->SequenceLength)
-		{
-			Segment.AnimEndTime = Segment.AnimReference->SequenceLength;
+			if(Segment.AnimReference && Segment.AnimEndTime > Segment.AnimReference->SequenceLength)
+			{
+				Segment.AnimEndTime = Segment.AnimReference->SequenceLength;
+			}
 		}
 	}
 }
