Fixed Clang compile error in NiagaraSimulation.h

diff --git a/Engine/Source/Runtime/Engine/Classes/Engine/NiagaraSimulation.h b/Engine/Source/Runtime/Engine/Classes/Engine/NiagaraSimulation.h
--- a/Engine/Source/Runtime/Engine/Classes/Engine/NiagaraSimulation.h
+++ b/Engine/Source/Runtime/Engine/Classes/Engine/NiagaraSimulation.h
@@ -37,7 +37,7 @@ struct FNiagaraEmitterProperties
 {
 	GENERATED_USTRUCT_BODY()
 public:
-	FNiagaraEmitterProperties() : RenderModuleType(RMT_Sprites), Material(nullptr), SpawnRate(50), UpdateScript(nullptr), SpawnScript(nullptr), bIsEnabled(true)
+	FNiagaraEmitterProperties() : bIsEnabled(true), SpawnRate(50), UpdateScript(nullptr), SpawnScript(nullptr), Material(nullptr), RenderModuleType(RMT_Sprites)
 	{
 	}
 
