Missing change. For some reason this file was part of the last checkin, but the change didn't get submitted.

diff --git a/Engine/Source/Runtime/Engine/Private/NiagaraSimulation.cpp b/Engine/Source/Runtime/Engine/Private/NiagaraSimulation.cpp
--- a/Engine/Source/Runtime/Engine/Private/NiagaraSimulation.cpp
+++ b/Engine/Source/Runtime/Engine/Private/NiagaraSimulation.cpp
@@ -5,7 +5,8 @@
 #include "VectorVM.h"
 
 
-FNiagaraSimulation::FNiagaraSimulation(FNiagaraEmitterProperties *InProps) : Age(0.0f)
+FNiagaraSimulation::FNiagaraSimulation(FNiagaraEmitterProperties *InProps) 
+: Age(0.0f)
 , bIsEnabled(true)
 , SpawnRemainder(0.0f)
 , CachedBounds(ForceInit)
@@ -14,6 +15,11 @@ FNiagaraSimulation::FNiagaraSimulation(FNiagaraEmitterProperties *InProps) : Age
 	Props = InProps;
 }
 
+FNiagaraSimulation::FNiagaraSimulation(FNiagaraEmitterProperties *InProps, ERHIFeatureLevel::Type InFeatureLevel)
+:FNiagaraSimulation(InProps)
+{
+	SetRenderModuleType(InProps->RenderModuleType, InFeatureLevel);
+}
 
 void FNiagaraSimulation::Tick(float DeltaSeconds)
 {
