Include the address of the FirstNode in the NodeWaveInstanceHash

diff --git a/Engine/Source/Runtime/Engine/Private/SoundCue.cpp b/Engine/Source/Runtime/Engine/Private/SoundCue.cpp
--- a/Engine/Source/Runtime/Engine/Private/SoundCue.cpp
+++ b/Engine/Source/Runtime/Engine/Private/SoundCue.cpp
@@ -331,7 +331,7 @@ void USoundCue::Parse( FAudioDevice* AudioDevice, const UPTRINT NodeWaveInstance
 {
 	if (FirstNode)
 	{
-		FirstNode->ParseNodes(AudioDevice,NodeWaveInstanceHash,ActiveSound,ParseParams,WaveInstances);
+		FirstNode->ParseNodes(AudioDevice,(UPTRINT)FirstNode,ActiveSound,ParseParams,WaveInstances);
 	}
 }
 
