Fix focus issues when bringing up the console in PIE

diff --git a/Engine/Source/Runtime/Engine/Private/UserInterface/Console.cpp b/Engine/Source/Runtime/Engine/Private/UserInterface/Console.cpp
--- a/Engine/Source/Runtime/Engine/Private/UserInterface/Console.cpp
+++ b/Engine/Source/Runtime/Engine/Private/UserInterface/Console.cpp
@@ -1231,12 +1231,16 @@ void UConsole::FakeGotoState(FName NextStateName)
 	if (NextStateName == NAME_Typing)
 	{
 		BeginState_Typing(ConsoleState);
+		auto CurrentFocus = FSlateApplication::Get().GetKeyboardFocusedWidget();
 		FSlateApplication::Get().ResetToDefaultInputSettings();
+		FSlateApplication::Get().SetKeyboardFocus( CurrentFocus );
 	}
 	else if (NextStateName == NAME_Open)
 	{
 		BeginState_Open(ConsoleState);
+		auto CurrentFocus = FSlateApplication::Get().GetKeyboardFocusedWidget();
 		FSlateApplication::Get().ResetToDefaultInputSettings();
+		FSlateApplication::Get().SetKeyboardFocus( CurrentFocus );
 	}
 	else if( NextStateName == NAME_None )
 	{
