#ue4 - Change Pawn::PostNetReceiveLocationAndRotation() to use SetActorLocationAndRotation() rather than TeleportTo(), for efficiency.

diff --git a/Engine/Source/Runtime/Engine/Private/Pawn.cpp b/Engine/Source/Runtime/Engine/Private/Pawn.cpp
--- a/Engine/Source/Runtime/Engine/Private/Pawn.cpp
+++ b/Engine/Source/Runtime/Engine/Private/Pawn.cpp
@@ -990,8 +990,7 @@ void APawn::PostNetReceiveLocationAndRotation()
 		ReplicatedMovement.Location.Z += 0.01f;
 
 		const FVector OldLocation = GetActorLocation();
-		TeleportTo( ReplicatedMovement.Location, ReplicatedMovement.Rotation, false, true );
-		// SetActorLocationAndRotation(ReplicatedMovement.Location, ReplicatedMovement.Rotation); <-- preferred, but awaiting answer to question about UpdateNavOctree() missing in SceneComponent::MoveComponent
+		SetActorLocationAndRotation(ReplicatedMovement.Location, ReplicatedMovement.Rotation, /*bSweep=*/ false);
 
 		INetworkPredictionInterface* PredictionInterface = Cast<INetworkPredictionInterface>(GetMovementComponent());
 		if (PredictionInterface)
