#ue4 - Made the following GameNetworkManager functions virtual: WithinUpdateDelayBounds, ExceedsAllowablePositionError, NetworkVelocityNearZero.

diff --git a/Engine/Source/Runtime/Engine/Classes/GameFramework/GameNetworkManager.h b/Engine/Source/Runtime/Engine/Classes/GameFramework/GameNetworkManager.h
--- a/Engine/Source/Runtime/Engine/Classes/GameFramework/GameNetworkManager.h
+++ b/Engine/Source/Runtime/Engine/Classes/GameFramework/GameNetworkManager.h
@@ -122,13 +122,13 @@ class ENGINE_API AGameNetworkManager : public AInfo
 	// Player replication
 
 	/** @return true if last player client to server update was sufficiently recent.  Used to limit frequency of corrections if connection speed is limited. */
-	bool WithinUpdateDelayBounds(class APlayerController* PC, float LastUpdateTime) const;
+	virtual bool WithinUpdateDelayBounds(class APlayerController* PC, float LastUpdateTime) const;
 
 	/** @return true if position error exceeds max allowable amount */
-	bool ExceedsAllowablePositionError(FVector LocDiff) const;
+	virtual bool ExceedsAllowablePositionError(FVector LocDiff) const;
 
 	/** @return true if velocity vector passed in is considered near zero for networking purposes */
-	bool NetworkVelocityNearZero(FVector InVelocity) const;
+	virtual bool NetworkVelocityNearZero(FVector InVelocity) const;
 	virtual void PostInitializeComponents() override;
 
 	/** @RETURN new per/client bandwidth given number of players in the game */
